[INTEGRATE] Change 2418858 by [REDACTED_PERSON][REDACTED_HANDLE] on 2015/01/26 13:33:26

	Added method to create a new blueprinted component from the AddComponent combo

diff --git a/Engine/Source/Editor/Kismet/Private/SSCSEditor.cpp b/Engine/Source/Editor/Kismet/Private/SSCSEditor.cpp
--- a/Engine/Source/Editor/Kismet/Private/SSCSEditor.cpp
+++ b/Engine/Source/Editor/Kismet/Private/SSCSEditor.cpp
@@ -45,6 +45,8 @@
 
 #include "GameProjectGenerationModule.h"
 #include "HotReloadInterface.h"
+#include "AssetRegistryModule.h"
+#include "SCreateAssetFromActor.h"
 
 #define LOCTEXT_NAMESPACE "SSCSEditor"
 
@@ -3359,23 +3361,13 @@ void SSCSEditor::PerformComboAddClass(TSubclassOf<UActorComponent> ComponentClas
 
 	if( ComponentCreateAction == EComponentCreateAction::CreateNewCPPClass )
 	{
-		TSharedPtr<SWindow> ParentWindow = FSlateApplication::Get().FindWidgetWindow( SharedThis(this) );
-
-		FString AddedClassName;
-		auto OnCodeAddedToProject = [&AddedClassName](const FString& ClassName, const FString& ClassPath, const FString& ModuleName )
-		{
-			if( !ClassName.IsEmpty() && !ClassPath.IsEmpty() )
-			{
-				AddedClassName = FString::Printf(TEXT("/Script/%s.%s"), *ModuleName, *ClassName);
-			}
-		};
-
-		const bool bModal = true;
-		FGameProjectGenerationModule::Get().OpenAddCodeToProjectDialog( ComponentClass, FString(), ParentWindow, bModal, FOnCodeAddedToProject::CreateLambda(OnCodeAddedToProject) );
-
-		NewClass = LoadClass<UActorComponent>(nullptr, *AddedClassName, nullptr, LOAD_None, nullptr );
+		NewClass = CreateNewCPPComponent( ComponentClass );
 	}
-	
+	else if( ComponentCreateAction == EComponentCreateAction::CreateNewBlueprintClass )
+	{
+		NewClass = CreateNewBPComponent( ComponentClass );
+	}
+
 	if( NewClass != nullptr )
 	{
 		FEditorDelegates::LoadSelectedAssetsIfNeeded.Broadcast();
@@ -3828,6 +3820,80 @@ TSharedPtr<FSCSEditorTreeNode> SSCSEditor::AddRootComponentTreeNode(UActorCompon
 	return NewTreeNode;
 }
 
+UClass* SSCSEditor::CreateNewCPPComponent( TSubclassOf<UActorComponent> ComponentClass )
+{
+	TSharedPtr<SWindow> ParentWindow = FSlateApplication::Get().FindWidgetWindow(SharedThis(this));
+
+	FString AddedClassName;
+	auto OnCodeAddedToProject = [&AddedClassName](const FString& ClassName, const FString& ClassPath, const FString& ModuleName)
+	{
+		if(!ClassName.IsEmpty() && !ClassPath.IsEmpty())
+		{
+			AddedClassName = FString::Printf(TEXT("/Script/%s.%s"), *ModuleName, *ClassName);
+		}
+	};
+
+	const bool bModal = true;
+	FGameProjectGenerationModule::Get().OpenAddCodeToProjectDialog(ComponentClass, FString(), ParentWindow, bModal, FOnCodeAddedToProject::CreateLambda(OnCodeAddedToProject));
+
+	return LoadClass<UActorComponent>(nullptr, *AddedClassName, nullptr, LOAD_None, nullptr);
+}
+
+UClass* SSCSEditor::CreateNewBPComponent(TSubclassOf<UActorComponent> ComponentClass)
+{
+	UClass* NewClass = nullptr;
+
+	TSharedPtr<SWindow> PickBlueprintPathWidget;
+	SAssignNew(PickBlueprintPathWidget, SWindow)
+		.Title(LOCTEXT("SelectPath", "Select Path"))
+		.ToolTipText(LOCTEXT("SelectPathTooltip", "Select the path where the Blueprint will be created"))
+		.ClientSize(FVector2D(400, 400));
+
+
+	FString PackagePath;
+	FString AssetName;
+	auto OnPathPicked = [&PackagePath, &AssetName](const FString& Path)
+	{
+		PackagePath = Path;
+		AssetName = FPackageName::GetLongPackageAssetName(Path);
+	};
+
+	TSharedPtr<SCreateAssetFromObject> CreateBlueprintFromActorDialog;
+	PickBlueprintPathWidget->SetContent
+	(
+		SAssignNew(CreateBlueprintFromActorDialog, SCreateAssetFromObject, PickBlueprintPathWidget)
+		.HeadingText(LOCTEXT("CreateBlueprintFromActor_Heading", "Blueprint Name"))
+		.CreateButtonText(LOCTEXT("CreateBlueprintFromActor_ButtonLabel", "Create Blueprint"))
+		.DefaultNameOverride(FText::FromString(TEXT("BlueprintComponent")))
+		.OnCreateAssetAction(FOnPathChosen::CreateLambda(OnPathPicked))
+	);
+
+
+	FSlateApplication::Get().AddModalWindow(PickBlueprintPathWidget.ToSharedRef(), AsShared());
+
+	if(!PackagePath.IsEmpty() && !AssetName.IsEmpty())
+	{
+		UPackage* Package = CreatePackage(NULL, *PackagePath);
+
+		if(Package)
+		{
+			// Create and init a new Blueprint
+			if(UBlueprint* NewBP = FKismetEditorUtilities::CreateBlueprint(ComponentClass, Package, FName(*AssetName), BPTYPE_Normal, UBlueprint::StaticClass(), UBlueprintGeneratedClass::StaticClass()))
+			{
+				// Notify the asset registry
+				FAssetRegistryModule::AssetCreated(NewBP);
+
+				// Mark the package dirty...
+				Package->MarkPackageDirty();
+
+				NewClass = NewBP->GeneratedClass;
+			}
+		}
+	}
+
+	return NewClass;
+}
+
 void SSCSEditor::ClearSelection()
 {
 	if ( bUpdatingSelection == false )
diff --git a/Engine/Source/Editor/Kismet/Public/SSCSEditor.h b/Engine/Source/Editor/Kismet/Public/SSCSEditor.h
--- a/Engine/Source/Editor/Kismet/Public/SSCSEditor.h
+++ b/Engine/Source/Editor/Kismet/Public/SSCSEditor.h
@@ -726,6 +726,23 @@ protected:
 	/** Adds a root component tree node */
 	TSharedPtr<FSCSEditorTreeNode> AddRootComponentTreeNode(UActorComponent* ActorComp);
 
+	/**
+	 * Creates a new C++ component from the specified class type
+	 * The user will be prompted to pick a new subclass name and code will be recompiled
+	 *
+	 * @return The new class that was created
+	 */
+	UClass* CreateNewCPPComponent(TSubclassOf<UActorComponent> ComponentClass);
+
+	
+	/**
+	 * Creates a new Blueprint component from the specified class type
+	 * The user will be prompted to pick a new subclass name and a blueprint asset will be created
+	 *
+	 * @return The new class that was created
+	 */
+	UClass* CreateNewBPComponent(TSubclassOf<UActorComponent> ComponentClass);
+
 public:
 	/** Tree widget */
 	TSharedPtr<SSCSTreeType> SCSTreeWidget;
diff --git a/Engine/Source/Editor/KismetWidgets/Private/CreateBlueprintFromActorDialog.cpp b/Engine/Source/Editor/KismetWidgets/Private/CreateBlueprintFromActorDialog.cpp
--- a/Engine/Source/Editor/KismetWidgets/Private/CreateBlueprintFromActorDialog.cpp
+++ b/Engine/Source/Editor/KismetWidgets/Private/CreateBlueprintFromActorDialog.cpp
@@ -23,10 +23,10 @@ void FCreateBlueprintFromActorDialog::OpenDialog(bool bInHarvest, AActor* InActo
 		.ToolTipText(LOCTEXT("SelectPathTooltip", "Select the path where the Blueprint will be created at"))
 		.ClientSize(FVector2D(400, 400));
 
-	TSharedPtr<SCreateAssetFromActor> CreateBlueprintFromActorDialog;
+	TSharedPtr<SCreateAssetFromObject> CreateBlueprintFromActorDialog;
 	PickBlueprintPathWidget->SetContent
 	(
-		SAssignNew(CreateBlueprintFromActorDialog, SCreateAssetFromActor, PickBlueprintPathWidget)
+		SAssignNew(CreateBlueprintFromActorDialog, SCreateAssetFromObject, PickBlueprintPathWidget)
 		.AssetFilenameSuffix(TEXT("Blueprint"))
 		.HeadingText(LOCTEXT("CreateBlueprintFromActor_Heading", "Blueprint Name"))
 		.CreateButtonText(LOCTEXT("CreateBlueprintFromActor_ButtonLabel", "Create Blueprint"))
diff --git a/Engine/Source/Editor/UnrealEd/Private/Editor.cpp b/Engine/Source/Editor/UnrealEd/Private/Editor.cpp
--- a/Engine/Source/Editor/UnrealEd/Private/Editor.cpp
+++ b/Engine/Source/Editor/UnrealEd/Private/Editor.cpp
@@ -5725,10 +5725,10 @@ void UEditorEngine::ConvertActors( const TArray<AActor*>& ActorsToConvert, UClas
 			.ToolTipText(LOCTEXT("SelectPathTooltip", "Select the path where the static mesh will be created"))
 			.ClientSize(FVector2D(400, 400));
 
-		TSharedPtr<SCreateAssetFromActor> CreateAssetFromActorWidget;
+		TSharedPtr<SCreateAssetFromObject> CreateAssetFromActorWidget;
 		CreateAssetFromActorWindow->SetContent
 			(
-			SAssignNew(CreateAssetFromActorWidget, SCreateAssetFromActor, CreateAssetFromActorWindow)
+			SAssignNew(CreateAssetFromActorWidget, SCreateAssetFromObject, CreateAssetFromActorWindow)
 			.AssetFilenameSuffix(TEXT("StaticMesh"))
 			.HeadingText(LOCTEXT("ConvertBrushesToStaticMesh_Heading", "Static Mesh Name:"))
 			.CreateButtonText(LOCTEXT("ConvertBrushesToStaticMesh_ButtonLabel", "Create Static Mesh"))
diff --git a/Engine/Source/Editor/UnrealEd/Private/SComponentClassCombo.cpp b/Engine/Source/Editor/UnrealEd/Private/SComponentClassCombo.cpp
--- a/Engine/Source/Editor/UnrealEd/Private/SComponentClassCombo.cpp
+++ b/Engine/Source/Editor/UnrealEd/Private/SComponentClassCombo.cpp
@@ -97,6 +97,7 @@ void SComponentClassCombo::Construct(const FArguments& InArgs)
 
 	SComboButton::Construct(Args);
 
+	ComponentClassListView->EnableToolTipForceField( true );
 	// The base class can automatically handle setting focus to a specified control when the combo button is opened
 	SetMenuContentWidgetToFocus( SearchBox );
 }
@@ -105,12 +106,6 @@ void SComponentClassCombo::ClearSelection()
 {
 	PrevSelectedIndex = INDEX_NONE;
 
-	// Scroll to the first item if there is one, to reset the scroll position
-	if(FilteredComponentClassList.Num())
-	{
-		ComponentClassListView->RequestScrollIntoView(FilteredComponentClassList[0]);
-	}
-
 	// Clear the selection in such a way as to also clear the keyboard selector
 	ComponentClassListView->SetSelection(NULL, ESelectInfo::OnNavigation);
 }
@@ -248,7 +243,7 @@ TSharedRef<ITableRow> SComponentClassCombo::GenerateAddComponentRow( FComponentC
 	}
 	else
 	{
-
+		
 		return
 			SNew( SComboRow< TSharedPtr<FString> >, OwnerTable )
 			.ToolTip( FEditorClassUtils::GetTooltip(Entry->GetComponentClass()) )
@@ -328,12 +323,16 @@ void SComponentClassCombo::UpdateComponentClassList()
 
 	if( GetDefault<UEditorExperimentalSettings>()->bScriptableComponentsOnActors )
 	{
+		FString NewComponentsHeading = LOCTEXT("NewComponentsHeading", "Behavioral").ToString();
 		// Add new C++ component class
-		FComponentClassComboEntryPtr NewClassHeader = MakeShareable(new FComponentClassComboEntry(LOCTEXT("ScriptingHeading", "Scripting").ToString()));
+		FComponentClassComboEntryPtr NewClassHeader = MakeShareable(new FComponentClassComboEntry(NewComponentsHeading));
 		ComponentClassList.Add(NewClassHeader);
 
-		FComponentClassComboEntryPtr NewCPPClass = MakeShareable(new FComponentClassComboEntry(LOCTEXT("ScriptingHeading", "Scripting").ToString(), UActorComponent::StaticClass(), true, EComponentCreateAction::CreateNewCPPClass));
+		FComponentClassComboEntryPtr NewCPPClass = MakeShareable(new FComponentClassComboEntry(NewComponentsHeading, UActorComponent::StaticClass(), true, EComponentCreateAction::CreateNewCPPClass));
 		ComponentClassList.Add(NewCPPClass);
+
+		FComponentClassComboEntryPtr NewBPClass = MakeShareable(new FComponentClassComboEntry(NewComponentsHeading, UActorComponent::StaticClass(), true, EComponentCreateAction::CreateNewBlueprintClass));
+		ComponentClassList.Add(NewBPClass);
 	}
 	
 	TArray<FComponentClassComboEntryPtr> SortedClassList;
@@ -414,7 +413,11 @@ FText SComponentClassCombo::GetFriendlyComponentName(FComponentClassComboEntryPt
 
 	if( Entry->GetComponentCreateAction() == EComponentCreateAction::CreateNewCPPClass )
 	{
-		FriendlyComponentName = LOCTEXT("NewCPPBehaviorFriendlyName", "New C++ Behavior").ToString();
+		FriendlyComponentName = LOCTEXT("NewCPPBehaviorFriendlyName", "New C++ Component").ToString();
+	}
+	else if (Entry->GetComponentCreateAction() == EComponentCreateAction::CreateNewBlueprintClass )
+	{
+		FriendlyComponentName = LOCTEXT("NewCPPBehaviorFriendlyName", "New Blueprint Component").ToString();
 	}
 	else
 	{
diff --git a/Engine/Source/Editor/UnrealEd/Private/SCreateAssetFromActor.cpp b/Engine/Source/Editor/UnrealEd/Private/SCreateAssetFromActor.cpp
deleted file mode 100644
diff --git a/Engine/Source/Editor/UnrealEd/Private/SCreateAssetFromObject.cpp b/Engine/Source/Editor/UnrealEd/Private/SCreateAssetFromObject.cpp
new file mode 100644
--- /dev/null
+++ b/Engine/Source/Editor/UnrealEd/Private/SCreateAssetFromObject.cpp
@@ -0,0 +1,189 @@
+// Copyright 1998-2015 Epic Games, Inc. All Rights Reserved.
+
+#include "UnrealEd.h"
+#include "SCreateAssetFromObject.h"
+#include "ContentBrowserModule.h"
+#include "AssetToolsModule.h"
+#include "AssetRegistryModule.h"
+#include "PackageTools.h"
+#include "Engine/Selection.h"
+
+#define LOCTEXT_NAMESPACE "SCreateAssetFromActor"
+
+void SCreateAssetFromObject::Construct(const FArguments& InArgs, TSharedPtr<SWindow> InParentWindow)
+{
+	AssetFilenameSuffix = InArgs._AssetFilenameSuffix;
+	HeadingText = InArgs._HeadingText;
+	CreateButtonText = InArgs._CreateButtonText;
+	OnCreateAssetAction = InArgs._OnCreateAssetAction;
+
+	bIsReportingError = false;
+	AssetPath = FString("/Game");
+	FPathPickerConfig PathPickerConfig;
+	PathPickerConfig.DefaultPath = AssetPath;
+	PathPickerConfig.OnPathSelected = FOnPathSelected::CreateRaw(this, &SCreateAssetFromObject::OnSelectAssetPath);
+
+	USelection::SelectionChangedEvent.AddRaw(this, &SCreateAssetFromObject::OnLevelSelectionChanged);
+
+	// Set up PathPickerConfig.
+	FContentBrowserModule& ContentBrowserModule = FModuleManager::LoadModuleChecked<FContentBrowserModule>("ContentBrowser");
+
+	ParentWindow = InParentWindow;
+
+	FString PackageName;
+	ActorInstanceLabel.Empty();
+
+	if( InArgs._DefaultNameOverride.IsEmpty() )
+	{
+		USelection* SelectedActors = GEditor->GetSelectedActors();
+		for(FSelectionIterator Iter(*SelectedActors); Iter; ++Iter)
+		{
+			AActor* Actor = Cast<AActor>(*Iter);
+			if(Actor)
+			{
+				ActorInstanceLabel += Actor->GetActorLabel();
+				ActorInstanceLabel += TEXT("_");
+				break;
+			}
+		}
+	}
+	else
+	{
+		ActorInstanceLabel = InArgs._DefaultNameOverride.ToString();
+	}
+
+	ActorInstanceLabel = PackageTools::SanitizePackageName(ActorInstanceLabel + AssetFilenameSuffix);
+
+	FString AssetName = ActorInstanceLabel;
+	FString BasePath = AssetPath / AssetName;
+	FAssetToolsModule& AssetToolsModule = FModuleManager::LoadModuleChecked<FAssetToolsModule>("AssetTools");
+	AssetToolsModule.Get().CreateUniqueAssetName(BasePath, TEXT(""), PackageName, AssetName);
+
+	ChildSlot
+	[
+		SNew(SVerticalBox)
+		+ SVerticalBox::Slot()
+		.FillHeight(1.0f)
+		[
+			ContentBrowserModule.Get().CreatePathPicker(PathPickerConfig)
+		]
+		+ SVerticalBox::Slot()
+		.AutoHeight()
+		[
+			SNew(SHorizontalBox)
+			+ SHorizontalBox::Slot()
+			.Padding(0.0f, 0.0f, 4.0f, 0.0f)
+			.AutoWidth()
+			[
+				SNew(STextBlock)
+				.Text(HeadingText)
+			]
+			+ SHorizontalBox::Slot()
+			.FillWidth(1.0f)
+			[
+				SAssignNew(FileNameWidget, SEditableTextBox)
+				.Text(FText::FromString(AssetName))
+				.OnTextChanged(this, &SCreateAssetFromObject::OnFilenameChanged)
+			]
+		]
+		+ SVerticalBox::Slot()
+		.HAlign(HAlign_Right)
+		.Padding(0, 20, 0, 0)
+		.AutoHeight()
+		[
+			SNew(SHorizontalBox)
+			+ SHorizontalBox::Slot()
+			.Padding(0, 2, 6, 0)
+			.AutoWidth()
+			[
+				SNew(SButton)
+				.VAlign(VAlign_Bottom)
+				.ContentPadding(FMargin(8, 2, 8, 2))
+				.OnClicked(this, &SCreateAssetFromObject::OnCreateAssetFromActorClicked)
+				.IsEnabled(this, &SCreateAssetFromObject::IsCreateAssetFromActorEnabled)
+				[
+					SNew(STextBlock)
+					.Text(CreateButtonText)
+				]
+			]
+			+ SHorizontalBox::Slot()
+			.Padding(0, 2, 0, 0)
+			.AutoWidth()
+			[
+				SNew(SButton)
+				.VAlign(VAlign_Bottom)
+				.ContentPadding(FMargin(8, 2, 8, 2))
+				.OnClicked(this, &SCreateAssetFromObject::OnCancelCreateAssetFromActor)
+				[
+					SNew(STextBlock)
+					.Text(LOCTEXT("CancelButtonText", "Cancel"))
+				]
+			]
+		]
+	];
+
+	OnFilenameChanged(FText::FromString(AssetName));
+}
+
+FReply SCreateAssetFromObject::OnCreateAssetFromActorClicked()
+{
+	ParentWindow->RequestDestroyWindow();
+	OnCreateAssetAction.ExecuteIfBound(AssetPath / FileNameWidget->GetText().ToString());
+	return FReply::Handled();
+}
+
+FReply SCreateAssetFromObject::OnCancelCreateAssetFromActor()
+{
+	ParentWindow->RequestDestroyWindow();
+	return FReply::Handled();
+}
+
+void SCreateAssetFromObject::OnSelectAssetPath(const FString& Path)
+{
+	AssetPath = Path;
+	OnFilenameChanged(FileNameWidget->GetText());
+}
+
+void SCreateAssetFromObject::OnLevelSelectionChanged(UObject* InObjectSelected)
+{
+	// When actor selection changes, this window should be destroyed.
+	ParentWindow->RequestDestroyWindow();
+}
+
+void SCreateAssetFromObject::OnFilenameChanged(const FText& InNewName)
+{
+	TArray<FAssetData> AssetData;
+	FAssetRegistryModule& AssetRegistryModule = FModuleManager::LoadModuleChecked<FAssetRegistryModule>("AssetRegistry");
+	AssetRegistryModule.Get().GetAssetsByPath(FName(*AssetPath), AssetData);
+
+	FText ErrorText;
+	if (!FEditorFileUtils::IsFilenameValidForSaving(InNewName.ToString(), ErrorText) || !FName(*InNewName.ToString()).IsValidObjectName(ErrorText))
+	{
+		FileNameWidget->SetError(ErrorText);
+		bIsReportingError = true;
+		return;
+	}
+	else
+	{
+		// Check to see if the name conflicts
+		for (auto Iter = AssetData.CreateConstIterator(); Iter; ++Iter)
+		{
+			if (Iter->AssetName.ToString() == InNewName.ToString())
+			{
+				FileNameWidget->SetError(LOCTEXT("AssetInUseError", "Asset name already in use!"));
+				bIsReportingError = true;
+				return;
+			}
+		}
+	}
+
+	FileNameWidget->SetError(FText::FromString(TEXT("")));
+	bIsReportingError = false;
+}
+
+bool SCreateAssetFromObject::IsCreateAssetFromActorEnabled() const
+{
+	return !bIsReportingError;
+}
+
+#undef LOCTEXT_NAMESPACE
diff --git a/Engine/Source/Editor/UnrealEd/Public/SCreateAssetFromActor.h b/Engine/Source/Editor/UnrealEd/Public/SCreateAssetFromActor.h
deleted file mode 100644
diff --git a/Engine/Source/Editor/UnrealEd/Public/SCreateAssetFromObject.h b/Engine/Source/Editor/UnrealEd/Public/SCreateAssetFromObject.h
new file mode 100644
--- /dev/null
+++ b/Engine/Source/Editor/UnrealEd/Public/SCreateAssetFromObject.h
@@ -0,0 +1,85 @@
+// Copyright 1998-2015 Epic Games, Inc. All Rights Reserved.
+
+#pragma once
+
+//////////////////////////////////////////////////////////////////////////
+// SCreateAssetFromActor
+
+DECLARE_DELEGATE_OneParam(FOnPathChosen, const FString&);
+
+class SCreateAssetFromObject : public SCompoundWidget
+{
+public:
+	SLATE_BEGIN_ARGS(SCreateAssetFromObject)
+		: _AssetFilenameSuffix()
+		, _HeadingText()
+		, _CreateButtonText()
+	{}
+
+	/** The default suffix to use for the asset filename */
+	SLATE_ARGUMENT(FString, AssetFilenameSuffix)
+
+	/** The text to display at the top of the dialog */
+	SLATE_ARGUMENT(FText, HeadingText)
+
+	/** The label for the create button */
+	SLATE_ARGUMENT(FText, CreateButtonText)
+
+	SLATE_ARGUMENT(FText, DefaultNameOverride)
+
+	/** Action to perform when create clicked */
+	SLATE_EVENT(FOnPathChosen, OnCreateAssetAction)
+
+	SLATE_END_ARGS()
+
+	/** Constructs this widget with InArgs */
+	UNREALED_API void Construct(const FArguments& InArgs, TSharedPtr<SWindow> InParentWindow);
+
+private:
+
+	/** Callback when the "create asset" button is clicked. */
+	FReply OnCreateAssetFromActorClicked();
+
+	/** Callback when the selected asset path has changed. */
+	void OnSelectAssetPath(const FString& Path);
+
+	/** Destroys the window when the operation is cancelled. */
+	FReply OnCancelCreateAssetFromActor();
+
+	/** Callback when level selection has changed, will destroy the window. */
+	void OnLevelSelectionChanged(UObject* InObjectSelected);
+
+	/** Callback when the user changes the filename for the Blueprint */
+	void OnFilenameChanged(const FText& InNewName);
+
+	/** Callback to see if creating an asset is enabled */
+	bool IsCreateAssetFromActorEnabled() const;
+
+private:
+	/** The window this widget is nested in */
+	TSharedPtr<SWindow> ParentWindow;
+
+	/** The selected path to create the asset */
+	FString AssetPath;
+
+	/** The resultant actor instance label, based on the original actor labels */
+	FString ActorInstanceLabel;
+
+	/** The default suffix to use for the asset filename */
+	FString AssetFilenameSuffix;
+
+	/** The text to display as a heading for the dialog */
+	FText HeadingText;
+
+	/** The label to be displayed on the create button */
+	FText CreateButtonText;
+
+	/** Filename textbox widget */
+	TSharedPtr<SEditableTextBox> FileNameWidget;
+
+	/** True if an error is currently being reported */
+	bool bIsReportingError;
+
+	/** Called when the create button is clicked */
+	FOnPathChosen OnCreateAssetAction;
+};
