[INTEGRATE] Change 2418983 by [REDACTED_PERSON][REDACTED_HANDLE].7 on 2015/01/26 14:27:06

	Non-unity compile fixes

diff --git a/Engine/Source/Editor/UnrealEd/Private/AutoReimport/AutoReimportUtilities.cpp b/Engine/Source/Editor/UnrealEd/Private/AutoReimport/AutoReimportUtilities.cpp
--- a/Engine/Source/Editor/UnrealEd/Private/AutoReimport/AutoReimportUtilities.cpp
+++ b/Engine/Source/Editor/UnrealEd/Private/AutoReimport/AutoReimportUtilities.cpp
@@ -3,6 +3,8 @@
 #include "UnrealEd.h"
 #include "AutoReimportUtilities.h"
 #include "EditorReimportHandler.h"
+#include "ARFilter.h"
+#include "IAssetRegistry.h"
 
 namespace Utils
 {
diff --git a/Engine/Source/Editor/UnrealEd/Private/AutoReimport/ReimportFeedbackContext.h b/Engine/Source/Editor/UnrealEd/Private/AutoReimport/ReimportFeedbackContext.h
--- a/Engine/Source/Editor/UnrealEd/Private/AutoReimport/ReimportFeedbackContext.h
+++ b/Engine/Source/Editor/UnrealEd/Private/AutoReimport/ReimportFeedbackContext.h
@@ -5,6 +5,8 @@
 #include "FeedbackContext.h"
 #include "SNotificationList.h"
 #include "INotificationWidget.h"
+#include "MessageLog.h"
+#include "AutoReimportUtilities.h"
 
 class SWidgetStack;
 
@@ -51,7 +53,7 @@ public:
 	/** Tick this context - potentially opening a new notification if the process has taken a while */
 	void Tick();
 
-	TSharedPtr<class SReimportFeedback> GetContent() { return NotificationContent; }
+	TSharedPtr<SReimportFeedback> GetContent() { return NotificationContent; }
 
 	/** Add a message to the context (and message log) */
 	void AddMessage(EMessageSeverity::Type Severity, const FText& Message);
@@ -79,5 +81,5 @@ private:
 	TSharedPtr<class SNotificationItem> Notification;
 
 	/** The notification content */
-	TSharedPtr<class SReimportFeedback> NotificationContent;
+	TSharedPtr<SReimportFeedback> NotificationContent;
 };
\ No newline at end of file
