#UE4 The 32 characters reserved for compressed package intermediates can now be waived via ini setting.

diff --git a/Engine/Config/BaseEditor.ini b/Engine/Config/BaseEditor.ini
--- a/Engine/Config/BaseEditor.ini
+++ b/Engine/Config/BaseEditor.ini
@@ -290,6 +290,7 @@ HighSpawnRateOrBurstThreshold=35
 FarLODDistanceTheshold=3000
 
 [CookSettings]
+bConsiderCompressedPackageFileLengthRequirements=true
 +FullGCAssetClassNames=/Script/Engine.World
 +VersionedIniParams=Engine:/Script/Engine.RendererSettings:Compat.UseDXT5NormalMaps
 +VersionedIniParams=Engine:/Script/AndroidRuntimeSettings.AndroidRuntimeSettings:AndroidAudio
diff --git a/Engine/Source/Editor/UnrealEd/Classes/CookOnTheSide/CookOnTheFlyServer.h b/Engine/Source/Editor/UnrealEd/Classes/CookOnTheSide/CookOnTheFlyServer.h
--- a/Engine/Source/Editor/UnrealEd/Classes/CookOnTheSide/CookOnTheFlyServer.h
+++ b/Engine/Source/Editor/UnrealEd/Classes/CookOnTheSide/CookOnTheFlyServer.h
@@ -1028,6 +1028,9 @@ private:
 	 */
 	bool GetPackageTimestamp( const FString& InFilename, FDateTime& OutDateTime );
 
+	/** If true, the maximum file length of a package being saved will be reduced by 32 to compensate for compressed package intermediate files */
+	bool ShouldConsiderCompressedPackageFileLengthRequirements() const;
+
 	/**
 	 *	Cook (save) the given package
 	 *
diff --git a/Engine/Source/Editor/UnrealEd/Private/CookOnTheFlyServer.cpp b/Engine/Source/Editor/UnrealEd/Private/CookOnTheFlyServer.cpp
--- a/Engine/Source/Editor/UnrealEd/Private/CookOnTheFlyServer.cpp
+++ b/Engine/Source/Editor/UnrealEd/Private/CookOnTheFlyServer.cpp
@@ -1624,6 +1624,13 @@ bool UCookOnTheFlyServer::ShouldCook(const FString& InFileName, const FName &InP
 	return bDoCook;*/
 }
 
+bool UCookOnTheFlyServer::ShouldConsiderCompressedPackageFileLengthRequirements() const
+{
+	bool bConsiderCompressedPackageFileLengthRequirements = true;
+	GConfig->GetBool(TEXT("CookSettings"), TEXT("bConsiderCompressedPackageFileLengthRequirements"), bConsiderCompressedPackageFileLengthRequirements, GEditorIni);
+	return bConsiderCompressedPackageFileLengthRequirements;
+}
+
 bool UCookOnTheFlyServer::SaveCookedPackage( UPackage* Package, uint32 SaveFlags, bool& bOutWasUpToDate, TArray<FName> &TargetPlatformNames )
 {
 	bool bSavedCorrectly = true;
@@ -1776,8 +1783,12 @@ bool UCookOnTheFlyServer::SaveCookedPackage( UPackage* Package, uint32 SaveFlags
 					}
 				}
 
+				// need to subtract 32 because the SavePackage code creates temporary files with longer file names then the one we provide
+				// projects may ignore this restriction if desired
+				static bool bConsiderCompressedPackageFileLengthRequirements = ShouldConsiderCompressedPackageFileLengthRequirements();
+				const int32 CompressedPackageFileLengthRequirement = bConsiderCompressedPackageFileLengthRequirements ? 32 : 0;
 				const FString FullFilename = FPaths::ConvertRelativePathToFull( PlatFilename );
-				if( FullFilename.Len() >= (PLATFORM_MAX_FILEPATH_LENGTH-32) ) // need to subtract 32 because the SavePackage code creates temporary files with longer file names then the one we provide
+				if( FullFilename.Len() >= (PLATFORM_MAX_FILEPATH_LENGTH-CompressedPackageFileLengthRequirement) )
 				{
 					LogCookerMessage( FString::Printf(TEXT("Couldn't save package, filename is too long: %s"), *PlatFilename), EMessageSeverity::Error );
 					UE_LOG( LogCook, Error, TEXT( "Couldn't save package, filename is too long :%s" ), *PlatFilename );
