Fix for broken build.

diff --git a/Engine/Source/Editor/UnrealEd/Private/AutoReimport/ReimportFeedbackContext.cpp b/Engine/Source/Editor/UnrealEd/Private/AutoReimport/ReimportFeedbackContext.cpp
--- a/Engine/Source/Editor/UnrealEd/Private/AutoReimport/ReimportFeedbackContext.cpp
+++ b/Engine/Source/Editor/UnrealEd/Private/AutoReimport/ReimportFeedbackContext.cpp
@@ -39,7 +39,7 @@ class SWidgetStack : public SVerticalBox
 		return Size;
 	}
 
-	virtual FVector2D ComputeDesiredSize() const override
+	virtual FVector2D ComputeDesiredSize(float) const override
 	{
 		FVector2D Size = ComputeTotalSize();
 		Size.Y -= LerpSizeOffset;
