[INTEGRATE] Change 2419308 by [REDACTED_PERSON][REDACTED_HANDLE] on 2015/01/26 17:35:27

	#ue4 Blueprintable Components.  Enable in the Experimental editor settings

diff --git a/Engine/Source/Editor/BlueprintGraph/Private/K2ActionMenuBuilder.cpp b/Engine/Source/Editor/BlueprintGraph/Private/K2ActionMenuBuilder.cpp
--- a/Engine/Source/Editor/BlueprintGraph/Private/K2ActionMenuBuilder.cpp
+++ b/Engine/Source/Editor/BlueprintGraph/Private/K2ActionMenuBuilder.cpp
@@ -324,7 +324,7 @@ static void GetAddComponentClasses(UBlueprint const* BlueprintIn, FGraphActionLi
 		UClass* Class = *It;
 
 		// If this is a subclass of ActorComponent, not abstract, and tagged as spawnable from Kismet
-		if (Class->IsChildOf(UActorComponent::StaticClass()) && !Class->HasAnyClassFlags(CLASS_Abstract) && Class->HasMetaData(FBlueprintMetadata::MD_BlueprintSpawnableComponent) )
+		if (Class->IsChildOf(UActorComponent::StaticClass()) && !Class->HasAnyClassFlags(CLASS_Abstract) && Class->HasMetaData(FBlueprintMetadata::MD_BlueprintSpawnableComponent) && !FKismetEditorUtilities::IsClassABlueprintSkeleton(Class))
 		{
 			TSharedPtr<FEdGraphSchemaAction_K2AddComponent> NewAction = FK2ActionMenuBuilder::CreateAddComponentAction(ActionMenuBuilder.OwnerOfTemporaries, BlueprintIn, Class, /*Asset=*/ NULL);
 			ActionMenuBuilder.AddAction(NewAction);
diff --git a/Engine/Source/Editor/Kismet/Private/SSCSEditor.cpp b/Engine/Source/Editor/Kismet/Private/SSCSEditor.cpp
--- a/Engine/Source/Editor/Kismet/Private/SSCSEditor.cpp
+++ b/Engine/Source/Editor/Kismet/Private/SSCSEditor.cpp
@@ -3997,6 +3997,13 @@ UActorComponent* SSCSEditor::AddNewComponent( UClass* NewComponentClass, UObject
 		ActorInstance->Modify();
 
 		FName NewComponentName = NAME_None;
+		if (NewComponentClass->ClassGeneratedBy && NewComponentClass->GetName().EndsWith(TEXT("_C")))
+		{
+			const FString NewClassName = NewComponentClass->GetName();
+			const int32 NewStrLen = NewClassName.Len() - 2;
+			NewComponentName = MakeUniqueObjectName(ActorInstance, NewComponentClass, FName(*NewClassName.Left(NewStrLen)));
+		}
+
 		if (Asset != nullptr)
 		{
 			NewComponentName = *FComponentEditorUtils::GenerateValidVariableNameFromAsset(Asset, ActorInstance);
diff --git a/Engine/Source/Editor/KismetCompiler/Private/KismetCompiler.cpp b/Engine/Source/Editor/KismetCompiler/Private/KismetCompiler.cpp
--- a/Engine/Source/Editor/KismetCompiler/Private/KismetCompiler.cpp
+++ b/Engine/Source/Editor/KismetCompiler/Private/KismetCompiler.cpp
@@ -1700,6 +1700,13 @@ void FKismetCompilerContext::FinishCompilingClass(UClass* Class)
 		{
 			Class->SetMetaData(TEXT("AutoCollapseCategories"), *ParentClass->GetMetaData("AutoCollapseCategories"));
 		}
+		
+		// Blueprinted Components are always Blueprint Spawnable
+		if (ParentClass->IsChildOf(UActorComponent::StaticClass()))
+		{
+			Class->SetMetaData(TEXT("BlueprintSpawnableComponent"), TEXT("true"));
+		}
+
 
 		// Add a category if one has been specified
 		if(Blueprint->BlueprintCategory.Len() > 0)
diff --git a/Engine/Source/Editor/UnrealEd/Classes/Settings/EditorExperimentalSettings.h b/Engine/Source/Editor/UnrealEd/Classes/Settings/EditorExperimentalSettings.h
--- a/Engine/Source/Editor/UnrealEd/Classes/Settings/EditorExperimentalSettings.h
+++ b/Engine/Source/Editor/UnrealEd/Classes/Settings/EditorExperimentalSettings.h
@@ -41,6 +41,10 @@ public:
 	UPROPERTY(EditAnywhere, config, Category=Tools)
 	bool bScriptableComponentsOnActors;
 
+	/** Enable being able to subclass components in blueprints */
+	UPROPERTY(EditAnywhere, config, Category = Tools, meta = (DisplayName = "Blueprintable Component"))
+	bool bBlueprintableComponents;
+
 	/** The Messaging Debugger provides a visual utility for debugging the messaging system. */
 	UPROPERTY(EditAnywhere, config, Category=Tools, meta=(DisplayName="Messaging Debugger"))
 	bool bMessagingDebugger;
diff --git a/Engine/Source/Editor/UnrealEd/Private/Kismet2/Kismet2.cpp b/Engine/Source/Editor/UnrealEd/Private/Kismet2/Kismet2.cpp
--- a/Engine/Source/Editor/UnrealEd/Private/Kismet2/Kismet2.cpp
+++ b/Engine/Source/Editor/UnrealEd/Private/Kismet2/Kismet2.cpp
@@ -866,6 +866,8 @@ bool FKismetEditorUtilities::CanCreateBlueprintOfClass(const UClass* Class)
 	bool bAllowDerivedBlueprints = false;
 	GConfig->GetBool(TEXT("Kismet"), TEXT("AllowDerivedBlueprints"), /*out*/ bAllowDerivedBlueprints, GEngineIni);
 
+	const bool bAllowBlueprintableComponents = GetDefault<UEditorExperimentalSettings>()->bBlueprintableComponents;
+	
 	const bool bCanCreateBlueprint =
 		!Class->HasAnyClassFlags(CLASS_Deprecated)
 		&& !Class->HasAnyClassFlags(CLASS_NewerVersionExists)
@@ -876,6 +878,7 @@ bool FKismetEditorUtilities::CanCreateBlueprintOfClass(const UClass* Class)
 	const UEdGraphSchema_K2* K2Schema = GetDefault<UEdGraphSchema_K2>();
 	const bool bIsValidClass = Class->GetBoolMetaDataHierarchical(FBlueprintMetadata::MD_IsBlueprintBase)
 		|| (Class == UObject::StaticClass())
+		|| (Class == UActorComponent::StaticClass() && bAllowBlueprintableComponents)
 		|| bIsBPGC;  // BPs are always considered inheritable
 
 	return bCanCreateBlueprint && bIsValidClass;
diff --git a/Engine/Source/Editor/UnrealEd/Private/Kismet2/KismetReinstanceUtilities.cpp b/Engine/Source/Editor/UnrealEd/Private/Kismet2/KismetReinstanceUtilities.cpp
--- a/Engine/Source/Editor/UnrealEd/Private/Kismet2/KismetReinstanceUtilities.cpp
+++ b/Engine/Source/Editor/UnrealEd/Private/Kismet2/KismetReinstanceUtilities.cpp
@@ -8,6 +8,7 @@
 #include "Layers/ILayers.h"
 #include "ComponentInstanceDataCache.h"
 #include "Engine/DynamicBlueprintBinding.h"
+#include "BlueprintEditor.h"
 
 DECLARE_CYCLE_STAT(TEXT("Replace Instances"), EKismetReinstancerStats_ReplaceInstancesOfClass, STATGROUP_KismetReinstancer );
 DECLARE_CYCLE_STAT(TEXT("Find Referencers"), EKismetReinstancerStats_FindReferencers, STATGROUP_KismetReinstancer );
@@ -305,9 +306,10 @@ void FBlueprintCompileReinstancer::ReinstanceObjects(bool bAlwaysReinstance)
 
 				const bool bIsActor = ClassToReinstance->IsChildOf<AActor>();
 				const bool bIsAnimInstance = ClassToReinstance->IsChildOf<UAnimInstance>();
+				const bool bIsComponent = ClassToReinstance->IsChildOf<UActorComponent>();
 				for (auto Obj : ObjectsToReplace)
 				{
-					if (!Obj->IsTemplate() && !Obj->IsPendingKill())
+					if ((!Obj->IsTemplate() || bIsComponent) && !Obj->IsPendingKill())
 					{
 						Obj->SetClass(ClassToReinstance);
 						if (bIsActor)
@@ -636,6 +638,9 @@ void FBlueprintCompileReinstancer::ReplaceInstancesOfClass(UClass* OldClass, UCl
 	// actors being replace
 	TArray<FActorReplacementHelper> ReplacementActors;
 
+	// A list of objects (e.g. Blueprints) that potentially have editors open that we need to refresh
+	TArray<UObject*> PotentialEditorsForRefreshing;
+
 	// Set global flag to let system know we are reconstructing blueprint instances
 	TGuardValue<bool> GuardTemplateNameFlag(GIsReconstructingBlueprintInstances, true);
 
@@ -664,7 +669,9 @@ void FBlueprintCompileReinstancer::ReplaceInstancesOfClass(UClass* OldClass, UCl
 		// Then fix 'real' (non archetype) instances of the class
 		for (UObject* OldObject : ObjectsToReplace)
 		{
-			if (OldObject->IsTemplate() || OldObject->IsPendingKill())
+			// Skip non-archetype instances, EXCEPT for component templates
+			const bool bIsComponent = NewClass->IsChildOf(UActorComponent::StaticClass());
+			if ((!bIsComponent && OldObject->IsTemplate()) || OldObject->IsPendingKill())
 			{
 				continue;
 			}
@@ -783,6 +790,22 @@ void FBlueprintCompileReinstancer::ReplaceInstancesOfClass(UClass* OldClass, UCl
 				OldObject->MarkPendingKill();
 
 				OldToNewInstanceMap.Add(OldObject, NewObject);
+
+				if (bIsComponent)
+				{
+					UActorComponent* Component = Cast<UActorComponent>(NewObject);
+					AActor* OwningActor = Component->GetOwner();
+					if (OwningActor)
+					{
+						OwningActor->ResetOwnedComponents();
+
+						// Check to see if they have an editor that potentially needs to be refreshed
+						if (OwningActor->GetClass()->ClassGeneratedBy)
+						{
+							PotentialEditorsForRefreshing.AddUnique(OwningActor->GetClass()->ClassGeneratedBy);
+						}
+					}
+				}
 			}
 
 			// If this original object came from a blueprint and it was in the selected debug set, change the debugging to the new object.
@@ -837,6 +860,19 @@ void FBlueprintCompileReinstancer::ReplaceInstancesOfClass(UClass* OldClass, UCl
 	{
 		GEditor->NoteSelectionChange();
 	}
+
+	if (GEditor)
+	{
+		// Refresh any editors for objects that we've updated components for
+		for (auto BlueprintAsset : PotentialEditorsForRefreshing)
+		{
+			FBlueprintEditor* BlueprintEditor = static_cast<FBlueprintEditor*>(FAssetEditorManager::Get().FindEditorForAsset(BlueprintAsset, /*bFocusIfOpen =*/false));
+			if (BlueprintEditor)
+			{
+				BlueprintEditor->RefreshEditors();
+			}
+		}
+	}
 }
 
 void FBlueprintCompileReinstancer::VerifyReplacement()
diff --git a/Engine/Source/Editor/UnrealEd/Private/SComponentClassCombo.cpp b/Engine/Source/Editor/UnrealEd/Private/SComponentClassCombo.cpp
--- a/Engine/Source/Editor/UnrealEd/Private/SComponentClassCombo.cpp
+++ b/Engine/Source/Editor/UnrealEd/Private/SComponentClassCombo.cpp
@@ -340,7 +340,7 @@ void SComponentClassCombo::UpdateComponentClassList()
 	{
 		UClass* Class = *It;
 		// If this is a subclass of Actor Component, not abstract, and tagged as spawnable from Kismet
-		if (Class->IsChildOf(UActorComponent::StaticClass()) && !Class->HasAnyClassFlags(CLASS_Abstract) && Class->HasMetaData(FBlueprintMetadata::MD_BlueprintSpawnableComponent)) //@TODO: Fold this logic together with the one in UEdGraphSchema_K2::GetAddComponentClasses
+		if (Class->IsChildOf(UActorComponent::StaticClass()) && !Class->HasAnyClassFlags(CLASS_Abstract) && Class->HasMetaData(FBlueprintMetadata::MD_BlueprintSpawnableComponent) && !FKismetEditorUtilities::IsClassABlueprintSkeleton(Class)) //@TODO: Fold this logic together with the one in UEdGraphSchema_K2::GetAddComponentClasses
 		{
 			TArray<FString> ClassGroupNames;
 			Class->GetClassGroupNames( ClassGroupNames );
@@ -478,8 +478,15 @@ FString SComponentClassCombo::GetSanitizedComponentName( UClass* ComponentClass
 	}
 	else
 	{
-		FString OriginalName = ComponentClass->GetName();
-		DisplayName = OriginalName.Replace( TEXT("Component"), TEXT(""), ESearchCase::IgnoreCase );
+		DisplayName = ComponentClass->GetName();
+		DisplayName = DisplayName.Replace( TEXT("Component"), TEXT(""), ESearchCase::IgnoreCase );
+	}
+
+	// Purge the _C for BP Components
+	if (DisplayName.EndsWith(TEXT("_C")))
+	{
+		const int32 NewLen = DisplayName.Len() - 2;
+		DisplayName = DisplayName.Left(NewLen);
 	}
 
 	return FName::NameToDisplayString(DisplayName, false);
diff --git a/Engine/Source/Runtime/CoreUObject/Private/UObject/Obj.cpp b/Engine/Source/Runtime/CoreUObject/Private/UObject/Obj.cpp
--- a/Engine/Source/Runtime/CoreUObject/Private/UObject/Obj.cpp
+++ b/Engine/Source/Runtime/CoreUObject/Private/UObject/Obj.cpp
@@ -2923,6 +2923,79 @@ bool StaticExec( UWorld* InWorld, const TCHAR* Cmd, FOutputDevice& Ar )
 			Ar.Logf(TEXT("Non-permanent: %d objects, %d edges, %d strongly connected components, %d objects are included in cycles."), IndexSet.TempObjects.Num(), IndexSet.Edges.Num(), TotalCnt, TotalNum);
 			return true;
 		}
+		else if (FParse::Command(&Str, TEXT("VERIFYCOMPONENTS")))
+		{
+			Ar.Logf(TEXT("------------------------------------------------------------------------------"));
+
+			for (FObjectIterator It; It; ++It)
+			{
+				UObject* Target = *It;
+
+				// Skip objects that are trashed
+				if ((Target->GetOutermost() == GetTransientPackage())
+					|| Target->GetClass()->HasAnyClassFlags(CLASS_NewerVersionExists)
+					|| Target->HasAnyFlags(RF_PendingKill))
+				{
+					continue;
+				}
+
+				TArray<UObject*> SubObjects;
+				GetObjectsWithOuter(Target, SubObjects);
+
+				TArray<FString> Errors;
+
+				for (auto SubObjIt : SubObjects)
+				{
+					const UObject* SubObj = SubObjIt;
+					const UClass* SubObjClass = SubObj->GetClass();
+					const FString SubObjName = SubObj->GetName();
+
+					if (SubObj->IsPendingKill())
+					{
+						continue;
+					}
+
+					if (SubObjClass->HasAnyClassFlags(CLASS_NewerVersionExists))
+					{
+						Errors.Add(FString::Printf(TEXT("  - %s has a stale class"), *SubObjName));
+					}
+
+					if (SubObjClass->GetOutermost() == GetTransientPackage())
+					{
+						Errors.Add(FString::Printf(TEXT("  - %s has a class in the transient package"), *SubObjName));
+					}
+
+					if (SubObj->GetOutermost() != Target->GetOutermost())
+					{
+						Errors.Add(FString::Printf(TEXT("  - %s has a different outer than its parent"), *SubObjName));
+					}
+					
+					if (SubObj->GetName().Find(TEXT("TRASH_")) != INDEX_NONE)
+					{
+						Errors.Add(FString::Printf(TEXT("  - %s is TRASH'd"), *SubObjName));
+					}
+
+					if (SubObj->GetName().Find(TEXT("REINST_")) != INDEX_NONE)
+					{
+						Errors.Add(FString::Printf(TEXT("  - %s is a REINST"), *SubObjName));
+					}
+				}
+
+				if (Errors.Num() > 0)
+				{
+					const FString ErrorStr = FString::Printf(TEXT("Errors for %s"), *Target->GetName());
+					Ar.Logf(*ErrorStr);
+
+					for (auto ErrorStr : Errors)
+					{
+						Ar.Logf(*(FString(TEXT("  - ") + ErrorStr)));
+					}
+				}
+			}
+
+			Ar.Logf(TEXT("------------------------------------------------------------------------------"));
+			return true;
+		}
 		else if( FParse::Command(&Str,TEXT("TRANSACTIONAL")) )
 		{
 			int32 Num=0;
diff --git a/Engine/Source/Runtime/Engine/Private/SCS_Node.cpp b/Engine/Source/Runtime/Engine/Private/SCS_Node.cpp
--- a/Engine/Source/Runtime/Engine/Private/SCS_Node.cpp
+++ b/Engine/Source/Runtime/Engine/Private/SCS_Node.cpp
@@ -447,7 +447,15 @@ FName USCS_Node::GenerateNewComponentName( TArray<FName>& CurrentNames, FName De
 			}
 			else
 			{
-				ComponentName = ComponentTemplate->GetClass()->GetName().Replace(TEXT("Component"), TEXT(""));
+				const UClass* ComponentClass = ComponentTemplate->GetClass();
+				ComponentName = ComponentClass->GetName().Replace(TEXT("Component"), TEXT(""));
+
+				// Trim the _C, if needed
+				if (ComponentClass->ClassGeneratedBy && ComponentName.EndsWith(TEXT("_C")))
+				{
+					const int32 NameLen = ComponentName.Len() - 1;
+					ComponentName = ComponentName.Left(NameLen);
+				}
 			}
 			
 			NewName = *ComponentName;
diff --git a/Engine/Source/Runtime/Engine/Private/SimpleConstructionScript.cpp b/Engine/Source/Runtime/Engine/Private/SimpleConstructionScript.cpp
--- a/Engine/Source/Runtime/Engine/Private/SimpleConstructionScript.cpp
+++ b/Engine/Source/Runtime/Engine/Private/SimpleConstructionScript.cpp
@@ -758,7 +758,16 @@ USCS_Node* USimpleConstructionScript::CreateNode(UClass* NewComponentClass, FNam
 	check(NewComponentClass->IsChildOf(UActorComponent::StaticClass()));
 
 	ensure(NULL != Cast<UBlueprintGeneratedClass>(Blueprint->GeneratedClass));
-	UActorComponent* NewComponentTemplate = ConstructObject<UActorComponent>(NewComponentClass, Blueprint->GeneratedClass);
+
+	FName NewComponentName(NAME_None);
+	if (NewComponentClass->ClassGeneratedBy && NewComponentClass->GetName().EndsWith(TEXT("_C")))
+	{
+		const FString NewClassName = NewComponentClass->GetName();
+		const int32 NewStrLen = NewClassName.Len() - 2;
+		NewComponentName = MakeUniqueObjectName(Blueprint->GeneratedClass, NewComponentClass, FName(*NewClassName.Left(NewStrLen)));
+	}
+
+	UActorComponent* NewComponentTemplate = ConstructObject<UActorComponent>(NewComponentClass, Blueprint->GeneratedClass, NewComponentName);
 	NewComponentTemplate->SetFlags(RF_ArchetypeObject|RF_Transactional|RF_Public);
 
 	return CreateNode(NewComponentTemplate, NewComponentVariableName);
@@ -778,6 +787,13 @@ USCS_Node* USimpleConstructionScript::CreateNode(UActorComponent* NewComponentTe
 		TArray<FName> CurrentNames;
 		NewNode->GenerateListOfExistingNames( CurrentNames );
 
+		if (NewComponentVariableName.ToString().EndsWith(TEXT("_C")))
+		{
+			const FString NameAsString = NewComponentVariableName.ToString();
+			const int32 NewStrLen = NameAsString.Len() - 2;
+			NewComponentVariableName = FName(*NameAsString.Left(NewStrLen));
+		}
+
 		// Now create a name for the new component.
 		NewNode->VariableName = NewNode->GenerateNewComponentName( CurrentNames, NewComponentVariableName );
 
