Remove the dependency on Simplygon from UnrealEd. Not sure why it's here (I suspect legacy reasons), but it doesn't seem to be necessary. SimplygonMeshReduction is the only module which actually needs it.

diff --git a/Engine/Source/Editor/UnrealEd/UnrealEd.Build.cs b/Engine/Source/Editor/UnrealEd/UnrealEd.Build.cs
--- a/Engine/Source/Editor/UnrealEd/UnrealEd.Build.cs
+++ b/Engine/Source/Editor/UnrealEd/UnrealEd.Build.cs
@@ -249,13 +249,6 @@ public class UnrealEd : ModuleRules
 
 		SetupModulePhysXAPEXSupport(Target);
 
-		if ((UEBuildConfiguration.bCompileSimplygon == true) &&
-			(Directory.Exists(UEBuildConfiguration.UEThirdPartySourceDirectory + "NotForLicensees") == true) &&
-			(Directory.Exists(UEBuildConfiguration.UEThirdPartySourceDirectory + "NotForLicensees/Simplygon") == true))
-		{
-			AddThirdPartyPrivateStaticDependencies(Target, "Simplygon");
-		}
-
 		if (UEBuildConfiguration.bCompileRecast)
 		{
             PrivateDependencyModuleNames.Add("Navmesh");
