[INTEGRATE] Change 2420362 by [REDACTED_PERSON][REDACTED_HANDLE] on 2015/01/27 10:34:50

	Remove #pragma once from .cpp file; causes build failure in non-unity on Mac.

diff --git a/Engine/Source/Editor/UnrealEd/Private/AutoReimport/ReimportFeedbackContext.cpp b/Engine/Source/Editor/UnrealEd/Private/AutoReimport/ReimportFeedbackContext.cpp
--- a/Engine/Source/Editor/UnrealEd/Private/AutoReimport/ReimportFeedbackContext.cpp
+++ b/Engine/Source/Editor/UnrealEd/Private/AutoReimport/ReimportFeedbackContext.cpp
@@ -1,7 +1,5 @@
 // Copyright 1998-2015 Epic Games, Inc. All Rights Reserved.
 
-#pragma once
-
 #include "UnrealEd.h"
 
 #include "ReimportFeedbackContext.h"
