Add frame tracking for the network profiler in UEditorEngine so that network profiles can be captured properly in the editor.
Thanks to [REDACTED_PERSON]!
#github https://github.com/EpicGames/UnrealEngine/pull/771

diff --git a/Engine/Source/Editor/UnrealEd/Private/Editor.cpp b/Engine/Source/Editor/UnrealEd/Private/Editor.cpp
--- a/Engine/Source/Editor/UnrealEd/Private/Editor.cpp
+++ b/Engine/Source/Editor/UnrealEd/Private/Editor.cpp
@@ -11,6 +11,7 @@
 #include "Factories.h"
 #include "BSPOps.h"
 #include "EditorCommandLineUtils.h"
+#include "Net/NetworkProfiler.h"
 
 // needed for the RemotePropagator
 #include "SoundDefinitions.h"
@@ -902,6 +903,8 @@ void UEditorEngine::AddReferencedObjects(UObject* InThis, FReferenceCollector& C
 
 void UEditorEngine::Tick( float DeltaSeconds, bool bIdleMode )
 {
+	NETWORK_PROFILER(GNetworkProfiler.TrackFrameBegin());
+
 	UWorld* CurrentGWorld = GWorld;
 	check( CurrentGWorld );
 	check( CurrentGWorld != PlayWorld || bIsSimulatingInEditor );
diff --git a/Engine/Source/Runtime/Engine/Public/Net/NetworkProfiler.h b/Engine/Source/Runtime/Engine/Public/Net/NetworkProfiler.h
--- a/Engine/Source/Runtime/Engine/Public/Net/NetworkProfiler.h
+++ b/Engine/Source/Runtime/Engine/Public/Net/NetworkProfiler.h
@@ -93,7 +93,7 @@ public:
 	/**
 	 * Marks the beginning of a frame.
 	 */
-	void TrackFrameBegin();
+	ENGINE_API void TrackFrameBegin();
 	
 	/**
 	 * Tracks and RPC being sent.
