Integrating change 2420817 from 4.7 branch

Stopped ctrl/shift/alt click on a selected actor (without any components selected) from selecting a component

diff --git a/Engine/Source/Editor/UnrealEd/Private/LevelEditorViewport.cpp b/Engine/Source/Editor/UnrealEd/Private/LevelEditorViewport.cpp
--- a/Engine/Source/Editor/UnrealEd/Private/LevelEditorViewport.cpp
+++ b/Engine/Source/Editor/UnrealEd/Private/LevelEditorViewport.cpp
@@ -1788,13 +1788,16 @@ void FLevelEditorViewportClient::ProcessClick(FSceneView& View, HHitProxy* HitPr
 			// 1. The actor clicked is already selected
 			// 2. The actor selected is the only actor selected
 			// 3. The LMB was pressed or we already have a component selected (so that right-clicking a selected actor won't select a component)
-			// 4. The click was not a double click (unless a component has already been selected)
-			// 5. The level viewport didn't just receive focus this frame (again unless a component has already been selected)
+			// And, if a component is not selected already:
+			// 4. The click was not a double click
+			// 5. The level viewport didn't just receive focus this frame
+			// 6. No modifier keys are pressed
 			const bool bActorAlreadySelectedExclusively = GEditor->GetSelectedActors()->IsSelected(ActorHitProxy->Actor) && ( GEditor->GetSelectedActorCount() == 1 );
 			const bool bComponentAlreadySelected = GEditor->GetSelectedComponentCount() > 0;
 			const bool bCanBeginSelectingComponents =  (Click.GetKey() == EKeys::LeftMouseButton) 
 													&& (Click.GetEvent() != IE_DoubleClick) 
-													&& !bReceivedFocusRecently;
+													&& !bReceivedFocusRecently
+													&& !( Click.IsAltDown() || Click.IsControlDown() || Click.IsShiftDown() );
 
 			const bool bSelectComponent = bActorAlreadySelectedExclusively && ( bComponentAlreadySelected || bCanBeginSelectingComponents );
 
