Changed FBX option window to use standard detail view

diff --git a/Engine/Source/Editor/UnrealEd/Private/Fbx/FbxOptionWindow.cpp b/Engine/Source/Editor/UnrealEd/Private/Fbx/FbxOptionWindow.cpp
--- a/Engine/Source/Editor/UnrealEd/Private/Fbx/FbxOptionWindow.cpp
+++ b/Engine/Source/Editor/UnrealEd/Private/Fbx/FbxOptionWindow.cpp
@@ -3,7 +3,8 @@
 #include "FbxOptionWindow.h"
 #include "Editor/ContentBrowser/Public/ContentBrowserModule.h"
 #include "IDocumentation.h"
-#include "SKismetInspector.h"
+#include "PropertyEditorModule.h"
+#include "IDetailsView.h"
 
 #define LOCTEXT_NAMESPACE "FBXOption"
 
@@ -93,14 +94,14 @@ void SFbxOptionWindow::Construct(const FArguments& InArgs)
 			]
 	];
 
-	TSharedPtr<SKismetInspector> Inspector;
-	SAssignNew(Inspector, SKismetInspector)
-		.ShowTitleArea(false)
-		.HideNameArea(true);
+	FPropertyEditorModule& PropertyEditorModule = FModuleManager::GetModuleChecked<FPropertyEditorModule>("PropertyEditor");
+	FDetailsViewArgs DetailsViewArgs;
+	DetailsViewArgs.bAllowSearch = false;
+	DetailsViewArgs.bHideActorNameArea = true;
+	TSharedPtr<IDetailsView> DetailsView = PropertyEditorModule.CreateDetailView(DetailsViewArgs);
 
-	InspectorBox->SetContent(Inspector->AsShared());
-
-	Inspector->ShowDetailsForSingleObject(ImportUI);
+	InspectorBox->SetContent(DetailsView->AsShared());
+	DetailsView->SetObject(ImportUI);
 }
 
 bool SFbxOptionWindow::CanImport()  const
