Added commandlet helpers - primarily, a function for spawning a new process of UE4 to run a commandlet.

diff --git a/Engine/Source/Editor/UnrealEd/Private/Commandlets/CommandletHelpers.cpp b/Engine/Source/Editor/UnrealEd/Private/Commandlets/CommandletHelpers.cpp
new file mode 100644
--- /dev/null
+++ b/Engine/Source/Editor/UnrealEd/Private/Commandlets/CommandletHelpers.cpp
@@ -0,0 +1,33 @@
+// Copyright 1998-2014 Epic Games, Inc. All Rights Reserved.
+
+#include "UnrealEd.h"
+#include "Commandlets/CommandletHelpers.h"
+
+namespace CommandletHelpers
+{
+	FProcHandle CreateCommandletProcess(const TCHAR* CommandletName, const TCHAR* ProjectPath, const TCHAR* Parms, bool bLaunchDetached, bool bLaunchHidden, bool bLaunchReallyHidden, uint32* OutProcessID, int32 PriorityModifier, const TCHAR* OptionalWorkingDirectory, void* PipeWrite)
+	{
+		const FString ExecutableURL = FUnrealEdMisc::Get().GetExecutableForCommandlets();
+
+		FProcHandle Result;
+
+		// Construct arguments string.
+		if (CommandletName && FCString::Strlen(CommandletName) != 0)
+		{
+			FString Arguments;
+			const FString RunCommand = FString::Printf( TEXT("-run=%s"), CommandletName );
+			if (ProjectPath && FCString::Strlen(ProjectPath) != 0)
+			{
+				Arguments = FString::Printf( TEXT("%s %s"), ProjectPath, *RunCommand );
+			}
+			if (Parms && FCString::Strlen(Parms) != 0)
+			{
+				Arguments = FString::Printf( TEXT("%s %s"), *Arguments, Parms );
+			}
+
+			Result = FPlatformProcess::CreateProc(*ExecutableURL, *Arguments, bLaunchDetached, bLaunchHidden, bLaunchReallyHidden, OutProcessID, PriorityModifier, OptionalWorkingDirectory, PipeWrite);
+		}
+
+		return Result;
+	}
+}
\ No newline at end of file
diff --git a/Engine/Source/Editor/UnrealEd/Public/Commandlets/CommandletHelpers.h b/Engine/Source/Editor/UnrealEd/Public/Commandlets/CommandletHelpers.h
new file mode 100644
--- /dev/null
+++ b/Engine/Source/Editor/UnrealEd/Public/Commandlets/CommandletHelpers.h
@@ -0,0 +1,8 @@
+// Copyright 1998-2014 Epic Games, Inc. All Rights Reserved.
+
+#pragma once
+
+namespace CommandletHelpers
+{
+	UNREALED_API FProcHandle CreateCommandletProcess(const TCHAR* CommandletName, const TCHAR* ProjectPath, const TCHAR* Parms, bool bLaunchDetached, bool bLaunchHidden, bool bLaunchReallyHidden, uint32* OutProcessID, int32 PriorityModifier, const TCHAR* OptionalWorkingDirectory, void* PipeWrite);
+}
\ No newline at end of file
