Fix for hide name area flag error, changed to new enum value.

diff --git a/Engine/Source/Editor/UnrealEd/Private/Fbx/FbxOptionWindow.cpp b/Engine/Source/Editor/UnrealEd/Private/Fbx/FbxOptionWindow.cpp
--- a/Engine/Source/Editor/UnrealEd/Private/Fbx/FbxOptionWindow.cpp
+++ b/Engine/Source/Editor/UnrealEd/Private/Fbx/FbxOptionWindow.cpp
@@ -97,7 +97,7 @@ void SFbxOptionWindow::Construct(const FArguments& InArgs)
 	FPropertyEditorModule& PropertyEditorModule = FModuleManager::GetModuleChecked<FPropertyEditorModule>("PropertyEditor");
 	FDetailsViewArgs DetailsViewArgs;
 	DetailsViewArgs.bAllowSearch = false;
-	DetailsViewArgs.bHideActorNameArea = true;
+	DetailsViewArgs.NameAreaSettings = FDetailsViewArgs::HideNameArea;
 	TSharedPtr<IDetailsView> DetailsView = PropertyEditorModule.CreateDetailView(DetailsViewArgs);
 
 	InspectorBox->SetContent(DetailsView->AsShared());
