fixed check() triggering

diff --git a/Engine/Source/Editor/UnrealEd/Private/AssetThumbnail.cpp b/Engine/Source/Editor/UnrealEd/Private/AssetThumbnail.cpp
--- a/Engine/Source/Editor/UnrealEd/Private/AssetThumbnail.cpp
+++ b/Engine/Source/Editor/UnrealEd/Private/AssetThumbnail.cpp
@@ -196,7 +196,7 @@ public:
 			.Padding(this, &SAssetThumbnail::GetAssetColorStripPadding)
 		];
 
-		if( InArgs._AllowAssetSpecificThumbnailOverlay )
+		if( InArgs._AllowAssetSpecificThumbnailOverlay && AssetTypeActions.IsValid() )
 		{
 			// Does the asset provide an additional thumbnail overlay?
 			TSharedPtr<SWidget> AssetSpecificThumbnailOverlay = AssetTypeActions->GetThumbnailOverlay(AssetData);
