Standardized dialog when importing an asset which already exists.

diff --git a/Engine/Plugins/Runtime/SoundMod/Source/SoundModImporter/Private/SoundModImporterFactory.cpp b/Engine/Plugins/Runtime/SoundMod/Source/SoundModImporter/Private/SoundModImporterFactory.cpp
--- a/Engine/Plugins/Runtime/SoundMod/Source/SoundModImporter/Private/SoundModImporterFactory.cpp
+++ b/Engine/Plugins/Runtime/SoundMod/Source/SoundModImporter/Private/SoundModImporterFactory.cpp
@@ -56,13 +56,9 @@ UObject* USoundModImporterFactory::FactoryCreateBinary
 
 	if (ExistingSound && !bUseExistingSettings && !GIsAutomationTesting)
 	{
-		// Prompt the user for what to do if a 'To All' response wasn't already given.
-		if (OverwriteYesOrNoToAllState != EAppReturnType::YesAll && OverwriteYesOrNoToAllState != EAppReturnType::NoAll)
-		{
-			OverwriteYesOrNoToAllState = FMessageDialog::Open(EAppMsgType::YesNoYesAllNoAllCancel, FText::Format(
-				NSLOCTEXT("UnrealEd", "ImportedSoundAlreadyExists_F", "You are about to import '{0}' over an existing sound. Would you like to overwrite the existing settings?\n\nYes or Yes to All: Overwrite the existing settings.\nNo or No to All: Preserve the existing settings.\nCancel: Abort the operation."),
-				FText::FromName(Name)));
-		}
+		DisplayOverwriteOptionsDialog(FText::Format(
+			NSLOCTEXT("SoundModImporterFactory", "ImportOverwriteWarning", "You are about to import '{0}' over an existing sound."),
+			FText::FromName(Name)));
 
 		switch (OverwriteYesOrNoToAllState)
 		{
diff --git a/Engine/Source/Editor/UnrealEd/Classes/Factories/Factory.h b/Engine/Source/Editor/UnrealEd/Classes/Factories/Factory.h
--- a/Engine/Source/Editor/UnrealEd/Classes/Factories/Factory.h
+++ b/Engine/Source/Editor/UnrealEd/Classes/Factories/Factory.h
@@ -115,6 +115,11 @@ public:
 	 */
 	static void ResetState();
 
+	/**
+	 * Pop up message to the user asking whether they wish to overwrite existing state or not
+	 */
+	static void DisplayOverwriteOptionsDialog(const FText& Message);
+
 	/** Opens a dialog to configure the factory properties. Return false if user opted out of configuring properties */
 	virtual bool ConfigureProperties() { return true; }
 
diff --git a/Engine/Source/Editor/UnrealEd/Private/Factories/EditorFactories.cpp b/Engine/Source/Editor/UnrealEd/Private/Factories/EditorFactories.cpp
--- a/Engine/Source/Editor/UnrealEd/Private/Factories/EditorFactories.cpp
+++ b/Engine/Source/Editor/UnrealEd/Private/Factories/EditorFactories.cpp
@@ -1703,13 +1703,9 @@ UObject* USoundFactory::FactoryCreateBinary
 
 		if( ExistingSound && !bSoundFactorySuppressImportOverwriteDialog && !GIsAutomationTesting )
 		{
-			// Prompt the user for what to do if a 'To All' response wasn't already given.
-			if( OverwriteYesOrNoToAllState != EAppReturnType::YesAll && OverwriteYesOrNoToAllState != EAppReturnType::NoAll )
-			{
-				OverwriteYesOrNoToAllState = FMessageDialog::Open( EAppMsgType::YesNoYesAllNoAllCancel, FText::Format(
-						NSLOCTEXT("UnrealEd", "ImportedSoundAlreadyExists_F", "You are about to import '{0}' over an existing sound. Would you like to overwrite the existing settings?\n\nYes or Yes to All: Overwrite the existing settings.\nNo or No to All: Preserve the existing settings.\nCancel: Abort the operation." ),
-						FText::FromName(Name) ) );
-			}
+			DisplayOverwriteOptionsDialog(FText::Format(
+				NSLOCTEXT("SoundFactory", "ImportOverwriteWarning", "You are about to import '{0}' over an existing sound."),
+				FText::FromName(Name)));
 
 			switch( OverwriteYesOrNoToAllState )
 			{
@@ -4304,12 +4300,9 @@ UObject* UTextureFactory::FactoryCreateBinary
 
 	if(ExistingTexture && !bSuppressImportOverwriteDialog)
 	{
-		// Prompt the user for what to do if a 'To All' response wasn't already given.
-		if( OverwriteYesOrNoToAllState != EAppReturnType::YesAll && OverwriteYesOrNoToAllState != EAppReturnType::NoAll )
-		{
-			OverwriteYesOrNoToAllState = FMessageDialog::Open(EAppMsgType::YesNoYesAllNoAllCancel,
-				NSLOCTEXT("UnrealEd", "ImportedTextureAlreadyExists_F", "The specified texture already exists and will be overwritten.  Do you want to keep the existing texture's settings?" ) );
-		}
+		DisplayOverwriteOptionsDialog(FText::Format(
+			NSLOCTEXT("TextureFactory", "ImportOverwriteWarning", "You are about to import '{0}' over an existing texture."),
+			FText::FromName(Name)));
 
 		switch( OverwriteYesOrNoToAllState )
 		{
@@ -4317,15 +4310,15 @@ UObject* UTextureFactory::FactoryCreateBinary
 		case EAppReturnType::Yes:
 		case EAppReturnType::YesAll:
 			{
-				// Preserve existing settings
-				bUsingExistingSettings = true;
+				// Overwrite existing settings
+				bUsingExistingSettings = false;
 				break;
 			}
 		case EAppReturnType::No:
 		case EAppReturnType::NoAll:
 			{
-				// Overwrite existing settings
-				bUsingExistingSettings = false;
+				// Preserve existing settings
+				bUsingExistingSettings = true;
 				break;
 			}
 		case EAppReturnType::Cancel:
diff --git a/Engine/Source/Editor/UnrealEd/Private/Factories/Factory.cpp b/Engine/Source/Editor/UnrealEd/Private/Factories/Factory.cpp
--- a/Engine/Source/Editor/UnrealEd/Private/Factories/Factory.cpp
+++ b/Engine/Source/Editor/UnrealEd/Private/Factories/Factory.cpp
@@ -163,6 +163,17 @@ void UFactory::ResetState()
 	bAllowOneTimeWarningMessages = true;
 }
 
+void UFactory::DisplayOverwriteOptionsDialog(const FText& Message)
+{
+	// Prompt the user for what to do if a 'To All' response wasn't already given.
+	if (OverwriteYesOrNoToAllState != EAppReturnType::YesAll && OverwriteYesOrNoToAllState != EAppReturnType::NoAll)
+	{
+		OverwriteYesOrNoToAllState = FMessageDialog::Open(EAppMsgType::YesNoYesAllNoAllCancel, FText::Format(
+			NSLOCTEXT("UnrealEd", "ImportedAssetAlreadyExists", "{0} Would you like to overwrite the existing settings?\n\nYes or Yes to All: Overwrite the existing settings.\nNo or No to All: Preserve the existing settings.\nCancel: Abort the operation."),
+			Message));
+	}
+}
+
 UObject* UFactory::StaticImportObject
 (
 UClass*				Class,
