Fixed deprecation warnings

diff --git a/Engine/Source/Editor/UnrealEd/Private/AutoReimport/ReimportFeedbackContext.cpp b/Engine/Source/Editor/UnrealEd/Private/AutoReimport/ReimportFeedbackContext.cpp
--- a/Engine/Source/Editor/UnrealEd/Private/AutoReimport/ReimportFeedbackContext.cpp
+++ b/Engine/Source/Editor/UnrealEd/Private/AutoReimport/ReimportFeedbackContext.cpp
@@ -104,18 +104,18 @@ class SWidgetStack : public SVerticalBox
 		StartSizeOffset = ComputeTotalSize().Y - GetDesiredSize().Y;
 		LerpSizeOffset = StartSizeOffset;
 
-		SlideCurve.Play();
+		SlideCurve.Play(AsShared());
 
 		if (Children.Num() > MaxNumVisible)
 		{
-			auto& Curve = StaticCastSharedRef<SWidgetStackItem>(Children[MaxNumVisible].GetWidget())->OpacityCurve;
-			if (Curve.IsPlaying())
+			auto Widget = StaticCastSharedRef<SWidgetStackItem>(Children[MaxNumVisible].GetWidget());
+			if (Widget->OpacityCurve.IsPlaying())
 			{
-				Curve.Reverse();
+				Widget->OpacityCurve.Reverse();
 			}
 			else
 			{
-				Curve.PlayReverse();
+				Widget->OpacityCurve.PlayReverse(Widget);
 			}
 		}
 	}
@@ -162,7 +162,7 @@ class SWidgetStack : public SVerticalBox
 		void Construct(const FArguments& InArgs)
 		{
 			OpacityCurve = FCurveSequence(0.f, .5f, ECurveEaseFunction::QuadOut);
-			OpacityCurve.Play();
+			OpacityCurve.Play(AsShared());
 
 			ChildSlot
 			[
