Fixed some new code that was providing text to Slate as FString rather than FText

UETOOL-213 - Minimize Slate FString -> FText conversion (remove SLATE_TEXT_ATTRIBUTE)

diff --git a/Engine/Source/Editor/LocalizationDashboard/Private/SLocalizationTargetEditor.cpp b/Engine/Source/Editor/LocalizationDashboard/Private/SLocalizationTargetEditor.cpp
--- a/Engine/Source/Editor/LocalizationDashboard/Private/SLocalizationTargetEditor.cpp
+++ b/Engine/Source/Editor/LocalizationDashboard/Private/SLocalizationTargetEditor.cpp
@@ -555,7 +555,7 @@ namespace
 					.Content()
 					[
 						SNew(STextBlock)
-						.Text(OtherLocalizationTarget->Settings.Name)
+						.Text(FText::FromString(OtherLocalizationTarget->Settings.Name))
 					]
 			];
 	}
diff --git a/Engine/Source/Editor/UnrealEd/Private/SCSVImportOptions.cpp b/Engine/Source/Editor/UnrealEd/Private/SCSVImportOptions.cpp
--- a/Engine/Source/Editor/UnrealEd/Private/SCSVImportOptions.cpp
+++ b/Engine/Source/Editor/UnrealEd/Private/SCSVImportOptions.cpp
@@ -211,7 +211,7 @@ FString SCSVImportOptions::GetImportTypeText(TSharedPtr<ECSVImportType> Type) co
 TSharedRef<SWidget> SCSVImportOptions::MakeImportTypeItemWidget(TSharedPtr<ECSVImportType> Type)
 {
 	return	SNew(STextBlock)
-			.Text(GetImportTypeText(Type));
+			.Text(FText::FromString(GetImportTypeText(Type)));
 }
 
 /** Called to create a widget for each struct */
@@ -219,7 +219,7 @@ TSharedRef<SWidget> SCSVImportOptions::MakeRowStructItemWidget(UScriptStruct* St
 {
 	check( Struct != NULL );
 	return	SNew(STextBlock)
-			.Text(Struct->GetName());
+			.Text(FText::FromString(Struct->GetName()));
 }
 
 FString SCSVImportOptions::GetCurveTypeText(CurveInterpModePtr InterpMode) const
@@ -247,7 +247,7 @@ FString SCSVImportOptions::GetCurveTypeText(CurveInterpModePtr InterpMode) const
 TSharedRef<SWidget> SCSVImportOptions::MakeCurveTypeWidget(CurveInterpModePtr InterpMode)
 {
 	FString Label = GetCurveTypeText(InterpMode);
-	return SNew(STextBlock) .Text( Label );
+	return SNew(STextBlock) .Text( FText::FromString(Label) );
 }
 
 /** Called when 'OK' button is pressed */
@@ -278,29 +278,29 @@ FReply SCSVImportOptions::OnCancel()
 	return FReply::Handled();
 }
 
-FString SCSVImportOptions::GetSelectedItemText() const
+FText SCSVImportOptions::GetSelectedItemText() const
 {
 	TSharedPtr<ECSVImportType> SelectedType = ImportTypeCombo->GetSelectedItem();
 
 	return (SelectedType.IsValid())
-		? GetImportTypeText(SelectedType)
-		: FString();
+		? FText::FromString(GetImportTypeText(SelectedType))
+		: FText::GetEmpty();
 }
 
-FString SCSVImportOptions::GetSelectedRowOptionText() const
+FText SCSVImportOptions::GetSelectedRowOptionText() const
 {
 	UScriptStruct* SelectedScript = RowStructCombo->GetSelectedItem();
 	return (SelectedScript)
-		? SelectedScript->GetName()
-		: FString();
+		? FText::FromString(SelectedScript->GetName())
+		: FText::GetEmpty();
 }
 
-FString SCSVImportOptions::GetSelectedCurveTypeText() const
+FText SCSVImportOptions::GetSelectedCurveTypeText() const
 {
 	CurveInterpModePtr CurveModePtr = CurveInterpCombo->GetSelectedItem();
 	return (CurveModePtr.IsValid())
-		? GetCurveTypeText(CurveModePtr)
-		: FString();
+		? FText::FromString(GetCurveTypeText(CurveModePtr))
+		: FText::GetEmpty();
 }
 
 #undef LOCTEXT_NAMESPACE
diff --git a/Engine/Source/Editor/UnrealEd/Public/SCSVImportOptions.h b/Engine/Source/Editor/UnrealEd/Public/SCSVImportOptions.h
--- a/Engine/Source/Editor/UnrealEd/Public/SCSVImportOptions.h
+++ b/Engine/Source/Editor/UnrealEd/Public/SCSVImportOptions.h
@@ -65,11 +65,11 @@ public:
 	/** Called when 'Cancel' button is pressed */
 	FReply OnCancel();
 
-	FString GetSelectedItemText() const;
+	FText GetSelectedItemText() const;
 
-	FString GetSelectedRowOptionText() const;
+	FText GetSelectedRowOptionText() const;
 
-	FString GetSelectedCurveTypeText() const;
+	FText GetSelectedCurveTypeText() const;
 
 private:
 	/** Whether we should go ahead with import */
diff --git a/Engine/Source/Programs/UnrealSync/Private/GUI.cpp b/Engine/Source/Programs/UnrealSync/Private/GUI.cpp
--- a/Engine/Source/Programs/UnrealSync/Private/GUI.cpp
+++ b/Engine/Source/Programs/UnrealSync/Private/GUI.cpp
@@ -493,7 +493,7 @@ class SPickLabel : public SCompoundWidget, public ILabelNameProvider
 {
 public:
 	SLATE_BEGIN_ARGS(SPickLabel) {}
-		SLATE_ATTRIBUTE(FString, Title)
+		SLATE_ATTRIBUTE(FText, Title)
 	SLATE_END_ARGS()
 
 	BEGIN_SLATE_FUNCTION_BUILD_OPTIMIZATION
@@ -510,9 +510,9 @@ public:
 					SNew(STextBlock).Text(InArgs._Title)
 				]
 				+ SHorizontalBox::Slot().HAlign(HAlign_Fill)
-					[
-						LabelsCombo.ToSharedRef()
-					]
+				[
+					LabelsCombo.ToSharedRef()
+				]
 			];
 	}
 	END_SLATE_FUNCTION_BUILD_OPTIMIZATION
@@ -847,9 +847,9 @@ public:
 		RadioSelection = SNew(SRadioContentSelection);
 
 		AddToRadioSelection("Sync to the latest promoted", SNew(SLatestPromoted));
-		AddToRadioSelection("Sync to chosen promoted label", SNew(SPickPromoted).Title("Pick promoted label: "));
-		AddToRadioSelection("Sync to chosen promotable label since last promoted", SNew(SPickPromotable).Title("Pick promotable label: "));
-		AddToRadioSelection("Sync to any chosen label", SNew(SPickAny).Title("Pick label: "));
+		AddToRadioSelection("Sync to chosen promoted label", SNew(SPickPromoted).Title(FText::FromString("Pick promoted label: ")));
+		AddToRadioSelection("Sync to chosen promotable label since last promoted", SNew(SPickPromotable).Title(FText::FromString("Pick promotable label: ")));
+		AddToRadioSelection("Sync to any chosen label", SNew(SPickAny).Title(FText::FromString("Pick label: ")));
 
 		PickGameWidget = SNew(SPickGameWidget)
 			.OnGamePicked(this, &SMainTabWidget::OnCurrentGameChanged);
