Fix crash applying materials when there is too many material slots on the component and that dont match the number of materials on the mesh

diff --git a/Engine/Source/Editor/UnrealEd/Private/AssetSelection.cpp b/Engine/Source/Editor/UnrealEd/Private/AssetSelection.cpp
--- a/Engine/Source/Editor/UnrealEd/Private/AssetSelection.cpp
+++ b/Engine/Source/Editor/UnrealEd/Private/AssetSelection.cpp
@@ -765,8 +765,10 @@ bool FActorFactoryAssetProxy::ApplyMaterialToActor( AActor* TargetActor, UMateri
 							}
 							else
 							{
-								check(OptionalMaterialSlot < MaterialCount);
-								FoundMeshComponent->SetMaterial(OptionalMaterialSlot, MaterialToApply);
+								if( OptionalMaterialSlot < MaterialCount )
+								{
+									FoundMeshComponent->SetMaterial(OptionalMaterialSlot, MaterialToApply);
+								}
 							}
 
 							TargetActor->MarkComponentsRenderStateDirty();
