Made the auto-saved directory path hard-coded

UE-8432 - Crash saving an asset in a project with a long path name

The auto-save dir used to be saved in the editor config file, but was always the same and couldn't be moved outside of the "Saved" directory. This was brought up due to the config issues raised by UE-8542, so we've now agreed to just hard-code this path.

diff --git a/Engine/Config/BaseEngine.ini b/Engine/Config/BaseEngine.ini
--- a/Engine/Config/BaseEngine.ini
+++ b/Engine/Config/BaseEngine.ini
@@ -1309,7 +1309,6 @@ UseDynamicStreaming=True
 BoostPlayerTextures=3.0
 
 [/Script/UnrealEd.EditorEngine]
-AutoSaveDir=%GAMEDIR%Saved/Autosaves
 LocalPlayerClassName=/Script/Engine.LocalPlayer
 bSubtitlesEnabled=True
 GameCommandLine=-log
diff --git a/Engine/Source/Editor/UnrealEd/Classes/Editor/EditorEngine.h b/Engine/Source/Editor/UnrealEd/Classes/Editor/EditorEngine.h
--- a/Engine/Source/Editor/UnrealEd/Classes/Editor/EditorEngine.h
+++ b/Engine/Source/Editor/UnrealEd/Classes/Editor/EditorEngine.h
@@ -288,10 +288,6 @@ class UNREALED_API UEditorEngine : public UEngine
 	UPROPERTY(EditAnywhere, config, Category=Advanced)
 	uint32 GodMode:1;
 
-	/** The location to autosave to. */
-	UPROPERTY(EditAnywhere, config, Category=Advanced)
-	FString AutoSaveDir;
-
 	UPROPERTY(EditAnywhere, config, Category=Advanced)
 	uint32 InvertwidgetZAxis:1;
 
diff --git a/Engine/Source/Editor/UnrealEd/Private/AutoSaveUtils.cpp b/Engine/Source/Editor/UnrealEd/Private/AutoSaveUtils.cpp
new file mode 100644
--- /dev/null
+++ b/Engine/Source/Editor/UnrealEd/Private/AutoSaveUtils.cpp
@@ -0,0 +1,9 @@
+// Copyright 1998-2015 Epic Games, Inc. All Rights Reserved.
+
+#include "UnrealEd.h"
+#include "AutoSaveUtils.h"
+
+FString AutoSaveUtils::GetAutoSaveDir()
+{
+	return FPaths::GameSavedDir() / TEXT("Autosaves");
+}
diff --git a/Engine/Source/Editor/UnrealEd/Private/AutoSaveUtils.h b/Engine/Source/Editor/UnrealEd/Private/AutoSaveUtils.h
new file mode 100644
--- /dev/null
+++ b/Engine/Source/Editor/UnrealEd/Private/AutoSaveUtils.h
@@ -0,0 +1,9 @@
+// Copyright 1998-2015 Epic Games, Inc. All Rights Reserved.
+
+#pragma once
+
+namespace AutoSaveUtils
+{
+	/** Get the root directory that we should save auto-saved assets and maps to */
+	FString GetAutoSaveDir();
+}
diff --git a/Engine/Source/Editor/UnrealEd/Private/Commandlets/FixupRedirectsCommandlet.cpp b/Engine/Source/Editor/UnrealEd/Private/Commandlets/FixupRedirectsCommandlet.cpp
--- a/Engine/Source/Editor/UnrealEd/Private/Commandlets/FixupRedirectsCommandlet.cpp
+++ b/Engine/Source/Editor/UnrealEd/Private/Commandlets/FixupRedirectsCommandlet.cpp
@@ -7,6 +7,7 @@
 #include "UnrealEd.h"
 #include "ISourceControlModule.h"
 #include "Engine/UserDefinedEnum.h"
+#include "AutoSaveUtils.h"
 
 DEFINE_LOG_CATEGORY_STATIC(LogFixupRedirectsCommandlet, Log, All);
 
@@ -140,7 +141,7 @@ int32 UFixupRedirectsCommandlet::Main( const FString& Params )
 			const FString& Filename = PackageList[PackageIndex];
 
 			// already loaded code, skip them, and skip autosave packages
-			if (Filename.Left(GEditor->AutoSaveDir.Len()) == GEditor->AutoSaveDir)
+			if (Filename.StartsWith(AutoSaveUtils::GetAutoSaveDir(), ESearchCase::IgnoreCase))
 			{
 				continue;
 			}
diff --git a/Engine/Source/Editor/UnrealEd/Private/FileHelpers.cpp b/Engine/Source/Editor/UnrealEd/Private/FileHelpers.cpp
--- a/Engine/Source/Editor/UnrealEd/Private/FileHelpers.cpp
+++ b/Engine/Source/Editor/UnrealEd/Private/FileHelpers.cpp
@@ -31,6 +31,7 @@
 #include "NotificationManager.h"
 #include "Engine/LevelStreaming.h"
 #include "GameMapsSettings.h"
+#include "AutoSaveUtils.h"
 
 DEFINE_LOG_CATEGORY_STATIC(LogFileHelpers, Log, All);
 
@@ -392,7 +393,7 @@ static bool SaveWorld(UWorld* World,
 		{
 			// If package exists, but doesn't feature the default extension, it will not load when launched,
 			// Change the extension of the map to the default for the auto-save
-			Path			= GEditor->AutoSaveDir;
+			Path			= AutoSaveUtils::GetAutoSaveDir();
 			CleanFilename	= FPackageName::GetLongPackageAssetName(PackageName) + FPackageName::GetMapPackageExtension();
 		}
 		else
@@ -412,7 +413,7 @@ static bool SaveWorld(UWorld* World,
 	else
 	{
 		// No package filename exists and none was specified, so save the package in the autosaves folder.
-		Path			= GEditor->AutoSaveDir;
+		Path			= AutoSaveUtils::GetAutoSaveDir();
 		CleanFilename	= FPackageName::GetLongPackageAssetName(PackageName) + FPackageName::GetMapPackageExtension();
 	}
 
diff --git a/Engine/Source/Editor/UnrealEd/Private/PackageAutoSaver.cpp b/Engine/Source/Editor/UnrealEd/Private/PackageAutoSaver.cpp
--- a/Engine/Source/Editor/UnrealEd/Private/PackageAutoSaver.cpp
+++ b/Engine/Source/Editor/UnrealEd/Private/PackageAutoSaver.cpp
@@ -6,6 +6,7 @@
 #include "Json.h"
 #include "SNotificationList.h"
 #include "NotificationManager.h"
+#include "AutoSaveUtils.h"
 
 namespace PackageAutoSaverJson
 {
@@ -160,7 +161,7 @@ void FPackageAutoSaver::AttemptAutoSave()
 			GUnrealEd->SaveConfig();
 
 			// Make sure the auto-save directory exists before attempting to write the file
-			const FString& AutoSaveDir = GUnrealEd->AutoSaveDir;
+			const FString AutoSaveDir = AutoSaveUtils::GetAutoSaveDir();
 			IFileManager::Get().MakeDirectory(*AutoSaveDir, true);
 
 			// Auto-save maps and/or content packages based on user settings.
@@ -303,7 +304,7 @@ void FPackageAutoSaver::OnPackageSaved(const FString& Filename, UObject* Obj)
 		{
 			// Make the filename relative to the auto-save directory
 			// Note: MakePathRelativeTo modifies in-place, hence the copy of Filename
-			const FString AutoSaveDir = GUnrealEd->AutoSaveDir / "";
+			const FString AutoSaveDir = AutoSaveUtils::GetAutoSaveDir() / "";
 			FString RelativeFilename = Filename;
 			FPaths::MakePathRelativeTo(RelativeFilename, *AutoSaveDir);
 
@@ -592,7 +593,7 @@ void FPackageAutoSaver::ClearStalePointers()
 
 FString PackageAutoSaverJson::GetRestoreFilename(const bool bEnsurePath)
 {
-	const FString& AutoSaveDir = GUnrealEd->AutoSaveDir;
+	const FString AutoSaveDir = AutoSaveUtils::GetAutoSaveDir();
 	if(bEnsurePath)
 	{
 		// Make sure the auto-save directory exists before attempting to write the file
diff --git a/Engine/Source/Editor/UnrealEd/Private/PackageRestore.cpp b/Engine/Source/Editor/UnrealEd/Private/PackageRestore.cpp
--- a/Engine/Source/Editor/UnrealEd/Private/PackageRestore.cpp
+++ b/Engine/Source/Editor/UnrealEd/Private/PackageRestore.cpp
@@ -3,6 +3,7 @@
 #include "UnrealEd.h"
 #include "PackageRestore.h"
 #include "PackageTools.h"
+#include "AutoSaveUtils.h"
 
 #define LOCTEXT_NAMESPACE "PackageRestore"
 
@@ -514,7 +515,7 @@ namespace PackageRestore
 
 FEditorFileUtils::EPromptReturnCode PackageRestore::PromptToRestorePackages(const TMap<FString, FString>& PackagesToRestore, TArray<FString>* OutFailedPackages)
 {
-	const FString& AutoSaveDir = GUnrealEd->AutoSaveDir;
+	const FString AutoSaveDir = AutoSaveUtils::GetAutoSaveDir();
 
 	FPackageRestoreItems PackageRestoreItems;
 	PackageRestoreItems.Reserve(PackagesToRestore.Num());
diff --git a/Engine/Source/Editor/UnrealEd/Private/PackageRestore.h b/Engine/Source/Editor/UnrealEd/Private/PackageRestore.h
--- a/Engine/Source/Editor/UnrealEd/Private/PackageRestore.h
+++ b/Engine/Source/Editor/UnrealEd/Private/PackageRestore.h
@@ -1,7 +1,6 @@
 // Copyright 1998-2015 Epic Games, Inc. All Rights Reserved.
 
-#ifndef __PACKAGERESTORE_H__
-#define __PACKAGERESTORE_H__
+#pragma once
 
 #include "FileHelpers.h"
 
@@ -22,5 +21,3 @@ namespace PackageRestore
 	 */
 	FEditorFileUtils::EPromptReturnCode PromptToRestorePackages(const TMap<FString, FString>& PackagesToRestore, TArray<FString>* OutFailedPackages = nullptr);
 }
-
-#endif // __PACKAGERESTORE_H__
