better messaging during Launch On
UE-8127
#editor

diff --git a/Engine/Source/Editor/UnrealEd/Classes/Editor/EditorEngine.h b/Engine/Source/Editor/UnrealEd/Classes/Editor/EditorEngine.h
--- a/Engine/Source/Editor/UnrealEd/Classes/Editor/EditorEngine.h
+++ b/Engine/Source/Editor/UnrealEd/Classes/Editor/EditorEngine.h
@@ -2656,6 +2656,7 @@ private:
 	/** The platform to run on (as selected in dreop down) */
 	FString PlayUsingLauncherDeviceId;
 	FString PlayUsingLauncherDeviceName;
+	bool bPlayUsingLauncherHasCode;
 
 	/** List of files we are deferring adding to source control */
 	TArray<FString> DeferredFilesToAddToSourceControl;
diff --git a/Engine/Source/Editor/UnrealEd/Private/PlayLevel.cpp b/Engine/Source/Editor/UnrealEd/Private/PlayLevel.cpp
--- a/Engine/Source/Editor/UnrealEd/Private/PlayLevel.cpp
+++ b/Engine/Source/Editor/UnrealEd/Private/PlayLevel.cpp
@@ -1279,18 +1279,60 @@ private:
 void UEditorEngine::HandleStageStarted(const FString& InStage, TWeakPtr<SNotificationItem> NotificationItemPtr)
 {
 	FFormatNamedArguments Arguments;
-	Arguments.Add(TEXT("StageName"), FText::FromString(InStage));
-	Arguments.Add(TEXT("DeviceName"), FText::FromString(PlayUsingLauncherDeviceName));
-
 	FText NotificationText;
-	if(PlayUsingLauncherDeviceName.Len() == 0)
+	if (InStage.Contains(TEXT("Cooking")) || InStage.Contains(TEXT("Cook Task")))
 	{
-		NotificationText = FText::Format(LOCTEXT("LauncherTaskStageNotificationNoDevice", "{StageName} for launch..."), Arguments);
+		FString PlatformName = PlayUsingLauncherDeviceId.Left(PlayUsingLauncherDeviceId.Find(TEXT("@")));
+		if (PlatformName.Contains(TEXT("NoEditor")))
+		{
+			PlatformName = PlatformName.Left(PlatformName.Find(TEXT("NoEditor")));
+		}
+		Arguments.Add(TEXT("PlatformName"), FText::FromString(PlatformName));
+		NotificationText = FText::Format(LOCTEXT("LauncherTaskProcessingNotification", "Processing Assets for {PlatformName}..."), Arguments);
 	}
-	else
+	else if (InStage.Contains(TEXT("Build Task")))
 	{
-		NotificationText = FText::Format(LOCTEXT("LauncherTaskStageNotification", "{StageName} for launch on {DeviceName}..."), Arguments);
+		FString PlatformName = PlayUsingLauncherDeviceId.Left(PlayUsingLauncherDeviceId.Find(TEXT("@")));
+		if (PlatformName.Contains(TEXT("NoEditor")))
+		{
+			PlatformName = PlatformName.Left(PlatformName.Find(TEXT("NoEditor")));
+		}
+		Arguments.Add(TEXT("PlatformName"), FText::FromString(PlatformName));
+		if (FRocketSupport::IsRocket() || !bPlayUsingLauncherHasCode || !FSourceCodeNavigation::IsCompilerAvailable())
+		{
+			NotificationText = FText::Format(LOCTEXT("LauncherTaskValidateNotification", "Validating Executable for {PlatformName}..."), Arguments);
+		}
+		else
+		{
+			NotificationText = FText::Format(LOCTEXT("LauncherTaskBuildNotification", "Building Executable for {PlatformName}..."), Arguments);
+		}
 	}
+	else if (InStage.Contains(TEXT("Deploy Task")))
+	{
+		Arguments.Add(TEXT("DeviceName"), FText::FromString(PlayUsingLauncherDeviceName));
+		if(PlayUsingLauncherDeviceName.Len() == 0)
+		{
+			NotificationText = FText::Format(LOCTEXT("LauncherTaskStageNotificationNoDevice", "Deploying Executable and Assets..."), Arguments);
+		}
+		else
+		{
+			NotificationText = FText::Format(LOCTEXT("LauncherTaskStageNotification", "Deploying Executable and Assets to {DeviceName}..."), Arguments);
+		}
+	}
+	else if (InStage.Contains(TEXT("Run Task")))
+	{
+		Arguments.Add(TEXT("GameName"), FText::FromString(FApp::GetGameName()));
+		Arguments.Add(TEXT("DeviceName"), FText::FromString(PlayUsingLauncherDeviceName));
+		if(PlayUsingLauncherDeviceName.Len() == 0)
+		{
+			NotificationText = FText::Format(LOCTEXT("LauncherTaskStageNotificationNoDevice", "Running {GameName}..."), Arguments);
+		}
+		else
+		{
+			NotificationText = FText::Format(LOCTEXT("LauncherTaskStageNotification", "Running {GameName} on {DeviceName}..."), Arguments);
+		}
+	}
+
 	NotificationItemPtr.Pin()->SetText(NotificationText);
 }
 
@@ -1404,11 +1446,11 @@ void UEditorEngine::PlayUsingLauncher()
 
 		// does the project have any code?
 		FGameProjectGenerationModule& GameProjectModule = FModuleManager::LoadModuleChecked<FGameProjectGenerationModule>(TEXT("GameProjectGeneration"));
-		bool bHasCode = GameProjectModule.Get().ProjectHasCodeFiles();
+		bPlayUsingLauncherHasCode = GameProjectModule.Get().ProjectHasCodeFiles();
 
 		// Setup launch profile, keep the setting here to a minimum.
 		ILauncherProfileRef LauncherProfile = LauncherServicesModule.CreateProfile(TEXT("Play On Device"));
-		LauncherProfile->SetBuildGame(bHasCode && FSourceCodeNavigation::IsCompilerAvailable());
+		LauncherProfile->SetBuildGame(bPlayUsingLauncherHasCode && FSourceCodeNavigation::IsCompilerAvailable());
 
 		// select the quickest cook mode based on which in editor cook mode is enabled
 		bool bIncrimentalCooking = true;
@@ -1507,12 +1549,6 @@ void UEditorEngine::PlayUsingLauncher()
 
 		// create notification item
 		FText LaunchingText = LOCTEXT("LauncherTaskInProgressNotificationNoDevice", "Launching...");
-		if(PlayUsingLauncherDeviceName.Len() > 0)
-		{
-			FFormatNamedArguments Arguments;
-			Arguments.Add(TEXT("Device"), FText::FromString(PlayUsingLauncherDeviceName));
-			LaunchingText = FText::Format( LOCTEXT("LauncherTaskInProgressNotification", "Launching on {Device}..."), Arguments);
-		}
 		FNotificationInfo Info(LaunchingText);
 
 		Info.Image = FEditorStyle::GetBrush(TEXT("MainFrame.CookContent"));
@@ -1537,7 +1573,7 @@ void UEditorEngine::PlayUsingLauncher()
 
 		// analytics for launch on
 		int32 ErrorCode = 0;
-		FEditorAnalytics::ReportEvent(TEXT( "Editor.LaunchOn.Started" ), PlayUsingLauncherDeviceId.Left(PlayUsingLauncherDeviceId.Find(TEXT("@"))), bHasCode);
+		FEditorAnalytics::ReportEvent(TEXT( "Editor.LaunchOn.Started" ), PlayUsingLauncherDeviceId.Left(PlayUsingLauncherDeviceId.Find(TEXT("@"))), bPlayUsingLauncherHasCode);
 
 		NotificationItem->SetCompletionState(SNotificationItem::CS_Pending);
 		
@@ -1549,9 +1585,9 @@ void UEditorEngine::PlayUsingLauncher()
 			GEditor->PlayEditorSound(TEXT("/Engine/EditorSounds/Notifications/CompileStart_Cue.CompileStart_Cue"));
 			GEditor->LauncherWorker->OnOutputReceived().AddStatic(HandleOutputReceived);
 			GEditor->LauncherWorker->OnStageStarted().AddUObject(this, &UEditorEngine::HandleStageStarted, NotificationItemPtr);
-			GEditor->LauncherWorker->OnStageCompleted().AddUObject(this, &UEditorEngine::HandleStageCompleted, bHasCode, NotificationItemPtr);
-			GEditor->LauncherWorker->OnCompleted().AddUObject(this, &UEditorEngine::HandleLaunchCompleted, bHasCode, NotificationItemPtr, MessageLog);
-			GEditor->LauncherWorker->OnCanceled().AddUObject(this, &UEditorEngine::HandleLaunchCanceled, bHasCode, NotificationItemPtr);
+			GEditor->LauncherWorker->OnStageCompleted().AddUObject(this, &UEditorEngine::HandleStageCompleted, bPlayUsingLauncherHasCode, NotificationItemPtr);
+			GEditor->LauncherWorker->OnCompleted().AddUObject(this, &UEditorEngine::HandleLaunchCompleted, bPlayUsingLauncherHasCode, NotificationItemPtr, MessageLog);
+			GEditor->LauncherWorker->OnCanceled().AddUObject(this, &UEditorEngine::HandleLaunchCanceled, bPlayUsingLauncherHasCode, NotificationItemPtr);
 		}
 		else
 		{
@@ -1566,7 +1602,7 @@ void UEditorEngine::PlayUsingLauncher()
 			// analytics for launch on
 			TArray<FAnalyticsEventAttribute> ParamArray;
 			ParamArray.Add(FAnalyticsEventAttribute(TEXT("Time"), 0.0));
-			FEditorAnalytics::ReportEvent(TEXT( "Editor.LaunchOn.Failed" ), PlayUsingLauncherDeviceId.Left(PlayUsingLauncherDeviceId.Find(TEXT("@"))), bHasCode, EAnalyticsErrorCodes::LauncherFailed, ParamArray );
+			FEditorAnalytics::ReportEvent(TEXT( "Editor.LaunchOn.Failed" ), PlayUsingLauncherDeviceId.Left(PlayUsingLauncherDeviceId.Find(TEXT("@"))), bPlayUsingLauncherHasCode, EAnalyticsErrorCodes::LauncherFailed, ParamArray );
 		}
 	}
 }
