Cook in editor checks to see if cvars have changed (mobilehdr) before allowing cook in editor to be active.

diff --git a/Engine/Source/Editor/UnrealEd/Classes/CookOnTheSide/CookOnTheFlyServer.h b/Engine/Source/Editor/UnrealEd/Classes/CookOnTheSide/CookOnTheFlyServer.h
--- a/Engine/Source/Editor/UnrealEd/Classes/CookOnTheSide/CookOnTheFlyServer.h
+++ b/Engine/Source/Editor/UnrealEd/Classes/CookOnTheSide/CookOnTheFlyServer.h
@@ -836,6 +836,7 @@ public:
 	void OnObjectModified( UObject *ObjectMoving );
 	void OnObjectPropertyChanged(UObject* ObjectBeingModified, FPropertyChangedEvent& PropertyChangedEvent);
 	void OnObjectUpdated( UObject *Object );
+	void OnObjectSaved( UObject *ObjectSaved );
 
 	/**
 	 * Marks a package as dirty for cook
diff --git a/Engine/Source/Editor/UnrealEd/Classes/Editor/EditorEngine.h b/Engine/Source/Editor/UnrealEd/Classes/Editor/EditorEngine.h
--- a/Engine/Source/Editor/UnrealEd/Classes/Editor/EditorEngine.h
+++ b/Engine/Source/Editor/UnrealEd/Classes/Editor/EditorEngine.h
@@ -1477,12 +1477,12 @@ public:
 	/**
 	 * Can the editor do cook by the book in the editor process space
 	 */
-	virtual bool CanCookByTheBookInEditor() const { return false; }
+	virtual bool CanCookByTheBookInEditor(const FString& PlatformName ) const { return false; }
 
 	/**
 	 * Can the editor act as a cook on the fly server
 	 */
-	virtual bool CanCookOnTheFlyInEditor() const { return false; }
+	virtual bool CanCookOnTheFlyInEditor(const FString& PlatformName) const { return false; }
 
 	/**
 	 * Start cook by the book in the editor process space
diff --git a/Engine/Source/Editor/UnrealEd/Classes/Editor/UnrealEdEngine.h b/Engine/Source/Editor/UnrealEd/Classes/Editor/UnrealEdEngine.h
--- a/Engine/Source/Editor/UnrealEd/Classes/Editor/UnrealEdEngine.h
+++ b/Engine/Source/Editor/UnrealEd/Classes/Editor/UnrealEdEngine.h
@@ -549,12 +549,12 @@ class UNREALED_API UUnrealEdEngine : public UEditorEngine, public FNotifyHook
 	/**
 	 * Can the editor do cook by the book in the editor process space
 	 */
-	virtual bool CanCookByTheBookInEditor() const override;
+	virtual bool CanCookByTheBookInEditor(const FString& PlatformName) const override;
 
 	/**
 	 * Can the editor act as a cook on the fly server
 	 */
-	virtual bool CanCookOnTheFlyInEditor() const override;
+	virtual bool CanCookOnTheFlyInEditor(const FString& PlatformName) const override;
 
 	/**
 	 * Start cook by the book in the editor process space
diff --git a/Engine/Source/Editor/UnrealEd/Private/CookOnTheFlyServer.cpp b/Engine/Source/Editor/UnrealEd/Private/CookOnTheFlyServer.cpp
--- a/Engine/Source/Editor/UnrealEd/Private/CookOnTheFlyServer.cpp
+++ b/Engine/Source/Editor/UnrealEd/Private/CookOnTheFlyServer.cpp
@@ -1448,6 +1448,11 @@ void UCookOnTheFlyServer::OnObjectPropertyChanged(UObject* ObjectBeingModified,
 	OnObjectUpdated( ObjectBeingModified );
 }
 
+void UCookOnTheFlyServer::OnObjectSaved( UObject* ObjectSaved )
+{
+	OnObjectUpdated( ObjectSaved );
+}
+
 void UCookOnTheFlyServer::OnObjectUpdated( UObject *Object )
 {
 	// get the outer of the object
@@ -1468,10 +1473,10 @@ void UCookOnTheFlyServer::MarkPackageDirtyForCooker( UPackage *Package )
 		/*FString PackageFilename(GetPackageFilename(Package));
 		FPaths::MakeStandardFilename(PackageFilename);*/
 
-		ClearPackageFilenameCacheForPackage( Package );
+		FString PackageFilename = GetCachedStandardPackageFilename(Package);
 
+		ClearPackageFilenameCacheForPackage( Package );
 
-		FString PackageFilename = GetCachedStandardPackageFilename(Package);
 #if DEBUG_COOKONTHEFLY
 		UE_LOG(LogCook, Display, TEXT("Modification detected to package %s"), *PackageFilename);
 #endif
@@ -1485,12 +1490,12 @@ void UCookOnTheFlyServer::MarkPackageDirtyForCooker( UPackage *Package )
 			{
 				if ( IsCookByTheBookRunning() )
 				{
-				// if this package was previously cooked and we are doing a cook by the book 
-				// we need to recook this package before finishing cook by the book
-				CookRequests.EnqueueUnique( FFilePlatformRequest(PackageFFileName, CookedPlatforms) );
+					// if this package was previously cooked and we are doing a cook by the book 
+					// we need to recook this package before finishing cook by the book
+					CookRequests.EnqueueUnique( FFilePlatformRequest(PackageFFileName, CookedPlatforms) );
+				}
 			}
 		}
-		}
 		CookedPackages.RemoveFile( PackageFFileName );
 	}
 }
diff --git a/Engine/Source/Editor/UnrealEd/Private/PlayLevel.cpp b/Engine/Source/Editor/UnrealEd/Private/PlayLevel.cpp
--- a/Engine/Source/Editor/UnrealEd/Private/PlayLevel.cpp
+++ b/Engine/Source/Editor/UnrealEd/Private/PlayLevel.cpp
@@ -1454,19 +1454,32 @@ void UEditorEngine::PlayUsingLauncher()
 
 		// select the quickest cook mode based on which in editor cook mode is enabled
 		bool bIncrimentalCooking = true;
+		LauncherProfile->AddCookedPlatform(PlayUsingLauncherDeviceId.Left(PlayUsingLauncherDeviceId.Find(TEXT("@"))));
 		ELauncherProfileCookModes::Type CurrentLauncherCookMode = ELauncherProfileCookModes::ByTheBook;
-		if ( CanCookByTheBookInEditor() )
+		bool bCanCookByTheBookInEditor = true;
+		bool bCanCookOnTheFlyInEditor = true;
+		for ( const auto &PlatformName : LauncherProfile->GetCookedPlatforms() )
+		{
+			if ( CanCookByTheBookInEditor(PlatformName) == false )
+			{
+				bCanCookByTheBookInEditor = false;
+			}
+			if ( CanCookOnTheFlyInEditor(PlatformName)== false )
+			{
+				bCanCookOnTheFlyInEditor = false;
+			}
+		}
+		if ( bCanCookByTheBookInEditor )
 		{
 			CurrentLauncherCookMode = ELauncherProfileCookModes::ByTheBookInEditor;
 		}
-		if ( CanCookOnTheFlyInEditor() )
+		if ( bCanCookOnTheFlyInEditor )
 		{
 			CurrentLauncherCookMode = ELauncherProfileCookModes::OnTheFlyInEditor;
 			bIncrimentalCooking = false;
 		}
 		LauncherProfile->SetCookMode( CurrentLauncherCookMode );
 		LauncherProfile->SetIncrementalCooking(bIncrimentalCooking);
-		LauncherProfile->AddCookedPlatform(PlayUsingLauncherDeviceId.Left(PlayUsingLauncherDeviceId.Find(TEXT("@"))));
 		LauncherProfile->SetDeployedDeviceGroup(DeviceGroup);
 		LauncherProfile->SetIncrementalDeploying(bIncrimentalCooking);
 		LauncherProfile->SetEditorExe(FUnrealEdMisc::Get().GetExecutableForCommandlets());
@@ -1490,6 +1503,18 @@ void UEditorEngine::PlayUsingLauncher()
 			// Or if using by book in editor don't need to resave the package just cook it by the book 
 			FString MapName = EditorContext.World()->GetOutermost()->GetName();
 			MapNames.Add(MapName);
+
+
+			// Daniel: Only reason we actually need to save any packages is because if a new package is created it won't be on disk yet and CookOnTheFly will early out if the package doesn't exist (even though it could be in memory and not require loading at all)
+			//			future me can optimize this by either adding extra allowances to CookOnTheFlyServer code or only saving packages which doesn't exist if it becomes a problem
+			// if this returns false, it means we should stop what we're doing and return to the editor
+			bool bPromptUserToSave = true;
+			bool bSaveMapPackages = false;
+			bool bSaveContentPackages = true;
+			if (!FEditorFileUtils::SaveDirtyPackages(bPromptUserToSave, bSaveMapPackages, bSaveContentPackages))
+			{
+				return;
+			}
 		}
 		else
 		{
diff --git a/Engine/Source/Editor/UnrealEd/Private/UnrealEdEngine.cpp b/Engine/Source/Editor/UnrealEd/Private/UnrealEdEngine.cpp
--- a/Engine/Source/Editor/UnrealEd/Private/UnrealEdEngine.cpp
+++ b/Engine/Source/Editor/UnrealEd/Private/UnrealEdEngine.cpp
@@ -126,6 +126,7 @@ void UUnrealEdEngine::Init(IEngineLoop* InEngineLoop)
 
 		FCoreUObjectDelegates::OnObjectPropertyChanged.AddUObject(CookServer, &UCookOnTheFlyServer::OnObjectPropertyChanged);
 		FCoreUObjectDelegates::OnObjectModified.AddUObject(CookServer, &UCookOnTheFlyServer::OnObjectModified);
+		FCoreUObjectDelegates::OnObjectSaved.AddUObject( CookServer, &UCookOnTheFlyServer::OnObjectSaved );
 	}
 	else if ( !ExperimentalSettings->bDisableCookInEditor)
 	{
@@ -134,11 +135,43 @@ void UUnrealEdEngine::Init(IEngineLoop* InEngineLoop)
 
 		FCoreUObjectDelegates::OnObjectPropertyChanged.AddUObject(CookServer, &UCookOnTheFlyServer::OnObjectPropertyChanged);
 		FCoreUObjectDelegates::OnObjectModified.AddUObject(CookServer, &UCookOnTheFlyServer::OnObjectModified);
+		FCoreUObjectDelegates::OnObjectSaved.AddUObject( CookServer, &UCookOnTheFlyServer::OnObjectSaved );
 	}
 }
 
-bool UUnrealEdEngine::CanCookByTheBookInEditor() const 
-{ 
+bool CanCookForPlatformInThisProcess( const FString& PlatformName )
+{
+	////////////////////////////////////////
+	// hack remove this hack when we properly support changing the mobileHDR setting 
+	// check if our mobile hdr setting in memory is different from the one which is saved in the config file
+	
+	bool ConfigSetting = false;
+	if ( !GConfig->GetBool( TEXT("/Script/Engine.RendererSettings"), TEXT("r.MobileHDR"), ConfigSetting, GEngineIni) )
+	{
+		// if we can't get the config setting then don't risk it
+		return false;
+	}
+
+	// this was stolen from void IsMobileHDR()
+	static auto* MobileHDRCvar = IConsoleManager::Get().FindTConsoleVariableDataInt(TEXT("r.MobileHDR"));
+	const bool CurrentRSetting = MobileHDRCvar->GetValueOnAnyThread() == 1;
+
+	if ( CurrentRSetting != ConfigSetting )
+	{
+		return false;
+	}
+	////////////////////////////////////////
+	return true;
+}
+
+
+bool UUnrealEdEngine::CanCookByTheBookInEditor(const FString& PlatformName) const 
+{ 	
+	if ( CanCookForPlatformInThisProcess(PlatformName) == false )
+	{
+		return false;
+	}
+
 	if ( CookServer )
 	{
 		return CookServer->GetCookMode() == ECookMode::CookByTheBookFromTheEditor; 
@@ -146,8 +179,13 @@ bool UUnrealEdEngine::CanCookByTheBookInEditor() const
 	return false;
 }
 
-bool UUnrealEdEngine::CanCookOnTheFlyInEditor() const
+bool UUnrealEdEngine::CanCookOnTheFlyInEditor(const FString& PlatformName) const
 {
+	if ( CanCookForPlatformInThisProcess(PlatformName) == false )
+	{
+		return false;
+	}
+
 	if ( CookServer )
 	{
 		return CookServer->GetCookMode() == ECookMode::CookOnTheFlyFromTheEditor;
