Fixed crash in SSCSEditor after converting BSP to a static mesh.

diff --git a/Engine/Source/Editor/UnrealEd/Private/Editor.cpp b/Engine/Source/Editor/UnrealEd/Private/Editor.cpp
--- a/Engine/Source/Editor/UnrealEd/Private/Editor.cpp
+++ b/Engine/Source/Editor/UnrealEd/Private/Editor.cpp
@@ -5807,6 +5807,7 @@ void UEditorEngine::DoConvertActors( const TArray<AActor*>& ActorsToConvert, UCl
 		if( BrushList.Num() )
 		{
 			AActor* ConvertedBrushActor = ConvertBrushesToStaticMesh(InStaticMeshPackageName, BrushList, CachePivotLocation);
+			ConvertedActors.Add(ConvertedBrushActor);
 
 			// If only one brush is being converted, reattach it to whatever it was attached to before.
 			// Multiple brushes become impossible to reattach due to the single actor returned.
