[INTEGRATE] Change 2420410 by [REDACTED_PERSON][REDACTED_HANDLE].7 on 2015/01/27 10:53:04

	Added options to enable and disable auto creation and deletion of assets for the auto reimport manager

	Users now specify all the paths that they want monitoring. By default we monitor /Game/.

diff --git a/Engine/Source/Editor/UnrealEd/Classes/Settings/EditorLoadingSavingSettings.h b/Engine/Source/Editor/UnrealEd/Classes/Settings/EditorLoadingSavingSettings.h
--- a/Engine/Source/Editor/UnrealEd/Classes/Settings/EditorLoadingSavingSettings.h
+++ b/Engine/Source/Editor/UnrealEd/Classes/Settings/EditorLoadingSavingSettings.h
@@ -32,9 +32,14 @@ public:
 	bool bMonitorContentDirectories;
 
 	/** Directories being monitored for Auto Reimport */
-	UPROPERTY(EditAnywhere, config, Category=AutoReimport, meta=(DisplayName="Additional Directories to Monitor"))
+	UPROPERTY(EditAnywhere, config, AdvancedDisplay,Category=AutoReimport, meta=(DisplayName="Directories to Monitor", ToolTip="Lists every directory to monitor for content changes. Can be virtual package paths (eg /Game/ or /MyPlugin/), or absolute paths on disk."))
 	TArray<FString> AutoReimportDirectories;
 
+	UPROPERTY(EditAnywhere, config, AdvancedDisplay, Category=AutoReimport, meta=(DisplayName="Auto Create Assets", ToolTip="When enabled, newly added source content files will be automatically imported into new assets."))
+	bool bAutoCreateAssets;
+	UPROPERTY(EditAnywhere, config, AdvancedDisplay, Category=AutoReimport, meta=(DisplayName="Auto Delete Assets", ToolTip="When enabled, deleting a source content file will automatically prompt the deletion of any related assets."))
+	bool bAutoDeleteAssets;
+
 	/** Internal setting to control whether we should ask the user whether we should automatically delete source files when their assets are deleted */
 	UPROPERTY(config)
 	bool bDeleteSourceFilesWithAssets;
diff --git a/Engine/Source/Editor/UnrealEd/Private/AutoReimport/AutoReimportManager.cpp b/Engine/Source/Editor/UnrealEd/Private/AutoReimport/AutoReimportManager.cpp
--- a/Engine/Source/Editor/UnrealEd/Private/AutoReimport/AutoReimportManager.cpp
+++ b/Engine/Source/Editor/UnrealEd/Private/AutoReimport/AutoReimportManager.cpp
@@ -123,11 +123,8 @@ private:
 	/** Callback for when an editor user setting has changed */
 	void HandleLoadingSavingSettingChanged(FName PropertyName);
 
-	/** Set up monitors to all the mounted content directories */
-	void SetUpMountedContentDirectories();
-
-	/** Set up the user directories from the editor Loading/Saving settings */
-	void SetUpUserDirectories();
+	/** Set up monitors to all the monitored content directories */
+	void SetUpDirectoryMonitors();
 
 	/** Retrieve a semi-colon separated string of file extensions supported by all available editor import factories */
 	static FString GetAllFactoryExtensions();
@@ -176,11 +173,8 @@ private:
 	/** Cached string of extensions that are supported by all available factories */
 	FString SupportedExtensions;
 
-	/** Array of directory monitors for automatically managed mounted content paths */
-	TArray<FContentDirectoryMonitor> MountedDirectoryMonitors;
-
-	/** Additional external directories to monitor for changes specified by the user */
-	TArray<FContentDirectoryMonitor> UserDirectoryMonitors;
+	/** Array of objects that detect changes to directories */
+	TArray<FContentDirectoryMonitor> DirectoryMonitors;
 
 	/** A list of packages to save when we've added a bunch of assets */
 	TArray<UPackage*> PackagesToSave;
@@ -218,12 +212,9 @@ FAutoReimportManager::FAutoReimportManager()
 	// Only set this up for content directories if the user has this enabled
 	if (Settings->bMonitorContentDirectories)
 	{
-		SetUpMountedContentDirectories();
+		SetUpDirectoryMonitors();
 	}
 
-	// We always monitor user directories if they are set
-	SetUpUserDirectories();
-
 	StateMachine.Add(ECurrentState::Idle,					EStateMachineNode::CallOnce,	[this](const FTimeLimit& T){ return this->Idle(T);					});
 	StateMachine.Add(ECurrentState::PendingResponse,		EStateMachineNode::CallOnce,	[this](const FTimeLimit& T){ return this->PendingResponse();		});
 	StateMachine.Add(ECurrentState::ProcessAdditions,		EStateMachineNode::CallMany,	[this](const FTimeLimit& T){ return this->ProcessAdditions(T);		});
@@ -248,29 +239,31 @@ void FAutoReimportManager::Destroy()
 	}
 
 	// Force a save of all the caches
-	MountedDirectoryMonitors.Empty();
-	UserDirectoryMonitors.Empty();
+	DirectoryMonitors.Empty();
 }
 
 int32 FAutoReimportManager::GetNumUnprocessedChanges() const
 {
-	auto Pred = [](const FContentDirectoryMonitor& Monitor, int32 Total){ return Total + Monitor.GetNumUnprocessedChanges(); };
-	return Utils::Reduce(MountedDirectoryMonitors, Pred, 0) + Utils::Reduce(UserDirectoryMonitors, Pred, 0);
+	return Utils::Reduce(DirectoryMonitors, [](const FContentDirectoryMonitor& Monitor, int32 Total){
+		return Total + Monitor.GetNumUnprocessedChanges();
+	}, 0);
 }
 
 FText FAutoReimportManager::GetProgressText() const
 {
 	FFormatOrderedArguments Args;
 	{
-		auto Pred = [](const FContentDirectoryMonitor& Monitor, int32 Total){ return Total + Monitor.GetWorkProgress(); };
-		const int32 Progress = Utils::Reduce(MountedDirectoryMonitors, Pred, 0) + Utils::Reduce(UserDirectoryMonitors, Pred, 0);
+		const int32 Progress = Utils::Reduce(DirectoryMonitors, [](const FContentDirectoryMonitor& Monitor, int32 Total){
+			return Total + Monitor.GetWorkProgress();
+		}, 0);
 
 		Args.Add(Progress);
 	}
 
 	{
-		auto Pred = [](const FContentDirectoryMonitor& Monitor, int32 Total){ return Total + Monitor.GetTotalWork(); };
-		const int32 Total = Utils::Reduce(MountedDirectoryMonitors, Pred, 0) + Utils::Reduce(UserDirectoryMonitors, Pred, 0);
+		const int32 Total = Utils::Reduce(DirectoryMonitors, [](const FContentDirectoryMonitor& Monitor, int32 Total){
+			return Total + Monitor.GetTotalWork();
+		}, 0);
 
 		Args.Add(Total);
 	}
@@ -314,13 +307,7 @@ TOptional<ECurrentState> FAutoReimportManager::ProcessAdditions(const FTimeLimit
 		Pair.Value.Sort([](const UFactory& A, const UFactory& B) { return A.ImportPriority > B.ImportPriority; });
 	}
 
-	for (auto& Monitor : MountedDirectoryMonitors)
-	{
-		Monitor.ProcessAdditions(PackagesToSave, TimeLimit, Factories, *FeedbackContextOverride);
-		yield TOptional<ECurrentState>();
-	}
-
-	for (auto& Monitor : UserDirectoryMonitors)
+	for (auto& Monitor : DirectoryMonitors)
 	{
 		Monitor.ProcessAdditions(PackagesToSave, TimeLimit, Factories, *FeedbackContextOverride);
 		yield TOptional<ECurrentState>();
@@ -358,13 +345,7 @@ TOptional<ECurrentState> FAutoReimportManager::ProcessModifications(const FTimeL
 
 	const IAssetRegistry& Registry = FModuleManager::LoadModuleChecked<FAssetRegistryModule>("AssetRegistry").Get();
 
-	for (auto& Monitor : MountedDirectoryMonitors)
-	{
-		Monitor.ProcessModifications(Registry, TimeLimit, *FeedbackContextOverride);
-		yield TOptional<ECurrentState>();
-	}
-
-	for (auto& Monitor : UserDirectoryMonitors)
+	for (auto& Monitor : DirectoryMonitors)
 	{
 		Monitor.ProcessModifications(Registry, TimeLimit, *FeedbackContextOverride);
 		yield TOptional<ECurrentState>();
@@ -381,11 +362,7 @@ TOptional<ECurrentState> FAutoReimportManager::ProcessDeletions()
 
 	TArray<FAssetData> AssetsToDelete;
 
-	for (auto& Monitor : MountedDirectoryMonitors)
-	{
-		Monitor.ExtractAssetsToDelete(Registry, AssetsToDelete);
-	}
-	for (auto& Monitor : UserDirectoryMonitors)
+	for (auto& Monitor : DirectoryMonitors)
 	{
 		Monitor.ExtractAssetsToDelete(Registry, AssetsToDelete);
 	}
@@ -407,12 +384,7 @@ TOptional<ECurrentState> FAutoReimportManager::PendingResponse()
 
 TOptional<ECurrentState> FAutoReimportManager::Idle(const FTimeLimit& TimeLimit)
 {
-	for (auto& Monitor : MountedDirectoryMonitors)
-	{
-		Monitor.Tick(TimeLimit);
-		yield TOptional<ECurrentState>();
-	}
-	for (auto& Monitor : UserDirectoryMonitors)
+	for (auto& Monitor : DirectoryMonitors)
 	{
 		Monitor.Tick(TimeLimit);
 		yield TOptional<ECurrentState>();
@@ -431,21 +403,25 @@ TOptional<ECurrentState> FAutoReimportManager::Idle(const FTimeLimit& TimeLimit)
 		{
 			if (StartProcessingDelay.Exceeded())
 			{
-				// Create a new feedback context override
-				FeedbackContextOverride = MakeShareable(new FReimportFeedbackContext);
-				FeedbackContextOverride->Initialize(SNew(SReimportFeedback, GetProgressText()));
-
 				// Deal with changes to the file system second
-				for (auto& Monitor : MountedDirectoryMonitors)
+				for (auto& Monitor : DirectoryMonitors)
 				{
 					Monitor.StartProcessing();
 				}
-				for (auto& Monitor : UserDirectoryMonitors)
+
+				// Check that we actually have anything to do before kicking off a process
+				const int32 WorkTotal = Utils::Reduce(DirectoryMonitors, [](const FContentDirectoryMonitor& Monitor, int32 Total){
+					return Total + Monitor.GetTotalWork();
+				}, 0);
+
+				if (WorkTotal > 0)
 				{
-					Monitor.StartProcessing();
-				}
+					// Create a new feedback context override
+					FeedbackContextOverride = MakeShareable(new FReimportFeedbackContext);
+					FeedbackContextOverride->Initialize(SNew(SReimportFeedback, GetProgressText()));
 
-				return ECurrentState::ProcessAdditions;
+					return ECurrentState::ProcessAdditions;
+				}
 			}
 		}
 	}
@@ -488,98 +464,99 @@ void FAutoReimportManager::AddReferencedObjects(FReferenceCollector& Collector)
 
 void FAutoReimportManager::HandleLoadingSavingSettingChanged(FName PropertyName)
 {
-	if (PropertyName == GET_MEMBER_NAME_CHECKED(UEditorLoadingSavingSettings, AutoReimportDirectories))
-	{
-		SetUpUserDirectories();
-	}
-	else if (PropertyName == GET_MEMBER_NAME_CHECKED(UEditorLoadingSavingSettings, bMonitorContentDirectories))
+	if (PropertyName == GET_MEMBER_NAME_CHECKED(UEditorLoadingSavingSettings, bMonitorContentDirectories) ||
+		PropertyName == GET_MEMBER_NAME_CHECKED(UEditorLoadingSavingSettings, AutoReimportDirectories))
 	{
 		const auto* Settings = GetDefault<UEditorLoadingSavingSettings>();
 		if (Settings->bMonitorContentDirectories)
 		{
-			// Should already be empty, but just to be sure
-			MountedDirectoryMonitors.Empty();
-			SetUpMountedContentDirectories();
+			DirectoryMonitors.Empty();
+			SetUpDirectoryMonitors();
 		}
 		else
 		{
 			// Destroy all the existing monitors, including their file caches
-			for (auto& Monitor : MountedDirectoryMonitors)
+			for (auto& Monitor : DirectoryMonitors)
 			{
 				Monitor.Destroy();
 			}
-			MountedDirectoryMonitors.Empty();
+			DirectoryMonitors.Empty();
 		}
 	}
 }
 
 void FAutoReimportManager::OnContentPathMounted(const FString& InAssetPath, const FString& FileSystemPath)
 {
-	MountedDirectoryMonitors.Emplace(FPaths::ConvertRelativePathToFull(FileSystemPath), SupportedExtensions, InAssetPath);
+	// @todo: check the paths for mounted directories
+	//DirectoryMonitors.Emplace(FPaths::ConvertRelativePathToFull(FileSystemPath), SupportedExtensions, InAssetPath);
 }
 
 void FAutoReimportManager::OnContentPathDismounted(const FString& InAssetPath, const FString& FileSystemPath)
 {
 	const auto FullDismountedPath = FPaths::ConvertRelativePathToFull(FileSystemPath);
 
-	MountedDirectoryMonitors.RemoveAll([&](const FContentDirectoryMonitor& Monitor){
+	DirectoryMonitors.RemoveAll([&](const FContentDirectoryMonitor& Monitor){
 		return Monitor.GetDirectory() == FullDismountedPath;
 	});
 }
 
-void FAutoReimportManager::SetUpMountedContentDirectories()
+void FAutoReimportManager::SetUpDirectoryMonitors()
 {
-	TArray<FString> ContentPaths;
-	FPackageName::QueryRootContentPaths(ContentPaths);
-	for (const auto& ContentPath : ContentPaths)
-	{
-		const FString Directory = FPaths::ConvertRelativePathToFull(FPackageName::LongPackageNameToFilename(ContentPath));
+	struct FPathAndMountPoint { FString Path; FString MountPoint; };
+	TArray<FPathAndMountPoint> CollapsedDirectories;
 
-		// We ignore the engine content directory for now
-		if (Directory != FPaths::ConvertRelativePathToFull(FPaths::EngineContentDir()))
+	// Build an array of directory / mounted path so we can match up absolute paths that the user has typed in
+	TArray<TPair<FString, FString>> MountedPaths;
+	{
+		TArray<FString> RootContentPaths;
+		FPackageName::QueryRootContentPaths( RootContentPaths );
+		for (FString& RootPath : RootContentPaths)
 		{
-			MountedDirectoryMonitors.Emplace(Directory, SupportedExtensions, ContentPath);
+			FString ContentFolder = FPaths::ConvertRelativePathToFull(FPackageName::LongPackageNameToFilename(RootPath));
+			MountedPaths.Add( TPairInitializer<FString, FString>(MoveTemp(ContentFolder), MoveTemp(RootPath)) );
 		}
 	}
-}
 
-void FAutoReimportManager::SetUpUserDirectories()
-{
-	TArray<bool> PersistedDirectories;
-	PersistedDirectories.AddZeroed(UserDirectoryMonitors.Num());
 
-	// Also set up user-specified directories
-	TArray<FString> NewDirectories = GetDefault<UEditorLoadingSavingSettings>()->AutoReimportDirectories;
-
-	for (int32 Index = 0; Index < NewDirectories.Num(); ++Index)
+	auto& FileManager = IFileManager::Get();
+	for (const FString& Path : GetDefault<UEditorLoadingSavingSettings>()->AutoReimportDirectories)
 	{
-		NewDirectories[Index] = FPaths::ConvertRelativePathToFull(NewDirectories[Index]);
-		const int32 Existing = UserDirectoryMonitors.IndexOfByPredicate([&](const FContentDirectoryMonitor& InMonitor){ return InMonitor.GetDirectory() == NewDirectories[Index]; });
-		if (Existing != INDEX_NONE)
+		FPathAndMountPoint ThisPath;
+
+		const FName MountPoint = FPackageName::GetPackageMountPoint(Path);
+		if (!MountPoint.IsNone())
 		{
-			PersistedDirectories[Existing] = true;
+			ThisPath.Path = FPaths::ConvertRelativePathToFull(FPackageName::LongPackageNameToFilename(Path / TEXT(""))) / "";
 		}
 		else
 		{
-			// Add a new one
-			UserDirectoryMonitors.Emplace(NewDirectories[Index], SupportedExtensions);
-			PersistedDirectories.Add(true);
+			ThisPath.Path = FPaths::ConvertRelativePathToFull(Path) / "";
 		}
-	}
 
-	// Delete anything that is no longer needed (we delete the cache file altogether)
-	for (int32 Index = PersistedDirectories.Num() - 1; Index >= 0; --Index)
-	{
-		if (!PersistedDirectories[Index])
+		if (FileManager.DirectoryExists(*ThisPath.Path))
 		{
-			UserDirectoryMonitors[Index].Destroy();
-			UserDirectoryMonitors.RemoveAt(Index, 1, false);
+			// Set the mounted path if necessary
+			auto* Pair = MountedPaths.FindByPredicate([&](const TPair<FString, FString>& Pair){
+				return ThisPath.Path.StartsWith(Pair.Key);
+			});
+
+			if (Pair)
+			{
+				ThisPath.MountPoint = Pair->Value / ThisPath.Path.RightChop(Pair->Key.Len());
+			}
+
+			CollapsedDirectories.Add(ThisPath);
 		}
 	}
 
-	if (UserDirectoryMonitors.GetSlack() > UserDirectoryMonitors.Max() / 2)
+	Utils::RemoveDuplicates(CollapsedDirectories, [](const FPathAndMountPoint& A, const FPathAndMountPoint& B){
+		return A.Path.StartsWith(B.Path, ESearchCase::IgnoreCase) ||
+			   B.Path.StartsWith(A.Path, ESearchCase::IgnoreCase);
+	});
+
+	for (const auto& Dir : CollapsedDirectories)
 	{
-		UserDirectoryMonitors.Shrink();
+		DirectoryMonitors.Emplace(Dir.Path, SupportedExtensions, Dir.MountPoint);
 	}
 }
 
diff --git a/Engine/Source/Editor/UnrealEd/Private/AutoReimport/AutoReimportUtilities.cpp b/Engine/Source/Editor/UnrealEd/Private/AutoReimport/AutoReimportUtilities.cpp
--- a/Engine/Source/Editor/UnrealEd/Private/AutoReimport/AutoReimportUtilities.cpp
+++ b/Engine/Source/Editor/UnrealEd/Private/AutoReimport/AutoReimportUtilities.cpp
@@ -29,7 +29,8 @@ namespace Utils
 				{
 					// Either relative to the asset itself, or relative to the base path, or absolute.
 					// Would ideally use FReimportManager::ResolveImportFilename here, but we don't want to ask the file system whether the file exists.
-					if (AbsoluteFilename == FPaths::ConvertRelativePathToFull(FPackageName::LongPackageNameToFilename(Asset.PackagePath.ToString()) / Pair.Value) ||
+					auto RelativeToPackagePath = FPackageName::LongPackageNameToFilename(Asset.PackagePath.ToString() / TEXT("")) / Pair.Value;
+					if (AbsoluteFilename == FPaths::ConvertRelativePathToFull(RelativeToPackagePath) ||
 						AbsoluteFilename == FPaths::ConvertRelativePathToFull(Pair.Value))
 					{
 						return false;
diff --git a/Engine/Source/Editor/UnrealEd/Private/AutoReimport/AutoReimportUtilities.h b/Engine/Source/Editor/UnrealEd/Private/AutoReimport/AutoReimportUtilities.h
--- a/Engine/Source/Editor/UnrealEd/Private/AutoReimport/AutoReimportUtilities.h
+++ b/Engine/Source/Editor/UnrealEd/Private/AutoReimport/AutoReimportUtilities.h
@@ -97,6 +97,43 @@ namespace Utils
 		return Accumulator;
 	}
 
+	/** Helper template function to remove duplicates from an array, given some predicate */
+	template<typename T, typename P>
+	void RemoveDuplicates(TArray<T>& Array, P Predicate)
+	{
+		if (Array.Num() == 0)
+			return;
+
+		const int32 NumElements = Array.Num();
+		int32 Write = 0;
+		for (int32 Read = 0; Read < NumElements; ++Read)
+		{
+			for (int32 DuplRead = Read + 1; DuplRead < NumElements; ++DuplRead)
+			{
+				if (Predicate(Array[Read], Array[DuplRead]))
+				{
+					goto next;
+				}
+			}
+
+			if (Write != Read)
+			{
+				Swap(Array[Write++], Array[Read]);
+			}
+			else
+			{
+				Write++;
+			}
+	next:
+			;
+		}
+		const int32 NumDuplicates = NumElements - Write;
+		if (NumDuplicates != 0)
+		{
+			Array.RemoveAt(Write, NumDuplicates, false);
+		}
+	}
+
 	/** Find a list of assets that were once imported from the specified filename */
 	TArray<FAssetData> FindAssetsPertainingToFile(const IAssetRegistry& Registry, const FString& AbsoluteFilename);
 
diff --git a/Engine/Source/Editor/UnrealEd/Private/AutoReimport/ContentDirectoryMonitor.cpp b/Engine/Source/Editor/UnrealEd/Private/AutoReimport/ContentDirectoryMonitor.cpp
--- a/Engine/Source/Editor/UnrealEd/Private/AutoReimport/ContentDirectoryMonitor.cpp
+++ b/Engine/Source/Editor/UnrealEd/Private/AutoReimport/ContentDirectoryMonitor.cpp
@@ -60,13 +60,37 @@ void FContentDirectoryMonitor::StartProcessing()
 	auto OutstandingChanges = Cache.GetOutstandingChanges();
 	if (OutstandingChanges.Num() != 0)
 	{
+		const auto* Settings = GetDefault<UEditorLoadingSavingSettings>();
+
 		for (auto& Transaction : OutstandingChanges)
 		{
 			switch(Transaction.Action)
 			{
-				case FFileChangeData::FCA_Added:		AddedFiles.Emplace(MoveTemp(Transaction));		break;
-				case FFileChangeData::FCA_Modified:		ModifiedFiles.Emplace(MoveTemp(Transaction));	break;
-				case FFileChangeData::FCA_Removed:		DeletedFiles.Emplace(MoveTemp(Transaction));	break;
+				case FFileChangeData::FCA_Added:
+					if (Settings->bAutoCreateAssets)
+					{
+						AddedFiles.Emplace(MoveTemp(Transaction));
+					}
+					else
+					{
+						Cache.CompleteTransaction(MoveTemp(Transaction));
+					}
+					break;
+					
+				case FFileChangeData::FCA_Modified:
+					ModifiedFiles.Emplace(MoveTemp(Transaction));
+					break;
+
+				case FFileChangeData::FCA_Removed:
+					if (Settings->bAutoDeleteAssets)
+					{
+						DeletedFiles.Emplace(MoveTemp(Transaction));
+					}
+					else
+					{
+						Cache.CompleteTransaction(MoveTemp(Transaction));
+					}
+					break;
 			}
 		}
 
diff --git a/Engine/Source/Editor/UnrealEd/Private/AutoReimport/FileCache.cpp b/Engine/Source/Editor/UnrealEd/Private/AutoReimport/FileCache.cpp
--- a/Engine/Source/Editor/UnrealEd/Private/AutoReimport/FileCache.cpp
+++ b/Engine/Source/Editor/UnrealEd/Private/AutoReimport/FileCache.cpp
@@ -5,50 +5,11 @@
 #include "FileCache.h"
 #include "IDirectoryWatcher.h"
 #include "DirectoryWatcherModule.h"
+#include "AutoReimportUtilities.h"
 
 const FGuid FFileCacheCustomVersion::Key(0x8E7DDCB3, 0x80DA47BB, 0x9FD346A2, 0x93984DF6);
 FCustomVersionRegistration GRegisterFileCacheVersion(FFileCacheCustomVersion::Key, FFileCacheCustomVersion::Latest, TEXT("FileCacheVersion"));
 
-namespace
-{
-	/** Helper template function to remove duplicates from an array, given some predicate */
-	template<typename T, typename P>
-	void RemoveDuplicates(TArray<T>& Array, P Predicate)
-	{
-		if (Array.Num() == 0)
-			return;
-
-		const int32 NumElements = Array.Num();
-		int32 Write = 0;
-		for (int32 Read = 0; Read < NumElements; ++Read)
-		{
-			for (int32 DuplRead = Read + 1; DuplRead < NumElements; ++DuplRead)
-			{
-				if (Predicate(Array[Read], Array[DuplRead]))
-				{
-					goto next;
-				}
-			}
-
-			if (Write != Read)
-			{
-				Swap(Array[Write++], Array[Read]);
-			}
-			else
-			{
-				Write++;
-			}
-	next:
-			;
-		}
-		const int32 NumDuplicates = NumElements - Write;
-		if (NumDuplicates != 0)
-		{
-			Array.RemoveAt(Write, NumDuplicates, false);
-		}
-	}
-}
-
 FAsyncDirectoryReader::FAsyncDirectoryReader(const FString& InDirectory, EPathType InPathType)
 	: RootPath(InDirectory), PathType(InPathType)
 {
@@ -419,7 +380,7 @@ void FFileCache::OnDirectoryChanged(const TArray<FFileChangeData>& FileChanges)
 		}
 	}
 
-	RemoveDuplicates(OutstandingChanges, [](const FUpdateCacheTransaction& A, const FUpdateCacheTransaction& B){
+	Utils::RemoveDuplicates(OutstandingChanges, [](const FUpdateCacheTransaction& A, const FUpdateCacheTransaction& B){
 		return A.Action == B.Action && A.Filename == B.Filename;
 	});
 }
diff --git a/Engine/Source/Editor/UnrealEd/Private/Settings/SettingsClasses.cpp b/Engine/Source/Editor/UnrealEd/Private/Settings/SettingsClasses.cpp
--- a/Engine/Source/Editor/UnrealEd/Private/Settings/SettingsClasses.cpp
+++ b/Engine/Source/Editor/UnrealEd/Private/Settings/SettingsClasses.cpp
@@ -94,9 +94,12 @@ void UEditorExperimentalSettings::PostEditChangeProperty( struct FPropertyChange
 UEditorLoadingSavingSettings::UEditorLoadingSavingSettings( const FObjectInitializer& ObjectInitializer )
 	: Super(ObjectInitializer)
 	, bMonitorContentDirectories(false)
+	, bAutoCreateAssets(true)
+	, bAutoDeleteAssets(true)
 	, bDeleteSourceFilesWithAssets(true)
 {
 	TextDiffToolPath.FilePath = TEXT("P4Merge.exe");
+	AutoReimportDirectories.Add("/Game/");
 }
 
 
