[INTEGRATE] Change 2420481 by [REDACTED_PERSON][REDACTED_HANDLE] on 2015/01/27 11:33:10

	Ensure that variables get created for components when creating a blueprint class from an instance

diff --git a/Engine/Source/Editor/UnrealEd/Private/Kismet2/Kismet2.cpp b/Engine/Source/Editor/UnrealEd/Private/Kismet2/Kismet2.cpp
--- a/Engine/Source/Editor/UnrealEd/Private/Kismet2/Kismet2.cpp
+++ b/Engine/Source/Editor/UnrealEd/Private/Kismet2/Kismet2.cpp
@@ -1069,6 +1069,8 @@ UBlueprint* FKismetEditorUtilities::CreateBlueprintFromActor(const FName Bluepri
 				}
 			}
 
+			FKismetEditorUtilities::CompileBlueprint(NewBlueprint);
+
 			if(bReplaceActor)
 			{
 				TArray<AActor*> Actors;
