[INTEGRATE] Change 2420495 by [REDACTED_PERSON][REDACTED_HANDLE].7 on 2015/01/27 11:39:27

	Auto reimport directory watchers now update when the mounted directory paths change

diff --git a/Engine/Source/Editor/UnrealEd/Private/AutoReimport/AutoReimportManager.cpp b/Engine/Source/Editor/UnrealEd/Private/AutoReimport/AutoReimportManager.cpp
--- a/Engine/Source/Editor/UnrealEd/Private/AutoReimport/AutoReimportManager.cpp
+++ b/Engine/Source/Editor/UnrealEd/Private/AutoReimport/AutoReimportManager.cpp
@@ -112,11 +112,8 @@ private:
 
 private:
 
-	/** Called when a new asset path has been mounted */
-	void OnContentPathMounted(const FString& InAssetPath, const FString& FileSystemPath);
-
-	/** Called when a new asset path has been dismounted */
-	void OnContentPathDismounted(const FString& InAssetPath, const FString& FileSystemPath);
+	/** Called when a new asset path has been mounted or unmounted */
+	void OnContentPathChanged(const FString& InAssetPath, const FString& FileSystemPath);
 
 private:
 
@@ -200,8 +197,8 @@ FAutoReimportManager::FAutoReimportManager()
 	auto* Settings = GetMutableDefault<UEditorLoadingSavingSettings>();
 	Settings->OnSettingChanged().AddRaw(this, &FAutoReimportManager::HandleLoadingSavingSettingChanged);
 
-	FPackageName::OnContentPathMounted().AddRaw(this, &FAutoReimportManager::OnContentPathMounted);
-	FPackageName::OnContentPathDismounted().AddRaw(this, &FAutoReimportManager::OnContentPathDismounted);
+	FPackageName::OnContentPathMounted().AddRaw(this, &FAutoReimportManager::OnContentPathChanged);
+	FPackageName::OnContentPathDismounted().AddRaw(this, &FAutoReimportManager::OnContentPathChanged);
 
 	FMessageLogModule& MessageLogModule = FModuleManager::LoadModuleChecked<FMessageLogModule>("MessageLog");
 	if (!MessageLogModule.IsRegisteredLogListing("AssetReimport"))
@@ -485,19 +482,14 @@ void FAutoReimportManager::HandleLoadingSavingSettingChanged(FName PropertyName)
 	}
 }
 
-void FAutoReimportManager::OnContentPathMounted(const FString& InAssetPath, const FString& FileSystemPath)
+void FAutoReimportManager::OnContentPathChanged(const FString& InAssetPath, const FString& FileSystemPath)
 {
-	// @todo: check the paths for mounted directories
-	//DirectoryMonitors.Emplace(FPaths::ConvertRelativePathToFull(FileSystemPath), SupportedExtensions, InAssetPath);
-}
-
-void FAutoReimportManager::OnContentPathDismounted(const FString& InAssetPath, const FString& FileSystemPath)
-{
-	const auto FullDismountedPath = FPaths::ConvertRelativePathToFull(FileSystemPath);
-
-	DirectoryMonitors.RemoveAll([&](const FContentDirectoryMonitor& Monitor){
-		return Monitor.GetDirectory() == FullDismountedPath;
-	});
+	const auto* Settings = GetDefault<UEditorLoadingSavingSettings>();
+	if (Settings->bMonitorContentDirectories)
+	{
+		DirectoryMonitors.Empty();
+		SetUpDirectoryMonitors();
+	}
 }
 
 void FAutoReimportManager::SetUpDirectoryMonitors()
