[INTEGRATE] Change 2420903 by [REDACTED_PERSON][REDACTED_HANDLE] on 2015/01/27 15:29:19

	Added "..." to new "Create Component" buttons in Add Component dropdown

diff --git a/Engine/Source/Editor/UnrealEd/Private/SComponentClassCombo.cpp b/Engine/Source/Editor/UnrealEd/Private/SComponentClassCombo.cpp
--- a/Engine/Source/Editor/UnrealEd/Private/SComponentClassCombo.cpp
+++ b/Engine/Source/Editor/UnrealEd/Private/SComponentClassCombo.cpp
@@ -414,11 +414,11 @@ FText SComponentClassCombo::GetFriendlyComponentName(FComponentClassComboEntryPt
 
 	if( Entry->GetComponentCreateAction() == EComponentCreateAction::CreateNewCPPClass )
 	{
-		FriendlyComponentName = LOCTEXT("NewCPPBehaviorFriendlyName", "New C++ Component").ToString();
+		FriendlyComponentName = LOCTEXT("NewCPPComponentFriendlyName", "New C++ Component...").ToString();
 	}
 	else if (Entry->GetComponentCreateAction() == EComponentCreateAction::CreateNewBlueprintClass )
 	{
-		FriendlyComponentName = LOCTEXT("NewCPPBehaviorFriendlyName", "New Blueprint Component").ToString();
+		FriendlyComponentName = LOCTEXT("NewBlueprintComponentFriendlyName", "New Blueprint Component...").ToString();
 	}
 	else
 	{
