[INTEGRATE] Change 2420939 by [REDACTED_PERSON][REDACTED_HANDLE] on 2015/01/27 15:50:03

	Enable snap to floor on components

diff --git a/Engine/Source/Editor/LevelEditor/Private/LevelEditor.cpp b/Engine/Source/Editor/LevelEditor/Private/LevelEditor.cpp
--- a/Engine/Source/Editor/LevelEditor/Private/LevelEditor.cpp
+++ b/Engine/Source/Editor/LevelEditor/Private/LevelEditor.cpp
@@ -723,7 +723,7 @@ void FLevelEditorModule::BindGlobalLevelEditorCommands()
 	bool bUsePivot = false;
 	ActionList.MapAction(
 		Commands.SnapToFloor,
-		FExecuteAction::CreateStatic(&FLevelEditorActionCallbacks::SnapActorToFloor_Clicked, bAlign, bUseLineTrace, bUseBounds, bUsePivot),
+		FExecuteAction::CreateStatic(&FLevelEditorActionCallbacks::SnapToFloor_Clicked, bAlign, bUseLineTrace, bUseBounds, bUsePivot),
 		FCanExecuteAction::CreateStatic(&FLevelEditorActionCallbacks::ActorSelected_CanExecute)
 		);
 
@@ -733,7 +733,7 @@ void FLevelEditorModule::BindGlobalLevelEditorCommands()
 	bUsePivot = false;
 	ActionList.MapAction(
 		Commands.AlignToFloor,
-		FExecuteAction::CreateStatic(&FLevelEditorActionCallbacks::SnapActorToFloor_Clicked, bAlign, bUseLineTrace, bUseBounds, bUsePivot),
+		FExecuteAction::CreateStatic(&FLevelEditorActionCallbacks::SnapToFloor_Clicked, bAlign, bUseLineTrace, bUseBounds, bUsePivot),
 		FCanExecuteAction::CreateStatic(&FLevelEditorActionCallbacks::ActorSelected_CanExecute)
 		);
 
@@ -743,7 +743,7 @@ void FLevelEditorModule::BindGlobalLevelEditorCommands()
 	bUsePivot = true;
 	ActionList.MapAction(
 		Commands.SnapPivotToFloor,
-		FExecuteAction::CreateStatic(&FLevelEditorActionCallbacks::SnapActorToFloor_Clicked, bAlign, bUseLineTrace, bUseBounds, bUsePivot),
+		FExecuteAction::CreateStatic(&FLevelEditorActionCallbacks::SnapToFloor_Clicked, bAlign, bUseLineTrace, bUseBounds, bUsePivot),
 		FCanExecuteAction::CreateStatic(&FLevelEditorActionCallbacks::ActorSelected_CanExecute)
 		);
 
@@ -753,7 +753,7 @@ void FLevelEditorModule::BindGlobalLevelEditorCommands()
 	bUsePivot = true;
 	ActionList.MapAction(
 		Commands.AlignPivotToFloor,
-		FExecuteAction::CreateStatic(&FLevelEditorActionCallbacks::SnapActorToFloor_Clicked, bAlign, bUseLineTrace, bUseBounds, bUsePivot),
+		FExecuteAction::CreateStatic(&FLevelEditorActionCallbacks::SnapToFloor_Clicked, bAlign, bUseLineTrace, bUseBounds, bUsePivot),
 		FCanExecuteAction::CreateStatic(&FLevelEditorActionCallbacks::ActorSelected_CanExecute)
 		);
 
@@ -763,7 +763,7 @@ void FLevelEditorModule::BindGlobalLevelEditorCommands()
 	bUsePivot = false;
 	ActionList.MapAction(
 		Commands.SnapBottomCenterBoundsToFloor,
-		FExecuteAction::CreateStatic(&FLevelEditorActionCallbacks::SnapActorToFloor_Clicked, bAlign, bUseLineTrace, bUseBounds, bUsePivot),
+		FExecuteAction::CreateStatic(&FLevelEditorActionCallbacks::SnapToFloor_Clicked, bAlign, bUseLineTrace, bUseBounds, bUsePivot),
 		FCanExecuteAction::CreateStatic(&FLevelEditorActionCallbacks::ActorSelected_CanExecute)
 		);
 
@@ -773,7 +773,7 @@ void FLevelEditorModule::BindGlobalLevelEditorCommands()
 	bUsePivot = false;
 	ActionList.MapAction(
 		Commands.AlignBottomCenterBoundsToFloor,
-		FExecuteAction::CreateStatic(&FLevelEditorActionCallbacks::SnapActorToFloor_Clicked, bAlign, bUseLineTrace, bUseBounds, bUsePivot),
+		FExecuteAction::CreateStatic(&FLevelEditorActionCallbacks::SnapToFloor_Clicked, bAlign, bUseLineTrace, bUseBounds, bUsePivot),
 		FCanExecuteAction::CreateStatic(&FLevelEditorActionCallbacks::ActorSelected_CanExecute)
 		);
 
diff --git a/Engine/Source/Editor/LevelEditor/Private/LevelEditorActions.cpp b/Engine/Source/Editor/LevelEditor/Private/LevelEditorActions.cpp
--- a/Engine/Source/Editor/LevelEditor/Private/LevelEditorActions.cpp
+++ b/Engine/Source/Editor/LevelEditor/Private/LevelEditorActions.cpp
@@ -2440,10 +2440,10 @@ void FLevelEditorActionCallbacks::MoveActorTo_Clicked( const bool InAlign, const
 	GEditor->RebuildAlteredBSP(); // Update the Bsp of any levels containing a modified brush
 }
 
-void FLevelEditorActionCallbacks::SnapActorToFloor_Clicked( bool InAlign, bool InUseLineTrace, bool InUseBounds, bool InUsePivot )
+void FLevelEditorActionCallbacks::SnapToFloor_Clicked( bool InAlign, bool InUseLineTrace, bool InUseBounds, bool InUsePivot )
 {
 	const FScopedTransaction Transaction( NSLOCTEXT("UnrealEd", "SnapActorsToFloor", "Snap Actors To Floor") );
-	SnapActorTo_Clicked( InAlign, InUseLineTrace, InUseBounds, InUsePivot );
+	SnapTo_Clicked( InAlign, InUseLineTrace, InUseBounds, InUsePivot );
 }
 
 void FLevelEditorActionCallbacks::SnapActorToActor_Clicked( bool InAlign, bool InUseLineTrace, bool InUseBounds, bool InUsePivot )
@@ -2452,41 +2452,73 @@ void FLevelEditorActionCallbacks::SnapActorToActor_Clicked( bool InAlign, bool I
 	if( Actor )
 	{
 		const FScopedTransaction Transaction( NSLOCTEXT("UnrealEd", "SnapActorsToActor", "Snap Actors To Actor") );
-		SnapActorTo_Clicked( InAlign, InUseLineTrace, InUseBounds, InUsePivot, Actor );
+		SnapTo_Clicked( InAlign, InUseLineTrace, InUseBounds, InUsePivot, Actor );
 	}
 }
 
-void FLevelEditorActionCallbacks::SnapActorTo_Clicked( const bool InAlign, const bool InUseLineTrace, const bool InUseBounds, const bool InUsePivot, const AActor* InDestination/* = NULL*/ )
+void FLevelEditorActionCallbacks::SnapTo_Clicked( const bool InAlign, const bool InUseLineTrace, const bool InUseBounds, const bool InUsePivot, AActor* InDestination )
 {
 	// Fires ULevel::LevelDirtiedEvent when falling out of scope.
 	FScopedLevelDirtied		LevelDirtyCallback;
 
-	for ( FSelectionIterator It( GEditor->GetSelectedActorIterator() ) ; It ; ++It )
+	bool bSnappedComponents = false;
+	if( GEditor->GetSelectedComponentCount() > 0 )
 	{
-		AActor* Actor = Cast<AActor>( *It );
-		checkSlow( Actor->IsA(AActor::StaticClass()) );
+		for(FSelectionIterator It(GEditor->GetSelectedComponentIterator()); It; ++It)
+		{
+			USceneComponent* SceneComponent = Cast<USceneComponent>(*It);
+			if(SceneComponent)
+			{
+				SceneComponent->Modify();
+				AActor* ActorOwner = SceneComponent->GetOwner();
+				bSnappedComponents = true;
+				if(ActorOwner)
+				{
+					ActorOwner->Modify();
+					GEditor->SnapObjectTo(FActorOrComponent(SceneComponent), InAlign, InUseLineTrace, InUseBounds, InUsePivot, FActorOrComponent(InDestination));
+					ActorOwner->InvalidateLightingCache();
+					ActorOwner->UpdateComponentTransforms();
 
-		Actor->Modify();
-		GEditor->SnapActorTo(Actor,InAlign,InUseLineTrace,InUseBounds,InUsePivot,InDestination);
-		Actor->InvalidateLightingCache();
-		Actor->UpdateComponentTransforms();
+					LevelDirtyCallback.Request();
+				}
+			}
+		}
 
-		Actor->MarkPackageDirty();
-		LevelDirtyCallback.Request();
+		USceneComponent* LastComp = GEditor->GetSelectedComponents()->GetBottom<USceneComponent>();
+
+		GEditor->SetPivot(LastComp->GetComponentLocation(), false, true);
 	}
 
-	AActor* Actor = GEditor->GetSelectedActors()->GetBottom<AActor>();
-	if( Actor )
+	if( !bSnappedComponents )
 	{
-		GEditor->SetPivot( Actor->GetActorLocation(), false, true );
+		for(FSelectionIterator It(GEditor->GetSelectedActorIterator()); It; ++It)
+		{
+			AActor* Actor = Cast<AActor>(*It);
+			if(Actor)
+			{
+				Actor->Modify();
+				GEditor->SnapObjectTo(FActorOrComponent(Actor), InAlign, InUseLineTrace, InUseBounds, InUsePivot, FActorOrComponent(InDestination));
+				Actor->InvalidateLightingCache();
+				Actor->UpdateComponentTransforms();
 
-		if( GEditor->bGroupingActive ) 
+				LevelDirtyCallback.Request();
+			}
+		}
+
+
+		AActor* Actor = GEditor->GetSelectedActors()->GetBottom<AActor>();
+		if(Actor)
 		{
-			// set group pivot for the root-most group
-			AGroupActor* ActorGroupRoot = AGroupActor::GetRootForActor(Actor, true, true);
-			if(ActorGroupRoot)
+			GEditor->SetPivot(Actor->GetActorLocation(), false, true);
+
+			if(GEditor->bGroupingActive)
 			{
-				ActorGroupRoot->CenterGroupLocation();
+				// set group pivot for the root-most group
+				AGroupActor* ActorGroupRoot = AGroupActor::GetRootForActor(Actor, true, true);
+				if(ActorGroupRoot)
+				{
+					ActorGroupRoot->CenterGroupLocation();
+				}
 			}
 		}
 	}
@@ -2715,8 +2747,8 @@ void FLevelEditorCommands::RegisterCommands()
 	UI_COMMAND( SnapOriginToGrid, "Snap Origin to Grid", "Snaps the actor to the nearest grid location at its origin", EUserInterfaceActionType::Button, FInputGesture( EModifierKey::Control, EKeys::End ) );
 	UI_COMMAND( SnapOriginToGridPerActor, "Snap Origin to Grid Per Actor", "Snaps each selected actor separately to the nearest grid location at its origin", EUserInterfaceActionType::Button, FInputGesture() );
 	UI_COMMAND( AlignOriginToGrid, "Align Origin to Grid", "Aligns the actor to the nearest grid location at its origin", EUserInterfaceActionType::Button, FInputGesture() );
-	UI_COMMAND( SnapToFloor, "Snap to Floor", "Snaps the actor to the floor below it", EUserInterfaceActionType::Button, FInputGesture( EKeys::End ) );
-	UI_COMMAND( AlignToFloor, "Align to Floor", "Aligns the actor with the floor", EUserInterfaceActionType::Button, FInputGesture() );
+	UI_COMMAND( SnapToFloor, "Snap to Floor", "Snaps the actor or component to the floor below it", EUserInterfaceActionType::Button, FInputGesture( EKeys::End ) );
+	UI_COMMAND( AlignToFloor, "Align to Floor", "Aligns the actor or component with the floor", EUserInterfaceActionType::Button, FInputGesture() );
 	UI_COMMAND( SnapPivotToFloor, "Snap Pivot to Floor", "Snaps the actor to the floor at its pivot point", EUserInterfaceActionType::Button, FInputGesture( EModifierKey::Alt, EKeys::End ) );	
 	UI_COMMAND( AlignPivotToFloor, "Align Pivot to Floor", "Aligns the actor with the floor at its pivot point", EUserInterfaceActionType::Button, FInputGesture() );
 	UI_COMMAND( SnapBottomCenterBoundsToFloor, "Snap Bottom Center Bounds to Floor", "Snaps the actor to the floor at its bottom center bounds", EUserInterfaceActionType::Button, FInputGesture( EModifierKey::Shift, EKeys::End ) );
diff --git a/Engine/Source/Editor/LevelEditor/Public/LevelEditorActions.h b/Engine/Source/Editor/LevelEditor/Public/LevelEditorActions.h
--- a/Engine/Source/Editor/LevelEditor/Public/LevelEditorActions.h
+++ b/Engine/Source/Editor/LevelEditor/Public/LevelEditorActions.h
@@ -1162,7 +1162,7 @@ public:
 	 * @param InUseBounds		Whether or not to base the line trace off of the bounds.
 	 * @param InUsePivot		Whether or not to use the pivot position.
 	 */
-	static void SnapActorToFloor_Clicked( bool InAlign, bool InUseLineTrace, bool InUseBounds, bool InUsePivot );
+	static void SnapToFloor_Clicked( bool InAlign, bool InUseLineTrace, bool InUseBounds, bool InUsePivot );
 
 	/**
 	 * Snaps an actor to another actor.  Optionally will align with the trace normal.
@@ -1193,14 +1193,14 @@ private:
 	static void MoveActorTo_Clicked( const bool InAlign, const AActor* InDestination = NULL, bool bInPerActor = false );
 
 	/** 
-	 * Snaps an actor...  Optionally will align with the trace normal.
+	 * Snaps an actor or component...  Optionally will align with the trace normal.
 	 * @param InAlign			Whether or not to rotate the actor to align with the trace normal.
 	 * @param InUseLineTrace	Whether or not to only trace with a line through the world.
 	 * @param InUseBounds		Whether or not to base the line trace off of the bounds.
 	 * @param InUsePivot		Whether or not to use the pivot position.
 	 * @param InDestination		The destination actor we want to move this actor to, NULL assumes we just want to go towards the floor
 	 */
-	static void SnapActorTo_Clicked( const bool InAlign, const bool InUseLineTrace, const bool InUseBounds, const bool InUsePivot, const AActor* InDestination = NULL );
+	static void SnapTo_Clicked( const bool InAlign, const bool InUseLineTrace, const bool InUseBounds, const bool InUsePivot, AActor* InDestination = NULL );
 
 	/** 
 	 * Create and apply animation to the SkeletalMeshComponent if Simulating
diff --git a/Engine/Source/Editor/UnrealEd/Classes/Editor/EditorEngine.h b/Engine/Source/Editor/UnrealEd/Classes/Editor/EditorEngine.h
--- a/Engine/Source/Editor/UnrealEd/Classes/Editor/EditorEngine.h
+++ b/Engine/Source/Editor/UnrealEd/Classes/Editor/EditorEngine.h
@@ -198,6 +198,77 @@ private:
 	TSet<FString> ActorLabels;
 };
 
+/** 
+ * Represents an actor or a component for use in editor functionality such as snapping which can operate on either type
+ */
+struct FActorOrComponent
+{
+	AActor* Actor;
+	USceneComponent* Component;
+
+	FActorOrComponent()
+		: Actor( nullptr )
+		, Component( nullptr )
+	{}
+
+	FActorOrComponent( AActor* InActor )
+		: Actor( InActor )
+		, Component( nullptr )
+	{}
+
+	FActorOrComponent( USceneComponent* InComponent )
+		: Actor( nullptr )
+		, Component( InComponent )
+	{}
+
+	UWorld* GetWorld() const
+	{
+		return Actor ? Actor->GetWorld() : Component->GetWorld();
+	}
+
+	bool operator==( const FActorOrComponent& Other ) const
+	{
+		return Actor == Other.Actor && Component == Other.Component;
+	}
+
+	const FBoxSphereBounds& GetBounds() const
+	{
+		return Actor ? Actor->GetRootComponent()->Bounds : Component->Bounds;
+	}
+
+	FVector GetWorldLocation() const
+	{
+		return Actor ? Actor->GetActorLocation() : Component->GetComponentLocation();
+	}
+
+	void SetWorldLocation( const FVector& NewLocation )
+	{
+		if( Actor )
+		{
+			Actor->SetActorLocation( NewLocation );
+		}
+		else
+		{
+			Component->SetWorldLocation( NewLocation );
+		}
+	}
+
+	void SetWorldRotation( const FRotator& NewRotation )
+	{
+		if(Actor)
+		{
+			Actor->SetActorRotation(NewRotation);
+		}
+		else
+		{
+			Component->SetWorldRotation(NewRotation);
+		}
+	}
+
+	/** @return true if this is a valid actor or component but not both */
+	bool IsValid() const { return (Actor != nullptr) ^ (Component != nullptr);}
+};
+
 /**
  * Engine that drives the Editor.
  * Separate from UGameEngine because it may have much different functionality than desired for an instance of a game itself.
@@ -877,7 +948,7 @@ public:
 	 * @param InDestination		The destination actor we want to move this actor to, NULL assumes we just want to go towards the floor
 	 * @return					Whether or not the actor was moved.
 	 */
-	bool SnapActorTo( AActor* InActor, const bool InAlign, const bool InUseLineTrace, const bool InUseBounds, const bool InUsePivot, const AActor* InDestination = NULL );
+	bool SnapObjectTo( FActorOrComponent Object, const bool InAlign, const bool InUseLineTrace, const bool InUseBounds, const bool InUsePivot, FActorOrComponent InDestination = FActorOrComponent() );
 
 	/**
 	 * Snaps the view of the camera to that of the provided actor.
diff --git a/Engine/Source/Editor/UnrealEd/Private/EditorServer.cpp b/Engine/Source/Editor/UnrealEd/Private/EditorServer.cpp
--- a/Engine/Source/Editor/UnrealEd/Private/EditorServer.cpp
+++ b/Engine/Source/Editor/UnrealEd/Private/EditorServer.cpp
@@ -4472,25 +4472,23 @@ void UEditorEngine::MoveViewportCamerasToActor(const TArray<AActor*> &Actors, co
  * @param InDestination		The destination actor we want to move this actor to, NULL assumes we just want to go towards the floor
  * @return					Whether or not the actor was moved.
  */
-bool UEditorEngine::SnapActorTo( AActor* InActor, const bool InAlign, const bool InUseLineTrace, const bool InUseBounds, const bool InUsePivot, const AActor* InDestination/* = NULL*/ )
+bool UEditorEngine::SnapObjectTo( FActorOrComponent Object, const bool InAlign, const bool InUseLineTrace, const bool InUseBounds, const bool InUsePivot, FActorOrComponent InDestination )
 {
-	check( InActor );
-	if ( InActor == InDestination )	// Early out
+	if ( !Object.IsValid() || Object == InDestination )	// Early out
 	{
 		return false;
 	}
 
 
-	FVector	StartLocation = InActor->GetActorLocation();
+	FVector	StartLocation = Object.GetWorldLocation();
 	FVector	LocationOffset = FVector::ZeroVector;
 	FVector	Extent = FVector::ZeroVector;
-	ABrush* Brush = Cast< ABrush >( InActor );
+	ABrush* Brush = Cast< ABrush >( Object.Actor );
 	bool UseLineTrace = Brush ? true: InUseLineTrace;
 	bool UseBounds = Brush ? true: InUseBounds;
 
 	if( UseLineTrace && UseBounds )
 	{
-		check(InActor->GetRootComponent()->IsRegistered());
 		if (InUsePivot)
 		{
 			// Will do a line trace from the pivot location.
@@ -4499,68 +4497,65 @@ bool UEditorEngine::SnapActorTo( AActor* InActor, const bool InAlign, const bool
 		else
 		{
 			// Will do a line trace from the center bottom of the bounds through the world. Will begin at the bottom center of the component's bounds.
-			StartLocation = InActor->GetRootComponent()->Bounds.Origin;
-			StartLocation.Z -= InActor->GetRootComponent()->Bounds.BoxExtent.Z;
+			StartLocation = Object.GetBounds().Origin;
+			StartLocation.Z -= Object.GetBounds().BoxExtent.Z;
 		}
 
 		// Forces a line trace.
 		Extent = FVector::ZeroVector;
-		LocationOffset = StartLocation - InActor->GetActorLocation();
+		LocationOffset = StartLocation - Object.GetWorldLocation();
 	}
 	else if( UseLineTrace )
 	{
 		// This will be false if multiple objects are selected. In that case the actor's position should be used so all the objects do not go to the same point.
-		if( InUsePivot && !InDestination )	// @todo: If the destination actor is part of the selection tho, we can't use the pivot! (remove check if not)
+		if( InUsePivot && !InDestination.IsValid() )	// @todo: If the destination actor is part of the selection tho, we can't use the pivot! (remove check if not)
 		{
 			StartLocation = GetPivotLocation();
 		}
 		else
 		{
-			StartLocation = InActor->GetActorLocation();
+			StartLocation = Object.GetWorldLocation();
 		}
 
 		// Forces a line trace.
 		Extent = FVector::ZeroVector;
-		LocationOffset = StartLocation - InActor->GetActorLocation();
+		LocationOffset = StartLocation - Object.GetWorldLocation();
 	}
-	else if(InActor->GetRootComponent()) // Casts the entire bounds through the world to find collision.
+	else 
 	{
-		check(InActor->GetRootComponent()->IsRegistered());
-		StartLocation = InActor->GetRootComponent()->Bounds.Origin;
+		StartLocation = Object.GetBounds().Origin;
 
-		Extent = InActor->GetRootComponent()->Bounds.BoxExtent;
-		LocationOffset = StartLocation - InActor->GetActorLocation();
+		Extent = Object.GetBounds().BoxExtent;
+		LocationOffset = StartLocation - Object.GetWorldLocation();
 	}
 
 
 	FVector Direction = FVector(0.f,0.f,-1.f);
-	if ( InDestination )	// If a destination actor was specified, work out the direction
+	if ( InDestination.IsValid() )	// If a destination actor was specified, work out the direction
 	{
-		FVector	EndLocation = InDestination->GetActorLocation();
+		FVector	EndLocation = InDestination.GetWorldLocation();
 
 		// Code here assumes you want to same type of end point as the start point used, comment out to just use the destination actors origin!
 		if( UseLineTrace && UseBounds )
 		{
-			check(InDestination->GetRootComponent()->IsRegistered());
-			EndLocation = InDestination->GetRootComponent()->Bounds.Origin;
-			EndLocation.Z -= InDestination->GetRootComponent()->Bounds.BoxExtent.Z;
+			EndLocation = InDestination.GetBounds().Origin;
+			EndLocation.Z -= InDestination.GetBounds().BoxExtent.Z;
 		}
 		else if( UseLineTrace )
 		{
 			// This will be false if multiple objects are selected. In that case the actor's position should be used so all the objects do not go to the same point.
-			if( InUsePivot && !InDestination )	// @todo: If the destination actor is part of the selection tho, we can't use the pivot! (remove check if not)
+			if( InUsePivot && !InDestination.IsValid() )	// @todo: If the destination actor is part of the selection tho, we can't use the pivot! (remove check if not)
 			{
 				EndLocation = GetPivotLocation();
 			}
 			else
 			{
-				EndLocation = InDestination->GetActorLocation();
+				EndLocation = InDestination.GetWorldLocation();
 			}
 		}
-		else if( InDestination->GetRootComponent() )
+		else
 		{
-			check(InDestination->GetRootComponent()->IsRegistered());
-			EndLocation = InDestination->GetRootComponent()->Bounds.Origin;
+			EndLocation = InDestination.GetBounds().Origin;
 		}
 
 		if ( EndLocation.Equals( StartLocation ) )
@@ -4576,32 +4571,43 @@ bool UEditorEngine::SnapActorTo( AActor* InActor, const bool InAlign, const bool
 	if (Brush)
 	{
 		const float fTinyOffset = 0.01f;
-		StartLocation.Z = InActor->GetRootComponent()->Bounds.Origin.Z - InActor->GetRootComponent()->Bounds.BoxExtent.Z - fTinyOffset;
+		StartLocation.Z = Brush->GetRootComponent()->Bounds.Origin.Z - Brush->GetRootComponent()->Bounds.BoxExtent.Z - fTinyOffset;
 	}
 
 	// Do the actual actor->world check.  We try to collide against the world, straight down from our current position.
 	// If we hit anything, we will move the actor to a position that lets it rest on the floor.
 	FHitResult Hit(1.0f);
-	FCollisionQueryParams Params(FName(TEXT("MoveActorToTrace")), false, InActor);
-	if ( InActor->GetWorld()->SweepSingle(Hit, StartLocation, StartLocation + Direction*WORLD_MAX, FQuat::Identity, FCollisionShape::MakeBox(Extent), Params, FCollisionObjectQueryParams(ECC_WorldStatic)))
+	FCollisionQueryParams Params(FName(TEXT("MoveActorToTrace")), false);
+	if( Object.Actor )
+	{
+		Params.AddIgnoredActor( Object.Actor );
+	}
+	else
+	{
+		Params.AddIgnoredComponent( Cast<UPrimitiveComponent>(Object.Component) );
+	}
+
+	if ( Object.GetWorld()->SweepSingle(Hit, StartLocation, StartLocation + Direction*WORLD_MAX, FQuat::Identity, FCollisionShape::MakeBox(Extent), Params, FCollisionObjectQueryParams(ECC_WorldStatic)))
 	{
 		FVector NewLocation = Hit.Location - LocationOffset;
 		NewLocation.Z += KINDA_SMALL_NUMBER;	// Move the new desired location up by an error tolerance
-		
-		InActor->TeleportTo( NewLocation, InActor->GetActorRotation(), false,true );
+
+		Object.SetWorldLocation( NewLocation );
+		//InActor->TeleportTo( NewLocation, InActor->GetActorRotation(), false,true );
 		
 		if( InAlign )
 		{
 			//@todo: This doesn't take into account that rotating the actor changes LocationOffset.
 			FRotator NewRotation( Hit.Normal.Rotation() );
 			NewRotation.Pitch -= 90.f;
-			InActor->SetActorRotation(NewRotation);
+			Object.SetWorldRotation( NewRotation );
 		}
 
 		// Switch to the pie world if we have one
 		FScopedConditionalWorldSwitcher WorldSwitcher( GCurrentLevelEditingViewportClient );
 
-		InActor->PostEditMove( true );
+		Object.Actor ? Object.Actor->PostEditMove(true) : Object.Component->GetOwner()->PostEditMove(true);
+		//InActor->PostEditMove( true );
 		if (Brush)
 		{
 			RebuildAlteredBSP();
