[INTEGRATE] Change 2421922 by bruce.nesbit@bnesbit_Releases on 2015/01/28 05:46:59

	Added message to indicate compiling is going to happen when importing code FP. Also made code wait until its complete
	Added provision for focus on a specifc asset after adding a feature pack
	Revised class type names in feature pack manifests.

diff --git a/Engine/Source/Editor/AddContentDialog/AddContentDialog.Build.cs b/Engine/Source/Editor/AddContentDialog/AddContentDialog.Build.cs
--- a/Engine/Source/Editor/AddContentDialog/AddContentDialog.Build.cs
+++ b/Engine/Source/Editor/AddContentDialog/AddContentDialog.Build.cs
@@ -9,6 +9,7 @@ public class AddContentDialog : ModuleRules
 		PrivateIncludePathModuleNames.AddRange(
 			new string[] {
 				"AssetTools",
+				"ContentBrowser",
 			}
 		);
 
diff --git a/Engine/Source/Editor/AddContentDialog/Private/ContentSourceProviders/FeaturePack/FeaturePackContentSource.cpp b/Engine/Source/Editor/AddContentDialog/Private/ContentSourceProviders/FeaturePack/FeaturePackContentSource.cpp
--- a/Engine/Source/Editor/AddContentDialog/Private/ContentSourceProviders/FeaturePack/FeaturePackContentSource.cpp
+++ b/Engine/Source/Editor/AddContentDialog/Private/ContentSourceProviders/FeaturePack/FeaturePackContentSource.cpp
@@ -4,6 +4,7 @@
 #include "FeaturePackContentSource.h"
 
 #include "AssetToolsModule.h"
+#include "ContentBrowserModule.h"
 #include "IPlatformFilePak.h"
 #include "FileHelpers.h"
 
@@ -64,6 +65,9 @@ FFeaturePackContentSource::FFeaturePackContentSource(FString InFeaturePackPath)
 	// Parse class types field
 	ClassTypes = ManifestObject->GetStringField("ClassTypes");
 	
+	// Parse initial focus asset
+	FocusAssetIdent = ManifestObject->GetStringField("FocusAsset");
+
 	FString CategoryString = ManifestObject->GetStringField("Category");	
 	if (CategoryString == "CodeFeature")
 	{
@@ -164,7 +168,8 @@ bool FFeaturePackContentSource::InstallToProject(FString InstallPath)
 		FAssetToolsModule& AssetToolsModule = FModuleManager::Get().LoadModuleChecked<FAssetToolsModule>("AssetTools");
 		TArray<FString> AssetPaths;
 		AssetPaths.Add(FeaturePackPath);
-		TArray<UObject*> ImportedObjects = AssetToolsModule.Get().ImportAssets(AssetPaths, InstallPath);
+
+		TArray<UObject*> ImportedObjects = AssetToolsModule.Get().ImportAssets(AssetPaths, InstallPath );
 		if( ImportedObjects.Num() == 0 )
 		{
 			UE_LOG(LogFeaturePack, Warning, TEXT("No objects imported installing pack %s"), *InstallPath);
@@ -179,6 +184,19 @@ bool FFeaturePackContentSource::InstallToProject(FString InstallPath)
 			}
 			FEditorFileUtils::PromptForCheckoutAndSave( ToSave, /*bCheckDirty=*/ false, /*bPromptToSave=*/ false );
 			bResult = true;
+			
+			// Focus on a specific asset if we want to.
+			if( GetFocusAssetName().IsEmpty() == false )
+			{
+				UObject* FocusAsset = LoadObject<UObject>(nullptr, *GetFocusAssetName());
+				if (FocusAsset)
+				{
+					FContentBrowserModule& ContentBrowserModule = FModuleManager::Get().LoadModuleChecked<FContentBrowserModule>("ContentBrowser");
+					TArray<UObject*> SyncObjects;
+					SyncObjects.Add(FocusAsset);
+					ContentBrowserModule.Get().SyncBrowserToAssets(SyncObjects);
+				}
+			}
 		}
 	}
 	return bResult;
@@ -196,4 +214,9 @@ bool FFeaturePackContentSource::IsDataValid() const
 	}
 	// To Do maybe validate other data here
 	return true;	
-}
\ No newline at end of file
+}
+
+FString FFeaturePackContentSource::GetFocusAssetName() const
+{
+	return FocusAssetIdent;
+}
diff --git a/Engine/Source/Editor/AddContentDialog/Private/ContentSourceProviders/FeaturePack/FeaturePackContentSource.h b/Engine/Source/Editor/AddContentDialog/Private/ContentSourceProviders/FeaturePack/FeaturePackContentSource.h
--- a/Engine/Source/Editor/AddContentDialog/Private/ContentSourceProviders/FeaturePack/FeaturePackContentSource.h
+++ b/Engine/Source/Editor/AddContentDialog/Private/ContentSourceProviders/FeaturePack/FeaturePackContentSource.h
@@ -22,11 +22,12 @@ public:
 
 	virtual bool InstallToProject(FString InstallPath) override;
 	virtual bool IsDataValid() const override;
-
+	
 	virtual ~FFeaturePackContentSource();
 	
 private:
 	bool LoadPakFileToBuffer(FPakPlatformFile& PakPlatformFile, FString Path, TArray<uint8>& Buffer);
+	FString GetFocusAssetName() const;
 
 private:
 	FString FeaturePackPath;
@@ -38,4 +39,5 @@ private:
 	TArray<FLocalizedText> LocalizedAssetTypesList;
 	FString ClassTypes;
 	bool bPackValid;
+	FString FocusAssetIdent;
 };
\ No newline at end of file
diff --git a/Engine/Source/Editor/UnrealEd/Private/Factories/PackFactory.cpp b/Engine/Source/Editor/UnrealEd/Private/Factories/PackFactory.cpp
--- a/Engine/Source/Editor/UnrealEd/Private/Factories/PackFactory.cpp
+++ b/Engine/Source/Editor/UnrealEd/Private/Factories/PackFactory.cpp
@@ -246,7 +246,7 @@ UObject* UPackFactory::FactoryCreateBinary
 	const uint8*		BufferEnd,
 	FFeedbackContext*	Warn
 )
-{
+{ 
 	FBufferReader PakReader((void*)Buffer, BufferEnd-Buffer, false);
 	FPakFile PakFile(&PakReader);
 
@@ -336,7 +336,6 @@ UObject* UPackFactory::FactoryCreateBinary
 					// Setup the game name redirect
 					if (FString* GameName = ConfigParameters.Find("GameName"))
 					{
-
 						const FString EngineIniFilename = FPaths::ConvertRelativePathToFull(GetDefault<UEngine>()->GetDefaultConfigFilename());
 
 						if (ISourceControlModule::Get().IsEnabled())
@@ -478,15 +477,9 @@ UObject* UPackFactory::FactoryCreateBinary
 		if (WrittenFiles.Num() > 0)
 		{
 			// If we wrote out source files, kick off the hot reload process
-			// TODO: Ideally we would block on finishing the factory until the hot reload completed
 			if (WrittenSourceFiles.Num() > 0)
 			{
-				IHotReloadInterface& HotReloadSupport = FModuleManager::LoadModuleChecked<IHotReloadInterface>("HotReload");
-				if( !HotReloadSupport.IsCurrentlyCompiling() )
-				{
-					HotReloadSupport.DoHotReloadFromEditor();
-				}
-
+				// Ask about editing code where applicable
 				if (FSlateApplication::Get().SupportsSourceAccess() )
 				{
 					// Code successfully added, notify the user and ask about opening the IDE now
@@ -496,11 +489,31 @@ UObject* UPackFactory::FactoryCreateBinary
 						FSourceCodeNavigation::OpenSourceFiles(WrittenSourceFiles);
 					}
 				}
-
+				
+				// Tell the user we are about to compile and fire it off (if we need to)
+				IHotReloadInterface& HotReloadSupport = FModuleManager::LoadModuleChecked<IHotReloadInterface>("HotReload");
+				if (!HotReloadSupport.IsCurrentlyCompiling())
+				{					
+					const FText Message = NSLOCTEXT("PackFactory", "WaitingForCompile", "Press OK to begin compling. Pack will installation will complete when the compiler has finished.");
+					FMessageDialog::Open(EAppMsgType::Ok, Message);
+				
+					HotReloadSupport.DoHotReloadFromEditor();
+					
+					UE_LOG(LogPackFactory, Warning, TEXT("Starting compile."));
+					// Wait until the compiler finishes
+					while (HotReloadSupport.IsCurrentlyCompiling())
+					{
+						// Tick slate while we wait
+						FSlateApplication::Get().PumpMessages();
+						FSlateApplication::Get().Tick();						
+						// Tick hot reload - or nothing will happen !
+						HotReloadSupport.Tick();
+					}
+					UE_LOG(LogPackFactory, Warning, TEXT("Compile complete."));
+				}
 			}
-
-			// Find an asset to return as the file to be selected in the content browser
-			// TODO: Should we load all objects? Should the config file be able to specify which asset is selected?
+			
+			// Find an asset to return (It will be marked as dirty)
 			for (const FString& Filename : WrittenFiles)
 			{
 				static const FString AssetExtension(TEXT(".uasset"));
@@ -515,7 +528,7 @@ UObject* UPackFactory::FactoryCreateBinary
 						{
 							const FString AssetName = GameFileName.RightChop(SlashIndex + 1);
 							ReturnAsset = LoadObject<UObject>(nullptr, *(GameFileName + TEXT(".") + AssetName));
-							if (ReturnAsset) 
+							if (ReturnAsset)
 							{
 								break;
 							}
@@ -543,6 +556,5 @@ UObject* UPackFactory::FactoryCreateBinary
 		UE_LOG(LogPackFactory, Warning, TEXT("Invalid pak file."));
 	}
 
-
 	return ReturnAsset;
-}
\ No newline at end of file
+}
