[INTEGRATE] Change 2421951 by [REDACTED_PERSON][REDACTED_HANDLE].7 on 2015/01/28 06:47:27

	Fixed slow tasks sometimes not updating their progress dialogs

	SlateTick was being throttled along with slow task progress updates

diff --git a/Engine/Source/Editor/UnrealEd/Private/FeedbackContextEditor.cpp b/Engine/Source/Editor/UnrealEd/Private/FeedbackContextEditor.cpp
--- a/Engine/Source/Editor/UnrealEd/Private/FeedbackContextEditor.cpp
+++ b/Engine/Source/Editor/UnrealEd/Private/FeedbackContextEditor.cpp
@@ -36,9 +36,6 @@ public:
 	/** Construct this widget */
 	void Construct( const FArguments& InArgs )
 	{
-		// Ensure this gets ticked at least once
-		LastTickTime = 0;
-
 		OnCancelClickedDelegate = InArgs._OnCancelClickedDelegate;
 		ScopeStack = InArgs._ScopeStack;
 
@@ -306,10 +303,6 @@ private:
 
 	/** Array mapping progress bar index -> scope stack index. Updated every tick. */
 	TArray<int32> DynamicProgressIndices;
-
-public:
-	/** Publicly accessible tick time to throttle slate ticking */
-	static double LastTickTime;
 };
 
 /** Static integer definitions required on some builds where the linker needs access to these */
@@ -318,21 +311,14 @@ const int32 SSlowTaskWidget::FixedWidth;
 const int32 SSlowTaskWidget::FixedPaddingH;;
 const int32 SSlowTaskWidget::MainBarHeight;
 const int32 SSlowTaskWidget::SecondaryBarHeight;
-double SSlowTaskWidget::LastTickTime = 0;
 
 static void TickSlate()
 {
-	static double MinFrameTime = 0.05;		// Only update at 20fps so as not to slow down the actual task
-	if (FPlatformTime::Seconds() - SSlowTaskWidget::LastTickTime > MinFrameTime)
-	{
-		SSlowTaskWidget::LastTickTime = FPlatformTime::Seconds();
-
-		// Tick Slate application
-		FSlateApplication::Get().Tick();
+	// Tick Slate application
+	FSlateApplication::Get().Tick();
 
-		// Sync the game thread and the render thread. This is needed if many StatusUpdate are called
-		FSlateApplication::Get().GetRenderer()->Sync();
-	}
+	// Sync the game thread and the render thread. This is needed if many StatusUpdate are called
+	FSlateApplication::Get().GetRenderer()->Sync();
 }
 
 FFeedbackContextEditor::FFeedbackContextEditor()
