[INTEGRATE] Change 2422025 by bruce.nesbit@bnesbit_Releases on 2015/01/28 08:30:06

	Removed unrequired call to PumpMessages in PackFactory to

diff --git a/Engine/Source/Editor/UnrealEd/Private/Factories/PackFactory.cpp b/Engine/Source/Editor/UnrealEd/Private/Factories/PackFactory.cpp
--- a/Engine/Source/Editor/UnrealEd/Private/Factories/PackFactory.cpp
+++ b/Engine/Source/Editor/UnrealEd/Private/Factories/PackFactory.cpp
@@ -504,7 +504,6 @@ UObject* UPackFactory::FactoryCreateBinary
 					while (HotReloadSupport.IsCurrentlyCompiling())
 					{
 						// Tick slate while we wait
-						FSlateApplication::Get().PumpMessages();
 						FSlateApplication::Get().Tick();						
 						// Tick hot reload - or nothing will happen !
 						HotReloadSupport.Tick();
