[INTEGRATE] Change 2422261 by [REDACTED_PERSON][REDACTED_HANDLE].7 on 2015/01/28 10:58:34

	Auto monitoring of content directories is now on by default

diff --git a/Engine/Source/Editor/UnrealEd/Private/Settings/SettingsClasses.cpp b/Engine/Source/Editor/UnrealEd/Private/Settings/SettingsClasses.cpp
--- a/Engine/Source/Editor/UnrealEd/Private/Settings/SettingsClasses.cpp
+++ b/Engine/Source/Editor/UnrealEd/Private/Settings/SettingsClasses.cpp
@@ -93,7 +93,7 @@ void UEditorExperimentalSettings::PostEditChangeProperty( struct FPropertyChange
 
 UEditorLoadingSavingSettings::UEditorLoadingSavingSettings( const FObjectInitializer& ObjectInitializer )
 	: Super(ObjectInitializer)
-	, bMonitorContentDirectories(false)
+	, bMonitorContentDirectories(true)
 	, bAutoCreateAssets(true)
 	, bAutoDeleteAssets(true)
 	, bDeleteSourceFilesWithAssets(true)
