[INTEGRATE] Change 2422295 by [REDACTED_PERSON][REDACTED_HANDLE] on 2015/01/28 11:19:26

	[UE-7812] - Pivot location of a component in world is properly transformed to be in the same space as the component

diff --git a/Engine/Source/Editor/UnrealEd/Private/LevelEditorViewport.cpp b/Engine/Source/Editor/UnrealEd/Private/LevelEditorViewport.cpp
--- a/Engine/Source/Editor/UnrealEd/Private/LevelEditorViewport.cpp
+++ b/Engine/Source/Editor/UnrealEd/Private/LevelEditorViewport.cpp
@@ -3036,32 +3036,32 @@ void FLevelEditorViewportClient::ApplyDeltaToActors(const FVector& InDrag,
 			}
 			else
 			{
-			AGroupActor* ParentGroup = AGroupActor::GetRootForActor(Actor, true, true);
-			if(ParentGroup && GEditor->bGroupingActive )
-			{
-				ActorGroups.AddUnique(ParentGroup);
-			}
-			else
-			{
-				// Finally, verify that no actor in the parent hierarchy is also selected
-				bool bHasParentInSelection = false;
-				AActor* ParentActor = Actor->GetAttachParentActor();
-				while(ParentActor!=NULL && !bHasParentInSelection)
+				AGroupActor* ParentGroup = AGroupActor::GetRootForActor(Actor, true, true);
+				if (ParentGroup && GEditor->bGroupingActive)
 				{
-					if(ParentActor->IsSelected())
-					{
-						bHasParentInSelection = true;
-					}
-					ParentActor = ParentActor->GetAttachParentActor();
+					ActorGroups.AddUnique(ParentGroup);
 				}
-				if(!bHasParentInSelection)
+				else
 				{
-					ApplyDeltaToActor( Actor, InDrag, InRot, ModifiedScale );
+					// Finally, verify that no actor in the parent hierarchy is also selected
+					bool bHasParentInSelection = false;
+					AActor* ParentActor = Actor->GetAttachParentActor();
+					while (ParentActor != NULL && !bHasParentInSelection)
+					{
+						if (ParentActor->IsSelected())
+						{
+							bHasParentInSelection = true;
+						}
+						ParentActor = ParentActor->GetAttachParentActor();
+					}
+					if (!bHasParentInSelection)
+					{
+						ApplyDeltaToActor(Actor, InDrag, InRot, ModifiedScale);
+					}
 				}
 			}
 		}
 	}
-	}
 	AGroupActor::RemoveSubGroupsFromArray(ActorGroups);
 	for(int32 ActorGroupsIndex=0; ActorGroupsIndex<ActorGroups.Num(); ++ActorGroupsIndex)
 	{
@@ -3091,13 +3091,19 @@ void FLevelEditorViewportClient::ApplyDeltaToComponent(USceneComponent* InCompon
 	FRotator AdjustedRot = InDeltaRot;
 	FComponentEditorUtils::AdjustComponentDelta(InComponent, AdjustedDrag, AdjustedRot);
 
+	FVector EditorWorldPivotLocation = GEditor->GetPivotLocation();
+
+	// If necessary, transform the editor pivot location to be relative to the component's parent
+	const bool bIsRootComponent = InComponent->GetOwner()->GetRootComponent() == InComponent;
+	FVector RelativePivotLocation = bIsRootComponent ? EditorWorldPivotLocation : InComponent->GetAttachParent()->GetComponentToWorld().Inverse().TransformPosition(EditorWorldPivotLocation);
+
 	GEditor->ApplyDeltaToComponent(
 		InComponent,
 		true,
 		&AdjustedDrag,
 		&AdjustedRot,
 		&ModifiedDeltaScale,
-		InComponent->RelativeLocation);
+		RelativePivotLocation);
 }
 
 /** Helper function for ModifyScale - Convert the active Dragging Axis to per-axis flags */
diff --git a/Engine/Source/Editor/UnrealEd/Private/LevelViewportClickHandlers.cpp b/Engine/Source/Editor/UnrealEd/Private/LevelViewportClickHandlers.cpp
--- a/Engine/Source/Editor/UnrealEd/Private/LevelViewportClickHandlers.cpp
+++ b/Engine/Source/Editor/UnrealEd/Private/LevelViewportClickHandlers.cpp
@@ -302,7 +302,7 @@ namespace ClickHandlers
 		// Pivot snapping
 		if (Click.GetKey() == EKeys::MiddleMouseButton && Click.IsAltDown())
 		{
-			GEditor->SetPivot(GEditor->ClickLocation, true, false, true);
+			GEditor->SetPivot(GEditor->ClickLocation, true, false);
 
 			return true;
 		}
