[INTEGRATE] Change 2422476 by [REDACTED_PERSON][REDACTED_HANDLE] on 2015/01/28 13:20:07

	[UE-8526] Allowing blueprinted components to be super-classes for other blueprints.

diff --git a/Engine/Source/Editor/UnrealEd/Private/Kismet2/Kismet2.cpp b/Engine/Source/Editor/UnrealEd/Private/Kismet2/Kismet2.cpp
--- a/Engine/Source/Editor/UnrealEd/Private/Kismet2/Kismet2.cpp
+++ b/Engine/Source/Editor/UnrealEd/Private/Kismet2/Kismet2.cpp
@@ -879,6 +879,7 @@ bool FKismetEditorUtilities::CanCreateBlueprintOfClass(const UClass* Class)
 	const bool bIsValidClass = Class->GetBoolMetaDataHierarchical(FBlueprintMetadata::MD_IsBlueprintBase)
 		|| (Class == UObject::StaticClass())
 		|| (Class == UActorComponent::StaticClass() && bAllowBlueprintableComponents)
+		|| (bAllowBlueprintableComponents && Class->IsChildOf<UActorComponent>() && Class->HasAnyClassFlags(CLASS_CompiledFromBlueprint))
 		|| bIsBPGC;  // BPs are always considered inheritable
 
 	return bCanCreateBlueprint && bIsValidClass;
