[INTEGRATE] Change 2422576 by [REDACTED_PERSON][REDACTED_HANDLE] on 2015/01/28 14:11:18

	[UE-7296][UE-7780] More circular dependency test cases.

diff --git a/Engine/Source/Editor/UnrealEd/Classes/Tests/CircularDependencyTestActor.h b/Engine/Source/Editor/UnrealEd/Classes/Tests/CircularDependencyTestActor.h
--- a/Engine/Source/Editor/UnrealEd/Classes/Tests/CircularDependencyTestActor.h
+++ b/Engine/Source/Editor/UnrealEd/Classes/Tests/CircularDependencyTestActor.h
@@ -29,19 +29,25 @@ public:
 	void VisualizeNewTestState(ETestResult::Type OldState, ETestResult::Type NewState);
 
 	UFUNCTION(BlueprintCallable, Category="Development|Tests")
-	bool TestVerifyClass(bool bCheckPropertyType = true);
+	bool TestVerifyClass(UClass* Class, bool bCheckPropertyType = true, bool bCheckPropertyVals = true);
 
 	UFUNCTION(BlueprintCallable, Category = "Development|Tests")
-	bool TestVerifyBlueprint();
+	bool TestVerifyBlueprint(UClass* BPClass);
 
 	//UFUNCTION(BlueprintCallable, Category="Development|Tests")
-	bool TestVerifyStructMember(UScriptStruct* Struct, uint8* Container, bool bCheckPropertyType = true);
+	bool TestVerifyStructMember(UScriptStruct* Struct, UClass* AncestorClass, uint8* Container, bool bCheckPropertyType = true, bool bCheckPropertyVals = true);
+
+	UFUNCTION(BlueprintCallable, Category = "Development|Tests")
+	bool TestVerifySubObjects(UObject* ObjInstace);
 
 	UFUNCTION(BlueprintCallable, BlueprintNativeEvent, Category="Development|Tests")
 	bool RunVerificationTests();
 
+	UFUNCTION(BlueprintCallable, Category = "Development|Tests")
+	bool TestObjectValidity(UObject* ObjInst, bool bCheckPropertyType = true, bool bCheckPropertyVals = true);
+
 	UFUNCTION(BlueprintCallable, Category="Development|Tests")
-	bool TestVerifyIsBlueprintTypeVar(FName VarName, bool bCheckPropertyType);
+	bool TestVerifyIsBlueprintTypeVar(FName VarName, UClass* ClassOuter, bool bCheckPropertyType, bool bCheckPropertyVal);
 
 	UFUNCTION(BlueprintCallable, Category="Development|Tests")
 	void ForceResetFailure();
@@ -54,5 +60,6 @@ protected:
 	TEnumAsByte<ETestResult::Type> TestState;
 
 private: 
-	bool TestVerifyProperty(UProperty* Property, uint8* Container, bool bCheckPropertyType);
+	bool TestVerifyClass(UClass* Class, UObject* ObjInst, bool bCheckPropertyType, bool bCheckPropertyVals);
+	bool TestVerifyProperty(UProperty* Property, UClass* PropOwner, uint8* Container, bool bCheckPropertyType, bool bCheckPropertyVal);
 };
diff --git a/Engine/Source/Editor/UnrealEd/Private/Tests/CircularDependencyTestActor.cpp b/Engine/Source/Editor/UnrealEd/Private/Tests/CircularDependencyTestActor.cpp
--- a/Engine/Source/Editor/UnrealEd/Private/Tests/CircularDependencyTestActor.cpp
+++ b/Engine/Source/Editor/UnrealEd/Private/Tests/CircularDependencyTestActor.cpp
@@ -37,24 +37,54 @@ static bool IsValidClass(UClass* ClassToVerify)
 	return (ClassToVerify != nullptr) && !IsPlaceholderClass(ClassToVerify);
 };
 
-bool ACircularDependencyTestActor::TestVerifyClass(bool bCheckPropertyType)
+bool ACircularDependencyTestActor::TestVerifyClass(UClass* Class, bool bCheckPropertyType, bool bCheckPropertyVals)
 {
-	UClass* BPClass = GetClass();
+	UObject* ClassInst = this;
+	if (Class && !GetClass()->IsChildOf(Class))
+	{
+		ClassInst = Class->GetDefaultObject();
+	}
+	return TestVerifyClass(Class, ClassInst, bCheckPropertyType, bCheckPropertyVals);
+}
+
+bool ACircularDependencyTestActor::TestVerifyClass(UClass* BPClass, UObject* ClassInst, bool bCheckPropertyType, bool bCheckPropertyVals)
+{
+	if (ClassInst == nullptr || !ClassInst->GetClass()->IsChildOf(BPClass))
+	{
+		return false;
+	}
+
+	if (Cast<UBlueprintGeneratedClass>(BPClass) == nullptr)
+	{
+		return false;
+	}
+
+// 	if (ULinkerLoad* Linker = BPClass->GetLinker())
+// 	{
+// // 		if (Linker->HasUnresolvedDependencies() || Linker->IsBlueprintFinalizationPending())
+// // 		{
+// // 			return false;
+// // 		}
+// 	}
+// 	else
+// 	{
+// 		return false;
+// 	}
 
 	for (UProperty* ClassProp = BPClass->PropertyLink; ClassProp != nullptr; ClassProp = ClassProp->PropertyLinkNext)
 	{
-		if (!TestVerifyProperty(ClassProp, (uint8*)this, bCheckPropertyType))
+		if (!TestVerifyProperty(ClassProp, BPClass, (uint8*)ClassInst, bCheckPropertyType, bCheckPropertyVals))
 		{
-			SetTestState(ETestResult::Failed);
+			//SetTestState(ETestResult::Failed);
 			return false;
 		}
 
 		if (UStructProperty* StructProp = Cast<UStructProperty>(ClassProp))
 		{
-			uint8* StructInst = StructProp->ContainerPtrToValuePtr<uint8>(this);
-			if ((StructProp->Struct == nullptr) || !TestVerifyStructMember(StructProp->Struct, (uint8*)StructInst, bCheckPropertyType))
+			uint8* StructInst = StructProp->ContainerPtrToValuePtr<uint8>(ClassInst);
+			if ((StructProp->Struct == nullptr) || !TestVerifyStructMember(StructProp->Struct, BPClass, (uint8*)StructInst, bCheckPropertyType))
 			{
-				SetTestState(ETestResult::Failed);
+				//SetTestState(ETestResult::Failed);
 				return false;
 			}
 		}
@@ -62,10 +92,9 @@ bool ACircularDependencyTestActor::TestVerifyClass(bool bCheckPropertyType)
 	return true;
 }
 
-bool ACircularDependencyTestActor::TestVerifyBlueprint()
+bool ACircularDependencyTestActor::TestVerifyBlueprint(UClass* BPClass)
 {
 #if WITH_EDITOR
-	UClass* BPClass = GetClass();
 	if (UBlueprint* Blueprint = Cast<UBlueprint>(BPClass->ClassGeneratedBy))
 	{
 		TArray<UEdGraph*> BpGraphs;
@@ -131,11 +160,11 @@ bool ACircularDependencyTestActor::TestVerifyBlueprint()
 	return true;
 }
 
-bool ACircularDependencyTestActor::TestVerifyStructMember(UScriptStruct* Struct, uint8* StructInst, bool bCheckPropertyType)
+bool ACircularDependencyTestActor::TestVerifyStructMember(UScriptStruct* Struct, UClass* AncestorClass, uint8* StructInst, bool bCheckPropertyType, bool bCheckPropertyVals)
 {
 	for (UProperty* StructProp = Struct->PropertyLink; StructProp != nullptr; StructProp = StructProp->NextRef)
 	{
-		if (!TestVerifyProperty(StructProp, StructInst, bCheckPropertyType))
+		if (!TestVerifyProperty(StructProp, AncestorClass, StructInst, bCheckPropertyType, bCheckPropertyVals))
 		{
 			SetTestState(ETestResult::Failed);
 			return false;
@@ -144,7 +173,7 @@ bool ACircularDependencyTestActor::TestVerifyStructMember(UScriptStruct* Struct,
 		if (UStructProperty* StructProp2 = Cast<UStructProperty>(StructProp))
 		{
 			uint8* StructInst2 = StructProp2->ContainerPtrToValuePtr<uint8>(StructInst);
-			if ((StructProp2->Struct == nullptr) || !TestVerifyStructMember(StructProp2->Struct, StructInst2, bCheckPropertyType))
+			if ((StructProp2->Struct == nullptr) || !TestVerifyStructMember(StructProp2->Struct, AncestorClass, StructInst2, bCheckPropertyType))
 			{
 				SetTestState(ETestResult::Failed);
 				return false;
@@ -154,10 +183,78 @@ bool ACircularDependencyTestActor::TestVerifyStructMember(UScriptStruct* Struct,
 	return true;
 }
 
-bool ACircularDependencyTestActor::TestVerifyIsBlueprintTypeVar(FName VarName, bool bCheckPropertyType)
+bool ACircularDependencyTestActor::TestVerifySubObjects(UObject* ObjInstace)
+{
+	TArray<UObject*> SubObjects;
+	GetObjectsWithOuter(ObjInstace, SubObjects);
+
+	TArray<FString> Errors;
+
+	for (auto SubObjIt : SubObjects)
+	{
+		const UObject* SubObj = SubObjIt;
+		const UClass*  SubObjClass = SubObj->GetClass();
+		const FString  SubObjName  = SubObj->GetName();
+
+		if (SubObj->IsPendingKill())
+		{
+			continue;
+		}
+
+		if (SubObjClass->HasAnyClassFlags(CLASS_NewerVersionExists))
+		{
+			return false;
+		}
+
+		if (SubObjClass->GetOutermost() == GetTransientPackage())
+		{
+			return false;
+		}
+
+		if (SubObj->GetOutermost() != ObjInstace->GetOutermost())
+		{
+			return false;
+		}
+
+		if (SubObj->GetName().Find(TEXT("TRASH_")) != INDEX_NONE)
+		{
+			return false;
+		}
+
+		if (SubObj->GetName().Find(TEXT("REINST_")) != INDEX_NONE)
+		{
+			return false;
+		}
+
+		if (SubObj->GetName().Find(TEXT("PLACEHOLDER-")) != INDEX_NONE)
+		{
+			return false;
+		}
+	}
+
+	if (AActor* AsActor = Cast<AActor>(ObjInstace))
+	{
+		for (UActorComponent* Component : AsActor->GetComponents())
+		{
+			if (!TestVerifySubObjects(Component))
+			{
+				return false;
+			}
+		}
+	}
+	return true;
+}
+
+bool ACircularDependencyTestActor::TestVerifyIsBlueprintTypeVar(FName VarName, UClass* ClassOuter, bool bCheckPropertyType, bool bCheckPropertyVal)
 {
 	FMemberReference VarRef;
-	VarRef.SetExternalMember(VarName, GetClass());
+	VarRef.SetExternalMember(VarName, ClassOuter);
+
+	UObject* ClassInst = this;
+	if (ClassOuter != GetClass())
+	{
+		ClassInst = ClassOuter->GetDefaultObject();
+	}
 
 	UProperty* VarProperty = VarRef.ResolveMember<UProperty>(GetClass());
 	if (UObjectProperty* ObjProp = Cast<UObjectProperty>(VarProperty))
@@ -165,11 +262,11 @@ bool ACircularDependencyTestActor::TestVerifyIsBlueprintTypeVar(FName VarName, b
 		UClass* ClassToVerify = ObjProp->PropertyClass;
 		if (UClassProperty* ClassProp = Cast<UClassProperty>(ObjProp))
 		{
-			ClassToVerify = Cast<UClass>(ClassProp->GetPropertyValue_InContainer(this));
+			ClassToVerify = Cast<UClass>(ClassProp->GetPropertyValue_InContainer(ClassInst));
 		}
-		else
+		else if (bCheckPropertyVal)
 		{
-			UObject* ObjVal = ObjProp->GetPropertyValue_InContainer(this);
+			UObject* ObjVal = ObjProp->GetPropertyValue_InContainer(ClassInst);
 			if (UChildActorComponent* AsChildComponent = Cast<UChildActorComponent>(ObjVal))
 			{
 				ObjVal = AsChildComponent->ChildActor;
@@ -198,21 +295,27 @@ bool ACircularDependencyTestActor::TestVerifyIsBlueprintTypeVar(FName VarName, b
 
 bool ACircularDependencyTestActor::RunVerificationTests_Implementation()
 {
-	return TestVerifyClass() && TestVerifyBlueprint();
+	return TestObjectValidity(this);
+}
+
+bool ACircularDependencyTestActor::TestObjectValidity(UObject* ObjInst, bool bCheckPropertyType, bool bCheckPropertyVals)
+{
+	return TestVerifyClass(ObjInst->GetClass(), ObjInst, bCheckPropertyType, bCheckPropertyVals) && TestVerifyBlueprint(ObjInst->GetClass()) && TestVerifySubObjects(ObjInst);// &&
+		//(ObjInst->GetClass()->GetDefaultObject<ACircularDependencyTestActor>()->TestState != ETestResult::Failed);
 }
 
-bool ACircularDependencyTestActor::TestVerifyProperty(UProperty* Property, uint8* Container, bool bCheckPropertyType)
+bool ACircularDependencyTestActor::TestVerifyProperty(UProperty* Property, UClass* PropOwner, uint8* Container, bool bCheckPropertyType, bool bCheckPropertyVal)
 {
 	bool bBelongsToThisBlueprintClass = false;
 	if (UClass* OwnerClass = Property->GetOwnerClass())
 	{
-		bBelongsToThisBlueprintClass = (OwnerClass == GetClass()) && (Cast<UBlueprintGeneratedClass>(OwnerClass) != nullptr);
+		bBelongsToThisBlueprintClass = (OwnerClass == PropOwner) && (Cast<UBlueprintGeneratedClass>(OwnerClass) != nullptr);
 		if (UClass* SuperClass = OwnerClass->GetSuperClass())
 		{
 			FMemberReference VarRef;
 			VarRef.SetExternalMember(Property->GetFName(), SuperClass);
 			// make sure the parent doesn't have this variable
-			bBelongsToThisBlueprintClass &= (VarRef.ResolveMember<UProperty>(GetClass()) == nullptr);
+			bBelongsToThisBlueprintClass &= (VarRef.ResolveMember<UProperty>(PropOwner) == nullptr);
 		}
 	}
 	else if (UStruct* OwnerStruct = Property->GetOwnerStruct())
@@ -237,7 +340,7 @@ bool ACircularDependencyTestActor::TestVerifyProperty(UProperty* Property, uint8
 		{
 			// assume any object properties that we've added to the BP 
 			// directly should be circular refs (and thus, pointers to other blueprint types) 
-			return TestVerifyIsBlueprintTypeVar(Property->GetFName(), bCheckPropertyType);
+			return TestVerifyIsBlueprintTypeVar(Property->GetFName(), PropOwner, bCheckPropertyType, bCheckPropertyVal);
 		}
 	}
 	else if (UInterfaceProperty* InterfaceProp = Cast<UInterfaceProperty>(Property))
@@ -251,7 +354,7 @@ bool ACircularDependencyTestActor::TestVerifyProperty(UProperty* Property, uint8
 		{
 			// assume any object properties that we've added to the BP 
 			// directly should be circular refs (and thus, pointers to other blueprint types) 
-			return TestVerifyIsBlueprintTypeVar(Property->GetFName(), bCheckPropertyType);
+			return TestVerifyIsBlueprintTypeVar(Property->GetFName(), PropOwner, bCheckPropertyType, bCheckPropertyVal);
 		}
 	}
 
