[INTEGRATE] Change 2422593 by [REDACTED_PERSON][REDACTED_HANDLE] on 2015/01/28 14:23:11

	Delete Blueprint Class instances when users force delete the source asset.

	change summary:
	- alert and clear transaction buffer on delete when instances are present
	- modified ForceDeleteObjects() to find and destroy BPGC instances (routing through standard editor code paths for Actor-based instances to ensure that references are cleared, panels are updated, etc.)

diff --git a/Engine/Source/Editor/UnrealEd/Private/AssetDeleteModel.cpp b/Engine/Source/Editor/UnrealEd/Private/AssetDeleteModel.cpp
--- a/Engine/Source/Editor/UnrealEd/Private/AssetDeleteModel.cpp
+++ b/Engine/Source/Editor/UnrealEd/Private/AssetDeleteModel.cpp
@@ -537,6 +537,36 @@ void FPendingDelete::CheckForReferences()
 	const int32 NonUndoReferenceCount = MemoryReferences.ExternalReferences.Num() + MemoryReferences.InternalReferences.Num();
 
 	bIsReferencedInMemoryByUndo = TotalReferenceCount > NonUndoReferenceCount;
+
+	// If the object itself isn't in the transaction buffer, check to see if it's a Blueprint asset. We might have instances of the
+	// Blueprint in the transaction buffer, in which case we also want to both alert the user and clear it prior to deleting the asset.
+	if ( !bIsReferencedInMemoryByUndo )
+	{
+		UBlueprint* Blueprint = Cast<UBlueprint>( Object );
+		if ( Blueprint && Blueprint->GeneratedClass )
+		{
+			TArray<FReferencerInformation> ExternalMemoryReferences = MemoryReferences.ExternalReferences;
+			for ( auto RefIt = ExternalMemoryReferences.CreateIterator(); RefIt && !bIsReferencedInMemoryByUndo; ++RefIt )
+			{
+				FReferencerInformation& RefInfo = *RefIt;
+				if ( RefInfo.Referencer->IsA( Blueprint->GeneratedClass ) )
+				{
+					if ( IsReferenced( RefInfo.Referencer, GARBAGE_COLLECTION_KEEPFLAGS, true, &ReferencesIncludingUndo ) )
+					{
+						GEditor->Trans->DisableObjectSerialization();
+
+						FReferencerInformationList ReferencesExcludingUndo;
+						if ( IsReferenced( RefInfo.Referencer, GARBAGE_COLLECTION_KEEPFLAGS, true, &ReferencesExcludingUndo ) )
+						{
+							bIsReferencedInMemoryByUndo = ( ReferencesIncludingUndo.InternalReferences.Num() + ReferencesIncludingUndo.ExternalReferences.Num() ) > ( ReferencesExcludingUndo.InternalReferences.Num() + ReferencesExcludingUndo.ExternalReferences.Num() );
+						}
+
+						GEditor->Trans->EnableObjectSerialization();
+					}
+				}
+			}
+		}
+	}
 }
 
 bool FPendingDelete::operator == ( const FPendingDelete& Other ) const
diff --git a/Engine/Source/Editor/UnrealEd/Private/ObjectTools.cpp b/Engine/Source/Editor/UnrealEd/Private/ObjectTools.cpp
--- a/Engine/Source/Editor/UnrealEd/Private/ObjectTools.cpp
+++ b/Engine/Source/Editor/UnrealEd/Private/ObjectTools.cpp
@@ -32,6 +32,7 @@
 #include "ComponentReregisterContext.h"
 #include "SNotificationList.h"
 #include "NotificationManager.h"
+#include "Layers/ILayers.h"
 
 DEFINE_LOG_CATEGORY_STATIC(LogObjectTools, Log, All);
 
@@ -1755,7 +1756,9 @@ namespace ObjectTools
 
 		GWarn->BeginSlowTask( NSLOCTEXT("UnrealEd", "Deleting", "Deleting"), true );
 
+		TArray<AActor*> ActorsToDelete;
 		TArray<UObject*> ObjectsToDelete;
+		bool bNeedsGarbageCollection = false;
 	
 		// Clear audio components to allow previewed sounds to be consolidated
 		GEditor->ClearPreviewComponents();
@@ -1767,6 +1770,65 @@ namespace ObjectTools
 			GEditor->GetSelectedObjects()->Deselect( CurrentObject );
 
 			ObjectsToDelete.Add( CurrentObject );
+
+			// If the object about to be deleted is a Blueprint asset, make sure that any instances of the Blueprint class get deleted as well
+			UBlueprint *BlueprintObject = Cast<UBlueprint>(CurrentObject);
+			if ( BlueprintObject && BlueprintObject->GeneratedClass && BlueprintObject->GeneratedClass->ClassDefaultObject )
+			{
+				TArray<UObject*> InstancesToDelete;
+				BlueprintObject->GeneratedClass->ClassDefaultObject->GetArchetypeInstances( InstancesToDelete );
+
+				for ( TArray<UObject*>::TConstIterator InstanceItr( InstancesToDelete ); InstanceItr; ++InstanceItr )
+				{
+					UObject* CurrentInstance = *InstanceItr;
+
+					AActor* CurrentInstanceAsActor = Cast<AActor>( CurrentInstance );
+					if ( CurrentInstanceAsActor )
+					{
+						ActorsToDelete.Add( CurrentInstanceAsActor );
+					}
+					else
+					{
+						ObjectsToDelete.Add( CurrentInstance );
+					}
+				}
+			}
+		}
+
+		// Destroy all Actor instances
+		if ( ActorsToDelete.Num() > 0 )
+		{
+			bool bSelectionChanged = false;
+			for ( TArray<AActor*>::TConstIterator ActorItr( ActorsToDelete ); ActorItr; ++ActorItr )
+			{
+				AActor* CurActor = *ActorItr;
+
+				// Skip if already pending GC
+				if ( !CurActor->IsPendingKill() )
+				{
+					// Deselect if active
+					USelection* SelectedActors = GEditor->GetSelectedActors();
+					if ( SelectedActors && CurActor->IsSelected() )
+					{
+						SelectedActors->Deselect( CurActor );
+
+						bSelectionChanged = true;
+					}
+
+					// Destroy the Actor instance. This is similar to edactDeleteSelected(), but we don't request user confirmation here.
+					GEditor->Layers->DisassociateActorFromLayers( CurActor );
+					GEditor->GetEditorWorldContext().World()->EditorDestroyActor( CurActor, false );
+
+					bNeedsGarbageCollection = true;
+				}
+
+				GWarn->StatusUpdate( ActorItr.GetIndex(), ActorsToDelete.Num(), NSLOCTEXT( "UnrealEd", "ConsolidateAssetsUpdate_DeletingInstances", "Deleting Instances..." ) );
+			}
+
+			if ( bSelectionChanged )
+			{
+				GEditor->NoteSelectionChange();
+			}
 		}
 
 		// If the current editor world is in this list, transition to a new map and reload the world to finish the delete
@@ -1775,7 +1837,6 @@ namespace ObjectTools
 		{
 			int32 ReplaceableObjectsNum = 0;
 			{
-				bool bNeedsGarbageCollection = false;
 				TArray<UObject*> ObjectsToReplace = ObjectsToDelete;
 				for (TArray<UObject*>::TIterator ObjectItr(ObjectsToReplace); ObjectItr; ++ObjectItr)
 				{
