[INTEGRATE] Change 2422948 by [REDACTED_PERSON][REDACTED_HANDLE] on 2015/01/28 17:58:51

	[UE-8490] - With an actor selected, you must now double-click to select a component. Once a component is selected, single clicks will select components and a double-click will select the actor.

diff --git a/Engine/Source/Editor/UnrealEd/Private/LevelEditorViewport.cpp b/Engine/Source/Editor/UnrealEd/Private/LevelEditorViewport.cpp
--- a/Engine/Source/Editor/UnrealEd/Private/LevelEditorViewport.cpp
+++ b/Engine/Source/Editor/UnrealEd/Private/LevelEditorViewport.cpp
@@ -1794,19 +1794,13 @@ void FLevelEditorViewportClient::ProcessClick(FSceneView& View, HHitProxy* HitPr
 			// We want to process the click on the component only if:
 			// 1. The actor clicked is already selected
 			// 2. The actor selected is the only actor selected
-			// 3. The LMB was pressed or we already have a component selected (so that right-clicking a selected actor won't select a component)
-			// And, if a component is not selected already:
-			// 4. The click was not a double click
-			// 5. The level viewport didn't just receive focus this frame
-			// 6. No modifier keys are pressed
+			// 3. No components are already selected and the click was a double click
+			// 4. OR, a component is already selected and the click was NOT a double click
 			const bool bActorAlreadySelectedExclusively = GEditor->GetSelectedActors()->IsSelected(ActorHitProxy->Actor) && ( GEditor->GetSelectedActorCount() == 1 );
 			const bool bComponentAlreadySelected = GEditor->GetSelectedComponentCount() > 0;
-			const bool bCanBeginSelectingComponents =  (Click.GetKey() == EKeys::LeftMouseButton) 
-													&& (Click.GetEvent() != IE_DoubleClick) 
-													&& !bReceivedFocusRecently
-													&& !( Click.IsAltDown() || Click.IsControlDown() || Click.IsShiftDown() );
+			const bool bWasDoubleClick = ( Click.GetEvent() == IE_DoubleClick );
 
-			const bool bSelectComponent = bActorAlreadySelectedExclusively && ( bComponentAlreadySelected || bCanBeginSelectingComponents );
+			const bool bSelectComponent = bActorAlreadySelectedExclusively && (bComponentAlreadySelected != bWasDoubleClick);
 
 			if (bSelectComponent && GetDefault<UEditorExperimentalSettings>()->bInWorldBPEditing)
 			{
