[INTEGRATE] Change 2423078 by [REDACTED_PERSON][REDACTED_HANDLE] on 2015/01/28 19:47:09

	Need to check all objects, not just components when determining whether a name can be used

diff --git a/Engine/Source/Editor/UnrealEd/Private/Kismet2/ComponentEditorUtils.cpp b/Engine/Source/Editor/UnrealEd/Private/Kismet2/ComponentEditorUtils.cpp
--- a/Engine/Source/Editor/UnrealEd/Private/Kismet2/ComponentEditorUtils.cpp
+++ b/Engine/Source/Editor/UnrealEd/Private/Kismet2/ComponentEditorUtils.cpp
@@ -114,26 +114,10 @@ bool FComponentEditorUtils::IsValidVariableNameString(const UActorComponent* InC
 	return bIsValid;
 }
 
-bool FComponentEditorUtils::IsComponentNameAvailable(const FString& InString, const AActor* ComponentOwner, const UActorComponent* ComponentToIgnore)
+bool FComponentEditorUtils::IsComponentNameAvailable(const FString& InString, AActor* ComponentOwner, const UActorComponent* ComponentToIgnore)
 {
-	bool bNameIsAvailable = true;
-
-	if (ComponentOwner != nullptr)
-	{
-		TInlineComponentArray<UActorComponent*> Components;
-		ComponentOwner->GetComponents(Components);
-
-		for (auto Component : Components)
-		{
-			if (Component != ComponentToIgnore && Component->GetName() == InString)
-			{
-				bNameIsAvailable = false;
-				break;
-			}
-		}
-	}
-
-	return bNameIsAvailable;
+	UObject* ExistingObject = FindObject<UObject>(ComponentOwner, *InString);
+	return (ExistingObject == nullptr || ExistingObject == ComponentToIgnore);
 }
 
 FString FComponentEditorUtils::GenerateValidVariableName(TSubclassOf<UActorComponent> ComponentClass, AActor* ComponentOwner)
diff --git a/Engine/Source/Editor/UnrealEd/Public/Kismet2/ComponentEditorUtils.h b/Engine/Source/Editor/UnrealEd/Public/Kismet2/ComponentEditorUtils.h
--- a/Engine/Source/Editor/UnrealEd/Public/Kismet2/ComponentEditorUtils.h
+++ b/Engine/Source/Editor/UnrealEd/Public/Kismet2/ComponentEditorUtils.h
@@ -23,7 +23,7 @@ public:
 	 * Optionally excludes an existing component from the check (ex. a component currently being renamed)
 	 * @return True if the InString is an available name for a component of ComponentOwner
 	 */
-	static bool IsComponentNameAvailable(const FString& InString, const AActor* ComponentOwner, const UActorComponent* ComponentToIgnore = nullptr);
+	static bool IsComponentNameAvailable(const FString& InString, AActor* ComponentOwner, const UActorComponent* ComponentToIgnore = nullptr);
 		
 	/** Generate a valid variable name string for the given component instance */
 	static FString GenerateValidVariableName(TSubclassOf<UActorComponent> InComponentClass, AActor* ComponentOwner);
