[INTEGRATE] Change 2423160 by [REDACTED_PERSON][REDACTED_HANDLE] on 2015/01/28 21:01:24

	Added ability to spawn  basic shapes from the placement browser or component dropdown

diff --git a/Engine/Source/Editor/Kismet/Private/SSCSEditor.cpp b/Engine/Source/Editor/Kismet/Private/SSCSEditor.cpp
--- a/Engine/Source/Editor/Kismet/Private/SSCSEditor.cpp
+++ b/Engine/Source/Editor/Kismet/Private/SSCSEditor.cpp
@@ -3434,10 +3434,12 @@ void SSCSEditor::OnGetChildrenForTree( FSCSEditorTreeNodePtrType InNodePtr, TArr
 }
 
 
-void SSCSEditor::PerformComboAddClass(TSubclassOf<UActorComponent> ComponentClass, EComponentCreateAction::Type ComponentCreateAction )
+UActorComponent* SSCSEditor::PerformComboAddClass(TSubclassOf<UActorComponent> ComponentClass, EComponentCreateAction::Type ComponentCreateAction, UObject* AssetOverride)
 {
 	UClass* NewClass = ComponentClass;
 
+	UActorComponent* NewComponent = nullptr;
+
 	if( ComponentCreateAction == EComponentCreateAction::CreateNewCPPClass )
 	{
 		NewClass = CreateNewCPPComponent( ComponentClass );
@@ -3457,7 +3459,7 @@ void SSCSEditor::PerformComboAddClass(TSubclassOf<UActorComponent> ComponentClas
 		// This adds components according to the type selected in the drop down. If the user
 		// has the appropriate objects selected in the content browser then those are added,
 		// else we go down the previous route of adding components by type.
-		if(Selection->Num() > 0)
+		if(Selection->Num() > 0 && !AssetOverride)
 		{
 			for(FSelectionIterator ObjectIter(*Selection); ObjectIter; ++ObjectIter)
 			{
@@ -3471,7 +3473,7 @@ void SSCSEditor::PerformComboAddClass(TSubclassOf<UActorComponent> ComponentClas
 				{
 					if(ComponentClasses[ComponentIndex]->IsChildOf(NewClass))
 					{
-						AddNewComponent(NewClass, Object);
+						NewComponent = AddNewComponent(NewClass, Object);
 						bAddedComponent = true;
 						break;
 					}
@@ -3482,9 +3484,11 @@ void SSCSEditor::PerformComboAddClass(TSubclassOf<UActorComponent> ComponentClas
 		if(!bAddedComponent)
 		{
 			// As the SCS splits up the scene and actor components, can now add directly
-			AddNewComponent(NewClass, NULL);
+			NewComponent = AddNewComponent(NewClass, AssetOverride);
 		}
 	}
+
+	return NewComponent;
 }
 
 TArray<FSCSEditorTreeNodePtrType>  SSCSEditor::GetSelectedNodes() const
diff --git a/Engine/Source/Editor/Kismet/Public/SSCSEditor.h b/Engine/Source/Editor/Kismet/Public/SSCSEditor.h
--- a/Engine/Source/Editor/Kismet/Public/SSCSEditor.h
+++ b/Engine/Source/Editor/Kismet/Public/SSCSEditor.h
@@ -656,7 +656,7 @@ protected:
 	FString GetSelectedClassText() const;
 
 	/** Add a component from the selection in the combo box */
-	void PerformComboAddClass(TSubclassOf<UActorComponent> ComponentClass, EComponentCreateAction::Type ComponentCreateAction );
+	UActorComponent* PerformComboAddClass(TSubclassOf<UActorComponent> ComponentClass, EComponentCreateAction::Type ComponentCreateAction, UObject* AssetOverride);
 
 	/** Called to display context menu when right clicking on the widget */
 	TSharedPtr< SWidget > CreateContextMenu();
diff --git a/Engine/Source/Editor/PlacementMode/Private/SPlacementModeTools.cpp b/Engine/Source/Editor/PlacementMode/Private/SPlacementModeTools.cpp
--- a/Engine/Source/Editor/PlacementMode/Private/SPlacementModeTools.cpp
+++ b/Engine/Source/Editor/PlacementMode/Private/SPlacementModeTools.cpp
@@ -49,18 +49,20 @@ public:
 
 	SLATE_END_ARGS()
 
-	void Construct( const FArguments& InArgs, const FAssetData& InAsset )
+	void Construct( const FArguments& InArgs, const FAssetData& InAsset, bool bAlwaysUseClassThumbnail )
 	{
 		Asset = InAsset;
 
 		FLevelEditorModule& LevelEditorModule = FModuleManager::LoadModuleChecked<FLevelEditorModule>( "LevelEditor" );
 		TSharedPtr<FAssetThumbnailPool> ThumbnailPool = LevelEditorModule.GetFirstLevelEditor()->GetThumbnailPool();
 
-		Thumbnail = MakeShareable( new FAssetThumbnail( Asset, InArgs._Width, InArgs._Height, ThumbnailPool ) );
+		Thumbnail = MakeShareable(new FAssetThumbnail(Asset, InArgs._Width, InArgs._Height, ThumbnailPool));
 
+		FAssetThumbnailConfig Config;
+		Config.bForceGenericThumbnail = bAlwaysUseClassThumbnail;
 		ChildSlot
 		[
-			Thumbnail->MakeThumbnailWidget()
+			Thumbnail->MakeThumbnailWidget( Config )
 		];
 	}
 
@@ -165,7 +167,7 @@ void SPlacementAssetEntry::Construct(const FArguments& InArgs, UActorFactory* In
 					.WidthOverride( 35 )
 					.HeightOverride( 35 )
 					[
-						SNew( SPlacementAssetThumbnail, AssetData )
+						SNew( SPlacementAssetThumbnail, AssetData, InArgs._UseClassThumbnailForAsset )
 					]
 				]
 			]
@@ -567,10 +569,10 @@ TSharedRef< SWidget > BuildDraggableAssetWidget( UClass* InAssetClass )
 	return SNew( SPlacementAssetEntry, Factory, AssetData );
 }
 
-TSharedRef< SWidget > BuildDraggableAssetWidget( UClass* InAssetClass, const FAssetData& InAssetData )
+TSharedRef< SWidget > BuildDraggableAssetWidget( UClass* InAssetClass, const FAssetData& InAssetData, bool bUseClassThumbnail )
 {
 	UActorFactory* Factory = GEditor->FindActorFactoryByClass( InAssetClass );
-	return SNew( SPlacementAssetEntry, Factory, InAssetData );
+	return SNew( SPlacementAssetEntry, Factory, InAssetData ).UseClassThumbnailForAsset( bUseClassThumbnail );
 }
 
 TSharedRef< SWidget > SPlacementModeTools::BuildLightsWidget()
@@ -665,6 +667,30 @@ TSharedRef< SWidget > SPlacementModeTools::BuildBasicWidget()
 	]
 	+ SVerticalBox::Slot()
 	.AutoHeight()
+	[
+		// Cube
+		BuildDraggableAssetWidget(UActorFactoryBasicShape::StaticClass(), FAssetData( LoadObject<UStaticMesh>(nullptr,TEXT("/Engine/BasicShapes/Cube.Cube"))), true )
+	]
+	+ SVerticalBox::Slot()
+	.AutoHeight()
+	[
+		// Sphere
+		BuildDraggableAssetWidget(UActorFactoryBasicShape::StaticClass(), FAssetData(LoadObject<UStaticMesh>(nullptr,TEXT("/Engine/BasicShapes/Sphere.Sphere"))), true )
+	]
+	+ SVerticalBox::Slot()
+	.AutoHeight()
+	[
+		// Cylinder
+		BuildDraggableAssetWidget(UActorFactoryBasicShape::StaticClass(), FAssetData(LoadObject<UStaticMesh>(nullptr,TEXT("/Engine/BasicShapes/Cylinder.Cylinder"))), true )
+	]
+	+ SVerticalBox::Slot()
+	.AutoHeight()
+	[
+		// Cone
+		BuildDraggableAssetWidget(UActorFactoryBasicShape::StaticClass(), FAssetData(LoadObject<UStaticMesh>(nullptr,TEXT("/Engine/BasicShapes/Cone.Cone"))), true )
+	]
+	+ SVerticalBox::Slot()
+	.AutoHeight()
 	[
 		BuildDraggableAssetWidget( UActorFactoryTriggerBox::StaticClass() )
 	]
diff --git a/Engine/Source/Editor/PlacementMode/Private/SPlacementModeTools.h b/Engine/Source/Editor/PlacementMode/Private/SPlacementModeTools.h
--- a/Engine/Source/Editor/PlacementMode/Private/SPlacementModeTools.h
+++ b/Engine/Source/Editor/PlacementMode/Private/SPlacementModeTools.h
@@ -15,13 +15,15 @@ public:
 
 	SLATE_BEGIN_ARGS(SPlacementAssetEntry)
 		: _LabelOverride()
+		, _UseClassThumbnailForAsset(false)
 	{}
 	
 	/** Highlight this text in the text block */
 	SLATE_ATTRIBUTE(FText, HighlightText)
 
 	SLATE_ARGUMENT(FText, LabelOverride)
-	
+
+	SLATE_ARGUMENT(bool, UseClassThumbnailForAsset)
 	SLATE_END_ARGS()
 
 	void Construct(const FArguments& InArgs, UActorFactory* InFactory, const FAssetData& InAsset);
diff --git a/Engine/Source/Editor/UnrealEd/Classes/ActorFactories/ActorFactoryBasicShape.h b/Engine/Source/Editor/UnrealEd/Classes/ActorFactories/ActorFactoryBasicShape.h
new file mode 100644
--- /dev/null
+++ b/Engine/Source/Editor/UnrealEd/Classes/ActorFactories/ActorFactoryBasicShape.h
@@ -0,0 +1,13 @@
+// Copyright 1998-2015 Epic Games, Inc. All Rights Reserved.
+
+
+#pragma once
+#include "ActorFactoryBasicShape.generated.h"
+
+UCLASS(MinimalAPI,config=Editor)
+class UActorFactoryBasicShape : public UActorFactoryStaticMesh
+{
+	GENERATED_UCLASS_BODY()
+
+	virtual void PostSpawnActor(UObject* Asset, AActor* NewActor) override;
+};
diff --git a/Engine/Source/Editor/UnrealEd/Private/Factories/ActorFactory.cpp b/Engine/Source/Editor/UnrealEd/Private/Factories/ActorFactory.cpp
--- a/Engine/Source/Editor/UnrealEd/Private/Factories/ActorFactory.cpp
+++ b/Engine/Source/Editor/UnrealEd/Private/Factories/ActorFactory.cpp
@@ -326,6 +326,38 @@ FQuat UActorFactoryStaticMesh::AlignObjectToSurfaceNormal(const FVector& InSurfa
 	return FindActorAlignmentRotation(ActorRotation, FVector(0.f, 0.f, 1.f), InSurfaceNormal);
 }
 
+/*-----------------------------------------------------------------------------
+UActorFactoryBasicShape
+-----------------------------------------------------------------------------*/
+UActorFactoryBasicShape::UActorFactoryBasicShape(const FObjectInitializer& ObjectInitializer)
+	: Super(ObjectInitializer)
+{
+	DisplayName = LOCTEXT("UActorFactoryBasicShapeDisplayName", "Basic Shape");
+	NewActorClass = AStaticMeshActor::StaticClass();
+	bUseSurfaceOrientation = true;
+}
+
+void UActorFactoryBasicShape::PostSpawnActor(UObject* Asset, AActor* NewActor)
+{
+	// Change properties
+	UStaticMesh* StaticMesh = CastChecked<UStaticMesh>(Asset);
+	GEditor->SetActorLabelUnique(NewActor, StaticMesh->GetName());
+
+	AStaticMeshActor* StaticMeshActor = CastChecked<AStaticMeshActor>(NewActor);
+	UStaticMeshComponent* StaticMeshComponent = StaticMeshActor->GetStaticMeshComponent();
+
+	if( StaticMeshComponent )
+	{
+		StaticMeshComponent->UnregisterComponent();
+
+		StaticMeshComponent->StaticMesh = StaticMesh;
+		StaticMeshComponent->StaticMeshDerivedDataKey = StaticMesh->RenderData->DerivedDataKey;
+		StaticMeshComponent->SetMaterial(0, LoadObject<UMaterial>(nullptr, TEXT("/Engine/BasicShapes/BasicShapeMaterial.BasicShapeMaterial")));
+		// Init Component
+		StaticMeshComponent->RegisterComponent();
+	}
+}
+
 /*-----------------------------------------------------------------------------
 UActorFactoryDeferredDecal
 -----------------------------------------------------------------------------*/
diff --git a/Engine/Source/Editor/UnrealEd/Private/SComponentClassCombo.cpp b/Engine/Source/Editor/UnrealEd/Private/SComponentClassCombo.cpp
--- a/Engine/Source/Editor/UnrealEd/Private/SComponentClassCombo.cpp
+++ b/Engine/Source/Editor/UnrealEd/Private/SComponentClassCombo.cpp
@@ -128,7 +128,7 @@ void SComponentClassCombo::GenerateFilteredComponentList(const FString& InSearch
 
 			if (CurrentEntry->IsClass() && CurrentEntry->IsIncludedInFilter())
 			{
-				FString FriendlyComponentName = GetSanitizedComponentName( CurrentEntry->GetComponentClass() );
+				FString FriendlyComponentName = GetSanitizedComponentName( CurrentEntry->GetComponentClass(), CurrentEntry->GetComponentNameOverride() );
 
 				if ( FriendlyComponentName.Contains( InSearchText, ESearchCase::IgnoreCase ) )
 				{
@@ -176,10 +176,17 @@ void SComponentClassCombo::OnAddComponentSelectionChanged( FComponentClassComboE
 
 		if ( InItem->IsClass() )
 		{
-			OnComponentClassSelected.ExecuteIfBound(InItem->GetComponentClass(), InItem->GetComponentCreateAction() );
-
 			// Neither do we want the combo dropdown staying open once the user has clicked on a valid option
-			SetIsOpen( false, false );
+			SetIsOpen(false, false);
+
+			if( OnComponentClassSelected.IsBound() )
+			{
+				UActorComponent* NewActorComponent = OnComponentClassSelected.Execute(InItem->GetComponentClass(), InItem->GetComponentCreateAction(), InItem->GetAssetOverride());
+				if(NewActorComponent)
+				{
+					InItem->GetOnComponentCreated().ExecuteIfBound(NewActorComponent);
+				}
+			}
 		}
 	}
 	else if ( InItem.IsValid() && SelectInfo != ESelectInfo::OnMouseClick )
@@ -318,6 +325,61 @@ struct SortComboEntry
 };
 const FString SortComboEntry::CommonClassGroup(TEXT("Common"));
 
+void SComponentClassCombo::OnBasicShapeCreated(UActorComponent* Component)
+{
+	UStaticMeshComponent* SMC = Cast<UStaticMeshComponent>(Component);
+	if(SMC)
+	{
+		SMC->SetMaterial(0, LoadObject<UMaterial>(nullptr, TEXT("/Engine/BasicShapes/BasicShapeMaterial.BasicShapeMaterial")));
+	}
+};
+
+void SComponentClassCombo::AddBasicShapeComponents( TArray<FComponentClassComboEntryPtr>& SortedClassList )
+{
+	FString BasicShapesHeading = LOCTEXT("BasicShapesHeading", "Basic Shapes").ToString();
+
+	{
+		FComponentEntryCustomizationArgs Args;
+		Args.AssetOverride = LoadObject<UStaticMesh>(nullptr, TEXT("/Engine/BasicShapes/Cube.Cube"));;
+		Args.OnComponentCreated = FOnComponentCreated::CreateSP(this, &SComponentClassCombo::OnBasicShapeCreated);
+		Args.ComponentNameOverride = LOCTEXT("BasicCubeShapeDisplayName", "Cube").ToString();
+
+		FComponentClassComboEntryPtr NewShape = MakeShareable(new FComponentClassComboEntry(BasicShapesHeading, UStaticMeshComponent::StaticClass(), true, EComponentCreateAction::SpawnExistingClass, Args));
+		SortedClassList.Add(NewShape);
+	}
+
+	{
+		FComponentEntryCustomizationArgs Args;
+		Args.AssetOverride = LoadObject<UStaticMesh>(nullptr, TEXT("/Engine/BasicShapes/Sphere.Sphere"));;
+		Args.OnComponentCreated = FOnComponentCreated::CreateSP(this, &SComponentClassCombo::OnBasicShapeCreated);
+		Args.ComponentNameOverride = LOCTEXT("BasicSphereShapeDisplayName", "Sphere").ToString();
+
+		FComponentClassComboEntryPtr NewShape = MakeShareable(new FComponentClassComboEntry(BasicShapesHeading, UStaticMeshComponent::StaticClass(), true, EComponentCreateAction::SpawnExistingClass, Args));
+		SortedClassList.Add(NewShape);
+	}
+
+	{
+		FComponentEntryCustomizationArgs Args;
+		Args.AssetOverride = LoadObject<UStaticMesh>(nullptr, TEXT("/Engine/BasicShapes/Cylinder.Cylinder"));;
+		Args.OnComponentCreated = FOnComponentCreated::CreateSP(this, &SComponentClassCombo::OnBasicShapeCreated);
+		Args.ComponentNameOverride = LOCTEXT("BasicCylinderShapeDisplayName", "Cylinder").ToString();
+
+		FComponentClassComboEntryPtr NewShape = MakeShareable(new FComponentClassComboEntry(BasicShapesHeading, UStaticMeshComponent::StaticClass(), true, EComponentCreateAction::SpawnExistingClass, Args));
+		SortedClassList.Add(NewShape);
+	}
+
+	{
+		FComponentEntryCustomizationArgs Args;
+		Args.AssetOverride = LoadObject<UStaticMesh>(nullptr, TEXT("/Engine/BasicShapes/Cone.Cone"));;
+		Args.OnComponentCreated = FOnComponentCreated::CreateSP(this, &SComponentClassCombo::OnBasicShapeCreated);
+		Args.ComponentNameOverride = LOCTEXT("BasicConeShapeDisplayName", "Cone").ToString();
+
+		FComponentClassComboEntryPtr NewShape = MakeShareable(new FComponentClassComboEntry(BasicShapesHeading, UStaticMeshComponent::StaticClass(), true, EComponentCreateAction::SpawnExistingClass, Args));
+		SortedClassList.Add(NewShape);
+	}
+
+}
+
 void SComponentClassCombo::UpdateComponentClassList()
 {
 	ComponentClassList.Empty();
@@ -365,6 +427,8 @@ void SComponentClassCombo::UpdateComponentClassList()
 		}
 	}
 
+	AddBasicShapeComponents(SortedClassList);
+
 	if (SortedClassList.Num() > 0)
 	{
 		Sort(SortedClassList.GetData(), SortedClassList.Num(), SortComboEntry());
@@ -419,75 +483,84 @@ FText SComponentClassCombo::GetFriendlyComponentName(FComponentClassComboEntryPt
 	}
 	else
 	{
-		FriendlyComponentName = GetSanitizedComponentName(Entry->GetComponentClass());
-
+		FriendlyComponentName = GetSanitizedComponentName(Entry->GetComponentClass(), Entry->GetComponentNameOverride() );
 
-		// Search the selected assets and look for any that can be used as a source asset for this type of component
-		// If there is one we append the asset name to the component name, if there are many we append "Multiple Assets"
-		FString AssetName;
-		UObject* PreviousMatchingAsset = NULL;
-
-		FEditorDelegates::LoadSelectedAssetsIfNeeded.Broadcast();
-		USelection* Selection = GEditor->GetSelectedObjects();
-		for(FSelectionIterator ObjectIter(*Selection); ObjectIter; ++ObjectIter)
+		if( Entry->GetComponentNameOverride().IsEmpty() )
 		{
-			UObject* Object = *ObjectIter;
-			UClass* Class = Object->GetClass();
-
-			TArray<TSubclassOf<UActorComponent> > ComponentClasses = FComponentAssetBrokerage::GetComponentsForAsset(Object);
-			for(int32 ComponentIndex = 0; ComponentIndex < ComponentClasses.Num(); ComponentIndex++)
+			// Search the selected assets and look for any that can be used as a source asset for this type of component
+			// If there is one we append the asset name to the component name, if there are many we append "Multiple Assets"
+			FString AssetName;
+			UObject* PreviousMatchingAsset = NULL;
+
+			FEditorDelegates::LoadSelectedAssetsIfNeeded.Broadcast();
+			USelection* Selection = GEditor->GetSelectedObjects();
+			for(FSelectionIterator ObjectIter(*Selection); ObjectIter; ++ObjectIter)
 			{
-				if(ComponentClasses[ComponentIndex]->IsChildOf(Entry->GetComponentClass()))
+				UObject* Object = *ObjectIter;
+				UClass* Class = Object->GetClass();
+
+				TArray<TSubclassOf<UActorComponent> > ComponentClasses = FComponentAssetBrokerage::GetComponentsForAsset(Object);
+				for(int32 ComponentIndex = 0; ComponentIndex < ComponentClasses.Num(); ComponentIndex++)
 				{
-					if(AssetName.IsEmpty())
+					if(ComponentClasses[ComponentIndex]->IsChildOf(Entry->GetComponentClass()))
 					{
-						// If there is no previous asset then we just accept the name
-						AssetName = Object->GetName();
-						PreviousMatchingAsset = Object;
-					}
-					else
-					{
-						// if there is a previous asset then check that we didn't just find multiple appropriate components
-						// in a single asset - if the asset differs then we don't display the name, just "Multiple Assets"
-						if(PreviousMatchingAsset != Object)
+						if(AssetName.IsEmpty())
 						{
-							AssetName = LOCTEXT("MultipleAssetsForComponentAnnotation", "Multiple Assets").ToString();
+							// If there is no previous asset then we just accept the name
+							AssetName = Object->GetName();
 							PreviousMatchingAsset = Object;
 						}
+						else
+						{
+							// if there is a previous asset then check that we didn't just find multiple appropriate components
+							// in a single asset - if the asset differs then we don't display the name, just "Multiple Assets"
+							if(PreviousMatchingAsset != Object)
+							{
+								AssetName = LOCTEXT("MultipleAssetsForComponentAnnotation", "Multiple Assets").ToString();
+								PreviousMatchingAsset = Object;
+							}
+						}
 					}
 				}
 			}
-		}
 
-		if(!AssetName.IsEmpty())
-		{
-			FriendlyComponentName += FString(" (") + AssetName + FString(")");
+			if(!AssetName.IsEmpty())
+			{
+				FriendlyComponentName += FString(" (") + AssetName + FString(")");
+			}
 		}
 	}
 	return FText::FromString(FriendlyComponentName);
 }
 
-FString SComponentClassCombo::GetSanitizedComponentName( UClass* ComponentClass )
+FString SComponentClassCombo::GetSanitizedComponentName( UClass* ComponentClass, const FString& InComponentNameOverride )
 {
-	FString DisplayName;
-	if (ComponentClass->HasMetaData(TEXT("DisplayName")))
+	if( InComponentNameOverride.IsEmpty() )
 	{
-		DisplayName = ComponentClass->GetMetaData(TEXT("DisplayName"));
+		FString DisplayName;
+		if(ComponentClass->HasMetaData(TEXT("DisplayName")))
+		{
+			DisplayName = ComponentClass->GetMetaData(TEXT("DisplayName"));
+		}
+		else
+		{
+			DisplayName = ComponentClass->GetName();
+			DisplayName = DisplayName.Replace(TEXT("Component"), TEXT(""), ESearchCase::IgnoreCase);
+		}
+
+		// Purge the _C for BP Components
+		if(DisplayName.EndsWith(TEXT("_C")))
+		{
+			const int32 NewLen = DisplayName.Len() - 2;
+			DisplayName = DisplayName.Left(NewLen);
+		}
+
+		return FName::NameToDisplayString(DisplayName, false);
 	}
 	else
 	{
-		DisplayName = ComponentClass->GetName();
-		DisplayName = DisplayName.Replace( TEXT("Component"), TEXT(""), ESearchCase::IgnoreCase );
+		return InComponentNameOverride;
 	}
-
-	// Purge the _C for BP Components
-	if (DisplayName.EndsWith(TEXT("_C")))
-	{
-		const int32 NewLen = DisplayName.Len() - 2;
-		DisplayName = DisplayName.Left(NewLen);
-	}
-
-	return FName::NameToDisplayString(DisplayName, false);
 }
 
 #undef LOCTEXT_NAMESPACE
diff --git a/Engine/Source/Editor/UnrealEd/Public/SComponentClassCombo.h b/Engine/Source/Editor/UnrealEd/Public/SComponentClassCombo.h
--- a/Engine/Source/Editor/UnrealEd/Public/SComponentClassCombo.h
+++ b/Engine/Source/Editor/UnrealEd/Public/SComponentClassCombo.h
@@ -19,16 +19,37 @@ namespace EComponentCreateAction
 	};
 }
 
-DECLARE_DELEGATE_TwoParams(FComponentClassSelected, TSubclassOf<UActorComponent>, EComponentCreateAction::Type );
 
+DECLARE_DELEGATE_OneParam(FOnComponentCreated, UActorComponent*);
+
+DECLARE_DELEGATE_RetVal_ThreeParams( UActorComponent*, FComponentClassSelected, TSubclassOf<UActorComponent>, EComponentCreateAction::Type, UObject*);
+
+struct FComponentEntryCustomizationArgs
+{
+	/** Specific asset to use instead of the selected asset in the content browser */
+	TWeakObjectPtr<UObject> AssetOverride;
+	/** Custom name to display */
+	FString ComponentNameOverride;
+	/** Callback when a new component is created */
+	FOnComponentCreated OnComponentCreated;
+
+	FComponentEntryCustomizationArgs()
+		: AssetOverride( nullptr )
+		, ComponentNameOverride()
+		, OnComponentCreated()
+	{
+	
+	}
+};
 class FComponentClassComboEntry: public TSharedFromThis<FComponentClassComboEntry>
 {
 public:
-	FComponentClassComboEntry( const FString& InHeadingText, TSubclassOf<UActorComponent> InComponentClass, bool InIncludedInFilter, EComponentCreateAction::Type InComponentCreateAction )
+	FComponentClassComboEntry( const FString& InHeadingText, TSubclassOf<UActorComponent> InComponentClass, bool InIncludedInFilter, EComponentCreateAction::Type InComponentCreateAction, FComponentEntryCustomizationArgs InCustomizationArgs = FComponentEntryCustomizationArgs() )
 		: ComponentClass(InComponentClass)
 		, HeadingText(InHeadingText)
 		, bIncludedInFilter(InIncludedInFilter)
 		, ComponentCreateAction(InComponentCreateAction)
+		, CustomizationArgs(InCustomizationArgs)
 	{}
 
 	FComponentClassComboEntry( const FString& InHeadingText )
@@ -65,16 +86,31 @@ public:
 		return bIncludedInFilter;
 	}
 	
+	const FString& GetComponentNameOverride() const
+	{
+		return CustomizationArgs.ComponentNameOverride;
+	}
+
 	EComponentCreateAction::Type GetComponentCreateAction() const
 	{
 		return ComponentCreateAction;
 	}
 
+	FOnComponentCreated& GetOnComponentCreated()
+	{
+		return CustomizationArgs.OnComponentCreated;
+	}
+
+	UObject* GetAssetOverride()
+	{
+		return CustomizationArgs.AssetOverride.Get();
+	}
 private:
 	TSubclassOf<UActorComponent> ComponentClass;
 	FString HeadingText;
 	bool bIncludedInFilter;
 	EComponentCreateAction::Type ComponentCreateAction;
+	FComponentEntryCustomizationArgs CustomizationArgs;
 };
 
 //////////////////////////////////////////////////////////////////////////
@@ -121,13 +157,18 @@ public:
 	void UpdateComponentClassList();
 
 	/** Returns a component name without the substring "Component" and sanitized for display */
-	static FString GetSanitizedComponentName( UClass* ComponentClass );
+	static FString GetSanitizedComponentName( UClass* ComponentClass, const FString& DisplayNameOverride );
 
 protected:
 
 	/** Called when a project is hot reloaded to refresh the components list */
 	void OnProjectHotReloaded( bool bWasTriggeredAutomatically );
 
+	/** Adds basic shapes to the component dropdown */
+	void AddBasicShapeComponents( TArray<FComponentClassComboEntryPtr>& SortedClassList );
+
+	/** Called when a new basic shape is added */
+	void OnBasicShapeCreated( UActorComponent* CreatedComponent);
 private:
 
 	FText GetFriendlyComponentName(FComponentClassComboEntryPtr Entry) const;
