[INTEGRATE] Change 2423602 by [REDACTED_PERSON][REDACTED_HANDLE].7 on 2015/01/29 06:12:36

	Suppress popup progress dialogs when importing new files

diff --git a/Engine/Source/Editor/UnrealEd/Private/AutoReimport/AutoReimportManager.cpp b/Engine/Source/Editor/UnrealEd/Private/AutoReimport/AutoReimportManager.cpp
--- a/Engine/Source/Editor/UnrealEd/Private/AutoReimport/AutoReimportManager.cpp
+++ b/Engine/Source/Editor/UnrealEd/Private/AutoReimport/AutoReimportManager.cpp
@@ -270,6 +270,9 @@ FText FAutoReimportManager::GetProgressText() const
 
 TOptional<ECurrentState> FAutoReimportManager::ProcessAdditions(const FTimeLimit& TimeLimit)
 {
+	// Override the global feedback context while we do this to avoid popping up dialogs
+	TGuardValue<FFeedbackContext*> ScopedContextOverride(GWarn, FeedbackContextOverride.Get());
+	
 	FeedbackContextOverride->GetContent()->SetMainText(GetProgressText());
 
 	TMap<FString, TArray<UFactory*>> Factories;
