[INTEGRATE] Change 2423603 by [REDACTED_PERSON][REDACTED_HANDLE].7 on 2015/01/29 06:13:50

	Source content residing in directories external to a mounted path is no longer deleted when deleting assets.

	We also explicitly don't auto-add and delete assets when source content in such directories is added/removed.

diff --git a/Engine/Source/Editor/UnrealEd/Private/AssetDeleteModel.cpp b/Engine/Source/Editor/UnrealEd/Private/AssetDeleteModel.cpp
--- a/Engine/Source/Editor/UnrealEd/Private/AssetDeleteModel.cpp
+++ b/Engine/Source/Editor/UnrealEd/Private/AssetDeleteModel.cpp
@@ -128,10 +128,28 @@ void FAssetDeleteModel::DeleteSourceContentFiles()
 		Utils::ExtractSourceFilePaths(PendingDelete->GetObject(), FilesToDelete);
 	}
 
-	IFileManager& FileManager = IFileManager::Get();
-	for (const auto& Path : FilesToDelete)
+	if (FilesToDelete.Num() != 0)
 	{
-		FileManager.Delete(*Path, false /* RequireExists */, true /* Even if read only */, true /* Quiet */);
+		TArray<FString> RootContentPaths;
+		FPackageName::QueryRootContentPaths( RootContentPaths );
+		for (FString& RootPath : RootContentPaths)
+		{
+			RootPath = FPaths::ConvertRelativePathToFull(FPackageName::LongPackageNameToFilename(RootPath));
+		}
+
+		IFileManager& FileManager = IFileManager::Get();
+		for (const auto& Path : FilesToDelete)
+		{
+			const FString FullPath = FPaths::ConvertRelativePathToFull(Path);
+			const bool bFileIsExternal = !RootContentPaths.ContainsByPredicate([&](const FString& ContentDir){
+				return FullPath.StartsWith(ContentDir);
+			});
+
+			if (!bFileIsExternal)
+			{
+				FileManager.Delete(*Path, false /* RequireExists */, true /* Even if read only */, true /* Quiet */);
+			}
+		}
 	}
 }
 
diff --git a/Engine/Source/Editor/UnrealEd/Private/AutoReimport/ContentDirectoryMonitor.cpp b/Engine/Source/Editor/UnrealEd/Private/AutoReimport/ContentDirectoryMonitor.cpp
--- a/Engine/Source/Editor/UnrealEd/Private/AutoReimport/ContentDirectoryMonitor.cpp
+++ b/Engine/Source/Editor/UnrealEd/Private/AutoReimport/ContentDirectoryMonitor.cpp
@@ -67,7 +67,7 @@ void FContentDirectoryMonitor::StartProcessing()
 			switch(Transaction.Action)
 			{
 				case FFileChangeData::FCA_Added:
-					if (Settings->bAutoCreateAssets)
+					if (Settings->bAutoCreateAssets && !MountedContentPath.IsEmpty())
 					{
 						AddedFiles.Emplace(MoveTemp(Transaction));
 					}
@@ -76,13 +76,13 @@ void FContentDirectoryMonitor::StartProcessing()
 						Cache.CompleteTransaction(MoveTemp(Transaction));
 					}
 					break;
-					
+
 				case FFileChangeData::FCA_Modified:
 					ModifiedFiles.Emplace(MoveTemp(Transaction));
 					break;
 
 				case FFileChangeData::FCA_Removed:
-					if (Settings->bAutoDeleteAssets)
+					if (Settings->bAutoDeleteAssets && !MountedContentPath.IsEmpty())
 					{
 						DeletedFiles.Emplace(MoveTemp(Transaction));
 					}
@@ -100,12 +100,6 @@ void FContentDirectoryMonitor::StartProcessing()
 
 void FContentDirectoryMonitor::ProcessAdditions(TArray<UPackage*>& OutPackagesToSave, const FTimeLimit& TimeLimit, const TMap<FString, TArray<UFactory*>>& InFactoriesByExtension, FReimportFeedbackContext& Context)
 {
-	if (MountedContentPath.IsEmpty())
-	{
-		// @todo: arodham: Prompt user for import location
-		return;
-	}
-	
 	bool bCancelled = false;
 	for (int32 Index = 0; Index < AddedFiles.Num(); ++Index)
 	{
@@ -166,10 +160,15 @@ void FContentDirectoryMonitor::ProcessAdditions(TArray<UPackage*>& OutPackagesTo
 
 				UObject* NewAsset = nullptr;
 				UFactory* FactoryInstance = FactoryType ? ConstructObject<UFactory>(FactoryType->GetClass()) : nullptr;
-				if (FactoryInstance && FactoryInstance->ConfigureProperties())
+				if (FactoryInstance)
 				{
-					UClass* ImportAssetType = FactoryInstance->SupportedClass;
-					NewAsset = UFactory::StaticImportObject(ImportAssetType, NewPackage, FName(*NewAssetName), RF_Public | RF_Standalone, bCancelled, *FullFilename, nullptr, FactoryInstance);
+					FactoryInstance->AddToRoot();
+					if (FactoryInstance->ConfigureProperties())
+					{
+						UClass* ImportAssetType = FactoryInstance->SupportedClass;
+						NewAsset = UFactory::StaticImportObject(ImportAssetType, NewPackage, FName(*NewAssetName), RF_Public | RF_Standalone, bCancelled, *FullFilename, nullptr, FactoryInstance);
+					}
+					FactoryInstance->RemoveFromRoot();
 				}
 
 				if (!bCancelled)
@@ -215,6 +214,7 @@ void FContentDirectoryMonitor::ProcessModifications(const IAssetRegistry& Regist
 		auto& Change = ModifiedFiles[Index];
 
 		const FString FullFilename = Cache.GetDirectory() + Change.Filename.Get();
+
 		for (const auto& AssetData : Utils::FindAssetsPertainingToFile(Registry, FullFilename))
 		{
 			UObject* Asset = AssetData.GetAsset();
