[INTEGRATE] Change 2423608 by [REDACTED_PERSON][REDACTED_HANDLE].7 on 2015/01/29 06:28:22

	"Delete source content files" no longer defaults to true when deleting assets

diff --git a/Engine/Source/Editor/UnrealEd/Private/Settings/SettingsClasses.cpp b/Engine/Source/Editor/UnrealEd/Private/Settings/SettingsClasses.cpp
--- a/Engine/Source/Editor/UnrealEd/Private/Settings/SettingsClasses.cpp
+++ b/Engine/Source/Editor/UnrealEd/Private/Settings/SettingsClasses.cpp
@@ -96,7 +96,7 @@ UEditorLoadingSavingSettings::UEditorLoadingSavingSettings( const FObjectInitial
 	, bMonitorContentDirectories(true)
 	, bAutoCreateAssets(true)
 	, bAutoDeleteAssets(true)
-	, bDeleteSourceFilesWithAssets(true)
+	, bDeleteSourceFilesWithAssets(false)
 {
 	TextDiffToolPath.FilePath = TEXT("P4Merge.exe");
 	AutoReimportDirectories.Add("/Game/");
