[INTEGRATE] Change 2423751 by [REDACTED_PERSON][REDACTED_HANDLE].7 on 2015/01/29 10:18:23

	Improved error and progress reporting for automatic content monitoring

diff --git a/Engine/Source/Editor/UnrealEd/Private/AutoReimport/AutoReimportManager.cpp b/Engine/Source/Editor/UnrealEd/Private/AutoReimport/AutoReimportManager.cpp
--- a/Engine/Source/Editor/UnrealEd/Private/AutoReimport/AutoReimportManager.cpp
+++ b/Engine/Source/Editor/UnrealEd/Private/AutoReimport/AutoReimportManager.cpp
@@ -370,6 +370,10 @@ TOptional<ECurrentState> FAutoReimportManager::ProcessDeletions()
 	if (AssetsToDelete.Num() > 0)
 	{
 		ObjectTools::DeleteAssets(AssetsToDelete);
+		for (const auto& AssetData : AssetsToDelete)
+		{
+			FeedbackContextOverride->AddMessage(EMessageSeverity::Info, FText::Format(LOCTEXT("Success_DeletedAsset", "Deleted {0}."), FText::FromName(AssetData.AssetName)));
+		}
 	}
 
 	Cleanup();
diff --git a/Engine/Source/Editor/UnrealEd/Private/AutoReimport/ContentDirectoryMonitor.cpp b/Engine/Source/Editor/UnrealEd/Private/AutoReimport/ContentDirectoryMonitor.cpp
--- a/Engine/Source/Editor/UnrealEd/Private/AutoReimport/ContentDirectoryMonitor.cpp
+++ b/Engine/Source/Editor/UnrealEd/Private/AutoReimport/ContentDirectoryMonitor.cpp
@@ -175,10 +175,21 @@ void FContentDirectoryMonitor::ProcessAdditions(TArray<UPackage*>& OutPackagesTo
 				{
 					if (!NewAsset)
 					{
+						TArray<UPackage*> Packages;
+						Packages.Add(NewPackage);
+
+						FText ErrorMessage;
+						if (!PackageTools::UnloadPackages(Packages, ErrorMessage))
+						{
+							Context.AddMessage(EMessageSeverity::Error, FText::Format(LOCTEXT("Error_UnloadingPackage", "There was an error unloading a package: {0}."), ErrorMessage));
+						}						
+
 						Context.AddMessage(EMessageSeverity::Error, FText::Format(LOCTEXT("Error_FailedToImportAsset", "Failed to import file {0}."), FText::FromString(FullFilename)));
 						continue;
 					}
 
+					Context.AddMessage(EMessageSeverity::Info, FText::Format(LOCTEXT("Success_CreatedNewAsset", "Created new asset {0}."), FText::FromString(PackagePath)));
+
 					FAssetRegistryModule::AssetCreated(NewAsset);
 					GEditor->BroadcastObjectReimported(NewAsset);
 
@@ -222,6 +233,10 @@ void FContentDirectoryMonitor::ProcessModifications(const IAssetRegistry& Regist
 			{
 				Context.AddMessage(EMessageSeverity::Error, FText::Format(LOCTEXT("Error_FailedToReimportAsset", "Failed to reimport asset {0}."), FText::FromString(Asset->GetName())));
 			}
+			else
+			{
+				Context.AddMessage(EMessageSeverity::Info, FText::Format(LOCTEXT("Success_CreatedNewAsset", "Reimported asset {0} from {1}."), FText::FromString(Asset->GetName()), FText::FromString(FullFilename)));
+			}
 		}
 
 		// Let the cache know that we've dealt with this change
diff --git a/Engine/Source/Editor/UnrealEd/Private/AutoReimport/ReimportFeedbackContext.cpp b/Engine/Source/Editor/UnrealEd/Private/AutoReimport/ReimportFeedbackContext.cpp
--- a/Engine/Source/Editor/UnrealEd/Private/AutoReimport/ReimportFeedbackContext.cpp
+++ b/Engine/Source/Editor/UnrealEd/Private/AutoReimport/ReimportFeedbackContext.cpp
@@ -6,7 +6,8 @@
 #include "NotificationManager.h"
 #include "SNotificationList.h"
 #include "INotificationWidget.h"
-
+#include "MessageLogModule.h"
+#include "SHyperlink.h"
 
 class SWidgetStack : public SVerticalBox
 {
@@ -17,14 +18,15 @@ class SWidgetStack : public SVerticalBox
 	{
 		MaxNumVisible = InMaxNumVisible;
 		SlideCurve = FCurveSequence(0.f, .5f, ECurveEaseFunction::QuadOut);
+		SizeCurve = FCurveSequence(0.f, .5f, ECurveEaseFunction::QuadOut);
 
-		StartSlideOffset = LerpSlideOffset = 0;
-		StartSizeOffset = LerpSizeOffset = 0;
+		StartSlideOffset = 0;
+		StartSizeOffset = FVector2D(ForceInitToZero);
 	}
 
 	FVector2D ComputeTotalSize() const
 	{
-		FVector2D Size;
+		FVector2D Size(ForceInitToZero);
 		for (int32 Index = 0; Index < FMath::Min(NumSlots(), MaxNumVisible); ++Index)
 		{
 			const FVector2D& ChildSize = Children[Index].GetWidget()->GetDesiredSize();
@@ -39,9 +41,11 @@ class SWidgetStack : public SVerticalBox
 
 	virtual FVector2D ComputeDesiredSize(float) const override
 	{
-		FVector2D Size = ComputeTotalSize();
-		Size.Y -= LerpSizeOffset;
-		return Size;
+		const float Lerp = SizeCurve.GetLerp();
+		FVector2D DesiredSize = ComputeTotalSize() * Lerp + StartSizeOffset * (1.f-Lerp);
+		DesiredSize.X = FMath::Min(500.f, DesiredSize.X);
+
+		return DesiredSize;
 	}
 
 	virtual void OnArrangeChildren( const FGeometry& AllottedGeometry, FArrangedChildren& ArrangedChildren ) const
@@ -51,7 +55,8 @@ class SWidgetStack : public SVerticalBox
 			return;
 		}
 
-		float PositionSoFar = AllottedGeometry.GetLocalSize().Y + LerpSlideOffset;
+		const float Alpha = 1.f - SlideCurve.GetLerp();
+		float PositionSoFar = AllottedGeometry.GetLocalSize().Y + StartSlideOffset*Alpha;
 
 		for (int32 Index = 0; Index < NumSlots(); ++Index)
 		{
@@ -79,58 +84,58 @@ class SWidgetStack : public SVerticalBox
 		}
 	}
 
-	void Remove(const TSharedRef<SWidget>& InWidget)
-	{
-
-	}
-
 	void Add(const TSharedRef<SWidget>& InWidget)
 	{
+		TSharedPtr<SWidgetStackItem> NewItem;
+
 		InsertSlot(0)
 		.AutoHeight()
 		[
-			SNew(SWidgetStackItem)
+			SAssignNew(NewItem, SWidgetStackItem)
 			[
 				InWidget
 			]
 		];
-
-		auto Widget = Children[0].GetWidget();
-		Widget->SlatePrepass();
 		
-		StartSlideOffset = LerpSlideOffset + Widget->GetDesiredSize().Y;
-		LerpSlideOffset = StartSlideOffset;
-
-		StartSizeOffset = ComputeTotalSize().Y - GetDesiredSize().Y;
-		LerpSizeOffset = StartSizeOffset;
+		{
+			auto Widget = Children[0].GetWidget();
+			Widget->SlatePrepass();
 
-		SlideCurve.Play(AsShared());
+			const float WidgetHeight = Widget->GetDesiredSize().Y;
+			StartSlideOffset += WidgetHeight;
+			// Fade in time is 1 second x the proportion of the slide amount that this widget takes up
+			NewItem->FadeIn(WidgetHeight / StartSlideOffset);
 
-		if (Children.Num() > MaxNumVisible)
-		{
-			auto Widget = StaticCastSharedRef<SWidgetStackItem>(Children[MaxNumVisible].GetWidget());
-			if (Widget->OpacityCurve.IsPlaying())
+			if (!SlideCurve.IsPlaying())
 			{
-				Widget->OpacityCurve.Reverse();
+				SlideCurve.Play(AsShared());
 			}
-			else
+		}
+
+		const FVector2D NewSize = ComputeTotalSize();
+		if (NewSize != StartSizeOffset)
+		{
+			StartSizeOffset = NewSize;
+
+			if (!SizeCurve.IsPlaying())
 			{
-				Widget->OpacityCurve.PlayReverse(Widget);
+				SizeCurve.Play(AsShared());
 			}
 		}
 	}
 	
 	virtual void Tick(const FGeometry& AllottedGeometry, const double InCurrentTime, const float InDeltaTime)
 	{
-		const float Alpha = 1.f - SlideCurve.GetLerp();
-
-		LerpSlideOffset = StartSlideOffset * Alpha;
-		LerpSizeOffset = StartSizeOffset * Alpha;
+		if (!SlideCurve.IsPlaying())
+		{
+			StartSlideOffset = 0;
+		}
 
 		// Delete any widgets that are now offscreen
 		if (Children.Num() != 0)
 		{
-			float PositionSoFar = AllottedGeometry.GetLocalSize().Y + LerpSlideOffset;
+			const float Alpha = 1.f - SlideCurve.GetLerp();
+			float PositionSoFar = AllottedGeometry.GetLocalSize().Y + Alpha * StartSlideOffset;
 
 			int32 Index = 0;
 			for (; PositionSoFar > 0 && Index < NumSlots(); ++Index)
@@ -145,9 +150,16 @@ class SWidgetStack : public SVerticalBox
 				}
 			}
 
-			while (Children.Num() > MaxNumVisible + 1)
+			for (int32 Index = MaxNumVisible; Index < Children.Num(); )
 			{
-				Children.RemoveAt(Children.Num() - 1);
+				if (StaticCastSharedRef<SWidgetStackItem>(Children[Index].GetWidget())->bIsFinished)
+				{
+					Children.RemoveAt(Index);
+				}
+				else
+				{
+					++Index;
+				}
 			}
 		}
 	}
@@ -156,13 +168,11 @@ class SWidgetStack : public SVerticalBox
 	{
 		SLATE_BEGIN_ARGS(SWidgetStackItem){}
 			SLATE_DEFAULT_SLOT(FArguments, Content)
-			SLATE_EVENT(FSimpleDelegate, OnFadeOut)
 		SLATE_END_ARGS()
 
 		void Construct(const FArguments& InArgs)
 		{
-			OpacityCurve = FCurveSequence(0.f, .5f, ECurveEaseFunction::QuadOut);
-			OpacityCurve.Play(AsShared());
+			bIsFinished = false;
 
 			ChildSlot
 			[
@@ -176,12 +186,17 @@ class SWidgetStack : public SVerticalBox
 			];
 		}
 		
+		void FadeIn(float Time)
+		{
+			OpacityCurve = FCurveSequence(0.f, Time, ECurveEaseFunction::QuadOut);
+			OpacityCurve.Play();
+		}
+
 		virtual void Tick(const FGeometry& AllottedGeometry, const double InCurrentTime, const float InDeltaTime)
 		{
-			if (OnFadeOut.IsBound() && OpacityCurve.IsAtStart() && OpacityCurve.IsInReverse())
+			if (!bIsFinished && OpacityCurve.IsAtStart() && OpacityCurve.IsInReverse())
 			{
-				OnFadeOut.Execute();
-				OnFadeOut = FSimpleDelegate();
+				bIsFinished = true;
 			}
 		}
 
@@ -190,18 +205,22 @@ class SWidgetStack : public SVerticalBox
 			return FLinearColor(1.f, 1.f, 1.f, OpacityCurve.GetLerp());
 		}
 
+		bool bIsFinished;
 		FCurveSequence OpacityCurve;
-		FSimpleDelegate OnFadeOut;
 	};
 
-	FCurveSequence SlideCurve;
+	FCurveSequence SlideCurve, SizeCurve;
 
-	float StartSlideOffset, LerpSlideOffset;
-	float StartSizeOffset, LerpSizeOffset;
+	float StartSlideOffset;
+	FVector2D StartSizeOffset;
 
 	int32 MaxNumVisible;
 };
 
+void SReimportFeedback::Disable()
+{
+	WidgetStack->SetVisibility(EVisibility::HitTestInvisible);
+}
 
 void SReimportFeedback::Add(const TSharedRef<SWidget>& Widget)
 {
@@ -218,6 +237,11 @@ FText SReimportFeedback::GetMainText() const
 	return MainText;
 }
 
+EVisibility SReimportFeedback::GetHyperlinkVisibility() const
+{
+	return WidgetStack->NumSlots() != 0 ? EVisibility::Visible : EVisibility::Collapsed;
+}
+
 void SReimportFeedback::Construct(const FArguments& InArgs, FText InMainText)
 {
 	MainText = InMainText;
@@ -243,7 +267,22 @@ void SReimportFeedback::Construct(const FArguments& InArgs, FText InMainText)
 			.Padding(FMargin(0, 5, 0, 0))
 			.AutoHeight()
 			[
-				SAssignNew(WidgetStack, SWidgetStack, 1)
+				SAssignNew(WidgetStack, SWidgetStack, 3)
+			]
+
+			+ SVerticalBox::Slot()
+			.Padding(FMargin(0, 5, 0, 0))
+			.AutoHeight()
+			.HAlign(HAlign_Right)
+			[
+				SNew(SHyperlink)
+				.Visibility(this, &SReimportFeedback::GetHyperlinkVisibility)
+				.Text(NSLOCTEXT("ReimportContext", "OpenMessageLog", "Open message log"))
+				.TextStyle(FCoreStyle::Get(), "SmallText")
+				.OnNavigate_Lambda([]{
+					FMessageLogModule& MessageLogModule = FModuleManager::LoadModuleChecked<FMessageLogModule>("MessageLog");
+					MessageLogModule.OpenMessageLog("AssetReimport");
+				})
 			]
 		]
 	];
@@ -257,7 +296,8 @@ void FReimportFeedbackContext::Initialize(TSharedRef<SReimportFeedback> Widget)
 {
 	ShowNotificationDelay = FTimeLimit(.5f);
 	NotificationContent = Widget;
-	MessageLog.NewPage(Widget->GetMainText());
+
+	MessageLog.NewPage(FText::Format(NSLOCTEXT("ReimportContext", "MessageLogPageLabel", "Outstanding source content changes {0}"), FText::AsTime(FDateTime::Now())));
 }
 
 void FReimportFeedbackContext::Tick()
@@ -265,7 +305,7 @@ void FReimportFeedbackContext::Tick()
 	if (!Notification.IsValid() && ShowNotificationDelay.Exceeded())
 	{
 		FNotificationInfo Info(SharedThis(this));
-		Info.ExpireDuration = 1.f;
+		Info.ExpireDuration = 3.f;
 		Info.bFireAndForget = false;
 
 		Notification = FSlateNotificationManager::Get().AddNotification(Info);
@@ -274,12 +314,12 @@ void FReimportFeedbackContext::Tick()
 
 void FReimportFeedbackContext::Destroy()
 {
-	MessageLog.Notify();
+	MessageLog.Notify(FText(), EMessageSeverity::Warning);
 
 	if (Notification.IsValid())
 	{
-		NotificationContent->SetVisibility(EVisibility::HitTestInvisible);
-		//Notification->SetCompletionState(SNotificationItem::CS_Success);
+		NotificationContent->Disable();
+		Notification->SetCompletionState(SNotificationItem::CS_Success);
 		Notification->ExpireAndFadeout();
 	}
 }
@@ -287,10 +327,7 @@ void FReimportFeedbackContext::Destroy()
 void FReimportFeedbackContext::AddMessage(EMessageSeverity::Type Severity, const FText& Message)
 {
 	MessageLog.Message(Severity, Message);
-	if (Severity >= EMessageSeverity::Error)
-	{
-		AddWidget(SNew(STextBlock).Text(Message));
-	}
+	AddWidget(SNew(STextBlock).Text(Message));
 }
 
 void FReimportFeedbackContext::AddWidget(const TSharedRef<SWidget>& Widget)
@@ -298,10 +335,12 @@ void FReimportFeedbackContext::AddWidget(const TSharedRef<SWidget>& Widget)
 	NotificationContent->Add(Widget);
 }
 
-void FReimportFeedbackContext::ProgressReported(const float TotalProgressInterp, FText DisplayMessage)
+void FReimportFeedbackContext::StartSlowTask(const FText& Task, bool bShowCancelButton)
 {
-	if (!DisplayMessage.IsEmpty())
+	FFeedbackContext::StartSlowTask(Task, bShowCancelButton);
+
+	if (!Task.IsEmpty())
 	{
-		NotificationContent->Add(SNew(STextBlock).Text(DisplayMessage));
+		NotificationContent->Add(SNew(STextBlock).Text(Task));
 	}
 }
\ No newline at end of file
diff --git a/Engine/Source/Editor/UnrealEd/Private/AutoReimport/ReimportFeedbackContext.h b/Engine/Source/Editor/UnrealEd/Private/AutoReimport/ReimportFeedbackContext.h
--- a/Engine/Source/Editor/UnrealEd/Private/AutoReimport/ReimportFeedbackContext.h
+++ b/Engine/Source/Editor/UnrealEd/Private/AutoReimport/ReimportFeedbackContext.h
@@ -23,6 +23,9 @@ public:
 	/** Add a widget to this feedback's widget stack */
 	void Add(const TSharedRef<SWidget>& Widget);
 
+	/** Disable input to this widget's dynamic content (except the message log hyperlink) */
+	void Disable();
+
 	/** Set the main text of this widget */
 	void SetMainText(FText InText);
 	
@@ -30,6 +33,12 @@ public:
 	FText GetMainText() const;
 
 private:
+	
+	/** Get the visibility of the hyperlink to open the message log */
+	EVisibility GetHyperlinkVisibility() const;
+
+private:
+
 	/** Cached main text for the notification */
 	FText MainText;
 
@@ -65,7 +74,7 @@ public:
 	virtual void Serialize( const TCHAR* V, ELogVerbosity::Type Verbosity, const class FName& Category ) override {}
 	virtual void OnSetCompletionState(SNotificationItem::ECompletionState State) override {}
 	virtual TSharedRef<SWidget> AsWidget() override { return NotificationContent.ToSharedRef(); };
-	virtual void ProgressReported(const float TotalProgressInterp, FText DisplayMessage) override;
+	virtual void StartSlowTask( const FText& Task, bool bShowCancelButton=false) override;
 
 private:
 	/** Message log for output of errors and log messages */
