[INTEGRATE] Change 2424621 by [REDACTED_PERSON][REDACTED_HANDLE] on 2015/01/29 17:58:00

	[UE-8535] Fixing up how we display blueprintable component names in the SCS tree.

diff --git a/Engine/Source/Editor/UnrealEd/Private/SComponentClassCombo.cpp b/Engine/Source/Editor/UnrealEd/Private/SComponentClassCombo.cpp
--- a/Engine/Source/Editor/UnrealEd/Private/SComponentClassCombo.cpp
+++ b/Engine/Source/Editor/UnrealEd/Private/SComponentClassCombo.cpp
@@ -544,15 +544,11 @@ FString SComponentClassCombo::GetSanitizedComponentName( UClass* ComponentClass,
 		}
 		else
 		{
-			DisplayName = ComponentClass->GetName();
-			DisplayName = DisplayName.Replace(TEXT("Component"), TEXT(""), ESearchCase::IgnoreCase);
-		}
-
-		// Purge the _C for BP Components
-		if(DisplayName.EndsWith(TEXT("_C")))
-		{
-			const int32 NewLen = DisplayName.Len() - 2;
-			DisplayName = DisplayName.Left(NewLen);
+			DisplayName = ComponentClass->GetDisplayNameText().ToString();
+			if (!ComponentClass->HasAnyClassFlags(CLASS_CompiledFromBlueprint))
+			{
+				DisplayName.RemoveFromEnd(TEXT("Component"), ESearchCase::IgnoreCase);
+			}
 		}
 
 		return FName::NameToDisplayString(DisplayName, false);
diff --git a/Engine/Source/Runtime/Engine/Private/SCS_Node.cpp b/Engine/Source/Runtime/Engine/Private/SCS_Node.cpp
--- a/Engine/Source/Runtime/Engine/Private/SCS_Node.cpp
+++ b/Engine/Source/Runtime/Engine/Private/SCS_Node.cpp
@@ -448,13 +448,15 @@ FName USCS_Node::GenerateNewComponentName( TArray<FName>& CurrentNames, FName De
 			else
 			{
 				const UClass* ComponentClass = ComponentTemplate->GetClass();
-				ComponentName = ComponentClass->GetName().Replace(TEXT("Component"), TEXT(""));
+				ComponentName = ComponentClass->GetName();
 
-				// Trim the _C, if needed
-				if (ComponentClass->ClassGeneratedBy && ComponentName.EndsWith(TEXT("_C")))
+				if (!ComponentClass->HasAnyClassFlags(CLASS_CompiledFromBlueprint))
 				{
-					const int32 NameLen = ComponentName.Len() - 1;
-					ComponentName = ComponentName.Left(NameLen);
+					ComponentName.RemoveFromEnd(TEXT("Component"));
+				}
+				else
+				{
+					ComponentName.RemoveFromEnd("_C");
 				}
 			}
 			
