Add missing header. Optimistic build fix for Mac.

diff --git a/Engine/Source/Editor/UnrealEd/Public/AssetThumbnail.h b/Engine/Source/Editor/UnrealEd/Public/AssetThumbnail.h
--- a/Engine/Source/Editor/UnrealEd/Public/AssetThumbnail.h
+++ b/Engine/Source/Editor/UnrealEd/Public/AssetThumbnail.h
@@ -7,6 +7,7 @@ class FSlateTexture2DRHIRef;
 class FSlateTextureRenderTarget2DResource;
 
 #include "Runtime/AssetRegistry/Public/AssetData.h"
+#include "TickableEditorObject.h"
 
 namespace EThumbnailLabel
 {
