[INTEGRATE] Change 2425051 by [REDACTED_PERSON][REDACTED_HANDLE] on 2015/01/29 22:14:00

	#ue4 - "Possess" button during simulate now attempts to use any pawn with "AutoPossessPlayer" set to index 0. If not set, it spawns the pawn as normal, and if there was no PlayerStart, it spawns at the camera location.

	UE-8182

diff --git a/Engine/Source/Editor/UnrealEd/Private/PlayLevel.cpp b/Engine/Source/Editor/UnrealEd/Private/PlayLevel.cpp
--- a/Engine/Source/Editor/UnrealEd/Private/PlayLevel.cpp
+++ b/Engine/Source/Editor/UnrealEd/Private/PlayLevel.cpp
@@ -3002,7 +3002,42 @@ void UEditorEngine::ToggleBetweenPIEandSIE( bool bNewSession )
 					AuthGameMode->RemovePlayerControllerFromPlayerCount(PC);
 					PC->PlayerState->bOnlySpectator = false;
 					AuthGameMode->NumPlayers++;
-					AuthGameMode->RestartPlayer(PC);
+
+					bool bNeedsRestart = true;
+					if (PC->GetPawn() == NULL)
+					{
+						// Use the "auto-possess" pawn in the world, if there is one.
+						for (FConstPawnIterator Iterator = World->GetPawnIterator(); Iterator; ++Iterator)
+						{
+							APawn* Pawn = *Iterator;
+							if (Pawn && Pawn->AutoPossessPlayer == EAutoReceiveInput::Player0)
+							{
+								if (Pawn->Controller == nullptr)
+								{
+									PC->Possess(Pawn);
+									bNeedsRestart = false;
+								}
+								break;
+							}
+						}
+					}
+
+					if (bNeedsRestart)
+					{
+						AuthGameMode->RestartPlayer(PC);
+
+						if (PC->GetPawn())
+						{
+							// If there was no player start, then try to place the pawn where the camera was.						
+							if (PC->StartSpot == nullptr || Cast<AWorldSettings>(PC->StartSpot.Get()))
+							{
+								const FVector Location = EditorViewportClient.GetViewLocation();
+								const FRotator Rotation = EditorViewportClient.GetViewRotation();
+								PC->SetControlRotation(Rotation);
+								PC->GetPawn()->TeleportTo(Location, Rotation);
+							}
+						}
+					}
 				}
 
 				OnSwitchWorldsForPIE(false);
