[INTEGRATE] Change 2425145 by [REDACTED_PERSON][REDACTED_HANDLE] on 2015/01/29 23:38:05

	Blueprints: Fix a typo in a LOCTEXT key

diff --git a/Engine/Source/Editor/UnrealEd/Private/Kismet2/BlueprintEditorUtils.cpp b/Engine/Source/Editor/UnrealEd/Private/Kismet2/BlueprintEditorUtils.cpp
--- a/Engine/Source/Editor/UnrealEd/Private/Kismet2/BlueprintEditorUtils.cpp
+++ b/Engine/Source/Editor/UnrealEd/Private/Kismet2/BlueprintEditorUtils.cpp
@@ -6089,7 +6089,7 @@ TSharedRef<SWidget> FBlueprintEditorUtils::ConstructBlueprintParentClassPicker(
 
 	TSharedPtr<FBlueprintReparentFilter> Filter = MakeShareable(new FBlueprintReparentFilter);
 	Options.ClassFilter = Filter;
-	Options.ViewerTitleString = LOCTEXT("ReparentBblueprint", "Reparent blueprint");
+	Options.ViewerTitleString = LOCTEXT("ReparentBlueprint", "Reparent blueprint");
 
 	// Only allow parenting to base blueprints.
 	Options.bIsBlueprintBaseOnly = true;
