Fix for deprecation warning after integrating back from 4.7.

diff --git a/Engine/Source/Editor/UnrealEd/Private/AutoReimport/ReimportFeedbackContext.cpp b/Engine/Source/Editor/UnrealEd/Private/AutoReimport/ReimportFeedbackContext.cpp
--- a/Engine/Source/Editor/UnrealEd/Private/AutoReimport/ReimportFeedbackContext.cpp
+++ b/Engine/Source/Editor/UnrealEd/Private/AutoReimport/ReimportFeedbackContext.cpp
@@ -189,7 +189,7 @@ class SWidgetStack : public SVerticalBox
 		void FadeIn(float Time)
 		{
 			OpacityCurve = FCurveSequence(0.f, Time, ECurveEaseFunction::QuadOut);
-			OpacityCurve.Play();
+			OpacityCurve.Play(AsShared());
 		}
 
 		virtual void Tick(const FGeometry& AllottedGeometry, const double InCurrentTime, const float InDeltaTime)
