Auto reimport manager no longer operates when there is no project path set

diff --git a/Engine/Source/Editor/UnrealEd/Private/UnrealEdEngine.cpp b/Engine/Source/Editor/UnrealEd/Private/UnrealEdEngine.cpp
--- a/Engine/Source/Editor/UnrealEd/Private/UnrealEdEngine.cpp
+++ b/Engine/Source/Editor/UnrealEd/Private/UnrealEdEngine.cpp
@@ -101,8 +101,11 @@ void UUnrealEdEngine::Init(IEngineLoop* InEngineLoop)
 		SpriteIDToIndexMap.Add( SpriteInfo.Category, InfoIndex );
 	}
 
-	AutoReimportManager = ConstructObject<UAutoReimportManager>(UAutoReimportManager::StaticClass());
-	AutoReimportManager->Initialize();
+	if (FPaths::IsProjectFilePathSet())
+	{
+		AutoReimportManager = ConstructObject<UAutoReimportManager>(UAutoReimportManager::StaticClass());
+		AutoReimportManager->Initialize();
+	}
 
 	// register details panel customizations
 	if (!HasAnyFlags(RF_ClassDefaultObject))
