fix for UE-7292 Replacing an actor with a new actor twice will result in a crash

diff --git a/Engine/Source/Editor/UnrealEd/Private/Editor.cpp b/Engine/Source/Editor/UnrealEd/Private/Editor.cpp
--- a/Engine/Source/Editor/UnrealEd/Private/Editor.cpp
+++ b/Engine/Source/Editor/UnrealEd/Private/Editor.cpp
@@ -12,6 +12,7 @@
 #include "BSPOps.h"
 #include "EditorCommandLineUtils.h"
 #include "Net/NetworkProfiler.h"
+#include "UObjectGlobals.h"
 
 // needed for the RemotePropagator
 #include "SoundDefinitions.h"
@@ -4241,15 +4242,6 @@ bool UEditorEngine::CanParentActors( const AActor* ParentActor, const AActor* Ch
 		return false;
 	}
 
-	if(ParentActor->IsSelected())
-	{
-		if (ReasonText)
-		{
-			*ReasonText = NSLOCTEXT("ActorAttachmentError", "ParentSelected_ActorAttachementError", "Cannot attach actor to an actor which is also selected.");
-		}
-		return false;
-	}
-
 	if(ParentRoot->IsAttachedTo( ChildRoot ))
 	{
 		if (ReasonText)
@@ -5110,7 +5102,8 @@ void UEditorEngine::ReplaceActors(UActorFactory* Factory, const FAssetData& Asse
 		AActor* NewActor = NULL;
 
 		const FName OldActorName = OldActor->GetFName();
-		OldActor->Rename(*FString::Printf(TEXT("%s_REPLACED"), *OldActorName.ToString()));
+		const FName OldActorReplacedNamed = MakeUniqueObjectName(OldActor->GetOuter(), OldActor->GetClass(), *FString::Printf(TEXT("%s_REPLACED"), *OldActorName.ToString()));
+		OldActor->Rename(*OldActorReplacedNamed.ToString());
 
 		const FTransform OldTransform = OldActor->ActorToWorld();
 
diff --git a/Engine/Source/Runtime/Engine/Private/Actor.cpp b/Engine/Source/Runtime/Engine/Private/Actor.cpp
--- a/Engine/Source/Runtime/Engine/Private/Actor.cpp
+++ b/Engine/Source/Runtime/Engine/Private/Actor.cpp
@@ -2173,7 +2173,7 @@ static void DispatchOnComponentsCreated(AActor* NewActor)
 	for (int32 Idx = 0; Idx < Components.Num(); Idx++)
 	{
 		UActorComponent* ActorComp = Components[Idx];
-		if (ActorComp != NULL)
+		if (ActorComp != NULL && !ActorComp->bHasBeenCreated)
 		{
 			ActorComp->OnComponentCreated();
 		}
