Trim down the amount of UObject constructors.

diff --git a/Engine/Source/Editor/UnrealEd/Private/Factories/EditorFactories.cpp b/Engine/Source/Editor/UnrealEd/Private/Factories/EditorFactories.cpp
--- a/Engine/Source/Editor/UnrealEd/Private/Factories/EditorFactories.cpp
+++ b/Engine/Source/Editor/UnrealEd/Private/Factories/EditorFactories.cpp
@@ -200,12 +200,12 @@ UMaterialFactoryNew::UMaterialFactoryNew(const FObjectInitializer& ObjectInitial
 
 UObject* UMaterialFactoryNew::FactoryCreateNew(UClass* Class,UObject* InParent,FName Name,EObjectFlags Flags,UObject* Context,FFeedbackContext* Warn)
 {
-	UMaterial* NewMaterial = ConstructObject<UMaterial>(Class, InParent, Name, Flags);
+	UMaterial* NewMaterial = NewObject<UMaterial>(InParent, Class, Name, Flags);
 
 	if ( InitialTexture != nullptr )
 	{
 		// An initial texture was specified, add it and assign it to the BaseColor
-		UMaterialExpressionTextureSample* Expression = ConstructObject<UMaterialExpressionTextureSample>( UMaterialExpressionTextureSample::StaticClass(), NewMaterial );
+		UMaterialExpressionTextureSample* Expression = NewObject<UMaterialExpressionTextureSample>(NewMaterial);
 		NewMaterial->Expressions.Add( Expression );
 
 		NewMaterial->BaseColor.Expression = Expression;
@@ -1202,7 +1202,7 @@ UObject* UPolysFactory::FactoryCreateText
 	FEditorDelegates::OnAssetPreImport.Broadcast(this, Class, InParent, Name, Type);
 
 	// Create polys.	
-	UPolys* Polys = Context ? CastChecked<UPolys>(Context) : NewNamedObject<UPolys>(InParent, Name, Flags);
+	UPolys* Polys = Context ? CastChecked<UPolys>(Context) : NewObject<UPolys>(InParent, Name, Flags);
 
 	// Eat up if present.
 	GetBEGIN( &Buffer, TEXT("POLYLIST") );
@@ -1474,7 +1474,7 @@ UObject* UModelFactory::FactoryCreateText
 	FEditorDelegates::OnAssetPreImport.Broadcast(this, Class, InParent, Name, Type);
 
 	ABrush* TempOwner = (ABrush*)Context;
-	UModel* Model = NewNamedObject<UModel>(InParent, Name, Flags);
+	UModel* Model = NewObject<UModel>(InParent, Name, Flags);
 	Model->Initialize(TempOwner, true);
 
 	const TCHAR* StrPtr;
@@ -1560,7 +1560,7 @@ void CreateSoundCue( USoundWave* Sound, UObject* InParent, EObjectFlags Flags, b
 	FString SoundCueName = FString::Printf( TEXT( "%s_Cue" ), *Sound->GetName() );
 
 	// Create sound cue and wave player
-	USoundCue* SoundCue = ConstructObject<USoundCue>( USoundCue::StaticClass(), InParent, *SoundCueName, Flags );
+	USoundCue* SoundCue = NewObject<USoundCue>(InParent, *SoundCueName, Flags);
 	USoundNodeWavePlayer* WavePlayer = SoundCue->ConstructSoundNode<USoundNodeWavePlayer>(); 
 
 	int32 NodeIndex = ( int32 )bIncludeAttenuationNode + ( int32 )bIncludeModulatorNode + ( int32 )bIncludeLoopingNode;
@@ -1615,7 +1615,7 @@ UReverbEffectFactory::UReverbEffectFactory(const FObjectInitializer& ObjectIniti
 
 UObject* UReverbEffectFactory::FactoryCreateNew( UClass* InClass, UObject* InParent, FName InName, EObjectFlags Flags, UObject* Context, FFeedbackContext* Warn )
 {
-	UReverbEffect* ReverbEffect = ConstructObject<UReverbEffect>( UReverbEffect::StaticClass(), InParent, InName, Flags );
+	UReverbEffect* ReverbEffect = NewObject<UReverbEffect>(InParent, InName, Flags);
 
 	return ReverbEffect;
 }
@@ -1771,7 +1771,7 @@ UObject* USoundFactory::FactoryCreateBinary
 
 		// Use pre-existing sound if it exists and we want to keep settings,
 		// otherwise create new sound and import raw data.
-		USoundWave* Sound = (bUseExistingSettings && ExistingSound) ? ExistingSound : NewNamedObject<USoundWave>(InParent, Name, Flags);
+		USoundWave* Sound = (bUseExistingSettings && ExistingSound) ? ExistingSound : NewObject<USoundWave>(InParent, Name, Flags);
 		
 		if (bUseExistingSettings && ExistingSound)
 		{
@@ -2078,7 +2078,7 @@ UObject* USoundSurroundFactory::FactoryCreateBinary
 
 			if (Sound == nullptr)
 			{
-				Sound = NewNamedObject<USoundWave>(InParent, BaseName, Flags);
+				Sound = NewObject<USoundWave>(InParent, BaseName, Flags);
 			}
 		}
 
@@ -2389,7 +2389,7 @@ USoundCueFactoryNew::USoundCueFactoryNew(const FObjectInitializer& ObjectInitial
 
 UObject* USoundCueFactoryNew::FactoryCreateNew( UClass* Class, UObject* InParent, FName Name, EObjectFlags Flags, UObject* Context, FFeedbackContext* Warn )
 {
-	USoundCue* SoundCue = ConstructObject<USoundCue>( USoundCue::StaticClass(), InParent, Name, Flags );
+	USoundCue* SoundCue = NewObject<USoundCue>(InParent, Name, Flags);
 
 	if (InitialSoundWave)
 	{
@@ -2420,7 +2420,7 @@ USoundMixFactory::USoundMixFactory(const FObjectInitializer& ObjectInitializer)
 
 UObject* USoundMixFactory::FactoryCreateNew( UClass* InClass, UObject* InParent, FName InName, EObjectFlags Flags, UObject* Context, FFeedbackContext* Warn )
 {
-	USoundMix* Mix = ConstructObject<USoundMix>( USoundMix::StaticClass(), InParent, InName, Flags );
+	USoundMix* Mix = NewObject<USoundMix>(InParent, InName, Flags);
 
 	return Mix;
 }
@@ -2441,7 +2441,7 @@ USoundClassFactory::USoundClassFactory(const FObjectInitializer& ObjectInitializ
 
 UObject* USoundClassFactory::FactoryCreateNew( UClass* InClass, UObject* InParent, FName InName, EObjectFlags Flags, UObject* Context, FFeedbackContext* Warn )
 {
-	USoundClass* SoundClass = ConstructObject<USoundClass>( USoundClass::StaticClass(), InParent, InName, Flags );
+	USoundClass* SoundClass = NewObject<USoundClass>(InParent, InName, Flags);
 	
 	FAudioDevice* AudioDevice = GEngine ? GEngine->GetAudioDevice() : nullptr;
 	if (AudioDevice)
@@ -2467,7 +2467,7 @@ USoundAttenuationFactory::USoundAttenuationFactory(const FObjectInitializer& Obj
 
 UObject* USoundAttenuationFactory::FactoryCreateNew( UClass* Class, UObject* InParent, FName Name, EObjectFlags Flags, UObject* Context, FFeedbackContext* Warn )
 {
-	return ConstructObject<USoundAttenuation>( USoundAttenuation::StaticClass(), InParent, Name, Flags );
+	return NewObject<USoundAttenuation>(InParent, Name, Flags);
 }
 
 /*------------------------------------------------------------------------------
@@ -4517,7 +4517,7 @@ UObject* UTextureFactory::FactoryCreateBinary
 		FAssetRegistryModule::AssetCreated(Material);
 
 		// Create a texture reference for the texture we just imported and hook it up to the diffuse channel
-		UMaterialExpression* Expression = ConstructObject<UMaterialExpression>( UMaterialExpressionTextureSample::StaticClass(), Material );
+		UMaterialExpression* Expression = NewObject<UMaterialExpression>(Material, UMaterialExpressionTextureSample::StaticClass());
 		Material->Expressions.Add( Expression );
 		TArray<FExpressionOutput> Outputs;
 
@@ -5213,7 +5213,7 @@ UFontFactory::UFontFactory(const FObjectInitializer& ObjectInitializer)
 
 UObject* UFontFactory::FactoryCreateNew(UClass* InClass, UObject* InParent, FName InName, EObjectFlags InFlags, UObject* InContext, FFeedbackContext* InWarn)
 {
-	UFont* const Font = ConstructObject<UFont>(InClass, InParent, InName, InFlags);
+	UFont* const Font = NewObject<UFont>(InParent, InClass, InName, InFlags);
 	if(Font)
 	{
 		Font->FontCacheType = EFontCacheType::Runtime;
@@ -5240,7 +5240,7 @@ UObject* UFontFileImportFactory::FactoryCreateBinary(UClass* InClass, UObject* I
 {
 	FEditorDelegates::OnAssetPreImport.Broadcast(this, InClass, InParent, InName, InType);
 
-	UFont* const Font = ConstructObject<UFont>(InClass, InParent, InName, InFlags);
+	UFont* const Font = NewObject<UFont>(InParent, InClass, InName, InFlags);
 	if(Font)
 	{
 		Font->FontCacheType = EFontCacheType::Runtime;
@@ -5328,7 +5328,7 @@ void FCustomizableTextObjectFactory::ProcessBuffer(UObject* InParent, EObjectFla
 				ClearObjectNameUsage(InParent, ObjName);
 
 				// Spawn the object and reset it's archetype
-				UObject* CreatedObject = ConstructObject<UObject>( ObjClass, InParent, ObjName, Flags, ObjArchetype, !!InParent, &InstanceGraph );
+				UObject* CreatedObject = NewObject<UObject>(InParent, ObjClass, ObjName, Flags, ObjArchetype, !!InParent, &InstanceGraph);
 
 				// Get property text for the new object.
 				FString PropText, PropLine;
@@ -5634,7 +5634,7 @@ EReimportResult::Type UReimportFbxStaticMeshFactory::Reimport( UObject* Obj )
 
 	UFbxStaticMeshImportData* ImportData = Cast<UFbxStaticMeshImportData>(Mesh->AssetImportData);
 	
-	UFbxImportUI* ReimportUI = ConstructObject<UFbxImportUI>(UFbxImportUI::StaticClass());
+	UFbxImportUI* ReimportUI = NewObject<UFbxImportUI>();
 	ReimportUI->MeshTypeToImport = FBXIT_StaticMesh;
 	ReimportUI->bOverrideFullName = false;
 	ReimportUI->bCombineMeshes = true;
@@ -5815,7 +5815,7 @@ EReimportResult::Type UReimportFbxSkeletalMeshFactory::Reimport( UObject* Obj )
 
 	if( !ImportUI )
 	{
-		ImportUI = ConstructObject<UFbxImportUI>( UFbxImportUI::StaticClass(), this, NAME_None, RF_Public );
+		ImportUI = NewObject<UFbxImportUI>(this, NAME_None, RF_Public);
 	}
 
 	UnFbx::FFbxImporter* FFbxImporter = UnFbx::FFbxImporter::GetInstance();
@@ -5824,7 +5824,7 @@ EReimportResult::Type UReimportFbxSkeletalMeshFactory::Reimport( UObject* Obj )
 	UFbxSkeletalMeshImportData* ImportData = Cast<UFbxSkeletalMeshImportData>(SkeletalMesh->AssetImportData);
 	
 	// Prepare the import options
-	UFbxImportUI* ReimportUI = ConstructObject<UFbxImportUI>(UFbxImportUI::StaticClass());
+	UFbxImportUI* ReimportUI = NewObject<UFbxImportUI>();
 	ReimportUI->MeshTypeToImport = FBXIT_SkeletalMesh;
 	ReimportUI->bOverrideFullName = false;
 	ReimportUI->Skeleton = SkeletalMesh->Skeleton;
@@ -6763,7 +6763,7 @@ void UReimportDestructibleMeshFactory::SetReimportPaths( UObject* Obj, const TAr
 		if ( DestructibleMesh->AssetImportData == nullptr )
 		{
 			// @todo AssetImportData make an apex destructible import data class
-			DestructibleMesh->AssetImportData = ConstructObject<UAssetImportData>(UAssetImportData::StaticClass(), DestructibleMesh);
+			DestructibleMesh->AssetImportData = NewObject<UAssetImportData>(DestructibleMesh);
 		}
 
 		DestructibleMesh->AssetImportData->SourceFilePath = FReimportManager::SanitizeImportFilename(NewReimportPaths[0], DestructibleMesh);
@@ -6783,7 +6783,7 @@ EReimportResult::Type UReimportDestructibleMeshFactory::Reimport( UObject* Obj )
 	if ( DestructibleMesh->AssetImportData == nullptr )
 	{
 		// @todo AssetImportData make an apex destructible import data class
-		DestructibleMesh->AssetImportData = ConstructObject<UAssetImportData>(UAssetImportData::StaticClass(), DestructibleMesh);
+		DestructibleMesh->AssetImportData = NewObject<UAssetImportData>(DestructibleMesh);
 	}
 
 	const FString Filename = FReimportManager::ResolveImportFilename(DestructibleMesh->AssetImportData->SourceFilePath, DestructibleMesh);
@@ -7107,7 +7107,7 @@ UForceFeedbackEffectFactory::UForceFeedbackEffectFactory(const FObjectInitialize
 
 UObject* UForceFeedbackEffectFactory::FactoryCreateNew( UClass* InClass, UObject* InParent, FName InName, EObjectFlags Flags, UObject* Context, FFeedbackContext* Warn )
 {
-	return ConstructObject<UForceFeedbackEffect>( UForceFeedbackEffect::StaticClass(), InParent, InName, Flags );
+	return NewObject<UForceFeedbackEffect>(InParent, InName, Flags);
 }
 
 /*-----------------------------------------------------------------------------
@@ -7125,7 +7125,7 @@ USubsurfaceProfileFactory::USubsurfaceProfileFactory(const FObjectInitializer& O
 
 UObject* USubsurfaceProfileFactory::FactoryCreateNew(UClass* InClass, UObject* InParent, FName InName, EObjectFlags Flags, UObject* Context, FFeedbackContext* Warn)
 {
-	return ConstructObject<USubsurfaceProfile>(USubsurfaceProfile::StaticClass(), InParent, InName, Flags);
+	return NewObject<USubsurfaceProfile>(InParent, InName, Flags);
 }
 
 /*-----------------------------------------------------------------------------
@@ -7144,7 +7144,7 @@ UTouchInterfaceFactory::UTouchInterfaceFactory(const FObjectInitializer& ObjectI
 
 UObject* UTouchInterfaceFactory::FactoryCreateNew( UClass* InClass, UObject* InParent, FName InName, EObjectFlags Flags, UObject* Context, FFeedbackContext* Warn )
 {
-	return ConstructObject<UTouchInterface>( UTouchInterface::StaticClass(), InParent, InName, Flags );
+	return NewObject<UTouchInterface>(InParent, InName, Flags);
 }
 
 /*------------------------------------------------------------------------------
@@ -7176,7 +7176,7 @@ uint32 UCameraAnimFactory::GetMenuCategories() const
 UObject* UCameraAnimFactory::FactoryCreateNew(UClass* Class,UObject* InParent,FName Name,EObjectFlags Flags,UObject* Context,FFeedbackContext* Warn)
 {
 	UCameraAnim* NewCamAnim = CastChecked<UCameraAnim>(StaticConstructObject(Class,InParent,Name,Flags));
-	NewCamAnim->CameraInterpGroup = ConstructObject<UInterpGroupCamera>(UInterpGroupCamera::StaticClass(), NewCamAnim);
+	NewCamAnim->CameraInterpGroup = NewObject<UInterpGroupCamera>(NewCamAnim);
 	NewCamAnim->CameraInterpGroup->GroupName = Name;
 	return NewCamAnim;
 }
@@ -7315,7 +7315,7 @@ UObject* UDataTableFactory::FactoryCreateNew(UClass* Class, UObject* InParent, F
 	if (Struct && ensure(UDataTable::StaticClass() == Class))
 	{
 		ensure(0 != (RF_Public & Flags));
-		DataTable = NewNamedObject<UDataTable>(InParent, Name, Flags);
+		DataTable = NewObject<UDataTable>(InParent, Name, Flags);
 		if (DataTable)
 		{
 			DataTable->RowStruct = Struct;
diff --git a/Engine/Source/Editor/UnrealEd/Private/Factories/Factory.cpp b/Engine/Source/Editor/UnrealEd/Private/Factories/Factory.cpp
--- a/Engine/Source/Editor/UnrealEd/Private/Factories/Factory.cpp
+++ b/Engine/Source/Editor/UnrealEd/Private/Factories/Factory.cpp
@@ -224,6 +224,7 @@ UObject* UFactory::StaticImportObject
 	}
 	else
 	{
+		auto TransientPackage = GetTransientPackage();
 		// Try all automatic factories, sorted by priority.
 		for( TObjectIterator<UClass> It; It; ++It )
 		{
@@ -232,7 +233,7 @@ UObject* UFactory::StaticImportObject
 				UFactory* Default = It->GetDefaultObject<UFactory>();
 				if (Class->IsChildOf(Default->SupportedClass) && Default->ImportPriority >= 0)
 				{
-					Factories.Add( ConstructObject<UFactory>( *It ) );
+					Factories.Add(NewObject<UFactory>(TransientPackage, *It));
 				}
 			}
 		}
diff --git a/Engine/Source/Editor/UnrealEd/Private/Factories/NiagaraEffectFactory.cpp b/Engine/Source/Editor/UnrealEd/Private/Factories/NiagaraEffectFactory.cpp
--- a/Engine/Source/Editor/UnrealEd/Private/Factories/NiagaraEffectFactory.cpp
+++ b/Engine/Source/Editor/UnrealEd/Private/Factories/NiagaraEffectFactory.cpp
@@ -25,7 +25,7 @@ UObject* UNiagaraEffectFactoryNew::FactoryCreateNew(UClass* Class, UObject* InPa
 	check(Class->IsChildOf(UNiagaraEffect::StaticClass()));
 
 	// First allocate runtime script 
-	UNiagaraEffect* NewEffect = ConstructObject<UNiagaraEffect>(Class, InParent, Name, Flags);
+	UNiagaraEffect* NewEffect = NewObject<UNiagaraEffect>(InParent, Class, Name, Flags);
 	if (NewEffect != NULL)
 	{
 		// Then allocate editor-only objects
diff --git a/Engine/Source/Editor/UnrealEd/Private/Factories/NiagaraScriptFactory.cpp b/Engine/Source/Editor/UnrealEd/Private/Factories/NiagaraScriptFactory.cpp
--- a/Engine/Source/Editor/UnrealEd/Private/Factories/NiagaraScriptFactory.cpp
+++ b/Engine/Source/Editor/UnrealEd/Private/Factories/NiagaraScriptFactory.cpp
@@ -25,15 +25,15 @@ UObject* UNiagaraScriptFactoryNew::FactoryCreateNew(UClass* Class,UObject* InPar
 	check(Class->IsChildOf(UNiagaraScript::StaticClass()));
 
 	// First allocate runtime script 
-	UNiagaraScript* NewScript = ConstructObject<UNiagaraScript>(Class, InParent, Name, Flags);
+	UNiagaraScript* NewScript = NewObject<UNiagaraScript>(InParent, Class, Name, Flags);
 	if(NewScript != NULL)
 	{
 		// Then allocate editor-only 'source' object
-		UNiagaraScriptSource* Source = ConstructObject<UNiagaraScriptSource>(UNiagaraScriptSource::StaticClass(), NewScript, NAME_None, RF_Transactional);
+		UNiagaraScriptSource* Source = NewObject<UNiagaraScriptSource>(NewScript, NAME_None, RF_Transactional);
 		if(Source)
 		{
 			// Create 'update graph'
-			UNiagaraGraph* CreatedGraph  = ConstructObject<UNiagaraGraph>(UNiagaraGraph::StaticClass(), Source, NAME_None, RF_Transactional);
+			UNiagaraGraph* CreatedGraph = NewObject<UNiagaraGraph>(Source, NAME_None, RF_Transactional);
 			Source->NodeGraph = CreatedGraph;
 
 			FGraphNodeCreator<UNiagaraNodeOutput> OutputNodeCreator(*CreatedGraph);
diff --git a/Engine/Source/Editor/UnrealEd/Private/Factories/SkeletonFactory.cpp b/Engine/Source/Editor/UnrealEd/Private/Factories/SkeletonFactory.cpp
--- a/Engine/Source/Editor/UnrealEd/Private/Factories/SkeletonFactory.cpp
+++ b/Engine/Source/Editor/UnrealEd/Private/Factories/SkeletonFactory.cpp
@@ -75,7 +75,7 @@ UObject* USkeletonFactory::FactoryCreateNew(UClass* Class, UObject* InParent, FN
 		return NULL;
 	}
 
-	USkeleton* NewAsset = ConstructObject<USkeleton>( Class, InParent, Name, Flags );
+	USkeleton* NewAsset = NewObject<USkeleton>(InParent, Class, Name, Flags);
 	if(NewAsset)
 	{
 		// this should not fail. If so there is problem.
diff --git a/Engine/Source/Editor/UnrealEd/Private/Factories/TTFontImport.cpp b/Engine/Source/Editor/UnrealEd/Private/Factories/TTFontImport.cpp
--- a/Engine/Source/Editor/UnrealEd/Private/Factories/TTFontImport.cpp
+++ b/Engine/Source/Editor/UnrealEd/Private/Factories/TTFontImport.cpp
@@ -76,7 +76,7 @@ int32 FromHex( TCHAR Ch )
 void UTrueTypeFontFactory::SetupFontImportOptions()
 {
 	// Allocate our import options object if it hasn't been created already!
-	ImportOptions = ConstructObject< UFontImportOptions >( UFontImportOptions::StaticClass(), this, NAME_None );
+	ImportOptions = NewObject<UFontImportOptions>(this, NAME_None);
 }
 
 bool UTrueTypeFontFactory::ConfigureProperties()
@@ -153,7 +153,7 @@ UObject* UTrueTypeFontFactory::FactoryCreateNew(
 	}
 
 	// Create font and its texture.
-	auto Font = NewNamedObject<UFont>(InParent, Name, Flags);
+	auto Font = NewObject<UFont>(InParent, Name, Flags);
 	
 	if (ImportOptions->Data.bUseDistanceFieldAlpha)
 	{
@@ -777,7 +777,7 @@ UTexture2D* UTrueTypeFontFactory::CreateTextureFromDC( UFont* Font, HDC dc, int3
 	}
 
 	// Create texture for page.
-	auto Texture = NewNamedObject<UTexture2D>(Font, *TextureString);
+	auto Texture = NewObject<UTexture2D>(Font, *TextureString);
 
 	// note RF_Public because font textures can be referenced directly by material expressions
 	Texture->SetFlags(RF_Public);
diff --git a/Engine/Source/Editor/UnrealEd/Private/Factories/WorldFactory.cpp b/Engine/Source/Editor/UnrealEd/Private/Factories/WorldFactory.cpp
--- a/Engine/Source/Editor/UnrealEd/Private/Factories/WorldFactory.cpp
+++ b/Engine/Source/Editor/UnrealEd/Private/Factories/WorldFactory.cpp
@@ -26,7 +26,7 @@ UObject* UWorldFactory::FactoryCreateNew(UClass* Class, UObject* InParent, FName
 	const bool bAddToRoot = false;
 	UWorld* NewWorld = UWorld::CreateWorld(WorldType, bInformEngineOfWorld, Name, Cast<UPackage>(InParent), bAddToRoot, FeatureLevel);
 	NewWorld->SetFlags(Flags);
-	NewWorld->ThumbnailInfo = ConstructObject<UWorldThumbnailInfo>(UWorldThumbnailInfo::StaticClass(), NewWorld);
+	NewWorld->ThumbnailInfo = NewObject<UWorldThumbnailInfo>(NewWorld);
 
 	return NewWorld;
 }
diff --git a/Engine/Source/Editor/UnrealEd/Private/Fbx/FbxAnimSequenceImportData.cpp b/Engine/Source/Editor/UnrealEd/Private/Fbx/FbxAnimSequenceImportData.cpp
--- a/Engine/Source/Editor/UnrealEd/Private/Fbx/FbxAnimSequenceImportData.cpp
+++ b/Engine/Source/Editor/UnrealEd/Private/Fbx/FbxAnimSequenceImportData.cpp
@@ -16,7 +16,7 @@ UFbxAnimSequenceImportData* UFbxAnimSequenceImportData::GetImportDataForAnimSequ
 	UFbxAnimSequenceImportData* ImportData = Cast<UFbxAnimSequenceImportData>(AnimSequence->AssetImportData);
 	if ( !ImportData )
 	{
-		ImportData = ConstructObject<UFbxAnimSequenceImportData>(UFbxAnimSequenceImportData::StaticClass(), AnimSequence, NAME_None, RF_NoFlags, TemplateForCreation);
+		ImportData = NewObject<UFbxAnimSequenceImportData>(AnimSequence, NAME_None, RF_NoFlags, TemplateForCreation);
 
 		// Try to preserve the source file path if possible
 		if ( AnimSequence->AssetImportData != NULL )
diff --git a/Engine/Source/Editor/UnrealEd/Private/Fbx/FbxFactory.cpp b/Engine/Source/Editor/UnrealEd/Private/Fbx/FbxFactory.cpp
--- a/Engine/Source/Editor/UnrealEd/Private/Fbx/FbxFactory.cpp
+++ b/Engine/Source/Editor/UnrealEd/Private/Fbx/FbxFactory.cpp
@@ -63,7 +63,7 @@ void UFbxFactory::PostInitProperties()
 	bEditorImport = true;
 	bText = false;
 
-	ImportUI = ConstructObject<UFbxImportUI>( UFbxImportUI::StaticClass(), this, NAME_None, RF_NoFlags );
+	ImportUI = NewObject<UFbxImportUI>(this, NAME_None, RF_NoFlags);
 }
 
 
@@ -586,10 +586,10 @@ UFbxImportUI::UFbxImportUI(const FObjectInitializer& ObjectInitializer)
 {
 	bCombineMeshes = true;
 
-	StaticMeshImportData = ConstructObject<UFbxStaticMeshImportData>(UFbxStaticMeshImportData::StaticClass(), this);
-	SkeletalMeshImportData = ConstructObject<UFbxSkeletalMeshImportData>(UFbxSkeletalMeshImportData::StaticClass(), this);
-	AnimSequenceImportData = ConstructObject<UFbxAnimSequenceImportData>(UFbxAnimSequenceImportData::StaticClass(), this);
-	TextureImportData = ConstructObject<UFbxTextureImportData>(UFbxTextureImportData::StaticClass(), this);
+	StaticMeshImportData = CreateDefaultSubobject<UFbxStaticMeshImportData>(TEXT("StaticMeshImportData"));
+	SkeletalMeshImportData = CreateDefaultSubobject<UFbxSkeletalMeshImportData>(TEXT("SkeletalMeshImportData"));
+	AnimSequenceImportData = CreateDefaultSubobject<UFbxAnimSequenceImportData>(TEXT("AnimSequenceImportData"));
+	TextureImportData = CreateDefaultSubobject<UFbxTextureImportData>(TEXT("TextureImportData"));
 }
 
 
diff --git a/Engine/Source/Editor/UnrealEd/Private/Fbx/FbxMaterialImport.cpp b/Engine/Source/Editor/UnrealEd/Private/Fbx/FbxMaterialImport.cpp
--- a/Engine/Source/Editor/UnrealEd/Private/Fbx/FbxMaterialImport.cpp
+++ b/Engine/Source/Editor/UnrealEd/Private/Fbx/FbxMaterialImport.cpp
@@ -255,7 +255,7 @@ bool UnFbx::FFbxImporter::CreateAndLinkExpressionForMaterialProperty(
 						float ScaleV = FbxTexture->GetScaleV();
 
 						// and link it to the material 
-						UMaterialExpressionTextureSample* UnrealTextureExpression = ConstructObject<UMaterialExpressionTextureSample>( UMaterialExpressionTextureSample::StaticClass(), UnrealMaterial );
+						UMaterialExpressionTextureSample* UnrealTextureExpression = NewObject<UMaterialExpressionTextureSample>(UnrealMaterial);
 						UnrealMaterial->Expressions.Add( UnrealTextureExpression );
 						MaterialInput.Expression = UnrealTextureExpression;
 						UnrealTextureExpression->Texture = UnrealTexture;
@@ -270,7 +270,7 @@ bool UnFbx::FFbxImporter::CreateAndLinkExpressionForMaterialProperty(
 						if( (SetIndex != 0 && SetIndex != INDEX_NONE) || ScaleU != 1.0f || ScaleV != 1.0f )
 						{
 							// Create a texture coord node for the texture sample
-							UMaterialExpressionTextureCoordinate* MyCoordExpression = ConstructObject<UMaterialExpressionTextureCoordinate>(UMaterialExpressionTextureCoordinate::StaticClass(), UnrealMaterial);
+							UMaterialExpressionTextureCoordinate* MyCoordExpression = NewObject<UMaterialExpressionTextureCoordinate>(UnrealMaterial);
 							UnrealMaterial->Expressions.Add(MyCoordExpression);
 							MyCoordExpression->CoordinateIndex = (SetIndex >= 0) ? SetIndex : 0;
 							MyCoordExpression->UTiling = ScaleU;
@@ -312,7 +312,7 @@ void UnFbx::FFbxImporter::FixupMaterial( FbxSurfaceMaterial& FbxMaterial, UMater
 	{
 		FbxDouble3 DiffuseColor;
 		
-		UMaterialExpressionVectorParameter* MyColorExpression = ConstructObject<UMaterialExpressionVectorParameter>( UMaterialExpressionVectorParameter::StaticClass(), UnrealMaterial );
+		UMaterialExpressionVectorParameter* MyColorExpression = NewObject<UMaterialExpressionVectorParameter>(UnrealMaterial);
 		UnrealMaterial->Expressions.Add( MyColorExpression );
 		UnrealMaterial->BaseColor.Expression = MyColorExpression;
 
diff --git a/Engine/Source/Editor/UnrealEd/Private/Fbx/FbxMatineeImport.cpp b/Engine/Source/Editor/UnrealEd/Private/Fbx/FbxMatineeImport.cpp
--- a/Engine/Source/Editor/UnrealEd/Private/Fbx/FbxMatineeImport.cpp
+++ b/Engine/Source/Editor/UnrealEd/Private/Fbx/FbxMatineeImport.cpp
@@ -387,9 +387,9 @@ void FFbxImporter::ImportAnimatedProperty(float* Value, const TCHAR* ValueName,
 	// If a track for this property was not found, create one.
 	if (PropertyTrack == NULL)
 	{
-		PropertyTrack = ConstructObject<UInterpTrackFloatProp>( UInterpTrackFloatProp::StaticClass(), MatineeGroup->Group, NAME_None, RF_Transactional );
+		PropertyTrack = NewObject<UInterpTrackFloatProp>(MatineeGroup->Group, NAME_None, RF_Transactional);
 		MatineeGroup->Group->InterpTracks.Add(PropertyTrack);
-		UInterpTrackInstFloatProp* PropertyTrackInst = ConstructObject<UInterpTrackInstFloatProp>( UInterpTrackInstFloatProp::StaticClass(), MatineeGroup, NAME_None, RF_Transactional );
+		UInterpTrackInstFloatProp* PropertyTrackInst = NewObject<UInterpTrackInstFloatProp>(MatineeGroup, NAME_None, RF_Transactional);
 		MatineeGroup->TrackInst.Add(PropertyTrackInst);
 		PropertyTrack->PropertyName = ValueName;
 		PropertyTrack->TrackTitle = ValueName;
@@ -461,12 +461,12 @@ void FFbxImporter::ImportAnimatedProperty(float* Value, const TCHAR* ValueName,
 UInterpGroupInst* FFbxImporter::CreateMatineeGroup(AMatineeActor* InMatineeActor, AActor* Actor, FString GroupName)
 {
 	// There are no groups for this actor: create the Matinee group data structure.
-	UInterpGroup* MatineeGroupData = ConstructObject<UInterpGroup>( UInterpGroup::StaticClass(), InMatineeActor->MatineeData, NAME_None, RF_Transactional );
+	UInterpGroup* MatineeGroupData = NewObject<UInterpGroup>(InMatineeActor->MatineeData, NAME_None, RF_Transactional);
 	MatineeGroupData->GroupName = FName( *GroupName );
 	InMatineeActor->MatineeData->InterpGroups.Add(MatineeGroupData);
 
 	// Instantiate the Matinee group data structure.
-	UInterpGroupInst* MatineeGroup = ConstructObject<UInterpGroupInst>( UInterpGroupInst::StaticClass(), InMatineeActor, NAME_None, RF_Transactional );
+	UInterpGroupInst* MatineeGroup = NewObject<UInterpGroupInst>(InMatineeActor, NAME_None, RF_Transactional);
 	InMatineeActor->GroupInst.Add(MatineeGroup);
 	MatineeGroup->InitGroupInst(MatineeGroupData, Actor);
 	MatineeGroup->SaveGroupActorState();
@@ -598,9 +598,9 @@ float FFbxImporter::ImportMatineeActor(FbxNode* Node, UInterpGroupInst* MatineeG
 	// Add a Movement track if the node is animated and the group does not already have one.
 	if (MovementTrack == NULL && bNodeAnimated)
 	{
-		MovementTrack = ConstructObject<UInterpTrackMove>( UInterpTrackMove::StaticClass(), MatineeGroup->Group, NAME_None, RF_Transactional );
+		MovementTrack = NewObject<UInterpTrackMove>(MatineeGroup->Group, NAME_None, RF_Transactional);
 		MatineeGroup->Group->InterpTracks.Add(MovementTrack);
-		UInterpTrackInstMove* MovementTrackInst = ConstructObject<UInterpTrackInstMove>( UInterpTrackInstMove::StaticClass(), MatineeGroup, NAME_None, RF_Transactional );
+		UInterpTrackInstMove* MovementTrackInst = NewObject<UInterpTrackInstMove>(MatineeGroup, NAME_None, RF_Transactional);
 		MatineeGroup->TrackInst.Add(MovementTrackInst);
 		MovementTrackInst->InitTrackInst(MovementTrack);
 	}
diff --git a/Engine/Source/Editor/UnrealEd/Private/Fbx/FbxSkeletalMeshImport.cpp b/Engine/Source/Editor/UnrealEd/Private/Fbx/FbxSkeletalMeshImport.cpp
--- a/Engine/Source/Editor/UnrealEd/Private/Fbx/FbxSkeletalMeshImport.cpp
+++ b/Engine/Source/Editor/UnrealEd/Private/Fbx/FbxSkeletalMeshImport.cpp
@@ -2785,7 +2785,7 @@ void UnFbx::FFbxImporter::ImportMorphTargetsInternal( TArray<FbxNode*>& SkelMesh
 		{
 			if (LODIndex == 0)
 			{
-				Result = NewNamedObject<UMorphTarget>(BaseSkelMesh, FName(*ShapeName));
+				Result = NewObject<UMorphTarget>(BaseSkelMesh, FName(*ShapeName));
 			}
 			else
 			{
diff --git a/Engine/Source/Editor/UnrealEd/Private/Fbx/FbxSkeletalMeshImportData.cpp b/Engine/Source/Editor/UnrealEd/Private/Fbx/FbxSkeletalMeshImportData.cpp
--- a/Engine/Source/Editor/UnrealEd/Private/Fbx/FbxSkeletalMeshImportData.cpp
+++ b/Engine/Source/Editor/UnrealEd/Private/Fbx/FbxSkeletalMeshImportData.cpp
@@ -15,7 +15,7 @@ UFbxSkeletalMeshImportData* UFbxSkeletalMeshImportData::GetImportDataForSkeletal
 	UFbxSkeletalMeshImportData* ImportData = Cast<UFbxSkeletalMeshImportData>(SkeletalMesh->AssetImportData);
 	if ( !ImportData )
 	{
-		ImportData = ConstructObject<UFbxSkeletalMeshImportData>(UFbxSkeletalMeshImportData::StaticClass(), SkeletalMesh, NAME_None, RF_NoFlags, TemplateForCreation);
+		ImportData = NewObject<UFbxSkeletalMeshImportData>(SkeletalMesh, NAME_None, RF_NoFlags, TemplateForCreation);
 
 		// Try to preserve the source file path if possible
 		if ( SkeletalMesh->AssetImportData != NULL )
diff --git a/Engine/Source/Editor/UnrealEd/Private/Fbx/FbxStaticMeshImport.cpp b/Engine/Source/Editor/UnrealEd/Private/Fbx/FbxStaticMeshImport.cpp
--- a/Engine/Source/Editor/UnrealEd/Private/Fbx/FbxStaticMeshImport.cpp
+++ b/Engine/Source/Editor/UnrealEd/Private/Fbx/FbxStaticMeshImport.cpp
@@ -956,7 +956,7 @@ UStaticMesh* UnFbx::FFbxImporter::ImportStaticMeshAsSingle(UObject* InParent, TA
 	}
 	else
 	{
-		StaticMesh = NewNamedObject<UStaticMesh>(Package, FName(*MeshName), Flags | RF_Public);
+		StaticMesh = NewObject<UStaticMesh>(Package, FName(*MeshName), Flags | RF_Public);
 	}
 
 	if (StaticMesh->SourceModels.Num() < LODIndex+1)
@@ -1329,7 +1329,7 @@ void UnFbx::FFbxImporter::ImportStaticMeshSockets( UStaticMesh* StaticMesh )
 		if( !Socket )
 		{
 			// If the socket didn't exist create a new one now
-			Socket = ConstructObject<UStaticMeshSocket>( UStaticMeshSocket::StaticClass(), StaticMesh );
+			Socket = NewObject<UStaticMeshSocket>(StaticMesh);
 			check(Socket);
 
 			Socket->SocketName = SocketNode.SocketName;
diff --git a/Engine/Source/Editor/UnrealEd/Private/Fbx/FbxStaticMeshImportData.cpp b/Engine/Source/Editor/UnrealEd/Private/Fbx/FbxStaticMeshImportData.cpp
--- a/Engine/Source/Editor/UnrealEd/Private/Fbx/FbxStaticMeshImportData.cpp
+++ b/Engine/Source/Editor/UnrealEd/Private/Fbx/FbxStaticMeshImportData.cpp
@@ -21,7 +21,7 @@ UFbxStaticMeshImportData* UFbxStaticMeshImportData::GetImportDataForStaticMesh(U
 	UFbxStaticMeshImportData* ImportData = Cast<UFbxStaticMeshImportData>(StaticMesh->AssetImportData);
 	if ( !ImportData )
 	{
-		ImportData = ConstructObject<UFbxStaticMeshImportData>(UFbxStaticMeshImportData::StaticClass(), StaticMesh, NAME_None, RF_NoFlags, TemplateForCreation);
+		ImportData = NewObject<UFbxStaticMeshImportData>(StaticMesh, NAME_None, RF_NoFlags, TemplateForCreation);
 
 		// Try to preserve the source file path if possible
 		if ( StaticMesh->AssetImportData != NULL )
diff --git a/Engine/Source/Editor/UnrealEd/Private/Kismet2/BlueprintEditorUtils.cpp b/Engine/Source/Editor/UnrealEd/Private/Kismet2/BlueprintEditorUtils.cpp
--- a/Engine/Source/Editor/UnrealEd/Private/Kismet2/BlueprintEditorUtils.cpp
+++ b/Engine/Source/Editor/UnrealEd/Private/Kismet2/BlueprintEditorUtils.cpp
@@ -1461,8 +1461,8 @@ void FBlueprintEditorUtils::PostDuplicateBlueprint(UBlueprint* Blueprint)
 			FName NewSkelClassName, NewGenClassName;
 			Blueprint->GetBlueprintClassNames(NewGenClassName, NewSkelClassName);
 
-			UClass* NewClass = ConstructObject<UClass>(
-				Blueprint->GetBlueprintClass(), Blueprint->GetOutermost(), NewGenClassName, RF_Public | RF_Transactional);
+			UClass* NewClass = NewObject<UClass>(
+				Blueprint->GetOutermost(), Blueprint->GetBlueprintClass(), NewGenClassName, RF_Public | RF_Transactional);
 
 			Blueprint->GeneratedClass = NewClass;
 			NewClass->ClassGeneratedBy = Blueprint;
@@ -1803,13 +1803,13 @@ UEdGraph* FBlueprintEditorUtils::CreateNewGraph(UObject* ParentScope, const FNam
 		}
 
 		// Construct new graph with the supplied name
-		NewGraph = ConstructObject<UEdGraph>(GraphClass, ParentScope, NAME_None, RF_Transactional);
+		NewGraph = NewObject<UEdGraph>(ParentScope, GraphClass, NAME_None, RF_Transactional);
 		bRename = true;
 	}
 	else
 	{
 		// Construct a new graph with a default name
-		NewGraph = ConstructObject<UEdGraph>(GraphClass, ParentScope, NAME_None, RF_Transactional);
+		NewGraph = NewObject<UEdGraph>(ParentScope, GraphClass, NAME_None, RF_Transactional);
 	}
 
 	NewGraph->Schema = SchemaClass;
@@ -5483,7 +5483,7 @@ UTimelineTemplate* FBlueprintEditorUtils::AddNewTimeline(UBlueprint* Blueprint,
 		check(NULL != Blueprint->GeneratedClass);
 		// Construct new graph with the supplied name
 		const FName TimelineTemplateName = *UTimelineTemplate::TimelineVariableNameToTemplateName(TimelineVarName);
-		Timeline = NewNamedObject<UTimelineTemplate>(Blueprint->GeneratedClass, TimelineTemplateName, RF_Transactional);
+		Timeline = NewObject<UTimelineTemplate>(Blueprint->GeneratedClass, TimelineTemplateName, RF_Transactional);
 		Blueprint->Timelines.Add(Timeline);
 
 		// Potentially adjust variable names for any child blueprints
diff --git a/Engine/Source/Editor/UnrealEd/Private/Kismet2/EnumEditorUtils.cpp b/Engine/Source/Editor/UnrealEd/Private/Kismet2/EnumEditorUtils.cpp
--- a/Engine/Source/Editor/UnrealEd/Private/Kismet2/EnumEditorUtils.cpp
+++ b/Engine/Source/Editor/UnrealEd/Private/Kismet2/EnumEditorUtils.cpp
@@ -30,7 +30,7 @@ UEnum* FEnumEditorUtils::CreateUserDefinedEnum(UObject* InParent, FName EnumName
 {
 	ensure(0 != (RF_Public & Flags));
 
-	UEnum* Enum = NewNamedObject<UUserDefinedEnum>(InParent, EnumName, Flags);
+	UEnum* Enum = NewObject<UUserDefinedEnum>(InParent, EnumName, Flags);
 
 	if (NULL != Enum)
 	{
diff --git a/Engine/Source/Editor/UnrealEd/Private/Kismet2/Kismet2.cpp b/Engine/Source/Editor/UnrealEd/Private/Kismet2/Kismet2.cpp
--- a/Engine/Source/Editor/UnrealEd/Private/Kismet2/Kismet2.cpp
+++ b/Engine/Source/Editor/UnrealEd/Private/Kismet2/Kismet2.cpp
@@ -376,7 +376,7 @@ UBlueprint* FKismetEditorUtilities::CreateBlueprint(UClass* ParentClass, UObject
 	}
 	
 	// Create new UBlueprint object
-	UBlueprint* NewBP = ConstructObject<UBlueprint>(*BlueprintClassType, Outer, NewBPName, RF_Public|RF_Standalone|RF_Transactional|RF_LoadCompleted);
+	UBlueprint* NewBP = NewObject<UBlueprint>(Outer, *BlueprintClassType, NewBPName, RF_Public | RF_Standalone | RF_Transactional | RF_LoadCompleted);
 	NewBP->Status = BS_BeingCreated;
 	NewBP->BlueprintType = BlueprintType;
 	NewBP->ParentClass = ParentClass;
@@ -393,8 +393,8 @@ UBlueprint* FKismetEditorUtilities::CreateBlueprint(UClass* ParentClass, UObject
 		// >>> Temporary workaround, before a BlueprintGeneratedClass is the main asset.
 		FName NewSkelClassName, NewGenClassName;
 		NewBP->GetBlueprintClassNames(NewGenClassName, NewSkelClassName);
-		UBlueprintGeneratedClass* NewClass = ConstructObject<UBlueprintGeneratedClass>(
-			*BlueprintGeneratedClassType, NewBP->GetOutermost(), NewGenClassName, RF_Public | RF_Transactional);
+		UBlueprintGeneratedClass* NewClass = NewObject<UBlueprintGeneratedClass>(
+			NewBP->GetOutermost(), *BlueprintGeneratedClassType, NewGenClassName, RF_Public | RF_Transactional);
 		NewBP->GeneratedClass = NewClass;
 		NewClass->ClassGeneratedBy = NewBP;
 		NewClass->SetSuperStruct(ParentClass);
diff --git a/Engine/Source/Editor/UnrealEd/Private/Kismet2/StructureEditorUtils.cpp b/Engine/Source/Editor/UnrealEd/Private/Kismet2/StructureEditorUtils.cpp
--- a/Engine/Source/Editor/UnrealEd/Private/Kismet2/StructureEditorUtils.cpp
+++ b/Engine/Source/Editor/UnrealEd/Private/Kismet2/StructureEditorUtils.cpp
@@ -32,9 +32,9 @@ UUserDefinedStruct* FStructureEditorUtils::CreateUserDefinedStruct(UObject* InPa
 	
 	if (UserDefinedStructEnabled())
 	{
-		Struct = NewNamedObject<UUserDefinedStruct>(InParent, Name, Flags);
+		Struct = NewObject<UUserDefinedStruct>(InParent, Name, Flags);
 		check(Struct);
-		Struct->EditorData = NewNamedObject<UUserDefinedStructEditorData>(Struct, NAME_None, RF_Transactional);
+		Struct->EditorData = NewObject<UUserDefinedStructEditorData>(Struct, NAME_None, RF_Transactional);
 		check(Struct->EditorData);
 
 		Struct->SetMetaData(TEXT("BlueprintType"), TEXT("true"));
diff --git a/Engine/Source/Editor/UnrealEd/Private/Layers/Layers.cpp b/Engine/Source/Editor/UnrealEd/Private/Layers/Layers.cpp
--- a/Engine/Source/Editor/UnrealEd/Private/Layers/Layers.cpp
+++ b/Engine/Source/Editor/UnrealEd/Private/Layers/Layers.cpp
@@ -1019,7 +1019,7 @@ void FLayers::AddAllLayersTo( TArray< TWeakObjectPtr< ULayer > >& OutLayers ) co
 
 TWeakObjectPtr< ULayer > FLayers::CreateLayer( const FName& LayerName )
 {
-	ULayer* NewLayer = ConstructObject< ULayer >( ULayer::StaticClass(), GetWorld(), NAME_None, RF_Transactional );
+	ULayer* NewLayer = NewObject<ULayer>(GetWorld(), NAME_None, RF_Transactional);
 	check( NewLayer != NULL );
 
 	GetWorld()->Modify();
diff --git a/Engine/Source/Editor/UnrealEd/Private/LevelEditorViewport.cpp b/Engine/Source/Editor/UnrealEd/Private/LevelEditorViewport.cpp
--- a/Engine/Source/Editor/UnrealEd/Private/LevelEditorViewport.cpp
+++ b/Engine/Source/Editor/UnrealEd/Private/LevelEditorViewport.cpp
@@ -347,7 +347,7 @@ static bool TryAndCreateMaterialInput( UMaterial* UnrealMaterial, EMaterialKind:
 	bool bSetupAsNormalMap = UnrealTexture->IsNormalMap();
 
 	// Create a new texture sample expression, this is our texture input node into the material output.
-	UMaterialExpressionTextureSample* UnrealTextureExpression = ConstructObject<UMaterialExpressionTextureSample>( UMaterialExpressionTextureSample::StaticClass(), UnrealMaterial );
+	UMaterialExpressionTextureSample* UnrealTextureExpression = NewObject<UMaterialExpressionTextureSample>(UnrealMaterial);
 	UnrealMaterial->Expressions.Add( UnrealTextureExpression );
 	MaterialInput.Expression = UnrealTextureExpression;
 	UnrealTextureExpression->Texture = UnrealTexture;
diff --git a/Engine/Source/Editor/UnrealEd/Private/Lightmass/Lightmass.cpp b/Engine/Source/Editor/UnrealEd/Private/Lightmass/Lightmass.cpp
--- a/Engine/Source/Editor/UnrealEd/Private/Lightmass/Lightmass.cpp
+++ b/Engine/Source/Editor/UnrealEd/Private/Lightmass/Lightmass.cpp
@@ -4020,7 +4020,7 @@ bool FLightmassProcessor::ImportTextureMapping(int32 Channel, FTextureMappingImp
 	int32 TotalMemory = FMath::TruncToInt(TotalMemoryAmount * BytesPerPixel * MIP_FACTOR * LightMapTypeModifier);
 	
 	FStatsViewerModule& StatsViewerModule = FModuleManager::Get().LoadModuleChecked<FStatsViewerModule>(TEXT("StatsViewer"));
-	ULightingBuildInfo* LightingBuildInfo = ConstructObject<ULightingBuildInfo>(ULightingBuildInfo::StaticClass());
+	ULightingBuildInfo* LightingBuildInfo = NewObject<ULightingBuildInfo>();
 	LightingBuildInfo->Set( MappedObject, TMImport.ExecutionTime, TMImport.UnmappedTexelsPercentage, WastedMemory, TotalMemory );
 	StatsViewerModule.GetPage(EStatsPage::LightingBuildInfo)->AddEntry( LightingBuildInfo );
 
diff --git a/Engine/Source/Editor/UnrealEd/Private/ObjectTools.cpp b/Engine/Source/Editor/UnrealEd/Private/ObjectTools.cpp
--- a/Engine/Source/Editor/UnrealEd/Private/ObjectTools.cpp
+++ b/Engine/Source/Editor/UnrealEd/Private/ObjectTools.cpp
@@ -2805,13 +2805,15 @@ namespace ObjectTools
 	 */
 	void AssembleListOfExporters(TArray<UExporter*>& OutExporters)
 	{
+		auto TransientPackage = GetTransientPackage();
+
 		// @todo DB: Assemble this set once.
 		OutExporters.Empty();
 		for( TObjectIterator<UClass> It ; It ; ++It )
 		{
 			if( It->IsChildOf(UExporter::StaticClass()) && !It->HasAnyClassFlags(CLASS_Abstract) )
 			{
-				UExporter* Exporter = ConstructObject<UExporter>( *It );
+				UExporter* Exporter = NewObject<UExporter>(TransientPackage, *It);
 				OutExporters.Add( Exporter );		
 			}
 		}
diff --git a/Engine/Source/Editor/UnrealEd/Private/PhysicsAssetUtils.cpp b/Engine/Source/Editor/UnrealEd/Private/PhysicsAssetUtils.cpp
--- a/Engine/Source/Editor/UnrealEd/Private/PhysicsAssetUtils.cpp
+++ b/Engine/Source/Editor/UnrealEd/Private/PhysicsAssetUtils.cpp
@@ -677,7 +677,7 @@ int32 CreateNewConstraint(UPhysicsAsset* PhysAsset, FName InConstraintName, UPhy
 		return ConstraintIndex;
 	}
 
-	UPhysicsConstraintTemplate* NewConstraintSetup = ConstructObject<UPhysicsConstraintTemplate>( UPhysicsConstraintTemplate::StaticClass(), PhysAsset, NAME_None, RF_Transactional );
+	UPhysicsConstraintTemplate* NewConstraintSetup = NewObject<UPhysicsConstraintTemplate>(PhysAsset, NAME_None, RF_Transactional);
 	if(InConstraintSetup)
 	{
 		NewConstraintSetup->DefaultInstance.CopyConstraintParamsFrom( &InConstraintSetup->DefaultInstance );
@@ -706,7 +706,7 @@ int32 CreateNewBody(UPhysicsAsset* PhysAsset, FName InBodyName)
 		return BodyIndex; // if we already have one for this name - just return that.
 	}
 
-	UBodySetup* NewBodySetup = ConstructObject<UBodySetup>( UBodySetup::StaticClass(), PhysAsset, NAME_None, RF_Transactional );
+	UBodySetup* NewBodySetup = NewObject<UBodySetup>(PhysAsset, NAME_None, RF_Transactional);
 	// make default to be use complex as simple 
 	NewBodySetup->CollisionTraceFlag = CTF_UseSimpleAsComplex;
 	// newly created bodies default to simulating
diff --git a/Engine/Source/Editor/UnrealEd/Private/PhysicsManipulationMode.cpp b/Engine/Source/Editor/UnrealEd/Private/PhysicsManipulationMode.cpp
--- a/Engine/Source/Editor/UnrealEd/Private/PhysicsManipulationMode.cpp
+++ b/Engine/Source/Editor/UnrealEd/Private/PhysicsManipulationMode.cpp
@@ -41,7 +41,7 @@ TSharedRef<FEdMode> FPhysicsManipulationEdModeFactory::CreateMode() const
 
 FPhysicsManipulationEdMode::FPhysicsManipulationEdMode()
 {
-	HandleComp = ConstructObject<UPhysicsHandleComponent>(UPhysicsHandleComponent::StaticClass(), GetTransientPackage(), NAME_None, RF_NoFlags);
+	HandleComp = NewObject<UPhysicsHandleComponent>();
 }
 
 FPhysicsManipulationEdMode::~FPhysicsManipulationEdMode()
diff --git a/Engine/Source/Editor/UnrealEd/Private/PlayLevel.cpp b/Engine/Source/Editor/UnrealEd/Private/PlayLevel.cpp
--- a/Engine/Source/Editor/UnrealEd/Private/PlayLevel.cpp
+++ b/Engine/Source/Editor/UnrealEd/Private/PlayLevel.cpp
@@ -2525,7 +2525,7 @@ UGameInstance* UEditorEngine::CreatePIEGameInstance(int32 PIEInstance, bool bInS
 	{
 		GameInstanceClass = UGameInstance::StaticClass();
 	}
-	UGameInstance* GameInstance = ConstructObject<UGameInstance>(GameInstanceClass, this);
+	UGameInstance* GameInstance = NewObject<UGameInstance>(this, GameInstanceClass);
 
 	// We need to temporarily add the GameInstance to the root because the InitPIE call can do garbage collection wiping out the GameInstance
 	GameInstance->AddToRoot();
@@ -2631,7 +2631,7 @@ UGameInstance* UEditorEngine::CreatePIEGameInstance(int32 PIEInstance, bool bInS
 	
 	if (!PieWorldContext->RunAsDedicated)
 	{
-		ViewportClient = ConstructObject<UGameViewportClient>(GameViewportClientClass,this);
+		ViewportClient = NewObject<UGameViewportClient>(this, GameViewportClientClass);
 		ViewportClient->Init(*PieWorldContext, GameInstance);
 		GameViewport = ViewportClient;
 		GameViewport->bIsPlayInEditorViewport = true;
diff --git a/Engine/Source/Editor/UnrealEd/Private/PreviewMaterial.cpp b/Engine/Source/Editor/UnrealEd/Private/PreviewMaterial.cpp
--- a/Engine/Source/Editor/UnrealEd/Private/PreviewMaterial.cpp
+++ b/Engine/Source/Editor/UnrealEd/Private/PreviewMaterial.cpp
@@ -282,7 +282,7 @@ void UMaterialEditorInstanceConstant::RegenerateArrays()
 
 		for(int32 ParameterIdx=0; ParameterIdx<ParameterNames.Num(); ParameterIdx++)
 		{
-			UDEditorVectorParameterValue & ParameterValue = *(ConstructObject<UDEditorVectorParameterValue>(UDEditorVectorParameterValue::StaticClass()));
+			UDEditorVectorParameterValue & ParameterValue = *(NewObject<UDEditorVectorParameterValue>());
 			FName ParameterName = ParameterNames[ParameterIdx];
 			FLinearColor Value;
 
@@ -312,7 +312,7 @@ void UMaterialEditorInstanceConstant::RegenerateArrays()
 		ParentMaterial->GetAllScalarParameterNames(ParameterNames, Guids);
 		for(int32 ParameterIdx=0; ParameterIdx<ParameterNames.Num(); ParameterIdx++)
 		{			
-			UDEditorScalarParameterValue & ParameterValue = *(ConstructObject<UDEditorScalarParameterValue>(UDEditorScalarParameterValue::StaticClass()));
+			UDEditorScalarParameterValue & ParameterValue = *(NewObject<UDEditorScalarParameterValue>());
 			FName ParameterName = ParameterNames[ParameterIdx];
 			float Value;
 
@@ -344,7 +344,7 @@ void UMaterialEditorInstanceConstant::RegenerateArrays()
 		ParentMaterial->GetAllTextureParameterNames(ParameterNames, Guids);
 		for(int32 ParameterIdx=0; ParameterIdx<ParameterNames.Num(); ParameterIdx++)
 		{			
-			UDEditorTextureParameterValue& ParameterValue = *(ConstructObject<UDEditorTextureParameterValue>(UDEditorTextureParameterValue::StaticClass()));
+			UDEditorTextureParameterValue& ParameterValue = *(NewObject<UDEditorTextureParameterValue>());
 			FName ParameterName = ParameterNames[ParameterIdx];
 			UTexture* Value;
 
@@ -376,7 +376,7 @@ void UMaterialEditorInstanceConstant::RegenerateArrays()
 		ParentMaterial->GetAllFontParameterNames(ParameterNames, Guids);
 		for(int32 ParameterIdx=0; ParameterIdx<ParameterNames.Num(); ParameterIdx++)
 		{
-			UDEditorFontParameterValue& ParameterValue = *(ConstructObject<UDEditorFontParameterValue>(UDEditorFontParameterValue::StaticClass()));
+			UDEditorFontParameterValue& ParameterValue = *(NewObject<UDEditorFontParameterValue>());
 			FName ParameterName = ParameterNames[ParameterIdx];
 			UFont* FontValue;
 			int32 FontPage;
@@ -415,7 +415,7 @@ void UMaterialEditorInstanceConstant::RegenerateArrays()
 		for(int32 ParameterIdx=0; ParameterIdx<SourceStaticParameters.StaticSwitchParameters.Num(); ParameterIdx++)
 		{			
 			FStaticSwitchParameter StaticSwitchParameterValue = FStaticSwitchParameter(SourceStaticParameters.StaticSwitchParameters[ParameterIdx]);
-			UDEditorStaticSwitchParameterValue& ParameterValue = *(ConstructObject<UDEditorStaticSwitchParameterValue>(UDEditorStaticSwitchParameterValue::StaticClass()));
+			UDEditorStaticSwitchParameterValue& ParameterValue = *(NewObject<UDEditorStaticSwitchParameterValue>());
 			ParameterValue.ParameterValue =StaticSwitchParameterValue.Value;
 			ParameterValue.bOverride =StaticSwitchParameterValue.bOverride;
 			ParameterValue.ParameterName =StaticSwitchParameterValue.ParameterName;
@@ -429,7 +429,7 @@ void UMaterialEditorInstanceConstant::RegenerateArrays()
 		for(int32 ParameterIdx=0; ParameterIdx<SourceStaticParameters.StaticComponentMaskParameters.Num(); ParameterIdx++)
 		{
 			FStaticComponentMaskParameter StaticComponentMaskParameterValue = FStaticComponentMaskParameter(SourceStaticParameters.StaticComponentMaskParameters[ParameterIdx]);
-			UDEditorStaticComponentMaskParameterValue& ParameterValue = *(ConstructObject<UDEditorStaticComponentMaskParameterValue>(UDEditorStaticComponentMaskParameterValue::StaticClass()));
+			UDEditorStaticComponentMaskParameterValue& ParameterValue = *(NewObject<UDEditorStaticComponentMaskParameterValue>());
 			ParameterValue.ParameterValue.R = StaticComponentMaskParameterValue.R;
 			ParameterValue.ParameterValue.G = StaticComponentMaskParameterValue.G;
 			ParameterValue.ParameterValue.B = StaticComponentMaskParameterValue.B;
diff --git a/Engine/Source/Editor/UnrealEd/Private/SCurveEditor.cpp b/Engine/Source/Editor/UnrealEd/Private/SCurveEditor.cpp
--- a/Engine/Source/Editor/UnrealEd/Private/SCurveEditor.cpp
+++ b/Engine/Source/Editor/UnrealEd/Private/SCurveEditor.cpp
@@ -1955,7 +1955,7 @@ void SCurveEditor::OnCreateExternalCurveClicked()
 UObject* SCurveEditor::CreateCurveObject( TSubclassOf<UCurveBase> CurveType, UObject* PackagePtr, FName& AssetName )
 {
 	UObject* NewObj = NULL;
-	CurveFactory = Cast<UCurveFactory>(ConstructObject<UFactory>( UCurveFactory::StaticClass() ) );
+	CurveFactory = Cast<UCurveFactory>(NewObject<UFactory>(GetTransientPackage(), UCurveFactory::StaticClass()));
 	if(CurveFactory)
 	{
 		CurveFactory->CurveClass = CurveType;
diff --git a/Engine/Source/Editor/UnrealEd/Private/SSocketManager.cpp b/Engine/Source/Editor/UnrealEd/Private/SSocketManager.cpp
--- a/Engine/Source/Editor/UnrealEd/Private/SSocketManager.cpp
+++ b/Engine/Source/Editor/UnrealEd/Private/SSocketManager.cpp
@@ -421,7 +421,7 @@ void SSocketManager::CreateSocket()
 
 		const FScopedTransaction Transaction( LOCTEXT( "CreateSocket", "Create Socket" ) );
 
-		UStaticMeshSocket* NewSocket = ConstructObject<UStaticMeshSocket>(UStaticMeshSocket::StaticClass(), CurrentStaticMesh);
+		UStaticMeshSocket* NewSocket = NewObject<UStaticMeshSocket>(CurrentStaticMesh);
 		check(NewSocket);
 
 		if (FEngineAnalytics::IsAvailable())
diff --git a/Engine/Source/Editor/UnrealEd/Private/SkeletalMeshEdit.cpp b/Engine/Source/Editor/UnrealEd/Private/SkeletalMeshEdit.cpp
--- a/Engine/Source/Editor/UnrealEd/Private/SkeletalMeshEdit.cpp
+++ b/Engine/Source/Editor/UnrealEd/Private/SkeletalMeshEdit.cpp
@@ -109,7 +109,7 @@ bool UEditorEngine::ReimportFbxAnimation( USkeleton* Skeleton, UAnimSequence* An
 	if ( ImportData )
 	{
 		// Prepare the import options
-		UFbxImportUI* ReimportUI = ConstructObject<UFbxImportUI>(UFbxImportUI::StaticClass());
+		UFbxImportUI* ReimportUI = NewObject<UFbxImportUI>();
 		ReimportUI->MeshTypeToImport = FBXIT_Animation;
 		ReimportUI->bOverrideFullName = false;
 		ReimportUI->AnimSequenceImportData = ImportData;
@@ -535,7 +535,7 @@ UAnimSequence * UnFbx::FFbxImporter::ImportAnimations(USkeleton* Skeleton, UObje
 		// If not, create new one now.
 		if(!DestSeq)
 		{
-			DestSeq = ConstructObject<UAnimSequence>( UAnimSequence::StaticClass(), ParentPackage, *SequenceName, RF_Public|RF_Standalone );
+			DestSeq = NewObject<UAnimSequence>(ParentPackage, *SequenceName, RF_Public | RF_Standalone);
 	
 			// Notify the asset registry
 			FAssetRegistryModule::AssetCreated(DestSeq);
diff --git a/Engine/Source/Editor/UnrealEd/Private/Tests/BlueprintAutomationTests.cpp b/Engine/Source/Editor/UnrealEd/Private/Tests/BlueprintAutomationTests.cpp
--- a/Engine/Source/Editor/UnrealEd/Private/Tests/BlueprintAutomationTests.cpp
+++ b/Engine/Source/Editor/UnrealEd/Private/Tests/BlueprintAutomationTests.cpp
@@ -938,7 +938,7 @@ bool FBlueprintCompileOnLoadTest::RunTest(const FString& BlueprintAssetPath)
 	FName ReconstructedName = MakeUniqueObjectName(TransientPackage, UBlueprint::StaticClass(), BlueprintName);
 	// reconstruct the initial blueprint (using the serialized data from its initial load)
 	EObjectFlags const StandardBlueprintFlags = RF_Public | RF_Standalone | RF_Transactional;
-	InitialBlueprint = ConstructObject<UBlueprint>(UBlueprint::StaticClass(), TransientPackage, ReconstructedName, StandardBlueprintFlags | RF_Transient);
+	InitialBlueprint = NewObject<UBlueprint>(TransientPackage, ReconstructedName, StandardBlueprintFlags | RF_Transient);
 	FObjectReader(InitialBlueprint, InitialLoadData);
 	{
 		TMap<UObject*, UObject*> ClassRedirects;
diff --git a/Engine/Source/Editor/UnrealEd/Private/Tests/EditorAssetAutomationTests.cpp b/Engine/Source/Editor/UnrealEd/Private/Tests/EditorAssetAutomationTests.cpp
--- a/Engine/Source/Editor/UnrealEd/Private/Tests/EditorAssetAutomationTests.cpp
+++ b/Engine/Source/Editor/UnrealEd/Private/Tests/EditorAssetAutomationTests.cpp
@@ -426,7 +426,7 @@ bool FLogAssetCreationStatsCommand::Update()
 #define ASSET_TEST_CREATE( TAssetClass, TFactoryClass, NamePrefix, ExtraCommands ) \
 { \
 	FString NameString = FString::Printf(TEXT("%s_%s"), TEXT(#NamePrefix), *CurrentTimestamp); \
-	TFactoryClass* FactoryInst = ConstructObject<TFactoryClass>(TFactoryClass::StaticClass()); \
+	TFactoryClass* FactoryInst = NewObject<TFactoryClass>(); \
 	ExtraCommands \
 	TSharedPtr<CreateAssetHelper::FCreateAssetInfo> CreateInfo = MakeShareable(new CreateAssetHelper::FCreateAssetInfo(NameString, GamePath, TAssetClass::StaticClass(), FactoryInst, BuildStats)); \
 	AssetInfos.Add(CreateInfo); \
@@ -440,7 +440,7 @@ bool FLogAssetCreationStatsCommand::Update()
 	UClass* FactoryClass = StaticLoadClass(UFactory::StaticClass(), NULL, TEXT(#TFactoryClassName), NULL, LOAD_None, NULL); \
 	if (FactoryClass) \
 	{ \
-		UFactory* FactoryInst = ConstructObject<UFactory>(FactoryClass); \
+		UFactory* FactoryInst = NewObject<UFactory>(GetTransientPackage(), FactoryClass); \
 		ExtraCommands \
 		UClass* AssetClass = StaticLoadClass(UObject::StaticClass(), NULL, TEXT(#TAssetClassName), NULL, LOAD_None, NULL); \
 		TSharedPtr<CreateAssetHelper::FCreateAssetInfo> CreateInfo = MakeShareable(new CreateAssetHelper::FCreateAssetInfo(NameString, GamePath, AssetClass, FactoryInst, BuildStats)); \
@@ -821,7 +821,7 @@ namespace ImportExportAssetHelper
 			{
 				GWarn->BeginSlowTask(LOCTEXT("ImportSlowTask", "Importing"), true);
 
-				UFactory* ImportFactory = ConstructObject<UFactory>(FactoryClass);
+				UFactory* ImportFactory = NewObject<UFactory>(GetTransientPackage(), FactoryClass);
 
 				//Apply any custom settings to the factory
 				AutomationEditorCommonUtils::ApplyCustomFactorySettings(ImportFactory, FactorySettings);
diff --git a/Engine/Source/Editor/UnrealEd/Private/Tests/EditorBuildPromotionTests.cpp b/Engine/Source/Editor/UnrealEd/Private/Tests/EditorBuildPromotionTests.cpp
--- a/Engine/Source/Editor/UnrealEd/Private/Tests/EditorBuildPromotionTests.cpp
+++ b/Engine/Source/Editor/UnrealEd/Private/Tests/EditorBuildPromotionTests.cpp
@@ -259,7 +259,7 @@ namespace EditorBuildPromotionTestUtils
 	static UMaterial* CreateMaterialFromTexture(UTexture* InTexture)
 	{
 		// Create the factory used to generate the asset
-		UMaterialFactoryNew* Factory = ConstructObject<UMaterialFactoryNew>(UMaterialFactoryNew::StaticClass());
+		UMaterialFactoryNew* Factory = NewObject<UMaterialFactoryNew>();
 		Factory->InitialTexture = InTexture;
 
 
@@ -600,7 +600,7 @@ namespace EditorBuildPromotionTestUtils
 
 			ensure(NULL != Cast<UBlueprintGeneratedClass>(InBlueprint->GeneratedClass));
 			// Then create a new template object, and add to array in
-			UActorComponent* NewTemplate = ConstructObject<UActorComponent>(ComponentClass, InBlueprint->GeneratedClass);
+			UActorComponent* NewTemplate = NewObject<UActorComponent>(InBlueprint->GeneratedClass, ComponentClass);
 			NewTemplate->SetFlags(RF_ArchetypeObject);
 			InBlueprint->ComponentTemplates.Add(NewTemplate);
 
@@ -2350,7 +2350,7 @@ namespace BuildPromotionTestHelper
 	const FString FilePath = AutomationTestSettings->BuildPromotionTest.ImportWorkflow.ImportSetting.ImportFilePath.FilePath; \
 	if (FilePath.Len() > 0) \
 	{ \
-		TFactoryClass* FactoryInst = ConstructObject<TFactoryClass>(TFactoryClass::StaticClass()); \
+		TFactoryClass* FactoryInst = NewObject<TFactoryClass>(); \
 		ExtraSettings \
 		AutomationEditorCommonUtils::ApplyCustomFactorySettings(FactoryInst, AutomationTestSettings->BuildPromotionTest.ImportWorkflow.ImportSetting.FactorySettings); \
 		ClassVariable = Cast<TObjectClass>(EditorBuildPromotionTestUtils::ImportAsset(FactoryInst, FilePath)); \
@@ -2416,7 +2416,7 @@ namespace BuildPromotionTestHelper
 			{
 				const FString BaseFileName = FPaths::GetPath(SurroundFilePath) / FPaths::GetBaseFilename(SurroundFilePath).LeftChop(3);
 
-				USoundSurroundFactory* FactoryInst = ConstructObject<USoundSurroundFactory>(USoundSurroundFactory::StaticClass());
+				USoundSurroundFactory* FactoryInst = NewObject<USoundSurroundFactory>();
 				AutomationEditorCommonUtils::ApplyCustomFactorySettings(FactoryInst, AutomationTestSettings->BuildPromotionTest.ImportWorkflow.SurroundSound.FactorySettings); 
 
 				const FString SurroundChannels[] = { TEXT("_fl"), TEXT("_fr"), TEXT("_fc"), TEXT("_lf"), TEXT("_sl"), TEXT("_sr"), TEXT("_bl"), TEXT("_br") };
@@ -2464,7 +2464,7 @@ namespace BuildPromotionTestHelper
 					if (FactoryClass)
 					{
 						//Create the factory and import the asset
-						UFactory* FactoryInst = ConstructObject<UFactory>(FactoryClass);
+						UFactory* FactoryInst = NewObject<UFactory>(GetTransientPackage(), FactoryClass);
 						AutomationEditorCommonUtils::ApplyCustomFactorySettings(FactoryInst, AssetsToImport[i].FactorySettings);
 						UObject* NewObject = EditorBuildPromotionTestUtils::ImportAsset(FactoryInst, FilePath);
 						if (NewObject)
@@ -2988,7 +2988,7 @@ namespace BuildPromotionTestHelper
 		bool ContentBrowser_CreateAParticleSystem_Part1()
 		{
 			//Create a Particle system
-			UParticleSystemFactoryNew* PSFactory = ConstructObject<UParticleSystemFactoryNew>(UParticleSystemFactoryNew::StaticClass());
+			UParticleSystemFactoryNew* PSFactory = NewObject<UParticleSystemFactoryNew>();
 			const FString& PSName(TEXT("BP_ParticleSystem"));
 			CreatedPS = Cast<UParticleSystem>(EditorBuildPromotionTestUtils::CreateAsset(PSFactory, UParticleSystem::StaticClass(), PSName));
 			if (CreatedPS)
@@ -3093,7 +3093,7 @@ namespace BuildPromotionTestHelper
 
 			if (FirstBlueprintMesh && SecondBlueprintMesh && CreatedPS)
 			{
-				UBlueprintFactory* Factory = ConstructObject<UBlueprintFactory>(UBlueprintFactory::StaticClass());
+				UBlueprintFactory* Factory = NewObject<UBlueprintFactory>();
 				Factory->ParentClass = AActor::StaticClass();
 
 				const FString PackageName = EditorBuildPromotionTestUtils::GetGamePath() + TEXT("/") + EditorBuildPromotionTestUtils::BlueprintNameString;
diff --git a/Engine/Source/Editor/UnrealEd/Private/ThumbnailHelpers.cpp b/Engine/Source/Editor/UnrealEd/Private/ThumbnailHelpers.cpp
--- a/Engine/Source/Editor/UnrealEd/Private/ThumbnailHelpers.cpp
+++ b/Engine/Source/Editor/UnrealEd/Private/ThumbnailHelpers.cpp
@@ -29,7 +29,7 @@ FThumbnailPreviewScene::FThumbnailPreviewScene()
 						.SetTransactional(false))
 {
 	// A background sky sphere
- 	UStaticMeshComponent* BackgroundComponent = ConstructObject<UStaticMeshComponent>(UStaticMeshComponent::StaticClass());
+	UStaticMeshComponent* BackgroundComponent = NewObject<UStaticMeshComponent>();
  	BackgroundComponent->SetStaticMesh( GUnrealEd->GetThumbnailManager()->EditorSkySphere );
  	const float SkySphereScale = 2000.0f;
  	const FTransform BackgroundTransform(FRotator(0,0,0), FVector(0,0,0), FVector(SkySphereScale));
@@ -39,18 +39,18 @@ FThumbnailPreviewScene::FThumbnailPreviewScene()
 	DirectionalLight->Intensity = 0.2f;
 
 	// Add additional lights
-	UDirectionalLightComponent* DirectionalLight2 = ConstructObject<UDirectionalLightComponent>(UDirectionalLightComponent::StaticClass());
+	UDirectionalLightComponent* DirectionalLight2 = NewObject<UDirectionalLightComponent>();
 	DirectionalLight->Intensity = 5.0f;
 	AddComponent(DirectionalLight2, FTransform( FRotator(-40,-144.678, 0) ));
 
-	UDirectionalLightComponent* DirectionalLight3 = ConstructObject<UDirectionalLightComponent>(UDirectionalLightComponent::StaticClass());
+	UDirectionalLightComponent* DirectionalLight3 = NewObject<UDirectionalLightComponent>();
 	DirectionalLight->Intensity = 1.0f;
 	AddComponent(DirectionalLight3, FTransform( FRotator(299.235,144.993, 0) ));
 
 	// Add an infinite plane
 	const float FloorPlaneScale = 10000.f;
 	const FTransform FloorPlaneTransform(FRotator(-90.f,0,0), FVector::ZeroVector, FVector(FloorPlaneScale));
-	UStaticMeshComponent* FloorPlaneComponent = ConstructObject<UStaticMeshComponent>(UStaticMeshComponent::StaticClass());
+	UStaticMeshComponent* FloorPlaneComponent = NewObject<UStaticMeshComponent>();
 	FloorPlaneComponent->SetStaticMesh( GUnrealEd->GetThumbnailManager()->EditorPlane );
 	FloorPlaneComponent->SetMaterial( 0, GUnrealEd->GetThumbnailManager()->FloorPlaneMaterial );
 	FPreviewScene::AddComponent(FloorPlaneComponent, FloorPlaneTransform);
@@ -162,7 +162,7 @@ void FParticleSystemThumbnailScene::SetParticleSystem(UParticleSystem* ParticleS
 	// If no preview component currently existing - create it now and warm it up.
 	if ( ParticleSystem && !ParticleSystem->PreviewComponent)
 	{
-		ParticleSystem->PreviewComponent = ConstructObject<UParticleSystemComponent>(UParticleSystemComponent::StaticClass());
+		ParticleSystem->PreviewComponent = NewObject<UParticleSystemComponent>();
 		ParticleSystem->PreviewComponent->Template = ParticleSystem;
 
 		ParticleSystem->PreviewComponent->ComponentToWorld.SetIdentity();
@@ -986,7 +986,7 @@ UActorComponent* FClassActorThumbnailScene::CreateComponentInstanceFromTemplate(
 	{
 		// We can not instance components that use the within keyword.
 		// Make a placeholder scene component instead.
-		USceneComponent* NewSceneComp = ConstructObject<USceneComponent>(USceneComponent::StaticClass(), GetTransientPackage());
+		USceneComponent* NewSceneComp = NewObject<USceneComponent>();
 		NewComponent = NewSceneComp;
 		
 		// Preserve relative location, rotation, scale, parent, and children if the template was a scene component.
diff --git a/Engine/Source/Editor/UnrealEd/Private/ThumbnailManager.cpp b/Engine/Source/Editor/UnrealEd/Private/ThumbnailManager.cpp
--- a/Engine/Source/Editor/UnrealEd/Private/ThumbnailManager.cpp
+++ b/Engine/Source/Editor/UnrealEd/Private/ThumbnailManager.cpp
@@ -87,7 +87,7 @@ void UThumbnailManager::InitializeRenderTypeArray(TArray<FThumbnailRenderingInfo
 			UClass* RenderClass = LoadObject<UClass>(nullptr, *RenderInfo.RendererClassName, nullptr, LOAD_None, nullptr);
 			if (RenderClass != nullptr)
 			{
-				RenderInfo.Renderer = ConstructObject<UThumbnailRenderer>(RenderClass);
+				RenderInfo.Renderer = NewObject<UThumbnailRenderer>(GetTransientPackage(), RenderClass);
 			}
 		}
 
@@ -185,7 +185,7 @@ void UThumbnailManager::RegisterCustomRenderer(UClass* Class, TSubclassOf<UThumb
 	FThumbnailRenderingInfo& Info = *(new (RenderableThumbnailTypes) FThumbnailRenderingInfo());
 	Info.ClassNeedingThumbnailName = NewClassPathName;
 	Info.ClassNeedingThumbnail = Class;
-	Info.Renderer = ConstructObject<UThumbnailRenderer>(RendererClass);
+	Info.Renderer = NewObject<UThumbnailRenderer>(GetTransientPackage(), RendererClass);
 	Info.RendererClassName = RendererClass->GetPathName();
 
 	bMapNeedsUpdate = true;
@@ -225,14 +225,14 @@ UThumbnailManager& UThumbnailManager::Get()
 			if (Class != nullptr)
 			{
 				// Create an instance of this class
-				ThumbnailManagerSingleton = ConstructObject<UThumbnailManager>(Class);
+				ThumbnailManagerSingleton = NewObject<UThumbnailManager>(GetTransientPackage(), Class);
 			}
 		}
 
 		// If the class couldn't be loaded or is the wrong type, fallback to the default
 		if (ThumbnailManagerSingleton == nullptr)
 		{
-			ThumbnailManagerSingleton = ConstructObject<UThumbnailManager>(UThumbnailManager::StaticClass());
+			ThumbnailManagerSingleton = NewObject<UThumbnailManager>();
 		}
 
 		// Keep the singleton alive
diff --git a/Engine/Source/Editor/UnrealEd/Private/UnrealEdEngine.cpp b/Engine/Source/Editor/UnrealEd/Private/UnrealEdEngine.cpp
--- a/Engine/Source/Editor/UnrealEd/Private/UnrealEdEngine.cpp
+++ b/Engine/Source/Editor/UnrealEd/Private/UnrealEdEngine.cpp
@@ -103,7 +103,7 @@ void UUnrealEdEngine::Init(IEngineLoop* InEngineLoop)
 
 	if (FPaths::IsProjectFilePathSet())
 	{
-		AutoReimportManager = ConstructObject<UAutoReimportManager>(UAutoReimportManager::StaticClass());
+		AutoReimportManager = NewObject<UAutoReimportManager>();
 		AutoReimportManager->Initialize();
 	}
 
@@ -123,7 +123,7 @@ void UUnrealEdEngine::Init(IEngineLoop* InEngineLoop)
 	BaseCookingFlags |= ExperimentalSettings->bIterativeCookingForLaunchOn ? ECookInitializationFlags::Iterative : ECookInitializationFlags::None;
 	if ( ExperimentalSettings->bCookOnTheSide )
 	{
-		CookServer = ConstructObject<UCookOnTheFlyServer>( UCookOnTheFlyServer::StaticClass() );
+		CookServer = NewObject<UCookOnTheFlyServer>();
 		CookServer->Initialize( ECookMode::CookOnTheFlyFromTheEditor, BaseCookingFlags );
 		CookServer->StartNetworkFileServer( false );
 
@@ -133,7 +133,7 @@ void UUnrealEdEngine::Init(IEngineLoop* InEngineLoop)
 	}
 	else if ( !ExperimentalSettings->bDisableCookInEditor)
 	{
-		CookServer = ConstructObject<UCookOnTheFlyServer>( UCookOnTheFlyServer::StaticClass() );
+		CookServer = NewObject<UCookOnTheFlyServer>();
 		CookServer->Initialize( ECookMode::CookByTheBookFromTheEditor, BaseCookingFlags );
 
 		FCoreUObjectDelegates::OnObjectPropertyChanged.AddUObject(CookServer, &UCookOnTheFlyServer::OnObjectPropertyChanged);
@@ -814,7 +814,7 @@ void UUnrealEdOptions::PostInitProperties()
 	Super::PostInitProperties();
 	if (!HasAnyFlags(RF_ClassDefaultObject | RF_NeedLoad))
 	{
-		EditorKeyBindings = ConstructObject<UUnrealEdKeyBindings>(UUnrealEdKeyBindings::StaticClass(), this, FName("EditorKeyBindingsInst"));
+		EditorKeyBindings = NewObject<UUnrealEdKeyBindings>(this, FName("EditorKeyBindingsInst"));
 	}
 }
 
@@ -823,7 +823,7 @@ UUnrealEdOptions* UUnrealEdEngine::GetUnrealEdOptions()
 {
 	if(EditorOptionsInst == NULL)
 	{
-		EditorOptionsInst = ConstructObject<UUnrealEdOptions>(UUnrealEdOptions::StaticClass());
+		EditorOptionsInst = NewObject<UUnrealEdOptions>();
 	}
 
 	return EditorOptionsInst;
diff --git a/Engine/Source/Editor/UnrealEd/Private/UnrealEdSrv.cpp b/Engine/Source/Editor/UnrealEd/Private/UnrealEdSrv.cpp
--- a/Engine/Source/Editor/UnrealEd/Private/UnrealEdSrv.cpp
+++ b/Engine/Source/Editor/UnrealEd/Private/UnrealEdSrv.cpp
@@ -1929,7 +1929,7 @@ void CreateBoundingBoxBuilderBrush( UWorld* InWorld, const TArray<FPoly*> Select
 	{
 		const FScopedTransaction Transaction(NSLOCTEXT("UnrealEd", "BrushSet", "Brush Set"));
 
-		UCubeBuilder* CubeBuilder = ConstructObject<UCubeBuilder>(UCubeBuilder::StaticClass(), GetTransientPackage(), NAME_None, RF_Transactional);
+		UCubeBuilder* CubeBuilder = NewObject<UCubeBuilder>(GetTransientPackage(), NAME_None, RF_Transactional);
 		FVector Extent = BBox.GetExtent();
 		CubeBuilder->X = Extent.X * 2;
 		CubeBuilder->Y = Extent.Y * 2;
diff --git a/Engine/Source/Editor/UnrealEd/Public/AnimationEditorUtils.h b/Engine/Source/Editor/UnrealEd/Public/AnimationEditorUtils.h
--- a/Engine/Source/Editor/UnrealEd/Public/AnimationEditorUtils.h
+++ b/Engine/Source/Editor/UnrealEd/Public/AnimationEditorUtils.h
@@ -74,7 +74,7 @@ namespace AnimationEditorUtils
 				FString PackageName;
 				CreateUniqueAssetName(Object->GetOutermost()->GetName(), InSuffix, PackageName, Name);
 
-				TFactory* Factory = ConstructObject<TFactory>(TFactory::StaticClass());
+				TFactory* Factory = NewObject<TFactory>();
 				Factory->TargetSkeleton = Object;
 
 				FContentBrowserModule& ContentBrowserModule = FModuleManager::LoadModuleChecked<FContentBrowserModule>("ContentBrowser");
diff --git a/Engine/Source/Editor/WorldBrowser/Private/Tiles/WorldTileCollectionModel.cpp b/Engine/Source/Editor/WorldBrowser/Private/Tiles/WorldTileCollectionModel.cpp
--- a/Engine/Source/Editor/WorldBrowser/Private/Tiles/WorldTileCollectionModel.cpp
+++ b/Engine/Source/Editor/WorldBrowser/Private/Tiles/WorldTileCollectionModel.cpp
@@ -1483,7 +1483,7 @@ static ULandscapeLayerInfoObject* GetLandscapeLayerInfoObject(FName LayerName, c
 	ULandscapeLayerInfoObject* LayerInfo = FindObject<ULandscapeLayerInfoObject>(Package, *LayerObjectName);
 	if (LayerInfo == nullptr)
 	{
-		LayerInfo = ConstructObject<ULandscapeLayerInfoObject>(ULandscapeLayerInfoObject::StaticClass(), Package, FName(*LayerObjectName), RF_Public | RF_Standalone | RF_Transactional);
+		LayerInfo = NewObject<ULandscapeLayerInfoObject>(Package, FName(*LayerObjectName), RF_Public | RF_Standalone | RF_Transactional);
 		LayerInfo->LayerName = LayerName;
 		// Notify the asset registry
 		FAssetRegistryModule::AssetCreated(LayerInfo);
@@ -2045,7 +2045,7 @@ bool FWorldTileCollectionModel::GenerateLODLevels(FLevelModelList InLevelList, i
 				MeshOuter->Modify();
 			}
 
-			auto StaticMesh = NewNamedObject<UStaticMesh>(MeshOuter, *LandscapeMeshAssetName, RF_Public | RF_Standalone);
+			auto StaticMesh = NewObject<UStaticMesh>(MeshOuter, *LandscapeMeshAssetName, RF_Public | RF_Standalone);
 			StaticMesh->InitResources();
 			{
 				FString OutputPath = StaticMesh->GetPathName();
diff --git a/Engine/Source/Runtime/AIModule/Classes/Actions/PawnAction.h b/Engine/Source/Runtime/AIModule/Classes/Actions/PawnAction.h
--- a/Engine/Source/Runtime/AIModule/Classes/Actions/PawnAction.h
+++ b/Engine/Source/Runtime/AIModule/Classes/Actions/PawnAction.h
@@ -178,7 +178,7 @@ public:
 	static TActionClass* CreateActionInstance(UWorld& World)
 	{
 		TSubclassOf<UPawnAction> ActionClass = TActionClass::StaticClass();
-		return ConstructObject<TActionClass>(ActionClass, &World);
+		return NewObject<TActionClass>(&World, ActionClass);
 	}
 
 	//----------------------------------------------------------------------//
diff --git a/Engine/Source/Runtime/AIModule/Classes/BehaviorTree/BlackboardData.h b/Engine/Source/Runtime/AIModule/Classes/BehaviorTree/BlackboardData.h
--- a/Engine/Source/Runtime/AIModule/Classes/BehaviorTree/BlackboardData.h
+++ b/Engine/Source/Runtime/AIModule/Classes/BehaviorTree/BlackboardData.h
@@ -103,7 +103,7 @@ public:
 			FBlackboardEntry Entry;
 			Entry.EntryName = KeyName;
 
-			CreatedKeyType = ConstructObject<T>(T::StaticClass(), this);
+			CreatedKeyType = NewObject<T>(this);
 			Entry.KeyType = CreatedKeyType;		
 
 			Keys.Add(Entry);
diff --git a/Engine/Source/Runtime/AIModule/Private/AIController.cpp b/Engine/Source/Runtime/AIModule/Private/AIController.cpp
--- a/Engine/Source/Runtime/AIModule/Private/AIController.cpp
+++ b/Engine/Source/Runtime/AIModule/Private/AIController.cpp
@@ -13,6 +13,7 @@
 #include "Engine/Canvas.h"
 #include "GameFramework/PhysicsVolume.h"
 #include "AIController.h"
+#include "Perception/AIPerceptionComponent.h"
 
 // mz@todo these need to be removed, legacy code
 #define CLOSEPROXIMITY					500.f
@@ -765,7 +766,7 @@ bool AAIController::RunBehaviorTree(UBehaviorTree* BTAsset)
 		{
 			UE_VLOG(this, LogBehaviorTree, Log, TEXT("RunBehaviorTree: spawning BehaviorTreeComponent.."));
 
-			BrainComponent = BTComp = ConstructObject<UBehaviorTreeComponent>(UBehaviorTreeComponent::StaticClass(), this, TEXT("BTComponent"));
+			BrainComponent = BTComp = NewObject<UBehaviorTreeComponent>(this, TEXT("BTComponent"));
 			BrainComponent->RegisterComponent();
 		}
 
@@ -800,7 +801,7 @@ bool AAIController::UseBlackboard(UBlackboardData* BlackboardAsset)
 
 	if (BlackboardComp == NULL)
 	{
-		BlackboardComp = ConstructObject<UBlackboardComponent>(UBlackboardComponent::StaticClass(), this, TEXT("BlackboardComponent"));
+		BlackboardComp = NewObject<UBlackboardComponent>(this, TEXT("BlackboardComponent"));
 		if (BlackboardComp != NULL)
 		{
 			InitializeBlackboard(*BlackboardComp, *BlackboardAsset);
diff --git a/Engine/Source/Runtime/AIModule/Private/AIModule.cpp b/Engine/Source/Runtime/AIModule/Private/AIModule.cpp
--- a/Engine/Source/Runtime/AIModule/Private/AIModule.cpp
+++ b/Engine/Source/Runtime/AIModule/Private/AIModule.cpp
@@ -4,6 +4,7 @@
 #if WITH_EDITOR && ENABLE_VISUAL_LOG
 #	include "VisualLoggerExtension.h"
 #endif
+#include "AIModule.h"
 
 DEFINE_LOG_CATEGORY_STATIC(LogAIModule, Log, All);
 
@@ -54,5 +55,5 @@ UAISystemBase* FAIModule::CreateAISystemInstance(UWorld* World)
 	UE_LOG(LogAIModule, Log, TEXT("Creating AISystem for world %s"), *GetNameSafe(World));
 	
 	TSubclassOf<UAISystemBase> AISystemClass = LoadClass<UAISystemBase>(NULL, *UAISystem::GetAISystemClassName().ToString(), NULL, LOAD_None, NULL);
-	return ConstructObject<UAISystemBase>(AISystemClass, World);
+	return NewObject<UAISystemBase>(World, AISystemClass);
 }
diff --git a/Engine/Source/Runtime/AIModule/Private/Actions/PawnAction.cpp b/Engine/Source/Runtime/AIModule/Private/Actions/PawnAction.cpp
--- a/Engine/Source/Runtime/AIModule/Private/Actions/PawnAction.cpp
+++ b/Engine/Source/Runtime/AIModule/Private/Actions/PawnAction.cpp
@@ -386,7 +386,7 @@ UPawnAction* UPawnAction::CreateActionInstance(UObject* WorldContextObject, TSub
 	UWorld* World = GEngine->GetWorldFromContextObject(WorldContextObject);
 	if (World && ActionClass)
 	{
-		return ConstructObject<UPawnAction>(ActionClass, World);
+		return NewObject<UPawnAction>(World, ActionClass);
 	}
 	return NULL;
 }
diff --git a/Engine/Source/Runtime/AIModule/Private/BehaviorTree/BTNode.cpp b/Engine/Source/Runtime/AIModule/Private/BehaviorTree/BTNode.cpp
--- a/Engine/Source/Runtime/AIModule/Private/BehaviorTree/BTNode.cpp
+++ b/Engine/Source/Runtime/AIModule/Private/BehaviorTree/BTNode.cpp
@@ -74,7 +74,7 @@ void UBTNode::InitializeInSubtree(UBehaviorTreeComponent& OwnerComp, uint8* Node
 		UBTNode* NodeInstance = OwnerComp.NodeInstances.IsValidIndex(NextInstancedIndex) ? OwnerComp.NodeInstances[NextInstancedIndex] : NULL;
 		if (NodeInstance == NULL)
 		{
-			NodeInstance = ConstructObject<UBTNode>(GetClass(), &OwnerComp, GetFName(), RF_NoFlags, (UObject*)(this));
+			NodeInstance = NewObject<UBTNode>(&OwnerComp, GetClass(), GetFName(), RF_NoFlags, (UObject*)(this));
 			NodeInstance->InitializeNode(GetParentNode(), GetExecutionIndex(), GetMemoryOffset(), GetTreeDepth());
 			NodeInstance->bIsInstanced = true;
 
diff --git a/Engine/Source/Runtime/AIModule/Private/BehaviorTree/BehaviorTreeTypes.cpp b/Engine/Source/Runtime/AIModule/Private/BehaviorTree/BehaviorTreeTypes.cpp
--- a/Engine/Source/Runtime/AIModule/Private/BehaviorTree/BehaviorTreeTypes.cpp
+++ b/Engine/Source/Runtime/AIModule/Private/BehaviorTree/BehaviorTreeTypes.cpp
@@ -5,6 +5,7 @@
 #include "BehaviorTree/BTCompositeNode.h"
 #include "BehaviorTree/Tasks/BTTask_RunBehavior.h"
 #include "BehaviorTree/BehaviorTreeTypes.h"
+#include "BehaviorTree/BTDecorator.h"
 
 //----------------------------------------------------------------------//
 // FBehaviorTreeInstance
@@ -225,65 +226,65 @@ void FBlackboardKeySelector::InitSelectedKey(UBlackboardData* BlackboardAsset)
 
 void FBlackboardKeySelector::AddObjectFilter(UObject* Owner, TSubclassOf<UObject> AllowedClass)
 {
-	UBlackboardKeyType_Object* FilterOb = NewNamedObject<UBlackboardKeyType_Object>(Owner, TEXT("BlackboardKeyType_Object"));
+	UBlackboardKeyType_Object* FilterOb = Owner->CreateDefaultSubobject<UBlackboardKeyType_Object>(TEXT("BlackboardKeyType_Object"));
 	FilterOb->BaseClass = AllowedClass;
 	AllowedTypes.Add(FilterOb);
 }
 
 void FBlackboardKeySelector::AddClassFilter(UObject* Owner, TSubclassOf<UClass> AllowedClass)
 {
-	UBlackboardKeyType_Class* FilterOb = NewNamedObject<UBlackboardKeyType_Class>(Owner, TEXT("BlackboardKeyType_Class"));
+	UBlackboardKeyType_Class* FilterOb = Owner->CreateDefaultSubobject<UBlackboardKeyType_Class>(TEXT("BlackboardKeyType_Class"));
 	FilterOb->BaseClass = AllowedClass;
 	AllowedTypes.Add(FilterOb);
 }
 
 void FBlackboardKeySelector::AddEnumFilter(UObject* Owner, UEnum* AllowedEnum)
 {
-	UBlackboardKeyType_Enum* FilterOb = NewNamedObject<UBlackboardKeyType_Enum>(Owner, TEXT("BlackboardKeyType_Enum"));
+	UBlackboardKeyType_Enum* FilterOb = Owner->CreateDefaultSubobject<UBlackboardKeyType_Enum>(TEXT("BlackboardKeyType_Enum"));
 	FilterOb->EnumType = AllowedEnum;
 	AllowedTypes.Add(FilterOb);
 }
 
 void FBlackboardKeySelector::AddNativeEnumFilter(UObject* Owner, const FString& AllowedEnumName)
 {
-	UBlackboardKeyType_NativeEnum* FilterOb = NewNamedObject<UBlackboardKeyType_NativeEnum>(Owner, TEXT("BlackboardKeyType_NativeEnum"));
+	UBlackboardKeyType_NativeEnum* FilterOb = Owner->CreateDefaultSubobject<UBlackboardKeyType_NativeEnum>(TEXT("BlackboardKeyType_NativeEnum"));
 	FilterOb->EnumName = AllowedEnumName;
 	AllowedTypes.Add(FilterOb);
 }
 
 void FBlackboardKeySelector::AddIntFilter(UObject* Owner)
 {
-	AllowedTypes.Add(NewNamedObject<UBlackboardKeyType_Int>(Owner, TEXT("BlackboardKeyType_Int")));
+	AllowedTypes.Add(Owner->CreateDefaultSubobject<UBlackboardKeyType_Int>(TEXT("BlackboardKeyType_Int")));
 }
 
 void FBlackboardKeySelector::AddFloatFilter(UObject* Owner)
 {
-	AllowedTypes.Add(NewNamedObject<UBlackboardKeyType_Float>(Owner, TEXT("BlackboardKeyType_Float")));
+	AllowedTypes.Add(Owner->CreateDefaultSubobject<UBlackboardKeyType_Float>(TEXT("BlackboardKeyType_Float")));
 }
 
 void FBlackboardKeySelector::AddBoolFilter(UObject* Owner)
 {
-	AllowedTypes.Add(NewNamedObject<UBlackboardKeyType_Bool>(Owner, TEXT("BlackboardKeyType_Bool")));
+	AllowedTypes.Add(Owner->CreateDefaultSubobject<UBlackboardKeyType_Bool>(TEXT("BlackboardKeyType_Bool")));
 }
 
 void FBlackboardKeySelector::AddVectorFilter(UObject* Owner)
 {
-	AllowedTypes.Add(NewNamedObject<UBlackboardKeyType_Vector>(Owner, TEXT("BlackboardKeyType_Vector")));
+	AllowedTypes.Add(Owner->CreateDefaultSubobject<UBlackboardKeyType_Vector>(TEXT("BlackboardKeyType_Vector")));
 }
 
 void FBlackboardKeySelector::AddRotatorFilter(UObject* Owner)
 {
-	AllowedTypes.Add(NewNamedObject<UBlackboardKeyType_Rotator>(Owner, TEXT("BlackboardKeyType_Rotator")));
+	AllowedTypes.Add(Owner->CreateDefaultSubobject<UBlackboardKeyType_Rotator>(TEXT("BlackboardKeyType_Rotator")));
 }
 
 void FBlackboardKeySelector::AddStringFilter(UObject* Owner)
 {
-	AllowedTypes.Add(NewNamedObject<UBlackboardKeyType_String>(Owner, TEXT("BlackboardKeyType_String")));
+	AllowedTypes.Add(Owner->CreateDefaultSubobject<UBlackboardKeyType_String>(TEXT("BlackboardKeyType_String")));
 }
 
 void FBlackboardKeySelector::AddNameFilter(UObject* Owner)
 {
-	AllowedTypes.Add(NewNamedObject<UBlackboardKeyType_Name>(Owner, TEXT("BlackboardKeyType_Name")));
+	AllowedTypes.Add(Owner->CreateDefaultSubobject<UBlackboardKeyType_Name>(TEXT("BlackboardKeyType_Name")));
 }
 
 //----------------------------------------------------------------------//
diff --git a/Engine/Source/Runtime/AIModule/Private/Perception/AIPerceptionSystem.cpp b/Engine/Source/Runtime/AIModule/Private/Perception/AIPerceptionSystem.cpp
--- a/Engine/Source/Runtime/AIModule/Private/Perception/AIPerceptionSystem.cpp
+++ b/Engine/Source/Runtime/AIModule/Private/Perception/AIPerceptionSystem.cpp
@@ -8,6 +8,8 @@
 #include "Perception/AISense_Touch.h"
 #include "Perception/AISense_Prediction.h"
 #include "Perception/AISense_Damage.h"
+#include "Perception/AISenseConfig.h"
+#include "Perception/AISenseEvent.h"
 
 DECLARE_CYCLE_STAT(TEXT("Perception System"),STAT_AI_PerceptionSys,STATGROUP_AI);
 
@@ -60,7 +62,7 @@ void UAIPerceptionSystem::RegisterSenseClass(TSubclassOf<UAISense> SenseClass)
 
 	if (Senses[SenseID] == nullptr)
 	{
-		Senses[SenseID] = ConstructObject<UAISense>(SenseClass, this);
+		Senses[SenseID] = NewObject<UAISense>(this, SenseClass);
 		check(Senses[SenseID]);
 #if !UE_BUILD_SHIPPING
 		DebugSenseColors[SenseID] = Senses[SenseID]->GetDebugColor();
diff --git a/Engine/Source/Runtime/AIModule/Private/VisualLoggerExtension.cpp b/Engine/Source/Runtime/AIModule/Private/VisualLoggerExtension.cpp
--- a/Engine/Source/Runtime/AIModule/Private/VisualLoggerExtension.cpp
+++ b/Engine/Source/Runtime/AIModule/Private/VisualLoggerExtension.cpp
@@ -79,7 +79,7 @@ void FVisualLoggerExtension::DrawData(UWorld* InWorld, UCanvas* Canvas, AActor*
 			UEQSRenderingComponent* EQSRenderComp = HelperActor->FindComponentByClass<UEQSRenderingComponent>();
 			if (!EQSRenderComp)
 			{
-				EQSRenderComp = ConstructObject<UEQSRenderingComponent>(UEQSRenderingComponent::StaticClass(), HelperActor);
+				EQSRenderComp = NewObject<UEQSRenderingComponent>(HelperActor);
 				EQSRenderComp->bDrawOnlyWhenSelected = false;
 				EQSRenderComp->RegisterComponent();
 				EQSRenderComp->SetHiddenInGame(true);
diff --git a/Engine/Source/Runtime/CoreUObject/Private/UObject/CoreNative.cpp b/Engine/Source/Runtime/CoreUObject/Private/UObject/CoreNative.cpp
--- a/Engine/Source/Runtime/CoreUObject/Private/UObject/CoreNative.cpp
+++ b/Engine/Source/Runtime/CoreUObject/Private/UObject/CoreNative.cpp
@@ -173,7 +173,7 @@ UObject* FObjectInstancingGraph::GetInstancedSubobject( UObject* SourceSubobject
 							if (!InstancedSubobject)
 							{
 								// finally, create the component instance
-								InstancedSubobject = ConstructObject<UObject>(SourceSubobject->GetClass(), SubobjectOuter,
+								InstancedSubobject = StaticConstructObject(SourceSubobject->GetClass(), SubobjectOuter,
 									SubobjectName, SubobjectOuter->GetMaskedFlags(RF_PropagateToSubObjects), SourceSubobject,
 									true, this);
 							}
diff --git a/Engine/Source/Runtime/CoreUObject/Private/UObject/Obj.cpp b/Engine/Source/Runtime/CoreUObject/Private/UObject/Obj.cpp
--- a/Engine/Source/Runtime/CoreUObject/Private/UObject/Obj.cpp
+++ b/Engine/Source/Runtime/CoreUObject/Private/UObject/Obj.cpp
@@ -185,7 +185,7 @@ bool UObject::Rename( const TCHAR* InName, UObject* NewOuter, ERenameFlags Flags
 		if ( Redirector == NULL )
 		{
 			// create a UObjectRedirector with the same name as the old object we are redirecting
-			Redirector = ConstructObject<UObjectRedirector>(UObjectRedirector::StaticClass(), OldOuter, OldName, RF_Standalone | RF_Public);
+			Redirector = NewObject<UObjectRedirector>(OldOuter, OldName, RF_Standalone | RF_Public);
 		}
 
 		// point the redirector object to this object
@@ -3567,7 +3567,7 @@ void StaticUObjectInit()
 	UObjectBaseInit();
 
 	// Allocate special packages.
-	GObjTransientPkg = NewNamedObject<UPackage>(nullptr, TEXT("/Engine/Transient"));
+	GObjTransientPkg = NewObject<UPackage>(nullptr, TEXT("/Engine/Transient"));
 	GObjTransientPkg->AddToRoot();
 
 	if( FParse::Param( FCommandLine::Get(), TEXT("VERIFYGC") ) )
diff --git a/Engine/Source/Runtime/CoreUObject/Private/UObject/Package.cpp b/Engine/Source/Runtime/CoreUObject/Private/UObject/Package.cpp
--- a/Engine/Source/Runtime/CoreUObject/Private/UObject/Package.cpp
+++ b/Engine/Source/Runtime/CoreUObject/Private/UObject/Package.cpp
@@ -113,7 +113,7 @@ UMetaData* UPackage::GetMetaData()
 		// If MetaData is NULL then it wasn't loaded by linker, so we have to create it.
 		if(MetaData == NULL)
 		{
-			MetaData = ConstructObject<UMetaData>(UMetaData::StaticClass(), this, NAME_PackageMetaData, RF_Standalone | RF_LoadCompleted);
+			MetaData = NewObject<UMetaData>(this, NAME_PackageMetaData, RF_Standalone | RF_LoadCompleted);
 		}
 	}
 
diff --git a/Engine/Source/Runtime/CoreUObject/Public/UObject/Class.h b/Engine/Source/Runtime/CoreUObject/Public/UObject/Class.h
--- a/Engine/Source/Runtime/CoreUObject/Public/UObject/Class.h
+++ b/Engine/Source/Runtime/CoreUObject/Public/UObject/Class.h
@@ -2546,6 +2546,7 @@ FORCEINLINE bool UObject::Implements() const
  * @return	a pointer of type T to a new object of the specified class
  */
 template< class T >
+DEPRECATED(4.8, "ConstructObject is deprecated. Use NewObject instead")
 T* ConstructObject(UClass* Class, UObject* Outer, FName Name, EObjectFlags SetFlags, UObject* Template, bool bCopyTransientsFromClassDefaults, struct FObjectInstancingGraph* InstanceGraph )
 {
 	checkf(Class, TEXT("ConstructObject called with a NULL class object"));
diff --git a/Engine/Source/Runtime/CoreUObject/Public/UObject/UObjectGlobals.h b/Engine/Source/Runtime/CoreUObject/Public/UObject/UObjectGlobals.h
--- a/Engine/Source/Runtime/CoreUObject/Public/UObject/UObjectGlobals.h
+++ b/Engine/Source/Runtime/CoreUObject/Public/UObject/UObjectGlobals.h
@@ -1028,6 +1028,7 @@ public:
  * @return	a pointer of type T to a new object of the specified class
  */
 template< class T >
+DEPRECATED(4.8, "ConstructObject is deprecated. Use NewObject instead")
 T* ConstructObject(UClass* Class, UObject* Outer = (UObject*)GetTransientPackage(), FName Name=NAME_None, EObjectFlags SetFlags=RF_NoFlags, UObject* Template=NULL, bool bCopyTransientsFromClassDefaults=false, struct FObjectInstancingGraph* InstanceGraph=NULL );
 
 /**
@@ -1037,10 +1038,21 @@ T* ConstructObject(UClass* Class, UObject* Outer = (UObject*)GetTransientPackage
  * @param	Class		the class of object to construct
  */
 template< class T >
-T* NewObject(UObject* Outer=(UObject*)GetTransientPackage(),UClass* Class=T::StaticClass() )
+T* NewObject(UObject* Outer = (UObject*)GetTransientPackage(), UClass* Class = T::StaticClass(), FName Name = NAME_None, EObjectFlags Flags = RF_NoFlags, UObject* Template = nullptr, bool bCopyTransientsFromClassDefaults = false, FObjectInstancingGraph* InInstanceGraph = nullptr)
 {
-	FObjectInitializer::AssertIfInConstructor(Outer, TEXT("NewObject can't be used to create default subobjects (inside of UObject derived class constructor) as it produces inconsistent object names. Use ObjectInitializer.CreateDefaultSuobject<> instead."));
-	return ConstructObject<T>(Class,Outer);
+	if (Name == NAME_None)
+	{
+		FObjectInitializer::AssertIfInConstructor(Outer, TEXT("NewObject with empty name can't be used to create default subobjects (inside of UObject derived class constructor) as it produces inconsistent object names. Use ObjectInitializer.CreateDefaultSuobject<> instead."));
+	}
+	checkf(Class, TEXT("NewObject called with a nullptr class object"));
+	checkSlow(Class->IsChildOf(T::StaticClass()));
+	return static_cast<T*>(StaticConstructObject(Class, Outer, Name, Flags, Template, bCopyTransientsFromClassDefaults, InInstanceGraph));
+}
+
+template< class T >
+T* NewObject(UObject* Outer, FName Name, EObjectFlags Flags = RF_NoFlags, UObject* Template = nullptr, bool bCopyTransientsFromClassDefaults = false, FObjectInstancingGraph* InInstanceGraph = nullptr)
+{
+	return NewObject<T>(Outer, T::StaticClass(), Name, Flags, Template, bCopyTransientsFromClassDefaults, InInstanceGraph);
 }
 
 /**
@@ -1051,9 +1063,10 @@ T* NewObject(UObject* Outer=(UObject*)GetTransientPackage(),UClass* Class=T::Sta
  * @param	Flags	The object flags for the new object.
  */
 template< class TClass >
+DEPRECATED(4.8, "NewNamedObject is deprecated. Use NewObject instead")
 TClass* NewNamedObject(UObject* Outer, FName Name, EObjectFlags Flags = RF_NoFlags, UObject* Template=NULL)
 {
-	return ConstructObject<TClass>(TClass::StaticClass(), Outer, Name, Flags, Template);
+	return NewObject<TClass>(Outer, Name, Flags, Template);
 }
 
 
diff --git a/Engine/Source/Runtime/Engine/Classes/Sound/SoundCue.h b/Engine/Source/Runtime/Engine/Classes/Sound/SoundCue.h
--- a/Engine/Source/Runtime/Engine/Classes/Sound/SoundCue.h
+++ b/Engine/Source/Runtime/Engine/Classes/Sound/SoundCue.h
@@ -98,7 +98,7 @@ public:
 	template<class T>
 	T* ConstructSoundNode(TSubclassOf<USoundNode> SoundNodeClass = T::StaticClass(), bool bSelectNewNode = true)
 	{
-		T* SoundNode = ConstructObject<T>(SoundNodeClass, this);
+		T* SoundNode = NewObject<T>(this, SoundNodeClass);
 #if WITH_EDITOR
 		AllNodes.Add(SoundNode);
 		SetupSoundNode(SoundNode, bSelectNewNode);
diff --git a/Engine/Source/Runtime/Engine/Private/AI/Navigation/RecastNavMesh.cpp b/Engine/Source/Runtime/Engine/Private/AI/Navigation/RecastNavMesh.cpp
--- a/Engine/Source/Runtime/Engine/Private/AI/Navigation/RecastNavMesh.cpp
+++ b/Engine/Source/Runtime/Engine/Private/AI/Navigation/RecastNavMesh.cpp
@@ -438,7 +438,7 @@ UPrimitiveComponent* ARecastNavMesh::ConstructRenderingComponentImpl()
 
 	return NULL;
 #else
-	return NewNamedObject<UNavMeshRenderingComponent>(this, TEXT("NavMeshRenderer"));
+	return NewObject<UNavMeshRenderingComponent>(this, TEXT("NavMeshRenderer"));
 #endif // DO_NAVMESH_DEBUG_DRAWING_PER_TILE
 }
 
diff --git a/Engine/Source/Runtime/Engine/Private/Actor.cpp b/Engine/Source/Runtime/Engine/Private/Actor.cpp
--- a/Engine/Source/Runtime/Engine/Private/Actor.cpp
+++ b/Engine/Source/Runtime/Engine/Private/Actor.cpp
@@ -2394,7 +2394,7 @@ void AActor::EnableInput(APlayerController* PlayerController)
 		// If it doesn't exist create it and bind delegates
 		if (!InputComponent)
 		{
-			InputComponent = ConstructObject<UInputComponent>(UInputComponent::StaticClass(), this);
+			InputComponent = NewObject<UInputComponent>(this);
 			InputComponent->RegisterComponent();
 			InputComponent->bBlockInput = bBlockInput;
 			InputComponent->Priority = InputPriority;
diff --git a/Engine/Source/Runtime/Engine/Private/Animation/AnimCompress_LeastDestructive.cpp b/Engine/Source/Runtime/Engine/Private/Animation/AnimCompress_LeastDestructive.cpp
--- a/Engine/Source/Runtime/Engine/Private/Animation/AnimCompress_LeastDestructive.cpp
+++ b/Engine/Source/Runtime/Engine/Private/Animation/AnimCompress_LeastDestructive.cpp
@@ -21,7 +21,7 @@ UAnimCompress_LeastDestructive::UAnimCompress_LeastDestructive(const FObjectInit
 
 void UAnimCompress_LeastDestructive::DoReduction(UAnimSequence* AnimSeq, const TArray<FBoneData>& BoneData)
 {
-	UAnimCompress* BitwiseCompressor = ConstructObject<UAnimCompress_BitwiseCompressOnly>( UAnimCompress_BitwiseCompressOnly::StaticClass() );
+	UAnimCompress* BitwiseCompressor = NewObject<UAnimCompress_BitwiseCompressOnly>();
 	BitwiseCompressor->RotationCompressionFormat = ACF_Float96NoW;
 	BitwiseCompressor->TranslationCompressionFormat = ACF_None;
 	BitwiseCompressor->Reduce(AnimSeq, false);
diff --git a/Engine/Source/Runtime/Engine/Private/Animation/AnimNotifyState_Trail.cpp b/Engine/Source/Runtime/Engine/Private/Animation/AnimNotifyState_Trail.cpp
--- a/Engine/Source/Runtime/Engine/Private/Animation/AnimNotifyState_Trail.cpp
+++ b/Engine/Source/Runtime/Engine/Private/Animation/AnimNotifyState_Trail.cpp
@@ -94,7 +94,7 @@ void UAnimNotifyState_Trail::NotifyBegin(class USkeletalMeshComponent * MeshComp
 	if (!bFoundExistingTrail && !bError)
 	{
 		//Spawn a new component from PSTemplate. This notify is made the outer so that the component can be identified later.
-		UParticleSystemComponent* NewParticleComp = ConstructObject<UParticleSystemComponent>(UParticleSystemComponent::StaticClass(), MeshComp);
+		UParticleSystemComponent* NewParticleComp = NewObject<UParticleSystemComponent>(MeshComp);
 		NewParticleComp->bAutoDestroy = true;
 		NewParticleComp->SecondsBeforeInactive = 0.0f;
 		NewParticleComp->bAutoActivate = false;
diff --git a/Engine/Source/Runtime/Engine/Private/Animation/AnimSequence.cpp b/Engine/Source/Runtime/Engine/Private/Animation/AnimSequence.cpp
--- a/Engine/Source/Runtime/Engine/Private/Animation/AnimSequence.cpp
+++ b/Engine/Source/Runtime/Engine/Private/Animation/AnimSequence.cpp
@@ -280,7 +280,7 @@ void UAnimSequence::Serialize(FArchive& Ar)
 	{
 		if ( AssetImportData == NULL )
 		{
-			AssetImportData = ConstructObject<UAssetImportData>(UAssetImportData::StaticClass(), this);
+			AssetImportData = NewObject<UAssetImportData>(this);
 		}
 		
 		AssetImportData->SourceFilePath = SourceFilePath_DEPRECATED;
diff --git a/Engine/Source/Runtime/Engine/Private/Animation/AnimationUtils.cpp b/Engine/Source/Runtime/Engine/Private/Animation/AnimationUtils.cpp
--- a/Engine/Source/Runtime/Engine/Private/Animation/AnimationUtils.cpp
+++ b/Engine/Source/Runtime/Engine/Private/Animation/AnimationUtils.cpp
@@ -401,7 +401,7 @@ static inline UAnimCompress* ConstructDefaultCompressionAlgorithm()
 		UE_LOG(LogAnimation, Fatal, TEXT("Couldn't find animation compression algorithm named %s"), *DefaultCompressionAlgorithm );
 	}
 
-	UAnimCompress* NewAlgorithm = ConstructObject<UAnimCompress>( CompressionAlgorithmClass );
+	UAnimCompress* NewAlgorithm = NewObject<UAnimCompress>(GetTransientPackage(), CompressionAlgorithmClass);
 	NewAlgorithm->RotationCompressionFormat = RotationCompressionFormat;
 	NewAlgorithm->TranslationCompressionFormat = TranslationCompressionFormat;
 	NewAlgorithm->AddToRoot();
@@ -831,7 +831,7 @@ void FAnimationUtils::CompressAnimSequenceExplicit(
 				// Progressive Algorithm
 				if( bTryPerTrackBitwiseCompression )
 				{
-					UAnimCompress_PerTrackCompression* PerTrackCompressor = ConstructObject<UAnimCompress_PerTrackCompression>(UAnimCompress_PerTrackCompression::StaticClass());
+					UAnimCompress_PerTrackCompression* PerTrackCompressor = NewObject<UAnimCompress_PerTrackCompression>();
 
 					// Start not too aggressive
 					PerTrackCompressor->MaxPosDiffBitwise /= 10.f;
@@ -970,7 +970,7 @@ void FAnimationUtils::CompressAnimSequenceExplicit(
 				// Start with Bitwise Compress only
 				if( bTryFixedBitwiseCompression )
 				{
-					UAnimCompress_BitwiseCompressOnly* BitwiseCompressor = ConstructObject<UAnimCompress_BitwiseCompressOnly>( UAnimCompress_BitwiseCompressOnly::StaticClass() );
+					UAnimCompress_BitwiseCompressOnly* BitwiseCompressor = NewObject<UAnimCompress_BitwiseCompressOnly>();
 
 					// Try ACF_Float96NoW
 					BitwiseCompressor->RotationCompressionFormat = ACF_Float96NoW;
@@ -998,7 +998,7 @@ void FAnimationUtils::CompressAnimSequenceExplicit(
 				// this compressor has a minimum number of frames requirement. So no need to go there if we don't meet that...
 				if( bTryFixedBitwiseCompression && bTryIntervalKeyRemoval )
 				{
-					UAnimCompress_RemoveEverySecondKey* RemoveEveryOtherKeyCompressor = ConstructObject<UAnimCompress_RemoveEverySecondKey>( UAnimCompress_RemoveEverySecondKey::StaticClass() );
+					UAnimCompress_RemoveEverySecondKey* RemoveEveryOtherKeyCompressor = NewObject<UAnimCompress_RemoveEverySecondKey>();
 					if( AnimSeq->NumFrames > RemoveEveryOtherKeyCompressor->MinKeys )
 					{
 						RemoveEveryOtherKeyCompressor->bStartAtSecondKey = false;
@@ -1053,7 +1053,7 @@ void FAnimationUtils::CompressAnimSequenceExplicit(
 				// construct the proposed compressor		
 				if( bTryLinearKeyRemovalCompression && AnimSeq->NumFrames > 1 )
 				{
-					UAnimCompress_RemoveLinearKeys* LinearKeyRemover = ConstructObject<UAnimCompress_RemoveLinearKeys>( UAnimCompress_RemoveLinearKeys::StaticClass() );
+					UAnimCompress_RemoveLinearKeys* LinearKeyRemover = NewObject<UAnimCompress_RemoveLinearKeys>();
 					{
 						// Try ACF_Float96NoW
 						LinearKeyRemover->RotationCompressionFormat = ACF_Float96NoW;
@@ -1080,7 +1080,7 @@ void FAnimationUtils::CompressAnimSequenceExplicit(
 
 				if( bTryPerTrackBitwiseCompression )
 				{
-					UAnimCompress_PerTrackCompression* PerTrackCompressor = ConstructObject<UAnimCompress_PerTrackCompression>(UAnimCompress_PerTrackCompression::StaticClass());
+					UAnimCompress_PerTrackCompression* PerTrackCompressor = NewObject<UAnimCompress_PerTrackCompression>();
 
 					// Straight PerTrackCompression, no key decimation and no linear key removal
 					TRYCOMPRESSION(Bitwise_PerTrack, PerTrackCompressor);
@@ -1112,7 +1112,7 @@ void FAnimationUtils::CompressAnimSequenceExplicit(
 
 				if( bTryPerTrackBitwiseCompression )
 				{
-					UAnimCompress_PerTrackCompression* PerTrackCompressor = ConstructObject<UAnimCompress_PerTrackCompression>(UAnimCompress_PerTrackCompression::StaticClass());
+					UAnimCompress_PerTrackCompression* PerTrackCompressor = NewObject<UAnimCompress_PerTrackCompression>();
 					PerTrackCompressor->bUseAdaptiveError = true;
 
 					if ( AnimSeq->NumFrames > 1 )
@@ -1145,7 +1145,7 @@ void FAnimationUtils::CompressAnimSequenceExplicit(
 
 				if( bTryPerTrackBitwiseCompression )
 				{
-					UAnimCompress_PerTrackCompression* PerTrackCompressor = ConstructObject<UAnimCompress_PerTrackCompression>(UAnimCompress_PerTrackCompression::StaticClass());
+					UAnimCompress_PerTrackCompression* PerTrackCompressor = NewObject<UAnimCompress_PerTrackCompression>();
 					PerTrackCompressor->bUseAdaptiveError = true;
 
 					// Try the decimation algorithms
@@ -1197,7 +1197,7 @@ void FAnimationUtils::CompressAnimSequenceExplicit(
 					// Try the decimation algorithms
 					if (AnimSeq->NumFrames >= 3)
 					{
-						UAnimCompress_PerTrackCompression* NewPerTrackCompressor = ConstructObject<UAnimCompress_PerTrackCompression>(UAnimCompress_PerTrackCompression::StaticClass());
+						UAnimCompress_PerTrackCompression* NewPerTrackCompressor = NewObject<UAnimCompress_PerTrackCompression>();
 
 						// Downsampling with linear key removal and adaptive error metrics v2
 						NewPerTrackCompressor->MinKeysForResampling = 3;
@@ -1220,7 +1220,7 @@ void FAnimationUtils::CompressAnimSequenceExplicit(
 				if( bTryPerTrackBitwiseCompression)
 				{
 					// Adaptive error through probing the effect of perturbations at each track
-					UAnimCompress_PerTrackCompression* NewPerTrackCompressor = ConstructObject<UAnimCompress_PerTrackCompression>(UAnimCompress_PerTrackCompression::StaticClass());
+					UAnimCompress_PerTrackCompression* NewPerTrackCompressor = NewObject<UAnimCompress_PerTrackCompression>();
 					NewPerTrackCompressor->bUseAdaptiveError2 = true;
 					NewPerTrackCompressor->MaxPosDiffBitwise = 0.05;
 					NewPerTrackCompressor->MaxAngleDiffBitwise = 0.02;
diff --git a/Engine/Source/Runtime/Engine/Private/AudioComponent.cpp b/Engine/Source/Runtime/Engine/Private/AudioComponent.cpp
--- a/Engine/Source/Runtime/Engine/Private/AudioComponent.cpp
+++ b/Engine/Source/Runtime/Engine/Private/AudioComponent.cpp
@@ -50,7 +50,7 @@ void UAudioComponent::OnRegister()
 
 	if ( bVisualizeComponent && SpriteComponent == NULL && GetOwner() && !GetWorld()->IsGameWorld() )
 	{
-		SpriteComponent = ConstructObject<UBillboardComponent>(UBillboardComponent::StaticClass(), GetOwner(), NAME_None, RF_Transactional | RF_TextExportTransient);
+		SpriteComponent = NewObject<UBillboardComponent>(GetOwner(), NAME_None, RF_Transactional | RF_TextExportTransient);
 
 		UpdateSpriteTexture();
 		SpriteComponent->RelativeScale3D = FVector(0.5f, 0.5f, 0.5f);
diff --git a/Engine/Source/Runtime/Engine/Private/AudioDevice.cpp b/Engine/Source/Runtime/Engine/Private/AudioDevice.cpp
--- a/Engine/Source/Runtime/Engine/Private/AudioDevice.cpp
+++ b/Engine/Source/Runtime/Engine/Private/AudioDevice.cpp
@@ -533,7 +533,7 @@ bool FAudioDevice::HandlePlaySoundCueCommand( const TCHAR* Cmd, FOutputDevice& A
 	// Stop any existing sound playing
 	if( !TestAudioComponent.IsValid() )
 	{
-		TestAudioComponent = ConstructObject<UAudioComponent>( UAudioComponent::StaticClass() );
+		TestAudioComponent = NewObject<UAudioComponent>();
 	}
 
 	UAudioComponent* AudioComp = TestAudioComponent.Get();
@@ -570,7 +570,7 @@ bool FAudioDevice::HandlePlaySoundWaveCommand( const TCHAR* Cmd, FOutputDevice&
 	// Stop any existing sound playing
 	if( !TestAudioComponent.IsValid() )
 	{
-		TestAudioComponent = ConstructObject<UAudioComponent>( UAudioComponent::StaticClass() );
+		TestAudioComponent = NewObject<UAudioComponent>();
 	}
 
 	UAudioComponent* AudioComp = TestAudioComponent.Get();
@@ -2044,12 +2044,12 @@ UAudioComponent* FAudioDevice::CreateComponent( USoundBase* Sound, UWorld* World
 			// Use actor as outer if we have one.
 			if( Actor )
 			{
-				AudioComponent = ConstructObject<UAudioComponent>( UAudioComponent::StaticClass(), Actor );
+				AudioComponent = NewObject<UAudioComponent>(Actor);
 			}
 			// Let engine pick the outer (transient package).
 			else
 			{
-				AudioComponent = ConstructObject<UAudioComponent>( UAudioComponent::StaticClass() );
+				AudioComponent = NewObject<UAudioComponent>();
 			}
 
 			check( AudioComponent );
diff --git a/Engine/Source/Runtime/Engine/Private/BlueprintGeneratedClass.cpp b/Engine/Source/Runtime/Engine/Private/BlueprintGeneratedClass.cpp
--- a/Engine/Source/Runtime/Engine/Private/BlueprintGeneratedClass.cpp
+++ b/Engine/Source/Runtime/Engine/Private/BlueprintGeneratedClass.cpp
@@ -318,7 +318,7 @@ void UBlueprintGeneratedClass::CreateComponentsForActor(AActor* Actor) const
 		}
 
 		FName NewName = *(FString::Printf(TEXT("TimelineComp__%d"), Actor->BlueprintCreatedComponents.Num() ) );
-		UTimelineComponent* NewTimeline = NewNamedObject<UTimelineComponent>(Actor, NewName);
+		UTimelineComponent* NewTimeline = NewObject<UTimelineComponent>(Actor, NewName);
 		NewTimeline->CreationMethod = EComponentCreationMethod::ConstructionScript; // Indicate it comes from a blueprint so it gets cleared when we rerun construction scripts
 		Actor->BlueprintCreatedComponents.Add(NewTimeline); // Add to array so it gets saved
 		NewTimeline->SetNetAddressable();	// This component has a stable name that can be referenced for replication
diff --git a/Engine/Source/Runtime/Engine/Private/BoxComponent.cpp b/Engine/Source/Runtime/Engine/Private/BoxComponent.cpp
--- a/Engine/Source/Runtime/Engine/Private/BoxComponent.cpp
+++ b/Engine/Source/Runtime/Engine/Private/BoxComponent.cpp
@@ -37,7 +37,7 @@ void UBoxComponent::UpdateBodySetup()
 {
 	if(ShapeBodySetup == NULL || ShapeBodySetup->IsPendingKill())
 	{
-		ShapeBodySetup = ConstructObject<UBodySetup>(UBodySetup::StaticClass(), this);
+		ShapeBodySetup = NewObject<UBodySetup>(this);
 		ShapeBodySetup->CollisionTraceFlag = CTF_UseSimpleAsComplex;
 		ShapeBodySetup->AggGeom.BoxElems.Add(FKBoxElem());
 	}
