Warning cleanup after limiting number of ways to construct UObjects.

diff --git a/Engine/Plugins/2D/Paper2D/Source/Paper2DEditor/Private/ExtractSprites/SPaperExtractSpritesDialog.cpp b/Engine/Plugins/2D/Paper2D/Source/Paper2DEditor/Private/ExtractSprites/SPaperExtractSpritesDialog.cpp
--- a/Engine/Plugins/2D/Paper2D/Source/Paper2DEditor/Private/ExtractSprites/SPaperExtractSpritesDialog.cpp
+++ b/Engine/Plugins/2D/Paper2D/Source/Paper2DEditor/Private/ExtractSprites/SPaperExtractSpritesDialog.cpp
@@ -340,7 +340,7 @@ void SPaperExtractSpritesDialog::CreateExtractedSprites()
 	TArray<UObject*> ObjectsToSync;
 
 	// Create the factory used to generate the sprite
-	UPaperSpriteFactory* SpriteFactory = ConstructObject<UPaperSpriteFactory>(UPaperSpriteFactory::StaticClass());
+	auto SpriteFactory = NewObject<UPaperSpriteFactory>();
 	SpriteFactory->InitialTexture = SourceTexture;
 
 	// Create the sprite
diff --git a/Engine/Plugins/2D/Paper2D/Source/Paper2DEditor/Private/PaperTileMapFactory.cpp b/Engine/Plugins/2D/Paper2D/Source/Paper2DEditor/Private/PaperTileMapFactory.cpp
--- a/Engine/Plugins/2D/Paper2D/Source/Paper2DEditor/Private/PaperTileMapFactory.cpp
+++ b/Engine/Plugins/2D/Paper2D/Source/Paper2DEditor/Private/PaperTileMapFactory.cpp
@@ -17,7 +17,7 @@ UPaperTileMapFactory::UPaperTileMapFactory(const FObjectInitializer& ObjectIniti
 
 UObject* UPaperTileMapFactory::FactoryCreateNew(UClass* Class, UObject* InParent, FName Name, EObjectFlags Flags, UObject* Context, FFeedbackContext* Warn)
 {
-	UPaperTileMap* NewTileMap = ConstructObject<UPaperTileMap>(Class, InParent, Name, Flags | RF_Transactional);
+	UPaperTileMap* NewTileMap = NewObject<UPaperTileMap>(InParent, Class, Name, Flags | RF_Transactional);
 
 	if (InitialTileSet != nullptr)
 	{
diff --git a/Engine/Plugins/2D/Paper2D/Source/Paper2DEditor/Private/PaperTileSetFactory.cpp b/Engine/Plugins/2D/Paper2D/Source/Paper2DEditor/Private/PaperTileSetFactory.cpp
--- a/Engine/Plugins/2D/Paper2D/Source/Paper2DEditor/Private/PaperTileSetFactory.cpp
+++ b/Engine/Plugins/2D/Paper2D/Source/Paper2DEditor/Private/PaperTileSetFactory.cpp
@@ -17,7 +17,7 @@ UPaperTileSetFactory::UPaperTileSetFactory(const FObjectInitializer& ObjectIniti
 
 UObject* UPaperTileSetFactory::FactoryCreateNew(UClass* Class, UObject* InParent, FName Name, EObjectFlags Flags, UObject* Context, FFeedbackContext* Warn)
 {
-	UPaperTileSet* NewTileSet = ConstructObject<UPaperTileSet>(Class, InParent, Name, Flags | RF_Transactional);
+	auto NewTileSet = NewObject<UPaperTileSet>(InParent, Class, Name, Flags | RF_Transactional);
 	NewTileSet->TileSheet = InitialTexture;
 
 	return NewTileSet;
diff --git a/Engine/Plugins/Runtime/SoundMod/Source/SoundModImporter/Private/SoundModImporterFactory.cpp b/Engine/Plugins/Runtime/SoundMod/Source/SoundModImporter/Private/SoundModImporterFactory.cpp
--- a/Engine/Plugins/Runtime/SoundMod/Source/SoundModImporter/Private/SoundModImporterFactory.cpp
+++ b/Engine/Plugins/Runtime/SoundMod/Source/SoundModImporter/Private/SoundModImporterFactory.cpp
@@ -103,7 +103,7 @@ UObject* USoundModImporterFactory::FactoryCreateBinary
 
 	// Use pre-existing sound if it exists and we want to keep settings,
 	// otherwise create new sound and import raw data.
-	USoundMod* Sound = (bUseExistingSettings && ExistingSound) ? ExistingSound : NewNamedObject<USoundMod>(InParent, Name, Flags);
+	USoundMod* Sound = (bUseExistingSettings && ExistingSound) ? ExistingSound : NewObject<USoundMod>(InParent, Name, Flags);
 
 	Sound->Duration = xmpModuleInfo.seq_data->duration / 1000.f;
 
diff --git a/Engine/Plugins/ScriptGeneratorPlugin/Source/ScriptGeneratorPlugin/Private/LuaScriptCodeGenerator.cpp b/Engine/Plugins/ScriptGeneratorPlugin/Source/ScriptGeneratorPlugin/Private/LuaScriptCodeGenerator.cpp
--- a/Engine/Plugins/ScriptGeneratorPlugin/Source/ScriptGeneratorPlugin/Private/LuaScriptCodeGenerator.cpp
+++ b/Engine/Plugins/ScriptGeneratorPlugin/Source/ScriptGeneratorPlugin/Private/LuaScriptCodeGenerator.cpp
@@ -431,7 +431,7 @@ FString FLuaScriptCodeGenerator::ExportAdditionalClassGlue(const FString& ClassN
 		GeneratedGlue += TEXT("\r\n{\r\n");
 		GeneratedGlue += TEXT("\tUObject* Outer = (UObject*)lua_touserdata(InScriptContext, 1);\r\n");
 		GeneratedGlue += TEXT("\tFName Name = FName(luaL_checkstring(InScriptContext, 2));\r\n");
-		GeneratedGlue += FString::Printf(TEXT("\tUObject* Obj = NewNamedObject<%s>(Outer, Name);\r\n"), *ClassNameCPP);
+		GeneratedGlue += FString::Printf(TEXT("\tUObject* Obj = NewObject<%s>(Outer, Name);\r\n"), *ClassNameCPP);
 		GeneratedGlue += TEXT("\tif (Obj)\r\n\t{\r\n");
 		GeneratedGlue += TEXT("\t\tFScriptObjectReferencer::Get().AddObjectReference(Obj);\r\n");
 		GeneratedGlue += TEXT("\t}\r\n");
diff --git a/Engine/Source/Developer/MeshUtilities/Private/MeshUtilities.cpp b/Engine/Source/Developer/MeshUtilities/Private/MeshUtilities.cpp
--- a/Engine/Source/Developer/MeshUtilities/Private/MeshUtilities.cpp
+++ b/Engine/Source/Developer/MeshUtilities/Private/MeshUtilities.cpp
@@ -3716,7 +3716,7 @@ void FMeshUtilities::CreateProxyMesh(
 		MeshPackage->Modify();
 	}
 
-	auto StaticMesh = NewNamedObject<UStaticMesh>(MeshPackage, FName(*MeshAssetName), RF_Public | RF_Standalone);
+	auto StaticMesh = NewObject<UStaticMesh>(MeshPackage, FName(*MeshAssetName), RF_Public | RF_Standalone);
 	StaticMesh->InitResources();
 	{
 		FString OutputPath = StaticMesh->GetPathName();
@@ -4102,7 +4102,7 @@ void FMeshUtilities::MergeActors(
 		Package->FullyLoad();
 		Package->Modify();
 
-		auto StaticMesh = NewNamedObject<UStaticMesh>(Package, *AssetName, RF_Public | RF_Standalone);
+		auto StaticMesh = NewObject<UStaticMesh>(Package, *AssetName, RF_Public | RF_Standalone);
 		StaticMesh->InitResources();
 		
 		FString OutputPath = StaticMesh->GetPathName();
diff --git a/Engine/Source/Editor/FoliageEdit/Private/FoliageEditModule.cpp b/Engine/Source/Editor/FoliageEdit/Private/FoliageEditModule.cpp
--- a/Engine/Source/Editor/FoliageEdit/Private/FoliageEditModule.cpp
+++ b/Engine/Source/Editor/FoliageEdit/Private/FoliageEditModule.cpp
@@ -46,7 +46,7 @@ public:
 		PropertyEditor.RegisterCustomClassLayout("ProceduralFoliageComponent", FOnGetDetailCustomizationInstance::CreateStatic(&FProceduralFoliageComponentDetails::MakeInstance));
 
 		// Actor Factories
-		UActorFactoryProceduralFoliage* ProceduralFoliageActorFactory = ConstructObject<UActorFactoryProceduralFoliage>(UActorFactoryProceduralFoliage::StaticClass());
+		auto ProceduralFoliageActorFactory = NewObject<UActorFactoryProceduralFoliage>();
 		GEditor->ActorFactories.Add(ProceduralFoliageActorFactory);
 
 #if WITH_EDITOR
diff --git a/Engine/Source/Editor/FoliageEdit/Private/LandscapeGrassTypeFactory.cpp b/Engine/Source/Editor/FoliageEdit/Private/LandscapeGrassTypeFactory.cpp
--- a/Engine/Source/Editor/FoliageEdit/Private/LandscapeGrassTypeFactory.cpp
+++ b/Engine/Source/Editor/FoliageEdit/Private/LandscapeGrassTypeFactory.cpp
@@ -15,7 +15,7 @@ ULandscapeGrassTypeFactory::ULandscapeGrassTypeFactory(const FObjectInitializer&
 
 UObject* ULandscapeGrassTypeFactory::FactoryCreateNew(UClass* Class, UObject* InParent, FName Name, EObjectFlags Flags, UObject* Context, FFeedbackContext* Warn)
 {
-	ULandscapeGrassType* NewGrassType = ConstructObject<ULandscapeGrassType>(Class, InParent, Name, Flags | RF_Transactional);
+	auto NewGrassType = NewObject<ULandscapeGrassType>(InParent, Class, Name, Flags | RF_Transactional);
 
 	return NewGrassType;
 }
diff --git a/Engine/Source/Editor/FoliageEdit/Private/ProceduralFoliageFactory.cpp b/Engine/Source/Editor/FoliageEdit/Private/ProceduralFoliageFactory.cpp
--- a/Engine/Source/Editor/FoliageEdit/Private/ProceduralFoliageFactory.cpp
+++ b/Engine/Source/Editor/FoliageEdit/Private/ProceduralFoliageFactory.cpp
@@ -15,7 +15,7 @@ UProceduralFoliageFactory::UProceduralFoliageFactory(const FObjectInitializer& O
 
 UObject* UProceduralFoliageFactory::FactoryCreateNew(UClass* Class, UObject* InParent, FName Name, EObjectFlags Flags, UObject* Context, FFeedbackContext* Warn)
 {
-	UProceduralFoliage* NewProceduralFoliage = ConstructObject<UProceduralFoliage>(Class, InParent, Name, Flags | RF_Transactional);
+	auto NewProceduralFoliage = NewObject<UProceduralFoliage>(InParent, Class, Name, Flags | RF_Transactional);
 
 	return NewProceduralFoliage;
 }
diff --git a/Engine/Source/Editor/Kismet/Private/SSCSEditor.cpp b/Engine/Source/Editor/Kismet/Private/SSCSEditor.cpp
--- a/Engine/Source/Editor/Kismet/Private/SSCSEditor.cpp
+++ b/Engine/Source/Editor/Kismet/Private/SSCSEditor.cpp
@@ -4465,7 +4465,7 @@ UActorComponent* SSCSEditor::AddNewComponent( UClass* NewComponentClass, UObject
 			{
 				NewComponentName = *FComponentEditorUtils::GenerateValidVariableName(NewComponentClass, ActorInstance);
 			}
-			NewComponent = AddNewNodeForInstancedComponent(ConstructObject<UActorComponent>(NewComponentClass, ActorInstance, NewComponentName, RF_Transactional, ComponentTemplate), Asset, true);
+			NewComponent = AddNewNodeForInstancedComponent(NewObject<UActorComponent>(ActorInstance, NewComponentClass, NewComponentName, RF_Transactional, ComponentTemplate), Asset, true);
 		}
 	}
 
diff --git a/Engine/Source/Editor/UnrealEd/Private/AutoReimport/ContentDirectoryMonitor.cpp b/Engine/Source/Editor/UnrealEd/Private/AutoReimport/ContentDirectoryMonitor.cpp
--- a/Engine/Source/Editor/UnrealEd/Private/AutoReimport/ContentDirectoryMonitor.cpp
+++ b/Engine/Source/Editor/UnrealEd/Private/AutoReimport/ContentDirectoryMonitor.cpp
@@ -159,7 +159,7 @@ void FContentDirectoryMonitor::ProcessAdditions(TArray<UPackage*>& OutPackagesTo
 				}
 
 				UObject* NewAsset = nullptr;
-				UFactory* FactoryInstance = FactoryType ? ConstructObject<UFactory>(FactoryType->GetClass()) : nullptr;
+				UFactory* FactoryInstance = FactoryType ? NewObject<UFactory>(GetTransientPackage(), FactoryType->GetClass()) : nullptr;
 				if (FactoryInstance)
 				{
 					FactoryInstance->AddToRoot();
diff --git a/Engine/Source/Editor/UnrealEd/Private/BSPOps.cpp b/Engine/Source/Editor/UnrealEd/Private/BSPOps.cpp
--- a/Engine/Source/Editor/UnrealEd/Private/BSPOps.cpp
+++ b/Engine/Source/Editor/UnrealEd/Private/BSPOps.cpp
@@ -549,9 +549,9 @@ void FBSPOps::csgCopyBrush( ABrush* Dest, ABrush* Src, uint32 PolyFlags, EObject
 
 	// Duplicate the brush and its polys.
 	Dest->PolyFlags		= PolyFlags;
-	Dest->Brush = NewNamedObject<UModel>(Dest, NAME_None, ResFlags);
+	Dest->Brush = NewObject<UModel>(Dest, NAME_None, ResFlags);
 	Dest->Brush->Initialize(nullptr, Src->Brush->RootOutside);
-	Dest->Brush->Polys	= NewNamedObject<UPolys>(Dest->Brush, NAME_None, ResFlags);
+	Dest->Brush->Polys	= NewObject<UPolys>(Dest->Brush, NAME_None, ResFlags);
 	check(Dest->Brush->Polys->Element.GetOwner()==Dest->Brush->Polys);
 	Dest->Brush->Polys->Element.AssignButKeepOwner(Src->Brush->Polys->Element);
 	check(Dest->Brush->Polys->Element.GetOwner()==Dest->Brush->Polys);
diff --git a/Engine/Source/Editor/UnrealEd/Private/EditorExporters.cpp b/Engine/Source/Editor/UnrealEd/Private/EditorExporters.cpp
--- a/Engine/Source/Editor/UnrealEd/Private/EditorExporters.cpp
+++ b/Engine/Source/Editor/UnrealEd/Private/EditorExporters.cpp
@@ -2876,7 +2876,7 @@ namespace MaterialExportUtils
 			MaterialOuter->Modify();
 		}
 
-		UMaterial* Material = ConstructObject<UMaterial>(UMaterial::StaticClass(), MaterialOuter, FName(*MaterialAssetName), Flags);
+		UMaterial* Material = NewObject<UMaterial>(MaterialOuter, FName(*MaterialAssetName), Flags);
 		Material->TwoSided = false;
 		Material->SetShadingModel(MSM_DefaultLit);
 		OutGeneratedAssets.Add(Material);
@@ -2893,7 +2893,7 @@ namespace MaterialExportUtils
 			UTexture2D* Texture = CreateTexture(Outer, AssetLongName, InFlattenMaterial.DiffuseSize, InFlattenMaterial.DiffuseSamples, TC_Default, TEXTUREGROUP_World, Flags, bSRGB);
 			OutGeneratedAssets.Add(Texture);
 			
-			UMaterialExpressionTextureSample* BasecolorExpression = ConstructObject<UMaterialExpressionTextureSample>(UMaterialExpressionTextureSample::StaticClass(), Material);
+			auto BasecolorExpression = NewObject<UMaterialExpressionTextureSample>(Material);
 			BasecolorExpression->Texture = Texture;
 			BasecolorExpression->SamplerType = EMaterialSamplerType::SAMPLERTYPE_Color;
 			BasecolorExpression->MaterialExpressionEditorX = -400;
@@ -2912,7 +2912,7 @@ namespace MaterialExportUtils
 			UTexture2D* Texture = CreateTexture(Outer, AssetBasePath + AssetName, InFlattenMaterial.MetallicSize, InFlattenMaterial.MetallicSamples, TC_Grayscale, TEXTUREGROUP_World, Flags, bSRGB);
 			OutGeneratedAssets.Add(Texture);
 			
-			UMaterialExpressionTextureSample* MetallicExpression = ConstructObject<UMaterialExpressionTextureSample>(UMaterialExpressionTextureSample::StaticClass(), Material);
+			auto MetallicExpression = NewObject<UMaterialExpressionTextureSample>(Material);
 			MetallicExpression->Texture = Texture;
 			MetallicExpression->SamplerType = EMaterialSamplerType::SAMPLERTYPE_LinearGrayscale;
 			MetallicExpression->MaterialExpressionEditorX = -400;
@@ -2926,7 +2926,7 @@ namespace MaterialExportUtils
 		{
 			// Set Metallic to constant
 			float Metallic = InFlattenMaterial.MetallicSamples.Num() ? InFlattenMaterial.MetallicSamples[0].R : 0.0f;
-			UMaterialExpressionConstant* MetallicExpression = ConstructObject<UMaterialExpressionConstant>(UMaterialExpressionConstant::StaticClass(), Material);
+			auto MetallicExpression = NewObject<UMaterialExpressionConstant>(Material);
 			MetallicExpression->R = Metallic;
 			MetallicExpression->MaterialExpressionEditorX = -400;
 			MetallicExpression->MaterialExpressionEditorY = MaterialNodeY;
@@ -2944,7 +2944,7 @@ namespace MaterialExportUtils
 			UTexture2D* Texture = CreateTexture(Outer, AssetBasePath + AssetName, InFlattenMaterial.SpecularSize, InFlattenMaterial.SpecularSamples, TC_Default, TEXTUREGROUP_World, Flags, bSRGB);
 			OutGeneratedAssets.Add(Texture);
 			
-			UMaterialExpressionTextureSample* SpecularExpression = ConstructObject<UMaterialExpressionTextureSample>(UMaterialExpressionTextureSample::StaticClass(), Material);
+			auto SpecularExpression = NewObject<UMaterialExpressionTextureSample>(Material);
 			SpecularExpression->Texture = Texture;
 			SpecularExpression->SamplerType = EMaterialSamplerType::SAMPLERTYPE_LinearColor;
 			SpecularExpression->MaterialExpressionEditorX = -400;
@@ -2963,7 +2963,7 @@ namespace MaterialExportUtils
 			UTexture2D* Texture = CreateTexture(Outer, AssetBasePath + AssetName, InFlattenMaterial.RoughnessSize, InFlattenMaterial.RoughnessSamples, TC_Grayscale, TEXTUREGROUP_World, Flags, bSRGB);
 			OutGeneratedAssets.Add(Texture);
 			
-			UMaterialExpressionTextureSample* RoughnessExpression = ConstructObject<UMaterialExpressionTextureSample>(UMaterialExpressionTextureSample::StaticClass(), Material);
+			auto RoughnessExpression = NewObject<UMaterialExpressionTextureSample>(Material);
 			RoughnessExpression->Texture = Texture;
 			RoughnessExpression->SamplerType = EMaterialSamplerType::SAMPLERTYPE_LinearGrayscale;
 			RoughnessExpression->MaterialExpressionEditorX = -400;
@@ -2977,7 +2977,7 @@ namespace MaterialExportUtils
 		{
 			// Set Roughness to constant
 			float Roughness = InFlattenMaterial.RoughnessSamples.Num() ? InFlattenMaterial.RoughnessSamples[0].R : 0.8f;
-			UMaterialExpressionConstant* RoughnessExpression = ConstructObject<UMaterialExpressionConstant>(UMaterialExpressionConstant::StaticClass(), Material);
+			auto RoughnessExpression = NewObject<UMaterialExpressionConstant>(Material);
 			RoughnessExpression->R = Roughness;
 			RoughnessExpression->MaterialExpressionEditorX = -400;
 			RoughnessExpression->MaterialExpressionEditorY = MaterialNodeY;
@@ -2995,7 +2995,7 @@ namespace MaterialExportUtils
 			UTexture2D* Texture = CreateTexture(Outer, AssetBasePath + AssetName, InFlattenMaterial.NormalSize, InFlattenMaterial.NormalSamples, TC_Normalmap, TEXTUREGROUP_WorldNormalMap, Flags, bSRGB);
 			OutGeneratedAssets.Add(Texture);
 			
-			UMaterialExpressionTextureSample* NormalExpression = ConstructObject<UMaterialExpressionTextureSample>(UMaterialExpressionTextureSample::StaticClass(), Material);
+			auto NormalExpression = NewObject<UMaterialExpressionTextureSample>(Material);
 			NormalExpression->Texture = Texture;
 			NormalExpression->SamplerType = EMaterialSamplerType::SAMPLERTYPE_Normal;
 			NormalExpression->MaterialExpressionEditorX = -400;
diff --git a/Engine/Source/Editor/UnrealEd/Private/Kismet2/KismetReinstanceUtilities.cpp b/Engine/Source/Editor/UnrealEd/Private/Kismet2/KismetReinstanceUtilities.cpp
--- a/Engine/Source/Editor/UnrealEd/Private/Kismet2/KismetReinstanceUtilities.cpp
+++ b/Engine/Source/Editor/UnrealEd/Private/Kismet2/KismetReinstanceUtilities.cpp
@@ -685,7 +685,7 @@ void FBlueprintCompileReinstancer::ReplaceInstancesOfClass(UClass* OldClass, UCl
 			}
 
 			AActor*  OldActor = Cast<AActor>(OldObject);
-			UObject* NewObject = nullptr;
+			UObject* NewUObject = nullptr;
 			// if the object to replace is an actor...
 			if (OldActor != nullptr)
 			{
@@ -727,7 +727,7 @@ void FBlueprintCompileReinstancer::ReplaceInstancesOfClass(UClass* OldClass, UCl
 				}
 
 				check(NewActor != nullptr);
-				NewObject = NewActor;
+				NewUObject = NewActor;
 				// store the new actor for the second pass (NOTE: this detaches 
 				// OldActor from all child/parent attachments)
 				//
@@ -736,7 +736,7 @@ void FBlueprintCompileReinstancer::ReplaceInstancesOfClass(UClass* OldClass, UCl
 				// another instance that hasn't been replaced yet).
 				ReplacementActors.Add(FActorReplacementHelper(NewActor, OldActor));
 
-				ReinstancedObjectsWeakReferenceMap.Add(OldObject, NewObject);
+				ReinstancedObjectsWeakReferenceMap.Add(OldObject, NewUObject);
 
 				OldActor->DestroyConstructedComponents(); // don't want to serialize components from the old actor
 				// Unregister native components so we don't copy any sub-components they generate for themselves (like UCameraComponent does)
@@ -771,12 +771,12 @@ void FBlueprintCompileReinstancer::ReplaceInstancesOfClass(UClass* OldClass, UCl
 			{
 				FName OldName(OldObject->GetFName());
 				OldObject->Rename(NULL, OldObject->GetOuter(), REN_DoNotDirty | REN_DontCreateRedirectors);
-				NewObject = ConstructObject<UObject>(NewClass, OldObject->GetOuter(), OldName);
-				check(NewObject != nullptr);
+				NewUObject = NewObject<UObject>(OldObject->GetOuter(), NewClass, OldName);
+				check(NewUObject != nullptr);
 
-				UEditorEngine::CopyPropertiesForUnrelatedObjects(OldObject, NewObject);
+				UEditorEngine::CopyPropertiesForUnrelatedObjects(OldObject, NewUObject);
 
-				if (UAnimInstance* AnimTree = Cast<UAnimInstance>(NewObject))
+				if (UAnimInstance* AnimTree = Cast<UAnimInstance>(NewUObject))
 				{
 					// Initialising the anim instance isn't enough to correctly set up the skeletal mesh again in a
 					// paused world, need to initialise the skeletal mesh component that contains the anim instance.
@@ -789,11 +789,11 @@ void FBlueprintCompileReinstancer::ReplaceInstancesOfClass(UClass* OldClass, UCl
 				OldObject->RemoveFromRoot();
 				OldObject->MarkPendingKill();
 
-				OldToNewInstanceMap.Add(OldObject, NewObject);
+				OldToNewInstanceMap.Add(OldObject, NewUObject);
 
 				if (bIsComponent)
 				{
-					UActorComponent* Component = Cast<UActorComponent>(NewObject);
+					UActorComponent* Component = Cast<UActorComponent>(NewUObject);
 					AActor* OwningActor = Component->GetOwner();
 					if (OwningActor)
 					{
@@ -809,14 +809,14 @@ void FBlueprintCompileReinstancer::ReplaceInstancesOfClass(UClass* OldClass, UCl
 			}
 
 			// If this original object came from a blueprint and it was in the selected debug set, change the debugging to the new object.
-			if ((CorrespondingBlueprint) && (OldBlueprintDebugObject) && (NewObject))
+			if ((CorrespondingBlueprint) && (OldBlueprintDebugObject) && (NewUObject))
 			{
-				CorrespondingBlueprint->SetObjectBeingDebugged(NewObject);
+				CorrespondingBlueprint->SetObjectBeingDebugged(NewUObject);
 			}
 
 			if (bLogConversions)
 			{
-				UE_LOG(LogBlueprint, Log, TEXT("Converted instance '%s' to '%s'"), *OldObject->GetPathName(), *NewObject->GetPathName());
+				UE_LOG(LogBlueprint, Log, TEXT("Converted instance '%s' to '%s'"), *OldObject->GetPathName(), *NewUObject->GetPathName());
 			}
 		}
 	}
diff --git a/Engine/Source/Editor/UnrealEd/Private/StaticMeshEdit.cpp b/Engine/Source/Editor/UnrealEd/Private/StaticMeshEdit.cpp
--- a/Engine/Source/Editor/UnrealEd/Private/StaticMeshEdit.cpp
+++ b/Engine/Source/Editor/UnrealEd/Private/StaticMeshEdit.cpp
@@ -510,7 +510,7 @@ UStaticMesh* CreateStaticMesh(struct FRawMesh& RawMesh,TArray<UMaterialInterface
 {
 	// Create the UStaticMesh object.
 	FStaticMeshComponentRecreateRenderStateContext RecreateRenderStateContext(FindObject<UStaticMesh>(InOuter,*InName.ToString()));
-	auto StaticMesh = NewNamedObject<UStaticMesh>(InOuter, InName, RF_Public | RF_Standalone);
+	auto StaticMesh = NewObject<UStaticMesh>(InOuter, InName, RF_Public | RF_Standalone);
 
 	// Add one LOD for the base mesh
 	FStaticMeshSourceModel* SrcModel = new(StaticMesh->SourceModels) FStaticMeshSourceModel();
diff --git a/Engine/Source/Editor/UnrealEd/Private/UnrealEdMisc.cpp b/Engine/Source/Editor/UnrealEd/Private/UnrealEdMisc.cpp
--- a/Engine/Source/Editor/UnrealEd/Private/UnrealEdMisc.cpp
+++ b/Engine/Source/Editor/UnrealEd/Private/UnrealEdMisc.cpp
@@ -656,7 +656,7 @@ bool FUnrealEdMisc::EnableWorldComposition(UWorld* InWorld, bool bEnable)
 			return false;
 		}
 			
-		UWorldComposition* WorldCompostion = ConstructObject<UWorldComposition>(UWorldComposition::StaticClass(), InWorld);
+		auto WorldCompostion = NewObject<UWorldComposition>(InWorld);
 		// All map files found in the same and folder and all sub-folders will be added ass sub-levels to this map
 		// Make sure user understands this
 		int32 NumFoundSublevels = WorldCompostion->GetTilesList().Num();
diff --git a/Engine/Source/Runtime/AIModule/Private/BehaviorTree/Blackboard/BlackboardKeyType.cpp b/Engine/Source/Runtime/AIModule/Private/BehaviorTree/Blackboard/BlackboardKeyType.cpp
--- a/Engine/Source/Runtime/AIModule/Private/BehaviorTree/Blackboard/BlackboardKeyType.cpp
+++ b/Engine/Source/Runtime/AIModule/Private/BehaviorTree/Blackboard/BlackboardKeyType.cpp
@@ -23,7 +23,7 @@ void UBlackboardKeyType::InitializeKey(UBlackboardComponent& OwnerComp, FBlackbo
 	if (bCreateKeyInstance)
 	{
 		FBlackboardInstancedKeyMemory* MyMemory = (FBlackboardInstancedKeyMemory*)RawData;
-		UBlackboardKeyType* KeyInstance = ConstructObject<UBlackboardKeyType>(GetClass(), &OwnerComp, GetFName(), RF_NoFlags, (UObject*)(this));
+		UBlackboardKeyType* KeyInstance = NewObject<UBlackboardKeyType>(&OwnerComp, GetClass(), GetFName(), RF_NoFlags, (UObject*)(this));
 		MyMemory->KeyIdx = KeyID;
 		OwnerComp.KeyInstances[KeyID] = KeyInstance;
 
diff --git a/Engine/Source/Runtime/CoreUObject/Private/Blueprint/BlueprintSupport.cpp b/Engine/Source/Runtime/CoreUObject/Private/Blueprint/BlueprintSupport.cpp
--- a/Engine/Source/Runtime/CoreUObject/Private/Blueprint/BlueprintSupport.cpp
+++ b/Engine/Source/Runtime/CoreUObject/Private/Blueprint/BlueprintSupport.cpp
@@ -321,7 +321,7 @@ bool ULinkerLoad::DeferPotentialCircularImport(const int32 Index)
 					FName PlaceholderName(*FString::Printf(TEXT("PLACEHOLDER-CLASS_%s"), *Import.ObjectName.ToString()));
 					PlaceholderName = MakeUniqueObjectName(PlaceholderOuter, PlaceholderType, PlaceholderName);
 
-					ULinkerPlaceholderClass* Placeholder = ConstructObject<ULinkerPlaceholderClass>(PlaceholderType, PlaceholderOuter, PlaceholderName, RF_Public | RF_Transient);
+					ULinkerPlaceholderClass* Placeholder = NewObject<ULinkerPlaceholderClass>(PlaceholderOuter, PlaceholderType, PlaceholderName, RF_Public | RF_Transient);
 					// store the import index in the placeholder, so we can 
 					// easily look it up in the import map, given the 
 					// placeholder (needed, to find the corresponding import for 
diff --git a/Engine/Source/Runtime/CoreUObject/Private/UObject/UObjectGlobals.cpp b/Engine/Source/Runtime/CoreUObject/Private/UObject/UObjectGlobals.cpp
--- a/Engine/Source/Runtime/CoreUObject/Private/UObject/UObjectGlobals.cpp
+++ b/Engine/Source/Runtime/CoreUObject/Private/UObject/UObjectGlobals.cpp
@@ -504,7 +504,7 @@ UPackage* CreatePackage( UObject* InOuter, const TCHAR* PackageName )
 			}
 			else
 			{
-				Result = NewNamedObject<UPackage>(InOuter, NewPackageName, RF_Public);
+				Result = NewObject<UPackage>(InOuter, NewPackageName, RF_Public);
 			}
 		}
 	}
diff --git a/Engine/Source/Runtime/Engine/Private/Blueprint.cpp b/Engine/Source/Runtime/Engine/Private/Blueprint.cpp
--- a/Engine/Source/Runtime/Engine/Private/Blueprint.cpp
+++ b/Engine/Source/Runtime/Engine/Private/Blueprint.cpp
@@ -1240,7 +1240,7 @@ UInheritableComponentHandler* UBlueprint::GetInheritableComponentHandler(bool bC
 	{
 		auto BPGC = CastChecked<UBlueprintGeneratedClass>(GeneratedClass);
 		ensure(!BPGC->InheritableComponentHandler);
-		InheritableComponentHandler = NewNamedObject<UInheritableComponentHandler>(BPGC, FName(TEXT("InheritableComponentHandler")));
+		InheritableComponentHandler = NewObject<UInheritableComponentHandler>(BPGC, FName(TEXT("InheritableComponentHandler")));
 		BPGC->InheritableComponentHandler = InheritableComponentHandler;
 	}
 	return InheritableComponentHandler;
diff --git a/Engine/Source/Runtime/Engine/Private/InheritableComponentHandler.cpp b/Engine/Source/Runtime/Engine/Private/InheritableComponentHandler.cpp
--- a/Engine/Source/Runtime/Engine/Private/InheritableComponentHandler.cpp
+++ b/Engine/Source/Runtime/Engine/Private/InheritableComponentHandler.cpp
@@ -30,8 +30,8 @@ UActorComponent* UInheritableComponentHandler::CreateOverridenComponentTemplate(
 		return NULL;
 	}
 	ensure(Cast<UBlueprintGeneratedClass>(GetOuter()));
-	auto NewComponentTemplate = ConstructObject<UActorComponent>(
-		BestArchetype->GetClass(), GetOuter(), BestArchetype->GetFName(), RF_ArchetypeObject | RF_Public | RF_InheritableComponentTemplate, BestArchetype);
+	auto NewComponentTemplate = NewObject<UActorComponent>(
+		GetOuter(), BestArchetype->GetClass(), BestArchetype->GetFName(), RF_ArchetypeObject | RF_Public | RF_InheritableComponentTemplate, BestArchetype);
 
 	FComponentOverrideRecord NewRecord;
 	NewRecord.ComponentKey = Key;
diff --git a/Engine/Source/Runtime/Engine/Private/LightMap.cpp b/Engine/Source/Runtime/Engine/Private/LightMap.cpp
--- a/Engine/Source/Runtime/Engine/Private/LightMap.cpp
+++ b/Engine/Source/Runtime/Engine/Private/LightMap.cpp
@@ -815,7 +815,7 @@ void FLightMapPendingTexture::StartEncoding()
 
 	if (bNeedsSkyOcclusionTexture)
 	{
-		auto Texture = NewNamedObject<ULightMapTexture2D>(Outer, GetSkyOcclusionTextureName(GLightmapCounter));
+		auto Texture = NewObject<ULightMapTexture2D>(Outer, GetSkyOcclusionTextureName(GLightmapCounter));
 		SkyOcclusionTexture = Texture;
 
 		Texture->Source.Init2DWithMipChain(GetSizeX(), GetSizeY(), TSF_BGRA8);
@@ -914,7 +914,7 @@ void FLightMapPendingTexture::StartEncoding()
 		}
 
 		// Create the light-map texture for this coefficient.
-		auto Texture = NewNamedObject<ULightMapTexture2D>(Outer, GetLightmapName(GLightmapCounter, CoefficientIndex));
+		auto Texture = NewObject<ULightMapTexture2D>(Outer, GetLightmapName(GLightmapCounter, CoefficientIndex));
 		Textures[CoefficientIndex] = Texture;
 		Texture->Source.Init2DWithMipChain(GetSizeX(), GetSizeY() * 2, TSF_BGRA8);	// Top/bottom atlased
 		Texture->MipGenSettings = TMGS_LeaveExistingMips;
diff --git a/Engine/Source/Runtime/Engine/Private/Model.cpp b/Engine/Source/Runtime/Engine/Private/Model.cpp
--- a/Engine/Source/Runtime/Engine/Private/Model.cpp
+++ b/Engine/Source/Runtime/Engine/Private/Model.cpp
@@ -491,7 +491,7 @@ void UModel::EmptyModel( int32 EmptySurfInfo, int32 EmptyPolys )
 #if WITH_EDITOR
 	if( EmptyPolys )
 	{
-		Polys = NewNamedObject<UPolys>(GetOuter(), NAME_None, RF_Transactional);
+		Polys = NewObject<UPolys>(GetOuter(), NAME_None, RF_Transactional);
 	}
 #endif // WITH_EDITOR
 
diff --git a/Engine/Source/Runtime/Engine/Private/NiagaraEffect.cpp b/Engine/Source/Runtime/Engine/Private/NiagaraEffect.cpp
--- a/Engine/Source/Runtime/Engine/Private/NiagaraEffect.cpp
+++ b/Engine/Source/Runtime/Engine/Private/NiagaraEffect.cpp
@@ -26,7 +26,7 @@ void UNiagaraEffect::CreateEffectRendererProps(TSharedPtr<FNiagaraSimulation> Si
 	UClass *RendererProps = Sim->GetEffectRenderer()->GetPropertiesClass();
 	if (RendererProps)
 	{
-		Sim->GetProperties()->RendererProperties = ConstructObject<UNiagaraEffectRendererProperties>(RendererProps, this);
+		Sim->GetProperties()->RendererProperties = NewObject<UNiagaraEffectRendererProperties>(this, RendererProps);
 	}
 	else
 	{
diff --git a/Engine/Source/Runtime/Engine/Private/SceneComponent.cpp b/Engine/Source/Runtime/Engine/Private/SceneComponent.cpp
--- a/Engine/Source/Runtime/Engine/Private/SceneComponent.cpp
+++ b/Engine/Source/Runtime/Engine/Private/SceneComponent.cpp
@@ -295,7 +295,7 @@ void USceneComponent::DestroyComponent(bool bPromoteChildren/*= false*/)
 				if (ChildToPromote == nullptr || ChildToPromote->IsEditorOnly())
 				{
 					// Construct a new default root component
-					USceneComponent* NewRootComponent = ConstructObject<USceneComponent>(USceneComponent::StaticClass(), Owner, USceneComponent::GetDefaultSceneRootVariableName(), RF_Transactional);
+					auto NewRootComponent = NewObject<USceneComponent>(Owner, USceneComponent::GetDefaultSceneRootVariableName(), RF_Transactional);
 					NewRootComponent->Mobility = Mobility;
 					NewRootComponent->SetWorldLocationAndRotation(GetComponentLocation(), GetComponentRotation());
 
diff --git a/Engine/Source/Runtime/Engine/Private/SimpleConstructionScript.cpp b/Engine/Source/Runtime/Engine/Private/SimpleConstructionScript.cpp
--- a/Engine/Source/Runtime/Engine/Private/SimpleConstructionScript.cpp
+++ b/Engine/Source/Runtime/Engine/Private/SimpleConstructionScript.cpp
@@ -767,7 +767,7 @@ USCS_Node* USimpleConstructionScript::CreateNode(UClass* NewComponentClass, FNam
 		NewComponentName = MakeUniqueObjectName(Blueprint->GeneratedClass, NewComponentClass, FName(*NewClassName.Left(NewStrLen)));
 	}
 
-	UActorComponent* NewComponentTemplate = ConstructObject<UActorComponent>(NewComponentClass, Blueprint->GeneratedClass, NewComponentName);
+	auto NewComponentTemplate = NewObject<UActorComponent>(Blueprint->GeneratedClass, NewComponentClass, NewComponentName);
 	NewComponentTemplate->SetFlags(RF_ArchetypeObject|RF_Transactional|RF_Public);
 
 	return CreateNode(NewComponentTemplate, NewComponentVariableName);
@@ -779,7 +779,7 @@ USCS_Node* USimpleConstructionScript::CreateNode(UActorComponent* NewComponentTe
 	{
 		// Create a node for the script, and save a pointer to the template
 		// NewNamedObject to work around the fact we shouldn't use NewObject for default subobjects
-		USCS_Node* NewNode = NewNamedObject<USCS_Node>(this, NAME_None);
+		auto NewNode = NewObject<USCS_Node>(this);
 		NewNode->SetFlags(RF_Transactional);
 		NewNode->ComponentTemplate = NewComponentTemplate;
 
diff --git a/Engine/Source/Runtime/Foliage/Private/ProceduralFoliage.cpp b/Engine/Source/Runtime/Foliage/Private/ProceduralFoliage.cpp
--- a/Engine/Source/Runtime/Foliage/Private/ProceduralFoliage.cpp
+++ b/Engine/Source/Runtime/Foliage/Private/ProceduralFoliage.cpp
@@ -46,7 +46,7 @@ void UProceduralFoliage::CreateProceduralFoliageInstances()
 	{
 		if( TypeData.Type )
 		{
-			TypeData.TypeInstance = ConstructObject<UFoliageType_InstancedStaticMesh>(TypeData.Type, this);
+			TypeData.TypeInstance = NewObject<UFoliageType_InstancedStaticMesh>(this, TypeData.Type);
 			TypeData.ChangeCount = TypeData.TypeInstance->ChangeCount;
 		}
 		else
