This change forces blueprint components to only consider UActorComponent and child classes to appear when reparenting in class defaults.
#UE-8927 - Can change the parent class of a BP component to a non-component type
#Branch UE4
#Proj UnreadEd

diff --git a/Engine/Source/Editor/UnrealEd/Private/Kismet2/BlueprintEditorUtils.cpp b/Engine/Source/Editor/UnrealEd/Private/Kismet2/BlueprintEditorUtils.cpp
--- a/Engine/Source/Editor/UnrealEd/Private/Kismet2/BlueprintEditorUtils.cpp
+++ b/Engine/Source/Editor/UnrealEd/Private/Kismet2/BlueprintEditorUtils.cpp
@@ -6073,6 +6073,7 @@ TSharedRef<SWidget> FBlueprintEditorUtils::ConstructBlueprintParentClassPicker(
 	bool bIsActor = false;
 	bool bIsAnimBlueprint = false;
 	bool bIsLevelScriptActor = false;
+	bool bIsComponentBlueprint = false;
 	TArray<UClass*> BlueprintClasses;
 	for( auto BlueprintIter = Blueprints.CreateConstIterator(); (!bIsActor && !bIsAnimBlueprint) && BlueprintIter; ++BlueprintIter )
 	{
@@ -6080,6 +6081,7 @@ TSharedRef<SWidget> FBlueprintEditorUtils::ConstructBlueprintParentClassPicker(
 		bIsActor |= Blueprint->ParentClass->IsChildOf( AActor::StaticClass() );
 		bIsAnimBlueprint |= Blueprint->IsA(UAnimBlueprint::StaticClass());
 		bIsLevelScriptActor |= Blueprint->ParentClass->IsChildOf( ALevelScriptActor::StaticClass() );
+		bIsComponentBlueprint |= Blueprint->ParentClass->IsChildOf( UActorComponent::StaticClass() );
 		BlueprintClasses.Add(Blueprint->GeneratedClass);
 	}
 
@@ -6129,6 +6131,10 @@ TSharedRef<SWidget> FBlueprintEditorUtils::ConstructBlueprintParentClassPicker(
 		// If it's an anim blueprint, do not allow conversion to non anim
 		Filter->AllowedChildrenOfClasses.Add( UAnimInstance::StaticClass() );
 	}
+	else if(bIsComponentBlueprint)
+	{
+		Filter->AllowedChildrenOfClasses.Add( UActorComponent::StaticClass() );
+	}
 	else
 	{
 		Filter->DisallowedChildrenOfClasses.Add( AActor::StaticClass() );
