Adding a comment I somehow managed to omit in cl 2429928

diff --git a/Engine/Source/Editor/UnrealEd/Private/Kismet2/BlueprintEditorUtils.cpp b/Engine/Source/Editor/UnrealEd/Private/Kismet2/BlueprintEditorUtils.cpp
--- a/Engine/Source/Editor/UnrealEd/Private/Kismet2/BlueprintEditorUtils.cpp
+++ b/Engine/Source/Editor/UnrealEd/Private/Kismet2/BlueprintEditorUtils.cpp
@@ -6133,6 +6133,7 @@ TSharedRef<SWidget> FBlueprintEditorUtils::ConstructBlueprintParentClassPicker(
 	}
 	else if(bIsComponentBlueprint)
 	{
+		// If it is a component blueprint, only allow classes under and including UActorComponent
 		Filter->AllowedChildrenOfClasses.Add( UActorComponent::StaticClass() );
 	}
 	else
