Change SnapObjectTo so that it does a raycast using the WorldStatic channel instead of the object type. This fixes things like procedural foliage and other volumes from blocking the SnapTo.

UE-8905

diff --git a/Engine/Source/Editor/UnrealEd/Private/EditorServer.cpp b/Engine/Source/Editor/UnrealEd/Private/EditorServer.cpp
--- a/Engine/Source/Editor/UnrealEd/Private/EditorServer.cpp
+++ b/Engine/Source/Editor/UnrealEd/Private/EditorServer.cpp
@@ -4587,7 +4587,7 @@ bool UEditorEngine::SnapObjectTo( FActorOrComponent Object, const bool InAlign,
 		Params.AddIgnoredComponent( Cast<UPrimitiveComponent>(Object.Component) );
 	}
 
-	if ( Object.GetWorld()->SweepSingle(Hit, StartLocation, StartLocation + Direction*WORLD_MAX, FQuat::Identity, FCollisionShape::MakeBox(Extent), Params, FCollisionObjectQueryParams(ECC_WorldStatic)))
+	if (Object.GetWorld()->SweepSingle(Hit, StartLocation, StartLocation + Direction*WORLD_MAX, FQuat::Identity, ECC_WorldStatic, FCollisionShape::MakeBox(Extent), Params))
 	{
 		FVector NewLocation = Hit.Location - LocationOffset;
 		NewLocation.Z += KINDA_SMALL_NUMBER;	// Move the new desired location up by an error tolerance
