fix for accessing a module on a non-game thread
UE-8984

diff --git a/Engine/Source/Editor/UnrealEd/Classes/Editor/EditorEngine.h b/Engine/Source/Editor/UnrealEd/Classes/Editor/EditorEngine.h
--- a/Engine/Source/Editor/UnrealEd/Classes/Editor/EditorEngine.h
+++ b/Engine/Source/Editor/UnrealEd/Classes/Editor/EditorEngine.h
@@ -2728,6 +2728,7 @@ private:
 	FString PlayUsingLauncherDeviceId;
 	FString PlayUsingLauncherDeviceName;
 	bool bPlayUsingLauncherHasCode;
+	bool bPlayUsingLauncherHasCompiler;
 
 	/** List of files we are deferring adding to source control */
 	TArray<FString> DeferredFilesToAddToSourceControl;
diff --git a/Engine/Source/Editor/UnrealEd/Private/PlayLevel.cpp b/Engine/Source/Editor/UnrealEd/Private/PlayLevel.cpp
--- a/Engine/Source/Editor/UnrealEd/Private/PlayLevel.cpp
+++ b/Engine/Source/Editor/UnrealEd/Private/PlayLevel.cpp
@@ -1298,7 +1298,7 @@ void UEditorEngine::HandleStageStarted(const FString& InStage, TWeakPtr<SNotific
 			PlatformName = PlatformName.Left(PlatformName.Find(TEXT("NoEditor")));
 		}
 		Arguments.Add(TEXT("PlatformName"), FText::FromString(PlatformName));
-		if (FRocketSupport::IsRocket() || !bPlayUsingLauncherHasCode || !FSourceCodeNavigation::IsCompilerAvailable())
+		if (FRocketSupport::IsRocket() || !bPlayUsingLauncherHasCode || !bPlayerUsingLauncherHasCompiler)
 		{
 			NotificationText = FText::Format(LOCTEXT("LauncherTaskValidateNotification", "Validating Executable for {PlatformName}..."), Arguments);
 		}
@@ -1447,10 +1447,11 @@ void UEditorEngine::PlayUsingLauncher()
 		// does the project have any code?
 		FGameProjectGenerationModule& GameProjectModule = FModuleManager::LoadModuleChecked<FGameProjectGenerationModule>(TEXT("GameProjectGeneration"));
 		bPlayUsingLauncherHasCode = GameProjectModule.Get().ProjectHasCodeFiles();
+		bPlayUsingLauncherHasCompiler = FSourceCodeNavigation::IsCompilerAvailable();
 
 		// Setup launch profile, keep the setting here to a minimum.
 		ILauncherProfileRef LauncherProfile = LauncherServicesModule.CreateProfile(TEXT("Play On Device"));
-		LauncherProfile->SetBuildGame(bPlayUsingLauncherHasCode && FSourceCodeNavigation::IsCompilerAvailable());
+		LauncherProfile->SetBuildGame(bPlayUsingLauncherHasCode && bPlayUsingLauncherHasCompiler);
 
 		// select the quickest cook mode based on which in editor cook mode is enabled
 		bool bIncrimentalCooking = true;
