fix for header mis-match

diff --git a/Engine/Source/Editor/UnrealEd/Private/PlayLevel.cpp b/Engine/Source/Editor/UnrealEd/Private/PlayLevel.cpp
--- a/Engine/Source/Editor/UnrealEd/Private/PlayLevel.cpp
+++ b/Engine/Source/Editor/UnrealEd/Private/PlayLevel.cpp
@@ -1298,7 +1298,7 @@ void UEditorEngine::HandleStageStarted(const FString& InStage, TWeakPtr<SNotific
 			PlatformName = PlatformName.Left(PlatformName.Find(TEXT("NoEditor")));
 		}
 		Arguments.Add(TEXT("PlatformName"), FText::FromString(PlatformName));
-		if (FRocketSupport::IsRocket() || !bPlayUsingLauncherHasCode || !bPlayerUsingLauncherHasCompiler)
+		if (FRocketSupport::IsRocket() || !bPlayUsingLauncherHasCode || !bPlayUsingLauncherHasCompiler)
 		{
 			NotificationText = FText::Format(LOCTEXT("LauncherTaskValidateNotification", "Validating Executable for {PlatformName}..."), Arguments);
 		}
