SPackagesDialog now responds to changes to source control state.

diff --git a/Engine/Source/Editor/PackagesDialog/Private/PackagesDialog.cpp b/Engine/Source/Editor/PackagesDialog/Private/PackagesDialog.cpp
--- a/Engine/Source/Editor/PackagesDialog/Private/PackagesDialog.cpp
+++ b/Engine/Source/Editor/PackagesDialog/Private/PackagesDialog.cpp
@@ -32,13 +32,14 @@ void FPackagesDialogModule::ShutdownModule()
 	}
 }
 
-void FPackagesDialogModule::CreatePackagesDialog(const FText& Title, const FText& Message, bool InReadOnly, bool InAllowSourceControlConnection)
+void FPackagesDialogModule::CreatePackagesDialog(const FText& Title, const FText& Message, bool InReadOnly, bool InAllowSourceControlConnection, const FSimpleDelegate& InOnSourceControlStateChanged)
 {
 	PackageDialogTitle = Title;
 	PackagesDialogWidget = SNew(SPackagesDialog) 
 		.ReadOnly(InReadOnly)
 		.AllowSourceControlConnection(InAllowSourceControlConnection)
-		.Message(Message);
+		.Message(Message)
+		.OnSourceControlStateChanged(InOnSourceControlStateChanged);
 }
 
 /**
@@ -167,6 +168,14 @@ void FPackagesDialogModule::GetResults(OUT TArray<UPackage*>& OutPackages, EChec
 	}
 }
 
+/**
+ * Removes all package items from the dialog
+ */
+void FPackagesDialogModule::RemoveAllPackageItems()
+{
+	PackagesDialogWidget.Get()->RemoveAll();
+}
+
 /**
  * Adds a new item to the checkbox that represents a package
  *
diff --git a/Engine/Source/Editor/PackagesDialog/Private/SPackagesDialog.cpp b/Engine/Source/Editor/PackagesDialog/Private/SPackagesDialog.cpp
--- a/Engine/Source/Editor/PackagesDialog/Private/SPackagesDialog.cpp
+++ b/Engine/Source/Editor/PackagesDialog/Private/SPackagesDialog.cpp
@@ -71,6 +71,7 @@ void SPackagesDialog::Construct(const FArguments& InArgs)
 	bAllowSourceControlConnection = InArgs._AllowSourceControlConnection.Get();
 	Message = InArgs._Message;
 	Warning = InArgs._Warning;
+	OnSourceControlStateChanged = InArgs._OnSourceControlStateChanged;
 	SortByColumn = SPackagesDialogDefs::ColumnID_FileLabel;
 	SortMode = EColumnSortMode::Ascending;
 
@@ -182,6 +183,14 @@ void SPackagesDialog::Construct(const FArguments& InArgs)
 	];
 }
 
+/**
+ * Removes all checkbox items from the dialog
+ */
+void SPackagesDialog::RemoveAll()
+{
+	Items.Reset();
+}
+
 /**
  * Adds a new checkbox item to the dialog
  *
@@ -589,6 +598,7 @@ EVisibility SPackagesDialog::GetConnectToSourceControlVisibility() const
 FReply SPackagesDialog::OnConnectToSourceControlClicked() const
 {
 	ISourceControlModule::Get().ShowLoginDialog(FSourceControlLoginClosed(), ELoginWindowMode::Modal);
+	OnSourceControlStateChanged.ExecuteIfBound();
 	return FReply::Handled();
 }
 
diff --git a/Engine/Source/Editor/PackagesDialog/Private/SPackagesDialog.h b/Engine/Source/Editor/PackagesDialog/Private/SPackagesDialog.h
--- a/Engine/Source/Editor/PackagesDialog/Private/SPackagesDialog.h
+++ b/Engine/Source/Editor/PackagesDialog/Private/SPackagesDialog.h
@@ -293,6 +293,9 @@ public:
 		/** The warning message of the widget */
 		SLATE_ARGUMENT(FText, Warning)
 
+		/** Called when source control state changes */
+		SLATE_EVENT(FSimpleDelegate, OnSourceControlStateChanged)
+
 	SLATE_END_ARGS()
 	
 	/**
@@ -312,6 +315,11 @@ public:
 	 */
 	TSharedRef<SWidget> GenerateWidgetForItemAndColumn( TSharedPtr<FPackageItem> Item, const FName ColumnID ) const;
 
+	/**
+	 * Removes all checkbox items from the dialog
+	 */
+	void RemoveAll();
+
 	/**
 	 * Adds a new checkbox item to the dialog
 	 *
@@ -520,6 +528,9 @@ private:
 
 	/** Currently selected sorting mode */
 	EColumnSortMode::Type SortMode;
+
+	/** Called when source control state changes */
+	FSimpleDelegate OnSourceControlStateChanged;
 };
 
 /** Widget that represents a row in the PackagesDialog's list view.  Generates widgets for each column on demand. */
diff --git a/Engine/Source/Editor/PackagesDialog/Public/PackagesDialog.h b/Engine/Source/Editor/PackagesDialog/Public/PackagesDialog.h
--- a/Engine/Source/Editor/PackagesDialog/Public/PackagesDialog.h
+++ b/Engine/Source/Editor/PackagesDialog/Public/PackagesDialog.h
@@ -42,8 +42,9 @@ public:
 	 * @param	Message							The message that gets displayed in the package dialog window
 	 * @param	InReadOnly						When true, this dialog only shows a list of packages without the ability to filter
 	 * @param	InAllowSourceControlConnection	When true, this dialog displays a 'connect to source control' button when needed
+	 * @param	InOnSourceControlStateChanged	Delegate called when the source control state changes
 	 */
-	virtual void CreatePackagesDialog(const FText& Title, const FText& Message, bool InReadOnly = false, bool InAllowSourceControlConnection = false);
+	virtual void CreatePackagesDialog(const FText& Title, const FText& Message, bool InReadOnly = false, bool InAllowSourceControlConnection = false, const FSimpleDelegate& InOnSourceControlStateChanged = FSimpleDelegate());
 
 	/**
 	 * Shows the package dialog window as a modal window
@@ -86,6 +87,11 @@ public:
 	 */
 	virtual void GetResults(TArray<UPackage*>& OutPackages, ECheckBoxState InChecked);
 
+	/**
+	 * Removes all package items from the dialog
+	 */
+	virtual void RemoveAllPackageItems();
+
 	/**
 	 * Adds a new item to the checkbox that represents a package
 	 *
diff --git a/Engine/Source/Editor/UnrealEd/Private/FileHelpers.cpp b/Engine/Source/Editor/UnrealEd/Private/FileHelpers.cpp
--- a/Engine/Source/Editor/UnrealEd/Private/FileHelpers.cpp
+++ b/Engine/Source/Editor/UnrealEd/Private/FileHelpers.cpp
@@ -1009,42 +1009,30 @@ static bool IsCheckOutSelectedDisabled()
 	return !(ISourceControlModule::Get().IsEnabled() && ISourceControlModule::Get().GetProvider().IsAvailable());
 }
 
-bool FEditorFileUtils::PromptToCheckoutPackages(bool bCheckDirty, const TArray<UPackage*>& PackagesToCheckOut, TArray<UPackage*>* OutPackagesCheckedOutOrMadeWritable, TArray<UPackage*>* OutPackagesNotNeedingCheckout, const bool bPromptingAfterModify )
+static bool AddCheckoutPackageItems(bool bCheckDirty, TArray<UPackage*> PackagesToCheckOut, TArray<UPackage*>* OutPackagesNotNeedingCheckout, bool* bOutHavePackageToCheckOut)
 {
-	bool bResult = true;
-
 	ISourceControlProvider& SourceControlProvider = ISourceControlModule::Get().GetProvider();
 	if (ISourceControlModule::Get().IsEnabled() && SourceControlProvider.IsAvailable())
 	{
 		// Update the source control status of all potentially relevant packages
 		SourceControlProvider.Execute(ISourceControlOperation::Create<FUpdateStatus>(), PackagesToCheckOut);
 	}
-	
-	// The checkout dialog to show users if any packages need to be checked out
-	const FText DialogTitle = NSLOCTEXT("PackagesDialogModule", "CheckoutPackagesDialogTitle", "Check Out Assets");
-	const FText DialogHeading = NSLOCTEXT("PackagesDialogModule", "CheckoutPackagesDialogMessage", "Select assets to check out.");
 
-	FPackagesDialogModule& CheckoutPackagesDialogModule = FModuleManager::LoadModuleChecked<FPackagesDialogModule>( TEXT("PackagesDialog") );
-	CheckoutPackagesDialogModule.CreatePackagesDialog(DialogTitle, DialogHeading, false, true);
+	FPackagesDialogModule& CheckoutPackagesDialogModule = FModuleManager::LoadModuleChecked<FPackagesDialogModule>(TEXT("PackagesDialog"));
 
-	// Add any of the packages which do not report as editable by source control, yet are currently in the source control depot
-	// If the user has specified to check for dirty packages, only add those which are dirty
 	bool bPackagesAdded = false;
-	
-	// If we found at least one package that can be checked out, this will be true
+	bool bShowWarning = false;
 	bool bHavePackageToCheckOut = false;
 
-	// If we got here and we have one package, it's because someone explicitly saved the asset, therefore remove the package from the ignore list.
-	if(PackagesToCheckOut.Num()==1)
+	if (OutPackagesNotNeedingCheckout)
 	{
-		const FString& PackageName = PackagesToCheckOut[0]->GetName();
-		PackagesNotSavedDuringSaveAll.Remove(PackageName);
+		OutPackagesNotNeedingCheckout->Reset();
 	}
 
-	bool bShowWarning = false;
+	CheckoutPackagesDialogModule.RemoveAllPackageItems();
 
 	// Iterate through all the packages and add them to the dialog if necessary.
-	for ( TArray<UPackage*>::TConstIterator PackageIter( PackagesToCheckOut ); PackageIter; ++PackageIter )
+	for (TArray<UPackage*>::TConstIterator PackageIter(PackagesToCheckOut); PackageIter; ++PackageIter)
 	{
 		UPackage* CurPackage = *PackageIter;
 		FString Filename;
@@ -1052,28 +1040,28 @@ bool FEditorFileUtils::PromptToCheckoutPackages(bool bCheckDirty, const TArray<U
 		bool bPkgReadOnly = true;
 		bool bCareAboutReadOnly = SourceControlProvider.UsesLocalReadOnlyState();
 		// Find the filename for this package
-		bool bFoundFile = FPackageName::DoesPackageExist( CurPackage->GetName(), NULL, &Filename );
-		if( bFoundFile )
+		bool bFoundFile = FPackageName::DoesPackageExist(CurPackage->GetName(), NULL, &Filename);
+		if (bFoundFile)
 		{
 			// determine if the package file is read only
-			bPkgReadOnly = IFileManager::Get().IsReadOnly( *Filename );
+			bPkgReadOnly = IFileManager::Get().IsReadOnly(*Filename);
 		}
-	
+
 		FSourceControlStatePtr SourceControlState = SourceControlProvider.GetState(CurPackage, EStateCacheUsage::Use);
 
 		// Package does not need to be checked out if its already checked out or we are ignoring it for source control
 		bool bSCCCanEdit = !SourceControlState.IsValid() || SourceControlState->CanCheckIn() || SourceControlState->IsIgnored() || SourceControlState->IsUnknown();
 		bool bIsSourceControlled = SourceControlState.IsValid() && SourceControlState->IsSourceControlled();
 
-		if ( !bSCCCanEdit && (bIsSourceControlled && ( !bCheckDirty || ( bCheckDirty && CurPackage->IsDirty() ) ) ) && !SourceControlState->IsCheckedOut() )
+		if (!bSCCCanEdit && (bIsSourceControlled && (!bCheckDirty || (bCheckDirty && CurPackage->IsDirty()))) && !SourceControlState->IsCheckedOut())
 		{
-			if( SourceControlState.IsValid() && !SourceControlState->IsCurrent() )
-			{				
+			if (SourceControlState.IsValid() && !SourceControlState->IsCurrent())
+			{
 				// This package is not at the head revision and it should be ghosted as a result
 				CheckoutPackagesDialogModule.AddPackageItem(CurPackage, CurPackage->GetName(), ECheckBoxState::Unchecked, true, TEXT("SavePackages.SCC_DlgNotCurrent"), SourceControlState->GetDisplayTooltip().ToString());
 				bShowWarning = true;
 			}
-			else if( SourceControlState.IsValid() && SourceControlState->IsCheckedOutOther() )
+			else if (SourceControlState.IsValid() && SourceControlState->IsCheckedOutOther())
 			{
 				// This package is checked out by someone else so it should be ghosted
 				CheckoutPackagesDialogModule.AddPackageItem(CurPackage, CurPackage->GetName(), ECheckBoxState::Unchecked, true, TEXT("SavePackages.SCC_DlgCheckedOutOther"), SourceControlState->GetDisplayTooltip().ToString());
@@ -1090,7 +1078,7 @@ bool FEditorFileUtils::PromptToCheckoutPackages(bool bCheckDirty, const TArray<U
 			}
 			bPackagesAdded = true;
 		}
-		else if ( bPkgReadOnly && bFoundFile && (IsCheckOutSelectedDisabled() || !bCareAboutReadOnly))
+		else if (bPkgReadOnly && bFoundFile && (IsCheckOutSelectedDisabled() || !bCareAboutReadOnly))
 		{
 			const FText Tooltip = SourceControlState.IsValid() ? SourceControlState->GetDisplayTooltip() : NSLOCTEXT("PackagesDialogModule", "Dlg_NotCheckedOutTip", "Not checked out");
 
@@ -1102,23 +1090,80 @@ bool FEditorFileUtils::PromptToCheckoutPackages(bool bCheckDirty, const TArray<U
 			CheckoutPackagesDialogModule.AddPackageItem(CurPackage, CurPackage->GetName(), ECheckBoxState::Unchecked, bIsDisabled, TEXT("SavePackages.SCC_DlgReadOnly"), Tooltip.ToString());
 			bPackagesAdded = true;
 		}
-		else if ( OutPackagesNotNeedingCheckout )
+		else if (OutPackagesNotNeedingCheckout)
 		{
 			// The current package does not need to be checked out in order to save.
-			OutPackagesNotNeedingCheckout->Add( CurPackage );
+			OutPackagesNotNeedingCheckout->Add(CurPackage);
 		}
-
 	}
 
-	// If any packages were added to the dialog, show the dialog to the user and allow them to select which files to check out
-	if ( bPackagesAdded )
+	if (bPackagesAdded)
 	{
 		if (bShowWarning)
 		{
 			CheckoutPackagesDialogModule.SetWarning(
 				NSLOCTEXT("PackagesDialogModule", "CheckoutPackagesWarnMessage", "Warning: There are modified assets which you will not be able to check out as they are locked or not at the head revision. You may lose your changes if you continue, as you will be unable to submit them to source control."));
 		}
+		else
+		{
+			CheckoutPackagesDialogModule.SetWarning(FText::GetEmpty());
+		}
+	}
+
+	if (bOutHavePackageToCheckOut)
+	{
+		*bOutHavePackageToCheckOut = bHavePackageToCheckOut;
+	}
+
+	return bPackagesAdded;
+}
 
+static void UpdateCheckoutPackageItems(bool bCheckDirty, TArray<UPackage*> PackagesToCheckOut, TArray<UPackage*>* OutPackagesNotNeedingCheckout)
+{
+	AddCheckoutPackageItems(bCheckDirty, PackagesToCheckOut, OutPackagesNotNeedingCheckout, nullptr);
+}
+
+bool FEditorFileUtils::PromptToCheckoutPackages(bool bCheckDirty, const TArray<UPackage*>& PackagesToCheckOut, TArray<UPackage*>* OutPackagesCheckedOutOrMadeWritable, TArray<UPackage*>* OutPackagesNotNeedingCheckout, const bool bPromptingAfterModify )
+{
+	bool bResult = true;
+
+	ISourceControlProvider& SourceControlProvider = ISourceControlModule::Get().GetProvider();
+	
+	// The checkout dialog to show users if any packages need to be checked out
+	const FText DialogTitle = NSLOCTEXT("PackagesDialogModule", "CheckoutPackagesDialogTitle", "Check Out Assets");
+	const FText DialogHeading = NSLOCTEXT("PackagesDialogModule", "CheckoutPackagesDialogMessage", "Select assets to check out.");
+
+	FPackagesDialogModule& CheckoutPackagesDialogModule = FModuleManager::LoadModuleChecked<FPackagesDialogModule>( TEXT("PackagesDialog") );
+
+	// Add any of the packages which do not report as editable by source control, yet are currently in the source control depot
+	// If the user has specified to check for dirty packages, only add those which are dirty
+	bool bPackagesAdded = false;
+
+	// If we found at least one package that can be checked out, this will be true
+	bool bHavePackageToCheckOut = false;
+
+	const bool bReadOnly = false;
+	const bool bAllowSourceControlConnection = true;
+	CheckoutPackagesDialogModule.CreatePackagesDialog(
+		DialogTitle,
+		DialogHeading,
+		bReadOnly,
+		bAllowSourceControlConnection,
+		FSimpleDelegate::CreateStatic(&UpdateCheckoutPackageItems, bCheckDirty, PackagesToCheckOut, OutPackagesNotNeedingCheckout)
+		);
+
+	// If we got here and we have one package, it's because someone explicitly saved the asset, therefore remove the package from the ignore list.
+	if(PackagesToCheckOut.Num()==1)
+	{
+		const FString& PackageName = PackagesToCheckOut[0]->GetName();
+		PackagesNotSavedDuringSaveAll.Remove(PackageName);
+	}
+
+	bPackagesAdded = AddCheckoutPackageItems(bCheckDirty, PackagesToCheckOut, OutPackagesNotNeedingCheckout, &bHavePackageToCheckOut);
+
+	// If any packages were added to the dialog, show the dialog to the user and allow them to select which files to check out
+	if ( bPackagesAdded )
+	{
 		TAttribute<bool> CheckOutSelectedDisabledAttrib;
 		if( !bHavePackageToCheckOut && !IsCheckOutSelectedDisabled() )
 		{
