Fixed issue with Static Mesh LoDs losing custom set materials / metadata when reimported.

diff --git a/Engine/Source/Editor/UnrealEd/Private/Fbx/FbxStaticMeshImport.cpp b/Engine/Source/Editor/UnrealEd/Private/Fbx/FbxStaticMeshImport.cpp
--- a/Engine/Source/Editor/UnrealEd/Private/Fbx/FbxStaticMeshImport.cpp
+++ b/Engine/Source/Editor/UnrealEd/Private/Fbx/FbxStaticMeshImport.cpp
@@ -777,6 +777,8 @@ UStaticMesh* UnFbx::FFbxImporter::ReimportStaticMesh(UStaticMesh* Mesh, UFbxStat
 		}
 	}
 	
+	struct ExistingStaticMeshData* ExistMeshDataPtr = SaveExistingStaticMeshData(Mesh);
+
 	if (Node)
 	{
 		FbxNode* NodeParent = Node->GetParent();
@@ -814,6 +816,8 @@ UStaticMesh* UnFbx::FFbxImporter::ReimportStaticMesh(UStaticMesh* Mesh, UFbxStat
 			AddTokenizedErrorMessage(FTokenizedMessage::Create(EMessageSeverity::Warning, FText::Format(LOCTEXT("Error_NoFBXMeshFound", "No FBX mesh found when reimport Unreal mesh '{0}'. The FBX file is crashed."), FText::FromString(Mesh->GetName()))), FFbxErrors::Generic_Mesh_MeshNotFound);
 		}
 	}
+
+	RestoreExistingMeshData(ExistMeshDataPtr, NewMesh);
 	return NewMesh;
 }
 
@@ -834,7 +838,6 @@ void UnFbx::FFbxImporter::VerifyGeometry(UStaticMesh* StaticMesh)
 UStaticMesh* UnFbx::FFbxImporter::ImportStaticMeshAsSingle(UObject* InParent, TArray<FbxNode*>& MeshNodeArray, const FName InName, EObjectFlags Flags, UFbxStaticMeshImportData* TemplateImportData, UStaticMesh* InStaticMesh, int LODIndex)
 {
 	bool bBuildStatus = true;
-	struct ExistingStaticMeshData* ExistMeshDataPtr = NULL;
 
 	// Make sure rendering is done - so we are not changing data being used by collision drawing.
 	FlushRenderingCommands();
@@ -912,7 +915,6 @@ UStaticMesh* UnFbx::FFbxImporter::ImportStaticMeshAsSingle(UObject* InParent, TA
 
 		// Free any RHI resources for existing mesh before we re-create in place.
 		ExistingMesh->PreEditChange(NULL);
-		ExistMeshDataPtr = SaveExistingStaticMeshData(ExistingMesh);
 	}
 	else if (ExistingObject)
 	{
@@ -1200,11 +1202,6 @@ UStaticMesh* UnFbx::FFbxImporter::ImportStaticMeshAsSingle(UObject* InParent, TA
 			StaticMesh->LightMapResolution = LODGroup.GetDefaultLightMapResolution();
 		}
 
-		if (ExistMeshDataPtr)
-		{
-			RestoreExistingMeshData(ExistMeshDataPtr, StaticMesh);
-		}
-
 		UFbxStaticMeshImportData* ImportData = UFbxStaticMeshImportData::GetImportDataForStaticMesh(StaticMesh, TemplateImportData);
 		ImportData->SourceFilePath = FReimportManager::SanitizeImportFilename(UFactory::CurrentFilename, StaticMesh);
 		ImportData->SourceFileTimestamp = IFileManager::Get().GetTimeStamp(*UFactory::CurrentFilename).ToString();
diff --git a/Engine/Source/Editor/UnrealEd/Private/StaticMeshEdit.cpp b/Engine/Source/Editor/UnrealEd/Private/StaticMeshEdit.cpp
--- a/Engine/Source/Editor/UnrealEd/Private/StaticMeshEdit.cpp
+++ b/Engine/Source/Editor/UnrealEd/Private/StaticMeshEdit.cpp
@@ -820,11 +820,13 @@ struct ExistingStaticMeshData
 	FMeshSectionInfoMap			ExistingSectionInfoMap;
 	TArray<FMeshBuildSettings>	ExistingBuildSettings;
 	TArray<FMeshReductionSettings> ExistingReductionSettings;
+	TArray<float>				ExistingScreenSize;
 	
 	TArray<UStaticMeshSocket*>	ExistingSockets;
 
 	bool						ExistingUseMaximumStreamingTexelRatio;
 	bool						ExistingCustomizedCollision;
+	bool						bAutoComputeLODScreenSize;
 
 	int32							ExistingLightMapResolution;
 	int32							ExistingLightMapCoordinateIndex;
@@ -857,12 +859,14 @@ ExistingStaticMeshData* SaveExistingStaticMeshData(UStaticMesh* ExistingMesh)
 		{
 			ExistingMeshDataPtr->ExistingBuildSettings.Add(ExistingMesh->SourceModels[i].BuildSettings);
 			ExistingMeshDataPtr->ExistingReductionSettings.Add(ExistingMesh->SourceModels[i].ReductionSettings);
+			ExistingMeshDataPtr->ExistingScreenSize.Add(ExistingMesh->SourceModels[i].ScreenSize);
 		}
 
 		ExistingMeshDataPtr->ExistingSockets = ExistingMesh->Sockets;
 
 		ExistingMeshDataPtr->ExistingUseMaximumStreamingTexelRatio = ExistingMesh->bUseMaximumStreamingTexelRatio;
 		ExistingMeshDataPtr->ExistingCustomizedCollision = ExistingMesh->bCustomizedCollision;
+		ExistingMeshDataPtr->bAutoComputeLODScreenSize = ExistingMesh->bAutoComputeLODScreenSize;
 
 		ExistingMeshDataPtr->ExistingLightMapResolution = ExistingMesh->LightMapResolution;
 		ExistingMeshDataPtr->ExistingLightMapCoordinateIndex = ExistingMesh->LightMapCoordinateIndex;
@@ -895,12 +899,14 @@ void RestoreExistingMeshData(struct ExistingStaticMeshData* ExistingMeshDataPtr,
 		{
 			NewMesh->SourceModels[i].BuildSettings = ExistingMeshDataPtr->ExistingBuildSettings[i];
 			NewMesh->SourceModels[i].ReductionSettings = ExistingMeshDataPtr->ExistingReductionSettings[i];
+			NewMesh->SourceModels[i].ScreenSize = ExistingMeshDataPtr->ExistingScreenSize[i];
 		}
 		for(int32 i=NumCommonLODs; i < ExistingMeshDataPtr->ExistingBuildSettings.Num(); ++i)
 		{
 			FStaticMeshSourceModel* SrcModel = new(NewMesh->SourceModels) FStaticMeshSourceModel();
 			SrcModel->BuildSettings = ExistingMeshDataPtr->ExistingBuildSettings[i];
 			SrcModel->ReductionSettings = ExistingMeshDataPtr->ExistingReductionSettings[i];
+			SrcModel->ScreenSize = ExistingMeshDataPtr->ExistingScreenSize[i];
 		}
 
 		// Assign sockets from old version of this StaticMesh.
@@ -911,6 +917,7 @@ void RestoreExistingMeshData(struct ExistingStaticMeshData* ExistingMeshDataPtr,
 
 		NewMesh->bUseMaximumStreamingTexelRatio = ExistingMeshDataPtr->ExistingUseMaximumStreamingTexelRatio;
 		NewMesh->bCustomizedCollision = ExistingMeshDataPtr->ExistingCustomizedCollision;
+		NewMesh->bAutoComputeLODScreenSize = ExistingMeshDataPtr->bAutoComputeLODScreenSize;
 
 		NewMesh->LightMapResolution = ExistingMeshDataPtr->ExistingLightMapResolution;
 		NewMesh->LightMapCoordinateIndex = ExistingMeshDataPtr->ExistingLightMapCoordinateIndex;
@@ -950,6 +957,8 @@ void RestoreExistingMeshData(struct ExistingStaticMeshData* ExistingMeshDataPtr,
 			}
 		}
 	}
+
+	delete ExistingMeshDataPtr;
 }
 
 #undef LOCTEXT_NAMESPACE
