Small change to preserve component names when adding components to blueprints.
#UE-8844 - In-World actor workflow - Component names are not kept when converting to BP
#Branch UE4
#Proj UnrealEd

diff --git a/Engine/Source/Editor/UnrealEd/Private/Kismet2/Kismet2.cpp b/Engine/Source/Editor/UnrealEd/Private/Kismet2/Kismet2.cpp
--- a/Engine/Source/Editor/UnrealEd/Private/Kismet2/Kismet2.cpp
+++ b/Engine/Source/Editor/UnrealEd/Private/Kismet2/Kismet2.cpp
@@ -938,7 +938,7 @@ void FKismetEditorUtilities::AddComponentsToBlueprint(UBlueprint* Blueprint, con
 				check(Actor);
 			}
 
-			USCS_Node* SCSNode = SCS->CreateNode(ActorComponent->GetClass());
+			USCS_Node* SCSNode = SCS->CreateNode(ActorComponent->GetClass(), ActorComponent->GetFName());
 			UEditorEngine::CopyPropertiesForUnrelatedObjects(ActorComponent,SCSNode->ComponentTemplate);
 
 			// Clear the instance component flag
