Fix FBX actor export not exporting vertex colors for static meshes

diff --git a/Engine/Source/Editor/UnrealEd/Private/Fbx/FbxMainExport.cpp b/Engine/Source/Editor/UnrealEd/Private/Fbx/FbxMainExport.cpp
--- a/Engine/Source/Editor/UnrealEd/Private/Fbx/FbxMainExport.cpp
+++ b/Engine/Source/Editor/UnrealEd/Private/Fbx/FbxMainExport.cpp
@@ -615,7 +615,7 @@ void FFbxExporter::ExportStaticMesh( AActor* Actor, UStaticMeshComponent* Static
 	{
 		return;
 	}
-	int32 LODIndex = StaticMeshComponent->ForcedLodModel;
+	int32 LODIndex = StaticMeshComponent->ForcedLodModel-1;
 	FStaticMeshLODResources& RenderMesh = StaticMesh->GetLODForExport(LODIndex);
 
 	FString FbxNodeName = GetActorNodeName(Actor, InMatineeActor);
@@ -628,7 +628,7 @@ void FFbxExporter::ExportStaticMesh( AActor* Actor, UStaticMeshComponent* Static
 	}
 
 	FbxNode* FbxActor = ExportActor( Actor, InMatineeActor );
-	ExportStaticMeshToFbx(StaticMesh, RenderMesh, *FbxNodeName, FbxActor);
+	ExportStaticMeshToFbx(StaticMesh, RenderMesh, *FbxNodeName, FbxActor, -1, ColorBuffer);
 }
 
 struct FBSPExportData
