Fix for unity build - duplicated stat definitions.

diff --git a/Engine/Source/Editor/UnrealEd/Private/Kismet2/BlueprintEditorUtils.cpp b/Engine/Source/Editor/UnrealEd/Private/Kismet2/BlueprintEditorUtils.cpp
--- a/Engine/Source/Editor/UnrealEd/Private/Kismet2/BlueprintEditorUtils.cpp
+++ b/Engine/Source/Editor/UnrealEd/Private/Kismet2/BlueprintEditorUtils.cpp
@@ -52,7 +52,7 @@
 
 DEFINE_LOG_CATEGORY(LogBlueprintDebug);
 
-DECLARE_CYCLE_STAT(TEXT("Notify Blueprint Changed"), EKismetCompilerStats_NotifyBlueprintChanged, STATGROUP_KismetCompiler);
+DEFINE_STAT(EKismetCompilerStats_NotifyBlueprintChanged);
 DECLARE_CYCLE_STAT(TEXT("Mark Blueprint as Structurally Modified"), EKismetCompilerStats_MarkBlueprintasStructurallyModified, STATGROUP_KismetCompiler);
 DECLARE_CYCLE_STAT(TEXT("Refresh External DependencyNodes"), EKismetCompilerStats_RefreshExternalDependencyNodes, STATGROUP_KismetCompiler);
 
diff --git a/Engine/Source/Editor/UnrealEd/Private/Kismet2/Kismet2.cpp b/Engine/Source/Editor/UnrealEd/Private/Kismet2/Kismet2.cpp
--- a/Engine/Source/Editor/UnrealEd/Private/Kismet2/Kismet2.cpp
+++ b/Engine/Source/Editor/UnrealEd/Private/Kismet2/Kismet2.cpp
@@ -42,7 +42,6 @@ DECLARE_CYCLE_STAT(TEXT("Compile Blueprint"), EKismetCompilerStats_CompileBluepr
 DECLARE_CYCLE_STAT(TEXT("Broadcast Precompile"), EKismetCompilerStats_BroadcastPrecompile, STATGROUP_KismetCompiler);
 DECLARE_CYCLE_STAT(TEXT("Update Search Metadata"), EKismetCompilerStats_UpdateSearchMetaData, STATGROUP_KismetCompiler);
 DECLARE_CYCLE_STAT(TEXT("Garbage Collection"), EKismetCompilerStats_GarbageCollection, STATGROUP_KismetCompiler);
-DECLARE_CYCLE_STAT(TEXT("Notify Blueprint Changed"), EKismetCompilerStats_NotifyBlueprintChanged, STATGROUP_KismetCompiler);
 DECLARE_CYCLE_STAT(TEXT("Refresh Dependent Blueprints"), EKismetCompilerStats_RefreshDependentBlueprints, STATGROUP_KismetCompiler);
 DECLARE_CYCLE_STAT(TEXT("Validate Generated Class"), EKismetCompilerStats_ValidateGeneratedClass, STATGROUP_KismetCompiler);
 
diff --git a/Engine/Source/Editor/UnrealEd/Public/Kismet2/BlueprintEditorUtils.h b/Engine/Source/Editor/UnrealEd/Public/Kismet2/BlueprintEditorUtils.h
--- a/Engine/Source/Editor/UnrealEd/Public/Kismet2/BlueprintEditorUtils.h
+++ b/Engine/Source/Editor/UnrealEd/Public/Kismet2/BlueprintEditorUtils.h
@@ -40,6 +40,8 @@ struct UNREALED_API FFunctionFromNodeHelper
 	FFunctionFromNodeHelper(UObject* Obj);
 };
 
+DECLARE_CYCLE_STAT_EXTERN(TEXT("Notify Blueprint Changed"), EKismetCompilerStats_NotifyBlueprintChanged, STATGROUP_KismetCompiler, );
+
 class UNREALED_API FBlueprintEditorUtils
 {
 public:
