#LEVEL:EXPORT: Fixed so that it exports skeletal mesh as well as ChildActorComponent

diff --git a/Engine/Source/Editor/UnrealEd/Private/Fbx/FbxAnimationExport.cpp b/Engine/Source/Editor/UnrealEd/Private/Fbx/FbxAnimationExport.cpp
--- a/Engine/Source/Editor/UnrealEd/Private/Fbx/FbxAnimationExport.cpp
+++ b/Engine/Source/Editor/UnrealEd/Private/Fbx/FbxAnimationExport.cpp
@@ -217,7 +217,7 @@ void FFbxExporter::CorrectAnimTrackInterpolation( TArray<FbxNode*>& BoneNodes, F
 }
 
 
-void FFbxExporter::ExportAnimSequence( const UAnimSequence* AnimSeq, USkeletalMesh* SkelMesh, bool bExportSkelMesh )
+void FFbxExporter::ExportAnimSequence( const UAnimSequence* AnimSeq, const USkeletalMesh* SkelMesh, bool bExportSkelMesh, const TCHAR* MeshName, FbxNode* ActorRootNode )
 {
 	if( Scene == NULL || AnimSeq == NULL || SkelMesh == NULL )
 	{
@@ -225,7 +225,7 @@ void FFbxExporter::ExportAnimSequence( const UAnimSequence* AnimSeq, USkeletalMe
 	}
 
 
-	FbxNode* RootNode = Scene->GetRootNode();
+	FbxNode* RootNode = (ActorRootNode)? ActorRootNode : Scene->GetRootNode();
 	// Create the Skeleton
 	TArray<FbxNode*> BoneNodes;
 	FbxNode* SkeletonRootNode = CreateSkeleton(*SkelMesh, BoneNodes);
@@ -250,11 +250,19 @@ void FFbxExporter::ExportAnimSequence( const UAnimSequence* AnimSeq, USkeletalMe
 	// Optionally export the mesh
 	if(bExportSkelMesh)
 	{
-		FString MeshName;
-		SkelMesh->GetName(MeshName);
+		FString MeshNodeName;
+		
+		if (MeshName)
+		{
+			MeshNodeName = MeshName;
+		}
+		else
+		{
+			SkelMesh->GetName(MeshNodeName);
+		}
 
 		// Add the mesh
-		FbxNode* MeshRootNode = CreateMesh(*SkelMesh, *MeshName);
+		FbxNode* MeshRootNode = CreateMesh(*SkelMesh, *MeshNodeName);
 		if(MeshRootNode)
 		{
 			RootNode->AddChild(MeshRootNode);
diff --git a/Engine/Source/Editor/UnrealEd/Private/Fbx/FbxMainExport.cpp b/Engine/Source/Editor/UnrealEd/Private/Fbx/FbxMainExport.cpp
--- a/Engine/Source/Editor/UnrealEd/Private/Fbx/FbxMainExport.cpp
+++ b/Engine/Source/Editor/UnrealEd/Private/Fbx/FbxMainExport.cpp
@@ -305,7 +305,7 @@ void FFbxExporter::ExportLevelMesh( ULevel* InLevel, AMatineeActor* InMatineeAct
 			{
 				ExportActor( Actor, InMatineeActor ); // Just export the placement of the particle emitter.
 			}
-			else if( Actor->GetClass()->ClassGeneratedBy != NULL )
+			else if( Actor != NULL )
 			{
 				// Export blueprint actors and all their components
 				ExportActor( Actor, InMatineeActor, true );
@@ -824,7 +824,7 @@ void FFbxExporter::ExportSkeletalMesh( USkeletalMesh* SkeletalMesh )
 	FbxNode* MeshNode = FbxNode::Create(Scene, TCHAR_TO_ANSI(*MeshName));
 	Scene->GetRootNode()->AddChild(MeshNode);
 
-	ExportSkeletalMeshToFbx(*SkeletalMesh, *MeshName, MeshNode);
+	ExportSkeletalMeshToFbx(*SkeletalMesh, NULL, *MeshName, MeshNode);
 }
 
 void FFbxExporter::ExportSkeletalMesh( AActor* Actor, USkeletalMeshComponent* SkeletalMeshComponent )
@@ -837,7 +837,7 @@ void FFbxExporter::ExportSkeletalMesh( AActor* Actor, USkeletalMeshComponent* Sk
 	FString FbxNodeName = GetActorNodeName(Actor, NULL);
 
 	FbxNode* FbxActorNode = ExportActor( Actor, NULL );
-	ExportSkeletalMeshToFbx(*SkeletalMesh, *FbxNodeName, FbxActorNode);
+	ExportSkeletalMeshToFbx(*SkeletalMesh, NULL, *FbxNodeName, FbxActorNode);
 }
 
 FbxSurfaceMaterial* FFbxExporter::CreateDefaultMaterial()
@@ -1093,16 +1093,17 @@ FbxNode* FFbxExporter::ExportActor(AActor* Actor, AMatineeActor* InMatineeActor,
 
 		if( bExportComponents )
 		{
-			TInlineComponentArray<UMeshComponent*> MeshComponents;
-			Actor->GetComponents(MeshComponents);
+			TInlineComponentArray<UActorComponent*> ActorComponents;
+			Actor->GetComponents(ActorComponents);
 
 			TInlineComponentArray<UActorComponent*> ComponentsToExport;
-			for( int32 ComponentIndex = 0; ComponentIndex < MeshComponents.Num(); ++ComponentIndex )
+			for( int32 ComponentIndex = 0; ComponentIndex < ActorComponents.Num(); ++ComponentIndex )
 			{
-				UMeshComponent* Component = MeshComponents[ComponentIndex];
+				UActorComponent* Component = ActorComponents[ComponentIndex];
 
 				UStaticMeshComponent* StaticMeshComp = Cast<UStaticMeshComponent>( Component );
 				USkeletalMeshComponent* SkelMeshComp = Cast<USkeletalMeshComponent>( Component );
+				UChildActorComponent* ChildActorComp = Cast<UChildActorComponent>( Component );
 
 				if( StaticMeshComp && StaticMeshComp->StaticMesh )
 				{
@@ -1112,6 +1113,10 @@ FbxNode* FFbxExporter::ExportActor(AActor* Actor, AMatineeActor* InMatineeActor,
 				{
 					ComponentsToExport.Add( SkelMeshComp );
 				}
+				else if (ChildActorComp && ChildActorComp->ChildActor)
+				{
+					ComponentsToExport.Add(ChildActorComp);
+				}
 			}
 
 
@@ -1144,6 +1149,7 @@ FbxNode* FFbxExporter::ExportActor(AActor* Actor, AMatineeActor* InMatineeActor,
 
 				UStaticMeshComponent* StaticMeshComp = Cast<UStaticMeshComponent>( Component );
 				USkeletalMeshComponent* SkelMeshComp = Cast<USkeletalMeshComponent>( Component );
+				UChildActorComponent* ChildActorComp = Cast<UChildActorComponent>( Component );
 
 				if (StaticMeshComp && StaticMeshComp->StaticMesh)
 				{
@@ -1161,7 +1167,13 @@ FbxNode* FFbxExporter::ExportActor(AActor* Actor, AMatineeActor* InMatineeActor,
 				}
 				else if (SkelMeshComp && SkelMeshComp->SkeletalMesh)
 				{
-					ExportSkeletalMeshToFbx( *SkelMeshComp->SkeletalMesh, *SkelMeshComp->GetName(), ExportNode );
+					UAnimSequence* AnimSeq = (SkelMeshComp->GetAnimationMode() == EAnimationMode::AnimationSingleNode)? Cast<UAnimSequence>(SkelMeshComp->AnimationData.AnimToPlay) : NULL;
+					ExportSkeletalMeshToFbx( *SkelMeshComp->SkeletalMesh, AnimSeq, *SkelMeshComp->GetName(), ExportNode);
+				}
+				else if (ChildActorComp && ChildActorComp->ChildActor)
+				{
+					FbxNode* ChildActorNode = ExportActor(ChildActorComp->ChildActor, InMatineeActor, true);
+					FbxActors.Add(ChildActorComp->ChildActor, ChildActorNode);
 				}
 			}
 		}
diff --git a/Engine/Source/Editor/UnrealEd/Private/Fbx/FbxSkeletalMeshExport.cpp b/Engine/Source/Editor/UnrealEd/Private/Fbx/FbxSkeletalMeshExport.cpp
--- a/Engine/Source/Editor/UnrealEd/Private/Fbx/FbxSkeletalMeshExport.cpp
+++ b/Engine/Source/Editor/UnrealEd/Private/Fbx/FbxSkeletalMeshExport.cpp
@@ -478,31 +478,38 @@ void FFbxExporter::CreateBindPose(FbxNode* MeshRootNode)
 /**
  * Add the given skeletal mesh to the Fbx scene in preparation for exporting.  Makes all new nodes a child of the given node
  */
-void FFbxExporter::ExportSkeletalMeshToFbx(const USkeletalMesh& SkelMesh, const TCHAR* MeshName, FbxNode* ActorRootNode)
+void FFbxExporter::ExportSkeletalMeshToFbx(const USkeletalMesh& SkelMesh, const UAnimSequence* AnimSeq, const TCHAR* MeshName, FbxNode* ActorRootNode)
 {
-	TArray<FbxNode*> BoneNodes;
-
-	// Add the skeleton to the scene
-	FbxNode* SkeletonRootNode = CreateSkeleton(SkelMesh, BoneNodes);
-	if(SkeletonRootNode)
+	if(AnimSeq)
 	{
-		ActorRootNode->AddChild(SkeletonRootNode);
+		ExportAnimSequence(AnimSeq, &SkelMesh, true, MeshName, ActorRootNode);
 	}
-
-	// Add the mesh
-	FbxNode* MeshRootNode = CreateMesh(SkelMesh, MeshName);
-	if(MeshRootNode)
+	else
 	{
-		ActorRootNode->AddChild(MeshRootNode);
-	}
+		TArray<FbxNode*> BoneNodes;
 
-	if(SkeletonRootNode && MeshRootNode)
-	{
-		// Bind the mesh to the skeleton
-		BindMeshToSkeleton(SkelMesh, MeshRootNode, BoneNodes);
+		// Add the skeleton to the scene
+		FbxNode* SkeletonRootNode = CreateSkeleton(SkelMesh, BoneNodes);
+		if(SkeletonRootNode)
+		{
+			ActorRootNode->AddChild(SkeletonRootNode);
+		}
 
-		// Add the bind pose
-		CreateBindPose(MeshRootNode);
+		// Add the mesh
+		FbxNode* MeshRootNode = CreateMesh(SkelMesh, MeshName);
+		if(MeshRootNode)
+		{
+			ActorRootNode->AddChild(MeshRootNode);
+		}
+
+		if(SkeletonRootNode && MeshRootNode)
+		{
+			// Bind the mesh to the skeleton
+			BindMeshToSkeleton(SkelMesh, MeshRootNode, BoneNodes);
+
+			// Add the bind pose
+			CreateBindPose(MeshRootNode);
+		}
 	}
 }
 
diff --git a/Engine/Source/Editor/UnrealEd/Private/FbxExporter.h b/Engine/Source/Editor/UnrealEd/Private/FbxExporter.h
--- a/Engine/Source/Editor/UnrealEd/Private/FbxExporter.h
+++ b/Engine/Source/Editor/UnrealEd/Private/FbxExporter.h
@@ -152,7 +152,7 @@ public:
 	/**
 	 * Exports a single UAnimSequence, and optionally a skeletal mesh
 	 */
-	void ExportAnimSequence( const UAnimSequence* AnimSeq, USkeletalMesh* SkelMesh, bool bExportSkelMesh );
+	void ExportAnimSequence( const UAnimSequence* AnimSeq, const USkeletalMesh* SkelMesh, bool bExportSkelMesh, const TCHAR* MeshNames=NULL, FbxNode* ActorRootNode=NULL);
 
 	/**
 	 * Exports the list of UAnimSequences as a single animation based on the settings in the TrackKeys
@@ -243,7 +243,7 @@ private:
 	/**
 	 * Add the given skeletal mesh to the Fbx scene in preparation for exporting.  Makes all new nodes a child of the given node
 	 */
-	void ExportSkeletalMeshToFbx(const USkeletalMesh& SkelMesh, const TCHAR* MeshName, FbxNode* FbxActor);
+	void ExportSkeletalMeshToFbx(const USkeletalMesh& SkelMesh, const UAnimSequence* AnimSeq, const TCHAR* MeshName, FbxNode* FbxActor);
 
 	/**
 	 * Add the given animation sequence as rotation and translation tracks to the given list of bone nodes
