[INTEGRATE] Change 2425673 by [REDACTED_PERSON][REDACTED_HANDLE] on 2015/01/30 09:38:32

	Editor - The source uility function link now allows formatted links to be generated.

diff --git a/Engine/Source/Editor/UnrealEd/Private/EditorClassUtils.cpp b/Engine/Source/Editor/UnrealEd/Private/EditorClassUtils.cpp
--- a/Engine/Source/Editor/UnrealEd/Private/EditorClassUtils.cpp
+++ b/Engine/Source/Editor/UnrealEd/Private/EditorClassUtils.cpp
@@ -73,6 +73,14 @@ TSharedRef<SWidget> FEditorClassUtils::GetDocumentationLinkWidget(const UClass*
 }
 
 TSharedRef<SWidget> FEditorClassUtils::GetSourceLink(const UClass* Class, const TWeakObjectPtr<UObject> ObjectWeakPtr)
+{
+	FText BlueprintFormat = NSLOCTEXT("SourceHyperlink", "EditBlueprint", "Edit {0}");
+	FText CodeFormat = NSLOCTEXT("SourceHyperlink", "GoToCode", "{0}");
+
+	return GetSourceLinkFormatted(Class, ObjectWeakPtr, BlueprintFormat, CodeFormat);
+}
+
+TSharedRef<SWidget> FEditorClassUtils::GetSourceLinkFormatted(const UClass* Class, const TWeakObjectPtr<UObject> ObjectWeakPtr, const FText& BlueprintFormat, const FText& CodeFormat)
 {
 	TSharedRef<SWidget> SourceHyperlink = SNew( SSpacer );
 	UBlueprint* Blueprint = (Class ? Cast<UBlueprint>(Class->ClassGeneratedBy) : nullptr);
@@ -103,7 +111,7 @@ TSharedRef<SWidget> FEditorClassUtils::GetSourceLink(const UClass* Class, const
 			.Style(FEditorStyle::Get(), "EditBPHyperlink")
 			.TextStyle(FEditorStyle::Get(), "DetailsView.EditBlueprintHyperlinkStyle")
 			.OnNavigate_Static(&Local::OnEditBlueprintClicked, BlueprintPtr, ObjectWeakPtr)
-			.Text(FText::Format(NSLOCTEXT("SourceHyperlink", "EditBlueprint", "Edit {0}"), FText::FromString( Blueprint->GetName() ) ))
+			.Text(FText::Format(BlueprintFormat, FText::FromString(Blueprint->GetName())))
 			.ToolTipText(NSLOCTEXT("SourceHyperlink", "EditBlueprint_ToolTip", "Click to edit the blueprint"));
 	}
 	else if( FSourceCodeNavigation::IsCompilerAvailable() )
@@ -124,7 +132,7 @@ TSharedRef<SWidget> FEditorClassUtils::GetSourceLink(const UClass* Class, const
 				.Style(FCoreStyle::Get(), "Hyperlink")
 				.TextStyle(FEditorStyle::Get(), "DetailsView.GoToCodeHyperlinkStyle")
 				.OnNavigate_Static(&Local::OnEditCodeClicked, ClassHeaderPath)
-				.Text(FText::Format(NSLOCTEXT("SourceHyperlink", "GoToCode", "{0}" ), FText::FromString(FPaths::GetCleanFilename( *ClassHeaderPath ) ) ) )
+				.Text(FText::Format(CodeFormat, FText::FromString(FPaths::GetCleanFilename( *ClassHeaderPath ) ) ) )
 				.ToolTipText(NSLOCTEXT("SourceHyperlink", "GoToCode_ToolTip", "Click to open this source file in a text editor"));
 		}
 	}
