[INTEGRATE] Change 2425681 by [REDACTED_PERSON][REDACTED_HANDLE] on 2015/01/30 09:47:34

	Editor - Forgot to add the header for the additional formatting source links.

diff --git a/Engine/Source/Editor/UnrealEd/Public/EditorClassUtils.h b/Engine/Source/Editor/UnrealEd/Public/EditorClassUtils.h
--- a/Engine/Source/Editor/UnrealEd/Public/EditorClassUtils.h
+++ b/Engine/Source/Editor/UnrealEd/Public/EditorClassUtils.h
@@ -73,6 +73,17 @@ namespace FEditorClassUtils
 	 */
 	UNREALED_API TSharedRef<SWidget> GetSourceLink(const UClass* Class, const TWeakObjectPtr<UObject> ObjectWeakPtr);
 
+	/**
+	 * Creates a link to the source code or blueprint for a given class formatted however you need. Example "Edit {0}"
+	 *
+	 * @param	Class			Class we want to build a link for
+	 * @param	ObjectWeakPtr	Optional object to set blueprint debugging to in the case we are choosing a blueprint
+	 * @param	BlueprintFormat	The text format for blueprint links.
+	 * @param	CodeFormat		The text format for C++ code file links.
+	 * @return					Shared pointer to the constructed tooltip
+	 */
+	UNREALED_API TSharedRef<SWidget> GetSourceLinkFormatted(const UClass* Class, const TWeakObjectPtr<UObject> ObjectWeakPtr, const FText& BlueprintFormat, const FText& CodeFormat);
+
 	/**
 	 * Fetches a UClass from the string name of the class
 	 *
