[INTEGRATE] Change 2425725 by [REDACTED_PERSON][REDACTED_HANDLE] on 2015/01/30 10:16:06

	UE-8533 Crash compiling Anim Blueprint with a Child
	UE-8680 Crash when reparenting when there is a derived BP as well

diff --git a/Engine/Source/Editor/UnrealEd/Private/Kismet2/KismetReinstanceUtilities.cpp b/Engine/Source/Editor/UnrealEd/Private/Kismet2/KismetReinstanceUtilities.cpp
--- a/Engine/Source/Editor/UnrealEd/Private/Kismet2/KismetReinstanceUtilities.cpp
+++ b/Engine/Source/Editor/UnrealEd/Private/Kismet2/KismetReinstanceUtilities.cpp
@@ -162,6 +162,7 @@ FBlueprintCompileReinstancer::FBlueprintCompileReinstancer(UClass* InClassToRein
 		DuplicatedClass->ClassDefaultObject->SetClass(DuplicatedClass);
 
 		OldCDO->SetClass(DuplicatedClass);
+		ObjectsThatShouldUseOldStuff.Add(DuplicatedClass); //CDO of REINST_ class can be used as archetype
 
 		if( !bIsBytecodeOnly )
 		{
