[INTEGRATE] Change 2425854 by [REDACTED_PERSON][REDACTED_HANDLE] on 2015/01/30 10:56:31

	Cleaned up text and menu separators in SCS editor's Add Component drop-down

diff --git a/Engine/Source/Editor/UnrealEd/Private/SComponentClassCombo.cpp b/Engine/Source/Editor/UnrealEd/Private/SComponentClassCombo.cpp
--- a/Engine/Source/Editor/UnrealEd/Private/SComponentClassCombo.cpp
+++ b/Engine/Source/Editor/UnrealEd/Private/SComponentClassCombo.cpp
@@ -388,7 +388,7 @@ void SComponentClassCombo::UpdateComponentClassList()
 
 	if( GetDefault<UEditorExperimentalSettings>()->bScriptableComponentsOnActors )
 	{
-		FString NewComponentsHeading = LOCTEXT("NewComponentsHeading", "Behavioral").ToString();
+		FString NewComponentsHeading = LOCTEXT("NewComponentsHeading", "New").ToString();
 		// Add new C++ component class
 		FComponentClassComboEntryPtr NewClassHeader = MakeShareable(new FComponentClassComboEntry(NewComponentsHeading));
 		ComponentClassList.Add(NewClassHeader);
@@ -398,6 +398,9 @@ void SComponentClassCombo::UpdateComponentClassList()
 
 		FComponentClassComboEntryPtr NewBPClass = MakeShareable(new FComponentClassComboEntry(NewComponentsHeading, UActorComponent::StaticClass(), true, EComponentCreateAction::CreateNewBlueprintClass));
 		ComponentClassList.Add(NewBPClass);
+
+		FComponentClassComboEntryPtr NewSeparator(new FComponentClassComboEntry());
+		ComponentClassList.Add( NewSeparator );
 	}
 	
 	TArray<FComponentClassComboEntryPtr> SortedClassList;
