[INTEGRATE] Change 2426019 by [REDACTED_PERSON][REDACTED_HANDLE] on 2015/01/30 12:25:20

	Missing file

diff --git a/Engine/Source/Editor/UnrealEd/Private/EditorSelectUtils.cpp b/Engine/Source/Editor/UnrealEd/Private/EditorSelectUtils.cpp
--- a/Engine/Source/Editor/UnrealEd/Private/EditorSelectUtils.cpp
+++ b/Engine/Source/Editor/UnrealEd/Private/EditorSelectUtils.cpp
@@ -569,7 +569,8 @@ void UUnrealEdEngine::SelectActor(AActor* Actor, bool bInSelected, bool bNotify,
 					auto PrimComponent = Cast<UPrimitiveComponent>(Component);
 					if (PrimComponent && PrimComponent->SelectionOverrideDelegate.IsBound())
 					{
-						PrimComponent->Modify();
+						const bool bShouldMarkPackageDirty = false;
+						PrimComponent->Modify( bShouldMarkPackageDirty );
 						PrimComponent->SelectionOverrideDelegate.Unbind();
 					}
 				}
@@ -582,7 +583,8 @@ void UUnrealEdEngine::SelectActor(AActor* Actor, bool bInSelected, bool bNotify,
 					auto PrimComponent = Cast<UPrimitiveComponent>(Component);
 					if (PrimComponent && !PrimComponent->SelectionOverrideDelegate.IsBound())
 					{
-						PrimComponent->Modify();
+						const bool bShouldMarkPackageDirty = false;
+						PrimComponent->Modify( bShouldMarkPackageDirty );
 						PrimComponent->SelectionOverrideDelegate = UPrimitiveComponent::FSelectionOverride::CreateUObject(this, &UUnrealEdEngine::IsComponentSelected);
 					}
 				}
