[INTEGRATE] Change 2426108 by [REDACTED_PERSON][REDACTED_HANDLE] on 2015/01/30 13:47:46

	[UE-8152] - All previously existing drag-and-drop functionality now works with components in the world as follows:

	1. If multiple actors are selected and the drop is on one of the selected actors, the dropped object will be applied (if possible) to all components in each selected actor.
	2. If multiple components are selected and the drop is on one of the selected components, the dropped object will be applied (if possible) to all selected components.
	3. If the drop is onto a component that is not selected, it will apply only to that component.
	4. If the drop is onto a component within a selected actor, but only one actor is selected, same as #3.

	- Fixed crash when trying to apply material to non-existent slot when applying to multiple actors

diff --git a/Engine/Source/Editor/LevelEditor/Private/SActorDetails.cpp b/Engine/Source/Editor/LevelEditor/Private/SActorDetails.cpp
--- a/Engine/Source/Editor/LevelEditor/Private/SActorDetails.cpp
+++ b/Engine/Source/Editor/LevelEditor/Private/SActorDetails.cpp
@@ -325,7 +325,7 @@ void SActorDetails::OnSCSEditorTreeViewSelectionChanged(const TArray<FSCSEditorT
 									auto SceneComponent = Cast<USceneComponent>(ComponentInstance);
 									if (SceneComponent)
 									{
-										FComponentEditorUtils::BindComponentSelectionOverride(SceneComponent);
+										FComponentEditorUtils::BindComponentSelectionOverride(SceneComponent, true);
 									}
 								}
 							}
diff --git a/Engine/Source/Editor/UnrealEd/Private/AssetSelection.cpp b/Engine/Source/Editor/UnrealEd/Private/AssetSelection.cpp
--- a/Engine/Source/Editor/UnrealEd/Private/AssetSelection.cpp
+++ b/Engine/Source/Editor/UnrealEd/Private/AssetSelection.cpp
@@ -8,6 +8,7 @@
 #include "LevelUtils.h"
 #include "EditorLevelUtils.h"
 #include "MainFrame.h"
+#include "ComponentEditorUtils.h"
 
 #include "ComponentAssetBroker.h"
 
@@ -708,7 +709,7 @@ bool FActorFactoryAssetProxy::ApplyMaterialToActor( AActor* TargetActor, UMateri
 
 			// Some actors could potentially have multiple mesh components, so we need to store all of the potentially valid ones
 			// (or else perform special cases with IsA checks on the target actor)
-			TArray<UActorComponent*> FoundMeshComponents;
+			TArray<USceneComponent*> FoundMeshComponents;
 
 			// Find which mesh the user clicked on first.
 			TInlineComponentArray<USceneComponent*> SceneComponents;
@@ -735,55 +736,10 @@ bool FActorFactoryAssetProxy::ApplyMaterialToActor( AActor* TargetActor, UMateri
 			if ( FoundMeshComponents.Num() > 0 )
 			{
 				// Check each component that was found
-				for ( TArray<UActorComponent*>::TConstIterator MeshCompIter( FoundMeshComponents ); MeshCompIter; ++MeshCompIter )
+				for ( TArray<USceneComponent*>::TConstIterator MeshCompIter( FoundMeshComponents ); MeshCompIter; ++MeshCompIter )
 				{
-					UActorComponent* ActorComp = *MeshCompIter;
-
-					UMeshComponent* FoundMeshComponent = Cast<UMeshComponent>(ActorComp);
-					UDecalComponent* DecalComponent = Cast<UDecalComponent>(ActorComp);
-
-					if(FoundMeshComponent)
-					{
-						// OK, we need to figure out how many material slots this mesh component/static mesh has.
-						// Start with the actor's material count, then drill into the static/skeletal mesh to make sure 
-						// we have the right total.
-						int32 MaterialCount = FMath::Max(FoundMeshComponent->OverrideMaterials.Num(), FoundMeshComponent->GetNumMaterials());
-
-						// Any materials to overwrite?
-						if( MaterialCount > 0 )
-						{
-							const FScopedTransaction Transaction( NSLOCTEXT("UnrealEd", "DropTarget_UndoSetActorMaterial", "Assign Material (Drag and Drop)") );
-							FoundMeshComponent->Modify();
-
-							// If the material slot is -1 then apply the material to every slot.
-							if ( OptionalMaterialSlot == -1 )
-							{
-								for ( int32 CurMaterialIndex = 0; CurMaterialIndex < MaterialCount; ++CurMaterialIndex )
-								{
-									FoundMeshComponent->SetMaterial(CurMaterialIndex, MaterialToApply);
-								}
-							}
-							else
-							{
-								if( OptionalMaterialSlot < MaterialCount )
-								{
-									FoundMeshComponent->SetMaterial(OptionalMaterialSlot, MaterialToApply);
-								}
-							}
-
-							TargetActor->MarkComponentsRenderStateDirty();
-							bResult = true;
-						}
-					}
-					else if(DecalComponent)
-					{
-						const FScopedTransaction Transaction( NSLOCTEXT("UnrealEd", "DropTarget_UndoSetActorMaterial", "Assign Material (Drag and Drop)") );
-						DecalComponent->Modify();
-						// set material #0
-						DecalComponent->SetMaterial(0, MaterialToApply);
-						TargetActor->MarkComponentsRenderStateDirty();
-						bResult = true;
-					}
+					USceneComponent* SceneComp = *MeshCompIter;
+					bResult = FComponentEditorUtils::AttemptApplyMaterialToComponent(SceneComp, MaterialToApply, OptionalMaterialSlot);
 				}
 			}
 		}
diff --git a/Engine/Source/Editor/UnrealEd/Private/EditorSelectUtils.cpp b/Engine/Source/Editor/UnrealEd/Private/EditorSelectUtils.cpp
--- a/Engine/Source/Editor/UnrealEd/Private/EditorSelectUtils.cpp
+++ b/Engine/Source/Editor/UnrealEd/Private/EditorSelectUtils.cpp
@@ -7,6 +7,7 @@
 #include "Editor/StatsViewer/Public/StatsViewerModule.h"
 #include "SnappingUtils.h"
 #include "MessageLog.h"
+#include "ComponentEditorUtils.h"
 
 #define LOCTEXT_NAMESPACE "EditorSelectUtils"
 
@@ -566,27 +567,17 @@ void UUnrealEdEngine::SelectActor(AActor* Actor, bool bInSelected, bool bNotify,
 					GetSelectedComponents()->Deselect( Component );
 
 					// Remove the selection override delegates from the deselected components
-					auto PrimComponent = Cast<UPrimitiveComponent>(Component);
-					if (PrimComponent && PrimComponent->SelectionOverrideDelegate.IsBound())
-					{
-						const bool bShouldMarkPackageDirty = false;
-						PrimComponent->Modify( bShouldMarkPackageDirty );
-						PrimComponent->SelectionOverrideDelegate.Unbind();
-					}
+					auto SceneComponent = Cast<USceneComponent>(Component);
+					FComponentEditorUtils::BindComponentSelectionOverride(SceneComponent, false);
 				}
 			}
 			else
 			{
+				// Bind the override delegates for the components in the selected actor
 				for (UActorComponent* Component : Actor->GetComponents())
 				{
-					// Establish the selection override delegate for each of the actor's components
-					auto PrimComponent = Cast<UPrimitiveComponent>(Component);
-					if (PrimComponent && !PrimComponent->SelectionOverrideDelegate.IsBound())
-					{
-						const bool bShouldMarkPackageDirty = false;
-						PrimComponent->Modify( bShouldMarkPackageDirty );
-						PrimComponent->SelectionOverrideDelegate = UPrimitiveComponent::FSelectionOverride::CreateUObject(this, &UUnrealEdEngine::IsComponentSelected);
-					}
+					auto SceneComponent = Cast<USceneComponent>(Component);
+					FComponentEditorUtils::BindComponentSelectionOverride(SceneComponent, true);
 				}
 			}
 
@@ -632,10 +623,11 @@ void UUnrealEdEngine::SelectComponent(UActorComponent* Component, bool bInSelect
 
 		GetSelectedComponents()->Select(Component, bInSelected);
 
-		auto PrimComponent = Cast<UPrimitiveComponent>(Component);
-		if (PrimComponent && !PrimComponent->SelectionOverrideDelegate.IsBound())
+		// Make sure the override delegate is bound properly
+		auto SceneComponent = Cast<USceneComponent>(Component);
+		if (SceneComponent)
 		{
-			PrimComponent->SelectionOverrideDelegate = UPrimitiveComponent::FSelectionOverride::CreateUObject(this, &UUnrealEdEngine::IsComponentSelected);
+			FComponentEditorUtils::BindComponentSelectionOverride(SceneComponent, true);
 		}
 
 		// Update the selection visualization
diff --git a/Engine/Source/Editor/UnrealEd/Private/Kismet2/ComponentEditorUtils.cpp b/Engine/Source/Editor/UnrealEd/Private/Kismet2/ComponentEditorUtils.cpp
--- a/Engine/Source/Editor/UnrealEd/Private/Kismet2/ComponentEditorUtils.cpp
+++ b/Engine/Source/Editor/UnrealEd/Private/Kismet2/ComponentEditorUtils.cpp
@@ -6,6 +6,7 @@
 #include "Runtime/Engine/Classes/Components/SceneComponent.h"
 #include "Engine/SCS_Node.h"
 #include "Engine/SimpleConstructionScript.h"
+#include "ScopedTransaction.h"
 
 
 USceneComponent* FComponentEditorUtils::GetSceneComponent( UObject* Object, UObject* SubObject /*= NULL*/ )
@@ -184,33 +185,100 @@ void FComponentEditorUtils::AdjustComponentDelta(USceneComponent* Component, FVe
 	}
 }
 
-void FComponentEditorUtils::BindComponentSelectionOverride(USceneComponent* SceneComponent)
+void FComponentEditorUtils::BindComponentSelectionOverride(USceneComponent* SceneComponent, bool bBind)
 {
 	if (SceneComponent)
 	{
-		// If the scene component is a primitive component, ensure the override is bound
+		// Modifying the selection override delegate is purely cosmetic and shouldn't dirty the package
+		const bool bShouldMarkPackageDirty = false;
+
+		// If the scene component is a primitive component, set the override for it
 		auto PrimComponent = Cast<UPrimitiveComponent>(SceneComponent);
-		if (PrimComponent && !PrimComponent->SelectionOverrideDelegate.IsBound())
+		if (PrimComponent && PrimComponent->SelectionOverrideDelegate.IsBound() != bBind)
 		{
-			PrimComponent->Modify();
-			PrimComponent->SelectionOverrideDelegate = UPrimitiveComponent::FSelectionOverride::CreateUObject(GUnrealEd, &UUnrealEdEngine::IsComponentSelected);
+			PrimComponent->Modify(bShouldMarkPackageDirty);
+			if (bBind)
+			{
+				PrimComponent->SelectionOverrideDelegate = UPrimitiveComponent::FSelectionOverride::CreateUObject(GUnrealEd, &UUnrealEdEngine::IsComponentSelected);
+			}
+			else
+			{
+				PrimComponent->SelectionOverrideDelegate.Unbind();
+			}
 		}
 		else
 		{
-			// Otherwise, make sure the override is bound on any attached primitive components (to make sure we catch any editor-only billboards and the like)
+			// Otherwise, make sure the override is set properly on any attached editor-only primitive components (ex: billboards)
 			for (auto Component : SceneComponent->AttachChildren)
 			{
 				PrimComponent = Cast<UPrimitiveComponent>(Component);
-				if (PrimComponent && !PrimComponent->SelectionOverrideDelegate.IsBound())
+				if (PrimComponent && PrimComponent->IsEditorOnly() && PrimComponent->SelectionOverrideDelegate.IsBound() != bBind)
 				{
-					PrimComponent->Modify();
-					PrimComponent->SelectionOverrideDelegate = UPrimitiveComponent::FSelectionOverride::CreateUObject(GUnrealEd, &UUnrealEdEngine::IsComponentSelected);
+					PrimComponent->Modify(bShouldMarkPackageDirty);
+					if (bBind)
+					{
+						PrimComponent->SelectionOverrideDelegate = UPrimitiveComponent::FSelectionOverride::CreateUObject(GUnrealEd, &UUnrealEdEngine::IsComponentSelected);
+					}
+					else
+					{
+						PrimComponent->SelectionOverrideDelegate.Unbind();
+					}
 				}
 			}
 		}
 	}
 }
 
+bool FComponentEditorUtils::AttemptApplyMaterialToComponent(USceneComponent* SceneComponent, UMaterialInterface* MaterialToApply, int32 OptionalMaterialSlot)
+{
+	bool bResult = false;
+	
+	auto MeshComponent = Cast<UMeshComponent>(SceneComponent);
+	auto DecalComponent = Cast<UDecalComponent>(SceneComponent);
+
+	// We can only apply a material to a mesh or a decal
+	if (MeshComponent || DecalComponent)
+	{
+		bResult = true;
+		const FScopedTransaction Transaction(NSLOCTEXT("UnrealEd", "DropTarget_UndoSetComponentMaterial", "Assign Material to Component (Drag and Drop)"));
+		SceneComponent->Modify();
+
+		if (MeshComponent)
+		{
+			// OK, we need to figure out how many material slots this mesh component/static mesh has.
+			// Start with the actor's material count, then drill into the static/skeletal mesh to make sure 
+			// we have the right total.
+			int32 MaterialCount = FMath::Max(MeshComponent->OverrideMaterials.Num(), MeshComponent->GetNumMaterials());
+
+			// Do we have an overridable material at the appropriate slot?
+			if (MaterialCount > 0 && OptionalMaterialSlot < MaterialCount)
+			{
+				if (OptionalMaterialSlot == -1)
+				{
+					// Apply the material to every slot.
+					for (int32 CurMaterialIndex = 0; CurMaterialIndex < MaterialCount; ++CurMaterialIndex)
+					{
+						MeshComponent->SetMaterial(CurMaterialIndex, MaterialToApply);
+					}
+				}
+				else
+				{
+					// Apply only to the indicated slot.
+					MeshComponent->SetMaterial(OptionalMaterialSlot, MaterialToApply);
+				}
+			}
+		}
+		else
+		{
+			DecalComponent->SetMaterial(0, MaterialToApply);
+		}
+
+		SceneComponent->MarkRenderStateDirty();
+	}
+
+	return bResult;
+}
+
 void FComponentEditorUtils::PropagateTransformPropertyChange(
 	class USceneComponent* InSceneComponentTemplate,
 	const FTransformData& OldDefaultTransform,
diff --git a/Engine/Source/Editor/UnrealEd/Private/LevelEditorViewport.cpp b/Engine/Source/Editor/UnrealEd/Private/LevelEditorViewport.cpp
--- a/Engine/Source/Editor/UnrealEd/Private/LevelEditorViewport.cpp
+++ b/Engine/Source/Editor/UnrealEd/Private/LevelEditorViewport.cpp
@@ -492,192 +492,187 @@ static UObject* GetOrCreateMaterialFromTexture( UTexture* UnrealTexture )
 }
 
 /**
- * Helper function that attempts to apply the supplied object to the supplied actor.
- *
- * @param	ObjToUse				Object to attempt to apply as specific asset
- * @param	ActorToApplyTo			Actor to whom the asset should be applied
- * @param   TargetMaterialSlot      When dealing with submeshes this will represent the target section/slot to apply materials to.
- * @param	bTest					Whether to test if the object would be successfully applied without actually doing it.
- * @param	bCreateDropPreview		Whether this is just a drop preview.
- *
- * @return	true if the provided object was successfully applied to the provided actor
- */
-static bool AttemptApplyObjToActor( UObject* ObjToUse, AActor* ActorToApplyTo, int32 TargetMaterialSlot = -1, bool bTest = false, bool bCreateDropPreview = false )
+* Helper function that attempts to apply the supplied object to the supplied actor.
+*
+* @param	ObjToUse				Object to attempt to apply as specific asset
+* @param	ComponentToApplyTo		Component to whom the asset should be applied
+* @param	TargetMaterialSlot      When dealing with submeshes this will represent the target section/slot to apply materials to.
+* @param	bTest					Whether to test if the object would be successfully applied without actually doing it.
+*
+* @return	true if the provided object was successfully applied to the provided actor
+*/
+static bool AttemptApplyObjToComponent(UObject* ObjToUse, USceneComponent* ComponentToApplyTo, int32 TargetMaterialSlot = -1, bool bTest = false)
 {
 	bool bResult = false;
 
-	if ( ActorToApplyTo )
+	if (ComponentToApplyTo && ComponentToApplyTo->CreationMethod != EComponentCreationMethod::ConstructionScript)
 	{
-		UTexture* DroppedObjAsTexture = Cast<UTexture>( ObjToUse );
-		if ( DroppedObjAsTexture != NULL )
+		// MESH/DECAL
+		auto MeshComponent = Cast<UMeshComponent>(ComponentToApplyTo);
+		auto DecalComponent = Cast<UDecalComponent>(ComponentToApplyTo);
+		if (MeshComponent || DecalComponent)
 		{
-			if ( bTest || bCreateDropPreview )
+			// Dropping a texture?
+			UTexture* DroppedObjAsTexture = Cast<UTexture>(ObjToUse);
+			if (DroppedObjAsTexture != NULL)
 			{
-				bResult = true;
-			}
-			else
-			{
-				ObjToUse = GetOrCreateMaterialFromTexture( DroppedObjAsTexture );
-			}
-		}
-
-		// Ensure the provided object is some form of material
-		UMaterialInterface* DroppedObjAsMaterial = Cast<UMaterialInterface>( ObjToUse );
-		if ( DroppedObjAsMaterial )
-		{
-			if (bTest || bCreateDropPreview)
-			{
-				bResult = true;
-			}
-			else
-			{
-				// Apply the material to the actor
-				FScopedTransaction Transaction( NSLOCTEXT("UnrealEd", "DragDrop_Transaction_ApplyMaterialToActor", "Apply Material to Actor") );
-				bResult = FActorFactoryAssetProxy::ApplyMaterialToActor( ActorToApplyTo, DroppedObjAsMaterial, TargetMaterialSlot );
-			}
-		}
-
-		USkeletalMesh* DroppedObjAsSkeletalMesh = Cast<USkeletalMesh>( ObjToUse );
-		USkeleton* DroppedObjAsSkeleton = Cast<USkeleton>( ObjToUse );
-		if ( DroppedObjAsSkeletalMesh ||
-			 DroppedObjAsSkeleton )
-		{
-			if ( bTest )
-			{
-				if ( ActorToApplyTo->IsA(ASkeletalMeshActor::StaticClass()) )
+				if (bTest)
 				{
 					bResult = true;
 				}
+				else
+				{
+					// Turn dropped textures into materials
+					ObjToUse = GetOrCreateMaterialFromTexture(DroppedObjAsTexture);
+				}
 			}
-			else
+
+			// Dropping a material?
+			UMaterialInterface* DroppedObjAsMaterial = Cast<UMaterialInterface>(ObjToUse);
+			if (DroppedObjAsMaterial)
 			{
-				if ( ASkeletalMeshActor* SkelMeshActor = Cast<ASkeletalMeshActor>(ActorToApplyTo) )
+				if (bTest)
 				{
-					const FScopedTransaction Transaction( LOCTEXT( "DropSkelMeshOnObject", "Drop Skeletal Mesh On Object" ) );
-					USkeletalMeshComponent* SkelMeshComponent = SkelMeshActor->GetSkeletalMeshComponent();
-					SkelMeshComponent->Modify();
-					if ( DroppedObjAsSkeletalMesh )
-					{
-						SkelMeshComponent->SetSkeletalMesh(DroppedObjAsSkeletalMesh);
-					}
-					else if ( DroppedObjAsSkeleton )
-					{
-						SkelMeshComponent->SetSkeletalMesh(DroppedObjAsSkeleton->GetPreviewMesh(true));
-					}
 					bResult = true;
 				}
+				else
+				{
+					bResult = FComponentEditorUtils::AttemptApplyMaterialToComponent(ComponentToApplyTo, DroppedObjAsMaterial, TargetMaterialSlot);
+				}
 			}
 		}
 
-		UAnimBlueprint* DroppedObjAsAnimBlueprint = Cast<UAnimBlueprint>( ObjToUse );
-		if ( DroppedObjAsAnimBlueprint )
+		// SKELETAL MESH COMPONENT
+		auto SkeletalMeshComponent = Cast<USkeletalMeshComponent>(ComponentToApplyTo);
+		if (SkeletalMeshComponent)
 		{
-			USkeleton* NeedsSkeleton = DroppedObjAsAnimBlueprint->TargetSkeleton;
-			if ( NeedsSkeleton )
+			// Dropping an Anim Blueprint?
+			UAnimBlueprint* DroppedObjAsAnimBlueprint = Cast<UAnimBlueprint>(ObjToUse);
+			if (DroppedObjAsAnimBlueprint)
 			{
-				if(bTest)
+				USkeleton* AnimBPSkeleton = DroppedObjAsAnimBlueprint->TargetSkeleton;
+				if (AnimBPSkeleton)
 				{
-					if(ActorToApplyTo->IsA(ASkeletalMeshActor::StaticClass()))
+					if (bTest)
 					{
 						bResult = true;
 					}
-				}
-				else
-				{
-					if(ASkeletalMeshActor* SkelMeshActor = Cast<ASkeletalMeshActor>(ActorToApplyTo))
+					else
 					{
 						const FScopedTransaction Transaction(LOCTEXT("DropAnimBlueprintOnObject", "Drop Anim Blueprint On Object"));
+						SkeletalMeshComponent->Modify();
 
-						USkeletalMeshComponent* SkelMeshComponent = SkelMeshActor->GetSkeletalMeshComponent();
-						// if anim blueprint skeleton and mesh skeleton does not match or component does not have any mesh, then change mesh
-						bool bShouldChangeMesh = (SkelMeshComponent->SkeletalMesh == NULL ||
-								!NeedsSkeleton->IsCompatible(SkelMeshComponent->SkeletalMesh->Skeleton));
+						// If the component doesn't have a mesh or the anim blueprint's skeleton isn't compatible with the existing mesh's skeleton, the mesh should change
+						const bool bShouldChangeMesh = !SkeletalMeshComponent->SkeletalMesh || !AnimBPSkeleton->IsCompatible(SkeletalMeshComponent->SkeletalMesh->Skeleton);
 
-						if(bShouldChangeMesh)
+						if (bShouldChangeMesh)
 						{
-							SkelMeshComponent->SetSkeletalMesh(NeedsSkeleton->GetPreviewMesh(true));
+							SkeletalMeshComponent->SetSkeletalMesh(AnimBPSkeleton->GetPreviewMesh(true));
 						}
 
-						// make sure if it's compabile now, if not we're not going to change anim blueprint
-						if(SkelMeshComponent->SkeletalMesh &&
-							NeedsSkeleton->IsCompatible(SkelMeshComponent->SkeletalMesh->Skeleton))
+						// Verify that the skeletons are compatible before changing the anim BP
+						if (SkeletalMeshComponent->SkeletalMesh && AnimBPSkeleton->IsCompatible(SkeletalMeshComponent->SkeletalMesh->Skeleton))
 						{
-							SkelMeshComponent->SetAnimInstanceClass(DroppedObjAsAnimBlueprint->GeneratedClass);
+							SkeletalMeshComponent->SetAnimInstanceClass(DroppedObjAsAnimBlueprint->GeneratedClass);
 							bResult = true;
 						}
 					}
 				}
 			}
-		}
 
-		UAnimationAsset* DroppedObjAsAnimationAsset = Cast<UAnimationAsset>( ObjToUse );
-		UVertexAnimation* DroppedObjAsVertexAnimation = Cast<UVertexAnimation>( ObjToUse );
-		// block anything else than just anim sequence
-		if( DroppedObjAsAnimationAsset != NULL )
-		{
-			if( ! DroppedObjAsAnimationAsset->IsA(UAnimSequenceBase::StaticClass()) )
+			// Dropping an Anim Sequence or Vertex Animation?
+			UAnimSequenceBase* DroppedObjAsAnimSequence = Cast<UAnimSequenceBase>(ObjToUse);
+			UVertexAnimation* DroppedObjAsVertexAnimation = Cast<UVertexAnimation>(ObjToUse);
+			if (DroppedObjAsAnimSequence || DroppedObjAsVertexAnimation)
 			{
-				DroppedObjAsAnimationAsset = NULL;
-			}
-		}
+				USkeleton* AnimSkeleton = nullptr;
+				const bool bVertexAnimHasValidMesh = DroppedObjAsVertexAnimation && DroppedObjAsVertexAnimation->BaseSkelMesh;
 
-		if ( DroppedObjAsAnimationAsset ||
-			 DroppedObjAsVertexAnimation)
-		{
-			USkeleton* NeedsSkeleton = DroppedObjAsAnimationAsset? DroppedObjAsAnimationAsset->GetSkeleton() :
-				(DroppedObjAsVertexAnimation && DroppedObjAsVertexAnimation->BaseSkelMesh? DroppedObjAsVertexAnimation->BaseSkelMesh->Skeleton : NULL);
+				if (DroppedObjAsAnimSequence)
+				{
+					AnimSkeleton = DroppedObjAsAnimSequence->GetSkeleton();
+				}
+				else if (bVertexAnimHasValidMesh)
+				{
+					AnimSkeleton = DroppedObjAsVertexAnimation->BaseSkelMesh->Skeleton;
+				}
 
-			if (NeedsSkeleton)
-			{
-				if(bTest)
+				if (AnimSkeleton)
 				{
-					if(ActorToApplyTo->IsA(ASkeletalMeshActor::StaticClass()))
+					if (bTest)
 					{
 						bResult = true;
 					}
-				}
-				else
-				{
-					if(ASkeletalMeshActor* SkelMeshActor = Cast<ASkeletalMeshActor>(ActorToApplyTo))
+					else
 					{
 						const FScopedTransaction Transaction(LOCTEXT("DropAnimationOnObject", "Drop Animation On Object"));
-						USkeletalMeshComponent* SkelComponent = SkelMeshActor->GetSkeletalMeshComponent();
-						SkelComponent->Modify();
-						// if asset skeleton and mesh skeleton does not match or component does not have any mesh, then change mesh
-						bool bShouldChangeMesh = SkelComponent->SkeletalMesh == NULL || 
-							!NeedsSkeleton->IsCompatible(SkelComponent->SkeletalMesh->Skeleton);
+						SkeletalMeshComponent->Modify();
+
+						// If the component doesn't have a mesh or the anim blueprint's skeleton isn't compatible with the existing mesh's skeleton, the mesh should change
+						const bool bShouldChangeMesh = !SkeletalMeshComponent->SkeletalMesh || !AnimSkeleton->IsCompatible(SkeletalMeshComponent->SkeletalMesh->Skeleton);
 
-						if(bShouldChangeMesh)
+						if (bShouldChangeMesh)
 						{
-							SkelComponent->SetSkeletalMesh(NeedsSkeleton->GetAssetPreviewMesh(DroppedObjAsAnimationAsset));
+							SkeletalMeshComponent->SetSkeletalMesh(AnimSkeleton->GetAssetPreviewMesh(DroppedObjAsAnimSequence));
 						}
 
-						if(DroppedObjAsAnimationAsset)
+						if (DroppedObjAsAnimSequence)
 						{
-							SkelComponent->SetAnimationMode(EAnimationMode::Type::AnimationSingleNode);
-							SkelComponent->AnimationData.AnimToPlay = DroppedObjAsAnimationAsset;
+							SkeletalMeshComponent->SetAnimationMode(EAnimationMode::Type::AnimationSingleNode);
+							SkeletalMeshComponent->AnimationData.AnimToPlay = DroppedObjAsAnimSequence;
 
 							// set runtime data
-							SkelComponent->SetAnimation(DroppedObjAsAnimationAsset);
+							SkeletalMeshComponent->SetAnimation(DroppedObjAsAnimSequence);
 						}
-						if(DroppedObjAsVertexAnimation)
+
+						if (DroppedObjAsVertexAnimation)
 						{
-							SkelComponent->SetAnimationMode(EAnimationMode::Type::AnimationSingleNode);
-							SkelComponent->AnimationData.VertexAnimToPlay = DroppedObjAsVertexAnimation;
+							SkeletalMeshComponent->SetAnimationMode(EAnimationMode::Type::AnimationSingleNode);
+							SkeletalMeshComponent->AnimationData.VertexAnimToPlay = DroppedObjAsVertexAnimation;
 
 							// set runtime data
-							SkelComponent->SetVertexAnimation(DroppedObjAsVertexAnimation);
+							SkeletalMeshComponent->SetVertexAnimation(DroppedObjAsVertexAnimation);
 						}
-						if(SkelComponent && SkelComponent->SkeletalMesh)
+
+						if (SkeletalMeshComponent && SkeletalMeshComponent->SkeletalMesh)
 						{
 							bResult = true;
-							SkelComponent->InitAnim(true);
+							SkeletalMeshComponent->InitAnim(true);
 						}
 					}
 				}
 			}
 		}
-		
+	}
+
+	return bResult;
+}
+
+/**
+ * Helper function that attempts to apply the supplied object to the supplied actor.
+ *
+ * @param	ObjToUse				Object to attempt to apply as specific asset
+ * @param	ActorToApplyTo			Actor to whom the asset should be applied
+ * @param   TargetMaterialSlot      When dealing with submeshes this will represent the target section/slot to apply materials to.
+ * @param	bTest					Whether to test if the object would be successfully applied without actually doing it.
+ *
+ * @return	true if the provided object was successfully applied to the provided actor
+ */
+static bool AttemptApplyObjToActor( UObject* ObjToUse, AActor* ActorToApplyTo, int32 TargetMaterialSlot = -1, bool bTest = false )
+{
+	bool bResult = false;
+
+	if ( ActorToApplyTo )
+	{
+		bResult = false;
+
+		TInlineComponentArray<USceneComponent*> SceneComponents;
+		ActorToApplyTo->GetComponents(SceneComponents);
+		for (auto SceneComp : SceneComponents)
+		{
+			bResult |= AttemptApplyObjToComponent(ObjToUse, SceneComp, TargetMaterialSlot, bTest);
+		}
+
 		// Notification hook for dropping asset onto actor
 		if(!bTest)
 		{
@@ -688,7 +683,6 @@ static bool AttemptApplyObjToActor( UObject* ObjToUse, AActor* ActorToApplyTo, i
 	return bResult;
 }
 
-
 /**
  * Helper function that attempts to apply the supplied object as a material to the BSP surface specified by the
  * provided model and index.
@@ -861,8 +855,8 @@ bool FLevelEditorViewportClient::DropObjectsOnActor(FViewportCursorLocation& Cur
 
 	for ( auto DroppedObject : DroppedObjects )
 	{
-		const bool bTest = false;
-		const bool bAppliedToActor = ( FactoryToUse == NULL ) ? AttemptApplyObjToActor( DroppedObject, DroppedUponActor, DroppedUponSlot, bTest, bCreateDropPreview ) : false;
+		const bool bIsTestApplication = bCreateDropPreview;
+		const bool bAppliedToActor = ( FactoryToUse == nullptr ) ? AttemptApplyObjToActor( DroppedObject, DroppedUponActor, DroppedUponSlot, bIsTestApplication ) : false;
 
 		if (!bAppliedToActor)
 		{
@@ -1208,12 +1202,14 @@ bool FLevelEditorViewportClient::DropObjectsAtCoordinates(int32 MouseX, int32 Mo
 		else if (HitProxy->IsA(HActor::StaticGetType()) || HitProxy->IsA(HBSPBrushVert::StaticGetType()))
 		{
 			AActor* TargetActor = NULL;
+			UPrimitiveComponent* TargetComponent = nullptr;
 			int32 TargetMaterialSlot = -1;
 
 			if (HitProxy->IsA(HActor::StaticGetType()))
 			{
 				HActor* TargetProxy = static_cast<HActor*>(HitProxy);
 				TargetActor = TargetProxy->Actor;
+				TargetComponent = const_cast<UPrimitiveComponent*>(TargetProxy->PrimComponent);
 				TargetMaterialSlot = TargetProxy->MaterialIndex;
 			}
 			else if (HitProxy->IsA(HBSPBrushVert::StaticGetType()))
@@ -1234,26 +1230,58 @@ bool FLevelEditorViewportClient::DropObjectsAtCoordinates(int32 MouseX, int32 Mo
 			{
 				FNavigationLockContext LockNavigationUpdates(TargetActor->GetWorld(), ENavigationLockReason::SpawnOnDragEnter, bCreateDropPreview);
 
-				// if the target actor is selected, we should drop onto all selected items
-				// otherwise, we should drop only onto this object
-				bool bDropOntoSelected = TargetActor->IsSelected();
-
-				if( !bDropOntoSelected || 
-					bOnlyDropOnTarget || 
-					FactoryToUse != NULL ||
-					!AttemptApplyObjToActor(DroppedObjects[0], TargetActor, TargetMaterialSlot, true) )
+				// If more than one actor is selected and the target actor is one of them, we should drop onto all selected actors
+				// otherwise, we should drop only onto the target object
+				const bool bDropOntoSelectedActors = TargetActor->IsSelected() && GEditor->GetSelectedActorCount() > 1;
+				const bool bCanApplyToComponent = AttemptApplyObjToComponent(DroppedObjects[0], TargetComponent, TargetMaterialSlot, true);
+				if (bOnlyDropOnTarget || !bDropOntoSelectedActors || !bCanApplyToComponent)
 				{
-					bResult = DropObjectsOnActor(Cursor, DroppedObjects, TargetActor, TargetMaterialSlot, ObjectFlags, OutNewActors, bCreateDropPreview, SelectActors, FactoryToUse);
+					if (bCanApplyToComponent)
+					{
+						const bool bIsTestAttempt = bCreateDropPreview;
+						bResult = AttemptApplyObjToComponent(DroppedObjects[0], TargetComponent, TargetMaterialSlot, bIsTestAttempt);
+					}
+					else
+					{
+						// Couldn't apply to a component, so try dropping the objects on the hit actor
+						bResult = DropObjectsOnActor(Cursor, DroppedObjects, TargetActor, TargetMaterialSlot, ObjectFlags, OutNewActors, bCreateDropPreview, SelectActors, FactoryToUse);
+					}
 				}
 				else
 				{
-					for ( FSelectionIterator It(*GEditor->GetSelectedActors()) ; It ; ++It )
+					// Are any components selected?
+					if (GEditor->GetSelectedComponentCount() > 0)
 					{
-						TargetActor = static_cast<AActor*>(*It);
-						if( TargetActor )
+						// Is the target component selected?
+						USelection* ComponentSelection = GEditor->GetSelectedComponents();
+						if (ComponentSelection->IsSelected(TargetComponent))
 						{
-							DropObjectsOnActor(Cursor, DroppedObjects, TargetActor, TargetMaterialSlot, ObjectFlags, OutNewActors, bCreateDropPreview, SelectActors, FactoryToUse);
-							bResult = true;
+							// The target component is selected, so try applying the object to every selected component
+							for (FSelectionIterator It(GEditor->GetSelectedComponentIterator()); It; ++It)
+							{
+								auto SceneComponent = Cast<USceneComponent>(*It);
+								AttemptApplyObjToComponent(DroppedObjects[0], SceneComponent, TargetMaterialSlot, bCreateDropPreview);
+								bResult = true;
+							}
+						}
+						else
+						{
+							// The target component is not selected, so apply the object exclusively to it
+							bResult = AttemptApplyObjToComponent(DroppedObjects[0], TargetComponent, TargetMaterialSlot, bCreateDropPreview);
+						}
+					}
+					
+					if (!bResult)
+					{
+						const FScopedTransaction Transaction(LOCTEXT("DropObjectsOnSelectedActors", "Drop Objects on Selected Actors"));
+						for (FSelectionIterator It(*GEditor->GetSelectedActors()); It; ++It)
+						{
+							TargetActor = static_cast<AActor*>( *It );
+							if (TargetActor)
+							{
+								DropObjectsOnActor(Cursor, DroppedObjects, TargetActor, TargetMaterialSlot, ObjectFlags, OutNewActors, bCreateDropPreview, SelectActors, FactoryToUse);
+								bResult = true;
+							}
 						}
 					}
 				}
@@ -1794,13 +1822,15 @@ void FLevelEditorViewportClient::ProcessClick(FSceneView& View, HHitProxy* HitPr
 			// We want to process the click on the component only if:
 			// 1. The actor clicked is already selected
 			// 2. The actor selected is the only actor selected
-			// 3. No components are already selected and the click was a double click
-			// 4. OR, a component is already selected and the click was NOT a double click
+			// 3. The actor selected is blueprintable
+			// 4. No components are already selected and the click was a double click
+			// 5. OR, a component is already selected and the click was NOT a double click
 			const bool bActorAlreadySelectedExclusively = GEditor->GetSelectedActors()->IsSelected(ActorHitProxy->Actor) && ( GEditor->GetSelectedActorCount() == 1 );
+			const bool bActorIsBlueprintable = FKismetEditorUtilities::CanCreateBlueprintOfClass(ActorHitProxy->Actor->GetClass());
 			const bool bComponentAlreadySelected = GEditor->GetSelectedComponentCount() > 0;
 			const bool bWasDoubleClick = ( Click.GetEvent() == IE_DoubleClick );
 
-			const bool bSelectComponent = bActorAlreadySelectedExclusively && (bComponentAlreadySelected != bWasDoubleClick);
+			const bool bSelectComponent = bActorAlreadySelectedExclusively && bActorIsBlueprintable && (bComponentAlreadySelected != bWasDoubleClick);
 
 			if (bSelectComponent && GetDefault<UEditorExperimentalSettings>()->bInWorldBPEditing)
 			{
diff --git a/Engine/Source/Editor/UnrealEd/Public/Kismet2/ComponentEditorUtils.h b/Engine/Source/Editor/UnrealEd/Public/Kismet2/ComponentEditorUtils.h
--- a/Engine/Source/Editor/UnrealEd/Public/Kismet2/ComponentEditorUtils.h
+++ b/Engine/Source/Editor/UnrealEd/Public/Kismet2/ComponentEditorUtils.h
@@ -30,10 +30,24 @@ public:
 	static FString GenerateValidVariableNameFromAsset(UObject* Asset, AActor* ComponentOwner);
 
 	/**
-	* Ensures that the selection override delegate is bound for the supplied component
-	* This includes any attached editor-only primitive components (such as billboard visualizers)
-	*/
-	static void BindComponentSelectionOverride(USceneComponent* SceneComponent);
+	 * Ensures that the selection override delegate is properly bound for the supplied component
+	 * This includes any attached editor-only primitive components (such as billboard visualizers)
+	 * 
+	 * @param SceneComponent The component to set the selection override for
+	 * @param bBind Whether the override should be bound
+	 */
+	static void BindComponentSelectionOverride(USceneComponent* SceneComponent, bool bBind);
+
+	/**
+	 * Attempts to apply a material to a component at the specified slot.
+	 *
+	 * @param SceneComponent The component to which we should attempt to apply the material
+	 * @param MaterialToApply The material to apply to the component
+	 * @param OptionalMaterialSlot The material slot on the component to which the material should be applied. -1 to apply to all slots on the component.
+	 *
+	 * @return	True if the material was successfully applied to the component.
+	 */
+	static bool AttemptApplyMaterialToComponent( USceneComponent* SceneComponent, UMaterialInterface* MaterialToApply, int32 OptionalMaterialSlot = -1 );
 
 	struct FTransformData
 	{
diff --git a/Engine/Source/Runtime/Engine/Private/ActorEditorUtils.cpp b/Engine/Source/Runtime/Engine/Private/ActorEditorUtils.cpp
--- a/Engine/Source/Runtime/Engine/Private/ActorEditorUtils.cpp
+++ b/Engine/Source/Runtime/Engine/Private/ActorEditorUtils.cpp
@@ -35,17 +35,14 @@ namespace FActorEditorUtils
 
 	void GetEditableComponents( const AActor* InActor, TArray<UActorComponent*>& OutEditableComponents )
 	{
-		for( TFieldIterator<UObjectProperty> PropIt(InActor->GetClass(), EFieldIteratorFlags::IncludeSuper); PropIt; ++PropIt )
+		// If a component wasn't created by a construction script (i.e. it's native or an instance), it's editable
+		TInlineComponentArray<UActorComponent*> InstanceComponents;
+		InActor->GetComponents(InstanceComponents);
+		for (auto Component : InstanceComponents)
 		{
-			UObjectProperty* ObjectProp = *PropIt;
-			if( ObjectProp->HasAnyPropertyFlags(CPF_Edit) )
+			if (Component->CreationMethod != EComponentCreationMethod::ConstructionScript)
 			{
-				UObject* ObjPtr = ObjectProp->GetObjectPropertyValue(ObjectProp->ContainerPtrToValuePtr<void>(InActor));
-
-				if( ObjPtr && ObjPtr->IsA<UActorComponent>() )
-				{
-					OutEditableComponents.Add( CastChecked<UActorComponent>( ObjPtr ) );
-				}
+				OutEditableComponents.Add(Component);
 			}
 		}
 	}
