[INTEGRATE] Change 2426302 by [REDACTED_PERSON][REDACTED_HANDLE] on 2015/01/30 15:48:11

	[UE-8767] Reruning BP-Component owner instances' construction scripts (so they update from component updates).

diff --git a/Engine/Source/Editor/UnrealEd/Private/Kismet2/KismetReinstanceUtilities.cpp b/Engine/Source/Editor/UnrealEd/Private/Kismet2/KismetReinstanceUtilities.cpp
--- a/Engine/Source/Editor/UnrealEd/Private/Kismet2/KismetReinstanceUtilities.cpp
+++ b/Engine/Source/Editor/UnrealEd/Private/Kismet2/KismetReinstanceUtilities.cpp
@@ -346,6 +346,11 @@ void FBlueprintCompileReinstancer::ReinstanceObjects(bool bAlwaysReinstance)
 		{
 			ReplaceInstancesOfClass(DuplicatedClass, ClassToReinstance, OriginalCDO, &ObjectsThatShouldUseOldStuff);
 		}
+		else if (ClassToReinstance->IsChildOf<UActorComponent>())
+		{
+			// ReplaceInstancesOfClass() handles this itself, if we had to re-instance
+			ReconstructOwnerInstances(ClassToReinstance);
+		}
 	
 		{ 
 			BP_SCOPED_COMPILER_EVENT_STAT(EKismetReinstancerStats_RecompileChildClasses);
@@ -642,6 +647,9 @@ void FBlueprintCompileReinstancer::ReplaceInstancesOfClass(UClass* OldClass, UCl
 	// A list of objects (e.g. Blueprints) that potentially have editors open that we need to refresh
 	TArray<UObject*> PotentialEditorsForRefreshing;
 
+	// A list of component owners that need their construction scripts re-ran (because a component of theirs has been reinstanced)
+	TSet<AActor*> OwnersToReconstruct;
+
 	// Set global flag to let system know we are reconstructing blueprint instances
 	TGuardValue<bool> GuardTemplateNameFlag(GIsReconstructingBlueprintInstances, true);
 
@@ -805,6 +813,11 @@ void FBlueprintCompileReinstancer::ReplaceInstancesOfClass(UClass* OldClass, UCl
 						{
 							PotentialEditorsForRefreshing.AddUnique(OwningActor->GetClass()->ClassGeneratedBy);
 						}
+
+						// we need to keep track of actor instances that need 
+						// their construction scripts re-ran (since we've just 
+						// replaced a component they own)
+						OwnersToReconstruct.Add(OwningActor);
 					}
 				}
 			}
@@ -874,6 +887,44 @@ void FBlueprintCompileReinstancer::ReplaceInstancesOfClass(UClass* OldClass, UCl
 			}
 		}
 	}
+
+	// in the case where we're replacing component instances, we need to make 
+	// sure to re-run their owner's construction scripts
+	for (AActor* ActorInstance : OwnersToReconstruct)
+	{
+		ActorInstance->RerunConstructionScripts();
+	}
+}
+
+void FBlueprintCompileReinstancer::ReconstructOwnerInstances(TSubclassOf<UActorComponent> ComponentClass)
+{
+	if (ComponentClass == nullptr)
+	{
+		return;
+	}
+
+	TArray<UObject*> ComponentInstances;
+	GetObjectsOfClass(ComponentClass, ComponentInstances, /*bIncludeDerivedClasses =*/false);
+
+	TSet<AActor*> OwnerInstances;
+	for (UObject* ComponentObj : ComponentInstances)
+	{
+	
+		UActorComponent* Component = CastChecked<UActorComponent>(ComponentObj);
+			
+		if (AActor* OwningActor = Component->GetOwner())
+		{
+			// we don't just rerun construction here, because we could end up 
+			// doing it twice for the same actor (if it had multiple components 
+			// of this kind), so we put that off as a secondary pass
+			OwnerInstances.Add(OwningActor);
+		}
+	}
+
+	for (AActor* ComponentOwner : OwnerInstances)
+	{
+		ComponentOwner->RerunConstructionScripts();
+	}
 }
 
 void FBlueprintCompileReinstancer::VerifyReplacement()
diff --git a/Engine/Source/Editor/UnrealEd/Public/Kismet2/KismetReinstanceUtilities.h b/Engine/Source/Editor/UnrealEd/Public/Kismet2/KismetReinstanceUtilities.h
--- a/Engine/Source/Editor/UnrealEd/Public/Kismet2/KismetReinstanceUtilities.h
+++ b/Engine/Source/Editor/UnrealEd/Public/Kismet2/KismetReinstanceUtilities.h
@@ -67,6 +67,15 @@ public:
 
 	/** Worker function to replace all instances of OldClass with a new instance of NewClass */
 	static void ReplaceInstancesOfClass(UClass* OldClass, UClass* NewClass, UObject* OriginalCDO = NULL, TSet<UObject*>* ObjectsThatShouldUseOldStuff = NULL);
+
+	/**
+	 * When re-instancing a component, we have to make sure all instance owners' 
+	 * construction scripts are re-ran (in-case modifying the component alters 
+	 * the construction of the actor).
+	 * 
+	 * @param  ComponentClass    Identifies the component that was altered (used to find all its instances, and thusly all instance owners).
+	 */
+	static void ReconstructOwnerInstances(TSubclassOf<UActorComponent> ComponentClass);
 	
 	/** Verify that all instances of the duplicated class have been replaced and collected */
 	void VerifyReplacement();
