#ANIM: Add status of importing animation -

- Helps to identify why it's taking so long

diff --git a/Engine/Source/Editor/UnrealEd/Private/Fbx/FbxMainImport.cpp b/Engine/Source/Editor/UnrealEd/Private/Fbx/FbxMainImport.cpp
--- a/Engine/Source/Editor/UnrealEd/Private/Fbx/FbxMainImport.cpp
+++ b/Engine/Source/Editor/UnrealEd/Private/Fbx/FbxMainImport.cpp
@@ -592,6 +592,9 @@ bool FFbxImporter::OpenFile(FString Filename, bool bParseStatistics, bool bForSc
 		return false;
 	}
 
+	GWarn->BeginSlowTask( LOCTEXT("OpeningFile", "Reading File"), true);
+	GWarn->StatusForceUpdate(20, 100, LOCTEXT("OpeningFile", "Reading File"));
+
 	int32 SDKMajor,  SDKMinor,  SDKRevision;
 
 	// Create an importer.
@@ -609,6 +612,8 @@ bool FFbxImporter::OpenFile(FString Filename, bool bParseStatistics, bool bForSc
 	
 	const bool bImportStatus = Importer->Initialize(TCHAR_TO_UTF8(*Filename));
 
+	GWarn->StatusForceUpdate(100, 100, LOCTEXT("OpeningFile", "Reading File"));
+	GWarn->EndSlowTask();
 	if( !bImportStatus )  // Problem with the file to be imported
 	{
 		UE_LOG(LogFbx, Error,TEXT("Call to FbxImporter::Initialize() failed."));
diff --git a/Engine/Source/Editor/UnrealEd/Private/SkeletalMeshEdit.cpp b/Engine/Source/Editor/UnrealEd/Private/SkeletalMeshEdit.cpp
--- a/Engine/Source/Editor/UnrealEd/Private/SkeletalMeshEdit.cpp
+++ b/Engine/Source/Editor/UnrealEd/Private/SkeletalMeshEdit.cpp
@@ -20,8 +20,6 @@ UAnimSequence * UEditorEngine::ImportFbxAnimation( USkeleton* Skeleton, UObject*
 
 	UAnimSequence * NewAnimation=NULL;
 
-	GWarn->BeginSlowTask( NSLOCTEXT("UnrealEd", "ImportingFbxAnimations", "Importing FBX animations"), true );
-
 	UnFbx::FFbxImporter* FFbxImporter = UnFbx::FFbxImporter::GetInstance();
 
 	const bool bPrevImportMorph = FFbxImporter->ImportOptions->bImportMorph;
@@ -47,7 +45,6 @@ UAnimSequence * UEditorEngine::ImportFbxAnimation( USkeleton* Skeleton, UObject*
 			FFbxImporter->AddTokenizedErrorMessage(FTokenizedMessage::Create(EMessageSeverity::Error, FText::Format(LOCTEXT("Error_CouldNotFindFbxTrack", "Mesh contains {0} bone as root but animation doesn't contain the root track.\nImport failed."), FText::FromName(Skeleton->GetReferenceSkeleton().GetBoneName(0)))), FFbxErrors::Animation_CouldNotFindRootTrack);
 
 			FFbxImporter->ReleaseScene();
-			GWarn->EndSlowTask();
 			return NULL;
 		}
 
@@ -89,7 +86,6 @@ UAnimSequence * UEditorEngine::ImportFbxAnimation( USkeleton* Skeleton, UObject*
 
 	FFbxImporter->ImportOptions->bImportMorph = bPrevImportMorph;
 	FFbxImporter->ReleaseScene();
-	GWarn->EndSlowTask();
 
 	return NewAnimation;
 }
@@ -98,7 +94,7 @@ bool UEditorEngine::ReimportFbxAnimation( USkeleton* Skeleton, UAnimSequence* An
 {
 	check(Skeleton);
 
-	GWarn->BeginSlowTask( NSLOCTEXT("UnrealEd", "ImportingFbxAnimations", "Importing FBX animations"), true );
+	GWarn->BeginSlowTask( LOCTEXT("ImportingFbxAnimations", "Importing FBX animations"), true );
 
 	UnFbx::FFbxImporter* FbxImporter = UnFbx::FFbxImporter::GetInstance();
 	// logger for all error/warnings
@@ -898,7 +894,7 @@ namespace AnimationTransformDebug
 					{
 						UnFbx::FFbxImporter* FFbxImporter = UnFbx::FFbxImporter::GetInstance();
 						// now print information - it doesn't match well, find out what it is
-						FFbxImporter->AddTokenizedErrorMessage(FTokenizedMessage::Create(EMessageSeverity::Warning, FText::Format(NSLOCTEXT("UnrealEd", "FBXImport_TransformError", "Imported bone transform is different from original. Please check Output Log to see detail of error. "),
+						FFbxImporter->AddTokenizedErrorMessage(FTokenizedMessage::Create(EMessageSeverity::Warning, FText::Format(LOCTEXT("FBXImport_TransformError", "Imported bone transform is different from original. Please check Output Log to see detail of error. "),
 							FText::FromName(Data.BoneName), FText::AsNumber(Data.BoneIndex), FText::FromString(Data.SourceGlobalTransform[Key].ToString()), FText::FromString(GlobalTransform.ToString()))), FFbxErrors::Animation_TransformError);
 
 						bShouldOutputToMessageLog = false;
@@ -994,9 +990,10 @@ bool UnFbx::FFbxImporter::ImportAnimation(USkeleton* Skeleton, UAnimSequence * D
 	}
 
 	//
-	// shape animation START
+	// import blend shape curves
 	//
 	{
+		GWarn->BeginSlowTask( LOCTEXT("BeginImportMorphTargetCurves", "Importing Morph Target Curves"), true);
 		for ( int32 NodeIndex = 0; NodeIndex < NodeArray.Num(); NodeIndex++ )
 		{
 			// consider blendshape animation curve
@@ -1029,6 +1026,10 @@ bool UnFbx::FFbxImporter::ImportAnimation(USkeleton* Skeleton, UAnimSequence * D
 							FbxAnimCurve* Curve = Geometry->GetShapeChannel(BlendShapeIndex, ChannelIndex, (FbxAnimLayer*)CurAnimStack->GetMember(0));
 							if (Curve && Curve->KeyGetCount() > 0)
 							{
+								FFormatNamedArguments Args;
+								Args.Add(TEXT("BlendShape"), FText::FromString(ChannelName));
+								const FText StatusUpate = FText::Format(LOCTEXT("ImportingMorphTargetCurvesDetail", "Importing Morph Target Curves [{BlendShape}]"), Args);
+								GWarn->StatusUpdate(NodeIndex + 1, NodeArray.Num(), StatusUpate);
 								// now see if we have one already exists. If so, just overwrite that. if not, add new one. 
 								ImportCurveToAnimSequence(DestSeq, *ChannelName, Curve,  ACF_DrivesMorphTarget | ACF_TriggerEvent, AnimTimeSpan, 0.01f /** for some reason blend shape values are coming as 100 scaled **/);
 							}
@@ -1037,17 +1038,17 @@ bool UnFbx::FFbxImporter::ImportAnimation(USkeleton* Skeleton, UAnimSequence * D
 				}
 			}
 		}
+		GWarn->EndSlowTask();
 	}
 
-	//
-	//shape animation END
-	//
-
 	// 
 	// importing custom attribute START
 	//
 	if (ImportOptions->bImportCustomAttribute)
 	{
+		GWarn->BeginSlowTask( LOCTEXT("BeginImportMorphTargetCurves", "Importing Custom Attirubte Curves"), true);
+		const int32 TotalLinks = SortedLinks.Num();
+		int32 CurLinkIndex=0;
 		for(auto Node: SortedLinks)
 		{
 			FbxProperty Property = Node->GetFirstProperty();
@@ -1069,14 +1070,22 @@ bool UnFbx::FFbxImporter::ImportAnimation(USkeleton* Skeleton, UAnimSequence * D
 
 						if (AnimCurve)
 						{
+							FString FinalCurveName;
 							if (TotalCount == 1)
 							{
-								ImportCurveToAnimSequence(DestSeq, CurveName, AnimCurve,  ACF_DefaultCurve, AnimTimeSpan);
+								FinalCurveName = CurveName;
 							}
 							else
 							{
-								ImportCurveToAnimSequence(DestSeq, CurveName + "_" + ChannelName, AnimCurve,  ACF_DefaultCurve, AnimTimeSpan);
+								FinalCurveName = CurveName + "_" + ChannelName;
 							}
+
+							FFormatNamedArguments Args;
+							Args.Add(TEXT("CurveName"), FText::FromString(FinalCurveName));
+							const FText StatusUpate = FText::Format(LOCTEXT("ImportingCustomAttributeCurvesDetail", "Importing Custom Attribute [{CurveName}]"), Args);
+							GWarn->StatusUpdate(CurLinkIndex + 1, TotalLinks, StatusUpate);
+
+							ImportCurveToAnimSequence(DestSeq, FinalCurveName, AnimCurve,  ACF_DefaultCurve, AnimTimeSpan);
 						}
 											
 					}
@@ -1084,163 +1093,191 @@ bool UnFbx::FFbxImporter::ImportAnimation(USkeleton* Skeleton, UAnimSequence * D
 
 				Property = Node->GetNextProperty(Property); 
 			}
+
+			CurLinkIndex++;
 		}
+
+		GWarn->EndSlowTask();
 	}
 
 	// importing custom attribute END
-	const bool bSourceDataExists = (DestSeq->SourceRawAnimationData.Num() > 0);
-	TArray<struct FRawAnimSequenceTrack>& RawAnimationData = bSourceDataExists? DestSeq->SourceRawAnimationData : DestSeq->RawAnimationData;
-	DestSeq->TrackToSkeletonMapTable.Empty();
-	DestSeq->AnimationTrackNames.Empty();
-	RawAnimationData.Empty();
-
-	TArray<FName> FbxRawBoneNames;
-	FillAndVerifyBoneNames(Skeleton, SortedLinks, FbxRawBoneNames, FileName);
 
+	const bool bSourceDataExists = (DestSeq->SourceRawAnimationData.Num() > 0);
 	TArray<AnimationTransformDebug::FAnimationTransformDebugData> TransformDebugData;
+	int32 TotalNumKeys = 0;
 	const FReferenceSkeleton& RefSkeleton = Skeleton->GetReferenceSkeleton();
 
-	UnFbx::FFbxImporter* FbxImporter = UnFbx::FFbxImporter::GetInstance();
+	// import animation
+	{
+		GWarn->BeginSlowTask( LOCTEXT("BeginImportAnimation", "Importing Animation"), true);
 
-	const bool bPreserveLocalTransform = FbxImporter->GetImportOptions()->bPreserveLocalTransform;
+		TArray<struct FRawAnimSequenceTrack>& RawAnimationData = bSourceDataExists? DestSeq->SourceRawAnimationData : DestSeq->RawAnimationData;
+		DestSeq->TrackToSkeletonMapTable.Empty();
+		DestSeq->AnimationTrackNames.Empty();
+		RawAnimationData.Empty();
 
-	// Build additional transform matrix
-	UFbxAnimSequenceImportData* TemplateData = Cast<UFbxAnimSequenceImportData>(DestSeq->AssetImportData);
-	FbxAMatrix FbxAddedMatrix;
-	BuildFbxMatrixForImportTransform(FbxAddedMatrix, TemplateData);
-	FMatrix AddedMatrix = Converter.ConvertMatrix(FbxAddedMatrix);
+		TArray<FName> FbxRawBoneNames;
+		FillAndVerifyBoneNames(Skeleton, SortedLinks, FbxRawBoneNames, FileName);
 
-	int32 TotalNumKeys = 0;
-	const int32 NumSamplingKeys = FMath::FloorToInt(AnimTimeSpan.GetDuration().GetSecondDouble() * ResampleRate);
-	const FbxTime TimeIncrement = AnimTimeSpan.GetDuration() / NumSamplingKeys;
-	for(int32 SourceTrackIdx = 0; SourceTrackIdx < FbxRawBoneNames.Num(); ++SourceTrackIdx)
-	{
-		int32 NumKeysForTrack = 0;
+		UnFbx::FFbxImporter* FbxImporter = UnFbx::FFbxImporter::GetInstance();
 
-		// see if it's found in Skeleton
-		FName BoneName = FbxRawBoneNames[SourceTrackIdx];
-		int32 BoneTreeIndex = RefSkeleton.FindBoneIndex(BoneName);
+		const bool bPreserveLocalTransform = FbxImporter->GetImportOptions()->bPreserveLocalTransform;
 
-		if (BoneTreeIndex!=INDEX_NONE)
+		// Build additional transform matrix
+		UFbxAnimSequenceImportData* TemplateData = Cast<UFbxAnimSequenceImportData>(DestSeq->AssetImportData);
+		FbxAMatrix FbxAddedMatrix;
+		BuildFbxMatrixForImportTransform(FbxAddedMatrix, TemplateData);
+		FMatrix AddedMatrix = Converter.ConvertMatrix(FbxAddedMatrix);
+
+		const int32 NumSamplingKeys = FMath::FloorToInt(AnimTimeSpan.GetDuration().GetSecondDouble() * ResampleRate);
+		const FbxTime TimeIncrement = AnimTimeSpan.GetDuration() / NumSamplingKeys;
+		for(int32 SourceTrackIdx = 0; SourceTrackIdx < FbxRawBoneNames.Num(); ++SourceTrackIdx)
 		{
-			bool bSuccess = true;
+			int32 NumKeysForTrack = 0;
+
+			// see if it's found in Skeleton
+			FName BoneName = FbxRawBoneNames[SourceTrackIdx];
+			int32 BoneTreeIndex = RefSkeleton.FindBoneIndex(BoneName);
+
+			// update status
+			FFormatNamedArguments Args;
+			Args.Add(TEXT("TrackName"), FText::FromName(BoneName));
+			Args.Add(TEXT("TotalKey"), FText::AsNumber(NumSamplingKeys));
+			Args.Add(TEXT("TrackIndex"), FText::AsNumber(SourceTrackIdx+1));
+			Args.Add(TEXT("TotalTracks"), FText::AsNumber(FbxRawBoneNames.Num()));
+			const FText StatusUpate = FText::Format(LOCTEXT("ImportingAnimTrackDetail", "Importing Animation Track [{TrackName}] ({TrackIndex}/{TotalTracks}) - TotalKey {TotalKey}"), Args);
+			GWarn->StatusForceUpdate(SourceTrackIdx + 1, FbxRawBoneNames.Num(), StatusUpate);
+
+			if (BoneTreeIndex!=INDEX_NONE)
+			{
+				bool bSuccess = true;
 
-			FRawAnimSequenceTrack RawTrack;
-			RawTrack.PosKeys.Empty();
-			RawTrack.RotKeys.Empty();
-			RawTrack.ScaleKeys.Empty();
+				FRawAnimSequenceTrack RawTrack;
+				RawTrack.PosKeys.Empty();
+				RawTrack.RotKeys.Empty();
+				RawTrack.ScaleKeys.Empty();
 
-			AnimationTransformDebug::FAnimationTransformDebugData NewDebugData;
+				AnimationTransformDebug::FAnimationTransformDebugData NewDebugData;
 
-			FbxNode* Link = SortedLinks[SourceTrackIdx];
-			FbxNode * LinkParent = Link->GetParent();
+				FbxNode* Link = SortedLinks[SourceTrackIdx];
+				FbxNode * LinkParent = Link->GetParent();
 			
-			for(FbxTime CurTime = AnimTimeSpan.GetStart(); CurTime <= AnimTimeSpan.GetStop(); CurTime += TimeIncrement)
-			{
-				// save global trasnform
-				FbxAMatrix GlobalMatrix = Link->EvaluateGlobalTransform(CurTime);
-				// we'd like to verify this before going to Transform. 
-				// currently transform has tons of NaN check, so it will crash there
-				FMatrix GlobalUEMatrix = Converter.ConvertMatrix(GlobalMatrix);
-				if (GlobalUEMatrix.ContainsNaN())
+				for(FbxTime CurTime = AnimTimeSpan.GetStart(); CurTime <= AnimTimeSpan.GetStop(); CurTime += TimeIncrement)
 				{
-					bSuccess = false;
-					AddTokenizedErrorMessage(FTokenizedMessage::Create(EMessageSeverity::Error, FText::Format(LOCTEXT("Error_InvalidTransform",
-						"Track {0} contains invalid transform. Could not import the track."), FText::FromName(BoneName))), FFbxErrors::Animation_TransformError);
-					break;
-				}
+					// save global trasnform
+					FbxAMatrix GlobalMatrix = Link->EvaluateGlobalTransform(CurTime);
+					// we'd like to verify this before going to Transform. 
+					// currently transform has tons of NaN check, so it will crash there
+					FMatrix GlobalUEMatrix = Converter.ConvertMatrix(GlobalMatrix);
+					if (GlobalUEMatrix.ContainsNaN())
+					{
+						bSuccess = false;
+						AddTokenizedErrorMessage(FTokenizedMessage::Create(EMessageSeverity::Error, FText::Format(LOCTEXT("Error_InvalidTransform",
+							"Track {0} contains invalid transform. Could not import the track."), FText::FromName(BoneName))), FFbxErrors::Animation_TransformError);
+						break;
+					}
 
-				FTransform GlobalTransform =  Converter.ConvertTransform(GlobalMatrix);
-				if (GlobalTransform.ContainsNaN())
-				{
-					bSuccess = false;
-					AddTokenizedErrorMessage(FTokenizedMessage::Create(EMessageSeverity::Error, FText::Format(LOCTEXT("Error_InvalidUnrealTransform",
-										"Track {0} did not yeild valid transform. Please report this to animation team."), FText::FromName(BoneName))), FFbxErrors::Animation_TransformError);
-					break;
-				}
+					FTransform GlobalTransform =  Converter.ConvertTransform(GlobalMatrix);
+					if (GlobalTransform.ContainsNaN())
+					{
+						bSuccess = false;
+						AddTokenizedErrorMessage(FTokenizedMessage::Create(EMessageSeverity::Error, FText::Format(LOCTEXT("Error_InvalidUnrealTransform",
+											"Track {0} did not yeild valid transform. Please report this to animation team."), FText::FromName(BoneName))), FFbxErrors::Animation_TransformError);
+						break;
+					}
 
-				// debug data, including import transformation
-				FTransform AddedTransform(AddedMatrix);
-				NewDebugData.SourceGlobalTransform.Add(GlobalTransform * AddedTransform);
+					// debug data, including import transformation
+					FTransform AddedTransform(AddedMatrix);
+					NewDebugData.SourceGlobalTransform.Add(GlobalTransform * AddedTransform);
 
-				FTransform LocalTransform;
-				if( !bPreserveLocalTransform && LinkParent)
-				{
-					// I can't rely on LocalMatrix. I need to recalculate quaternion/scale based on global transform if Parent exists
-					FbxAMatrix ParentGlobalMatrix = Link->GetParent()->EvaluateGlobalTransform(CurTime);
-					FTransform ParentGlobalTransform =  Converter.ConvertTransform(ParentGlobalMatrix);
+					FTransform LocalTransform;
+					if( !bPreserveLocalTransform && LinkParent)
+					{
+						// I can't rely on LocalMatrix. I need to recalculate quaternion/scale based on global transform if Parent exists
+						FbxAMatrix ParentGlobalMatrix = Link->GetParent()->EvaluateGlobalTransform(CurTime);
+						FTransform ParentGlobalTransform =  Converter.ConvertTransform(ParentGlobalMatrix);
 
-					LocalTransform = GlobalTransform.GetRelativeTransform(ParentGlobalTransform);
-					NewDebugData.SourceParentGlobalTransform.Add(ParentGlobalTransform);
-				} 
-				else
-				{
-					FbxAMatrix& LocalMatrix = Link->EvaluateLocalTransform(CurTime); 
-					FbxVector4 NewLocalT = LocalMatrix.GetT();
-					FbxVector4 NewLocalS = LocalMatrix.GetS();
-					FbxQuaternion NewLocalQ = LocalMatrix.GetQ();
+						LocalTransform = GlobalTransform.GetRelativeTransform(ParentGlobalTransform);
+						NewDebugData.SourceParentGlobalTransform.Add(ParentGlobalTransform);
+					} 
+					else
+					{
+						FbxAMatrix& LocalMatrix = Link->EvaluateLocalTransform(CurTime); 
+						FbxVector4 NewLocalT = LocalMatrix.GetT();
+						FbxVector4 NewLocalS = LocalMatrix.GetS();
+						FbxQuaternion NewLocalQ = LocalMatrix.GetQ();
 
-					LocalTransform.SetTranslation(Converter.ConvertPos(NewLocalT));
-					LocalTransform.SetScale3D(Converter.ConvertScale(NewLocalS));
-					LocalTransform.SetRotation(Converter.ConvertRotToQuat(NewLocalQ));
+						LocalTransform.SetTranslation(Converter.ConvertPos(NewLocalT));
+						LocalTransform.SetScale3D(Converter.ConvertScale(NewLocalS));
+						LocalTransform.SetRotation(Converter.ConvertRotToQuat(NewLocalQ));
 
-					NewDebugData.SourceParentGlobalTransform.Add(FTransform::Identity);
-				}
+						NewDebugData.SourceParentGlobalTransform.Add(FTransform::Identity);
+					}
 
-				if(TemplateData && BoneTreeIndex == 0)
-				{
-					// If we found template data earlier, apply the import transform matrix to
-					// the root track.
-					LocalTransform.SetFromMatrix(LocalTransform.ToMatrixWithScale() * AddedMatrix);
-				}
+					if(TemplateData && BoneTreeIndex == 0)
+					{
+						// If we found template data earlier, apply the import transform matrix to
+						// the root track.
+						LocalTransform.SetFromMatrix(LocalTransform.ToMatrixWithScale() * AddedMatrix);
+					}
 
-				if (LocalTransform.ContainsNaN())
-				{
-					bSuccess = false;
-					AddTokenizedErrorMessage(FTokenizedMessage::Create(EMessageSeverity::Error, FText::Format(LOCTEXT("Error_InvalidUnrealLocalTransform",
-										"Track {0} did not yeild valid local transform. Please report this to animation team."), FText::FromName(BoneName))), FFbxErrors::Animation_TransformError);
-					break;
-				}
+					if (LocalTransform.ContainsNaN())
+					{
+						bSuccess = false;
+						AddTokenizedErrorMessage(FTokenizedMessage::Create(EMessageSeverity::Error, FText::Format(LOCTEXT("Error_InvalidUnrealLocalTransform",
+											"Track {0} did not yeild valid local transform. Please report this to animation team."), FText::FromName(BoneName))), FFbxErrors::Animation_TransformError);
+						break;
+					}
 
-				RawTrack.ScaleKeys.Add(LocalTransform.GetScale3D());
-				RawTrack.PosKeys.Add(LocalTransform.GetTranslation());
-				RawTrack.RotKeys.Add(LocalTransform.GetRotation());
+					RawTrack.ScaleKeys.Add(LocalTransform.GetScale3D());
+					RawTrack.PosKeys.Add(LocalTransform.GetTranslation());
+					RawTrack.RotKeys.Add(LocalTransform.GetRotation());
 
-				NewDebugData.RecalculatedLocalTransform.Add(LocalTransform);
-				++NumKeysForTrack;
-			}
+					NewDebugData.RecalculatedLocalTransform.Add(LocalTransform);
+					++NumKeysForTrack;
+				}
 
-			if (bSuccess)
-			{
-				//add new track
-				int32 NewTrackIdx = RawAnimationData.Add(RawTrack);
-				DestSeq->AnimationTrackNames.Add(BoneName);
+				if (bSuccess)
+				{
+					//add new track
+					int32 NewTrackIdx = RawAnimationData.Add(RawTrack);
+					DestSeq->AnimationTrackNames.Add(BoneName);
 
-				NewDebugData.SetTrackData(NewTrackIdx, BoneTreeIndex, BoneName);
+					NewDebugData.SetTrackData(NewTrackIdx, BoneTreeIndex, BoneName);
 
-				// add mapping to skeleton bone track
-				DestSeq->TrackToSkeletonMapTable.Add(FTrackToSkeletonMap(BoneTreeIndex));
-				TransformDebugData.Add(NewDebugData);
+					// add mapping to skeleton bone track
+					DestSeq->TrackToSkeletonMapTable.Add(FTrackToSkeletonMap(BoneTreeIndex));
+					TransformDebugData.Add(NewDebugData);
+				}
 			}
+
+			TotalNumKeys = FMath::Max( TotalNumKeys, NumKeysForTrack );
 		}
 
-		TotalNumKeys = FMath::Max( TotalNumKeys, NumKeysForTrack );
+		DestSeq->NumFrames = TotalNumKeys;
+		GWarn->EndSlowTask();
 	}
 
-	DestSeq->NumFrames = TotalNumKeys;
-	// if source data exists, you should bake it to Raw to apply
-	if(bSourceDataExists)
+	// compress animation
 	{
-		DestSeq->BakeTrackCurvesToRawAnimation();
-	}
-	else
-	{
-		// otherwise just compress
-		DestSeq->PostProcessSequence();
+		GWarn->BeginSlowTask( LOCTEXT("BeginCompressAnimation", "Compress Animation"), true);
+		GWarn->StatusForceUpdate(1, 1, LOCTEXT("CompressAnimation", "Compressing Animation"));
+		// if source data exists, you should bake it to Raw to apply
+		if(bSourceDataExists)
+		{
+			DestSeq->BakeTrackCurvesToRawAnimation();
+		}
+		else
+		{
+			// otherwise just compress
+			DestSeq->PostProcessSequence();
+		}
+
+		// run debug mode
+		AnimationTransformDebug::OutputAnimationTransformDebugData(TransformDebugData, TotalNumKeys, RefSkeleton);
+		GWarn->EndSlowTask();
 	}
 
-	// run debug mode
-	AnimationTransformDebug::OutputAnimationTransformDebugData(TransformDebugData, TotalNumKeys, RefSkeleton);
 	return true;
 }
 
