[INTEGRATE] Change 2426831 by [REDACTED_PERSON][REDACTED_HANDLE] on 2015/01/30 20:17:16

	Fix for UE-8766, can now add blueprint created component types without opening the blueprint, we listen to the assetregistry's events and make use of the hierarchy information that it stores

diff --git a/Engine/Source/Editor/UnrealEd/Private/ComponentTypeRegistry.cpp b/Engine/Source/Editor/UnrealEd/Private/ComponentTypeRegistry.cpp
new file mode 100644
--- /dev/null
+++ b/Engine/Source/Editor/UnrealEd/Private/ComponentTypeRegistry.cpp
@@ -0,0 +1,323 @@
+// Copyright 1998-2015 Epic Games, Inc. All Rights Reserved.
+
+#include "UnrealEd.h"
+
+#include "AssetRegistryModule.h"
+#include "ComponentTypeRegistry.h"
+#include "EdGraphSchema_K2.h"
+#include "KismetEditorUtilities.h"
+
+#define LOCTEXT_NAMESPACE "ComponentTypeRegistry"
+
+struct FComponentTypeRegistryData
+	: public FTickableEditorObject
+{
+	FComponentTypeRegistryData();
+	void RefreshComponentList();
+
+	/** Implementation of FTickableEditorObject */
+	virtual void Tick(float) override;
+	virtual bool IsTickable() const override { return true; }
+	virtual TStatId GetStatId() const { RETURN_QUICK_DECLARE_CYCLE_STAT(FTypeDatabaseUpdater, STATGROUP_Tickables); }
+	/** End implementation of FTickableEditorObject */
+
+	TArray<FComponentClassComboEntryPtr> ComponentClassList;
+	TArray<FAssetData> PendingAssetData;
+	FOnComponentTypeListChanged ComponentListChanged;
+};
+
+static const FString CommonClassGroup(TEXT("Common"));
+// This has to stay in sync with logic in FKismetCompilerContext::FinishCompilingClass
+static const FString BlueprintComponents(TEXT("Custom"));
+
+static void AddBasicShapeComponents(TArray<FComponentClassComboEntryPtr>& SortedClassList)
+{
+	FString BasicShapesHeading = LOCTEXT("BasicShapesHeading", "Basic Shapes").ToString();
+
+	const auto OnBasicShapeCreated = [](UActorComponent* Component)
+	{
+		UStaticMeshComponent* SMC = Cast<UStaticMeshComponent>(Component);
+		if (SMC)
+		{
+			SMC->SetMaterial(0, LoadObject<UMaterial>(nullptr, TEXT("/Engine/BasicShapes/BasicShapeMaterial.BasicShapeMaterial")));
+		}
+	};
+
+	{
+		FComponentEntryCustomizationArgs Args;
+		Args.AssetOverride = LoadObject<UStaticMesh>(nullptr, TEXT("/Engine/BasicShapes/Cube.Cube"));;
+		Args.OnComponentCreated = FOnComponentCreated::CreateStatic(OnBasicShapeCreated);
+		Args.ComponentNameOverride = LOCTEXT("BasicCubeShapeDisplayName", "Cube").ToString();
+
+		FComponentClassComboEntryPtr NewShape = MakeShareable(new FComponentClassComboEntry(BasicShapesHeading, UStaticMeshComponent::StaticClass(), true, EComponentCreateAction::SpawnExistingClass, Args));
+		SortedClassList.Add(NewShape);
+	}
+
+	{
+		FComponentEntryCustomizationArgs Args;
+		Args.AssetOverride = LoadObject<UStaticMesh>(nullptr, TEXT("/Engine/BasicShapes/Sphere.Sphere"));;
+		Args.OnComponentCreated = FOnComponentCreated::CreateStatic(OnBasicShapeCreated);
+		Args.ComponentNameOverride = LOCTEXT("BasicSphereShapeDisplayName", "Sphere").ToString();
+
+		FComponentClassComboEntryPtr NewShape = MakeShareable(new FComponentClassComboEntry(BasicShapesHeading, UStaticMeshComponent::StaticClass(), true, EComponentCreateAction::SpawnExistingClass, Args));
+		SortedClassList.Add(NewShape);
+	}
+
+	{
+		FComponentEntryCustomizationArgs Args;
+		Args.AssetOverride = LoadObject<UStaticMesh>(nullptr, TEXT("/Engine/BasicShapes/Cylinder.Cylinder"));;
+		Args.OnComponentCreated = FOnComponentCreated::CreateStatic(OnBasicShapeCreated);
+		Args.ComponentNameOverride = LOCTEXT("BasicCylinderShapeDisplayName", "Cylinder").ToString();
+
+		FComponentClassComboEntryPtr NewShape = MakeShareable(new FComponentClassComboEntry(BasicShapesHeading, UStaticMeshComponent::StaticClass(), true, EComponentCreateAction::SpawnExistingClass, Args));
+		SortedClassList.Add(NewShape);
+	}
+
+	{
+		FComponentEntryCustomizationArgs Args;
+		Args.AssetOverride = LoadObject<UStaticMesh>(nullptr, TEXT("/Engine/BasicShapes/Cone.Cone"));;
+		Args.OnComponentCreated = FOnComponentCreated::CreateStatic(OnBasicShapeCreated);
+		Args.ComponentNameOverride = LOCTEXT("BasicConeShapeDisplayName", "Cone").ToString();
+
+		FComponentClassComboEntryPtr NewShape = MakeShareable(new FComponentClassComboEntry(BasicShapesHeading, UStaticMeshComponent::StaticClass(), true, EComponentCreateAction::SpawnExistingClass, Args));
+		SortedClassList.Add(NewShape);
+	}
+}
+
+FComponentTypeRegistryData::FComponentTypeRegistryData()
+{
+	const auto HandleAdded = [](const FAssetData& Data, FComponentTypeRegistryData* Parent)
+	{
+		Parent->PendingAssetData.Push(Data);
+	};
+
+	FAssetRegistryModule& AssetRegistryModule = FModuleManager::LoadModuleChecked<FAssetRegistryModule>(TEXT("AssetRegistry"));
+	AssetRegistryModule.Get().OnAssetAdded().AddStatic(HandleAdded, this);
+	AssetRegistryModule.Get().OnAssetRemoved().AddStatic(HandleAdded, this);
+
+	const auto HandleRenamed = [](const FAssetData& Data, const FString&, FComponentTypeRegistryData* Parent)
+	{
+		Parent->PendingAssetData.Push(Data);
+	};
+	AssetRegistryModule.Get().OnAssetRenamed().AddStatic(HandleRenamed, this);
+}
+
+void FComponentTypeRegistryData::RefreshComponentList()
+{
+	ComponentClassList.Empty();
+
+	struct SortComboEntry
+	{
+		bool operator () (const FComponentClassComboEntryPtr& A, const FComponentClassComboEntryPtr& B) const
+		{
+			bool bResult = false;
+
+			// check headings first, if they are the same compare the individual entries
+			int32 HeadingCompareResult = FCString::Stricmp(*A->GetHeadingText(), *B->GetHeadingText());
+			if (HeadingCompareResult == 0)
+			{
+				bResult = FCString::Stricmp(*A->GetClassName(), *B->GetClassName()) < 0;
+			}
+			else if (CommonClassGroup == A->GetHeadingText())
+			{
+				bResult = true;
+			}
+			else if (CommonClassGroup == B->GetHeadingText())
+			{
+				bResult = false;
+			}
+			else
+			{
+				bResult = HeadingCompareResult < 0;
+			}
+
+			return bResult;
+		}
+	};
+
+	const UEdGraphSchema_K2* K2Schema = GetDefault<UEdGraphSchema_K2>();
+
+	if (GetDefault<UEditorExperimentalSettings>()->bScriptableComponentsOnActors)
+	{
+		FString NewComponentsHeading = LOCTEXT("NewComponentsHeading", "New").ToString();
+		// Add new C++ component class
+		FComponentClassComboEntryPtr NewClassHeader = MakeShareable(new FComponentClassComboEntry(NewComponentsHeading));
+		ComponentClassList.Add(NewClassHeader);
+
+		FComponentClassComboEntryPtr NewCPPClass = MakeShareable(new FComponentClassComboEntry(NewComponentsHeading, UActorComponent::StaticClass(), true, EComponentCreateAction::CreateNewCPPClass));
+		ComponentClassList.Add(NewCPPClass);
+
+		FComponentClassComboEntryPtr NewBPClass = MakeShareable(new FComponentClassComboEntry(NewComponentsHeading, UActorComponent::StaticClass(), true, EComponentCreateAction::CreateNewBlueprintClass));
+		ComponentClassList.Add(NewBPClass);
+
+		FComponentClassComboEntryPtr NewSeparator(new FComponentClassComboEntry());
+		ComponentClassList.Add(NewSeparator);
+	}
+
+	TArray<FComponentClassComboEntryPtr> SortedClassList;
+	TArray<FName> InMemoryClasses;
+	for (TObjectIterator<UClass> It; It; ++It)
+	{
+		UClass* Class = *It;
+		// If this is a subclass of Actor Component, not abstract, and tagged as spawnable from Kismet
+		if (Class->IsChildOf(UActorComponent::StaticClass()))
+		{
+			InMemoryClasses.Push(Class->GetFName());
+			if (!Class->HasAnyClassFlags(CLASS_Abstract) && Class->HasMetaData(FBlueprintMetadata::MD_BlueprintSpawnableComponent) && !FKismetEditorUtilities::IsClassABlueprintSkeleton(Class)) //@TODO: Fold this logic together with the one in UEdGraphSchema_K2::GetAddComponentClasses
+			{
+				TArray<FString> ClassGroupNames;
+				Class->GetClassGroupNames(ClassGroupNames);
+
+				if (ClassGroupNames.Contains(CommonClassGroup))
+				{
+					FString ClassGroup = CommonClassGroup;
+					FComponentClassComboEntryPtr NewEntry(new FComponentClassComboEntry(ClassGroup, Class, ClassGroupNames.Num() <= 1, EComponentCreateAction::SpawnExistingClass));
+					SortedClassList.Add(NewEntry);
+				}
+				if (ClassGroupNames.Num() && !ClassGroupNames[0].Equals(CommonClassGroup))
+				{
+					const bool bIncludeInFilter = true;
+
+					FString ClassGroup = ClassGroupNames[0];
+					FComponentClassComboEntryPtr NewEntry(new FComponentClassComboEntry(ClassGroup, Class, bIncludeInFilter, EComponentCreateAction::SpawnExistingClass));
+					SortedClassList.Add(NewEntry);
+				}
+			}
+		}
+	}
+
+	AddBasicShapeComponents(SortedClassList);
+
+	{
+		// make sure that we add any user created classes immediately, generally this will not create anything (because assets have not been discovered yet), but asset discovery
+		// should be allowed to take place at any time:
+		FAssetRegistryModule& AssetRegistryModule = FModuleManager::LoadModuleChecked<FAssetRegistryModule>(TEXT("AssetRegistry"));
+		TArray<FName> ClassNames;
+		ClassNames.Add(UActorComponent::StaticClass()->GetFName());
+		TSet<FName> DerivedClassNames;
+		AssetRegistryModule.Get().GetDerivedClassNames(ClassNames, TSet<FName>(), DerivedClassNames);
+
+		const bool bIncludeInFilter = true;
+
+		auto InMemoryClassesSet = TSet<FName>(InMemoryClasses);
+		auto OnDiskClasses = DerivedClassNames.Difference(InMemoryClassesSet);
+
+		// GetAssetsByClass call is a cludget to get the full asset paths for the blueprints we care about, Bob T. thinks 
+		// that the Asset Registry could just keep asset paths:
+		TArray<FAssetData> BlueprintAssetData;
+		AssetRegistryModule.Get().GetAssetsByClass(UBlueprint::StaticClass()->GetFName(), BlueprintAssetData);
+
+		for (auto OnDiskClass : OnDiskClasses)
+		{
+			FName AssetPath;
+			FString FixedString = OnDiskClass.ToString();
+			FixedString.RemoveFromEnd(TEXT("_C"));
+			for (auto Blueprint : BlueprintAssetData)
+			{
+				if (Blueprint.AssetName.ToString() == FixedString)
+				{
+					AssetPath = Blueprint.ObjectPath;
+					break;
+				}
+			}
+
+			FComponentClassComboEntryPtr NewEntry(new FComponentClassComboEntry(BlueprintComponents, FixedString, AssetPath, bIncludeInFilter));
+			SortedClassList.Add(NewEntry);
+		}
+	}
+
+	if (SortedClassList.Num() > 0)
+	{
+		Sort(SortedClassList.GetData(), SortedClassList.Num(), SortComboEntry());
+
+		FString PreviousHeading;
+		for (int32 ClassIndex = 0; ClassIndex < SortedClassList.Num(); ClassIndex++)
+		{
+			FComponentClassComboEntryPtr& CurrentEntry = SortedClassList[ClassIndex];
+
+			const FString& CurrentHeadingText = CurrentEntry->GetHeadingText();
+
+			if (CurrentHeadingText != PreviousHeading)
+			{
+				// This avoids a redundant separator being added to the very top of the list
+				if (ClassIndex > 0)
+				{
+					FComponentClassComboEntryPtr NewSeparator(new FComponentClassComboEntry());
+					ComponentClassList.Add(NewSeparator);
+				}
+				FComponentClassComboEntryPtr NewHeading(new FComponentClassComboEntry(CurrentHeadingText));
+				ComponentClassList.Add(NewHeading);
+
+				PreviousHeading = CurrentHeadingText;
+			}
+
+			ComponentClassList.Add(CurrentEntry);
+		}
+	}
+}
+
+void FComponentTypeRegistryData::Tick(float)
+{
+	bool bRequiresRefresh = false;
+
+	if (PendingAssetData.Num() != 0)
+	{
+		FAssetRegistryModule& AssetRegistryModule = FModuleManager::LoadModuleChecked<FAssetRegistryModule>(TEXT("AssetRegistry"));
+		TArray<FName> ClassNames;
+		ClassNames.Add(UActorComponent::StaticClass()->GetFName());
+		TSet<FName> DerivedClassNames;
+		AssetRegistryModule.Get().GetDerivedClassNames(ClassNames, TSet<FName>(), DerivedClassNames);
+
+		for (auto Asset : PendingAssetData)
+		{
+			const FName BPParentClassName(TEXT("ParentClass"));
+
+			bool FoundBPNotify = false;
+			for (TMap<FName, FString>::TConstIterator TagIt(Asset.TagsAndValues); TagIt; ++TagIt)
+			{
+				FString TagValue = Asset.TagsAndValues.FindRef(BPParentClassName);
+				const FString ObjectPath = FPackageName::ExportTextPathToObjectPath(TagValue);
+				FName ObjectName = FName(*FPackageName::ObjectPathToObjectName(ObjectPath));
+				if (DerivedClassNames.Contains(ObjectName))
+				{
+					bRequiresRefresh = true;
+					break;
+				}
+			}
+		}
+		PendingAssetData.Empty();
+	}
+
+	if (bRequiresRefresh)
+	{
+		RefreshComponentList();
+		ComponentListChanged.Broadcast();
+	}
+}
+
+FComponentTypeRegistry& FComponentTypeRegistry::Get()
+{
+	static FComponentTypeRegistry ComponentRegistry;
+	return ComponentRegistry;
+}
+
+FOnComponentTypeListChanged& FComponentTypeRegistry::SubscribeToComponentList(TArray<FComponentClassComboEntryPtr>*& OutComponentList)
+{
+	check(Data);
+	OutComponentList = &Data->ComponentClassList;
+	return Data->ComponentListChanged;
+}
+
+FOnComponentTypeListChanged& FComponentTypeRegistry::GetOnComponentTypeListChanged()
+{
+	check(Data);
+	return Data->ComponentListChanged;
+}
+
+FComponentTypeRegistry::FComponentTypeRegistry()
+	: Data(nullptr)
+{
+	Data = new FComponentTypeRegistryData();
+	Data->RefreshComponentList();
+}
+
diff --git a/Engine/Source/Editor/UnrealEd/Private/SComponentClassCombo.cpp b/Engine/Source/Editor/UnrealEd/Private/SComponentClassCombo.cpp
--- a/Engine/Source/Editor/UnrealEd/Private/SComponentClassCombo.cpp
+++ b/Engine/Source/Editor/UnrealEd/Private/SComponentClassCombo.cpp
@@ -1,35 +1,42 @@
 // Copyright 1998-2015 Epic Games, Inc. All Rights Reserved.
 
 #include "UnrealEd.h"
-#include "SlateBasics.h"
-#include "SComponentClassCombo.h"
-#include "ComponentAssetBroker.h"
-#include "ClassIconFinder.h"
+
 #include "BlueprintGraphDefinitions.h"
-#include "IDocumentation.h"
-#include "SListViewSelectorDropdownMenu.h"
+#include "ClassIconFinder.h"
+#include "ComponentAssetBroker.h"
+#include "ComponentTypeRegistry.h"
 #include "EditorClassUtils.h"
-#include "SSearchBox.h"
 #include "Engine/Selection.h"
 #include "HotReloadInterface.h"
+#include "IDocumentation.h"
 #include "KismetEditorUtilities.h"
+#include "SComponentClassCombo.h"
+#include "SlateBasics.h"
+#include "SListViewSelectorDropdownMenu.h"
+#include "SSearchBox.h"
 
 #define LOCTEXT_NAMESPACE "ComponentClassCombo"
 
+FString FComponentClassComboEntry::GetClassName() const
+{
+	return ComponentClass != nullptr ? ComponentClass->GetDisplayNameText().ToString() : ComponentName;
+}
 
 void SComponentClassCombo::Construct(const FArguments& InArgs)
 {
 	PrevSelectedIndex = INDEX_NONE;
 	OnComponentClassSelected = InArgs._OnComponentClassSelected;
 
+	FComponentTypeRegistry::Get().SubscribeToComponentList(ComponentClassList).AddRaw(this, &SComponentClassCombo::UpdateComponentClassList);
+
 	UpdateComponentClassList();
-	GenerateFilteredComponentList(FString());
 
 	IHotReloadInterface& HotReloadSupport = FModuleManager::LoadModuleChecked<IHotReloadInterface>("HotReload");
 	HotReloadSupport.OnHotReload().AddSP( this, &SComponentClassCombo::OnProjectHotReloaded );
-	
+
 	SAssignNew(ComponentClassListView, SListView<FComponentClassComboEntryPtr>)
-		.ListItemsSource( &FilteredComponentClassList )
+		.ListItemsSource(ComponentClassList)
 		.OnSelectionChanged( this, &SComponentClassCombo::OnAddComponentSelectionChanged )
 		.OnGenerateRow( this, &SComponentClassCombo::GenerateAddComponentRow )
 		.SelectionMode(ESelectionMode::Single);
@@ -103,6 +110,11 @@ void SComponentClassCombo::Construct(const FArguments& InArgs)
 	SetMenuContentWidgetToFocus( SearchBox );
 }
 
+SComponentClassCombo::~SComponentClassCombo()
+{
+	FComponentTypeRegistry::Get().GetOnComponentTypeListChanged().RemoveAll(this);
+}
+
 void SComponentClassCombo::ClearSelection()
 {
 	SearchBox->SetText(FText::GetEmpty());
@@ -117,18 +129,18 @@ void SComponentClassCombo::GenerateFilteredComponentList(const FString& InSearch
 {
 	if ( InSearchText.IsEmpty() )
 	{
-		FilteredComponentClassList = ComponentClassList;
+		FilteredComponentClassList = *ComponentClassList;
 	}
 	else
 	{
 		FilteredComponentClassList.Empty();
-		for ( int32 ComponentIndex = 0; ComponentIndex < ComponentClassList.Num(); ComponentIndex++ )
+		for (int32 ComponentIndex = 0; ComponentIndex < ComponentClassList->Num(); ComponentIndex++)
 		{
-			FComponentClassComboEntryPtr& CurrentEntry = ComponentClassList[ComponentIndex];
+			FComponentClassComboEntryPtr& CurrentEntry = (*ComponentClassList)[ComponentIndex];
 
 			if (CurrentEntry->IsClass() && CurrentEntry->IsIncludedInFilter())
 			{
-				FString FriendlyComponentName = GetSanitizedComponentName( CurrentEntry->GetComponentClass(), CurrentEntry->GetComponentNameOverride() );
+				FString FriendlyComponentName = GetSanitizedComponentName( CurrentEntry );
 
 				if ( FriendlyComponentName.Contains( InSearchText, ESearchCase::IgnoreCase ) )
 				{
@@ -167,6 +179,25 @@ void SComponentClassCombo::OnSearchBoxTextCommitted(const FText& NewText, ETextC
 	}
 }
 
+// @todo: move this to FKismetEditorUtilities
+static UClass* GetAuthoritativeBlueprintClass(UBlueprint const* const Blueprint)
+{
+	UClass* BpClass = (Blueprint->SkeletonGeneratedClass != nullptr) ? Blueprint->SkeletonGeneratedClass :
+		Blueprint->GeneratedClass;
+
+	if (BpClass == nullptr)
+	{
+		BpClass = Blueprint->ParentClass;
+	}
+
+	UClass* AuthoritativeClass = BpClass;
+	if (BpClass != nullptr)
+	{
+		AuthoritativeClass = BpClass->GetAuthoritativeClass();
+	}
+	return AuthoritativeClass;
+}
+
 void SComponentClassCombo::OnAddComponentSelectionChanged( FComponentClassComboEntryPtr InItem, ESelectInfo::Type SelectInfo )
 {
 	if ( InItem.IsValid() && InItem->IsClass() && SelectInfo != ESelectInfo::OnNavigation)
@@ -181,7 +212,16 @@ void SComponentClassCombo::OnAddComponentSelectionChanged( FComponentClassComboE
 
 			if( OnComponentClassSelected.IsBound() )
 			{
-				UActorComponent* NewActorComponent = OnComponentClassSelected.Execute(InItem->GetComponentClass(), InItem->GetComponentCreateAction(), InItem->GetAssetOverride());
+				UClass* ComponentClass = InItem->GetComponentClass();
+				if (ComponentClass == nullptr)
+				{
+					// The class is not loaded yet, so load it:
+					const ELoadFlags LoadFlags = LOAD_None;
+					UBlueprint* LoadedObject = LoadObject<UBlueprint>(NULL, *InItem->GetComponentPath(), NULL, LoadFlags, NULL);
+					ComponentClass = GetAuthoritativeBlueprintClass(LoadedObject);
+				}
+
+				UActorComponent* NewActorComponent = OnComponentClassSelected.Execute(ComponentClass, InItem->GetComponentCreateAction(), InItem->GetAssetOverride());
 				if(NewActorComponent)
 				{
 					InItem->GetOnComponentCreated().ExecuteIfBound(NewActorComponent);
@@ -292,177 +332,11 @@ TSharedRef<ITableRow> SComponentClassCombo::GenerateAddComponentRow( FComponentC
 	}
 }
 
-struct SortComboEntry
-{
-	static const FString CommonClassGroup;
-
-	bool operator () (const FComponentClassComboEntryPtr& A, const FComponentClassComboEntryPtr& B) const
-	{
-		bool bResult = false;
-
-		// check headings first, if they are the same compare the individual entries
-		int32 HeadingCompareResult = FCString::Stricmp( *A->GetHeadingText(), *B->GetHeadingText() );
-		if ( HeadingCompareResult == 0 )
-		{
-			check(A->GetComponentClass()); check(B->GetComponentClass()); 
-			bResult = FCString::Stricmp(*A->GetComponentClass()->GetName(), *B->GetComponentClass()->GetName()) < 0;
-		}
-		else if (CommonClassGroup == A->GetHeadingText())
-		{
-			bResult = true;
-		}
-		else if (CommonClassGroup == B->GetHeadingText())
-		{
-			bResult = false;
-		}
-		else
-		{
-			bResult = HeadingCompareResult < 0;
-		}
-
-		return bResult;
-	}
-};
-const FString SortComboEntry::CommonClassGroup(TEXT("Common"));
-
-void SComponentClassCombo::OnBasicShapeCreated(UActorComponent* Component)
-{
-	UStaticMeshComponent* SMC = Cast<UStaticMeshComponent>(Component);
-	if(SMC)
-	{
-		SMC->SetMaterial(0, LoadObject<UMaterial>(nullptr, TEXT("/Engine/BasicShapes/BasicShapeMaterial.BasicShapeMaterial")));
-	}
-};
-
-void SComponentClassCombo::AddBasicShapeComponents( TArray<FComponentClassComboEntryPtr>& SortedClassList )
-{
-	FString BasicShapesHeading = LOCTEXT("BasicShapesHeading", "Basic Shapes").ToString();
-
-	{
-		FComponentEntryCustomizationArgs Args;
-		Args.AssetOverride = LoadObject<UStaticMesh>(nullptr, TEXT("/Engine/BasicShapes/Cube.Cube"));;
-		Args.OnComponentCreated = FOnComponentCreated::CreateSP(this, &SComponentClassCombo::OnBasicShapeCreated);
-		Args.ComponentNameOverride = LOCTEXT("BasicCubeShapeDisplayName", "Cube").ToString();
-
-		FComponentClassComboEntryPtr NewShape = MakeShareable(new FComponentClassComboEntry(BasicShapesHeading, UStaticMeshComponent::StaticClass(), true, EComponentCreateAction::SpawnExistingClass, Args));
-		SortedClassList.Add(NewShape);
-	}
-
-	{
-		FComponentEntryCustomizationArgs Args;
-		Args.AssetOverride = LoadObject<UStaticMesh>(nullptr, TEXT("/Engine/BasicShapes/Sphere.Sphere"));;
-		Args.OnComponentCreated = FOnComponentCreated::CreateSP(this, &SComponentClassCombo::OnBasicShapeCreated);
-		Args.ComponentNameOverride = LOCTEXT("BasicSphereShapeDisplayName", "Sphere").ToString();
-
-		FComponentClassComboEntryPtr NewShape = MakeShareable(new FComponentClassComboEntry(BasicShapesHeading, UStaticMeshComponent::StaticClass(), true, EComponentCreateAction::SpawnExistingClass, Args));
-		SortedClassList.Add(NewShape);
-	}
-
-	{
-		FComponentEntryCustomizationArgs Args;
-		Args.AssetOverride = LoadObject<UStaticMesh>(nullptr, TEXT("/Engine/BasicShapes/Cylinder.Cylinder"));;
-		Args.OnComponentCreated = FOnComponentCreated::CreateSP(this, &SComponentClassCombo::OnBasicShapeCreated);
-		Args.ComponentNameOverride = LOCTEXT("BasicCylinderShapeDisplayName", "Cylinder").ToString();
-
-		FComponentClassComboEntryPtr NewShape = MakeShareable(new FComponentClassComboEntry(BasicShapesHeading, UStaticMeshComponent::StaticClass(), true, EComponentCreateAction::SpawnExistingClass, Args));
-		SortedClassList.Add(NewShape);
-	}
-
-	{
-		FComponentEntryCustomizationArgs Args;
-		Args.AssetOverride = LoadObject<UStaticMesh>(nullptr, TEXT("/Engine/BasicShapes/Cone.Cone"));;
-		Args.OnComponentCreated = FOnComponentCreated::CreateSP(this, &SComponentClassCombo::OnBasicShapeCreated);
-		Args.ComponentNameOverride = LOCTEXT("BasicConeShapeDisplayName", "Cone").ToString();
-
-		FComponentClassComboEntryPtr NewShape = MakeShareable(new FComponentClassComboEntry(BasicShapesHeading, UStaticMeshComponent::StaticClass(), true, EComponentCreateAction::SpawnExistingClass, Args));
-		SortedClassList.Add(NewShape);
-	}
-
-}
-
 void SComponentClassCombo::UpdateComponentClassList()
 {
-	ComponentClassList.Empty();
-
-	const UEdGraphSchema_K2* K2Schema = GetDefault<UEdGraphSchema_K2>();
-
-	if( GetDefault<UEditorExperimentalSettings>()->bScriptableComponentsOnActors )
-	{
-		FString NewComponentsHeading = LOCTEXT("NewComponentsHeading", "New").ToString();
-		// Add new C++ component class
-		FComponentClassComboEntryPtr NewClassHeader = MakeShareable(new FComponentClassComboEntry(NewComponentsHeading));
-		ComponentClassList.Add(NewClassHeader);
-
-		FComponentClassComboEntryPtr NewCPPClass = MakeShareable(new FComponentClassComboEntry(NewComponentsHeading, UActorComponent::StaticClass(), true, EComponentCreateAction::CreateNewCPPClass));
-		ComponentClassList.Add(NewCPPClass);
-
-		FComponentClassComboEntryPtr NewBPClass = MakeShareable(new FComponentClassComboEntry(NewComponentsHeading, UActorComponent::StaticClass(), true, EComponentCreateAction::CreateNewBlueprintClass));
-		ComponentClassList.Add(NewBPClass);
-
-		FComponentClassComboEntryPtr NewSeparator(new FComponentClassComboEntry());
-		ComponentClassList.Add( NewSeparator );
-	}
-	
-	TArray<FComponentClassComboEntryPtr> SortedClassList;
-	for (TObjectIterator<UClass> It; It; ++It)
-	{
-		UClass* Class = *It;
-		// If this is a subclass of Actor Component, not abstract, and tagged as spawnable from Kismet
-		if (Class->IsChildOf(UActorComponent::StaticClass()) && !Class->HasAnyClassFlags(CLASS_Abstract) && Class->HasMetaData(FBlueprintMetadata::MD_BlueprintSpawnableComponent) && !FKismetEditorUtilities::IsClassABlueprintSkeleton(Class)) //@TODO: Fold this logic together with the one in UEdGraphSchema_K2::GetAddComponentClasses
-		{
-			TArray<FString> ClassGroupNames;
-			Class->GetClassGroupNames( ClassGroupNames );
-
-			if (ClassGroupNames.Contains(SortComboEntry::CommonClassGroup))
-			{
-				FString ClassGroup = SortComboEntry::CommonClassGroup;
-				FComponentClassComboEntryPtr NewEntry(new FComponentClassComboEntry(ClassGroup, Class, ClassGroupNames.Num() <= 1, EComponentCreateAction::SpawnExistingClass));
-				SortedClassList.Add(NewEntry);
-			}
-			if (ClassGroupNames.Num() && !ClassGroupNames[0].Equals(SortComboEntry::CommonClassGroup))
-			{
-				const bool bIncludeInFilter = true;
-
-				FString ClassGroup = ClassGroupNames[0];
-				FComponentClassComboEntryPtr NewEntry(new FComponentClassComboEntry(ClassGroup, Class, bIncludeInFilter, EComponentCreateAction::SpawnExistingClass));
-				SortedClassList.Add(NewEntry);
-			}
-		}
-	}
-
-	AddBasicShapeComponents(SortedClassList);
-
-	if (SortedClassList.Num() > 0)
-	{
-		Sort(SortedClassList.GetData(), SortedClassList.Num(), SortComboEntry());
-
-		FString PreviousHeading;
-		for ( int32 ClassIndex = 0; ClassIndex < SortedClassList.Num(); ClassIndex++ )
-		{
-			FComponentClassComboEntryPtr& CurrentEntry = SortedClassList[ClassIndex];
-
-			const FString& CurrentHeadingText = CurrentEntry->GetHeadingText();
-
-			if ( CurrentHeadingText != PreviousHeading )
-			{
-				// This avoids a redundant separator being added to the very top of the list
-				if ( ClassIndex > 0 )
-				{
-					FComponentClassComboEntryPtr NewSeparator(new FComponentClassComboEntry());
-					ComponentClassList.Add( NewSeparator );
-				}
-				FComponentClassComboEntryPtr NewHeading(new FComponentClassComboEntry( CurrentHeadingText ));
-				ComponentClassList.Add( NewHeading );
-
-				PreviousHeading = CurrentHeadingText;
-			}
-
-			ComponentClassList.Add( CurrentEntry );
-		}
-	}
+	GenerateFilteredComponentList(CurrentSearchString.ToString());
 }
 
-
 void SComponentClassCombo::OnProjectHotReloaded( bool bWasTriggeredAutomatically )
 {
 	UpdateComponentClassList();
@@ -486,7 +360,7 @@ FText SComponentClassCombo::GetFriendlyComponentName(FComponentClassComboEntryPt
 	}
 	else
 	{
-		FriendlyComponentName = GetSanitizedComponentName(Entry->GetComponentClass(), Entry->GetComponentNameOverride() );
+		FriendlyComponentName = GetSanitizedComponentName(Entry);
 
 		if( Entry->GetComponentNameOverride().IsEmpty() )
 		{
@@ -536,12 +410,16 @@ FText SComponentClassCombo::GetFriendlyComponentName(FComponentClassComboEntryPt
 	return FText::FromString(FriendlyComponentName);
 }
 
-FString SComponentClassCombo::GetSanitizedComponentName( UClass* ComponentClass, const FString& InComponentNameOverride )
+FString SComponentClassCombo::GetSanitizedComponentName(FComponentClassComboEntryPtr Entry)
 {
-	if( InComponentNameOverride.IsEmpty() )
+	FString DisplayName;
+	if (Entry->GetComponentNameOverride() != FString())
 	{
-		FString DisplayName;
-		if(ComponentClass->HasMetaData(TEXT("DisplayName")))
+		DisplayName = Entry->GetComponentNameOverride();
+	}
+	else if (UClass* ComponentClass = Entry->GetComponentClass())
+	{
+		if (ComponentClass->HasMetaData(TEXT("DisplayName")))
 		{
 			DisplayName = ComponentClass->GetMetaData(TEXT("DisplayName"));
 		}
@@ -553,13 +431,12 @@ FString SComponentClassCombo::GetSanitizedComponentName( UClass* ComponentClass,
 				DisplayName.RemoveFromEnd(TEXT("Component"), ESearchCase::IgnoreCase);
 			}
 		}
-
-		return FName::NameToDisplayString(DisplayName, false);
 	}
 	else
 	{
-		return InComponentNameOverride;
+		DisplayName = Entry->GetClassName();
 	}
+	return FName::NameToDisplayString(DisplayName, false);
 }
 
 #undef LOCTEXT_NAMESPACE
diff --git a/Engine/Source/Editor/UnrealEd/Public/ComponentTypeRegistry.h b/Engine/Source/Editor/UnrealEd/Public/ComponentTypeRegistry.h
new file mode 100644
--- /dev/null
+++ b/Engine/Source/Editor/UnrealEd/Public/ComponentTypeRegistry.h
@@ -0,0 +1,27 @@
+// Copyright 1998-2015 Epic Games, Inc. All Rights Reserved.
+
+#pragma once
+
+#include "SComponentClassCombo.h"
+
+DECLARE_MULTICAST_DELEGATE(FOnComponentTypeListChanged);
+
+typedef TSharedPtr<class FComponentClassComboEntry> FComponentClassComboEntryPtr;
+
+struct FComponentTypeRegistry
+{
+	static FComponentTypeRegistry& Get();
+
+	/**
+	* Called when the user changes the text in the search box.
+	* @OutComponentList Pointer that will be set to the (globally shared) component type list
+	* @return Deleate that can be used to handle change notifications. change notifications are raised when entries are 
+	*	added or removed from the component type list
+	*/
+	FOnComponentTypeListChanged& SubscribeToComponentList(TArray<FComponentClassComboEntryPtr>*& OutComponentList);
+	FOnComponentTypeListChanged& GetOnComponentTypeListChanged();
+
+private:
+	FComponentTypeRegistry();
+	struct FComponentTypeRegistryData* Data;
+};
diff --git a/Engine/Source/Editor/UnrealEd/Public/SComponentClassCombo.h b/Engine/Source/Editor/UnrealEd/Public/SComponentClassCombo.h
--- a/Engine/Source/Editor/UnrealEd/Public/SComponentClassCombo.h
+++ b/Engine/Source/Editor/UnrealEd/Public/SComponentClassCombo.h
@@ -46,14 +46,27 @@ class FComponentClassComboEntry: public TSharedFromThis<FComponentClassComboEntr
 public:
 	FComponentClassComboEntry( const FString& InHeadingText, TSubclassOf<UActorComponent> InComponentClass, bool InIncludedInFilter, EComponentCreateAction::Type InComponentCreateAction, FComponentEntryCustomizationArgs InCustomizationArgs = FComponentEntryCustomizationArgs() )
 		: ComponentClass(InComponentClass)
+		, ComponentName()
+		, ComponentPath()
 		, HeadingText(InHeadingText)
 		, bIncludedInFilter(InIncludedInFilter)
 		, ComponentCreateAction(InComponentCreateAction)
 		, CustomizationArgs(InCustomizationArgs)
 	{}
 
+	FComponentClassComboEntry(const FString& InHeadingText, const FString& InComponentName, FName InComponentPath, bool InIncludedInFilter)
+		: ComponentClass()
+		, ComponentName(InComponentName)
+		, ComponentPath(InComponentPath)
+		, HeadingText(InHeadingText)
+		, bIncludedInFilter(InIncludedInFilter)
+		, ComponentCreateAction(EComponentCreateAction::SpawnExistingClass)
+	{}
+
 	FComponentClassComboEntry( const FString& InHeadingText )
 		: ComponentClass(NULL)
+		, ComponentName()
+		, ComponentPath()
 		, HeadingText(InHeadingText)
 		, bIncludedInFilter(false)
 	{}
@@ -69,16 +82,16 @@ public:
 
 	bool IsHeading() const
 	{
-		return (ComponentClass==NULL && !HeadingText.IsEmpty());
+		return ((ComponentClass == NULL && ComponentName == FString()) && !HeadingText.IsEmpty());
 	}
 	bool IsSeparator() const
 	{
-		return (ComponentClass==NULL && HeadingText.IsEmpty());
+		return ((ComponentClass == NULL && ComponentName == FString()) && HeadingText.IsEmpty());
 	}
 	
 	bool IsClass() const
 	{
-		return (ComponentClass!=NULL);
+		return (ComponentClass != NULL || ComponentName != FString());
 	}
 	
 	bool IsIncludedInFilter() const
@@ -100,6 +113,8 @@ public:
 	{
 		return CustomizationArgs.OnComponentCreated;
 	}
+	FString GetClassName() const;
+	FString GetComponentPath() const { return ComponentPath.ToString(); }
 
 	UObject* GetAssetOverride()
 	{
@@ -107,6 +122,11 @@ public:
 	}
 private:
 	TSubclassOf<UActorComponent> ComponentClass;
+	// For components that are not loaded we just keep the name of the component,
+	// loading occurs when the blueprint is spawned, which should also trigger a refresh
+	// of the component list:
+	FString ComponentName;
+	FName ComponentPath;
 	FString HeadingText;
 	bool bIncludedInFilter;
 	EComponentCreateAction::Type ComponentCreateAction;
@@ -129,6 +149,8 @@ public:
 
 	void Construct(const FArguments& InArgs);
 
+	~SComponentClassCombo();
+
 	/** Clear the current combo list selection */
 	void ClearSelection();
 
@@ -157,18 +179,13 @@ public:
 	void UpdateComponentClassList();
 
 	/** Returns a component name without the substring "Component" and sanitized for display */
-	static FString GetSanitizedComponentName( UClass* ComponentClass, const FString& DisplayNameOverride );
+	static FString GetSanitizedComponentName(FComponentClassComboEntryPtr Entry);
 
 protected:
 
 	/** Called when a project is hot reloaded to refresh the components list */
 	void OnProjectHotReloaded( bool bWasTriggeredAutomatically );
 
-	/** Adds basic shapes to the component dropdown */
-	void AddBasicShapeComponents( TArray<FComponentClassComboEntryPtr>& SortedClassList );
-
-	/** Called when a new basic shape is added */
-	void OnBasicShapeCreated( UActorComponent* CreatedComponent);
 private:
 
 	FText GetFriendlyComponentName(FComponentClassComboEntryPtr Entry) const;
@@ -176,7 +193,7 @@ private:
 	FComponentClassSelected OnComponentClassSelected;
 
 	/** List of component class names used by combo box */
-	TArray<FComponentClassComboEntryPtr> ComponentClassList;
+	TArray<FComponentClassComboEntryPtr>* ComponentClassList;
 
 	/** List of component class names, filtered by the current search string */
 	TArray<FComponentClassComboEntryPtr> FilteredComponentClassList;
