[INTEGRATE] Change 2426903 by [REDACTED_PERSON][REDACTED_HANDLE] on 2015/01/30 22:15:11

	Editor: Fixing macro redefinition compile error on Mac

diff --git a/Engine/Source/Editor/UnrealEd/Private/ComponentTypeRegistry.cpp b/Engine/Source/Editor/UnrealEd/Private/ComponentTypeRegistry.cpp
--- a/Engine/Source/Editor/UnrealEd/Private/ComponentTypeRegistry.cpp
+++ b/Engine/Source/Editor/UnrealEd/Private/ComponentTypeRegistry.cpp
@@ -321,3 +321,4 @@ FComponentTypeRegistry::FComponentTypeRegistry()
 	Data->RefreshComponentList();
 }
 
+#undef LOCTEXT_NAMESPACE
\ No newline at end of file
