[INTEGRATE] Change 2427085 by [REDACTED_PERSON][REDACTED_HANDLE] on 2015/01/31 12:43:26

	Update basic shape icons in Add Components list.

diff --git a/Engine/Source/Editor/UnrealEd/Private/ComponentTypeRegistry.cpp b/Engine/Source/Editor/UnrealEd/Private/ComponentTypeRegistry.cpp
--- a/Engine/Source/Editor/UnrealEd/Private/ComponentTypeRegistry.cpp
+++ b/Engine/Source/Editor/UnrealEd/Private/ComponentTypeRegistry.cpp
@@ -48,6 +48,7 @@ static void AddBasicShapeComponents(TArray<FComponentClassComboEntryPtr>& Sorted
 		Args.AssetOverride = LoadObject<UStaticMesh>(nullptr, TEXT("/Engine/BasicShapes/Cube.Cube"));;
 		Args.OnComponentCreated = FOnComponentCreated::CreateStatic(OnBasicShapeCreated);
 		Args.ComponentNameOverride = LOCTEXT("BasicCubeShapeDisplayName", "Cube").ToString();
+		Args.IconOverrideBrushName = FName("ClassIcon.Cube");
 
 		FComponentClassComboEntryPtr NewShape = MakeShareable(new FComponentClassComboEntry(BasicShapesHeading, UStaticMeshComponent::StaticClass(), true, EComponentCreateAction::SpawnExistingClass, Args));
 		SortedClassList.Add(NewShape);
@@ -58,6 +59,7 @@ static void AddBasicShapeComponents(TArray<FComponentClassComboEntryPtr>& Sorted
 		Args.AssetOverride = LoadObject<UStaticMesh>(nullptr, TEXT("/Engine/BasicShapes/Sphere.Sphere"));;
 		Args.OnComponentCreated = FOnComponentCreated::CreateStatic(OnBasicShapeCreated);
 		Args.ComponentNameOverride = LOCTEXT("BasicSphereShapeDisplayName", "Sphere").ToString();
+		Args.IconOverrideBrushName = FName("ClassIcon.Sphere");
 
 		FComponentClassComboEntryPtr NewShape = MakeShareable(new FComponentClassComboEntry(BasicShapesHeading, UStaticMeshComponent::StaticClass(), true, EComponentCreateAction::SpawnExistingClass, Args));
 		SortedClassList.Add(NewShape);
@@ -68,6 +70,7 @@ static void AddBasicShapeComponents(TArray<FComponentClassComboEntryPtr>& Sorted
 		Args.AssetOverride = LoadObject<UStaticMesh>(nullptr, TEXT("/Engine/BasicShapes/Cylinder.Cylinder"));;
 		Args.OnComponentCreated = FOnComponentCreated::CreateStatic(OnBasicShapeCreated);
 		Args.ComponentNameOverride = LOCTEXT("BasicCylinderShapeDisplayName", "Cylinder").ToString();
+		Args.IconOverrideBrushName = FName("ClassIcon.Cylinder");
 
 		FComponentClassComboEntryPtr NewShape = MakeShareable(new FComponentClassComboEntry(BasicShapesHeading, UStaticMeshComponent::StaticClass(), true, EComponentCreateAction::SpawnExistingClass, Args));
 		SortedClassList.Add(NewShape);
@@ -78,6 +81,7 @@ static void AddBasicShapeComponents(TArray<FComponentClassComboEntryPtr>& Sorted
 		Args.AssetOverride = LoadObject<UStaticMesh>(nullptr, TEXT("/Engine/BasicShapes/Cone.Cone"));;
 		Args.OnComponentCreated = FOnComponentCreated::CreateStatic(OnBasicShapeCreated);
 		Args.ComponentNameOverride = LOCTEXT("BasicConeShapeDisplayName", "Cone").ToString();
+		Args.IconOverrideBrushName = FName("ClassIcon.Cone");
 
 		FComponentClassComboEntryPtr NewShape = MakeShareable(new FComponentClassComboEntry(BasicShapesHeading, UStaticMeshComponent::StaticClass(), true, EComponentCreateAction::SpawnExistingClass, Args));
 		SortedClassList.Add(NewShape);
diff --git a/Engine/Source/Editor/UnrealEd/Private/SComponentClassCombo.cpp b/Engine/Source/Editor/UnrealEd/Private/SComponentClassCombo.cpp
--- a/Engine/Source/Editor/UnrealEd/Private/SComponentClassCombo.cpp
+++ b/Engine/Source/Editor/UnrealEd/Private/SComponentClassCombo.cpp
@@ -311,7 +311,7 @@ TSharedRef<ITableRow> SComponentClassCombo::GenerateAddComponentRow( FComponentC
 				.AutoWidth()
 				[
 					SNew(SImage)
-					.Image( FClassIconFinder::FindIconForClass( Entry->GetComponentClass() ) )
+					.Image( FClassIconFinder::FindIconForClass( Entry->GetIconOverrideBrushName() == NAME_None ? Entry->GetComponentClass() : nullptr, Entry->GetIconOverrideBrushName() ) )
 				]
 				+SHorizontalBox::Slot()
 				.AutoWidth()
diff --git a/Engine/Source/Editor/UnrealEd/Public/SComponentClassCombo.h b/Engine/Source/Editor/UnrealEd/Public/SComponentClassCombo.h
--- a/Engine/Source/Editor/UnrealEd/Public/SComponentClassCombo.h
+++ b/Engine/Source/Editor/UnrealEd/Public/SComponentClassCombo.h
@@ -32,11 +32,14 @@ struct FComponentEntryCustomizationArgs
 	FString ComponentNameOverride;
 	/** Callback when a new component is created */
 	FOnComponentCreated OnComponentCreated;
+	/** Brush icon to use instead of the class icon */
+	FName IconOverrideBrushName;
 
 	FComponentEntryCustomizationArgs()
 		: AssetOverride( nullptr )
 		, ComponentNameOverride()
 		, OnComponentCreated()
+		, IconOverrideBrushName( NAME_None )
 	{
 	
 	}
@@ -120,6 +123,8 @@ public:
 	{
 		return CustomizationArgs.AssetOverride.Get();
 	}
+
+	virtual FName GetIconOverrideBrushName() const { return CustomizationArgs.IconOverrideBrushName; }
 private:
 	TSubclassOf<UActorComponent> ComponentClass;
 	// For components that are not loaded we just keep the name of the component,
