[INTEGRATE] Change 2427093 by [REDACTED_PERSON][REDACTED_HANDLE] on 2015/01/31 13:13:01

	Asset editor tabs now by default open docked in the tab well they were opened from

diff --git a/Engine/Source/Editor/UnrealEd/Private/Toolkits/AssetEditorToolkit.cpp b/Engine/Source/Editor/UnrealEd/Private/Toolkits/AssetEditorToolkit.cpp
--- a/Engine/Source/Editor/UnrealEd/Private/Toolkits/AssetEditorToolkit.cpp
+++ b/Engine/Source/Editor/UnrealEd/Private/Toolkits/AssetEditorToolkit.cpp
@@ -93,7 +93,7 @@ void FAssetEditorToolkit::InitAssetEditor( const EToolkitMode::Type Mode, const
 			static_assert(sizeof(EAssetEditorToolkitTabLocation) == sizeof(int32), "EAssetEditorToolkitTabLocation is the incorrect size");
 
 			// Work out where we should create this asset editor
-			EAssetEditorToolkitTabLocation SavedAssetEditorToolkitTabLocation = EAssetEditorToolkitTabLocation::Standalone;
+			EAssetEditorToolkitTabLocation SavedAssetEditorToolkitTabLocation = EAssetEditorToolkitTabLocation::Docked;
 			GConfig->GetInt(
 				TEXT("AssetEditorToolkitTabLocation"), 
 				*ObjectsToEdit[0]->GetPathName(), 
