[INTEGRATE] Change 2427095 by [REDACTED_PERSON][REDACTED_HANDLE] on 2015/01/31 13:14:51

	Fix function being virtual when it should not have been

diff --git a/Engine/Source/Editor/UnrealEd/Public/SComponentClassCombo.h b/Engine/Source/Editor/UnrealEd/Public/SComponentClassCombo.h
--- a/Engine/Source/Editor/UnrealEd/Public/SComponentClassCombo.h
+++ b/Engine/Source/Editor/UnrealEd/Public/SComponentClassCombo.h
@@ -124,7 +124,7 @@ public:
 		return CustomizationArgs.AssetOverride.Get();
 	}
 
-	virtual FName GetIconOverrideBrushName() const { return CustomizationArgs.IconOverrideBrushName; }
+	FName GetIconOverrideBrushName() const { return CustomizationArgs.IconOverrideBrushName; }
 private:
 	TSubclassOf<UActorComponent> ComponentClass;
 	// For components that are not loaded we just keep the name of the component,
