[INTEGRATE] Change 2427167 by [REDACTED_PERSON][REDACTED_HANDLE] on 2015/01/31 15:28:13

	Enabling 'Blueprintable Components' now allows any BlueprintSpawnableComponent to be Blueprinted
	- This unblocks testing of blueprinted scene components and other interesting types
	- Currently you need to use the regular "New Blueprint" flow -- the "Add Component" menu still only supports ActorComponents
	- Also removed the 'bScriptableComponentsOnActors' setting and collapsed that into the bBlueprintableComponents setting

diff --git a/Engine/Source/Editor/UnrealEd/Classes/Settings/EditorExperimentalSettings.h b/Engine/Source/Editor/UnrealEd/Classes/Settings/EditorExperimentalSettings.h
--- a/Engine/Source/Editor/UnrealEd/Classes/Settings/EditorExperimentalSettings.h
+++ b/Engine/Source/Editor/UnrealEd/Classes/Settings/EditorExperimentalSettings.h
@@ -41,12 +41,8 @@ public:
 	UPROPERTY(EditAnywhere, config, Category = Tools, meta = ( DisplayName = "Single Layout Blueprint Editor" ))
 	bool bUnifiedBlueprintEditor;
 
-	/** Ability to add C++ and blueprinted components to actors */
-	UPROPERTY(EditAnywhere, config, Category=Tools)
-	bool bScriptableComponentsOnActors;
-
 	/** Enable being able to subclass components in blueprints */
-	UPROPERTY(EditAnywhere, config, Category = Tools, meta = (DisplayName = "Blueprintable Component"))
+	UPROPERTY(EditAnywhere, config, Category=Tools)
 	bool bBlueprintableComponents;
 
 	/** The Messaging Debugger provides a visual utility for debugging the messaging system. */
diff --git a/Engine/Source/Editor/UnrealEd/Private/ComponentTypeRegistry.cpp b/Engine/Source/Editor/UnrealEd/Private/ComponentTypeRegistry.cpp
--- a/Engine/Source/Editor/UnrealEd/Private/ComponentTypeRegistry.cpp
+++ b/Engine/Source/Editor/UnrealEd/Private/ComponentTypeRegistry.cpp
@@ -141,7 +141,7 @@ void FComponentTypeRegistryData::RefreshComponentList()
 
 	const UEdGraphSchema_K2* K2Schema = GetDefault<UEdGraphSchema_K2>();
 
-	if (GetDefault<UEditorExperimentalSettings>()->bScriptableComponentsOnActors)
+	if (GetDefault<UEditorExperimentalSettings>()->bBlueprintableComponents)
 	{
 		FString NewComponentsHeading = LOCTEXT("NewComponentsHeading", "New").ToString();
 		// Add new C++ component class
diff --git a/Engine/Source/Editor/UnrealEd/Private/Kismet2/Kismet2.cpp b/Engine/Source/Editor/UnrealEd/Private/Kismet2/Kismet2.cpp
--- a/Engine/Source/Editor/UnrealEd/Private/Kismet2/Kismet2.cpp
+++ b/Engine/Source/Editor/UnrealEd/Private/Kismet2/Kismet2.cpp
@@ -876,8 +876,7 @@ bool FKismetEditorUtilities::CanCreateBlueprintOfClass(const UClass* Class)
 	const UEdGraphSchema_K2* K2Schema = GetDefault<UEdGraphSchema_K2>();
 	const bool bIsValidClass = Class->GetBoolMetaDataHierarchical(FBlueprintMetadata::MD_IsBlueprintBase)
 		|| (Class == UObject::StaticClass())
-		|| (Class == UActorComponent::StaticClass() && bAllowBlueprintableComponents)
-		|| (bAllowBlueprintableComponents && Class->IsChildOf<UActorComponent>() && Class->HasAnyClassFlags(CLASS_CompiledFromBlueprint))
+		|| (bAllowBlueprintableComponents && !Class->HasAnyClassFlags(CLASS_Abstract) && Class->IsChildOf<UActorComponent>() && ( Class->HasAnyClassFlags(CLASS_CompiledFromBlueprint) || Class->HasMetaData( FBlueprintMetadata::MD_BlueprintSpawnableComponent ) ) )
 		|| bIsBPGC;  // BPs are always considered inheritable
 
 	return bCanCreateBlueprint && bIsValidClass;
