[INTEGRATE] Change 2428140 by [REDACTED_PERSON][REDACTED_HANDLE].7 on 2015/02/02 10:18:43

	"Delete source content with assets" improvements

	Deleting assets now only shows the "Also delete source content" check box under the following circumstances:
	    * Auto reimport workflow is turned *on*, AND
	    * The assets currently being deleted were imported from files that still exist on disk, and reside under a monitored content directory.

	The tooltip now displays a list of files that will be deleted. It remains unchecked by default.

diff --git a/Engine/Source/Editor/UnrealEd/Private/AssetDeleteModel.cpp b/Engine/Source/Editor/UnrealEd/Private/AssetDeleteModel.cpp
--- a/Engine/Source/Editor/UnrealEd/Private/AssetDeleteModel.cpp
+++ b/Engine/Source/Editor/UnrealEd/Private/AssetDeleteModel.cpp
@@ -7,6 +7,7 @@
 #include "AssetRegistryModule.h"
 #include "AssetEditorManager.h"
 #include "AutoReimport/AutoReimportUtilities.h"
+#include "AutoReimport/AutoReimportManager.h"
 
 #define LOCTEXT_NAMESPACE "FAssetDeleteModel"
 
@@ -15,6 +16,7 @@ FAssetDeleteModel::FAssetDeleteModel( const TArray<UObject*>& InObjectsToDelete
 	, bPendingObjectsCanBeReplaced(false)
 	, bIsAnythingReferencedInMemoryByNonUndo(false)
 	, bIsAnythingReferencedInMemoryByUndo(false)
+	, bHasSourceContentFilesToDelete(false)
 	, PendingDeleteIndex(0)
 	, ObjectsDeleted(0)
 {
@@ -64,6 +66,7 @@ void FAssetDeleteModel::Tick( const float InDeltaTime )
 		OnDiskReferences = TSet<FName>();
 		bIsAnythingReferencedInMemoryByNonUndo = false;
 		bIsAnythingReferencedInMemoryByUndo = false;
+		bHasSourceContentFilesToDelete = false;
 		PendingDeleteIndex = 0;
 
 		SetState(Scanning);
@@ -99,8 +102,35 @@ void FAssetDeleteModel::Tick( const float InDeltaTime )
 			}
 			PendingDelete->RemainingMemoryReferences = NonPendingDeletedExternalInMemoryReferences;
 
+			if (GetDefault<UEditorLoadingSavingSettings>()->bMonitorContentDirectories)
+			{
+				Utils::ExtractSourceFilePaths(PendingDelete->GetObject(), PendingDelete->SourceContentFiles);
+
+				auto MonitoredDirectories = GUnrealEd->AutoReimportManager->GetMonitoredDirectories();
+
+				// Remove anything that's not under a monitored, mounted path
+				PendingDelete->SourceContentFiles.RemoveAll([&](const FString& InFilename){
+					for (const auto& Dir : MonitoredDirectories)
+					{
+						if (!Dir.MountPoint.IsEmpty() && InFilename.StartsWith(Dir.Path))
+						{
+							if (FPaths::FileExists(InFilename))
+							{
+								return false;
+							}
+							else
+							{
+								return true;
+							}
+						}
+					}
+					return true;
+				});
+			}
+
 			bIsAnythingReferencedInMemoryByNonUndo |= PendingDelete->RemainingMemoryReferences > 0;
 			bIsAnythingReferencedInMemoryByUndo |= PendingDelete->IsReferencedInMemoryByUndo();
+			bHasSourceContentFilesToDelete |= PendingDelete->SourceContentFiles.Num() != 0;
 
 			PendingDeleteIndex++;
 		}
@@ -120,35 +150,12 @@ void FAssetDeleteModel::Tick( const float InDeltaTime )
 
 void FAssetDeleteModel::DeleteSourceContentFiles()
 {
-	const IAssetRegistry& Registry = FModuleManager::LoadModuleChecked<FAssetRegistryModule>("AssetRegistry").Get();
-	
-	TArray<FString> FilesToDelete;
-	for ( const TSharedPtr< FPendingDelete >& PendingDelete : PendingDeletes )
+	IFileManager& FileManager = IFileManager::Get();
+	for (const auto& PendingDelete : PendingDeletes)
 	{
-		Utils::ExtractSourceFilePaths(PendingDelete->GetObject(), FilesToDelete);
-	}
-
-	if (FilesToDelete.Num() != 0)
-	{
-		TArray<FString> RootContentPaths;
-		FPackageName::QueryRootContentPaths( RootContentPaths );
-		for (FString& RootPath : RootContentPaths)
-		{
-			RootPath = FPaths::ConvertRelativePathToFull(FPackageName::LongPackageNameToFilename(RootPath));
-		}
-
-		IFileManager& FileManager = IFileManager::Get();
-		for (const auto& Path : FilesToDelete)
+		for (const auto& Path : PendingDelete->SourceContentFiles)
 		{
-			const FString FullPath = FPaths::ConvertRelativePathToFull(Path);
-			const bool bFileIsExternal = !RootContentPaths.ContainsByPredicate([&](const FString& ContentDir){
-				return FullPath.StartsWith(ContentDir);
-			});
-
-			if (!bFileIsExternal)
-			{
-				FileManager.Delete(*Path, false /* RequireExists */, true /* Even if read only */, true /* Quiet */);
-			}
+			FileManager.Delete(*Path, false /* RequireExists */, true /* Even if read only */, true /* Quiet */);
 		}
 	}
 }
diff --git a/Engine/Source/Editor/UnrealEd/Private/AutoReimport/AutoReimportManager.cpp b/Engine/Source/Editor/UnrealEd/Private/AutoReimport/AutoReimportManager.cpp
--- a/Engine/Source/Editor/UnrealEd/Private/AutoReimport/AutoReimportManager.cpp
+++ b/Engine/Source/Editor/UnrealEd/Private/AutoReimport/AutoReimportManager.cpp
@@ -98,6 +98,7 @@ class FAutoReimportManager : public FTickableEditorObject, public FGCObject, pub
 public:
 
 	FAutoReimportManager();
+	TArray<FPathAndMountPoint> GetMonitoredDirectories() const;
 	void Destroy();
 
 private:
@@ -220,6 +221,16 @@ FAutoReimportManager::FAutoReimportManager()
 	StateMachine.Add(ECurrentState::ProcessDeletions,		EStateMachineNode::CallMany,	[this](const FTimeLimit& T){ return this->ProcessDeletions();		});
 }
 
+TArray<FPathAndMountPoint> FAutoReimportManager::GetMonitoredDirectories() const
+{
+	TArray<FPathAndMountPoint> Dirs;
+	for (const auto& Monitor : DirectoryMonitors)
+	{
+		Dirs.Emplace(Monitor.GetDirectory(), Monitor.GetMountPoint());
+	}
+	return Dirs;
+}
+
 void FAutoReimportManager::Destroy()
 {
 	if (auto* Settings = GetMutableDefault<UEditorLoadingSavingSettings>())
@@ -272,7 +283,7 @@ TOptional<ECurrentState> FAutoReimportManager::ProcessAdditions(const FTimeLimit
 {
 	// Override the global feedback context while we do this to avoid popping up dialogs
 	TGuardValue<FFeedbackContext*> ScopedContextOverride(GWarn, FeedbackContextOverride.Get());
-	
+
 	FeedbackContextOverride->GetContent()->SetMainText(GetProgressText());
 
 	TMap<FString, TArray<UFactory*>> Factories;
@@ -501,7 +512,6 @@ void FAutoReimportManager::OnContentPathChanged(const FString& InAssetPath, cons
 
 void FAutoReimportManager::SetUpDirectoryMonitors()
 {
-	struct FPathAndMountPoint { FString Path; FString MountPoint; };
 	TArray<FPathAndMountPoint> CollapsedDirectories;
 
 	// Build an array of directory / mounted path so we can match up absolute paths that the user has typed in
@@ -611,6 +621,11 @@ void UAutoReimportManager::Initialize()
 	Implementation = MakeShareable(new FAutoReimportManager);
 }
 
+TArray<FPathAndMountPoint> UAutoReimportManager::GetMonitoredDirectories() const
+{
+	return Implementation->GetMonitoredDirectories();
+}
+
 void UAutoReimportManager::BeginDestroy()
 {
 	Super::BeginDestroy();
diff --git a/Engine/Source/Editor/UnrealEd/Private/AutoReimport/ContentDirectoryMonitor.h b/Engine/Source/Editor/UnrealEd/Private/AutoReimport/ContentDirectoryMonitor.h
--- a/Engine/Source/Editor/UnrealEd/Private/AutoReimport/ContentDirectoryMonitor.h
+++ b/Engine/Source/Editor/UnrealEd/Private/AutoReimport/ContentDirectoryMonitor.h
@@ -40,6 +40,9 @@ public:
 	/** Get the directory that this monitor applies to */
 	const FString& GetDirectory() const { return Cache.GetDirectory(); }
 
+	/** Get the content mount point that this monitor applies to */
+	const FString& GetMountPoint() const { return MountedContentPath; }
+
 public:
 
 	/** Get the number of outstanding changes that we potentially have to process (when not already processing) */
diff --git a/Engine/Source/Editor/UnrealEd/Private/Dialogs/SDeleteAssetsDialog.cpp b/Engine/Source/Editor/UnrealEd/Private/Dialogs/SDeleteAssetsDialog.cpp
--- a/Engine/Source/Editor/UnrealEd/Private/Dialogs/SDeleteAssetsDialog.cpp
+++ b/Engine/Source/Editor/UnrealEd/Private/Dialogs/SDeleteAssetsDialog.cpp
@@ -332,13 +332,14 @@ TSharedRef<SWidget> SDeleteAssetsDialog::BuildDeleteDialog()
 			.Padding(6,4)
 			[
 				SAssignNew(DeleteSourceFilesCheckbox, SCheckBox)
+				.Visibility(this, &SDeleteAssetsDialog::GetDeleteSourceFilesVisibility)
 				.IsChecked(LoadingSavingSettings->bDeleteSourceFilesWithAssets ? ECheckBoxState::Checked : ECheckBoxState::Unchecked)
 				[
 					SNew(STextBlock)
 					.Text(LOCTEXT("DeleteSourceFiles", "Also delete related source content files"))
 					.ToolTip(
 						SNew(SToolTip)
-						.Text(LOCTEXT("DeleteSourceFiles_Tooltip", "When checked, any source content files that were used to create the selected asset(s) will also be deleted. For instance, deleting a texture will also delete its associated *.png file."))
+						.Text(this, &SDeleteAssetsDialog::GetDeleteSourceContentTooltip)
 					)
 				]
 			]
@@ -421,6 +422,26 @@ FText SDeleteAssetsDialog::GetHandleText() const
 	}
 }
 
+FText SDeleteAssetsDialog::GetDeleteSourceContentTooltip() const
+{
+	const FText RootText = LOCTEXT("DeleteSourceFiles_Tooltip", "When checked, the following source content files will also be deleted along with the assets:\n\n{0}");
+
+	FString AllFiles;
+	for (const auto& PendingDelete : *DeleteModel->GetPendingDeletedAssets())
+	{
+		for (const FString& Path : PendingDelete->SourceContentFiles)
+		{
+			if (!AllFiles.IsEmpty())
+			{
+				AllFiles += TEXT("\n");
+			}
+			AllFiles += Path;
+		}
+	}
+
+	return FText::Format(RootText, FText::FromString(AllFiles));
+}
+
 EVisibility SDeleteAssetsDialog::GetAssetReferencesVisiblity() const
 {
 	return DeleteModel->GetAssetReferences().Num() == 0 ? EVisibility::Collapsed : EVisibility::Visible;
@@ -593,11 +614,25 @@ TSharedRef<ITableRow> SDeleteAssetsDialog::HandleGenerateAssetRow( TSharedPtr<FP
 		.Visibility(InItem->IsInternal() ? EVisibility::Collapsed : EVisibility::Visible);
 }
 
-FReply SDeleteAssetsDialog::Delete()
+void SDeleteAssetsDialog::DeleteRelevantSourceContent()
 {
-	auto* Settings = GetMutableDefault<UEditorLoadingSavingSettings>();
-	Settings->bDeleteSourceFilesWithAssets = DeleteSourceFilesCheckbox->GetCheckedState() == ECheckBoxState::Checked;
+	if (DeleteModel->HasAnySourceContentFilesToDelete())
+	{
+		auto* Settings = GetMutableDefault<UEditorLoadingSavingSettings>();
+		if (DeleteSourceFilesCheckbox->GetCheckedState() == ECheckBoxState::Checked)
+		{
+			Settings->bDeleteSourceFilesWithAssets = true;
+			DeleteModel->DeleteSourceContentFiles();
+		}
+		else
+		{
+			Settings->bDeleteSourceFilesWithAssets = false;			
+		}
+	}
+}
 
+FReply SDeleteAssetsDialog::Delete()
+{
 	ParentWindow.Get()->RequestDestroyWindow();
 
 	if (DeleteModel->IsAnythingReferencedInMemoryByUndo())
@@ -605,10 +640,7 @@ FReply SDeleteAssetsDialog::Delete()
 		GEditor->Trans->Reset(LOCTEXT("DeleteSelectedItem", "Delete Selected Item"));
 	}
 
-	if (Settings->bDeleteSourceFilesWithAssets)
-	{
-		DeleteModel->DeleteSourceContentFiles();
-	}
+	DeleteRelevantSourceContent();
 	DeleteModel->DoDelete();
 
 	return FReply::Handled();
@@ -623,9 +655,6 @@ FReply SDeleteAssetsDialog::Cancel()
 
 FReply SDeleteAssetsDialog::ForceDelete()
 {
-	auto* Settings = GetMutableDefault<UEditorLoadingSavingSettings>();
-	Settings->bDeleteSourceFilesWithAssets = DeleteSourceFilesCheckbox->GetCheckedState() == ECheckBoxState::Checked;
-
 	ParentWindow.Get()->RequestDestroyWindow();
 
 	if( DeleteModel->IsAnythingReferencedInMemoryByUndo() )
@@ -633,10 +662,7 @@ FReply SDeleteAssetsDialog::ForceDelete()
 		GEditor->Trans->Reset( LOCTEXT("DeleteSelectedItem", "Delete Selected Item") );
 	}
 
-	if (Settings->bDeleteSourceFilesWithAssets)
-	{
-		DeleteModel->DeleteSourceContentFiles();
-	}
+	DeleteRelevantSourceContent();
 	DeleteModel->DoForceDelete();
 
 	return FReply::Handled();
@@ -666,14 +692,8 @@ FReply SDeleteAssetsDialog::ReplaceReferences()
 
 	if ( EAppReturnType::Ok == OpenMsgDlgInt( EAppMsgType::OkCancel, Message, Title ) )
 	{
-		auto* Settings = GetMutableDefault<UEditorLoadingSavingSettings>();
-		Settings->bDeleteSourceFilesWithAssets = DeleteSourceFilesCheckbox->GetCheckedState() == ECheckBoxState::Checked;
-		if (Settings->bDeleteSourceFilesWithAssets)
-		{
-			DeleteModel->DeleteSourceContentFiles();
-		}
-
 		ParentWindow.Get()->RequestDestroyWindow();
+		DeleteRelevantSourceContent();
 		DeleteModel->DoReplaceReferences( ConsolidationAsset );
 	}
 
@@ -812,6 +832,11 @@ EVisibility SDeleteAssetsDialog::GetDeleteVisibility() const
 	return CanDelete() ? EVisibility::Visible : EVisibility::Collapsed;
 }
 
+EVisibility SDeleteAssetsDialog::GetDeleteSourceFilesVisibility() const
+{
+	return DeleteModel->HasAnySourceContentFilesToDelete() ? EVisibility::Visible : EVisibility::Collapsed;
+}
+
 bool SDeleteAssetsDialog::CanReplaceReferences() const
 {
 	return DeleteModel->CanReplaceReferences();
diff --git a/Engine/Source/Editor/UnrealEd/Public/AssetDeleteModel.h b/Engine/Source/Editor/UnrealEd/Public/AssetDeleteModel.h
--- a/Engine/Source/Editor/UnrealEd/Public/AssetDeleteModel.h
+++ b/Engine/Source/Editor/UnrealEd/Public/AssetDeleteModel.h
@@ -40,6 +40,9 @@ public:
 	/** The on disk references to this object */
 	TArray<FName> DiskReferences;
 
+	/** Source content files for this pending delete operation */
+	TArray<FString> SourceContentFiles;
+
 	/** In memory references to this object (*excluding* the undo buffer) */
 	FReferencerInformationList MemoryReferences;
 
@@ -159,7 +162,10 @@ public:
 
 	/** Is any of the pending deleted assets being referenced in the undo stack. */
 	bool IsAnythingReferencedInMemoryByUndo() const { return bIsAnythingReferencedInMemoryByUndo; }
-	
+
+	/** Check whether we have any source files residing under monitored, mounted paths to delete */
+	bool HasAnySourceContentFilesToDelete() const { return bHasSourceContentFilesToDelete; }
+
 	/** Goes to the next actor in the loaded level if it is available */
 	bool GoToNextReferenceInLevel() const;
 
@@ -205,6 +211,9 @@ private:
 	/** Is any of the pending deleted assets being referenced in the undo stack. */
 	bool bIsAnythingReferencedInMemoryByUndo;
 
+	/** Boolean that specifies whether we have any relevant (in game dir) source content files to remove with the assets */
+	bool bHasSourceContentFilesToDelete;
+
 	/** A tick-to-tick state tracking variable so we know what pending deleted object we checked last. */
 	int32 PendingDeleteIndex;
 
diff --git a/Engine/Source/Editor/UnrealEd/Public/AutoReimport/AutoReimportManager.h b/Engine/Source/Editor/UnrealEd/Public/AutoReimport/AutoReimportManager.h
--- a/Engine/Source/Editor/UnrealEd/Public/AutoReimport/AutoReimportManager.h
+++ b/Engine/Source/Editor/UnrealEd/Public/AutoReimport/AutoReimportManager.h
@@ -4,18 +4,34 @@
 
 #include "AutoReimportManager.generated.h"
 
+/** Struct representing a path on disk, and its virtual mount point */
+struct FPathAndMountPoint
+{
+	FPathAndMountPoint() {}
+	FPathAndMountPoint(FString InPath, FString InMountPoint) : Path(MoveTemp(InPath)), MountPoint(MoveTemp(InMountPoint)) {}
+
+	/** The directory on disk. Absolute. */
+	FString Path;
+
+	/** The mount point, if any to which this directory relates */
+	FString MountPoint;
+};
+	
 /* Deals with auto reimporting of objects when the objects file on disk is modified*/
 UCLASS(config=Editor, transient)
 class UNREALED_API UAutoReimportManager : public UObject
 {
 	GENERATED_UCLASS_BODY()
 public:
-	
+
 	~UAutoReimportManager();
 
 	/** Initialize during engine startup */
 	void Initialize();
 
+	/** Get a list of absolute directories that we are monitoring */
+	TArray<FPathAndMountPoint> GetMonitoredDirectories() const;
+
 private:
 
 	/** UObject Interface */
diff --git a/Engine/Source/Editor/UnrealEd/Public/Dialogs/SDeleteAssetsDialog.h b/Engine/Source/Editor/UnrealEd/Public/Dialogs/SDeleteAssetsDialog.h
--- a/Engine/Source/Editor/UnrealEd/Public/Dialogs/SDeleteAssetsDialog.h
+++ b/Engine/Source/Editor/UnrealEd/Public/Dialogs/SDeleteAssetsDialog.h
@@ -73,6 +73,7 @@ private:
 	FReply Cancel();
 	FReply ForceDelete();
 	FReply ReplaceReferences();
+	void DeleteRelevantSourceContent();
 
 	/** Gets the text to display in the on disk referencing assets section when it is empty. */
 	FText GetReferencingAssetsEmptyText() const;
@@ -83,6 +84,9 @@ private:
 	/** Gets the text to display in the header for the 'how to proceed' section */
 	FText GetHandleText() const;
 
+	/** Get the text for the delete source content files tooltip */
+	FText GetDeleteSourceContentTooltip() const;
+
 	/** Returns the visibility of the section showing asset references on disk. */
 	EVisibility GetAssetReferencesVisiblity() const;
 
@@ -95,6 +99,9 @@ private:
 	/** Returns the visibility of the 'Delete' option */
 	EVisibility GetDeleteVisibility() const;
 
+	/** Returns the visibility of the 'Delete source content files' option */
+	EVisibility GetDeleteSourceFilesVisibility() const;
+
 	/** Returns if the 'Replace References' option should be available. */
 	bool CanReplaceReferences() const;
 
