[INTEGRATE] Change 2428539 by [REDACTED_PERSON][REDACTED_HANDLE] on 2015/02/02 14:53:29

	Explicitly whitelisting USceneComponent. If we want to expose large parts of the component hierarchy we should use Blueprintable and NotBlueprintable markup as we do elsewhere. UActorComponent will show up as a blueprintable type again...

diff --git a/Engine/Source/Editor/UnrealEd/Private/Kismet2/Kismet2.cpp b/Engine/Source/Editor/UnrealEd/Private/Kismet2/Kismet2.cpp
--- a/Engine/Source/Editor/UnrealEd/Private/Kismet2/Kismet2.cpp
+++ b/Engine/Source/Editor/UnrealEd/Private/Kismet2/Kismet2.cpp
@@ -876,7 +876,7 @@ bool FKismetEditorUtilities::CanCreateBlueprintOfClass(const UClass* Class)
 	const UEdGraphSchema_K2* K2Schema = GetDefault<UEdGraphSchema_K2>();
 	const bool bIsValidClass = Class->GetBoolMetaDataHierarchical(FBlueprintMetadata::MD_IsBlueprintBase)
 		|| (Class == UObject::StaticClass())
-		|| (bAllowBlueprintableComponents && !Class->HasAnyClassFlags(CLASS_Abstract) && Class->IsChildOf<UActorComponent>() && ( Class->HasAnyClassFlags(CLASS_CompiledFromBlueprint) || Class->HasMetaData( FBlueprintMetadata::MD_BlueprintSpawnableComponent ) ) )
+		|| (bAllowBlueprintableComponents && (Class->HasAnyClassFlags(CLASS_CompiledFromBlueprint) || Class == USceneComponent::StaticClass() || Class == UActorComponent::StaticClass()))
 		|| bIsBPGC;  // BPs are always considered inheritable
 
 	return bCanCreateBlueprint && bIsValidClass;
