[INTEGRATE] Change 2428678 by [REDACTED_PERSON][REDACTED_HANDLE].7 on 2015/02/02 16:21:25

	Source content files that are referenced by multiple assets are now only deleted if all referencing assets are being deleted

diff --git a/Engine/Source/Editor/UnrealEd/Private/AssetDeleteModel.cpp b/Engine/Source/Editor/UnrealEd/Private/AssetDeleteModel.cpp
--- a/Engine/Source/Editor/UnrealEd/Private/AssetDeleteModel.cpp
+++ b/Engine/Source/Editor/UnrealEd/Private/AssetDeleteModel.cpp
@@ -8,6 +8,7 @@
 #include "AssetEditorManager.h"
 #include "AutoReimport/AutoReimportUtilities.h"
 #include "AutoReimport/AutoReimportManager.h"
+#include "ISourceControlModule.h"
 
 #define LOCTEXT_NAMESPACE "FAssetDeleteModel"
 
@@ -16,7 +17,6 @@ FAssetDeleteModel::FAssetDeleteModel( const TArray<UObject*>& InObjectsToDelete
 	, bPendingObjectsCanBeReplaced(false)
 	, bIsAnythingReferencedInMemoryByNonUndo(false)
 	, bIsAnythingReferencedInMemoryByUndo(false)
-	, bHasSourceContentFilesToDelete(false)
 	, PendingDeleteIndex(0)
 	, ObjectsDeleted(0)
 {
@@ -66,7 +66,7 @@ void FAssetDeleteModel::Tick( const float InDeltaTime )
 		OnDiskReferences = TSet<FName>();
 		bIsAnythingReferencedInMemoryByNonUndo = false;
 		bIsAnythingReferencedInMemoryByUndo = false;
-		bHasSourceContentFilesToDelete = false;
+		SourceFileToAssetCount.Empty();
 		PendingDeleteIndex = 0;
 
 		SetState(Scanning);
@@ -102,35 +102,10 @@ void FAssetDeleteModel::Tick( const float InDeltaTime )
 			}
 			PendingDelete->RemainingMemoryReferences = NonPendingDeletedExternalInMemoryReferences;
 
-			if (GetDefault<UEditorLoadingSavingSettings>()->bMonitorContentDirectories)
-			{
-				Utils::ExtractSourceFilePaths(PendingDelete->GetObject(), PendingDelete->SourceContentFiles);
-
-				auto MonitoredDirectories = GUnrealEd->AutoReimportManager->GetMonitoredDirectories();
-
-				// Remove anything that's not under a monitored, mounted path
-				PendingDelete->SourceContentFiles.RemoveAll([&](const FString& InFilename){
-					for (const auto& Dir : MonitoredDirectories)
-					{
-						if (!Dir.MountPoint.IsEmpty() && InFilename.StartsWith(Dir.Path))
-						{
-							if (FPaths::FileExists(InFilename))
-							{
-								return false;
-							}
-							else
-							{
-								return true;
-							}
-						}
-					}
-					return true;
-				});
-			}
+			DiscoverSourceFileReferences(*PendingDelete);
 
 			bIsAnythingReferencedInMemoryByNonUndo |= PendingDelete->RemainingMemoryReferences > 0;
 			bIsAnythingReferencedInMemoryByUndo |= PendingDelete->IsReferencedInMemoryByUndo();
-			bHasSourceContentFilesToDelete |= PendingDelete->SourceContentFiles.Num() != 0;
 
 			PendingDeleteIndex++;
 		}
@@ -148,15 +123,111 @@ void FAssetDeleteModel::Tick( const float InDeltaTime )
 	}
 }
 
+void FAssetDeleteModel::DiscoverSourceFileReferences(FPendingDelete& PendingDelete)
+{
+	if (!GetDefault<UEditorLoadingSavingSettings>()->bMonitorContentDirectories)
+	{
+		return;
+	}
+	
+	// Start by extracting the files from the object
+	TArray<FString> SourceContentFiles;
+	Utils::ExtractSourceFilePaths(PendingDelete.GetObject(), SourceContentFiles);
+
+	auto MonitoredDirectories = GUnrealEd->AutoReimportManager->GetMonitoredDirectories();
+
+	// Remove anything that's not under a monitored, mounted path, or doesn't exist
+	SourceContentFiles.RemoveAll([&](const FString& InFilename){
+		for (const auto& Dir : MonitoredDirectories)
+		{
+			if (!Dir.MountPoint.IsEmpty() && InFilename.StartsWith(Dir.Path))
+			{
+				if (FPaths::FileExists(InFilename))
+				{
+					return false;
+				}
+				else
+				{
+					return true;
+				}
+			}
+		}
+		return true;
+	});
+
+	// Now accumulate references to the same source content file. We only offer to delete a file if it is only referenced by the deleted object(s)
+	IAssetRegistry& AssetRegistry = FModuleManager::LoadModuleChecked<FAssetRegistryModule>(TEXT("AssetRegistry")).Get();
+	for (const FString& SourcePath : SourceContentFiles)
+	{
+		int32* NumReferences = SourceFileToAssetCount.Find(SourcePath);
+		if (NumReferences)
+		{
+			(*NumReferences)--;
+		}
+		else
+		{
+			SourceFileToAssetCount.Add(SourcePath, Utils::FindAssetsPertainingToFile(AssetRegistry, SourcePath).Num());
+		}
+	}
+}
+
+bool FAssetDeleteModel::HasAnySourceContentFilesToDelete() const
+{
+	for (const auto& Pair : SourceFileToAssetCount)
+	{
+		if (Pair.Value == 0)
+		{
+			return true;
+		}
+	}
+	
+	return false;
+}
+
 void FAssetDeleteModel::DeleteSourceContentFiles()
 {
 	IFileManager& FileManager = IFileManager::Get();
-	for (const auto& PendingDelete : PendingDeletes)
+	ISourceControlProvider& SourceControlProvider = ISourceControlModule::Get().GetProvider();
+
+	for (const auto& Pair : SourceFileToAssetCount)
 	{
-		for (const auto& Path : PendingDelete->SourceContentFiles)
+		const auto& Path = Pair.Key;
+		// We can only delete this path if there are no (non-deleted) objects referencing it
+		if (Pair.Value != 0)
 		{
-			FileManager.Delete(*Path, false /* RequireExists */, true /* Even if read only */, true /* Quiet */);
+			continue;
 		}
+
+		if (ISourceControlModule::Get().IsEnabled())
+		{
+			const FSourceControlStatePtr SourceControlState = SourceControlProvider.GetState(Path, EStateCacheUsage::ForceUpdate);
+			const bool bIsSourceControlled = SourceControlState.IsValid() && SourceControlState->IsSourceControlled();
+
+			if (bIsSourceControlled)
+			{
+				// The file is managed by source control. Delete it through there.
+				TArray<FString> DeleteFilenames;
+				DeleteFilenames.Add(Path);
+
+				// Revert the file if it is checked out
+				const bool bIsAdded = SourceControlState->IsAdded();
+				if (SourceControlState->IsCheckedOut() || bIsAdded || SourceControlState->IsDeleted())
+				{
+					SourceControlProvider.Execute(ISourceControlOperation::Create<FRevert>(), DeleteFilenames);
+				}
+
+				// If it wasn't already marked as an add, we can ask the source control provider to delete the file
+				if (!bIsAdded)
+				{
+					// Open the file for delete
+					SourceControlProvider.Execute(ISourceControlOperation::Create<FDelete>(), DeleteFilenames);
+					continue;
+				}
+			}
+		}
+
+		// We'll just delete it ourself
+		FileManager.Delete(*Path, false /* RequireExists */, true /* Even if read only */, true /* Quiet */);
 	}
 }
 
diff --git a/Engine/Source/Editor/UnrealEd/Private/Dialogs/SDeleteAssetsDialog.cpp b/Engine/Source/Editor/UnrealEd/Private/Dialogs/SDeleteAssetsDialog.cpp
--- a/Engine/Source/Editor/UnrealEd/Private/Dialogs/SDeleteAssetsDialog.cpp
+++ b/Engine/Source/Editor/UnrealEd/Private/Dialogs/SDeleteAssetsDialog.cpp
@@ -427,15 +427,16 @@ FText SDeleteAssetsDialog::GetDeleteSourceContentTooltip() const
 	const FText RootText = LOCTEXT("DeleteSourceFiles_Tooltip", "When checked, the following source content files will also be deleted along with the assets:\n\n{0}");
 
 	FString AllFiles;
-	for (const auto& PendingDelete : *DeleteModel->GetPendingDeletedAssets())
+	for (const auto& PathAndAssetCount : DeleteModel->GetPendingDeletedSourceFileCounts())
 	{
-		for (const FString& Path : PendingDelete->SourceContentFiles)
+		// If this path is no longer referenced by deleted files, it's toast.
+		if (PathAndAssetCount.Value == 0)
 		{
 			if (!AllFiles.IsEmpty())
 			{
 				AllFiles += TEXT("\n");
 			}
-			AllFiles += Path;
+			AllFiles += PathAndAssetCount.Key;
 		}
 	}
 
diff --git a/Engine/Source/Editor/UnrealEd/Public/AssetDeleteModel.h b/Engine/Source/Editor/UnrealEd/Public/AssetDeleteModel.h
--- a/Engine/Source/Editor/UnrealEd/Public/AssetDeleteModel.h
+++ b/Engine/Source/Editor/UnrealEd/Public/AssetDeleteModel.h
@@ -40,9 +40,6 @@ public:
 	/** The on disk references to this object */
 	TArray<FName> DiskReferences;
 
-	/** Source content files for this pending delete operation */
-	TArray<FString> SourceContentFiles;
-
 	/** In memory references to this object (*excluding* the undo buffer) */
 	FReferencerInformationList MemoryReferences;
 
@@ -112,6 +109,9 @@ public:
 	/** Returns the pending deleted assets. */
 	const TArray< TSharedPtr< FPendingDelete > >* GetPendingDeletedAssets() const { return &PendingDeletes; };
 
+	/** Returns a map of currently discovered source content files, and the number of times they are referenced by non-deleted assets */
+	const TMap< FString, int32 >& GetPendingDeletedSourceFileCounts() const { return SourceFileToAssetCount; };
+
 	/** Returns the current state of the deletion process */
 	EState GetState() const { return State; }
 
@@ -164,7 +164,7 @@ public:
 	bool IsAnythingReferencedInMemoryByUndo() const { return bIsAnythingReferencedInMemoryByUndo; }
 
 	/** Check whether we have any source files residing under monitored, mounted paths to delete */
-	bool HasAnySourceContentFilesToDelete() const { return bHasSourceContentFilesToDelete; }
+	bool HasAnySourceContentFilesToDelete() const;
 
 	/** Goes to the next actor in the loaded level if it is available */
 	bool GoToNextReferenceInLevel() const;
@@ -188,6 +188,9 @@ private:
 	/** Computes the value that should be used for CanReplaceReferences */
 	bool ComputeCanReplaceReferences();
 
+	/** Discover source file references for the specified object */
+	void DiscoverSourceFileReferences(FPendingDelete& PendingDelete);
+
 private:
 
 	/** Holds an event delegate that is executed when the state changes */
@@ -196,6 +199,9 @@ private:
 	/** The assets being deleted */
 	TArray< TSharedPtr< FPendingDelete > > PendingDeletes;
 
+	/** A running count of source content filename -> number of assets referencing. For files that are no longer referenced, the count will be 0. */
+	TMap<FString, int32> SourceFileToAssetCount;
+
 	/** On disk references to the currently to be deleted objects */
 	TSet< FName > OnDiskReferences;
 
@@ -211,9 +217,6 @@ private:
 	/** Is any of the pending deleted assets being referenced in the undo stack. */
 	bool bIsAnythingReferencedInMemoryByUndo;
 
-	/** Boolean that specifies whether we have any relevant (in game dir) source content files to remove with the assets */
-	bool bHasSourceContentFilesToDelete;
-
 	/** A tick-to-tick state tracking variable so we know what pending deleted object we checked last. */
 	int32 PendingDeleteIndex;
 
