[INTEGRATE] Change 2428682 by [REDACTED_PERSON][REDACTED_HANDLE].7 on 2015/02/02 16:23:35

	Adding a source content file with the same name as an asset now only prints out a warning, rather than a full blown error.

diff --git a/Engine/Source/Editor/UnrealEd/Private/AutoReimport/ContentDirectoryMonitor.cpp b/Engine/Source/Editor/UnrealEd/Private/AutoReimport/ContentDirectoryMonitor.cpp
--- a/Engine/Source/Editor/UnrealEd/Private/AutoReimport/ContentDirectoryMonitor.cpp
+++ b/Engine/Source/Editor/UnrealEd/Private/AutoReimport/ContentDirectoryMonitor.cpp
@@ -125,7 +125,7 @@ void FContentDirectoryMonitor::ProcessAdditions(TArray<UPackage*>& OutPackagesTo
 		}
 		else if (FindPackage(nullptr, *PackagePath))
 		{
-			Context.AddMessage(EMessageSeverity::Error, FText::Format(LOCTEXT("Error_ExistingAsset", "Attempted to create a new asset that already exists ({0})."), FText::FromString(PackagePath)));
+			Context.AddMessage(EMessageSeverity::Warning, FText::Format(LOCTEXT("Error_ExistingAsset", "Can't create a new asset at {0} - one already exists."), FText::FromString(PackagePath)));
 		}
 		else
 		{
diff --git a/Engine/Source/Editor/UnrealEd/Private/AutoReimport/ReimportFeedbackContext.cpp b/Engine/Source/Editor/UnrealEd/Private/AutoReimport/ReimportFeedbackContext.cpp
--- a/Engine/Source/Editor/UnrealEd/Private/AutoReimport/ReimportFeedbackContext.cpp
+++ b/Engine/Source/Editor/UnrealEd/Private/AutoReimport/ReimportFeedbackContext.cpp
@@ -314,7 +314,7 @@ void FReimportFeedbackContext::Tick()
 
 void FReimportFeedbackContext::Destroy()
 {
-	MessageLog.Notify(FText(), EMessageSeverity::Warning);
+	MessageLog.Notify(FText(), EMessageSeverity::Error);
 
 	if (Notification.IsValid())
 	{
