[INTEGRATE] Change 2428711 by [REDACTED_PERSON][REDACTED_HANDLE] on 2015/02/02 16:39:40

	Transform (location, rotation, scale) change from an original component template is propagated into overridden component templates.
	- Reset works properly.
	- Rotation is fixed

diff --git a/Engine/Source/Editor/DetailCustomizations/Private/ComponentTransformDetails.cpp b/Engine/Source/Editor/DetailCustomizations/Private/ComponentTransformDetails.cpp
--- a/Engine/Source/Editor/DetailCustomizations/Private/ComponentTransformDetails.cpp
+++ b/Engine/Source/Editor/DetailCustomizations/Private/ComponentTransformDetails.cpp
@@ -865,14 +865,26 @@ void FComponentTransformDetails::OnToggleAbsoluteScale( bool bEnable )
 
 }
 
+struct FGetRootComponentArchetype
+{
+	static USceneComponent* Get(UObject* Object)
+	{
+		auto RootComponent = Object ? FComponentEditorUtils::GetSceneComponent(Object) : nullptr;
+		return RootComponent ? Cast<USceneComponent>(RootComponent->GetArchetype()) : nullptr;
+	}
+};
+
 FReply FComponentTransformDetails::OnLocationResetClicked()
 {
 	const FText TransactionName = LOCTEXT("ResetLocation", "Reset Location");
 	FScopedTransaction Transaction(TransactionName);
 
-	OnSetLocation(0.f, ETextCommit::Default, 0);
-	OnSetLocation(0.f, ETextCommit::Default, 1);
-	OnSetLocation(0.f, ETextCommit::Default, 2);
+	const USceneComponent* Archetype = FGetRootComponentArchetype::Get(SelectedObjects[0].Get());
+	const FVector Data = Archetype ? Archetype->RelativeLocation : FVector();
+
+	OnSetLocation(Data[0], ETextCommit::Default, 0);
+	OnSetLocation(Data[1], ETextCommit::Default, 1);
+	OnSetLocation(Data[2], ETextCommit::Default, 2);
 
 	return FReply::Handled();
 }
@@ -882,9 +894,12 @@ FReply FComponentTransformDetails::OnRotationResetClicked()
 	const FText TransactionName = LOCTEXT("ResetRotation", "Reset Rotation");
 	FScopedTransaction Transaction(TransactionName);
 
-	OnSetRotation(0.f, true, 0);
-	OnSetRotation(0.f, true, 1);
-	OnSetRotation(0.f, true, 2);
+	const USceneComponent* Archetype = FGetRootComponentArchetype::Get(SelectedObjects[0].Get());
+	const FRotator Data = Archetype ? Archetype->RelativeRotation : FRotator();
+
+	OnSetRotation(Data.Roll, true, 0);
+	OnSetRotation(Data.Pitch, true, 1);
+	OnSetRotation(Data.Yaw, true, 2);
 
 	return FReply::Handled();
 }
@@ -894,9 +909,12 @@ FReply FComponentTransformDetails::OnScaleResetClicked()
 	const FText TransactionName = LOCTEXT("ResetScale", "Reset Scale");
 	FScopedTransaction Transaction(TransactionName);
 
-	ScaleObject(1.f, 0, false, TransactionName);
-	ScaleObject(1.f, 1, false, TransactionName);
-	ScaleObject(1.f, 2, false, TransactionName);
+	const USceneComponent* Archetype = FGetRootComponentArchetype::Get(SelectedObjects[0].Get());
+	const FVector Data = Archetype ? Archetype->RelativeScale3D : FVector();
+
+	ScaleObject(Data[0], 0, false, TransactionName);
+	ScaleObject(Data[1], 1, false, TransactionName);
+	ScaleObject(Data[2], 2, false, TransactionName);
 
 	return FReply::Handled();
 }
@@ -978,7 +996,7 @@ void FComponentTransformDetails::CacheTransform()
 			if( RootComponent )
 			{
 				Loc = RootComponent->RelativeLocation;
-				Rot = (bEditingRotationInUI && !Object->HasAnyFlags(RF_ClassDefaultObject|RF_DefaultSubObject)) ? ObjectToRelativeRotationMap.FindOrAdd(Object) : RootComponent->RelativeRotation;
+				Rot = (bEditingRotationInUI && !Object->HasAnyFlags(RF_ClassDefaultObject | RF_DefaultSubObject)) ? ObjectToRelativeRotationMap.FindOrAdd(Object) : RootComponent->RelativeRotation;
 				Scale = RootComponent->RelativeScale3D;
 
 				if( ObjectIndex == 0 )
diff --git a/Engine/Source/Editor/UnrealEd/Private/Kismet2/ComponentEditorUtils.cpp b/Engine/Source/Editor/UnrealEd/Private/Kismet2/ComponentEditorUtils.cpp
--- a/Engine/Source/Editor/UnrealEd/Private/Kismet2/ComponentEditorUtils.cpp
+++ b/Engine/Source/Editor/UnrealEd/Private/Kismet2/ComponentEditorUtils.cpp
@@ -288,25 +288,7 @@ void FComponentEditorUtils::PropagateTransformPropertyChangeAmongOverridenTempla
 		USceneComponent* InstancedSceneComponent = Cast<USceneComponent>(ArchetypeInstance);
 		if (InstancedSceneComponent && InstancedSceneComponent->HasAnyFlags(RF_InheritableComponentTemplate) && !UpdatedComponents.Contains(InstancedSceneComponent))
 		{
-
-			static const UProperty* RelativeLocationProperty = FindFieldChecked<UProperty>(USceneComponent::StaticClass(), GET_MEMBER_NAME_CHECKED(USceneComponent, RelativeLocation));
-			if (RelativeLocationProperty != nullptr)
-			{
-				PropagateTransformPropertyChange(InstancedSceneComponent, RelativeLocationProperty, OldDefaultTransform.Trans, NewDefaultTransform.Trans, UpdatedComponents);
-			}
-
-			static const UProperty* RelativeRotationProperty = FindFieldChecked<UProperty>(USceneComponent::StaticClass(), GET_MEMBER_NAME_CHECKED(USceneComponent, RelativeLocation));
-			if (RelativeRotationProperty != nullptr)
-			{
-				PropagateTransformPropertyChange(InstancedSceneComponent, RelativeRotationProperty, OldDefaultTransform.Rot, NewDefaultTransform.Rot, UpdatedComponents);
-			}
-
-			static const UProperty* RelativeScale3DProperty = FindFieldChecked<UProperty>(USceneComponent::StaticClass(), GET_MEMBER_NAME_CHECKED(USceneComponent, RelativeScale3D));
-			if (RelativeScale3DProperty != nullptr)
-			{
-				PropagateTransformPropertyChange(InstancedSceneComponent, RelativeScale3DProperty, OldDefaultTransform.Scale, NewDefaultTransform.Scale, UpdatedComponents);
-			}
-
+			PropagateTransformInner(InstancedSceneComponent, OldDefaultTransform, NewDefaultTransform, UpdatedComponents);
 		}
 	}
 }
@@ -326,23 +308,28 @@ void FComponentEditorUtils::PropagateTransformPropertyChange(
 		USceneComponent* InstancedSceneComponent = FComponentEditorUtils::GetSceneComponent(ArchetypeInstances[InstanceIndex], InSceneComponentTemplate);
 		if(InstancedSceneComponent != nullptr && !UpdatedComponents.Contains(InstancedSceneComponent))
 		{
-			static const UProperty* RelativeLocationProperty = FindFieldChecked<UProperty>( USceneComponent::StaticClass(), GET_MEMBER_NAME_CHECKED(USceneComponent, RelativeLocation));
-			if(RelativeLocationProperty != nullptr)
-			{
-				PropagateTransformPropertyChange(InstancedSceneComponent, RelativeLocationProperty, OldDefaultTransform.Trans, NewDefaultTransform.Trans, UpdatedComponents);
-			}
+			PropagateTransformInner(InstancedSceneComponent, OldDefaultTransform, NewDefaultTransform, UpdatedComponents);
+		}
+	}
+}
 
-			static const UProperty* RelativeRotationProperty = FindFieldChecked<UProperty>( USceneComponent::StaticClass(), GET_MEMBER_NAME_CHECKED(USceneComponent, RelativeLocation) );
-			if(RelativeRotationProperty != nullptr)
-			{
-				PropagateTransformPropertyChange(InstancedSceneComponent, RelativeRotationProperty, OldDefaultTransform.Rot, NewDefaultTransform.Rot, UpdatedComponents);
-			}
+void FComponentEditorUtils::PropagateTransformInner(class USceneComponent* InstancedSceneComponent, const FTransformData& OldDefaultTransform, const FTransformData& NewDefaultTransform, TSet<class USceneComponent*>& UpdatedComponents)
+{
+	static const UProperty* RelativeLocationProperty = FindFieldChecked<UProperty>(USceneComponent::StaticClass(), GET_MEMBER_NAME_CHECKED(USceneComponent, RelativeLocation));
+	if (RelativeLocationProperty != nullptr)
+	{
+		PropagateTransformPropertyChange(InstancedSceneComponent, RelativeLocationProperty, OldDefaultTransform.Trans, NewDefaultTransform.Trans, UpdatedComponents);
+	}
 
-			static const UProperty* RelativeScale3DProperty = FindFieldChecked<UProperty>( USceneComponent::StaticClass(), GET_MEMBER_NAME_CHECKED(USceneComponent, RelativeScale3D) );
-			if(RelativeScale3DProperty != nullptr)
-			{
-				PropagateTransformPropertyChange(InstancedSceneComponent, RelativeScale3DProperty, OldDefaultTransform.Scale, NewDefaultTransform.Scale, UpdatedComponents);
-			}
-		}
+	static const UProperty* RelativeRotationProperty = FindFieldChecked<UProperty>(USceneComponent::StaticClass(), GET_MEMBER_NAME_CHECKED(USceneComponent, RelativeRotation));
+	if (RelativeRotationProperty != nullptr)
+	{
+		PropagateTransformPropertyChange(InstancedSceneComponent, RelativeRotationProperty, OldDefaultTransform.Rot, NewDefaultTransform.Rot, UpdatedComponents);
+	}
+
+	static const UProperty* RelativeScale3DProperty = FindFieldChecked<UProperty>(USceneComponent::StaticClass(), GET_MEMBER_NAME_CHECKED(USceneComponent, RelativeScale3D));
+	if (RelativeScale3DProperty != nullptr)
+	{
+		PropagateTransformPropertyChange(InstancedSceneComponent, RelativeScale3DProperty, OldDefaultTransform.Scale, NewDefaultTransform.Scale, UpdatedComponents);
 	}
 }
\ No newline at end of file
diff --git a/Engine/Source/Editor/UnrealEd/Public/Kismet2/ComponentEditorUtils.h b/Engine/Source/Editor/UnrealEd/Public/Kismet2/ComponentEditorUtils.h
--- a/Engine/Source/Editor/UnrealEd/Public/Kismet2/ComponentEditorUtils.h
+++ b/Engine/Source/Editor/UnrealEd/Public/Kismet2/ComponentEditorUtils.h
@@ -119,4 +119,7 @@ public:
 			UpdatedComponents.Add(InSceneComponent);
 		}
 	}
+private:
+	static void PropagateTransformInner(class USceneComponent* InstancedSceneComponent, const FTransformData& OldDefaultTransform, const FTransformData& NewDefaultTransform, TSet<class USceneComponent*>& UpdatedComponents);
+
 };
\ No newline at end of file
