[INTEGRATE] Change 2428738 by [REDACTED_PERSON][REDACTED_HANDLE].7 on 2015/02/02 16:58:56

	FEditorReimportHandler::GetPriority now has a default implementation

diff --git a/Engine/Source/Editor/UnrealEd/Private/Editor.cpp b/Engine/Source/Editor/UnrealEd/Private/Editor.cpp
--- a/Engine/Source/Editor/UnrealEd/Private/Editor.cpp
+++ b/Engine/Source/Editor/UnrealEd/Private/Editor.cpp
@@ -2890,6 +2890,11 @@ FReimportManager::~FReimportManager()
 	Handlers.Empty();
 }
 
+int32 FReimportHandler::GetPriority() const
+{
+	return UFactory::DefaultImportPriority;
+}
+
 /*-----------------------------------------------------------------------------
 	PIE helpers.
 -----------------------------------------------------------------------------*/
diff --git a/Engine/Source/Editor/UnrealEd/Public/EditorReimportHandler.h b/Engine/Source/Editor/UnrealEd/Public/EditorReimportHandler.h
--- a/Engine/Source/Editor/UnrealEd/Public/EditorReimportHandler.h
+++ b/Engine/Source/Editor/UnrealEd/Public/EditorReimportHandler.h
@@ -171,5 +171,9 @@ public:
 	 */
 	virtual EReimportResult::Type Reimport( UObject* Obj ) = 0;
 
-	virtual int32 GetPriority() const = 0;
+	/**
+	 * Get the import priority for this handler.
+	 * Import handlers with higher priority values will take precedent over lower priorities.
+	 */
+	UNREALED_API virtual int32 GetPriority() const;
 };
