[INTEGRATE] Change 2428964 by [REDACTED_PERSON][REDACTED_HANDLE].7 on 2015/02/02 18:19:40

	Removed timeout from reimport feedback notification

diff --git a/Engine/Source/Editor/UnrealEd/Private/AutoReimport/AutoReimportManager.cpp b/Engine/Source/Editor/UnrealEd/Private/AutoReimport/AutoReimportManager.cpp
--- a/Engine/Source/Editor/UnrealEd/Private/AutoReimport/AutoReimportManager.cpp
+++ b/Engine/Source/Editor/UnrealEd/Private/AutoReimport/AutoReimportManager.cpp
@@ -454,11 +454,6 @@ void FAutoReimportManager::Cleanup()
 
 void FAutoReimportManager::Tick(float DeltaTime)
 {
-	if (FeedbackContextOverride.IsValid())
-	{
-		FeedbackContextOverride->Tick();
-	}
-	
 	// Never spend more than a 60fps frame doing this work (meaning we shouldn't drop below 30fps), we can do more if we're throttling CPU usage (ie editor running in background)
 	const FTimeLimit TimeLimit(GEditor->ShouldThrottleCPUUsage() ? 1 / 6.f : 1.f / 60.f);
 	StateMachine.Tick(TimeLimit);
diff --git a/Engine/Source/Editor/UnrealEd/Private/AutoReimport/ReimportFeedbackContext.cpp b/Engine/Source/Editor/UnrealEd/Private/AutoReimport/ReimportFeedbackContext.cpp
--- a/Engine/Source/Editor/UnrealEd/Private/AutoReimport/ReimportFeedbackContext.cpp
+++ b/Engine/Source/Editor/UnrealEd/Private/AutoReimport/ReimportFeedbackContext.cpp
@@ -294,22 +294,15 @@ FReimportFeedbackContext::FReimportFeedbackContext()
 
 void FReimportFeedbackContext::Initialize(TSharedRef<SReimportFeedback> Widget)
 {
-	ShowNotificationDelay = FTimeLimit(.5f);
 	NotificationContent = Widget;
 
-	MessageLog.NewPage(FText::Format(NSLOCTEXT("ReimportContext", "MessageLogPageLabel", "Outstanding source content changes {0}"), FText::AsTime(FDateTime::Now())));
-}
+	FNotificationInfo Info(SharedThis(this));
+	Info.ExpireDuration = 3.f;
+	Info.bFireAndForget = false;
 
-void FReimportFeedbackContext::Tick()
-{
-	if (!Notification.IsValid() && ShowNotificationDelay.Exceeded())
-	{
-		FNotificationInfo Info(SharedThis(this));
-		Info.ExpireDuration = 3.f;
-		Info.bFireAndForget = false;
+	Notification = FSlateNotificationManager::Get().AddNotification(Info);
 
-		Notification = FSlateNotificationManager::Get().AddNotification(Info);
-	}
+	MessageLog.NewPage(FText::Format(NSLOCTEXT("ReimportContext", "MessageLogPageLabel", "Outstanding source content changes {0}"), FText::AsTime(FDateTime::Now())));
 }
 
 void FReimportFeedbackContext::Destroy()
diff --git a/Engine/Source/Editor/UnrealEd/Private/AutoReimport/ReimportFeedbackContext.h b/Engine/Source/Editor/UnrealEd/Private/AutoReimport/ReimportFeedbackContext.h
--- a/Engine/Source/Editor/UnrealEd/Private/AutoReimport/ReimportFeedbackContext.h
+++ b/Engine/Source/Editor/UnrealEd/Private/AutoReimport/ReimportFeedbackContext.h
@@ -59,9 +59,6 @@ public:
 	/** Destroy this reimport context */
 	void Destroy();
 
-	/** Tick this context - potentially opening a new notification if the process has taken a while */
-	void Tick();
-
 	TSharedPtr<SReimportFeedback> GetContent() { return NotificationContent; }
 
 	/** Add a message to the context (and message log) */
@@ -83,9 +80,6 @@ private:
 	/** True if we need to notifiy about new messages on the message log */
 	bool NotifyMessageLog;
 
-	/** The time when the last change to the cache was reported */
-	FTimeLimit ShowNotificationDelay;
-
 	/** The notification that is shown when the context is active */
 	TSharedPtr<class SNotificationItem> Notification;
 
