[INTEGRATE] Change 2429030 by [REDACTED_PERSON][REDACTED_HANDLE] on 2015/02/02 19:10:35

	Fixed a code feature pack added to blank project not generating projects or compiling new code

	UE-8855 - When adding code pack, user is asked to open the code files, but the IDE shows a blank page

	It was trying to spawn VS without having actually generated projects, and it was also trying to use DoHotReloadFromEditor, which only works when you already have a module loaded.

diff --git a/Engine/Source/Editor/UnrealEd/Private/Factories/PackFactory.cpp b/Engine/Source/Editor/UnrealEd/Private/Factories/PackFactory.cpp
--- a/Engine/Source/Editor/UnrealEd/Private/Factories/PackFactory.cpp
+++ b/Engine/Source/Editor/UnrealEd/Private/Factories/PackFactory.cpp
@@ -305,13 +305,16 @@ UObject* UPackFactory::FactoryCreateBinary
 			}
 		}
 
+		bool bProjectHadSourceFiles = false;
+
 		// If we have source files, set up the project files if necessary and the game name redirects for blueprints saved with class
 		// references to the module name from the source template
 		if (bContainsSource)
 		{
 			FGameProjectGenerationModule& GameProjectModule = FModuleManager::LoadModuleChecked<FGameProjectGenerationModule>(TEXT("GameProjectGeneration"));
+			bProjectHadSourceFiles = GameProjectModule.Get().ProjectHasCodeFiles();
 
-			if (!GameProjectModule.Get().ProjectHasCodeFiles())
+			if (!bProjectHadSourceFiles)
 			{
 				TArray<FString> StartupModuleNames;
 				TArray<FString> CreatedFiles;
@@ -479,6 +482,46 @@ UObject* UPackFactory::FactoryCreateBinary
 			// If we wrote out source files, kick off the hot reload process
 			if (WrittenSourceFiles.Num() > 0)
 			{
+				// Update the game projects before we attempt to build
+				FGameProjectGenerationModule& GameProjectModule = FModuleManager::LoadModuleChecked<FGameProjectGenerationModule>(TEXT("GameProjectGeneration"));
+				FText FailReason;
+				if (!GameProjectModule.UpdateCodeProject(FailReason))
+				{
+					FMessageDialog::Open(EAppMsgType::Ok, FailReason);
+				}
+
+				// Compile the new code, either using the in editor hot-reload (if an existing module), or as a brand new module (if no existing code)
+				IHotReloadInterface& HotReloadSupport = FModuleManager::LoadModuleChecked<IHotReloadInterface>("HotReload");
+				if (bProjectHadSourceFiles)
+				{
+					// We can only hot-reload via DoHotReloadFromEditor when we already had code in our project
+					if (!HotReloadSupport.IsCurrentlyCompiling())
+					{					
+						HotReloadSupport.DoHotReloadFromEditor();
+					
+						UE_LOG(LogPackFactory, Warning, TEXT("Starting compile."));
+						// Wait until the compiler finishes
+						while (HotReloadSupport.IsCurrentlyCompiling())
+						{
+							// Tick slate while we wait
+							FSlateApplication::Get().Tick();						
+							// Tick hot reload - or nothing will happen !
+							HotReloadSupport.Tick();
+						}
+						UE_LOG(LogPackFactory, Warning, TEXT("Compile complete."));
+					}
+				}
+				else
+				{
+					const bool bReloadAfterCompiling = true;
+					const bool bForceCodeProject = true;
+					const bool bFailIfGeneratedCodeChanges = false;
+					if (!HotReloadSupport.RecompileModule(FApp::GetGameName(), bReloadAfterCompiling, *GWarn, bFailIfGeneratedCodeChanges, bForceCodeProject))
+					{
+						FMessageDialog::Open(EAppMsgType::Ok, NSLOCTEXT("PackFactory", "FailedToCompileNewGameModule", "Failed to compile newly created game module."));
+					}
+				}
+
 				// Ask about editing code where applicable
 				if (FSlateApplication::Get().SupportsSourceAccess() )
 				{
@@ -489,27 +532,6 @@ UObject* UPackFactory::FactoryCreateBinary
 						FSourceCodeNavigation::OpenSourceFiles(WrittenSourceFiles);
 					}
 				}
-				
-				// Tell the user we are about to compile and fire it off (if we need to)
-				IHotReloadInterface& HotReloadSupport = FModuleManager::LoadModuleChecked<IHotReloadInterface>("HotReload");
-				if (!HotReloadSupport.IsCurrentlyCompiling())
-				{					
-					const FText Message = NSLOCTEXT("PackFactory", "WaitingForCompile", "Press OK to begin compling. Pack will installation will complete when the compiler has finished.");
-					FMessageDialog::Open(EAppMsgType::Ok, Message);
-				
-					HotReloadSupport.DoHotReloadFromEditor();
-					
-					UE_LOG(LogPackFactory, Warning, TEXT("Starting compile."));
-					// Wait until the compiler finishes
-					while (HotReloadSupport.IsCurrentlyCompiling())
-					{
-						// Tick slate while we wait
-						FSlateApplication::Get().Tick();						
-						// Tick hot reload - or nothing will happen !
-						HotReloadSupport.Tick();
-					}
-					UE_LOG(LogPackFactory, Warning, TEXT("Compile complete."));
-				}
 			}
 			
 			// Find an asset to return (It will be marked as dirty)
