#ANIM: Add Animation editor setting to editor setting

diff --git a/Engine/Source/Editor/EditorSettingsViewer/Private/EditorSettingsViewerModule.cpp b/Engine/Source/Editor/EditorSettingsViewer/Private/EditorSettingsViewerModule.cpp
--- a/Engine/Source/Editor/EditorSettingsViewer/Private/EditorSettingsViewerModule.cpp
+++ b/Engine/Source/Editor/EditorSettingsViewer/Private/EditorSettingsViewerModule.cpp
@@ -215,6 +215,13 @@ protected:
 			LOCTEXT("ContentEditorsGraphBlueprintSettingsDescription", "Customize Blueprint Editors."),
 			GetMutableDefault<UBlueprintEditorSettings>()
 		);
+
+		// Persona editors
+		SettingsModule.RegisterSettings("Editor", "ContentEditors", "Persona",
+			LOCTEXT("ContentEditorsPersonaSettingsName", "Animation Editor"),
+			LOCTEXT("ContentEditorsPersonaSettingsDescription", "Customize Persona Editor."),
+			GetMutableDefault<UPersonaOptions>()
+		);
 	}
 
 	/** Unregisters all settings. */
@@ -242,6 +249,7 @@ protected:
 			SettingsModule->UnregisterSettings("Editor", "ContentEditors", "ContentBrowser");
 //			SettingsModule->UnregisterSettings("Editor", "ContentEditors", "DestructableMeshEditor");
 			SettingsModule->UnregisterSettings("Editor", "ContentEditors", "GraphEditor");
+			SettingsModule->UnregisterSettings("Editor", "ContentEditors", "Persona");
 		}
 	}
 
diff --git a/Engine/Source/Editor/Persona/Private/SAnimationEditorViewport.cpp b/Engine/Source/Editor/Persona/Private/SAnimationEditorViewport.cpp
--- a/Engine/Source/Editor/Persona/Private/SAnimationEditorViewport.cpp
+++ b/Engine/Source/Editor/Persona/Private/SAnimationEditorViewport.cpp
@@ -387,7 +387,7 @@ void SAnimationEditorViewportTabBody::Construct(const FArguments& InArgs)
 	TSharedRef<FAnimationViewportClient> AnimViewportClient = StaticCastSharedRef<FAnimationViewportClient>(LevelViewportClient.ToSharedRef());
 
 	// Load the view mode from config
-	AnimViewportClient->SetViewMode((EViewModeIndex)AnimViewportClient->ConfigOption->ViewModeIndex);
+	AnimViewportClient->SetViewMode(AnimViewportClient->ConfigOption->ViewModeIndex);
 	UpdateShowFlagForMeshEdges();
 
 	UpdateViewportClientPlaybackScale();
diff --git a/Engine/Source/Editor/UnrealEd/Classes/Preferences/PersonaOptions.h b/Engine/Source/Editor/UnrealEd/Classes/Preferences/PersonaOptions.h
--- a/Engine/Source/Editor/UnrealEd/Classes/Preferences/PersonaOptions.h
+++ b/Engine/Source/Editor/UnrealEd/Classes/Preferences/PersonaOptions.h
@@ -31,14 +31,14 @@ class UNREALED_API UPersonaOptions : public UObject
 	uint32 bMuteAudio:1;
 
 	// currently Stats can have None, Basic and Detailed. Please refer to EDisplayInfoMode.
-	UPROPERTY(EditAnywhere, config, Category = Options)
+	UPROPERTY(EditAnywhere, config, Category = Options, meta=(ClampMin ="0", ClampMax = "2", UIMin = "0", UIMax = "2"))
 	int32 ShowMeshStats;
 
 	UPROPERTY(EditAnywhere, config, Category=Options)
 	int32 GridSize;
 
 	UPROPERTY(EditAnywhere, config, Category=Options)
-	int32 ViewModeIndex;
+	TEnumAsByte<EViewModeIndex> ViewModeIndex;
 
 	UPROPERTY(EditAnywhere, config, Category=Options)
 	FLinearColor ViewportBackgroundColor;
@@ -57,7 +57,7 @@ public:
 	void SetShowSky( bool bInShowSky );
 	void SetMuteAudio( bool bInMuteAudio );
 	void SetGridSize( int32 InGridSize );
-	void SetViewModeIndex( int32 InViewModeIndex );
+	void SetViewModeIndex( EViewModeIndex InViewModeIndex );
 	void SetViewFOV( float InViewFOV );
 	void SetDefaultLocalAxesSelection( uint32 InDefaultLocalAxesSelection );
 	void SetShowMeshStats( int32 InShowMeshStats );
diff --git a/Engine/Source/Editor/UnrealEd/Private/PreferenceStubs.cpp b/Engine/Source/Editor/UnrealEd/Private/PreferenceStubs.cpp
--- a/Engine/Source/Editor/UnrealEd/Private/PreferenceStubs.cpp
+++ b/Engine/Source/Editor/UnrealEd/Private/PreferenceStubs.cpp
@@ -65,7 +65,7 @@ void UPersonaOptions::SetGridSize( int32 InGridSize )
 	SaveConfig();
 }
 
-void UPersonaOptions::SetViewModeIndex( int32 InViewModeIndex )
+void UPersonaOptions::SetViewModeIndex( EViewModeIndex InViewModeIndex )
 {
 	ViewModeIndex = InViewModeIndex;
 	SaveConfig();
diff --git a/Engine/Source/Runtime/Engine/Classes/Engine/EngineBaseTypes.h b/Engine/Source/Runtime/Engine/Classes/Engine/EngineBaseTypes.h
--- a/Engine/Source/Runtime/Engine/Classes/Engine/EngineBaseTypes.h
+++ b/Engine/Source/Runtime/Engine/Classes/Engine/EngineBaseTypes.h
@@ -701,9 +701,9 @@ enum EViewModeIndex
 //	VMI_VertexDensity = 15,
 	VMI_CollisionPawn = 15, 
 	VMI_CollisionVisibility = 16, 
-	VMI_Max,
+	VMI_Max UMETA(Hidden),
 
-	VMI_Unknown = 255,
+	VMI_Unknown = 255 UMETA(Hidden),
 };
 
 
