[INTEGRATE] Change 2429225 by [REDACTED_PERSON][REDACTED_HANDLE] on 2015/02/02 22:24:21

	Editor: Prevent crash after selecting and moving an actor (which has at least one ConstructionScript spawned component attached to another ConstructionScript spawned component) and then undoing twice
	- No longer dirty components on selection when adding the selection delegate, as the act of undoing selection will cause them to be restored or removed as needed
	- Add ensures in SCS tree rebuilding for pending kill to provide a more informative error message prior to crashing if a similar issue arises in the future
	[UE-8366]

diff --git a/Engine/Source/Editor/Kismet/Private/SSCSEditor.cpp b/Engine/Source/Editor/Kismet/Private/SSCSEditor.cpp
--- a/Engine/Source/Editor/Kismet/Private/SSCSEditor.cpp
+++ b/Engine/Source/Editor/Kismet/Private/SSCSEditor.cpp
@@ -399,20 +399,22 @@ void FSCSEditorTreeNode::AddChild(FSCSEditorTreeNodePtrType InChildNodePtr)
 	else if (IsInstanced())
 	{
 		USceneComponent* ChildInstance = Cast<USceneComponent>(InChildNodePtr->GetComponentTemplate());
-		check(ChildInstance != nullptr);
-
-		USceneComponent* ParentInstance = Cast<USceneComponent>(GetComponentTemplate());
-		check(ParentInstance != nullptr);
-
-		// Handle attachment at the instance level
-		if (ChildInstance->AttachParent != ParentInstance)
+		if (ensure(ChildInstance != nullptr))
 		{
-			AActor* Owner = ParentInstance->GetOwner();
-			if (Owner->GetRootComponent() == ChildInstance)
+			USceneComponent* ParentInstance = Cast<USceneComponent>(GetComponentTemplate());
+			if (ensure(ParentInstance != nullptr))
 			{
-				Owner->SetRootComponent(ParentInstance);
+				// Handle attachment at the instance level
+				if (ChildInstance->AttachParent != ParentInstance)
+				{
+					AActor* Owner = ParentInstance->GetOwner();
+					if (Owner->GetRootComponent() == ChildInstance)
+					{
+						Owner->SetRootComponent(ParentInstance);
+					}
+					ChildInstance->AttachTo(ParentInstance, NAME_None, EAttachLocation::KeepWorldPosition);
+				}
 			}
-			ChildInstance->AttachTo(ParentInstance, NAME_None, EAttachLocation::KeepWorldPosition);
 		}
 	}
 }
@@ -5186,6 +5188,7 @@ FSCSEditorTreeNodePtrType SSCSEditor::AddTreeNodeFromComponent(USceneComponent*
 	FSCSEditorTreeNodePtrType NewNodePtr;
 
 	check(InSceneComponent != NULL);
+	ensure(!InSceneComponent->HasAnyFlags(RF_PendingKill));
 
 	// If the given component has a parent, and if we're not in "instance" mode OR the owner of the parent matches the Actor instance we're editing
 	if(InSceneComponent->AttachParent != NULL
diff --git a/Engine/Source/Editor/UnrealEd/Private/Kismet2/ComponentEditorUtils.cpp b/Engine/Source/Editor/UnrealEd/Private/Kismet2/ComponentEditorUtils.cpp
--- a/Engine/Source/Editor/UnrealEd/Private/Kismet2/ComponentEditorUtils.cpp
+++ b/Engine/Source/Editor/UnrealEd/Private/Kismet2/ComponentEditorUtils.cpp
@@ -184,20 +184,16 @@ void FComponentEditorUtils::BindComponentSelectionOverride(USceneComponent* Scen
 {
 	if (SceneComponent)
 	{
-		// Modifying the selection override delegate is purely cosmetic and shouldn't dirty the package
-		const bool bShouldMarkPackageDirty = false;
-
 		// If the scene component is a primitive component, set the override for it
 		auto PrimComponent = Cast<UPrimitiveComponent>(SceneComponent);
 		if (PrimComponent && PrimComponent->SelectionOverrideDelegate.IsBound() != bBind)
 		{
-			PrimComponent->Modify(bShouldMarkPackageDirty);
 			if (bBind)
-		{
-			PrimComponent->SelectionOverrideDelegate = UPrimitiveComponent::FSelectionOverride::CreateUObject(GUnrealEd, &UUnrealEdEngine::IsComponentSelected);
-		}
-		else
-		{
+			{
+				PrimComponent->SelectionOverrideDelegate = UPrimitiveComponent::FSelectionOverride::CreateUObject(GUnrealEd, &UUnrealEdEngine::IsComponentSelected);
+			}
+			else
+			{
 				PrimComponent->SelectionOverrideDelegate.Unbind();
 			}
 		}
@@ -209,20 +205,19 @@ void FComponentEditorUtils::BindComponentSelectionOverride(USceneComponent* Scen
 				PrimComponent = Cast<UPrimitiveComponent>(Component);
 				if (PrimComponent && PrimComponent->IsEditorOnly() && PrimComponent->SelectionOverrideDelegate.IsBound() != bBind)
 				{
-					PrimComponent->Modify(bShouldMarkPackageDirty);
 					if (bBind)
-				{
-					PrimComponent->SelectionOverrideDelegate = UPrimitiveComponent::FSelectionOverride::CreateUObject(GUnrealEd, &UUnrealEdEngine::IsComponentSelected);
-				}
+					{
+						PrimComponent->SelectionOverrideDelegate = UPrimitiveComponent::FSelectionOverride::CreateUObject(GUnrealEd, &UUnrealEdEngine::IsComponentSelected);
+					}
 					else
 					{
 						PrimComponent->SelectionOverrideDelegate.Unbind();
 					}
 				}
 			}
-			}
 		}
 	}
+}
 
 bool FComponentEditorUtils::AttemptApplyMaterialToComponent(USceneComponent* SceneComponent, UMaterialInterface* MaterialToApply, int32 OptionalMaterialSlot)
 {
