#MATINEE: Fix Matinee export to be

1. Export Level would export everything including camera
2. Export Matinee will now look for FbxNode that is created and then attach curve to it
3. Export Matinee all will also export anim track

diff --git a/Engine/Source/Editor/UnrealEd/Private/Fbx/FbxAnimationExport.cpp b/Engine/Source/Editor/UnrealEd/Private/Fbx/FbxAnimationExport.cpp
--- a/Engine/Source/Editor/UnrealEd/Private/Fbx/FbxAnimationExport.cpp
+++ b/Engine/Source/Editor/UnrealEd/Private/Fbx/FbxAnimationExport.cpp
@@ -217,18 +217,18 @@ void FFbxExporter::CorrectAnimTrackInterpolation( TArray<FbxNode*>& BoneNodes, F
 }
 
 
-void FFbxExporter::ExportAnimSequence( const UAnimSequence* AnimSeq, const USkeletalMesh* SkelMesh, bool bExportSkelMesh, const TCHAR* MeshName, FbxNode* ActorRootNode )
+FbxNode* FFbxExporter::ExportAnimSequence( const UAnimSequence* AnimSeq, const USkeletalMesh* SkelMesh, bool bExportSkelMesh, const TCHAR* MeshName, FbxNode* ActorRootNode )
 {
 	if( Scene == NULL || AnimSeq == NULL || SkelMesh == NULL )
 	{
- 		return;
+ 		return NULL;
 	}
 
 
 	FbxNode* RootNode = (ActorRootNode)? ActorRootNode : Scene->GetRootNode();
 	// Create the Skeleton
 	TArray<FbxNode*> BoneNodes;
-	FbxNode* SkeletonRootNode = CreateSkeleton(*SkelMesh, BoneNodes);
+	FbxNode* SkeletonRootNode = CreateSkeleton(SkelMesh, BoneNodes);
 	RootNode->AddChild(SkeletonRootNode);
 
 
@@ -262,7 +262,7 @@ void FFbxExporter::ExportAnimSequence( const UAnimSequence* AnimSeq, const USkel
 		}
 
 		// Add the mesh
-		FbxNode* MeshRootNode = CreateMesh(*SkelMesh, *MeshNodeName);
+		FbxNode* MeshRootNode = CreateMesh(SkelMesh, *MeshNodeName);
 		if(MeshRootNode)
 		{
 			RootNode->AddChild(MeshRootNode);
@@ -271,12 +271,14 @@ void FFbxExporter::ExportAnimSequence( const UAnimSequence* AnimSeq, const USkel
 		if(SkeletonRootNode && MeshRootNode)
 		{
 			// Bind the mesh to the skeleton
-			BindMeshToSkeleton(*SkelMesh, MeshRootNode, BoneNodes);
+			BindMeshToSkeleton(SkelMesh, MeshRootNode, BoneNodes);
 
 			// Add the bind pose
 			CreateBindPose(MeshRootNode);
 		}
 	}
+
+	return SkeletonRootNode;
 }
 
 
@@ -299,7 +301,7 @@ void FFbxExporter::ExportAnimSequencesAsSingle( USkeletalMesh* SkelMesh, const A
 
 	// Create the Skeleton
 	TArray<FbxNode*> BoneNodes;
-	FbxNode* SkeletonRootNode = CreateSkeleton(*SkelMesh, BoneNodes);
+	FbxNode* SkeletonRootNode = CreateSkeleton(SkelMesh, BoneNodes);
 	BaseNode->AddChild(SkeletonRootNode);
 
 	bool bAnyObjectMissingSourceData = false;
@@ -339,10 +341,7 @@ void FFbxExporter::ExportAnimSequencesAsSingle( USkeletalMesh* SkelMesh, const A
  */
 void FFbxExporter::ExportMatineeGroup(class AMatineeActor* MatineeActor, USkeletalMeshComponent* SkeletalMeshComponent)
 {
-	static const float SamplingRate = 1.f / DEFAULT_SAMPLERATE;
-
 	float MatineeLength = MatineeActor->MatineeData->InterpLength;
-
 	
 	if (Scene == NULL || MatineeActor == NULL || SkeletalMeshComponent == NULL || MatineeLength == 0)
 	{
@@ -353,24 +352,43 @@ void FFbxExporter::ExportMatineeGroup(class AMatineeActor* MatineeActor, USkelet
 
 	FbxNode* BaseNode = FbxNode::Create(Scene, NodeName);
 	Scene->GetRootNode()->AddChild(BaseNode);
+
 	AActor* Owner = SkeletalMeshComponent->GetOwner();
-	if( Owner && Owner->GetRootComponent() )
+	if(Owner && Owner->GetRootComponent())
 	{
 		// Set the default position of the actor on the transforms
 		// The UE3 transformation is different from FBX's Z-up: invert the Y-axis for translations and the Y/Z angle values in rotations.
-		BaseNode->LclTranslation.Set(Converter.ConvertToFbxPos( Owner->GetActorLocation() ) );
-		BaseNode->LclRotation.Set(Converter.ConvertToFbxRot( Owner->GetActorRotation().Euler() ) ) ;
-		BaseNode->LclScaling.Set(Converter.ConvertToFbxScale( Owner->GetRootComponent()->RelativeScale3D ) );
+		BaseNode->LclTranslation.Set(Converter.ConvertToFbxPos(Owner->GetActorLocation()));
+		BaseNode->LclRotation.Set(Converter.ConvertToFbxRot(Owner->GetActorRotation().Euler()));
+		BaseNode->LclScaling.Set(Converter.ConvertToFbxScale(Owner->GetRootComponent()->RelativeScale3D));
 	}
 	// Create the Skeleton
 	TArray<FbxNode*> BoneNodes;
-	FbxNode* SkeletonRootNode = CreateSkeleton(*SkeletalMeshComponent->SkeletalMesh, BoneNodes);
+	FbxNode* SkeletonRootNode = CreateSkeleton(SkeletalMeshComponent->SkeletalMesh, BoneNodes);
+	FbxSkeletonRoots.Add(SkeletalMeshComponent, SkeletonRootNode);
 	BaseNode->AddChild(SkeletonRootNode);
 
+	ExportAnimTrack(MatineeActor, SkeletalMeshComponent);
+}
+
+void FFbxExporter::ExportAnimTrack(class AMatineeActor* MatineeActor, USkeletalMeshComponent* SkeletalMeshComponent)
+{
+	static const float SamplingRate = 1.f / DEFAULT_SAMPLERATE;
+
+	float MatineeLength = MatineeActor->MatineeData->InterpLength;
 	// show a status update every 1 second worth of samples
 	const float UpdateFrequency = 1.0f;
 	float NextUpdateTime = UpdateFrequency;
 
+	// find root and find the bone array
+	TArray<FbxNode*> BoneNodes;
+
+	if ( FindSkeleton(SkeletalMeshComponent, BoneNodes)==false )
+	{
+		// error
+		return;
+	}
+
 	float SampleTime;
 	for(SampleTime = 0.f; SampleTime <= MatineeLength; SampleTime += SamplingRate)
 	{
diff --git a/Engine/Source/Editor/UnrealEd/Private/Fbx/FbxMainExport.cpp b/Engine/Source/Editor/UnrealEd/Private/Fbx/FbxMainExport.cpp
--- a/Engine/Source/Editor/UnrealEd/Private/Fbx/FbxMainExport.cpp
+++ b/Engine/Source/Editor/UnrealEd/Private/Fbx/FbxMainExport.cpp
@@ -46,6 +46,8 @@
 #include "Matinee/InterpTrackFloatProp.h"
 #include "Matinee/InterpTrackInstFloatProp.h"
 #include "Matinee/InterpTrackInstMove.h"
+#include "Matinee/InterpTrackAnimControl.h"
+#include "Matinee/InterpTrackInstAnimControl.h"
 
 #include "StaticMeshResources.h"
 #include "LandscapeDataAccess.h"
@@ -217,6 +219,7 @@ void FFbxExporter::WriteToFile(const TCHAR* Filename)
 void FFbxExporter::CloseDocument()
 {
 	FbxActors.Reset();
+	FbxSkeletonRoots.Reset();
 	FbxMaterials.Reset();
 	FbxNodeNameToIndexMap.Reset();
 	
@@ -305,6 +308,10 @@ void FFbxExporter::ExportLevelMesh( ULevel* InLevel, AMatineeActor* InMatineeAct
 			{
 				ExportActor( Actor, InMatineeActor ); // Just export the placement of the particle emitter.
 			}
+			else if(Actor->IsA(ACameraActor::StaticClass()))
+			{
+				ExportCamera(CastChecked<ACameraActor>(Actor), InMatineeActor, true); // Just export the placement of the particle emitter.
+			}
 			else if( Actor != NULL )
 			{
 				// Export blueprint actors and all their components
@@ -381,12 +388,12 @@ void FFbxExporter::ExportLight( ALight* Actor, AMatineeActor* InMatineeActor )
 	FbxActor->SetNodeAttribute(Light);
 }
 
-void FFbxExporter::ExportCamera( ACameraActor* Actor, AMatineeActor* InMatineeActor )
+void FFbxExporter::ExportCamera( ACameraActor* Actor, AMatineeActor* InMatineeActor, bool bExportComponents )
 {
 	if (Scene == NULL || Actor == NULL) return;
 
 	// Export the basic actor information.
-	FbxNode* FbxActor = ExportActor( Actor, InMatineeActor ); // this is the pivot node
+	FbxNode* FbxActor = ExportActor( Actor, InMatineeActor, bExportComponents ); // this is the pivot node
 	// The real fbx camera node
 	FbxNode* FbxCameraNode = FbxActor->GetParent();
 
@@ -824,7 +831,7 @@ void FFbxExporter::ExportSkeletalMesh( USkeletalMesh* SkeletalMesh )
 	FbxNode* MeshNode = FbxNode::Create(Scene, TCHAR_TO_ANSI(*MeshName));
 	Scene->GetRootNode()->AddChild(MeshNode);
 
-	ExportSkeletalMeshToFbx(*SkeletalMesh, NULL, *MeshName, MeshNode);
+	ExportSkeletalMeshToFbx(SkeletalMesh, NULL, *MeshName, MeshNode);
 }
 
 void FFbxExporter::ExportSkeletalMesh( AActor* Actor, USkeletalMeshComponent* SkeletalMeshComponent )
@@ -836,8 +843,7 @@ void FFbxExporter::ExportSkeletalMesh( AActor* Actor, USkeletalMeshComponent* Sk
 
 	FString FbxNodeName = GetActorNodeName(Actor, NULL);
 
-	FbxNode* FbxActorNode = ExportActor( Actor, NULL );
-	ExportSkeletalMeshToFbx(*SkeletalMesh, NULL, *FbxNodeName, FbxActorNode);
+	ExportActor( Actor, NULL, true );
 }
 
 FbxSurfaceMaterial* FFbxExporter::CreateDefaultMaterial()
@@ -988,13 +994,9 @@ bool FFbxExporter::ExportMatinee(AMatineeActor* InMatineeActor)
 		AActor* Actor = Group->GetGroupActor();
 		if (Group->Group == NULL || Actor == NULL) continue;
 
-		// Look for the class-type of the actor.
-		if (Actor->IsA(ACameraActor::StaticClass()))
-		{
-			ExportCamera( (ACameraActor*) Actor, InMatineeActor );
-		}
-
-		FbxNode* FbxActor = ExportActor( Actor, InMatineeActor );
+		FbxNode* FbxActor = FindActor(Actor); 
+		// now it should export everybody
+		check (FbxActor);
 
 		// Look for the tracks that we currently support
 		int32 TrackCount = FMath::Min(Group->TrackInst.Num(), Group->Group->InterpTracks.Num());
@@ -1014,6 +1016,14 @@ bool FFbxExporter::ExportMatinee(AMatineeActor* InMatineeActor)
 				UInterpTrackFloatProp* PropertyTrack = (UInterpTrackFloatProp*) Track;
 				ExportMatineeTrackFloatProp(FbxActor, PropertyTrack);
 			}
+			else if (TrackInst->IsA(UInterpTrackInstAnimControl::StaticClass()) && Track->IsA(UInterpTrackAnimControl::StaticClass()))
+			{
+				USkeletalMeshComponent* SkeletalMeshComp = Cast<USkeletalMeshComponent>(Actor->GetComponentByClass(USkeletalMeshComponent::StaticClass()));
+				if ( SkeletalMeshComp )
+				{
+					ExportAnimTrack(InMatineeActor, SkeletalMeshComp);
+				}
+			}
 		}
 	}
 
@@ -1167,8 +1177,7 @@ FbxNode* FFbxExporter::ExportActor(AActor* Actor, AMatineeActor* InMatineeActor,
 				}
 				else if (SkelMeshComp && SkelMeshComp->SkeletalMesh)
 				{
-					UAnimSequence* AnimSeq = (SkelMeshComp->GetAnimationMode() == EAnimationMode::AnimationSingleNode)? Cast<UAnimSequence>(SkelMeshComp->AnimationData.AnimToPlay) : NULL;
-					ExportSkeletalMeshToFbx( *SkelMeshComp->SkeletalMesh, AnimSeq, *SkelMeshComp->GetName(), ExportNode);
+					ExportSkeletalMeshComponent( SkelMeshComp, *SkelMeshComp->GetName(), ExportNode);
 				}
 				else if (ChildActorComp && ChildActorComp->ChildActor)
 				{
@@ -1688,6 +1697,20 @@ FbxNode* FFbxExporter::FindActor(AActor* Actor)
 	}
 }
 
+bool FFbxExporter::FindSkeleton(const USkeletalMeshComponent* SkelComp, TArray<FbxNode*>& BoneNodes)
+{
+	FbxNode** SkelRoot = FbxSkeletonRoots.Find(SkelComp);
+
+	if (SkelRoot)
+	{
+		BoneNodes.Empty();
+		GetSkeleton(*SkelRoot, BoneNodes);
+
+		return true;
+	}
+
+	return false;
+}
 /**
  * Determines the UVS to weld when exporting a Static Mesh
  * 
diff --git a/Engine/Source/Editor/UnrealEd/Private/Fbx/FbxSkeletalMeshExport.cpp b/Engine/Source/Editor/UnrealEd/Private/Fbx/FbxSkeletalMeshExport.cpp
--- a/Engine/Source/Editor/UnrealEd/Private/Fbx/FbxSkeletalMeshExport.cpp
+++ b/Engine/Source/Editor/UnrealEd/Private/Fbx/FbxSkeletalMeshExport.cpp
@@ -45,9 +45,9 @@ namespace UnFbx
  * Adds FBX skeleton nodes to the FbxScene based on the skeleton in the given USkeletalMesh, and fills
  * the given array with the nodes created
  */
-FbxNode* FFbxExporter::CreateSkeleton(const USkeletalMesh& SkelMesh, TArray<FbxNode*>& BoneNodes)
+FbxNode* FFbxExporter::CreateSkeleton(const USkeletalMesh* SkelMesh, TArray<FbxNode*>& BoneNodes)
 {
-	const FReferenceSkeleton& RefSkeleton= SkelMesh.RefSkeleton;
+	const FReferenceSkeleton& RefSkeleton= SkelMesh->RefSkeleton;
 
 	if(RefSkeleton.GetNum() == 0)
 	{
@@ -107,13 +107,25 @@ FbxNode* FFbxExporter::CreateSkeleton(const USkeletalMesh& SkelMesh, TArray<FbxN
 	return BoneNodes[0];
 }
 
+void FFbxExporter::GetSkeleton(FbxNode* RootNode, TArray<FbxNode*>& BoneNodes)
+{
+	if (RootNode->GetSkeleton())
+	{
+		BoneNodes.Add(RootNode);
+	}
+
+	for (int32 ChildIndex=0; ChildIndex<RootNode->GetChildCount(); ++ChildIndex)
+	{
+		GetSkeleton(RootNode->GetChild(ChildIndex), BoneNodes);
+	}
+}
 
 /**
  * Adds an Fbx Mesh to the FBX scene based on the data in the given FStaticLODModel
  */
-FbxNode* FFbxExporter::CreateMesh(const USkeletalMesh& SkelMesh, const TCHAR* MeshName)
+FbxNode* FFbxExporter::CreateMesh(const USkeletalMesh* SkelMesh, const TCHAR* MeshName)
 {
-	const FSkeletalMeshResource* SkelMeshResource = SkelMesh.GetImportedResource();
+	const FSkeletalMeshResource* SkelMeshResource = SkelMesh->GetImportedResource();
 	const FStaticLODModel& SourceModel = SkelMeshResource->LODModels[0];
 	const int32 VertexCount = SourceModel.NumVertices;
 
@@ -251,11 +263,11 @@ FbxNode* FFbxExporter::CreateMesh(const USkeletalMesh& SkelMesh, const TCHAR* Me
 
 
 	// Add the materials for the mesh
-	int32 MaterialCount = SkelMesh.Materials.Num();
+	int32 MaterialCount = SkelMesh->Materials.Num();
 
 	for(int32 MaterialIndex = 0; MaterialIndex < MaterialCount; ++MaterialIndex)
 	{
-		UMaterialInterface* MatInterface = SkelMesh.Materials[MaterialIndex].MaterialInterface;
+		UMaterialInterface* MatInterface = SkelMesh->Materials[MaterialIndex].MaterialInterface;
 
 		FbxSurfaceMaterial* FbxMaterial = NULL;
 		if(MatInterface && !FbxMaterials.Find(MatInterface->GetMaterial()))
@@ -288,9 +300,9 @@ FbxNode* FFbxExporter::CreateMesh(const USkeletalMesh& SkelMesh, const TCHAR* Me
 /**
  * Adds Fbx Clusters necessary to skin a skeletal mesh to the bones in the BoneNodes list
  */
-void FFbxExporter::BindMeshToSkeleton(const USkeletalMesh& SkelMesh, FbxNode* MeshRootNode, TArray<FbxNode*>& BoneNodes)
+void FFbxExporter::BindMeshToSkeleton(const USkeletalMesh* SkelMesh, FbxNode* MeshRootNode, TArray<FbxNode*>& BoneNodes)
 {
-	const FSkeletalMeshResource* SkelMeshResource = SkelMesh.GetImportedResource();
+	const FSkeletalMeshResource* SkelMeshResource = SkelMesh->GetImportedResource();
 	const FStaticLODModel& SourceModel = SkelMeshResource->LODModels[0];
 	const int32 VertexCount = SourceModel.NumVertices;
 
@@ -474,15 +486,28 @@ void FFbxExporter::CreateBindPose(FbxNode* MeshRootNode)
 	}
 }
 
+void FFbxExporter::ExportSkeletalMeshComponent(USkeletalMeshComponent* SkelMeshComp, const TCHAR* MeshName, FbxNode* ActorRootNode)
+{
+	if (SkelMeshComp && SkelMeshComp->SkeletalMesh)
+	{
+		UAnimSequence* AnimSeq = (SkelMeshComp->GetAnimationMode() == EAnimationMode::AnimationSingleNode)? Cast<UAnimSequence>(SkelMeshComp->AnimationData.AnimToPlay) : NULL;
+		FbxNode* SkeletonRootNode = ExportSkeletalMeshToFbx(SkelMeshComp->SkeletalMesh, AnimSeq, MeshName, ActorRootNode);
+		if(SkeletonRootNode)
+		{
+			FbxSkeletonRoots.Add(SkelMeshComp, SkeletonRootNode);
+		}
+	}
+}
 
 /**
  * Add the given skeletal mesh to the Fbx scene in preparation for exporting.  Makes all new nodes a child of the given node
  */
-void FFbxExporter::ExportSkeletalMeshToFbx(const USkeletalMesh& SkelMesh, const UAnimSequence* AnimSeq, const TCHAR* MeshName, FbxNode* ActorRootNode)
+FbxNode* FFbxExporter::ExportSkeletalMeshToFbx(const USkeletalMesh* SkelMesh, const UAnimSequence* AnimSeq, const TCHAR* MeshName, FbxNode* ActorRootNode)
 {
 	if(AnimSeq)
 	{
-		ExportAnimSequence(AnimSeq, &SkelMesh, true, MeshName, ActorRootNode);
+		return ExportAnimSequence(AnimSeq, SkelMesh, true, MeshName, ActorRootNode);
+
 	}
 	else
 	{
@@ -510,7 +535,11 @@ void FFbxExporter::ExportSkeletalMeshToFbx(const USkeletalMesh& SkelMesh, const
 			// Add the bind pose
 			CreateBindPose(MeshRootNode);
 		}
+
+		return SkeletonRootNode;
 	}
+
+	return NULL;
 }
 
 } // namespace UnFbx
diff --git a/Engine/Source/Editor/UnrealEd/Private/FbxExporter.h b/Engine/Source/Editor/UnrealEd/Private/FbxExporter.h
--- a/Engine/Source/Editor/UnrealEd/Private/FbxExporter.h
+++ b/Engine/Source/Editor/UnrealEd/Private/FbxExporter.h
@@ -80,7 +80,7 @@ public:
 	/**
 	 * Exports the camera-specific information for a camera actor.
 	 */
-	virtual void ExportCamera( ACameraActor* Actor, AMatineeActor* InMatineeActor );
+	virtual void ExportCamera( ACameraActor* Actor, AMatineeActor* InMatineeActor, bool bExportComponents );
 
 	/**
 	 * Exports the mesh and the actor information for a brush actor.
@@ -152,13 +152,18 @@ public:
 	/**
 	 * Exports a single UAnimSequence, and optionally a skeletal mesh
 	 */
-	void ExportAnimSequence( const UAnimSequence* AnimSeq, const USkeletalMesh* SkelMesh, bool bExportSkelMesh, const TCHAR* MeshNames=NULL, FbxNode* ActorRootNode=NULL);
+	FbxNode* ExportAnimSequence( const UAnimSequence* AnimSeq, const USkeletalMesh* SkelMesh, bool bExportSkelMesh, const TCHAR* MeshNames=NULL, FbxNode* ActorRootNode=NULL);
 
 	/**
 	 * Exports the list of UAnimSequences as a single animation based on the settings in the TrackKeys
 	 */
 	void ExportAnimSequencesAsSingle( USkeletalMesh* SkelMesh, const ASkeletalMeshActor* SkelMeshActor, const FString& ExportName, const TArray<UAnimSequence*>& AnimSeqList, const TArray<FAnimControlTrackKey>& TrackKeys );
 
+	/**
+	 * Export Anim Track of the given SkeletalMeshComponent
+	 */
+	void ExportAnimTrack(class AMatineeActor* MatineeActor, USkeletalMeshComponent* SkeletalMeshComponent);
+
 private:
 	FFbxExporter();
 
@@ -174,6 +179,7 @@ private:
 	
 	TMap<FString,int32> FbxNodeNameToIndexMap;
 	TMap<AActor*, FbxNode*> FbxActors;
+	TMap<USkeletalMeshComponent*, FbxNode*> FbxSkeletonRoots;
 	TMap<UMaterial*, FbxSurfaceMaterial*> FbxMaterials;
 	
 	/** The frames-per-second (FPS) used when baking transforms */
@@ -223,17 +229,17 @@ private:
 	 * Adds FBX skeleton nodes to the FbxScene based on the skeleton in the given USkeletalMesh, and fills
 	 * the given array with the nodes created
 	 */
-	FbxNode* CreateSkeleton(const USkeletalMesh& SkelMesh, TArray<FbxNode*>& BoneNodes);
+	FbxNode* CreateSkeleton(const USkeletalMesh* SkelMesh, TArray<FbxNode*>& BoneNodes);
 
 	/**
 	 * Adds an Fbx Mesh to the FBX scene based on the data in the given FStaticLODModel
 	 */
-	FbxNode* CreateMesh(const USkeletalMesh& SkelMesh, const TCHAR* MeshName);
+	FbxNode* CreateMesh(const USkeletalMesh* SkelMesh, const TCHAR* MeshName);
 
 	/**
 	 * Adds Fbx Clusters necessary to skin a skeletal mesh to the bones in the BoneNodes list
 	 */
-	void BindMeshToSkeleton(const USkeletalMesh& SkelMesh, FbxNode* MeshRootNode, TArray<FbxNode*>& BoneNodes);
+	void BindMeshToSkeleton(const USkeletalMesh* SkelMesh, FbxNode* MeshRootNode, TArray<FbxNode*>& BoneNodes);
 
 	/**
 	 * Add a bind pose to the scene based on the FbxMesh and skinning settings of the given node
@@ -243,7 +249,10 @@ private:
 	/**
 	 * Add the given skeletal mesh to the Fbx scene in preparation for exporting.  Makes all new nodes a child of the given node
 	 */
-	void ExportSkeletalMeshToFbx(const USkeletalMesh& SkelMesh, const UAnimSequence* AnimSeq, const TCHAR* MeshName, FbxNode* FbxActor);
+	FbxNode* ExportSkeletalMeshToFbx(const USkeletalMesh* SkelMesh, const UAnimSequence* AnimSeq, const TCHAR* MeshName, FbxNode* ActorRootNode);
+
+	/** Export SkeletalMeshComponent */
+	void ExportSkeletalMeshComponent(USkeletalMeshComponent* SkelMeshComp, const TCHAR* MeshName, FbxNode* ActorRootNode);
 
 	/**
 	 * Add the given animation sequence as rotation and translation tracks to the given list of bone nodes
@@ -285,7 +294,15 @@ private:
 	 * @return FbxNode* the FBX node created from the UE4 actor
 	 */
 	FbxNode* FindActor(AActor* Actor);
-	
+
+	/**
+	 * Find bone array of FbxNOdes of the given skeletalmeshcomponent  
+	 */
+	bool FindSkeleton(const USkeletalMeshComponent* SkelComp, TArray<FbxNode*>& BoneNodes);
+
+	/** recursively get skeleton */
+	void GetSkeleton(FbxNode* RootNode, TArray<FbxNode*>& BoneNodes);
+
 	/**
 	 * Exports the profile_COMMON information for a material.
 	 */
diff --git a/Engine/Source/Editor/UnrealEd/Public/MatineeExporter.h b/Engine/Source/Editor/UnrealEd/Public/MatineeExporter.h
--- a/Engine/Source/Editor/UnrealEd/Public/MatineeExporter.h
+++ b/Engine/Source/Editor/UnrealEd/Public/MatineeExporter.h
@@ -52,7 +52,7 @@ public:
 	/**
 	 * Exports the camera-specific information for a camera actor.
 	 */
-	virtual void ExportCamera( ACameraActor* Actor, AMatineeActor* InMatineeActor ) = 0;
+	virtual void ExportCamera( ACameraActor* Actor, AMatineeActor* InMatineeActor, bool bExportComponents ) = 0;
 
 	/**
 	 * Exports the mesh and the actor information for a brush actor.
diff --git a/Engine/Source/Runtime/Engine/Classes/Matinee/InterpTrackInstAnimControl.h b/Engine/Source/Runtime/Engine/Classes/Matinee/InterpTrackInstAnimControl.h
--- a/Engine/Source/Runtime/Engine/Classes/Matinee/InterpTrackInstAnimControl.h
+++ b/Engine/Source/Runtime/Engine/Classes/Matinee/InterpTrackInstAnimControl.h
@@ -3,7 +3,7 @@
 #pragma once
 #include "InterpTrackInstAnimControl.generated.h"
 
-UCLASS()
+UCLASS(MinimalAPI)
 class UInterpTrackInstAnimControl : public UInterpTrackInst
 {
 	GENERATED_UCLASS_BODY()
