[INTEGRATE] Change 2429283 by [REDACTED_PERSON][REDACTED_HANDLE] on 2015/02/02 23:50:17

	Mac: Fix deprecated operator new warning for TTFontImport.cpp

diff --git a/Engine/Source/Editor/UnrealEd/Private/Factories/TTFontImport.cpp b/Engine/Source/Editor/UnrealEd/Private/Factories/TTFontImport.cpp
--- a/Engine/Source/Editor/UnrealEd/Private/Factories/TTFontImport.cpp
+++ b/Engine/Source/Editor/UnrealEd/Private/Factories/TTFontImport.cpp
@@ -1674,7 +1674,7 @@ UTexture2D* UTrueTypeFontFactory::CreateTextureFromBitmap( UFont* Font, uint8* B
 	}
 
 	// Create texture for page.
-	UTexture2D* Texture = new(Font, *TextureString)UTexture2D(FObjectInitializer());
+	UTexture2D* Texture = NewNamedObject<UTexture2D>(Font, *TextureString);
 
 	// note RF_Public because font textures can be referenced directly by material expressions
 	Texture->SetFlags(RF_Public);
