Rename StaticConstructObject to StaticConstructObject_Internal and use it only in CoreUObject module. Change rest of usages to NewObject.

diff --git a/Engine/Plugins/Experimental/LiveEditor/Source/LiveEditor/Private/LiveEditorManager.cpp b/Engine/Plugins/Experimental/LiveEditor/Source/LiveEditor/Private/LiveEditorManager.cpp
--- a/Engine/Plugins/Experimental/LiveEditor/Source/LiveEditor/Private/LiveEditorManager.cpp
+++ b/Engine/Plugins/Experimental/LiveEditor/Source/LiveEditor/Private/LiveEditorManager.cpp
@@ -524,7 +524,7 @@ bool FLiveEditorManager::Activate( const FString &Name )
 		FActiveBlueprintRecord Record;
 		Record.Name = Name;
 
-		ULiveEditorBlueprint *Instance = (ULiveEditorBlueprint*)StaticConstructObject( Blueprint->GeneratedClass, GetTransientPackage(), NAME_None, RF_Transient|RF_Public|RF_RootSet|RF_Standalone );
+		auto Instance = NewObject<ULiveEditorBlueprint>(GetTransientPackage(), Blueprint->GeneratedClass, NAME_None, RF_Transient | RF_Public | RF_RootSet | RF_Standalone);
 		Instance->DoInit();
 		Record.Blueprint = Instance;
 
diff --git a/Engine/Plugins/Media/MediaPlayerEditor/Source/MediaPlayerEditor/Private/Factories/MediaPlayerFactory.cpp b/Engine/Plugins/Media/MediaPlayerEditor/Source/MediaPlayerEditor/Private/Factories/MediaPlayerFactory.cpp
--- a/Engine/Plugins/Media/MediaPlayerEditor/Source/MediaPlayerEditor/Private/Factories/MediaPlayerFactory.cpp
+++ b/Engine/Plugins/Media/MediaPlayerEditor/Source/MediaPlayerEditor/Private/Factories/MediaPlayerFactory.cpp
@@ -34,7 +34,7 @@ UMediaPlayerFactory::UMediaPlayerFactory( const FObjectInitializer& ObjectInitia
 
 UObject* UMediaPlayerFactory::FactoryCreateBinary( UClass* Class, UObject* InParent, FName Name, EObjectFlags Flags, UObject* Context, const TCHAR* Type, const uint8*& Buffer, const uint8* BufferEnd, FFeedbackContext* Warn )
 {
-	UMediaPlayer* MediaPlayer = CastChecked<UMediaPlayer>(StaticConstructObject(Class, InParent, Name, Flags));
+	UMediaPlayer* MediaPlayer = NewObject<UMediaPlayer>(InParent, Class, Name, Flags);
 	MediaPlayer->OpenUrl(CurrentFilename);
 
 	return MediaPlayer;
diff --git a/Engine/Plugins/Media/MediaPlayerEditor/Source/MediaPlayerEditor/Private/Factories/MediaPlayerFactoryNew.cpp b/Engine/Plugins/Media/MediaPlayerEditor/Source/MediaPlayerEditor/Private/Factories/MediaPlayerFactoryNew.cpp
--- a/Engine/Plugins/Media/MediaPlayerEditor/Source/MediaPlayerEditor/Private/Factories/MediaPlayerFactoryNew.cpp
+++ b/Engine/Plugins/Media/MediaPlayerEditor/Source/MediaPlayerEditor/Private/Factories/MediaPlayerFactoryNew.cpp
@@ -20,7 +20,7 @@ UMediaPlayerFactoryNew::UMediaPlayerFactoryNew( const FObjectInitializer& Object
 
 UObject* UMediaPlayerFactoryNew::FactoryCreateNew( UClass* InClass, UObject* InParent, FName InName, EObjectFlags Flags, UObject* Context, FFeedbackContext* Warn )
 {
-	return CastChecked<UMediaPlayer>(StaticConstructObject(InClass, InParent, InName, Flags));
+	return NewObject<UMediaPlayer>(InParent, InClass, InName, Flags);
 }
 
 
diff --git a/Engine/Plugins/ScriptPlugin/Source/ScriptEditorPlugin/Private/ScriptBlueprintCompiler.cpp b/Engine/Plugins/ScriptPlugin/Source/ScriptEditorPlugin/Private/ScriptBlueprintCompiler.cpp
--- a/Engine/Plugins/ScriptPlugin/Source/ScriptEditorPlugin/Private/ScriptBlueprintCompiler.cpp
+++ b/Engine/Plugins/ScriptPlugin/Source/ScriptEditorPlugin/Private/ScriptBlueprintCompiler.cpp
@@ -184,7 +184,7 @@ void FScriptBlueprintCompiler::FinishCompilingClass(UClass* Class)
 	if (ContextProperty)
 	{
 		UObject* CDO = Class->GetDefaultObject();
-		UObject* ContextDefaultSubobject = StaticConstructObject(ContextProperty->PropertyClass, CDO, "ScriptContext", RF_DefaultSubObject|RF_Public);
+		UObject* ContextDefaultSubobject = NewObject<UObject>(CDO, ContextProperty->PropertyClass, "ScriptContext", RF_DefaultSubObject | RF_Public);
 		ContextProperty->SetObjectPropertyValue(ContextProperty->ContainerPtrToValuePtr<UObject*>(CDO), ContextDefaultSubobject);
 	}
 }
diff --git a/Engine/Source/Developer/AssetTools/Private/AssetTools.cpp b/Engine/Source/Developer/AssetTools/Private/AssetTools.cpp
--- a/Engine/Source/Developer/AssetTools/Private/AssetTools.cpp
+++ b/Engine/Source/Developer/AssetTools/Private/AssetTools.cpp
@@ -288,7 +288,7 @@ UObject* FAssetTools::CreateAsset(const FString& AssetName, const FString& Packa
 	}
 	else if ( AssetClass )
 	{
-		NewObj = StaticConstructObject(ClassToUse, Pkg, FName(*AssetName), Flags);
+		NewObj = NewObject<UObject>(Pkg, ClassToUse, FName(*AssetName), Flags);
 	}
 
 	if( NewObj )
diff --git a/Engine/Source/Developer/MovieSceneTools/Private/MovieSceneFactory.cpp b/Engine/Source/Developer/MovieSceneTools/Private/MovieSceneFactory.cpp
--- a/Engine/Source/Developer/MovieSceneTools/Private/MovieSceneFactory.cpp
+++ b/Engine/Source/Developer/MovieSceneTools/Private/MovieSceneFactory.cpp
@@ -24,7 +24,7 @@ UMovieSceneFactory::UMovieSceneFactory( const FObjectInitializer& ObjectInitiali
 
 UObject* UMovieSceneFactory::FactoryCreateNew(UClass* Class, UObject* InParent, FName Name, EObjectFlags Flags, UObject* Context, FFeedbackContext* Warn)
 {
-	return StaticConstructObject( UMovieScene::StaticClass(), InParent, Name, Flags );
+	return NewObject<UMovieScene>(InParent, Name, Flags);
 }
 
 #undef LOCTEXT_NAMESPACE
\ No newline at end of file
diff --git a/Engine/Source/Editor/LocalizationDashboard/Private/LocalizationDashboard.cpp b/Engine/Source/Editor/LocalizationDashboard/Private/LocalizationDashboard.cpp
--- a/Engine/Source/Editor/LocalizationDashboard/Private/LocalizationDashboard.cpp
+++ b/Engine/Source/Editor/LocalizationDashboard/Private/LocalizationDashboard.cpp
@@ -288,7 +288,7 @@ namespace
 			const TSharedPtr<IPropertyHandle> TargetObjectPropertyHandle = TargetObjectsArrayPropertyHandle->GetElement(NewEntryIndexToBeInitialized);
 			if (TargetObjectPropertyHandle.IsValid() && TargetObjectPropertyHandle->IsValidHandle())
 			{
-				ULocalizationTarget* const NewTarget = Cast<ULocalizationTarget>(StaticConstructObject(ULocalizationTarget::StaticClass(), Settings));
+				ULocalizationTarget* const NewTarget = NewObject<ULocalizationTarget>(Settings);
 
 				TArray<void*> RawData;
 				TargetObjectsPropertyHandle->AccessRawData(RawData);
diff --git a/Engine/Source/Editor/LocalizationDashboard/Private/ProjectLocalizationSettings.cpp b/Engine/Source/Editor/LocalizationDashboard/Private/ProjectLocalizationSettings.cpp
--- a/Engine/Source/Editor/LocalizationDashboard/Private/ProjectLocalizationSettings.cpp
+++ b/Engine/Source/Editor/LocalizationDashboard/Private/ProjectLocalizationSettings.cpp
@@ -16,7 +16,7 @@ void UProjectLocalizationSettings::PostInitProperties()
 	TargetObjects.Empty(Targets.Num());
 	for (const FLocalizationTargetSettings& Target : Targets)
 	{
-		ULocalizationTarget* const NewTargetObject = Cast<ULocalizationTarget>(StaticConstructObject(ULocalizationTarget::StaticClass(), this));
+		ULocalizationTarget* const NewTargetObject = NewObject<ULocalizationTarget>(this);
 		NewTargetObject->Settings = Target;
 		TargetObjects.Add(NewTargetObject);
 	}
diff --git a/Engine/Source/Editor/MaterialEditor/Private/MaterialEditor.cpp b/Engine/Source/Editor/MaterialEditor/Private/MaterialEditor.cpp
--- a/Engine/Source/Editor/MaterialEditor/Private/MaterialEditor.cpp
+++ b/Engine/Source/Editor/MaterialEditor/Private/MaterialEditor.cpp
@@ -267,7 +267,7 @@ void FMaterialEditor::InitEditorForMaterialFunction(UMaterialFunction* InMateria
 	ExpressionPreviewMaterial = NULL;
 
 	// Create a temporary material to preview the material function
-	Material = (UMaterial*)StaticConstructObject(UMaterial::StaticClass()); 
+	Material = NewObject<UMaterial>(); 
 	{
 		FArchiveUObject DummyArchive;
 		// Hack: serialize the new material with an archive that does nothing so that its material resources are created
@@ -2739,7 +2739,7 @@ void FMaterialEditor::SetPreviewExpression(UMaterialExpression* NewPreviewExpres
 		if( ExpressionPreviewMaterial == NULL )
 		{
 			// Create the expression preview material if it hasnt already been created
-			ExpressionPreviewMaterial = (UMaterial*)StaticConstructObject( UMaterial::StaticClass(), GetTransientPackage(), NAME_None, RF_Public);
+			ExpressionPreviewMaterial = NewObject<UMaterial>(GetTransientPackage(), NAME_None, RF_Public);
 			ExpressionPreviewMaterial->bIsPreviewMaterial = true;
 		}
 
diff --git a/Engine/Source/Editor/Matinee/Private/MatineeMenus.cpp b/Engine/Source/Editor/Matinee/Private/MatineeMenus.cpp
--- a/Engine/Source/Editor/Matinee/Private/MatineeMenus.cpp
+++ b/Engine/Source/Editor/Matinee/Private/MatineeMenus.cpp
@@ -2834,7 +2834,7 @@ void FMatinee::ExportCameraAnimationNameCommitted(const FText& InAnimationPackag
 					else
 					{
 						// create it, then copy params into it
-						ExistingObject = StaticConstructObject(UCameraAnim::StaticClass(), ExistingPackage, *ObjectName, RF_Public|RF_Standalone);
+						ExistingObject = NewObject<UCameraAnim>(ExistingPackage, *ObjectName, RF_Public | RF_Standalone);
 						bNewObject = true;
 					}
 				}
diff --git a/Engine/Source/Editor/MeshPaint/Private/MeshPaintEdMode.cpp b/Engine/Source/Editor/MeshPaint/Private/MeshPaintEdMode.cpp
--- a/Engine/Source/Editor/MeshPaint/Private/MeshPaintEdMode.cpp
+++ b/Engine/Source/Editor/MeshPaint/Private/MeshPaintEdMode.cpp
@@ -1876,7 +1876,7 @@ void FEdModeMeshPaint::StartPaintingTexture( UMeshComponent* InMeshComponent )
 					TextureData->PaintRenderTargetTexture->GetSurfaceHeight() != TextureHeight )
 				{
 					TextureData->PaintRenderTargetTexture = NULL;
-					TextureData->PaintRenderTargetTexture = CastChecked<UTextureRenderTarget2D>( StaticConstructObject( UTextureRenderTarget2D::StaticClass(), GetTransientPackage(), NAME_None, RF_Transient ) );
+					TextureData->PaintRenderTargetTexture = NewObject<UTextureRenderTarget2D>(GetTransientPackage(), NAME_None, RF_Transient);
 					TextureData->PaintRenderTargetTexture->bNeedsTwoCopies = true;
 					const bool bForceLinearGamma = true;
 					TextureData->PaintRenderTargetTexture->InitCustomFormat( TextureWidth, TextureHeight, PF_A16B16G16R16, bForceLinearGamma );
@@ -1897,7 +1897,7 @@ void FEdModeMeshPaint::StartPaintingTexture( UMeshComponent* InMeshComponent )
 					BrushRenderTargetTexture->GetSurfaceHeight() != BrushTargetTextureHeight )
 				{
 					BrushRenderTargetTexture = NULL;
-					BrushRenderTargetTexture = CastChecked<UTextureRenderTarget2D>( StaticConstructObject( UTextureRenderTarget2D::StaticClass(), GetTransientPackage(), NAME_None, RF_Transient ) );
+					BrushRenderTargetTexture = NewObject<UTextureRenderTarget2D>(GetTransientPackage(), NAME_None, RF_Transient);
 					const bool bForceLinearGamma = true;
 					BrushRenderTargetTexture->ClearColor = FLinearColor::Black;
 					BrushRenderTargetTexture->bNeedsTwoCopies = true;
@@ -1917,7 +1917,7 @@ void FEdModeMeshPaint::StartPaintingTexture( UMeshComponent* InMeshComponent )
 						BrushMaskRenderTargetTexture->GetSurfaceHeight() != BrushTargetTextureHeight )
 					{
 						BrushMaskRenderTargetTexture = NULL;
-						BrushMaskRenderTargetTexture = CastChecked<UTextureRenderTarget2D>( StaticConstructObject( UTextureRenderTarget2D::StaticClass(), GetTransientPackage(), NAME_None, RF_Transient ) );
+						BrushMaskRenderTargetTexture = NewObject<UTextureRenderTarget2D>(GetTransientPackage(), NAME_None, RF_Transient);
 						const bool bForceLinearGamma = true;
 						BrushMaskRenderTargetTexture->ClearColor = FLinearColor::Black;
 						BrushMaskRenderTargetTexture->bNeedsTwoCopies = true;
@@ -1933,7 +1933,7 @@ void FEdModeMeshPaint::StartPaintingTexture( UMeshComponent* InMeshComponent )
 						SeamMaskRenderTargetTexture->GetSurfaceHeight() != TextureHeight )
 					{
 						SeamMaskRenderTargetTexture = NULL;
-						SeamMaskRenderTargetTexture = CastChecked<UTextureRenderTarget2D>( StaticConstructObject( UTextureRenderTarget2D::StaticClass(), GetTransientPackage(), NAME_None, RF_Transient ) );
+						SeamMaskRenderTargetTexture = NewObject<UTextureRenderTarget2D>(GetTransientPackage(), NAME_None, RF_Transient);
 						const bool bForceLinearGamma = true;
 						SeamMaskRenderTargetTexture->ClearColor = FLinearColor::Black;
 						SeamMaskRenderTargetTexture->bNeedsTwoCopies = true;
diff --git a/Engine/Source/Editor/Persona/Private/EditorObjectsTracker.cpp b/Engine/Source/Editor/Persona/Private/EditorObjectsTracker.cpp
--- a/Engine/Source/Editor/Persona/Private/EditorObjectsTracker.cpp
+++ b/Engine/Source/Editor/Persona/Private/EditorObjectsTracker.cpp
@@ -22,7 +22,7 @@ UObject* FEditorObjectTracker::GetEditorObjectForClass( UClass* EdClass )
 	{
 		FString ObjName = MakeUniqueObjectName(GetTransientPackage(), EdClass).ToString();
 		ObjName += "_EdObj";
-		Obj = StaticConstructObject(EdClass, GetTransientPackage(), FName(*ObjName), RF_Public|RF_Standalone|RF_Transient);
+		Obj = NewObject<UObject>(GetTransientPackage(), EdClass, FName(*ObjName), RF_Public | RF_Standalone | RF_Transient);
 		EditorObjMap.Add(EdClass,Obj);
 	}
 	return Obj;
diff --git a/Engine/Source/Editor/Persona/Private/SAnimNotifyPanel.cpp b/Engine/Source/Editor/Persona/Private/SAnimNotifyPanel.cpp
--- a/Engine/Source/Editor/Persona/Private/SAnimNotifyPanel.cpp
+++ b/Engine/Source/Editor/Persona/Private/SAnimNotifyPanel.cpp
@@ -3401,7 +3401,7 @@ void SAnimNotifyPanel::OnTrackSelectionChanged()
 		for(int32 Idx : TrackIndices)
 		{
 			FString ObjName = MakeUniqueObjectName(GetTransientPackage(), UEditorNotifyObject::StaticClass()).ToString();
-			UEditorNotifyObject* NewNotifyObject = Cast<UEditorNotifyObject>(StaticConstructObject(UEditorNotifyObject::StaticClass(), GetTransientPackage(), FName(*ObjName), RF_Public|RF_Standalone|RF_Transient));
+			UEditorNotifyObject* NewNotifyObject = NewObject<UEditorNotifyObject>(GetTransientPackage(), FName(*ObjName), RF_Public | RF_Standalone | RF_Transient);
 			NewNotifyObject->InitFromAnim(Sequence, FOnAnimObjectChange::CreateSP(this, &SAnimNotifyPanel::OnNotifyObjectChanged));
 			NewNotifyObject->InitialiseNotify(NotifyAnimTracks.Num() - TrackIdx - 1, Idx);
 			SelectionSet.Add(NewNotifyObject);
diff --git a/Engine/Source/Editor/PropertyEditor/Private/UserInterface/PropertyEditor/SPropertyEditorEditInline.cpp b/Engine/Source/Editor/PropertyEditor/Private/UserInterface/PropertyEditor/SPropertyEditorEditInline.cpp
--- a/Engine/Source/Editor/PropertyEditor/Private/UserInterface/PropertyEditor/SPropertyEditorEditInline.cpp
+++ b/Engine/Source/Editor/PropertyEditor/Private/UserInterface/PropertyEditor/SPropertyEditorEditInline.cpp
@@ -203,9 +203,9 @@ void SPropertyEditorEditInline::OnClassPicked(UClass* InClass)
 				{
 					MaskedOuterFlags |= RF_ArchetypeObject;
 				}
-				UObject*		NewObject = StaticConstructObject(InClass, UseOuter, NAME_None, MaskedOuterFlags, NULL);
+				UObject*		NewUObject = NewObject<UObject>(UseOuter, InClass, NAME_None, MaskedOuterFlags, NULL);
 
-				NewValue = NewObject->GetPathName();
+				NewValue = NewUObject->GetPathName();
 			}
 			else
 			{
diff --git a/Engine/Source/Editor/UnrealEd/Private/Commandlets/DiffPackagesCommandlet.cpp b/Engine/Source/Editor/UnrealEd/Private/Commandlets/DiffPackagesCommandlet.cpp
--- a/Engine/Source/Editor/UnrealEd/Private/Commandlets/DiffPackagesCommandlet.cpp
+++ b/Engine/Source/Editor/UnrealEd/Private/Commandlets/DiffPackagesCommandlet.cpp
@@ -355,12 +355,12 @@ UObject* CopyObjectToPackage(UPackage* Package, UObject* Object)
 	FObjectWriter(Object, Bytes);
 
 	// make a new object
-	UObject* NewObject = StaticConstructObject(Object->GetClass(), NewOuter, Object->GetFName(), Object->GetFlags(), Object->GetArchetype(), true);
+	UObject* NewUObject = NewObject<UObject>(NewOuter, Object->GetClass(), Object->GetFName(), Object->GetFlags(), Object->GetArchetype(), true);
 
 	// serialize old objects on top of the new object
-	FObjectReader Reader(NewObject, Bytes);
+	FObjectReader Reader(NewUObject, Bytes);
 
-	return NewObject;
+	return NewUObject;
 }
 
 
diff --git a/Engine/Source/Editor/UnrealEd/Private/Editor.cpp b/Engine/Source/Editor/UnrealEd/Private/Editor.cpp
--- a/Engine/Source/Editor/UnrealEd/Private/Editor.cpp
+++ b/Engine/Source/Editor/UnrealEd/Private/Editor.cpp
@@ -2049,8 +2049,7 @@ void UEditorEngine::CloseEditedWorldAssets(UWorld* InWorld)
 
 UTextureRenderTarget2D* UEditorEngine::GetScratchRenderTarget( uint32 MinSize )
 {
-	UTextureRenderTargetFactoryNew* NewFactory = CastChecked<UTextureRenderTargetFactoryNew>( StaticConstructObject(UTextureRenderTargetFactoryNew::StaticClass()) );
-
+	auto NewFactory = NewObject<UTextureRenderTargetFactoryNew>();
 	UTextureRenderTarget2D* ScratchRenderTarget = NULL;
 
 	// We never allow render targets greater than 2048
diff --git a/Engine/Source/Editor/UnrealEd/Private/EditorLevelUtils.cpp b/Engine/Source/Editor/UnrealEd/Private/EditorLevelUtils.cpp
--- a/Engine/Source/Editor/UnrealEd/Private/EditorLevelUtils.cpp
+++ b/Engine/Source/Editor/UnrealEd/Private/EditorLevelUtils.cpp
@@ -200,7 +200,7 @@ namespace EditorLevelUtils
 
 			const FScopedBusyCursor BusyCursor;
 
-			ULevelStreaming* StreamingLevel = static_cast<ULevelStreaming*>( StaticConstructObject( LevelStreamingClass, InWorld, NAME_None, RF_NoFlags, NULL) );
+			ULevelStreaming* StreamingLevel = NewObject<ULevelStreaming>(InWorld, LevelStreamingClass, NAME_None, RF_NoFlags, NULL);
 
 			// Associate a package name.
 			StreamingLevel->SetWorldAssetByPackageName(LevelPackageName);
diff --git a/Engine/Source/Editor/UnrealEd/Private/EditorServer.cpp b/Engine/Source/Editor/UnrealEd/Private/EditorServer.cpp
--- a/Engine/Source/Editor/UnrealEd/Private/EditorServer.cpp
+++ b/Engine/Source/Editor/UnrealEd/Private/EditorServer.cpp
@@ -4951,7 +4951,7 @@ bool FixUpBadAnimNotifiers()
 				{
 					// fix animnotifiers
 					UE_LOG(LogEditorServer, Log, TEXT("Animation[%s] Notifier[%s:%d] is being fixed (Current Outer:%s)"), *AnimSeq->GetName(), *AnimSeq->Notifies[I].Notify->GetName(), I, *AnimSeq->Notifies[I].Notify->GetOuter()->GetName());
-					AnimSeq->Notifies[I].Notify = CastChecked<UAnimNotify>( StaticConstructObject(AnimSeq->Notifies[I].Notify->GetClass(), AnimSeq,NAME_None,RF_NoFlags,AnimSeq->Notifies[I].Notify) );
+					AnimSeq->Notifies[I].Notify = NewObject<UAnimNotify>(AnimSeq, AnimSeq->Notifies[I].Notify->GetClass(), NAME_None, RF_NoFlags, AnimSeq->Notifies[I].Notify);
 					UE_LOG(LogEditorServer, Log, TEXT("After fixed (Current Outer:%s)"), *AnimSeq->Notifies[I].Notify->GetOuter()->GetName());
 					AnimSeq->MarkPackageDirty();
 				}
diff --git a/Engine/Source/Editor/UnrealEd/Private/Factories/ApexDestructibleAssetImport.cpp b/Engine/Source/Editor/UnrealEd/Private/Factories/ApexDestructibleAssetImport.cpp
--- a/Engine/Source/Editor/UnrealEd/Private/Factories/ApexDestructibleAssetImport.cpp
+++ b/Engine/Source/Editor/UnrealEd/Private/Factories/ApexDestructibleAssetImport.cpp
@@ -991,7 +991,7 @@ UDestructibleMesh* ImportDestructibleMeshFromApexDestructibleAsset(UObject* InPa
 	if (DestructibleMesh == NULL)
 	{
 		// Create the new UDestructibleMesh object if the one with the same name does not exist
-		DestructibleMesh = CastChecked<UDestructibleMesh>(StaticConstructObject(UDestructibleMesh::StaticClass(), InParent, Name, Flags));
+		DestructibleMesh = NewObject<UDestructibleMesh>(InParent, Name, Flags);
 	}
 	
 	if (!(Options & EDestructibleImportOptions::PreserveSettings))
diff --git a/Engine/Source/Editor/UnrealEd/Private/Factories/CSVImportFactory.cpp b/Engine/Source/Editor/UnrealEd/Private/Factories/CSVImportFactory.cpp
--- a/Engine/Source/Editor/UnrealEd/Private/Factories/CSVImportFactory.cpp
+++ b/Engine/Source/Editor/UnrealEd/Private/Factories/CSVImportFactory.cpp
@@ -153,7 +153,7 @@ UObject* UCSVImportFactory::FactoryCreateText( UClass* InClass, UObject* InParen
 			}
 
 			// Create/reset table
-			UDataTable* NewTable = CastChecked<UDataTable>(StaticConstructObject(UDataTable::StaticClass(), InParent, InName, Flags));
+			UDataTable* NewTable = NewObject<UDataTable>(InParent, InName, Flags);
 			NewTable->RowStruct = ImportRowStruct;
 			NewTable->ImportPath = FReimportManager::SanitizeImportFilename(CurrentFilename, NewTable);
 			// Go ahead and create table from string
@@ -172,7 +172,7 @@ UObject* UCSVImportFactory::FactoryCreateText( UClass* InClass, UObject* InParen
 			}
 
 			// Create/reset table
-			UCurveTable* NewTable = CastChecked<UCurveTable>(StaticConstructObject(UCurveTable::StaticClass(), InParent, InName, Flags));
+			UCurveTable* NewTable = NewObject<UCurveTable>(InParent, InName, Flags);
 			NewTable->ImportPath = FReimportManager::SanitizeImportFilename(CurrentFilename, NewTable);
 
 			// Go ahead and create table from string
@@ -187,7 +187,7 @@ UObject* UCSVImportFactory::FactoryCreateText( UClass* InClass, UObject* InParen
 			UClass* CurveClass = GetCurveClass( ImportType );
 
 			// Create/reset curve
-			UCurveBase* NewCurve = CastChecked<UCurveBase>(StaticConstructObject(CurveClass, InParent, InName, Flags));
+			UCurveBase* NewCurve = NewObject<UCurveBase>(InParent, CurveClass, InName, Flags);
 
 			Problems = NewCurve->CreateCurveFromCSVString(String);
 
diff --git a/Engine/Source/Editor/UnrealEd/Private/Factories/EditorFactories.cpp b/Engine/Source/Editor/UnrealEd/Private/Factories/EditorFactories.cpp
--- a/Engine/Source/Editor/UnrealEd/Private/Factories/EditorFactories.cpp
+++ b/Engine/Source/Editor/UnrealEd/Private/Factories/EditorFactories.cpp
@@ -137,7 +137,7 @@ UObject* UTexture2DFactoryNew::FactoryCreateNew( UClass* InClass, UObject* InPar
 		return nullptr;
 	}
 
-	UTexture2D* Object = CastChecked<UTexture2D>(StaticConstructObject(InClass,InParent,InName,Flags) );
+	UTexture2D* Object = NewObject<UTexture2D>(InParent, InClass, InName, Flags);
 
 	Object->Source.Init2DWithMipChain(Width, Height, TSF_BGRA8);
 
@@ -170,7 +170,7 @@ UMaterialInstanceConstantFactoryNew::UMaterialInstanceConstantFactoryNew(const F
 
 UObject* UMaterialInstanceConstantFactoryNew::FactoryCreateNew(UClass* Class,UObject* InParent,FName Name,EObjectFlags Flags,UObject* Context,FFeedbackContext* Warn)
 {
-	UMaterialInstanceConstant* MIC = CastChecked<UMaterialInstanceConstant>(StaticConstructObject(Class,InParent,Name,Flags));
+	auto MIC = NewObject<UMaterialInstanceConstant>(InParent, Class, Name, Flags);
 	
 	if ( MIC )
 	{
@@ -242,7 +242,7 @@ UMaterialFunctionFactoryNew::UMaterialFunctionFactoryNew(const FObjectInitialize
 
 UObject* UMaterialFunctionFactoryNew::FactoryCreateNew(UClass* Class,UObject* InParent,FName Name,EObjectFlags Flags,UObject* Context,FFeedbackContext* Warn)
 {
-	return StaticConstructObject(Class,InParent,Name,Flags);
+	return NewObject<UObject>(InParent, Class, Name, Flags);
 }
 
 
@@ -260,7 +260,7 @@ UMaterialParameterCollectionFactoryNew::UMaterialParameterCollectionFactoryNew(c
 
 UObject* UMaterialParameterCollectionFactoryNew::FactoryCreateNew(UClass* Class,UObject* InParent,FName Name,EObjectFlags Flags,UObject* Context,FFeedbackContext* Warn)
 {
-	return StaticConstructObject(Class,InParent,Name,Flags);
+	return NewObject<UObject>(InParent, Class, Name, Flags);
 }
 
 /*------------------------------------------------------------------------------
@@ -1126,7 +1126,7 @@ UObject* UPackageFactory::FactoryCreateText( UClass* Class, UObject* InParent, F
 					}
 					check(ParentPkg);
 
-					UPackage* NewPackage = Cast<UPackage>(StaticConstructObject(UPackage::StaticClass(), ParentPkg, NewPackageName, RF_NoFlags, Archetype));
+					auto NewPackage = NewObject<UPackage>(ParentPkg, NewPackageName, RF_NoFlags, Archetype);
 					check(NewPackage);
 					NewPackage->SetFlags(RF_Standalone|RF_Public);
 					MapPackages.Add(NewPackageName.ToString(), NewPackage);
@@ -1854,7 +1854,7 @@ UDialogueVoiceFactory::UDialogueVoiceFactory(const FObjectInitializer& ObjectIni
 
 UObject* UDialogueVoiceFactory::FactoryCreateNew(UClass* Class,UObject* InParent,FName Name,EObjectFlags Flags,UObject* Context,FFeedbackContext* Warn)
 {
-	return StaticConstructObject(Class,InParent,Name,Flags);
+	return NewObject<UObject>(InParent, Class, Name, Flags);
 }
 
 /*-----------------------------------------------------------------------------
@@ -1871,7 +1871,7 @@ UDialogueWaveFactory::UDialogueWaveFactory(const FObjectInitializer& ObjectIniti
 
 UObject* UDialogueWaveFactory::FactoryCreateNew(UClass* Class,UObject* InParent,FName Name,EObjectFlags Flags,UObject* Context,FFeedbackContext* Warn)
 {
-	return StaticConstructObject(Class,InParent,Name,Flags);
+	return NewObject<UObject>(InParent, Class, Name, Flags);
 }
 
 /*-----------------------------------------------------------------------------
@@ -2484,7 +2484,7 @@ UParticleSystemFactoryNew::UParticleSystemFactoryNew(const FObjectInitializer& O
 
 UObject* UParticleSystemFactoryNew::FactoryCreateNew(UClass* Class,UObject* InParent,FName Name,EObjectFlags Flags,UObject* Context,FFeedbackContext* Warn)
 {
-	return StaticConstructObject(Class,InParent,Name,Flags);
+	return NewObject<UObject>(InParent, Class, Name, Flags);
 }
 
 /*------------------------------------------------------------------------------
@@ -2501,7 +2501,7 @@ UPhysicalMaterialFactoryNew::UPhysicalMaterialFactoryNew(const FObjectInitialize
 
 UObject* UPhysicalMaterialFactoryNew::FactoryCreateNew(UClass* Class,UObject* InParent,FName Name,EObjectFlags Flags,UObject* Context,FFeedbackContext* Warn)
 {
-	return StaticConstructObject(Class,InParent,Name,Flags);
+	return NewObject<UObject>(InParent, Class, Name, Flags);
 }
 
 /*------------------------------------------------------------------------------
@@ -2518,7 +2518,7 @@ UInterpDataFactoryNew::UInterpDataFactoryNew(const FObjectInitializer& ObjectIni
 
 UObject* UInterpDataFactoryNew::FactoryCreateNew(UClass* Class,UObject* InParent,FName Name,EObjectFlags Flags,UObject* Context,FFeedbackContext* Warn)
 {
-	return StaticConstructObject(Class,InParent,Name,Flags);
+	return NewObject<UObject>(InParent, Class, Name, Flags);
 }
 
 /*-----------------------------------------------------------------------------
@@ -2542,7 +2542,7 @@ UTextureRenderTargetFactoryNew::UTextureRenderTargetFactoryNew(const FObjectInit
 UObject* UTextureRenderTargetFactoryNew::FactoryCreateNew(UClass* Class,UObject* InParent,FName Name,EObjectFlags Flags,UObject* Context,FFeedbackContext* Warn)
 {
 	// create the new object
-	UTextureRenderTarget2D* Result = CastChecked<UTextureRenderTarget2D>( StaticConstructObject(Class,InParent,Name,Flags) );
+	UTextureRenderTarget2D* Result = NewObject<UTextureRenderTarget2D>(InParent, Class, Name, Flags);
 	// initialize the resource
 	Result->InitAutoFormat( Width, Height );
 	return( Result );
@@ -2568,7 +2568,7 @@ UTextureRenderTargetCubeFactoryNew::UTextureRenderTargetCubeFactoryNew(const FOb
 UObject* UTextureRenderTargetCubeFactoryNew::FactoryCreateNew(UClass* Class, UObject* InParent, FName Name, EObjectFlags Flags, UObject* Context, FFeedbackContext* Warn)
 {
 	// create the new object
-	UTextureRenderTargetCube* Result = CastChecked<UTextureRenderTargetCube>(StaticConstructObject(Class, InParent, Name, Flags));
+	UTextureRenderTargetCube* Result = NewObject<UTextureRenderTargetCube>(InParent, Class, Name, Flags);
 
 	// initialize the resource
 	Result->InitAutoFormat(Width);
@@ -6457,7 +6457,7 @@ UObject* UCurveFactory::FactoryCreateNew(UClass* Class, UObject* InParent, FName
 	UCurveBase* NewCurve = nullptr;
 	if(CurveClass != nullptr)
 	{
-		NewCurve = CastChecked<UCurveBase>(StaticConstructObject(CurveClass,InParent,Name,Flags));		
+		NewCurve = NewObject<UCurveBase>(InParent, CurveClass, Name, Flags);
 	}
 
 	return NewCurve;
@@ -6540,7 +6540,7 @@ UObject* UCurveImportFactory::FactoryCreateText( UClass* InClass, UObject* InPar
 		// make the curve object and set up the keys
 		if (FloatKeys.Num() > 0)
 		{
-			UCurveFloat* NewCurve = CastChecked<UCurveFloat>(StaticConstructObject(UCurveFloat::StaticClass(),InParent,InName,Flags));
+			UCurveFloat* NewCurve = NewObject<UCurveFloat>(InParent,InName,Flags);
 
 			if (NewCurve)
 			{
@@ -6578,7 +6578,7 @@ UObjectLibraryFactory::UObjectLibraryFactory(const FObjectInitializer& ObjectIni
 
 UObject* UObjectLibraryFactory::FactoryCreateNew(UClass* Class, UObject* InParent, FName Name, EObjectFlags Flags, UObject* Context, FFeedbackContext* Warn)
 {
-	return CastChecked<UObjectLibrary>(StaticConstructObject(Class,InParent,Name,Flags));
+	return NewObject<UObjectLibrary>(InParent, Class, Name, Flags);
 }
 
 /*------------------------------------------------------------------------------
@@ -6627,13 +6627,13 @@ UObject* UDataAssetFactory::FactoryCreateNew(UClass* Class, UObject* InParent, F
 {
 	if (DataAssetClass != nullptr)
 	{
-		return CastChecked<UDataAsset>(StaticConstructObject(DataAssetClass, InParent, Name, Flags | RF_Transactional));
+		return NewObject<UDataAsset>(InParent, DataAssetClass, Name, Flags | RF_Transactional);
 	}
 	else
 	{
 		// if we have no data asset class, use the passed-in class instead
 		check(Class->IsChildOf(UDataAsset::StaticClass()));
-		return CastChecked<UDataAsset>(StaticConstructObject(Class, InParent, Name, Flags));
+		return NewObject<UDataAsset>(InParent, Class, Name, Flags);
 	}
 }
 
@@ -6909,7 +6909,7 @@ UObject* UBlendSpaceFactoryNew::FactoryCreateNew(UClass* Class,UObject* InParent
 {
 	if (TargetSkeleton)
 	{
-		UBlendSpace * BlendSpace = Cast<UBlendSpace>(StaticConstructObject(Class,InParent,Name,Flags));
+		UBlendSpace * BlendSpace = NewObject<UBlendSpace>(InParent, Class, Name, Flags);
 
 		BlendSpace->SetSkeleton(TargetSkeleton);
 
@@ -6982,7 +6982,7 @@ UObject* UBlendSpaceFactory1D::FactoryCreateNew(UClass* Class,UObject* InParent,
 {
 	if (TargetSkeleton)
 	{
-		UBlendSpace1D * BlendSpace = Cast<UBlendSpace1D>(StaticConstructObject(Class,InParent,Name,Flags));
+		UBlendSpace1D * BlendSpace = NewObject<UBlendSpace1D>(InParent, Class, Name, Flags);
 
 		BlendSpace->SetSkeleton(TargetSkeleton);
 
@@ -7012,7 +7012,7 @@ UObject* UAimOffsetBlendSpaceFactoryNew::FactoryCreateNew(UClass* Class,UObject*
 {
 	if (TargetSkeleton)
 	{
-		UAimOffsetBlendSpace * BlendSpace = Cast<UAimOffsetBlendSpace>(StaticConstructObject(Class,InParent,Name,Flags));
+		UAimOffsetBlendSpace * BlendSpace = NewObject<UAimOffsetBlendSpace>(InParent, Class, Name, Flags);
 
 		BlendSpace->SetSkeleton(TargetSkeleton);
 
@@ -7036,7 +7036,7 @@ UObject* UAimOffsetBlendSpaceFactory1D::FactoryCreateNew(UClass* Class,UObject*
 {
 	if (TargetSkeleton)
 	{
-		UAimOffsetBlendSpace1D * BlendSpace = Cast<UAimOffsetBlendSpace1D>(StaticConstructObject(Class,InParent,Name,Flags));
+		UAimOffsetBlendSpace1D * BlendSpace = NewObject<UAimOffsetBlendSpace1D>(InParent, Class, Name, Flags);
 
 		BlendSpace->SetSkeleton(TargetSkeleton);
 
@@ -7175,7 +7175,7 @@ uint32 UCameraAnimFactory::GetMenuCategories() const
 
 UObject* UCameraAnimFactory::FactoryCreateNew(UClass* Class,UObject* InParent,FName Name,EObjectFlags Flags,UObject* Context,FFeedbackContext* Warn)
 {
-	UCameraAnim* NewCamAnim = CastChecked<UCameraAnim>(StaticConstructObject(Class,InParent,Name,Flags));
+	UCameraAnim* NewCamAnim = NewObject<UCameraAnim>(InParent, Class, Name, Flags);
 	NewCamAnim->CameraInterpGroup = NewObject<UInterpGroupCamera>(NewCamAnim);
 	NewCamAnim->CameraInterpGroup->GroupName = Name;
 	return NewCamAnim;
diff --git a/Engine/Source/Editor/UnrealEd/Private/Factories/Factory.cpp b/Engine/Source/Editor/UnrealEd/Private/Factories/Factory.cpp
--- a/Engine/Source/Editor/UnrealEd/Private/Factories/Factory.cpp
+++ b/Engine/Source/Editor/UnrealEd/Private/Factories/Factory.cpp
@@ -460,13 +460,13 @@ UObject* UFactory::CreateOrOverwriteAsset(UClass* InClass, UObject* InParent, FN
 	UObject* ExistingAsset = StaticFindObject(NULL, InParent, *InName.ToString());
 	if ( !ExistingAsset )
 	{
-		return StaticConstructObject(InClass, InParent, InName, InFlags, InTemplate);
+		return NewObject<UObject>(InParent, InClass, InName, InFlags, InTemplate);
 	}
 
 	// If it does exist then it overwrites it if possible.
 	if ( ExistingAsset->GetClass()->IsChildOf(InClass) )
 	{
-		return StaticConstructObject(InClass, InParent, InName, InFlags, InTemplate);
+		return NewObject<UObject>(InParent, InClass, InName, InFlags, InTemplate);
 	}
 	
 	// If it can not overwrite then it will delete and replace.
@@ -497,7 +497,7 @@ UObject* UFactory::CreateOrOverwriteAsset(UClass* InClass, UObject* InParent, FN
 		else
 		{
 			// We can now create the asset in the package
-			return StaticConstructObject(InClass, InParent, InName, InFlags, InTemplate);
+			return NewObject<UObject>(InParent, InClass, InName, InFlags, InTemplate);
 		}
 	}
 	else
diff --git a/Engine/Source/Editor/UnrealEd/Private/Factories/SlateBrushAssetFactory.cpp b/Engine/Source/Editor/UnrealEd/Private/Factories/SlateBrushAssetFactory.cpp
--- a/Engine/Source/Editor/UnrealEd/Private/Factories/SlateBrushAssetFactory.cpp
+++ b/Engine/Source/Editor/UnrealEd/Private/Factories/SlateBrushAssetFactory.cpp
@@ -29,7 +29,7 @@ bool USlateBrushAssetFactory::ConfigureProperties()
 
 UObject* USlateBrushAssetFactory::FactoryCreateNew(UClass* Class,UObject* InParent,FName Name,EObjectFlags Flags,UObject* Context,FFeedbackContext* Warn)
 {
-	USlateBrushAsset* NewSlateBrushAsset = CastChecked<USlateBrushAsset>(StaticConstructObject(USlateBrushAsset::StaticClass(), InParent, Name, Flags));
+	USlateBrushAsset* NewSlateBrushAsset = NewObject<USlateBrushAsset>(InParent, Name, Flags);
 	NewSlateBrushAsset->Brush = InitialTexture != NULL ? FSlateDynamicImageBrush( InitialTexture, FVector2D( InitialTexture->GetImportedSize() ), NAME_None ) : FSlateBrush();
 	return NewSlateBrushAsset;
 }
diff --git a/Engine/Source/Editor/UnrealEd/Private/Factories/SlateWidgetStyleAssetFactory.cpp b/Engine/Source/Editor/UnrealEd/Private/Factories/SlateWidgetStyleAssetFactory.cpp
--- a/Engine/Source/Editor/UnrealEd/Private/Factories/SlateWidgetStyleAssetFactory.cpp
+++ b/Engine/Source/Editor/UnrealEd/Private/Factories/SlateWidgetStyleAssetFactory.cpp
@@ -81,10 +81,10 @@ bool USlateWidgetStyleAssetFactory::ConfigureProperties()
 
 UObject* USlateWidgetStyleAssetFactory::FactoryCreateNew(UClass* Class,UObject* InParent,FName Name,EObjectFlags Flags,UObject* Context,FFeedbackContext* Warn)
 {
-	USlateWidgetStyleAsset* NewUSlateStyle = CastChecked<USlateWidgetStyleAsset>(StaticConstructObject(USlateWidgetStyleAsset::StaticClass(), InParent, Name, Flags));
+	USlateWidgetStyleAsset* NewUSlateStyle = NewObject<USlateWidgetStyleAsset>(InParent, Name, Flags);
 
 	//intialize
-	NewUSlateStyle->CustomStyle = CastChecked< USlateWidgetStyleContainerBase >( StaticConstructObject(StyleType, NewUSlateStyle, Name) );
+	NewUSlateStyle->CustomStyle = NewObject<USlateWidgetStyleContainerBase>(NewUSlateStyle, StyleType, Name);
 
 	return NewUSlateStyle;
 }
\ No newline at end of file
diff --git a/Engine/Source/Editor/UnrealEd/Private/Factories/VectorFieldFactory.cpp b/Engine/Source/Editor/UnrealEd/Private/Factories/VectorFieldFactory.cpp
--- a/Engine/Source/Editor/UnrealEd/Private/Factories/VectorFieldFactory.cpp
+++ b/Engine/Source/Editor/UnrealEd/Private/Factories/VectorFieldFactory.cpp
@@ -220,11 +220,7 @@ UObject* UVectorFieldStaticFactory::FactoryCreateBinary(
 			FFGAContents FileContents;
 			if ( ParseFGA( &FileContents, Stream, Warn ) )
 			{
-				VectorField = CastChecked<UVectorFieldStatic>( StaticConstructObject(
-					UVectorFieldStatic::StaticClass(),
-					InParent,
-					InName,
-					Flags ) );
+				VectorField = NewObject<UVectorFieldStatic>(InParent, InName,Flags);
 
 				VectorField->SizeX = FileContents.GridX;
 				VectorField->SizeY = FileContents.GridY;
diff --git a/Engine/Source/Editor/UnrealEd/Private/Fbx/FbxSkeletalMeshImport.cpp b/Engine/Source/Editor/UnrealEd/Private/Fbx/FbxSkeletalMeshImport.cpp
--- a/Engine/Source/Editor/UnrealEd/Private/Fbx/FbxSkeletalMeshImport.cpp
+++ b/Engine/Source/Editor/UnrealEd/Private/Fbx/FbxSkeletalMeshImport.cpp
@@ -1111,7 +1111,7 @@ USkeletalMesh* UnFbx::FFbxImporter::ImportSkeletalMesh(UObject* InParent, TArray
 	}
 
 	// [from USkeletalMeshFactory::FactoryCreateBinary]
-	USkeletalMesh* SkeletalMesh = CastChecked<USkeletalMesh>( StaticConstructObject( USkeletalMesh::StaticClass(), InParent, Name, Flags ) );
+	USkeletalMesh* SkeletalMesh = NewObject<USkeletalMesh>(InParent, Name, Flags);
 	
 	SkeletalMesh->PreEditChange(NULL);
 	
@@ -1458,7 +1458,7 @@ UObject* UnFbx::FFbxImporter::CreateAssetOfClass(UClass* AssetClass, FString Par
 	{
 		// add it to the set
 		// do not add to the set, now create independent asset
-		Object = StaticConstructObject( AssetClass, Parent, *ObjectName, RF_Public|RF_Standalone  );
+		Object = NewObject<UObject>(Parent, AssetClass, *ObjectName, RF_Public | RF_Standalone);
 		Object->MarkPackageDirty();
 		// Notify the asset registry
 		FAssetRegistryModule::AssetCreated(Object);
diff --git a/Engine/Source/Editor/UnrealEd/Private/ObjectTools.cpp b/Engine/Source/Editor/UnrealEd/Private/ObjectTools.cpp
--- a/Engine/Source/Editor/UnrealEd/Private/ObjectTools.cpp
+++ b/Engine/Source/Editor/UnrealEd/Private/ObjectTools.cpp
@@ -810,7 +810,7 @@ namespace ObjectTools
 
 					// Create a redirector with a unique name
 					// It will have the same name as the object that was consolidated after the garbage collect
-					UObjectRedirector* Redirector = Cast<UObjectRedirector>( StaticConstructObject( UObjectRedirector::StaticClass(), CurObjOuter, NAME_None, RF_Standalone | RF_Public ) );
+					UObjectRedirector* Redirector = NewObject<UObjectRedirector>(CurObjOuter, NAME_None, RF_Standalone | RF_Public);
 					check( Redirector );
 
 					// Set the redirector to redirect to the object to consolidate to
@@ -824,13 +824,13 @@ namespace ObjectTools
 					if ( BlueprintToConsolidateTo != NULL && BlueprintToConsolidate != NULL )
 					{
 						// One redirector for the class
-						UObjectRedirector* ClassRedirector = Cast<UObjectRedirector>( StaticConstructObject( UObjectRedirector::StaticClass(), CurObjOuter, NAME_None, RF_Standalone | RF_Public ) );
+						UObjectRedirector* ClassRedirector = NewObject<UObjectRedirector>(CurObjOuter, NAME_None, RF_Standalone | RF_Public);
 						check( ClassRedirector );
 						ClassRedirector->DestinationObject = BlueprintToConsolidateTo->GeneratedClass;
 						RedirectorToObjectNameMap.Add(ClassRedirector, BlueprintToConsolidate->GeneratedClass->GetFName());
 
 						// One redirector for the CDO
-						UObjectRedirector* CDORedirector = Cast<UObjectRedirector>( StaticConstructObject( UObjectRedirector::StaticClass(), CurObjOuter, NAME_None, RF_Standalone | RF_Public ) );
+						UObjectRedirector* CDORedirector = NewObject<UObjectRedirector>(CurObjOuter, NAME_None, RF_Standalone | RF_Public);
 						check( CDORedirector );
 						CDORedirector->DestinationObject = BlueprintToConsolidateTo->GeneratedClass->GetDefaultObject();
 						RedirectorToObjectNameMap.Add(CDORedirector, BlueprintToConsolidate->GeneratedClass->GetDefaultObject()->GetFName());
diff --git a/Engine/Source/Editor/UnrealEd/Private/StaticLightingSystem/StaticLightingExport.cpp b/Engine/Source/Editor/UnrealEd/Private/StaticLightingSystem/StaticLightingExport.cpp
--- a/Engine/Source/Editor/UnrealEd/Private/StaticLightingSystem/StaticLightingExport.cpp
+++ b/Engine/Source/Editor/UnrealEd/Private/StaticLightingSystem/StaticLightingExport.cpp
@@ -63,8 +63,7 @@ UObject* FBSPSurfaceStaticLighting::GetMappedObject() const
 {
 	//@todo. THIS WILL SCREW UP IF CALLED MORE THAN ONE TIME!!!!
 	// Create a collection object to allow selection of the surfaces in this mapping
-	ULightmappedSurfaceCollection* MappedObject = CastChecked<ULightmappedSurfaceCollection>(
-		StaticConstructObject(ULightmappedSurfaceCollection::StaticClass()));
+	auto MappedObject = NewObject<ULightmappedSurfaceCollection>();
 	// Set the owner model
 	MappedObject->SourceModel = Model.Get();
 	// Fill in the surface index array
diff --git a/Engine/Source/Editor/UnrealEd/Private/TexAlignTools.cpp b/Engine/Source/Editor/UnrealEd/Private/TexAlignTools.cpp
--- a/Engine/Source/Editor/UnrealEd/Private/TexAlignTools.cpp
+++ b/Engine/Source/Editor/UnrealEd/Private/TexAlignTools.cpp
@@ -441,10 +441,10 @@ void FTexAlignTools::Init()
 {
 	// Create the list of aligners.
 	Aligners.Empty();
-	Aligners.Add( CastChecked<UTexAligner>(StaticConstructObject(UTexAlignerDefault::StaticClass(),GetTransientPackage(),NAME_None,RF_Public|RF_RootSet|RF_Standalone) ) );
-	Aligners.Add( CastChecked<UTexAligner>(StaticConstructObject(UTexAlignerPlanar::StaticClass(),GetTransientPackage(),NAME_None,RF_Public|RF_RootSet|RF_Standalone) ) );
-	Aligners.Add( CastChecked<UTexAligner>(StaticConstructObject(UTexAlignerBox::StaticClass(),GetTransientPackage(),NAME_None,RF_Public|RF_RootSet|RF_Standalone) ) );
-	Aligners.Add( CastChecked<UTexAligner>(StaticConstructObject(UTexAlignerFit::StaticClass(),GetTransientPackage(),NAME_None,RF_Public|RF_RootSet|RF_Standalone) ) );
+	Aligners.Add(NewObject<UTexAligner>(GetTransientPackage(), NAME_None, RF_Public | RF_RootSet | RF_Standalone));
+	Aligners.Add(NewObject<UTexAligner>(GetTransientPackage(), NAME_None, RF_Public | RF_RootSet | RF_Standalone));
+	Aligners.Add(NewObject<UTexAligner>(GetTransientPackage(), NAME_None, RF_Public | RF_RootSet | RF_Standalone));
+	Aligners.Add(NewObject<UTexAligner>(GetTransientPackage(), NAME_None, RF_Public | RF_RootSet | RF_Standalone));
 	
 	FEditorDelegates::FitTextureToSurface.AddRaw(this, &FTexAlignTools::OnEditorFitTextureToSurface);
 }
diff --git a/Engine/Source/Editor/WorldBrowser/Private/Tiles/WorldTileModel.cpp b/Engine/Source/Editor/WorldBrowser/Private/Tiles/WorldTileModel.cpp
--- a/Engine/Source/Editor/WorldBrowser/Private/Tiles/WorldTileModel.cpp
+++ b/Engine/Source/Editor/WorldBrowser/Private/Tiles/WorldTileModel.cpp
@@ -29,11 +29,7 @@ FWorldTileModel::FWorldTileModel(const TWeakObjectPtr<UEditorEngine>& InEditor,
 	UWorldComposition* WorldComposition = LevelCollectionModel.GetWorld()->WorldComposition;
 
 	// Tile display details object
-	TileDetails = Cast<UWorldTileDetails>(
-		StaticConstructObject(UWorldTileDetails::StaticClass(), 
-			GetTransientPackage(), NAME_None, RF_RootSet|RF_Transient, NULL
-			)
-		);
+	TileDetails = NewObject<UWorldTileDetails>(GetTransientPackage(), NAME_None, RF_RootSet | RF_Transient);
 
 	// Subscribe to tile properties changes
 	TileDetails->PositionChangedEvent.AddRaw(this, &FWorldTileModel::OnPositionPropertyChanged);
@@ -600,9 +596,7 @@ ULevelStreaming* FWorldTileModel::GetAssosiatedStreamingLevel()
 	if (Index == INDEX_NONE)
 	{
 		// Create new streaming level
-		ULevelStreaming* AssociatedStreamingLevel = Cast<ULevelStreaming>(
-			StaticConstructObject(ULevelStreamingKismet::StaticClass(), PersistentWorld, NAME_None, RF_Transient, NULL)
-			);
+		auto AssociatedStreamingLevel = NewObject<ULevelStreamingKismet>(PersistentWorld, NAME_None, RF_Transient);
 
 		//
 		AssociatedStreamingLevel->SetWorldAssetByPackageName(PackageName);
diff --git a/Engine/Source/Runtime/CoreUObject/Private/Serialization/DuplicateDataWriter.cpp b/Engine/Source/Runtime/CoreUObject/Private/Serialization/DuplicateDataWriter.cpp
--- a/Engine/Source/Runtime/CoreUObject/Private/Serialization/DuplicateDataWriter.cpp
+++ b/Engine/Source/Runtime/CoreUObject/Private/Serialization/DuplicateDataWriter.cpp
@@ -132,7 +132,7 @@ UObject* FDuplicateDataWriter::GetDuplicatedObject(UObject* Object, bool bCreate
 			if(DupOuter != NULL)
 			{
 				// The object's outer is being duplicated, create a duplicate of this object.
-				UObject* NewEmptyDuplicate = StaticConstructObject(Object->GetClass(), DupOuter, Object->GetFName(), ApplyFlags|Object->GetMaskedFlags(FlagMask),
+				UObject* NewEmptyDuplicate = StaticConstructObject_Internal(Object->GetClass(), DupOuter, Object->GetFName(), ApplyFlags|Object->GetMaskedFlags(FlagMask),
 					Object->GetArchetype(), true, InstanceGraph);
 
 				Result = AddDuplicate(Object, NewEmptyDuplicate);
diff --git a/Engine/Source/Runtime/CoreUObject/Private/UObject/CoreNative.cpp b/Engine/Source/Runtime/CoreUObject/Private/UObject/CoreNative.cpp
--- a/Engine/Source/Runtime/CoreUObject/Private/UObject/CoreNative.cpp
+++ b/Engine/Source/Runtime/CoreUObject/Private/UObject/CoreNative.cpp
@@ -173,7 +173,7 @@ UObject* FObjectInstancingGraph::GetInstancedSubobject( UObject* SourceSubobject
 							if (!InstancedSubobject)
 							{
 								// finally, create the component instance
-								InstancedSubobject = StaticConstructObject(SourceSubobject->GetClass(), SubobjectOuter,
+								InstancedSubobject = StaticConstructObject_Internal(SourceSubobject->GetClass(), SubobjectOuter,
 									SubobjectName, SubobjectOuter->GetMaskedFlags(RF_PropagateToSubObjects), SourceSubobject,
 									true, this);
 							}
diff --git a/Engine/Source/Runtime/CoreUObject/Private/UObject/LinkerLoad.cpp b/Engine/Source/Runtime/CoreUObject/Private/UObject/LinkerLoad.cpp
--- a/Engine/Source/Runtime/CoreUObject/Private/UObject/LinkerLoad.cpp
+++ b/Engine/Source/Runtime/CoreUObject/Private/UObject/LinkerLoad.cpp
@@ -3388,7 +3388,7 @@ UObject* ULinkerLoad::CreateExport( int32 Index )
 		LoadClass->GetDefaultObject();
 
 
-		Export.Object = StaticConstructObject
+		Export.Object = StaticConstructObject_Internal
 		(
 			LoadClass,
 			ThisParent,
@@ -3401,7 +3401,7 @@ UObject* ULinkerLoad::CreateExport( int32 Index )
 		if (NewName != Export.ObjectName)
 		{
 			// create a UObjectRedirector with the same name as the old object we are redirecting
-			UObjectRedirector* Redir = (UObjectRedirector*)StaticConstructObject(UObjectRedirector::StaticClass(), Export.Object->GetOuter(), Export.ObjectName, RF_Standalone | RF_Public);
+			UObjectRedirector* Redir = NewObject<UObjectRedirector>(Export.Object->GetOuter(), Export.ObjectName, RF_Standalone | RF_Public);
 			// point the redirector object to this object
 			Redir->DestinationObject = Export.Object;
 		}
diff --git a/Engine/Source/Runtime/CoreUObject/Private/UObject/Obj.cpp b/Engine/Source/Runtime/CoreUObject/Private/UObject/Obj.cpp
--- a/Engine/Source/Runtime/CoreUObject/Private/UObject/Obj.cpp
+++ b/Engine/Source/Runtime/CoreUObject/Private/UObject/Obj.cpp
@@ -1846,7 +1846,7 @@ void UObject::ReinitializeProperties( UObject* SourceObject/*=NULL*/, FObjectIns
 	// the properties for this object ensures that any cleanup required when an object is reinitialized from defaults occurs properly
 	// for example, when re-initializing UPrimitiveComponents, the component must notify the rendering thread that its data structures are
 	// going to be re-initialized
-	StaticConstructObject( GetClass(), GetOuter(), GetFName(), GetFlags(), SourceObject, !HasAnyFlags(RF_ClassDefaultObject), InstanceGraph );
+	StaticConstructObject_Internal( GetClass(), GetOuter(), GetFName(), GetFlags(), SourceObject, !HasAnyFlags(RF_ClassDefaultObject), InstanceGraph );
 }
 
 
@@ -2091,7 +2091,7 @@ UObject* UObject::CreateArchetype( const TCHAR* ArchetypeName, UObject* Archetyp
 		ArchetypeObjectFlags |= RF_Standalone;
 	}
 
-	UObject* ArchetypeObject = StaticConstructObject(GetClass(), ArchetypeOuter, ArchetypeName, ArchetypeObjectFlags, this, true, InstanceGraph);
+	UObject* ArchetypeObject = StaticConstructObject_Internal(GetClass(), ArchetypeOuter, ArchetypeName, ArchetypeObjectFlags, this, true, InstanceGraph);
 	check(ArchetypeObject);
 
 	UObject* NewArchetype = AlternateArchetype == NULL
diff --git a/Engine/Source/Runtime/CoreUObject/Private/UObject/UObjectBase.cpp b/Engine/Source/Runtime/CoreUObject/Private/UObject/UObjectBase.cpp
--- a/Engine/Source/Runtime/CoreUObject/Private/UObject/UObjectBase.cpp
+++ b/Engine/Source/Runtime/CoreUObject/Private/UObject/UObjectBase.cpp
@@ -709,7 +709,7 @@ private:
 	{
 		auto* ClassWithin = Class->ClassWithin;
 		Class->ClassWithin = UPackage::StaticClass(); // We are just avoiding error checks with this... we don't care about this temp object other than to get the vtable.
-		UObject* TempObjectForVTable = StaticConstructObject(Class, GetTransientPackage(), NAME_None, RF_NeedLoad | RF_ClassDefaultObject);
+		UObject* TempObjectForVTable = StaticConstructObject_Internal(Class, GetTransientPackage(), NAME_None, RF_NeedLoad | RF_ClassDefaultObject);
 
 		if (!TempObjectForVTable->IsRooted())
 		{
diff --git a/Engine/Source/Runtime/CoreUObject/Private/UObject/UObjectGlobals.cpp b/Engine/Source/Runtime/CoreUObject/Private/UObject/UObjectGlobals.cpp
--- a/Engine/Source/Runtime/CoreUObject/Private/UObject/UObjectGlobals.cpp
+++ b/Engine/Source/Runtime/CoreUObject/Private/UObject/UObjectGlobals.cpp
@@ -1441,7 +1441,7 @@ UObject* StaticDuplicateObjectEx( FObjectDuplicationParameters& Parameters )
 	UObject* DupRootObject = Parameters.DuplicationSeed.FindRef(Parameters.SourceObject);
 	if ( DupRootObject == NULL )
 	{
-		DupRootObject = StaticConstructObject(	Parameters.DestClass,
+		DupRootObject = StaticConstructObject_Internal(	Parameters.DestClass,
 														Parameters.DestOuter,
 														Parameters.DestName,
 														Parameters.ApplyFlags|Parameters.SourceObject->GetMaskedFlags(Parameters.FlagMask),
@@ -2220,7 +2220,7 @@ bool DebugIsClassChildOf_Internal(UClass* Parent, UClass* Child)
 	return Child->IsChildOf(Parent);
 }
 
-UObject* StaticConstructObject
+UObject* StaticConstructObject_Internal
 (
 	UClass*			InClass,
 	UObject*		InOuter								/*=GetTransientPackage()*/,
@@ -2269,6 +2269,20 @@ UObject* StaticConstructObject
 	return Result;
 }
 
+UObject* StaticConstructObject
+(
+	UClass*			InClass,
+	UObject*		InOuter								/*=GetTransientPackage()*/,
+	FName			InName								/*=NAME_None*/,
+	EObjectFlags	InFlags								/*=0*/,
+	UObject*		InTemplate							/*=NULL*/,
+	bool			bCopyTransientsFromClassDefaults	/*=false*/, 
+	FObjectInstancingGraph* InInstanceGraph				/*=NULL*/
+)
+{
+	return StaticConstructObject_Internal(InClass, InOuter, InName, InFlags, InTemplate, bCopyTransientsFromClassDefaults, InInstanceGraph);
+}
+
 void FObjectInitializer::AssertIfInConstructor(UObject* Outer, const TCHAR* ErrorMessage)
 {
 	UE_CLOG(GIsInConstructor && Outer == GConstructedObject, LogUObjectGlobals, Fatal, TEXT("%s"), ErrorMessage);
@@ -2670,7 +2684,7 @@ UObject* FObjectInitializer::CreateDefaultSubobject(UObject* Outer, FName Subobj
 		{
 			UObject* Template = OverrideClass->GetDefaultObject(); // force the CDO to be created if it hasn't already
 			const EObjectFlags SubobjectFlags = Outer->GetMaskedFlags(RF_PropagateToSubObjects);
-			Result = StaticConstructObject(OverrideClass, Outer, SubobjectFName, SubobjectFlags);
+			Result = StaticConstructObject_Internal(OverrideClass, Outer, SubobjectFName, SubobjectFlags);
 			if (!bIsTransient && !Outer->GetArchetype()->GetClass()->HasAnyClassFlags(CLASS_Native | CLASS_Intrinsic))
 			{
 				// The archetype of the outer is not native, so we need to copy properties to the subobjects after the C++ constructor chain for the outer has run (because those sets properties on the subobjects)
diff --git a/Engine/Source/Runtime/CoreUObject/Public/UObject/Class.h b/Engine/Source/Runtime/CoreUObject/Public/UObject/Class.h
--- a/Engine/Source/Runtime/CoreUObject/Public/UObject/Class.h
+++ b/Engine/Source/Runtime/CoreUObject/Public/UObject/Class.h
@@ -2567,7 +2567,7 @@ T* ConstructObject(UClass* Class, UObject* Outer, FName Name, EObjectFlags SetFl
 {
 	checkf(Class, TEXT("ConstructObject called with a NULL class object"));
 	checkSlow(Class->IsChildOf(T::StaticClass()));
-	return (T*)StaticConstructObject(Class, Outer, Name, SetFlags, Template, bCopyTransientsFromClassDefaults, InstanceGraph);
+	return (T*)StaticConstructObject_Internal(Class, Outer, Name, SetFlags, Template, bCopyTransientsFromClassDefaults, InstanceGraph);
 }
 
 /**
diff --git a/Engine/Source/Runtime/CoreUObject/Public/UObject/UObjectGlobals.h b/Engine/Source/Runtime/CoreUObject/Public/UObject/UObjectGlobals.h
--- a/Engine/Source/Runtime/CoreUObject/Public/UObject/UObjectGlobals.h
+++ b/Engine/Source/Runtime/CoreUObject/Public/UObject/UObjectGlobals.h
@@ -183,6 +183,26 @@ COREUOBJECT_API bool ParseObject( const TCHAR* Stream, const TCHAR* Match, UClas
  */
 COREUOBJECT_API UObject* StaticLoadObject( UClass* Class, UObject* InOuter, const TCHAR* Name, const TCHAR* Filename = NULL, uint32 LoadFlags = LOAD_None, UPackageMap* Sandbox = NULL, bool bAllowObjectReconciliation = true );
 COREUOBJECT_API UClass* StaticLoadClass(UClass* BaseClass, UObject* InOuter, const TCHAR* Name, const TCHAR* Filename = NULL, uint32 LoadFlags = LOAD_None, UPackageMap* Sandbox = NULL);
+
+/**
+* Create a new instance of an object.  The returned object will be fully initialized.  If InFlags contains RF_NeedsLoad (indicating that the object still needs to load its object data from disk), components
+* are not instanced (this will instead occur in PostLoad()).  The different between StaticConstructObject and StaticAllocateObject is that StaticConstructObject will also call the class constructor on the object
+* and instance any components.
+*
+* @param	Class		the class of the object to create
+* @param	InOuter		the object to create this object within (the Outer property for the new object will be set to the value specified here).
+* @param	Name		the name to give the new object. If no value (NAME_None) is specified, the object will be given a unique name in the form of ClassName_#.
+* @param	SetFlags	the ObjectFlags to assign to the new object. some flags can affect the behavior of constructing the object.
+* @param	Template	if specified, the property values from this object will be copied to the new object, and the new object's ObjectArchetype value will be set to this object.
+*						If NULL, the class default object is used instead.
+* @param	bInCopyTransientsFromClassDefaults - if true, copy transient from the class defaults instead of the pass in archetype ptr (often these are the same)
+* @param	InstanceGraph
+*						contains the mappings of instanced objects and components to their templates
+*
+* @return	a pointer to a fully initialized object of the specified class.
+*/
+COREUOBJECT_API UObject* StaticConstructObject_Internal(UClass* Class, UObject* InOuter = (UObject*)GetTransientPackage(), FName Name = NAME_None, EObjectFlags SetFlags = RF_NoFlags, UObject* Template = NULL, bool bCopyTransientsFromClassDefaults = false, struct FObjectInstancingGraph* InstanceGraph = NULL);
+
 /**
  * Create a new instance of an object.  The returned object will be fully initialized.  If InFlags contains RF_NeedsLoad (indicating that the object still needs to load its object data from disk), components
  * are not instanced (this will instead occur in PostLoad()).  The different between StaticConstructObject and StaticAllocateObject is that StaticConstructObject will also call the class constructor on the object
@@ -200,7 +220,9 @@ COREUOBJECT_API UClass* StaticLoadClass(UClass* BaseClass, UObject* InOuter, con
  *
  * @return	a pointer to a fully initialized object of the specified class.
  */
+DEPRECATED(4.8, "StaticConstructObject is deprecated, please use NewObject instead. For internal CoreUObject module usage, please use StaticConstructObject_Internal.")
 COREUOBJECT_API UObject* StaticConstructObject( UClass* Class, UObject* InOuter=(UObject*)GetTransientPackage(), FName Name=NAME_None, EObjectFlags SetFlags=RF_NoFlags, UObject* Template=NULL, bool bCopyTransientsFromClassDefaults=false, struct FObjectInstancingGraph* InstanceGraph=NULL );
+
 /**
  * Creates a copy of SourceObject using the Outer and Name specified, as well as copies of all objects contained by SourceObject.  
  * Any objects referenced by SourceOuter or RootObject and contained by SourceOuter are also copied, maintaining their name relative to SourceOuter.  Any
@@ -1057,7 +1079,7 @@ T* NewObject(UObject* Outer = (UObject*)GetTransientPackage(), UClass* Class = T
 	}
 	checkf(Class, TEXT("NewObject called with a nullptr class object"));
 	checkSlow(DebugIsClassChildOf_Internal(T::StaticClass(), Class));
-	return static_cast<T*>(StaticConstructObject(Class, Outer, Name, Flags, Template, bCopyTransientsFromClassDefaults, InInstanceGraph));
+	return static_cast<T*>(StaticConstructObject_Internal(Class, Outer, Name, Flags, Template, bCopyTransientsFromClassDefaults, InInstanceGraph));
 }
 
 template< class T >
diff --git a/Engine/Source/Runtime/Engine/Private/DemoNetDriver.cpp b/Engine/Source/Runtime/Engine/Private/DemoNetDriver.cpp
--- a/Engine/Source/Runtime/Engine/Private/DemoNetDriver.cpp
+++ b/Engine/Source/Runtime/Engine/Private/DemoNetDriver.cpp
@@ -294,7 +294,7 @@ bool UDemoNetDriver::InitConnectInternal( FString& Error )
 	// Read meta data, if it exists
 	for ( int32 i = 0; i < MetadataHeader.NumStreamingLevels; ++i )
 	{
-		ULevelStreamingKismet* StreamingLevel = static_cast<ULevelStreamingKismet*>(StaticConstructObject(ULevelStreamingKismet::StaticClass(), GetWorld(), NAME_None, RF_NoFlags, NULL ) );
+		ULevelStreamingKismet* StreamingLevel = NewObject<ULevelStreamingKismet>(GetWorld(), NAME_None, RF_NoFlags, NULL);
 
 		StreamingLevel->bShouldBeLoaded		= true;
 		StreamingLevel->bShouldBeVisible	= true;
diff --git a/Engine/Source/Runtime/Engine/Private/DestructibleMesh.cpp b/Engine/Source/Runtime/Engine/Private/DestructibleMesh.cpp
--- a/Engine/Source/Runtime/Engine/Private/DestructibleMesh.cpp
+++ b/Engine/Source/Runtime/Engine/Private/DestructibleMesh.cpp
@@ -479,7 +479,7 @@ void UDestructibleMesh::CreateFractureSettings()
 #if WITH_EDITORONLY_DATA
 	if (FractureSettings == NULL)
 	{
-		FractureSettings = CastChecked<UDestructibleFractureSettings>(StaticConstructObject(UDestructibleFractureSettings::StaticClass(), this));
+		FractureSettings = NewObject<UDestructibleFractureSettings>(this);
 		check(FractureSettings);
 	}
 #endif	// WITH_EDITORONLY_DATA
diff --git a/Engine/Source/Runtime/Engine/Private/GameInstance.cpp b/Engine/Source/Runtime/Engine/Private/GameInstance.cpp
--- a/Engine/Source/Runtime/Engine/Private/GameInstance.cpp
+++ b/Engine/Source/Runtime/Engine/Private/GameInstance.cpp
@@ -359,7 +359,7 @@ ULocalPlayer* UGameInstance::CreateLocalPlayer(int32 ControllerId, FString& OutE
 			UE_LOG(LogPlayerManagement, Warning, TEXT("Controller ID (%d) is unlikely to map to any physical device, so this player will not receive input"), ControllerId);
 		}
 
-		NewPlayer = CastChecked<ULocalPlayer>(StaticConstructObject(GetEngine()->LocalPlayerClass, GetEngine()));
+		NewPlayer = NewObject<ULocalPlayer>(GetEngine(), GetEngine()->LocalPlayerClass);
 		InsertIndex = AddLocalPlayer(NewPlayer, ControllerId);
 		if (bSpawnActor && InsertIndex != INDEX_NONE && GetWorld() != NULL)
 		{
diff --git a/Engine/Source/Runtime/Engine/Private/GameViewportClient.cpp b/Engine/Source/Runtime/Engine/Private/GameViewportClient.cpp
--- a/Engine/Source/Runtime/Engine/Private/GameViewportClient.cpp
+++ b/Engine/Source/Runtime/Engine/Private/GameViewportClient.cpp
@@ -1403,7 +1403,7 @@ ULocalPlayer* UGameViewportClient::SetupInitialLocalPlayer(FString& OutError)
 
 #if !UE_BUILD_SHIPPING
 	// Create the viewport's console.
-	ViewportConsole = Cast<UConsole>(StaticConstructObject(GetOuterUEngine()->ConsoleClass, this));
+	ViewportConsole = NewObject<UConsole>(this, GetOuterUEngine()->ConsoleClass);
 	// register console to get all log messages
 	GLog->AddOutputDevice(ViewportConsole);
 #endif // !UE_BUILD_SHIPPING
diff --git a/Engine/Source/Runtime/Engine/Private/ImageUtils.cpp b/Engine/Source/Runtime/Engine/Private/ImageUtils.cpp
--- a/Engine/Source/Runtime/Engine/Private/ImageUtils.cpp
+++ b/Engine/Source/Runtime/Engine/Private/ImageUtils.cpp
@@ -126,7 +126,7 @@ UTexture2D* FImageUtils::CreateTexture2D(int32 SrcWidth, int32 SrcHeight, const
 #if WITH_EDITOR
 	UTexture2D* Tex2D;
 
-	Tex2D = CastChecked<UTexture2D>( StaticConstructObject(UTexture2D::StaticClass(), Outer, FName(*Name), Flags) );
+	Tex2D = NewObject<UTexture2D>(Outer, FName(*Name), Flags);
 	Tex2D->Source.Init(SrcWidth, SrcHeight, /*NumSlices=*/ 1, /*NumMips=*/ 1, TSF_BGRA8);
 	
 	// Create base mip for the texture we created.
diff --git a/Engine/Source/Runtime/Engine/Private/LevelStreaming.cpp b/Engine/Source/Runtime/Engine/Private/LevelStreaming.cpp
--- a/Engine/Source/Runtime/Engine/Private/LevelStreaming.cpp
+++ b/Engine/Source/Runtime/Engine/Private/LevelStreaming.cpp
@@ -663,7 +663,7 @@ ULevelStreaming* ULevelStreaming::CreateInstance(FString InstanceUniqueName)
 				
 		if (bUniqueName)
 		{
-			StreamingLevelInstance = Cast<ULevelStreaming>(StaticConstructObject(GetClass(), InWorld, NAME_None, RF_Transient, NULL));
+			StreamingLevelInstance = NewObject<ULevelStreaming>(InWorld, GetClass(), NAME_None, RF_Transient, NULL);
 			// new level streaming instance will load the same map package as this object
 			StreamingLevelInstance->PackageNameToLoad = (PackageNameToLoad == NAME_None ? GetWorldAssetPackageFName() : PackageNameToLoad);
 			// under a provided unique name
diff --git a/Engine/Source/Runtime/Engine/Private/Materials/MaterialExpressions.cpp b/Engine/Source/Runtime/Engine/Private/Materials/MaterialExpressions.cpp
--- a/Engine/Source/Runtime/Engine/Private/Materials/MaterialExpressions.cpp
+++ b/Engine/Source/Runtime/Engine/Private/Materials/MaterialExpressions.cpp
@@ -6076,7 +6076,7 @@ UMaterial* UMaterialFunction::GetPreviewMaterial()
 {
 	if( NULL == PreviewMaterial )
 	{
-		PreviewMaterial = (UMaterial*)StaticConstructObject( UMaterial::StaticClass(), GetTransientPackage(), NAME_None, RF_Transient | RF_Public);
+		PreviewMaterial = NewObject<UMaterial>(GetTransientPackage(), NAME_None, RF_Transient | RF_Public);
 
 		PreviewMaterial->Expressions = FunctionExpressions;
 
diff --git a/Engine/Source/Runtime/Engine/Private/Particles/ParticleComponents.cpp b/Engine/Source/Runtime/Engine/Private/Particles/ParticleComponents.cpp
--- a/Engine/Source/Runtime/Engine/Private/Particles/ParticleComponents.cpp
+++ b/Engine/Source/Runtime/Engine/Private/Particles/ParticleComponents.cpp
@@ -1938,7 +1938,7 @@ void UParticleSpriteEmitter::SetToSensibleDefaults()
 		}
 		ColorCurveDist->bIsDirty = true;
 	}
-	ColorModule->AlphaOverLife.Distribution = Cast<UDistributionFloatConstantCurve>(StaticConstructObject(UDistributionFloatConstantCurve::StaticClass(), ColorModule));
+	ColorModule->AlphaOverLife.Distribution = NewObject<UDistributionFloatConstantCurve>(ColorModule);
 	UDistributionFloatConstantCurve* AlphaCurveDist = Cast<UDistributionFloatConstantCurve>(ColorModule->AlphaOverLife.Distribution);
 	if (AlphaCurveDist)
 	{
diff --git a/Engine/Source/Runtime/Engine/Private/Particles/ParticleModules.cpp b/Engine/Source/Runtime/Engine/Private/Particles/ParticleModules.cpp
--- a/Engine/Source/Runtime/Engine/Private/Particles/ParticleModules.cpp
+++ b/Engine/Source/Runtime/Engine/Private/Particles/ParticleModules.cpp
@@ -826,7 +826,7 @@ private:
 			}
 			else
 			{
-				New = StaticConstructObject(Obj->GetClass(), Outer, Obj->GetFName(), RF_NoFlags, Obj);
+				New = NewObject<UObject>(Outer, Obj->GetClass(), Obj->GetFName(), RF_NoFlags, Obj);
 				UE_LOG(LogParticles, Verbose, TEXT("      Replacing with New %s"), *GetFullNameSafe(New));
 			}
 			Obj = New;
@@ -1624,7 +1624,7 @@ void UParticleModuleRotationRate::SpawnEx(FParticleEmitterInstance* Owner, int32
 
 void UParticleModuleRotationRate::SetToSensibleDefaults(UParticleEmitter* Owner)
 {
-	StartRotationRate.Distribution = Cast<UDistributionFloatUniform>(StaticConstructObject(UDistributionFloatUniform::StaticClass(), this));
+	StartRotationRate.Distribution = NewObject<UDistributionFloatUniform>(this);
 	UDistributionFloatUniform* StartRotationRateDist = Cast<UDistributionFloatUniform>(StartRotationRate.Distribution);
 	if (StartRotationRateDist)
 	{
@@ -1931,7 +1931,7 @@ float UParticleModuleSubUV::DetermineImageIndex(FParticleEmitterInstance* Owner,
 
 void UParticleModuleSubUV::SetToSensibleDefaults(UParticleEmitter* Owner)
 {
-	SubImageIndex.Distribution = Cast<UDistributionFloatConstantCurve>(StaticConstructObject(UDistributionFloatConstantCurve::StaticClass(), this));
+	SubImageIndex.Distribution = NewObject<UDistributionFloatConstantCurve>(this);
 	UDistributionFloatConstantCurve* SubImageIndexDist = Cast<UDistributionFloatConstantCurve>(SubImageIndex.Distribution);
 	if (SubImageIndexDist)
 	{
@@ -2166,7 +2166,7 @@ void UParticleModuleRotationRateMultiplyLife::Update(FParticleEmitterInstance* O
 
 void UParticleModuleRotationRateMultiplyLife::SetToSensibleDefaults(UParticleEmitter* Owner)
 {
-	LifeMultiplier.Distribution = Cast<UDistributionFloatConstantCurve>(StaticConstructObject(UDistributionFloatConstantCurve::StaticClass(), this));
+	LifeMultiplier.Distribution = NewObject<UDistributionFloatConstantCurve>(this);
 	UDistributionFloatConstantCurve* LifeMultiplierDist = Cast<UDistributionFloatConstantCurve>(LifeMultiplier.Distribution);
 	if (LifeMultiplierDist)
 	{
diff --git a/Engine/Source/Runtime/Engine/Private/Particles/ParticleModules_Color.cpp b/Engine/Source/Runtime/Engine/Private/Particles/ParticleModules_Color.cpp
--- a/Engine/Source/Runtime/Engine/Private/Particles/ParticleModules_Color.cpp
+++ b/Engine/Source/Runtime/Engine/Private/Particles/ParticleModules_Color.cpp
@@ -415,7 +415,7 @@ void UParticleModuleColorOverLife::Update(FParticleEmitterInstance* Owner, int32
 
 void UParticleModuleColorOverLife::SetToSensibleDefaults(UParticleEmitter* Owner)
 {
-	ColorOverLife.Distribution = Cast<UDistributionVectorConstantCurve>(StaticConstructObject(UDistributionVectorConstantCurve::StaticClass(), this));
+	ColorOverLife.Distribution = NewObject<UDistributionVectorConstantCurve>(this);
 	UDistributionVectorConstantCurve* ColorOverLifeDist = Cast<UDistributionVectorConstantCurve>(ColorOverLife.Distribution);
 	if (ColorOverLifeDist)
 	{
@@ -438,7 +438,7 @@ void UParticleModuleColorOverLife::SetToSensibleDefaults(UParticleEmitter* Owner
 		ColorOverLifeDist->bIsDirty = true;
 	}
 
-	AlphaOverLife.Distribution = Cast<UDistributionFloatConstantCurve>(StaticConstructObject(UDistributionFloatConstantCurve::StaticClass(), this));
+	AlphaOverLife.Distribution = NewObject<UDistributionFloatConstantCurve>(this);
 	UDistributionFloatConstantCurve* AlphaOverLifeDist = Cast<UDistributionFloatConstantCurve>(AlphaOverLife.Distribution);
 	if (AlphaOverLifeDist)
 	{
@@ -629,7 +629,7 @@ void UParticleModuleColorScaleOverLife::Update(FParticleEmitterInstance* Owner,
 
 void UParticleModuleColorScaleOverLife::SetToSensibleDefaults(UParticleEmitter* Owner)
 {
-	ColorScaleOverLife.Distribution = Cast<UDistributionVectorConstantCurve>(StaticConstructObject(UDistributionVectorConstantCurve::StaticClass(), this));
+	ColorScaleOverLife.Distribution = NewObject<UDistributionVectorConstantCurve>(this);
 	UDistributionVectorConstantCurve* ColorScaleOverLifeDist = Cast<UDistributionVectorConstantCurve>(ColorScaleOverLife.Distribution);
 	if (ColorScaleOverLifeDist)
 	{
@@ -669,7 +669,7 @@ bool UParticleModuleColorScaleOverLife::PerformCustomMenuEntry(int32 InEntryInde
 		if (InEntryIndex == 0)
 		{
 			UE_LOG(LogParticles, Log, TEXT("Setup color scale over life for particle param!"));
-			ColorScaleOverLife.Distribution = Cast<UDistributionVectorParticleParameter>(StaticConstructObject(UDistributionVectorParticleParameter::StaticClass(), this));
+			ColorScaleOverLife.Distribution = NewObject<UDistributionVectorParticleParameter>(this);
 			UDistributionVectorParticleParameter* ColorDist = Cast<UDistributionVectorParticleParameter>(ColorScaleOverLife.Distribution);
 			if (ColorDist)
 			{
@@ -681,7 +681,7 @@ bool UParticleModuleColorScaleOverLife::PerformCustomMenuEntry(int32 InEntryInde
 				ColorDist->bIsDirty = true;
 			}
 
-			AlphaScaleOverLife.Distribution = Cast<UDistributionFloatParticleParameter>(StaticConstructObject(UDistributionFloatParticleParameter::StaticClass(), this));
+			AlphaScaleOverLife.Distribution = NewObject<UDistributionFloatParticleParameter>(this);
 			UDistributionFloatParticleParameter* AlphaDist = Cast<UDistributionFloatParticleParameter>(AlphaScaleOverLife.Distribution);
 			if (AlphaDist)
 			{
diff --git a/Engine/Source/Runtime/Engine/Private/Particles/ParticleModules_Size.cpp b/Engine/Source/Runtime/Engine/Private/Particles/ParticleModules_Size.cpp
--- a/Engine/Source/Runtime/Engine/Private/Particles/ParticleModules_Size.cpp
+++ b/Engine/Source/Runtime/Engine/Private/Particles/ParticleModules_Size.cpp
@@ -301,7 +301,7 @@ void UParticleModuleSizeMultiplyLife::Update(FParticleEmitterInstance* Owner, in
 
 void UParticleModuleSizeMultiplyLife::SetToSensibleDefaults(UParticleEmitter* Owner)
 {
-	LifeMultiplier.Distribution = Cast<UDistributionVectorConstantCurve>(StaticConstructObject(UDistributionVectorConstantCurve::StaticClass(), this));
+	LifeMultiplier.Distribution = NewObject<UDistributionVectorConstantCurve>(this);
 	UDistributionVectorConstantCurve* LifeMultiplierDist = Cast<UDistributionVectorConstantCurve>(LifeMultiplier.Distribution);
 	if (LifeMultiplierDist)
 	{
diff --git a/Engine/Source/Runtime/Engine/Private/PlayerCameraManager.cpp b/Engine/Source/Runtime/Engine/Private/PlayerCameraManager.cpp
--- a/Engine/Source/Runtime/Engine/Private/PlayerCameraManager.cpp
+++ b/Engine/Source/Runtime/Engine/Private/PlayerCameraManager.cpp
@@ -604,7 +604,7 @@ void APlayerCameraManager::ApplyAudioFade()
 
 UCameraModifier* APlayerCameraManager::CreateCameraModifier(TSubclassOf<UCameraModifier> ModifierClass)
 {
-	UCameraModifier* NewMod = Cast<UCameraModifier>(StaticConstructObject(ModifierClass, this));
+	UCameraModifier* NewMod = NewObject<UCameraModifier>(this, ModifierClass);
 	NewMod->Init(this);
 	return NewMod;
 }
@@ -622,7 +622,7 @@ void APlayerCameraManager::PostInitializeComponents()
 	// create CameraAnimInsts in pool
 	for (int32 Idx=0; Idx<MAX_ACTIVE_CAMERA_ANIMS; ++Idx)
 	{
-		AnimInstPool[Idx] = Cast<UCameraAnimInst>(StaticConstructObject(UCameraAnimInst::StaticClass(), this));
+		AnimInstPool[Idx] = NewObject<UCameraAnimInst>(this);
 
 		// add everything to the free list initially
 		FreeAnims.Add(AnimInstPool[Idx]);
diff --git a/Engine/Source/Runtime/Engine/Private/PlayerController.cpp b/Engine/Source/Runtime/Engine/Private/PlayerController.cpp
--- a/Engine/Source/Runtime/Engine/Private/PlayerController.cpp
+++ b/Engine/Source/Runtime/Engine/Private/PlayerController.cpp
@@ -992,7 +992,7 @@ void APlayerController::AddCheats(bool bForce)
 	// Assuming that this never gets called for NM_Client without bForce=true
 	if ( ((CheatManager == NULL) && (World->GetAuthGameMode() != NULL) && World->GetAuthGameMode()->AllowCheats(this)) || bForce)
 	{
-		CheatManager = Cast<UCheatManager>(StaticConstructObject(CheatClass, this));
+		CheatManager = NewObject<UCheatManager>(this, CheatClass);
 		CheatManager->InitCheatManager();
 	}
 }
diff --git a/Engine/Source/Runtime/Engine/Private/SceneComponent.cpp b/Engine/Source/Runtime/Engine/Private/SceneComponent.cpp
--- a/Engine/Source/Runtime/Engine/Private/SceneComponent.cpp
+++ b/Engine/Source/Runtime/Engine/Private/SceneComponent.cpp
@@ -124,7 +124,7 @@ void USceneComponent::OnRegister()
 	if (bVisualizeComponent && SpriteComponent == nullptr && GetOwner() && !GetWorld()->IsGameWorld() )
 	{
 		// Create a new billboard component to serve as a visualization of the actor until there is another primitive component
-		SpriteComponent = ConstructObject<UBillboardComponent>(UBillboardComponent::StaticClass(), GetOwner(), NAME_None, RF_Transactional | RF_TextExportTransient);
+		SpriteComponent = NewObject<UBillboardComponent>(GetOwner(), NAME_None, RF_Transactional | RF_TextExportTransient);
 
 		SpriteComponent->Sprite = LoadObject<UTexture2D>(nullptr, TEXT("/Engine/EditorResources/EmptyActor.EmptyActor"));
 		SpriteComponent->RelativeScale3D = FVector(0.5f, 0.5f, 0.5f);
diff --git a/Engine/Source/Runtime/Engine/Private/ScriptPlatformInterface.cpp b/Engine/Source/Runtime/Engine/Private/ScriptPlatformInterface.cpp
--- a/Engine/Source/Runtime/Engine/Private/ScriptPlatformInterface.cpp
+++ b/Engine/Source/Runtime/Engine/Private/ScriptPlatformInterface.cpp
@@ -388,7 +388,7 @@ UObject* UCloudStorageBase::ParseDocumentAsObject(int32 Index, TSubclassOf<class
 
 	// NOTE: The following should be in shared functionality in UCloudStorageBase
 	// create the object
-	UObject* Obj = StaticConstructObject(ObjectClass);
+	UObject* Obj = NewObject<UObject>(GetTransientPackage(), ObjectClass);
 
 	// serialize the object
 	Obj->Serialize(Ar);
diff --git a/Engine/Source/Runtime/Engine/Private/Tests/Internationalization/TextPropertyTest.cpp b/Engine/Source/Runtime/Engine/Private/Tests/Internationalization/TextPropertyTest.cpp
--- a/Engine/Source/Runtime/Engine/Private/Tests/Internationalization/TextPropertyTest.cpp
+++ b/Engine/Source/Runtime/Engine/Private/Tests/Internationalization/TextPropertyTest.cpp
@@ -21,12 +21,12 @@ bool FTextPropertyTest::RunTest (const FString& Parameters)
 	UTextPropertyTestObject* const TextPropertyTestCDO = Cast<UTextPropertyTestObject>( TextPropertyTestObjectClass->ClassDefaultObject );
 
 	{
-		UTextPropertyTestObject* NewObject = Cast<UTextPropertyTestObject>( StaticConstructObject( TextPropertyTestObjectClass ) );
+		UTextPropertyTestObject* NewUObject = NewObject<UTextPropertyTestObject>(GetTransientPackage(), TextPropertyTestObjectClass);
 
 		// Test Identical - Newly constructed object properties should be identical to class default object properties.
-		if(	(DefaultedTextProperty->Identical(&(NewObject->DefaultedText), &(TextPropertyTestCDO->DefaultedText), 0) != true)
+		if ((DefaultedTextProperty->Identical(&(NewUObject->DefaultedText), &(TextPropertyTestCDO->DefaultedText), 0) != true)
 			||
-			(UndefaultedTextProperty->Identical(&(NewObject->UndefaultedText), &(TextPropertyTestCDO->UndefaultedText), 0) != true) )
+			(UndefaultedTextProperty->Identical(&(NewUObject->UndefaultedText), &(TextPropertyTestCDO->UndefaultedText), 0) != true))
 		{
 			AddError(TEXT("UTextProperty::Identical failed to return true comparing a newly constructed object and the class default object."));
 		}
@@ -34,8 +34,8 @@ bool FTextPropertyTest::RunTest (const FString& Parameters)
 		// Test ExportText - Export text should provide the localized form of the text.
 		{
 			FString ExportedStringValue;
-			DefaultedTextProperty->ExportTextItem(ExportedStringValue, &(NewObject->DefaultedText), NULL, NULL, 0, NULL);
-			if( ExportedStringValue != NewObject->DefaultedText.ToString() )
+			DefaultedTextProperty->ExportTextItem(ExportedStringValue, &(NewUObject->DefaultedText), NULL, NULL, 0, NULL);
+			if (ExportedStringValue != NewUObject->DefaultedText.ToString())
 			{
 				AddError(TEXT("UTextProperty::ExportTextItem failed to provide the display string."));
 			}
@@ -44,8 +44,8 @@ bool FTextPropertyTest::RunTest (const FString& Parameters)
 		// Test ImportText - Import text should set the source string to the input string.
 		{
 			FString ImportedStringValue = TEXT("ImportValue");
-			DefaultedTextProperty->ImportText(*ImportedStringValue, &(NewObject->DefaultedText), 0, NULL);
-			const FString* const SourceString = FTextInspector::GetSourceString(NewObject->DefaultedText);
+			DefaultedTextProperty->ImportText(*ImportedStringValue, &(NewUObject->DefaultedText), 0, NULL);
+			const FString* const SourceString = FTextInspector::GetSourceString(NewUObject->DefaultedText);
 			if( !SourceString || ImportedStringValue != *SourceString )
 			{
 				AddError(TEXT("UTextProperty::ImportText failed to alter the source string to the provided value."));
@@ -55,14 +55,14 @@ bool FTextPropertyTest::RunTest (const FString& Parameters)
 
 	// Test Identical - Altered text properties should not be identical to class default object properties.
 	{
-		UTextPropertyTestObject* NewObject = Cast<UTextPropertyTestObject>( StaticConstructObject( TextPropertyTestObjectClass ) );
+		UTextPropertyTestObject* NewUObject = NewObject<UTextPropertyTestObject>(GetTransientPackage(), TextPropertyTestObjectClass);
 
-		NewObject->DefaultedText = LOCTEXT("ModifiedDefaultedText", "Modified DefaultedText Value");
-		NewObject->UndefaultedText = LOCTEXT("ModifiedUndefaultedText", "Modified UndefaultedText Value");
+		NewUObject->DefaultedText = LOCTEXT("ModifiedDefaultedText", "Modified DefaultedText Value");
+		NewUObject->UndefaultedText = LOCTEXT("ModifiedUndefaultedText", "Modified UndefaultedText Value");
 		if( 
-			DefaultedTextProperty->Identical(&(NewObject->DefaultedText), &(TextPropertyTestCDO->DefaultedText), 0)
+			DefaultedTextProperty->Identical(&(NewUObject->DefaultedText), &(TextPropertyTestCDO->DefaultedText), 0)
 			||
-			UndefaultedTextProperty->Identical(&(NewObject->UndefaultedText), &(TextPropertyTestCDO->UndefaultedText), 0)
+			UndefaultedTextProperty->Identical(&(NewUObject->UndefaultedText), &(TextPropertyTestCDO->UndefaultedText), 0)
 		  )
 		{
 			AddError(TEXT("UTextProperty::Identical failed to return false comparing a modified object and the class default object."));
@@ -72,7 +72,7 @@ bool FTextPropertyTest::RunTest (const FString& Parameters)
 	{
 		TArray<uint8> BackingStore;
 		
-		UTextPropertyTestObject* SavedObject = Cast<UTextPropertyTestObject>(StaticConstructObject(UTextPropertyTestObject::StaticClass()));
+		auto SavedObject = NewObject<UTextPropertyTestObject>();
 
 		FText::FindText( TEXT("TextPropertyTest"), TEXT("DefaultedText"), /*OUT*/SavedObject->DefaultedText );
 		SavedObject->UndefaultedText = LOCTEXT("ModifiedUndefaultedText", "Modified UndefaultedText Value");
@@ -91,7 +91,7 @@ bool FTextPropertyTest::RunTest (const FString& Parameters)
 			SavedObject->Serialize(MemoryWriter);
 		}
 
-		UTextPropertyTestObject* LoadedObject = Cast<UTextPropertyTestObject>(StaticConstructObject(UTextPropertyTestObject::StaticClass()));
+		auto LoadedObject = NewObject<UTextPropertyTestObject>();
 
 		// Load.
 		{
diff --git a/Engine/Source/Runtime/Engine/Private/Texture2DDynamic.cpp b/Engine/Source/Runtime/Engine/Private/Texture2DDynamic.cpp
--- a/Engine/Source/Runtime/Engine/Private/Texture2DDynamic.cpp
+++ b/Engine/Source/Runtime/Engine/Private/Texture2DDynamic.cpp
@@ -121,7 +121,7 @@ UTexture2DDynamic* UTexture2DDynamic::Create(int32 InSizeX, int32 InSizeY, EPixe
 	if (InSizeX > 0 && InSizeY > 0 )
 	{
 		
-		UTexture2DDynamic* NewTexture = Cast<UTexture2DDynamic>(StaticConstructObject(UTexture2DDynamic::StaticClass(), GetTransientPackage(), NAME_None, RF_Transient));
+		auto NewTexture = NewObject<UTexture2DDynamic>(GetTransientPackage(), NAME_None, RF_Transient);
 		if (NewTexture != NULL)
 		{
 			// Disable compression
diff --git a/Engine/Source/Runtime/Engine/Private/TextureRenderTarget2D.cpp b/Engine/Source/Runtime/Engine/Private/TextureRenderTarget2D.cpp
--- a/Engine/Source/Runtime/Engine/Private/TextureRenderTarget2D.cpp
+++ b/Engine/Source/Runtime/Engine/Private/TextureRenderTarget2D.cpp
@@ -158,8 +158,7 @@ UTexture2D* UTextureRenderTarget2D::ConstructTexture2D(UObject* Outer, const FSt
 	}
 
 	// create the 2d texture
-	Result = CastChecked<UTexture2D>(
-		StaticConstructObject(UTexture2D::StaticClass(), Outer, FName(*NewTexName), ObjectFlags));
+	Result = NewObject<UTexture2D>(Outer, FName(*NewTexName), ObjectFlags);
 	// init to the same size as the 2d texture
 	Result->Source.Init(SizeX, SizeY, 1, 1, TextureFormat);
 
diff --git a/Engine/Source/Runtime/Engine/Private/TextureRenderTargetCube.cpp b/Engine/Source/Runtime/Engine/Private/TextureRenderTargetCube.cpp
--- a/Engine/Source/Runtime/Engine/Private/TextureRenderTargetCube.cpp
+++ b/Engine/Source/Runtime/Engine/Private/TextureRenderTargetCube.cpp
@@ -134,7 +134,7 @@ UTextureCube* UTextureRenderTargetCube::ConstructTextureCube(
 		if (CubeResource && TextureFormat != TSF_Invalid)
 		{
 			// create the cube texture
-			Result = CastChecked<UTextureCube>(StaticConstructObject(UTextureCube::StaticClass(), ObjOuter, FName(*NewTexName), InFlags));
+			Result = NewObject<UTextureCube>(ObjOuter, FName(*NewTexName), InFlags);
 
 			bool bSRGB = true;
 			// if render target gamma used was 1.0 then disable SRGB for the static texture
diff --git a/Engine/Source/Runtime/Engine/Private/UnrealEngine.cpp b/Engine/Source/Runtime/Engine/Private/UnrealEngine.cpp
--- a/Engine/Source/Runtime/Engine/Private/UnrealEngine.cpp
+++ b/Engine/Source/Runtime/Engine/Private/UnrealEngine.cpp
@@ -2631,7 +2631,7 @@ bool UEngine::HandleMergeMeshCommand( const TCHAR* Cmd, FOutputDevice& Ar, UWorl
 	if (PlayerPawn && SourceMeshList.Num() >= 2)
 	{
 		// create the composite mesh
-		USkeletalMesh* CompositeMesh = CastChecked<USkeletalMesh>(StaticConstructObject(USkeletalMesh::StaticClass(), GetTransientPackage(), NAME_None, RF_Transient));
+		auto CompositeMesh = NewObject<USkeletalMesh>(GetTransientPackage(), NAME_None, RF_Transient);
 
 		TArray<FSkelMeshMergeSectionMapping> InForceSectionMapping;
 		// create an instance of the FSkeletalMeshMerge utility
diff --git a/Engine/Source/Runtime/Engine/Private/UserInterface/Canvas.cpp b/Engine/Source/Runtime/Engine/Private/UserInterface/Canvas.cpp
--- a/Engine/Source/Runtime/Engine/Private/UserInterface/Canvas.cpp
+++ b/Engine/Source/Runtime/Engine/Private/UserInterface/Canvas.cpp
@@ -1873,7 +1873,7 @@ TWeakObjectPtr<class UReporterGraph> UCanvas::GetReporterGraph()
 {
 	if (!ReporterGraph)
 	{
-		ReporterGraph = Cast<UReporterGraph>(StaticConstructObject(UReporterGraph::StaticClass(), this));
+		ReporterGraph = NewObject<UReporterGraph>(this);
 	}
 
 	return ReporterGraph;
diff --git a/Engine/Source/Runtime/Engine/Private/World.cpp b/Engine/Source/Runtime/Engine/Private/World.cpp
--- a/Engine/Source/Runtime/Engine/Private/World.cpp
+++ b/Engine/Source/Runtime/Engine/Private/World.cpp
@@ -641,7 +641,7 @@ void UWorld::AddParameterCollectionInstance(UMaterialParameterCollection* Collec
 		}
 	}
 
-	UMaterialParameterCollectionInstance* NewInstance = (UMaterialParameterCollectionInstance*)StaticConstructObject(UMaterialParameterCollectionInstance::StaticClass());
+	auto NewInstance = NewObject<UMaterialParameterCollectionInstance>();
 	NewInstance->SetCollection(Collection, this);
 
 	if (ExistingIndex != INDEX_NONE)
diff --git a/Engine/Source/Runtime/Engine/Private/WorldComposition.cpp b/Engine/Source/Runtime/Engine/Private/WorldComposition.cpp
--- a/Engine/Source/Runtime/Engine/Private/WorldComposition.cpp
+++ b/Engine/Source/Runtime/Engine/Private/WorldComposition.cpp
@@ -291,7 +291,7 @@ ULevelStreaming* UWorldComposition::CreateStreamingLevel(const FWorldComposition
 {
 	UWorld* OwningWorld = GetWorld();
 	UClass* StreamingClass = ULevelStreamingKismet::StaticClass();
-	ULevelStreaming* StreamingLevel = Cast<ULevelStreaming>(StaticConstructObject(StreamingClass, OwningWorld, NAME_None, RF_Transient, NULL));
+	auto StreamingLevel = NewObject<ULevelStreaming>(OwningWorld, StreamingClass, NAME_None, RF_Transient, NULL);
 		
 	// Associate a package name.
 	StreamingLevel->SetWorldAssetByPackageName(InTile.PackageName);
diff --git a/Engine/Source/Runtime/GameplayAbilities/Private/GameplayEffectExtension_LifestealTest.cpp b/Engine/Source/Runtime/GameplayAbilities/Private/GameplayEffectExtension_LifestealTest.cpp
--- a/Engine/Source/Runtime/GameplayAbilities/Private/GameplayEffectExtension_LifestealTest.cpp
+++ b/Engine/Source/Runtime/GameplayAbilities/Private/GameplayEffectExtension_LifestealTest.cpp
@@ -36,7 +36,7 @@ void UGameplayEffectExtension_LifestealTest::PostGameplayEffectExecute(const FGa
 			UProperty *HealthProperty = FindFieldChecked<UProperty>(UAbilitySystemTestAttributeSet::StaticClass(), GET_MEMBER_NAME_CHECKED(UAbilitySystemTestAttributeSet, Health));
 
 			// Since this is a test class and we don't want to tie it any actual content assets, just construct a GameplayEffect here.
-			LocalHealthRestore = Cast<UGameplayEffect>(StaticConstructObject(UGameplayEffect::StaticClass(), GetTransientPackage(), FName(TEXT("LifestealHealthRestore"))));
+			LocalHealthRestore = NewObject<UGameplayEffect>(GetTransientPackage(), FName(TEXT("LifestealHealthRestore")));
 			LocalHealthRestore->Modifiers.SetNum(1);
 			LocalHealthRestore->Modifiers[0].Magnitude.SetValue(HealthToRestore);
 			LocalHealthRestore->Modifiers[0].ModifierOp = EGameplayModOp::Additive;
diff --git a/Engine/Source/Runtime/GameplayAbilities/Private/GameplayEffectTests.cpp b/Engine/Source/Runtime/GameplayAbilities/Private/GameplayEffectTests.cpp
--- a/Engine/Source/Runtime/GameplayAbilities/Private/GameplayEffectTests.cpp
+++ b/Engine/Source/Runtime/GameplayAbilities/Private/GameplayEffectTests.cpp
@@ -42,7 +42,7 @@ static UDataTable* CreateGameplayDataTable()
 	CSV.Append(TEXT("\r\n22,SpellDmg.Buff"));
 	CSV.Append(TEXT("\r\n23,GameplayCue.Burning"));
 
-	UDataTable * DataTable = Cast<UDataTable>(StaticConstructObject(UDataTable::StaticClass(), GetTransientPackage(), FName(TEXT("TempDataTable"))));
+	auto DataTable = NewObject<UDataTable>(GetTransientPackage(), FName(TEXT("TempDataTable")));
 	DataTable->RowStruct = FGameplayTagTableRow::StaticStruct();
 	DataTable->CreateTableFromCSVString(CSV);
 
@@ -55,7 +55,7 @@ static UDataTable* CreateGameplayDataTable()
 }
 
 #define GET_FIELD_CHECKED(Class, Field) FindFieldChecked<UProperty>(Class::StaticClass(), GET_MEMBER_NAME_CHECKED(Class, Field))
-#define CONSTRUCT_CLASS(Class, Name) Class * Name = Cast<Class>(StaticConstructObject(Class::StaticClass(), GetTransientPackage(), FName(TEXT(#Name))))
+#define CONSTRUCT_CLASS(Class, Name) Class* Name = NewObject<Class>(GetTransientPackage(), FName(TEXT(#Name)))
 
 class GameplayEffectsTestSuite
 {
diff --git a/Engine/Source/Runtime/GameplayTags/Private/Tests/GameplayTagTests.cpp b/Engine/Source/Runtime/GameplayTags/Private/Tests/GameplayTagTests.cpp
--- a/Engine/Source/Runtime/GameplayTags/Private/Tests/GameplayTagTests.cpp
+++ b/Engine/Source/Runtime/GameplayTags/Private/Tests/GameplayTagTests.cpp
@@ -9,7 +9,7 @@ static UDataTable* CreateGameplayDataTable()
 {
 	FString CSV(TEXT(",Tag,CategoryText,\r\n0,GameplayTagTest.Test1\r\n1,GameplayTagTest.Test2"));
 
-	UDataTable * DataTable = Cast<UDataTable>(StaticConstructObject(UDataTable::StaticClass(), GetTransientPackage(), FName(TEXT("TempDataTable"))));
+	auto DataTable = NewObject<UDataTable>(GetTransientPackage(), FName(TEXT("TempDataTable")));
 	DataTable->RowStruct = FGameplayTagTableRow::StaticStruct();
 	DataTable->CreateTableFromCSVString(CSV);
 
