[INTEGRATE] Change 2429878 by [REDACTED_PERSON][REDACTED_HANDLE] on 2015/02/03 09:38:44

	Fixed issue with static meshes not applying materials correctly when dropped.

diff --git a/Engine/Source/Editor/PlacementMode/Private/SPlacementModeTools.cpp b/Engine/Source/Editor/PlacementMode/Private/SPlacementModeTools.cpp
--- a/Engine/Source/Editor/PlacementMode/Private/SPlacementModeTools.cpp
+++ b/Engine/Source/Editor/PlacementMode/Private/SPlacementModeTools.cpp
@@ -695,25 +695,25 @@ TSharedRef< SWidget > SPlacementModeTools::BuildBasicWidget()
 	.AutoHeight()
 	[
 		// Cube
-		BuildDraggableAssetWidget(UActorFactoryBasicShape::StaticClass(), FAssetData( LoadObject<UStaticMesh>(nullptr,TEXT("/Engine/BasicShapes/Cube.Cube"))), FName("ClassThumbnail.Cube"), BasicShapeColorOverride )
+		BuildDraggableAssetWidget(UActorFactoryBasicShape::StaticClass(), FAssetData( LoadObject<UStaticMesh>(nullptr,*UActorFactoryBasicShape::BasicCube.ToString())), FName("ClassThumbnail.Cube"), BasicShapeColorOverride )
 	]
 	+ SVerticalBox::Slot()
 	.AutoHeight()
 	[
 		// Sphere
-		BuildDraggableAssetWidget(UActorFactoryBasicShape::StaticClass(), FAssetData(LoadObject<UStaticMesh>(nullptr,TEXT("/Engine/BasicShapes/Sphere.Sphere"))), FName("ClassThumbnail.Sphere"), BasicShapeColorOverride )
+		BuildDraggableAssetWidget(UActorFactoryBasicShape::StaticClass(), FAssetData(LoadObject<UStaticMesh>(nullptr,*UActorFactoryBasicShape::BasicSphere.ToString())), FName("ClassThumbnail.Sphere"), BasicShapeColorOverride )
 	]
 	+ SVerticalBox::Slot()
 	.AutoHeight()
 	[
 		// Cylinder
-		BuildDraggableAssetWidget(UActorFactoryBasicShape::StaticClass(), FAssetData(LoadObject<UStaticMesh>(nullptr,TEXT("/Engine/BasicShapes/Cylinder.Cylinder"))), FName("ClassThumbnail.Cylinder"), BasicShapeColorOverride )
+		BuildDraggableAssetWidget(UActorFactoryBasicShape::StaticClass(), FAssetData(LoadObject<UStaticMesh>(nullptr,*UActorFactoryBasicShape::BasicCylinder.ToString())), FName("ClassThumbnail.Cylinder"), BasicShapeColorOverride )
 	]
 	+ SVerticalBox::Slot()
 	.AutoHeight()
 	[
 		// Cone
-		BuildDraggableAssetWidget(UActorFactoryBasicShape::StaticClass(), FAssetData(LoadObject<UStaticMesh>(nullptr,TEXT("/Engine/BasicShapes/Cone.Cone"))), FName("ClassThumbnail.Cone"), BasicShapeColorOverride )
+		BuildDraggableAssetWidget(UActorFactoryBasicShape::StaticClass(), FAssetData(LoadObject<UStaticMesh>(nullptr,*UActorFactoryBasicShape::BasicCone.ToString())), FName("ClassThumbnail.Cone"), BasicShapeColorOverride )
 	]
 	+ SVerticalBox::Slot()
 	.AutoHeight()
diff --git a/Engine/Source/Editor/UnrealEd/Classes/ActorFactories/ActorFactoryBasicShape.h b/Engine/Source/Editor/UnrealEd/Classes/ActorFactories/ActorFactoryBasicShape.h
--- a/Engine/Source/Editor/UnrealEd/Classes/ActorFactories/ActorFactoryBasicShape.h
+++ b/Engine/Source/Editor/UnrealEd/Classes/ActorFactories/ActorFactoryBasicShape.h
@@ -9,5 +9,11 @@ class UActorFactoryBasicShape : public UActorFactoryStaticMesh
 {
 	GENERATED_UCLASS_BODY()
 
+	virtual bool CanCreateActorFrom( const FAssetData& AssetData, FText& OutErrorMsg ) override;
 	virtual void PostSpawnActor(UObject* Asset, AActor* NewActor) override;
+
+	UNREALED_API static const FName BasicCube;
+	UNREALED_API static const FName BasicSphere;
+	UNREALED_API static const FName BasicCylinder;
+	UNREALED_API static const FName BasicCone;
 };
diff --git a/Engine/Source/Editor/UnrealEd/Private/ComponentTypeRegistry.cpp b/Engine/Source/Editor/UnrealEd/Private/ComponentTypeRegistry.cpp
--- a/Engine/Source/Editor/UnrealEd/Private/ComponentTypeRegistry.cpp
+++ b/Engine/Source/Editor/UnrealEd/Private/ComponentTypeRegistry.cpp
@@ -45,7 +45,7 @@ static void AddBasicShapeComponents(TArray<FComponentClassComboEntryPtr>& Sorted
 
 	{
 		FComponentEntryCustomizationArgs Args;
-		Args.AssetOverride = LoadObject<UStaticMesh>(nullptr, TEXT("/Engine/BasicShapes/Cube.Cube"));;
+		Args.AssetOverride = LoadObject<UStaticMesh>(nullptr, *UActorFactoryBasicShape::BasicCube.ToString());
 		Args.OnComponentCreated = FOnComponentCreated::CreateStatic(OnBasicShapeCreated);
 		Args.ComponentNameOverride = LOCTEXT("BasicCubeShapeDisplayName", "Cube").ToString();
 		Args.IconOverrideBrushName = FName("ClassIcon.Cube");
@@ -56,7 +56,7 @@ static void AddBasicShapeComponents(TArray<FComponentClassComboEntryPtr>& Sorted
 
 	{
 		FComponentEntryCustomizationArgs Args;
-		Args.AssetOverride = LoadObject<UStaticMesh>(nullptr, TEXT("/Engine/BasicShapes/Sphere.Sphere"));;
+		Args.AssetOverride = LoadObject<UStaticMesh>(nullptr, *UActorFactoryBasicShape::BasicSphere.ToString());
 		Args.OnComponentCreated = FOnComponentCreated::CreateStatic(OnBasicShapeCreated);
 		Args.ComponentNameOverride = LOCTEXT("BasicSphereShapeDisplayName", "Sphere").ToString();
 		Args.IconOverrideBrushName = FName("ClassIcon.Sphere");
@@ -67,7 +67,7 @@ static void AddBasicShapeComponents(TArray<FComponentClassComboEntryPtr>& Sorted
 
 	{
 		FComponentEntryCustomizationArgs Args;
-		Args.AssetOverride = LoadObject<UStaticMesh>(nullptr, TEXT("/Engine/BasicShapes/Cylinder.Cylinder"));;
+		Args.AssetOverride = LoadObject<UStaticMesh>(nullptr, *UActorFactoryBasicShape::BasicCylinder.ToString());
 		Args.OnComponentCreated = FOnComponentCreated::CreateStatic(OnBasicShapeCreated);
 		Args.ComponentNameOverride = LOCTEXT("BasicCylinderShapeDisplayName", "Cylinder").ToString();
 		Args.IconOverrideBrushName = FName("ClassIcon.Cylinder");
@@ -78,7 +78,7 @@ static void AddBasicShapeComponents(TArray<FComponentClassComboEntryPtr>& Sorted
 
 	{
 		FComponentEntryCustomizationArgs Args;
-		Args.AssetOverride = LoadObject<UStaticMesh>(nullptr, TEXT("/Engine/BasicShapes/Cone.Cone"));;
+		Args.AssetOverride = LoadObject<UStaticMesh>(nullptr, *UActorFactoryBasicShape::BasicCone.ToString());
 		Args.OnComponentCreated = FOnComponentCreated::CreateStatic(OnBasicShapeCreated);
 		Args.ComponentNameOverride = LOCTEXT("BasicConeShapeDisplayName", "Cone").ToString();
 		Args.IconOverrideBrushName = FName("ClassIcon.Cone");
diff --git a/Engine/Source/Editor/UnrealEd/Private/Factories/ActorFactory.cpp b/Engine/Source/Editor/UnrealEd/Private/Factories/ActorFactory.cpp
--- a/Engine/Source/Editor/UnrealEd/Private/Factories/ActorFactory.cpp
+++ b/Engine/Source/Editor/UnrealEd/Private/Factories/ActorFactory.cpp
@@ -329,6 +329,12 @@ FQuat UActorFactoryStaticMesh::AlignObjectToSurfaceNormal(const FVector& InSurfa
 /*-----------------------------------------------------------------------------
 UActorFactoryBasicShape
 -----------------------------------------------------------------------------*/
+
+const FName UActorFactoryBasicShape::BasicCube("/Engine/BasicShapes/Cube.Cube");
+const FName UActorFactoryBasicShape::BasicSphere("/Engine/BasicShapes/Sphere.Sphere");
+const FName UActorFactoryBasicShape::BasicCylinder("/Engine/BasicShapes/Cylinder.Cylinder");
+const FName UActorFactoryBasicShape::BasicCone("/Engine/BasicShapes/Cone.Cone");
+
 UActorFactoryBasicShape::UActorFactoryBasicShape(const FObjectInitializer& ObjectInitializer)
 	: Super(ObjectInitializer)
 {
@@ -337,6 +343,16 @@ UActorFactoryBasicShape::UActorFactoryBasicShape(const FObjectInitializer& Objec
 	bUseSurfaceOrientation = true;
 }
 
+bool UActorFactoryBasicShape::CanCreateActorFrom( const FAssetData& AssetData, FText& OutErrorMsg )
+{
+	if(AssetData.IsValid() && (AssetData.ObjectPath == BasicCube || AssetData.ObjectPath == BasicSphere || AssetData.ObjectPath == BasicCone || AssetData.ObjectPath == BasicCylinder ) )
+	{
+		return true;
+	}
+
+	return false;
+}
+
 void UActorFactoryBasicShape::PostSpawnActor(UObject* Asset, AActor* NewActor)
 {
 	// Change properties
@@ -358,6 +374,7 @@ void UActorFactoryBasicShape::PostSpawnActor(UObject* Asset, AActor* NewActor)
 	}
 }
 
+
 /*-----------------------------------------------------------------------------
 UActorFactoryDeferredDecal
 -----------------------------------------------------------------------------*/
