[INTEGRATE] Change 2429904 by [REDACTED_PERSON][REDACTED_HANDLE] on 2015/02/03 10:05:43

	Only filter blueprint read-only properties when copy simulate results back from Simulate
	Fix spelling error

diff --git a/Engine/Source/Editor/Kismet/Private/SSCSEditor.cpp b/Engine/Source/Editor/Kismet/Private/SSCSEditor.cpp
--- a/Engine/Source/Editor/Kismet/Private/SSCSEditor.cpp
+++ b/Engine/Source/Editor/Kismet/Private/SSCSEditor.cpp
@@ -5540,7 +5540,7 @@ void SSCSEditor::OnApplyChangesToBlueprint() const
 			AActor* BlueprintCDO = Actor->GetClass()->GetDefaultObject<AActor>();
 			if (BlueprintCDO != NULL)
 			{
-				const auto CopyOptions = (EditorUtilities::ECopyOptions::Type)(EditorUtilities::ECopyOptions::OnlyCopyEditOrInterpProperties | EditorUtilities::ECopyOptions::PropagateChangesToArcheypeInstances);
+				const auto CopyOptions = (EditorUtilities::ECopyOptions::Type)(EditorUtilities::ECopyOptions::OnlyCopyEditOrInterpProperties | EditorUtilities::ECopyOptions::PropagateChangesToArchetypeInstances);
 				NumChangedProperties = EditorUtilities::CopyActorProperties(Actor, BlueprintCDO, CopyOptions);
 				if (Actor->GetInstanceComponents().Num() > 0)
 				{
diff --git a/Engine/Source/Editor/LevelEditor/Private/LevelEditorActions.cpp b/Engine/Source/Editor/LevelEditor/Private/LevelEditorActions.cpp
--- a/Engine/Source/Editor/LevelEditor/Private/LevelEditorActions.cpp
+++ b/Engine/Source/Editor/LevelEditor/Private/LevelEditorActions.cpp
@@ -1777,7 +1777,8 @@ void FLevelEditorActionCallbacks::OnKeepSimulationChanges()
 					const auto CopyOptions = ( EditorUtilities::ECopyOptions::Type )(
 						EditorUtilities::ECopyOptions::CallPostEditChangeProperty |
 						EditorUtilities::ECopyOptions::CallPostEditMove |
-						EditorUtilities::ECopyOptions::OnlyCopyEditOrInterpProperties );
+						EditorUtilities::ECopyOptions::OnlyCopyEditOrInterpProperties |
+						EditorUtilities::ECopyOptions::FilterBlueprintReadOnly);
 					const int32 CopiedPropertyCount = EditorUtilities::CopyActorProperties( SimWorldActor, EditorWorldActor, CopyOptions );
 
 					if( CopiedPropertyCount > 0 )
diff --git a/Engine/Source/Editor/UnrealEd/Private/Editor.cpp b/Engine/Source/Editor/UnrealEd/Private/Editor.cpp
--- a/Engine/Source/Editor/UnrealEd/Private/Editor.cpp
+++ b/Engine/Source/Editor/UnrealEd/Private/Editor.cpp
@@ -7235,7 +7235,7 @@ namespace EditorUtilities
 
 		// Get archetype instances for propagation (if requested)
 		TArray<UObject*> ArchetypeInstances;
-		if( Options & ECopyOptions::PropagateChangesToArcheypeInstances )
+		if( Options & ECopyOptions::PropagateChangesToArchetypeInstances )
 		{
 			TargetActor->GetArchetypeInstances(ArchetypeInstances);
 		}
@@ -7251,10 +7251,10 @@ namespace EditorUtilities
 				const bool bIsTransient = !!( Property->PropertyFlags & CPF_Transient );
 				const bool bIsComponentContainer = !!( Property->PropertyFlags & CPF_ContainsInstancedReference );
 				const bool bIsComponentProp = !!( Property->PropertyFlags & ( CPF_InstancedReference | CPF_ContainsInstancedReference ) );
-				const bool bIsReadonly = !!( Property->PropertyFlags & CPF_BlueprintReadOnly );
+				const bool bIsBlueprintReadonly = !!(Options & ECopyOptions::FilterBlueprintReadOnly) && !!( Property->PropertyFlags & CPF_BlueprintReadOnly );
 				const bool bIsIdentical = Property->Identical_InContainer( SourceActor, TargetActor );
 
-				if( !bIsTransient && !bIsIdentical && !bIsComponentContainer && !bIsComponentProp && !bIsReadonly && Property->GetName() != TEXT( "Tag" ) )
+				if( !bIsTransient && !bIsIdentical && !bIsComponentContainer && !bIsComponentProp && !bIsBlueprintReadonly && Property->GetName() != TEXT( "Tag" ) )
 				{
 					const bool bIsSafeToCopy = !( Options & ECopyOptions::OnlyCopyEditOrInterpProperties ) || ( Property->HasAnyPropertyFlags( CPF_Edit | CPF_Interp ) );
 					if( bIsSafeToCopy )
@@ -7274,7 +7274,7 @@ namespace EditorUtilities
 
 							// Determine which archetype instances match the current property value of the target actor (before it gets changed). We only want to propagate the change to those instances.
 							TArray<UObject*> ArchetypeInstancesToChange;
-							if( Options & ECopyOptions::PropagateChangesToArcheypeInstances )
+							if( Options & ECopyOptions::PropagateChangesToArchetypeInstances )
 							{
 								for( int32 InstanceIndex = 0; InstanceIndex < ArchetypeInstances.Num(); ++InstanceIndex )
 								{
@@ -7294,7 +7294,7 @@ namespace EditorUtilities
 								TargetActor->PostEditChangeProperty( PropertyChangedEvent );
 							}
 
-							if( Options & ECopyOptions::PropagateChangesToArcheypeInstances )
+							if( Options & ECopyOptions::PropagateChangesToArchetypeInstances )
 							{
 								for( int32 InstanceIndex = 0; InstanceIndex < ArchetypeInstancesToChange.Num(); ++InstanceIndex )
 								{
@@ -7348,7 +7348,7 @@ namespace EditorUtilities
 
 				// Build a list of matching component archetype instances for propagation (if requested)
 				TArray<UActorComponent*> ComponentArchetypeInstances;
-				if( Options & ECopyOptions::PropagateChangesToArcheypeInstances )
+				if( Options & ECopyOptions::PropagateChangesToArchetypeInstances )
 				{
 					for( int32 InstanceIndex = 0; InstanceIndex < ArchetypeInstances.Num(); ++InstanceIndex )
 					{
@@ -7398,7 +7398,7 @@ namespace EditorUtilities
 
 								// Determine which component archetype instances match the current property value of the target component (before it gets changed). We only want to propagate the change to those instances.
 								TArray<UActorComponent*> ComponentArchetypeInstancesToChange;
-								if( Options & ECopyOptions::PropagateChangesToArcheypeInstances )
+								if( Options & ECopyOptions::PropagateChangesToArchetypeInstances )
 								{
 									for( int32 InstanceIndex = 0; InstanceIndex < ComponentArchetypeInstances.Num(); ++InstanceIndex )
 									{
@@ -7418,7 +7418,7 @@ namespace EditorUtilities
 									TargetActor->PostEditChangeProperty( PropertyChangedEvent );
 								}
 
-								if( Options & ECopyOptions::PropagateChangesToArcheypeInstances )
+								if( Options & ECopyOptions::PropagateChangesToArchetypeInstances )
 								{
 									for( int32 InstanceIndex = 0; InstanceIndex < ComponentArchetypeInstancesToChange.Num(); ++InstanceIndex )
 									{
diff --git a/Engine/Source/Editor/UnrealEd/Public/Editor.h b/Engine/Source/Editor/UnrealEd/Public/Editor.h
--- a/Engine/Source/Editor/UnrealEd/Public/Editor.h
+++ b/Engine/Source/Editor/UnrealEd/Public/Editor.h
@@ -621,7 +621,10 @@ namespace EditorUtilities
 			OnlyCopyEditOrInterpProperties = 1 << 3,
 
 			/** Propagate property changes to archetype instances if the target actor is a CDO */
-			PropagateChangesToArcheypeInstances = 1 << 4,
+			PropagateChangesToArchetypeInstances = 1 << 4,
+
+			/** Filters out Blueprint Read-only properties */
+			FilterBlueprintReadOnly = 1 << 5,
 		};
 	}
 
