[INTEGRATE] Change 2430848 by [REDACTED_PERSON][REDACTED_HANDLE] on 2015/02/03 19:30:59

	Accidently using master list instead of filtered list for the SListView in SComponentClassCombo, this resulted in the list not being filterable.

diff --git a/Engine/Source/Editor/UnrealEd/Private/SComponentClassCombo.cpp b/Engine/Source/Editor/UnrealEd/Private/SComponentClassCombo.cpp
--- a/Engine/Source/Editor/UnrealEd/Private/SComponentClassCombo.cpp
+++ b/Engine/Source/Editor/UnrealEd/Private/SComponentClassCombo.cpp
@@ -36,7 +36,7 @@ void SComponentClassCombo::Construct(const FArguments& InArgs)
 	HotReloadSupport.OnHotReload().AddSP( this, &SComponentClassCombo::OnProjectHotReloaded );
 
 	SAssignNew(ComponentClassListView, SListView<FComponentClassComboEntryPtr>)
-		.ListItemsSource(ComponentClassList)
+		.ListItemsSource(&FilteredComponentClassList)
 		.OnSelectionChanged( this, &SComponentClassCombo::OnAddComponentSelectionChanged )
 		.OnGenerateRow( this, &SComponentClassCombo::GenerateAddComponentRow )
 		.SelectionMode(ESelectionMode::Single);
