[INTEGRATE] Change 2430917 by [REDACTED_PERSON][REDACTED_HANDLE] on 2015/02/03 20:16:41

	When typing into the search box it selects the first element of the filtered view to make hitting enter create that widget.
	Also changed the style of headers and separators so they don't highlight as if the can be clicked.
	Fixes UE-9006 along with CL 2430848

diff --git a/Engine/Source/Editor/UnrealEd/Private/SComponentClassCombo.cpp b/Engine/Source/Editor/UnrealEd/Private/SComponentClassCombo.cpp
--- a/Engine/Source/Editor/UnrealEd/Private/SComponentClassCombo.cpp
+++ b/Engine/Source/Editor/UnrealEd/Private/SComponentClassCombo.cpp
@@ -148,6 +148,11 @@ void SComponentClassCombo::GenerateFilteredComponentList(const FString& InSearch
 				}
 			}
 		}
+		if (FilteredComponentClassList.Num() > 0)
+		{
+			FComponentClassComboEntryPtr& FirstEntry = FilteredComponentClassList[0];
+			ComponentClassListView->SetSelection(FirstEntry, ESelectInfo::OnNavigation);
+		}
 	}
 }
 
@@ -265,6 +270,7 @@ TSharedRef<ITableRow> SComponentClassCombo::GenerateAddComponentRow( FComponentC
 	{
 		return 
 			SNew( STableRow< TSharedPtr<FString> >, OwnerTable )
+				.Style(&FEditorStyle::Get().GetWidgetStyle<FTableRowStyle>("TableView.NoHoverTableRow"))
 				.ShowSelection(false)
 			[
 				SNew(SBox)
@@ -280,6 +286,7 @@ TSharedRef<ITableRow> SComponentClassCombo::GenerateAddComponentRow( FComponentC
 	{
 		return 
 			SNew( STableRow< TSharedPtr<FString> >, OwnerTable )
+				.Style(&FEditorStyle::Get().GetWidgetStyle<FTableRowStyle>("TableView.NoHoverTableRow"))
 				.ShowSelection(false)
 			[
 				SNew(SBox)
