[INTEGRATE] Change 2431641 by [REDACTED_PERSON][REDACTED_HANDLE] on 2015/02/04 09:32:05

	Fix null deref crash shutting down the editor

diff --git a/Engine/Source/Editor/UnrealEd/Private/EditorModeManager.cpp b/Engine/Source/Editor/UnrealEd/Private/EditorModeManager.cpp
--- a/Engine/Source/Editor/UnrealEd/Private/EditorModeManager.cpp
+++ b/Engine/Source/Editor/UnrealEd/Private/EditorModeManager.cpp
@@ -53,8 +53,11 @@ FEditorModeTools::FEditorModeTools()
 	USelection::SelectionChangedEvent.AddRaw(this, &FEditorModeTools::OnEditorSelectionChanged);
 	USelection::SelectObjectEvent.AddRaw(this, &FEditorModeTools::OnEditorSelectionChanged);
 
-	// Register our callback for undo/redo
-	GEditor->RegisterForUndo(this);
+	if( GEditor )
+	{
+		// Register our callback for undo/redo
+		GEditor->RegisterForUndo(this);
+	}
 }
 
 FEditorModeTools::~FEditorModeTools()
