[INTEGRATE] Change 2431741 by [REDACTED_PERSON][REDACTED_HANDLE] on 2015/02/04 10:41:14

	Fixed issue where Surface properties for BSP would not appear due to multiple selection notify bug

diff --git a/Engine/Source/Editor/LevelEditor/Private/LevelEditorGenericDetails.cpp b/Engine/Source/Editor/LevelEditor/Private/LevelEditorGenericDetails.cpp
--- a/Engine/Source/Editor/LevelEditor/Private/LevelEditorGenericDetails.cpp
+++ b/Engine/Source/Editor/LevelEditor/Private/LevelEditorGenericDetails.cpp
@@ -31,7 +31,9 @@ void FLevelEditorGenericDetails::CustomizeDetails( IDetailLayoutBuilder& DetailL
 
 void FLevelEditorGenericDetails::GetSelectedSurfaceMaterials(IMaterialListBuilder& MaterialList) const
 {
-	for (TSelectedSurfaceIterator<> SurfaceIt(GWorld); SurfaceIt; ++SurfaceIt)
+	FSelectedActorInfo SelectionInfo = AssetSelectionUtils::GetSelectedActorInfo();	
+
+	for (TSelectedSurfaceIterator<> SurfaceIt(SelectionInfo.SharedWorld); SurfaceIt; ++SurfaceIt)
 	{
         FBspSurf* SelectedSurface = *SurfaceIt;
 
@@ -44,10 +46,12 @@ void FLevelEditorGenericDetails::OnMaterialChanged( UMaterialInterface* NewMater
 {
 	bool bModelDirtied = false;
 	{
+		FSelectedActorInfo SelectionInfo = AssetSelectionUtils::GetSelectedActorInfo();	
+
 		const FScopedTransaction Transaction( NSLOCTEXT("UnrealEd", "PolySetMaterial", "Set Material") );
-		GEditor->FlagModifyAllSelectedSurfacesInLevels( GWorld );
+		GEditor->FlagModifyAllSelectedSurfacesInLevels( SelectionInfo.SharedWorld );
 
-		for ( TSelectedSurfaceIterator<> It(GWorld) ; It ; ++It )
+		for ( TSelectedSurfaceIterator<> It(SelectionInfo.SharedWorld) ; It ; ++It )
 		{
 			UModel* Model = It.GetModel();
 			const int32 SurfaceIndex = It.GetSurfaceIndex();
@@ -124,7 +128,9 @@ void FLevelEditorGenericDetails::AddSurfaceDetails( IDetailLayoutBuilder& Detail
 	{
 		static FReply ExecuteExecCommand(FString InCommand)
 		{
-			GUnrealEd->Exec( GWorld, *InCommand );
+			FSelectedActorInfo SelectionInfo = AssetSelectionUtils::GetSelectedActorInfo();	
+
+			GUnrealEd->Exec( SelectionInfo.SharedWorld, *InCommand );
 			return FReply::Handled();
 		}
 	};
diff --git a/Engine/Source/Editor/UnrealEd/Private/LevelViewportClickHandlers.cpp b/Engine/Source/Editor/UnrealEd/Private/LevelViewportClickHandlers.cpp
--- a/Engine/Source/Editor/UnrealEd/Private/LevelViewportClickHandlers.cpp
+++ b/Engine/Source/Editor/UnrealEd/Private/LevelViewportClickHandlers.cpp
@@ -916,38 +916,40 @@ namespace ClickHandlers
 		}
 		else
 		{	
+			const FScopedTransaction Transaction( NSLOCTEXT("UnrealEd", "SelectBrushSurface", "Select Brush Surface") );
 			bool bDeselectAlreadyHandled = false;
-			if( GetDefault<ULevelEditorViewportSettings>()->bClickBSPSelectsBrush )
+			
+			// We are going to handle the notification ourselves
+			const bool bNotify = false;
+			if(GetDefault<ULevelEditorViewportSettings>()->bClickBSPSelectsBrush)
 			{
 				// Add to the actor selection set the brush actor that belongs to this BSP surface.
 				// Check Surf.Actor, as it can be NULL after deleting brushes and before rebuilding BSP.
-				if( Surf.Actor )
+				if(Surf.Actor)
 				{
-					const FScopedTransaction Transaction( NSLOCTEXT("UnrealEd", "SelectBrushFromSurface", "Select Brush from Surface") );
-					if( !Click.IsControlDown() )
+					if(!Click.IsControlDown())
 					{
-						GEditor->SelectNone( false, true );
+						GEditor->SelectNone(false, true);
 						bDeselectAlreadyHandled = true;
 					}
 					// If the builder brush is selected, first deselect it.
 					USelection* SelectedActors = GEditor->GetSelectedActors();
-					for ( FSelectionIterator It( *SelectedActors ) ; It ; ++It )
+					for(FSelectionIterator It(*SelectedActors); It; ++It)
 					{
-						ABrush* Brush = Cast<ABrush>( *It );
-						if ( Brush && FActorEditorUtils::IsABuilderBrush( Brush ) )
+						ABrush* Brush = Cast<ABrush>(*It);
+						if(Brush && FActorEditorUtils::IsABuilderBrush(Brush))
 						{
-							GEditor->SelectActor( Brush, false, false );
+							GEditor->SelectActor(Brush, false, bNotify);
 							break;
 						}
 					}
 
-					GEditor->SelectActor( Surf.Actor, true, true );
+					GEditor->SelectActor(Surf.Actor, true, bNotify);
 				}
 			}
+
 			// Select or deselect surfaces.
 			{
-				const FScopedTransaction Transaction( NSLOCTEXT("UnrealEd", "SelectSurfaces", "Select Surfaces") );
-
 				if( !Click.IsControlDown() && !bDeselectAlreadyHandled)
 				{
 					GEditor->SelectNone( false, true );
@@ -958,9 +960,12 @@ namespace ClickHandlers
 				// If there are no surfaces selected now, deselect the actor
 				if (!Model->HasSelectedSurfaces())
 				{
-					GEditor->SelectActor(Surf.Actor, false, true);
+					GEditor->SelectActor(Surf.Actor, false, bNotify);
 				}
 			}
+
+
+
 			GEditor->NoteSelectionChange();
 		}
 	}
