fix for UE-8741 Regression: Game Mode (G) no longer hides transform gizmo

diff --git a/Engine/Source/Editor/UnrealEd/Private/EditorViewportClient.cpp b/Engine/Source/Editor/UnrealEd/Private/EditorViewportClient.cpp
--- a/Engine/Source/Editor/UnrealEd/Private/EditorViewportClient.cpp
+++ b/Engine/Source/Editor/UnrealEd/Private/EditorViewportClient.cpp
@@ -4443,6 +4443,13 @@ void FEditorViewportClient::SetGameView(bool bGameViewEnable)
 	EngineShowFlags.CompositeEditorPrimitives = bCompositeEditorPrimitives;
 	LastEngineShowFlags.CompositeEditorPrimitives = bCompositeEditorPrimitives;
 
+	//reset game engine show flags that may have been turned on by making a selection in game view
+	if(bGameViewEnable)
+	{
+		EngineShowFlags.ModeWidgets = 0;
+		EngineShowFlags.Selection = 0;
+	}
+
 	EngineShowFlags.SelectionOutline = bGameViewEnable ? false : GetDefault<ULevelEditorViewportSettings>()->bUseSelectionOutline;
 
 	ApplyViewMode(GetViewMode(), IsPerspective(), EngineShowFlags);
