[INTEGRATE] Change 2432033 by [REDACTED_PERSON][REDACTED_HANDLE] on 2015/02/04 13:31:18

	Fixed weird editor layout until the editor's window was first activated
	- This reproduced when starting the editor while another app window was activated
	- Mode tools UI is only updated by viewport ticking.  It could only set its default layout after being ticked the first time.
	- Viewports never became "visible" until they were ticked at least once.  The editor doesn't tick viewports until it becomes activated.
	(I think this bug was introduced with CL 2284748)

diff --git a/Engine/Source/Editor/UnrealEd/Private/SEditorViewport.cpp b/Engine/Source/Editor/UnrealEd/Private/SEditorViewport.cpp
--- a/Engine/Source/Editor/UnrealEd/Private/SEditorViewport.cpp
+++ b/Engine/Source/Editor/UnrealEd/Private/SEditorViewport.cpp
@@ -434,7 +434,9 @@ bool SEditorViewport::IsVisible() const
 {
 	const float VisibilityTimeThreshold = .25f;
 	// The viewport is visible if we don't have a parent layout (likely a floating window) or this viewport is visible in the parent layout
-	return FPlatformTime::Seconds() - LastTickTime <= VisibilityTimeThreshold;
+	return 
+		LastTickTime == 0.0	||	// Never been ticked
+		FPlatformTime::Seconds() - LastTickTime <= VisibilityTimeThreshold;	// Ticked recently
 }
 
 void SEditorViewport::OnScreenCapture()
