[INTEGRATE] Change 2432449 by [REDACTED_PERSON][REDACTED_HANDLE] on 2015/02/04 16:57:41

	Dosn't allow the selection to affect adding a scene component. UE-8545

diff --git a/Engine/Source/Editor/Kismet/Private/SSCSEditor.cpp b/Engine/Source/Editor/Kismet/Private/SSCSEditor.cpp
--- a/Engine/Source/Editor/Kismet/Private/SSCSEditor.cpp
+++ b/Engine/Source/Editor/Kismet/Private/SSCSEditor.cpp
@@ -3846,7 +3846,9 @@ UActorComponent* SSCSEditor::PerformComboAddClass(TSubclassOf<UActorComponent> C
 		// This adds components according to the type selected in the drop down. If the user
 		// has the appropriate objects selected in the content browser then those are added,
 		// else we go down the previous route of adding components by type.
-		if(Selection->Num() > 0 && !AssetOverride)
+		//
+		// Furthermore don't try to match up assets for USceneComponent it will match lots of things and doesn't have any nice behavior for asset adds 
+		if (Selection->Num() > 0 && !AssetOverride && NewClass != USceneComponent::StaticClass())
 		{
 			for(FSelectionIterator ObjectIter(*Selection); ObjectIter; ++ObjectIter)
 			{
diff --git a/Engine/Source/Editor/UnrealEd/Private/SComponentClassCombo.cpp b/Engine/Source/Editor/UnrealEd/Private/SComponentClassCombo.cpp
--- a/Engine/Source/Editor/UnrealEd/Private/SComponentClassCombo.cpp
+++ b/Engine/Source/Editor/UnrealEd/Private/SComponentClassCombo.cpp
@@ -369,7 +369,8 @@ FText SComponentClassCombo::GetFriendlyComponentName(FComponentClassComboEntryPt
 	{
 		FriendlyComponentName = GetSanitizedComponentName(Entry);
 
-		if( Entry->GetComponentNameOverride().IsEmpty() )
+		// Don't try to match up assets for USceneComponent it will match lots of things and doesn't have any nice behavior for asset adds 
+		if (Entry->GetComponentClass() != USceneComponent::StaticClass() && Entry->GetComponentNameOverride().IsEmpty())
 		{
 			// Search the selected assets and look for any that can be used as a source asset for this type of component
 			// If there is one we append the asset name to the component name, if there are many we append "Multiple Assets"
