[INTEGRATE] Change 2432851 by [REDACTED_PERSON][REDACTED_HANDLE] on 2015/02/04 21:02:59

	Editor: Add UActorComponent and USceneComponent to the class picker list of parents for Blueprints [UE-8915]
	- Removed junk code for localizing the existing ini entries, switching to GetDisplayNameText and short tooltips instead
	- Add UActorComponent and USceneComponent to the list (gated by the experimental Blueprintable components flag)
	- Optimized the vertical screen usage of SClassPickerDialog for the BP parent picker case, keeping it visible in minspec monitors

diff --git a/Engine/Config/BaseEditor.ini b/Engine/Config/BaseEditor.ini
--- a/Engine/Config/BaseEditor.ini
+++ b/Engine/Config/BaseEditor.ini
@@ -181,11 +181,13 @@ bExpandClassPickerClassList=False
 
 ; Class picker dialog settings
 ; LocText refers to the localization id for the descriptive text to show next to the icon in the picker
-+NewAssetDefaultClasses=(ClassName="/Script/Engine.Actor", LocTextNameID="ActorName", LocTextDescriptionID="ActorDesc", AssetClass="/Script/Engine.Blueprint")
-+NewAssetDefaultClasses=(ClassName="/Script/Engine.Pawn", LocTextNameID="PawnName", LocTextDescriptionID="PawnDesc", AssetClass="/Script/Engine.Blueprint")
-+NewAssetDefaultClasses=(ClassName="/Script/Engine.Character", LocTextNameID="CharacterName", LocTextDescriptionID="CharacterDesc", AssetClass="/Script/Engine.Blueprint")
-+NewAssetDefaultClasses=(ClassName="/Script/Engine.PlayerController", LocTextNameID="PlayerControllerName", LocTextDescriptionID="PlayerControllerDesc", AssetClass="/Script/Engine.Blueprint")
-+NewAssetDefaultClasses=(ClassName="/Script/Engine.GameMode", LocTextNameID="GameModeName", LocTextDescriptionID="GameModeDesc", AssetClass="/Script/Engine.Blueprint")
++NewAssetDefaultClasses=(ClassName="/Script/Engine.Actor", AssetClass="/Script/Engine.Blueprint")
++NewAssetDefaultClasses=(ClassName="/Script/Engine.Pawn", AssetClass="/Script/Engine.Blueprint")
++NewAssetDefaultClasses=(ClassName="/Script/Engine.Character", AssetClass="/Script/Engine.Blueprint")
++NewAssetDefaultClasses=(ClassName="/Script/Engine.PlayerController", AssetClass="/Script/Engine.Blueprint")
++NewAssetDefaultClasses=(ClassName="/Script/Engine.GameMode", AssetClass="/Script/Engine.Blueprint")
++NewAssetDefaultClasses=(ClassName="/Script/Engine.ActorComponent", AssetClass="/Script/Engine.Blueprint")
++NewAssetDefaultClasses=(ClassName="/Script/Engine.SceneComponent", AssetClass="/Script/Engine.Blueprint")
 
 [TaskBrowser]
 DefaultServerName=
diff --git a/Engine/Source/Editor/UnrealEd/Classes/Preferences/UnrealEdOptions.h b/Engine/Source/Editor/UnrealEd/Classes/Preferences/UnrealEdOptions.h
--- a/Engine/Source/Editor/UnrealEd/Classes/Preferences/UnrealEdOptions.h
+++ b/Engine/Source/Editor/UnrealEd/Classes/Preferences/UnrealEdOptions.h
@@ -52,14 +52,6 @@ struct FClassPickerDefaults
 	UPROPERTY()
 	FString ClassName;
 
-	/** The loc text name for the class */
-	UPROPERTY()
-	FString LocTextNameID;
-
-	/** The loc text description for the class */
-	UPROPERTY()
-	FString LocTextDescriptionID;
-
 	/** The name of the asset type being created */
 	UPROPERTY()
 	FString AssetClass;
diff --git a/Engine/Source/Editor/UnrealEd/Private/SClassPickerDialog.cpp b/Engine/Source/Editor/UnrealEd/Private/SClassPickerDialog.cpp
--- a/Engine/Source/Editor/UnrealEd/Private/SClassPickerDialog.cpp
+++ b/Engine/Source/Editor/UnrealEd/Private/SClassPickerDialog.cpp
@@ -3,6 +3,7 @@
 
 #include "UnrealEd.h"
 #include "Editor/ClassViewer/Private/SClassViewer.h"
+#include "Editor/UnrealEd/Classes/Settings/EditorExperimentalSettings.h"
 #include "SClassPickerDialog.h"
 #include "EditorClassUtils.h"
 #include "SExpandableArea.h"
@@ -19,22 +20,30 @@ void SClassPickerDialog::Construct(const FArguments& InArgs)
 
 	ClassViewer = StaticCastSharedRef<SClassViewer>(FModuleManager::LoadModuleChecked<FClassViewerModule>("ClassViewer").CreateClassViewer(InArgs._Options, FOnClassPicked::CreateSP(this,&SClassPickerDialog::OnClassPicked)));
 
+	// Determine if components are allowed
+	const UEditorExperimentalSettings* ExperimentalSettings = GetDefault<UEditorExperimentalSettings>();
+	const bool bForBlueprint = (InArgs._AssetType != nullptr) && (InArgs._AssetType->IsChildOf(UBlueprint::StaticClass()));
+	const bool bAllowComponentSubclasses = !bForBlueprint || ExperimentalSettings->bBlueprintableComponents;
+
 	// Load in default settings
-	for(int ClassIdx = 0; ClassIdx < GUnrealEd->GetUnrealEdOptions()->NewAssetDefaultClasses.Num(); ++ClassIdx)
+	for (const FClassPickerDefaults& DefaultObj : GUnrealEd->GetUnrealEdOptions()->NewAssetDefaultClasses)
 	{
-		FClassPickerDefaults& DefaultObj = GUnrealEd->GetUnrealEdOptions()->NewAssetDefaultClasses[ClassIdx];
 		UClass* AssetType = LoadClass<UObject>(NULL, *DefaultObj.AssetClass, NULL, LOAD_None, NULL);
-		if (InArgs._AssetType->IsChildOf(AssetType))
+		UClass* AssetActualClass = LoadClass<UObject>(NULL, *DefaultObj.ClassName, NULL, LOAD_None, NULL);
+		const bool bIsComponent = AssetActualClass->IsChildOf(UActorComponent::StaticClass());
+
+		if (InArgs._AssetType->IsChildOf(AssetType) && (!bIsComponent || bAllowComponentSubclasses))
 		{
 			AssetDefaultClasses.Add(MakeShareable(new FClassPickerDefaults(DefaultObj)));
 		}
 	}
 
-	bool bHasDefaultClasses = AssetDefaultClasses.Num() > 0;
+	const bool bHasDefaultClasses = AssetDefaultClasses.Num() > 0;
+
 	bool bExpandDefaultClassPicker = true;
 	bool bExpandCustomClassPicker = !bHasDefaultClasses;
 
-	if(bHasDefaultClasses)
+	if (bHasDefaultClasses)
 	{
 		GConfig->GetBool(TEXT("/Script/UnrealEd.UnrealEdOptions"), TEXT("bExpandClassPickerDefaultClassList"), bExpandDefaultClassPicker, GEditorIni);
 		GConfig->GetBool(TEXT("/Script/UnrealEd.UnrealEdOptions"), TEXT("bExpandCustomClassPickerClassList"), bExpandCustomClassPicker, GEditorIni);
@@ -48,7 +57,7 @@ void SClassPickerDialog::Construct(const FArguments& InArgs)
 		[
 			SNew(SBox)
 			.Visibility(EVisibility::Visible)
-			.WidthOverride(500.0f)
+			.WidthOverride(520.0f)
 			[
 				SNew(SVerticalBox)
 				+SVerticalBox::Slot()
@@ -73,7 +82,7 @@ void SClassPickerDialog::Construct(const FArguments& InArgs)
 				.Padding(0.0f, 10.0f, 0.0f, 0.0f)
 				[
 					SNew(SExpandableArea)
-					.MaxHeight(385.f)
+					.MaxHeight(320.f)
 					.InitiallyCollapsed(!bExpandCustomClassPicker)
 					.AreaTitle(NSLOCTEXT("SClassPickerDialog", "AllClassesAreaTitle", "All Classes"))
 					.OnAreaExpansionChanged(this, &SClassPickerDialog::OnCustomAreaExpansionChanged)
@@ -162,7 +171,7 @@ TSharedRef<ITableRow> SClassPickerDialog::GenerateListRow(TSharedPtr<FClassPicke
 		SNew(SVerticalBox)
 		+SVerticalBox::Slot()
 		.MaxHeight(30.0f)
-		.Padding(0.0f, 6.0f, 0.0f, 4.0f)
+		.Padding(10.0f, 6.0f, 0.0f, 4.0f)
 		[
 			SNew(SHorizontalBox)
 			+SHorizontalBox::Slot()
@@ -206,13 +215,6 @@ TSharedRef<ITableRow> SClassPickerDialog::GenerateListRow(TSharedPtr<FClassPicke
 				FEditorClassUtils::GetDocumentationLinkWidget(ItemClass)
 			]
 		]
-		+SVerticalBox::Slot()
-		.AutoHeight()
-		.Padding(0.0f, 6.0f, 0.0f, 4.0f)
-		[
-			SNew(SSeparator)
-			.Orientation(Orient_Horizontal)
-		]
 	];
 }
 
diff --git a/Engine/Source/Editor/UnrealEd/Private/UnrealEdEngine.cpp b/Engine/Source/Editor/UnrealEd/Private/UnrealEdEngine.cpp
--- a/Engine/Source/Editor/UnrealEd/Private/UnrealEdEngine.cpp
+++ b/Engine/Source/Editor/UnrealEd/Private/UnrealEdEngine.cpp
@@ -728,71 +728,27 @@ void UUnrealEdEngine::SetMapBuildCancelled( bool InCancelled )
 
 FText FClassPickerDefaults::GetName() const
 {
-	static TMap< FString, FText > LocNames;
+	FText Result = LOCTEXT("NullClass", "(null class)");
 
-	if (LocNames.Num() == 0)
+	if (UClass* ItemClass = LoadClass<UObject>(NULL, *ClassName, NULL, LOAD_None, NULL))
 	{
-		LocNames.Add( TEXT("ActorName"), LOCTEXT("ActorName", "Actor") );
-		LocNames.Add( TEXT("PawnName"), LOCTEXT("PawnName", "Pawn") );
-		LocNames.Add( TEXT("CharacterName"), LOCTEXT("CharacterName", "Character") );
-		LocNames.Add( TEXT("PlayerControllerName"), LOCTEXT("PlayerControllerName", "PlayerController") );
-		LocNames.Add( TEXT("GameModeName"), LOCTEXT("GameModeName", "Game Mode") );
+		Result = ItemClass->GetDisplayNameText();
 	}
 
-	if ( LocTextNameID.IsEmpty() )
-	{
-		UClass* ItemClass = LoadClass<UObject>(NULL, *ClassName, NULL, LOAD_None, NULL);
-		check( ItemClass );
-		return FText::FromString(FName::NameToDisplayString(ItemClass->GetName(), false));
-	}
-	
-	const FText* PreExistingName = LocNames.Find( LocTextNameID );
-	if ( PreExistingName )
-	{
-		return *PreExistingName;
-	}
-
-	FText OutName;
-	if ( FText::FindText(TEXT("UnrealEd"), LocTextNameID, OutName) )
-	{
-		return OutName;
-	}
-		
-	return FText::FromString(LocTextNameID);
+	return Result;
 }
 
 
 FText FClassPickerDefaults::GetDescription() const
 {
-	static TMap< FString, FText > LocDescs;
+	FText Result = LOCTEXT("NullClass", "(null class)");
 
-	if (LocDescs.Num() == 0)
+	if (UClass* ItemClass = LoadClass<UObject>(NULL, *ClassName, NULL, LOAD_None, NULL))
 	{
-		LocDescs.Add( TEXT("ActorDesc"), LOCTEXT("ActorDesc", "An Actor is an object that can be placed or spawned in the world.") );
-		LocDescs.Add( TEXT("PawnDesc"), LOCTEXT("PawnDesc", "A Pawn is an actor that can be 'possessed' and receieve input from a controller.") );
-		LocDescs.Add( TEXT("CharacterDesc"), LOCTEXT("CharacterDesc", "A character is a type of Pawn that includes the ability to walk around.") );
-		LocDescs.Add( TEXT("PlayerControllerDesc"), LOCTEXT("PlayerControllerDesc", "A Player Controller is an actor responsible for controlling a Pawn used by the player.") );
-		LocDescs.Add( TEXT("GameModeDesc"), LOCTEXT("GameModeDesc", "Game Mode defines the game being played, its rules, scoring, and other facets of the game type.") );
+		Result = ItemClass->GetToolTipText(/*bShortTooltip=*/ true);
 	}
 
-	if ( LocTextDescriptionID.IsEmpty() )
-	{
-		return LOCTEXT("NoClassPickerDesc", "No Description.");
-	}
-
-	const FText* PreExistingDesc = LocDescs.Find( LocTextDescriptionID );
-	if ( PreExistingDesc )
-	{
-		return *PreExistingDesc;
-	}
-
-	FText OutDesc;
-	if ( FText::FindText(TEXT("UnrealEd"), LocTextDescriptionID, OutDesc) )
-	{
-		return OutDesc;
-	}
-		
-	return FText::FromString(LocTextDescriptionID);
+	return Result;
 }
 
 #undef LOCTEXT_NAMESPACE
