[INTEGRATE] Change 2432943 by [REDACTED_PERSON][REDACTED_HANDLE] on 2015/02/04 22:59:54

	Blueprints: Force the root component to zero rotation/translation during the convert Actor to BP process (SetRelativeLocation was earlying out due to no perceived movement, now clearing the relative rotation and translation directly) [UE-8190]

diff --git a/Engine/Source/Editor/UnrealEd/Private/Kismet2/Kismet2.cpp b/Engine/Source/Editor/UnrealEd/Private/Kismet2/Kismet2.cpp
--- a/Engine/Source/Editor/UnrealEd/Private/Kismet2/Kismet2.cpp
+++ b/Engine/Source/Editor/UnrealEd/Private/Kismet2/Kismet2.cpp
@@ -1014,15 +1014,15 @@ UBlueprint* FKismetEditorUtilities::CreateBlueprintFromActor(const FName Bluepri
 {
 	UBlueprint* NewBlueprint = nullptr;
 
-	if (Actor)
+	if (Actor != nullptr)
 	{
-		if (Outer)
+		if (Outer != nullptr)
 		{
 			// We don't have a factory, but we can still try to create a blueprint for this actor class
 			NewBlueprint = FKismetEditorUtilities::CreateBlueprint( Actor->GetClass(), Outer, BlueprintName, EBlueprintType::BPTYPE_Normal, UBlueprint::StaticClass(), UBlueprintGeneratedClass::StaticClass(), FName("CreateFromActor") );
 		}
 
-		if(NewBlueprint)
+		if (NewBlueprint != nullptr)
 		{
 			// Notify the asset registry
 			FAssetRegistryModule::AssetCreated(NewBlueprint);
@@ -1036,7 +1036,7 @@ UBlueprint* FKismetEditorUtilities::CreateBlueprintFromActor(const FName Bluepri
 				AddComponentsToBlueprint(NewBlueprint, Actor->GetInstanceComponents());
 			}
 
-			if(NewBlueprint->GeneratedClass)
+			if (NewBlueprint->GeneratedClass != nullptr)
 			{
 				// Since we already created SCS Nodes for the instance components, temporarily cache and clear the
 				// array to avoid creating duplicates in the new CDO
@@ -1051,11 +1051,12 @@ UBlueprint* FKismetEditorUtilities::CreateBlueprintFromActor(const FName Bluepri
 					Actor->AddInstanceComponent(Component);
 				}
 
-				if(AActor* CDOAsActor = Cast<AActor>(CDO))
+				if (AActor* CDOAsActor = Cast<AActor>(CDO))
 				{
-					if(USceneComponent* Scene = CDOAsActor->GetRootComponent())
+					if (USceneComponent* Scene = CDOAsActor->GetRootComponent())
 					{
-						Scene->SetRelativeLocation(FVector::ZeroVector);
+						Scene->RelativeLocation = FVector::ZeroVector;
+						Scene->RelativeRotation = FRotator::ZeroRotator;
 
 						// Clear out the attachment info after having copied the properties from the source actor
 						Scene->AttachParent = NULL;
@@ -1069,7 +1070,7 @@ UBlueprint* FKismetEditorUtilities::CreateBlueprintFromActor(const FName Bluepri
 
 			FKismetEditorUtilities::CompileBlueprint(NewBlueprint);
 
-			if(bReplaceActor)
+			if (bReplaceActor)
 			{
 				TArray<AActor*> Actors;
 				Actors.Add(Actor);
@@ -1082,7 +1083,7 @@ UBlueprint* FKismetEditorUtilities::CreateBlueprintFromActor(const FName Bluepri
 		}
 	}
 
-	if (NewBlueprint)
+	if (NewBlueprint != nullptr)
 	{
 		// Open the editor for the new blueprint
 		FAssetEditorManager::Get().OpenEditorForAsset(NewBlueprint);
