Actors referencing Blueprint assets now open the Blueprint in preference to any other referenced asset when pressing Ctrl+E in the editor.

diff --git a/Engine/Plugins/2D/Paper2D/Source/Paper2D/Private/PaperFlipbookActor.cpp b/Engine/Plugins/2D/Paper2D/Source/Paper2D/Private/PaperFlipbookActor.cpp
--- a/Engine/Plugins/2D/Paper2D/Source/Paper2D/Private/PaperFlipbookActor.cpp
+++ b/Engine/Plugins/2D/Paper2D/Source/Paper2D/Private/PaperFlipbookActor.cpp
@@ -17,6 +17,8 @@ APaperFlipbookActor::APaperFlipbookActor(const FObjectInitializer& ObjectInitial
 #if WITH_EDITOR
 bool APaperFlipbookActor::GetReferencedContentObjects(TArray<UObject*>& Objects) const
 {
+	Super::GetReferencedContentObjects(Objects);
+
 	if (UPaperFlipbook* FlipbookAsset = RenderComponent->GetFlipbook())
 	{
 		Objects.Add(FlipbookAsset);
diff --git a/Engine/Plugins/2D/Paper2D/Source/Paper2D/Private/PaperSpriteActor.cpp b/Engine/Plugins/2D/Paper2D/Source/Paper2D/Private/PaperSpriteActor.cpp
--- a/Engine/Plugins/2D/Paper2D/Source/Paper2D/Private/PaperSpriteActor.cpp
+++ b/Engine/Plugins/2D/Paper2D/Source/Paper2D/Private/PaperSpriteActor.cpp
@@ -19,6 +19,8 @@ APaperSpriteActor::APaperSpriteActor(const FObjectInitializer& ObjectInitializer
 #if WITH_EDITOR
 bool APaperSpriteActor::GetReferencedContentObjects(TArray<UObject*>& Objects) const
 {
+	Super::GetReferencedContentObjects(Objects);
+
 	if (UPaperSprite* SourceSprite = RenderComponent->GetSprite())
 	{
 		Objects.Add(SourceSprite);
diff --git a/Engine/Plugins/2D/Paper2D/Source/Paper2D/Private/PaperTileMapActor.cpp b/Engine/Plugins/2D/Paper2D/Source/Paper2D/Private/PaperTileMapActor.cpp
--- a/Engine/Plugins/2D/Paper2D/Source/Paper2D/Private/PaperTileMapActor.cpp
+++ b/Engine/Plugins/2D/Paper2D/Source/Paper2D/Private/PaperTileMapActor.cpp
@@ -16,6 +16,8 @@ APaperTileMapActor::APaperTileMapActor(const FObjectInitializer& ObjectInitializ
 #if WITH_EDITOR
 bool APaperTileMapActor::GetReferencedContentObjects(TArray<UObject*>& Objects) const
 {
+	Super::GetReferencedContentObjects(Objects);
+
 	if (const UObject* Asset = RenderComponent->AdditionalStatObject())
 	{
 		Objects.Add(const_cast<UObject*>(Asset));
diff --git a/Engine/Plugins/2D/Paper2D/Source/Paper2D/Private/Terrain/PaperTerrainActor.cpp b/Engine/Plugins/2D/Paper2D/Source/Paper2D/Private/Terrain/PaperTerrainActor.cpp
--- a/Engine/Plugins/2D/Paper2D/Source/Paper2D/Private/Terrain/PaperTerrainActor.cpp
+++ b/Engine/Plugins/2D/Paper2D/Source/Paper2D/Private/Terrain/PaperTerrainActor.cpp
@@ -21,6 +21,8 @@ APaperTerrainActor::APaperTerrainActor(const FObjectInitializer& ObjectInitializ
 #if WITH_EDITOR
 bool APaperTerrainActor::GetReferencedContentObjects(TArray<UObject*>& Objects) const
 {
+	Super::GetReferencedContentObjects(Objects);
+
 	if (RenderComponent->TerrainMaterial != nullptr)
 	{
 		Objects.Add(RenderComponent->TerrainMaterial);
diff --git a/Engine/Source/Editor/UnrealEd/Classes/Editor/EditorEngine.h b/Engine/Source/Editor/UnrealEd/Classes/Editor/EditorEngine.h
--- a/Engine/Source/Editor/UnrealEd/Classes/Editor/EditorEngine.h
+++ b/Engine/Source/Editor/UnrealEd/Classes/Editor/EditorEngine.h
@@ -2234,11 +2234,12 @@ public:
 	void GetObjectsToSyncToContentBrowser( TArray<UObject*>& Objects );
 
 	/**
-	 * Queries for a list of assets that are referenced by the current editor selection (actors, surfaces, etc.)
-	 *
-	 * @param	Objects	Array to be filled with asset objects referenced by the current editor selection
-	 */
-	void GetReferencedAssetsForEditorSelection( TArray<UObject*>& Objects );
+	* Queries for a list of assets that are referenced by the current editor selection (actors, surfaces, etc.)
+	*
+	* @param	Objects								Array to be filled with asset objects referenced by the current editor selection
+	* @param	bIgnoreOtherAssetsIfBPReferenced	If true, and a selected actor has a Blueprint asset, only that will be returned.
+	*/
+	void GetReferencedAssetsForEditorSelection(TArray<UObject*>& Objects, const bool bIgnoreOtherAssetsIfBPReferenced = false);
 
 	/** Returns the WorldContext for the editor world. For now, there will always be exactly 1 of these in the editor. 
 	 *
diff --git a/Engine/Source/Editor/UnrealEd/Private/Editor.cpp b/Engine/Source/Editor/UnrealEd/Private/Editor.cpp
--- a/Engine/Source/Editor/UnrealEd/Private/Editor.cpp
+++ b/Engine/Source/Editor/UnrealEd/Private/Editor.cpp
@@ -3159,7 +3159,7 @@ void UEditorEngine::SyncToContentBrowser()
 }
 
 
-void UEditorEngine::GetReferencedAssetsForEditorSelection( TArray<UObject*>& Objects )
+void UEditorEngine::GetReferencedAssetsForEditorSelection(TArray<UObject*>& Objects, const bool bIgnoreOtherAssetsIfBPReferenced)
 {
 	for ( TSelectedSurfaceIterator<> It(GWorld) ; It ; ++It )
 	{
@@ -3176,7 +3176,25 @@ void UEditorEngine::GetReferencedAssetsForEditorSelection( TArray<UObject*>& Obj
 		AActor* Actor = static_cast<AActor*>( *It );
 		checkSlow( Actor->IsA(AActor::StaticClass()) );
 
-		Actor->GetReferencedContentObjects(Objects);
+		TArray<UObject*> ActorObjects;
+		Actor->GetReferencedContentObjects(ActorObjects);
+
+		// If Blueprint assets should take precedence over any other referenced asset, check if there are any blueprints in this actor's list
+		// and if so, add only those.
+		if (bIgnoreOtherAssetsIfBPReferenced && ActorObjects.ContainsByPredicate([](UObject* Obj) { return Obj->IsA(UBlueprint::StaticClass()); }))
+		{
+			for (UObject* Object : ActorObjects)
+			{
+				if (Object->IsA(UBlueprint::StaticClass()))
+				{
+					Objects.Add(Object);
+				}
+			}
+		}
+		else
+		{
+			Objects.Append(ActorObjects);
+		}
 	}
 }
 
diff --git a/Engine/Source/Runtime/Engine/Private/AmbientSound.cpp b/Engine/Source/Runtime/Engine/Private/AmbientSound.cpp
--- a/Engine/Source/Runtime/Engine/Private/AmbientSound.cpp
+++ b/Engine/Source/Runtime/Engine/Private/AmbientSound.cpp
@@ -63,7 +63,9 @@ void AAmbientSound::CheckForErrors( void )
 
 bool AAmbientSound::GetReferencedContentObjects( TArray<UObject*>& Objects ) const
 {
-	if( AudioComponent->Sound )
+	Super::GetReferencedContentObjects(Objects);
+
+	if (AudioComponent->Sound)
 	{
 		Objects.Add( AudioComponent->Sound );
 	}
diff --git a/Engine/Source/Runtime/Engine/Private/DestructibleActor.cpp b/Engine/Source/Runtime/Engine/Private/DestructibleActor.cpp
--- a/Engine/Source/Runtime/Engine/Private/DestructibleActor.cpp
+++ b/Engine/Source/Runtime/Engine/Private/DestructibleActor.cpp
@@ -21,6 +21,8 @@ ADestructibleActor::ADestructibleActor(const FObjectInitializer& ObjectInitializ
 #if WITH_EDITOR
 bool ADestructibleActor::GetReferencedContentObjects( TArray<UObject*>& Objects ) const
 {
+	Super::GetReferencedContentObjects(Objects);
+
 	if (DestructibleComponent && DestructibleComponent->SkeletalMesh)
 	{
 		Objects.Add(DestructibleComponent->SkeletalMesh);
diff --git a/Engine/Source/Runtime/Engine/Private/Interpolation.cpp b/Engine/Source/Runtime/Engine/Private/Interpolation.cpp
--- a/Engine/Source/Runtime/Engine/Private/Interpolation.cpp
+++ b/Engine/Source/Runtime/Engine/Private/Interpolation.cpp
@@ -1631,6 +1631,8 @@ void AMatineeActor::InitGroupActorForGroup(class UInterpGroup* InGroup, class AA
 
 bool AMatineeActor::GetReferencedContentObjects( TArray<UObject*>& Objects ) const
 {
+	Super::GetReferencedContentObjects(Objects);
+
 	if (MatineeData)
 	{
 		Objects.Add(MatineeData);
diff --git a/Engine/Source/Runtime/Engine/Private/LODActor.cpp b/Engine/Source/Runtime/Engine/Private/LODActor.cpp
--- a/Engine/Source/Runtime/Engine/Private/LODActor.cpp
+++ b/Engine/Source/Runtime/Engine/Private/LODActor.cpp
@@ -53,6 +53,8 @@ void ALODActor::PostEditChangeChainProperty(FPropertyChangedChainEvent& Property
 
 bool ALODActor::GetReferencedContentObjects( TArray<UObject*>& Objects ) const
 {
+	Super::GetReferencedContentObjects(Objects);
+
 	if (StaticMeshComponent && StaticMeshComponent->StaticMesh)
 	{
 		Objects.Add(StaticMeshComponent->StaticMesh);
diff --git a/Engine/Source/Runtime/Engine/Private/Particles/ParticleComponents.cpp b/Engine/Source/Runtime/Engine/Private/Particles/ParticleComponents.cpp
--- a/Engine/Source/Runtime/Engine/Private/Particles/ParticleComponents.cpp
+++ b/Engine/Source/Runtime/Engine/Private/Particles/ParticleComponents.cpp
@@ -466,6 +466,8 @@ void AEmitter::SetMaterialParameter(FName ParameterName, UMaterialInterface* Par
 
 bool AEmitter::GetReferencedContentObjects( TArray<UObject*>& Objects ) const
 {
+	Super::GetReferencedContentObjects(Objects);
+
 	if (ParticleSystemComponent->Template)
 	{
 		Objects.Add(ParticleSystemComponent->Template);
diff --git a/Engine/Source/Runtime/Engine/Private/SkeletalMesh.cpp b/Engine/Source/Runtime/Engine/Private/SkeletalMesh.cpp
--- a/Engine/Source/Runtime/Engine/Private/SkeletalMesh.cpp
+++ b/Engine/Source/Runtime/Engine/Private/SkeletalMesh.cpp
@@ -3889,6 +3889,8 @@ void ASkeletalMeshActor::FinishAnimControl(UInterpGroup* InInterpGroup)
 
 bool ASkeletalMeshActor::GetReferencedContentObjects( TArray<UObject*>& Objects ) const
 {
+	Super::GetReferencedContentObjects(Objects);
+
 	if (SkeletalMeshComponent->SkeletalMesh)
 	{
 		Objects.Add(SkeletalMeshComponent->SkeletalMesh);
diff --git a/Engine/Source/Runtime/Engine/Private/SplineMeshActor.cpp b/Engine/Source/Runtime/Engine/Private/SplineMeshActor.cpp
--- a/Engine/Source/Runtime/Engine/Private/SplineMeshActor.cpp
+++ b/Engine/Source/Runtime/Engine/Private/SplineMeshActor.cpp
@@ -43,6 +43,8 @@ void ASplineMeshActor::SetMobility(EComponentMobility::Type InMobility)
 
 bool ASplineMeshActor::GetReferencedContentObjects(TArray<UObject*>& Objects) const
 {
+	Super::GetReferencedContentObjects(Objects);
+
 	check(SplineMeshComponent != nullptr);
 	if (SplineMeshComponent->StaticMesh != nullptr)
 	{
diff --git a/Engine/Source/Runtime/Engine/Private/StaticMeshActor.cpp b/Engine/Source/Runtime/Engine/Private/StaticMeshActor.cpp
--- a/Engine/Source/Runtime/Engine/Private/StaticMeshActor.cpp
+++ b/Engine/Source/Runtime/Engine/Private/StaticMeshActor.cpp
@@ -111,6 +111,8 @@ void AStaticMeshActor::PostEditChangeChainProperty(FPropertyChangedChainEvent& P
 
 bool AStaticMeshActor::GetReferencedContentObjects( TArray<UObject*>& Objects ) const
 {
+	Super::GetReferencedContentObjects(Objects);
+
 	if (StaticMeshComponent && StaticMeshComponent->StaticMesh)
 	{
 		Objects.Add(StaticMeshComponent->StaticMesh);
diff --git a/Engine/Source/Runtime/Foliage/Private/ProceduralFoliageActor.cpp b/Engine/Source/Runtime/Foliage/Private/ProceduralFoliageActor.cpp
--- a/Engine/Source/Runtime/Foliage/Private/ProceduralFoliageActor.cpp
+++ b/Engine/Source/Runtime/Foliage/Private/ProceduralFoliageActor.cpp
@@ -22,6 +22,8 @@ AProceduralFoliageActor::AProceduralFoliageActor(const FObjectInitializer& Objec
 
 bool AProceduralFoliageActor::GetReferencedContentObjects(TArray<UObject*>& Objects) const
 {
+	Super::GetReferencedContentObjects(Objects);
+
 	if (ProceduralComponent && ProceduralComponent->ProceduralFoliage)
 	{
 		Objects.Add(ProceduralComponent->ProceduralFoliage);
