#GDC - matinee export

- Make sure to export only enabled tracks
- Fixed crash with Exporting with bake transform

diff --git a/Engine/Source/Editor/UnrealEd/Private/Fbx/FbxMainExport.cpp b/Engine/Source/Editor/UnrealEd/Private/Fbx/FbxMainExport.cpp
--- a/Engine/Source/Editor/UnrealEd/Private/Fbx/FbxMainExport.cpp
+++ b/Engine/Source/Editor/UnrealEd/Private/Fbx/FbxMainExport.cpp
@@ -1004,24 +1004,27 @@ bool FFbxExporter::ExportMatinee(AMatineeActor* InMatineeActor)
 		{
 			UInterpTrackInst* TrackInst = Group->TrackInst[TrackIndex];
 			UInterpTrack* Track = Group->Group->InterpTracks[TrackIndex];
-			if (TrackInst->IsA(UInterpTrackInstMove::StaticClass()) && Track->IsA(UInterpTrackMove::StaticClass()))
+			if ( Track->IsDisabled() == false )
 			{
-				UInterpTrackInstMove* MoveTrackInst = (UInterpTrackInstMove*) TrackInst;
-				UInterpTrackMove* MoveTrack = (UInterpTrackMove*) Track;
-				ExportMatineeTrackMove(FbxActor, MoveTrackInst, MoveTrack, InMatineeActor->MatineeData->InterpLength);
-			}
-			else if (TrackInst->IsA(UInterpTrackInstFloatProp::StaticClass()) && Track->IsA(UInterpTrackFloatProp::StaticClass()))
-			{
-				UInterpTrackInstFloatProp* PropertyTrackInst = (UInterpTrackInstFloatProp*) TrackInst;
-				UInterpTrackFloatProp* PropertyTrack = (UInterpTrackFloatProp*) Track;
-				ExportMatineeTrackFloatProp(FbxActor, PropertyTrack);
-			}
-			else if (TrackInst->IsA(UInterpTrackInstAnimControl::StaticClass()) && Track->IsA(UInterpTrackAnimControl::StaticClass()))
-			{
-				USkeletalMeshComponent* SkeletalMeshComp = Cast<USkeletalMeshComponent>(Actor->GetComponentByClass(USkeletalMeshComponent::StaticClass()));
-				if ( SkeletalMeshComp )
+				if(TrackInst->IsA(UInterpTrackInstMove::StaticClass()) && Track->IsA(UInterpTrackMove::StaticClass()))
+				{
+					UInterpTrackInstMove* MoveTrackInst = (UInterpTrackInstMove*)TrackInst;
+					UInterpTrackMove* MoveTrack = (UInterpTrackMove*)Track;
+					ExportMatineeTrackMove(FbxActor, MoveTrackInst, MoveTrack, InMatineeActor->MatineeData->InterpLength);
+				}
+				else if(TrackInst->IsA(UInterpTrackInstFloatProp::StaticClass()) && Track->IsA(UInterpTrackFloatProp::StaticClass()))
 				{
-					ExportAnimTrack(InMatineeActor, SkeletalMeshComp);
+					UInterpTrackInstFloatProp* PropertyTrackInst = (UInterpTrackInstFloatProp*)TrackInst;
+					UInterpTrackFloatProp* PropertyTrack = (UInterpTrackFloatProp*)Track;
+					ExportMatineeTrackFloatProp(FbxActor, PropertyTrack);
+				}
+				else if(TrackInst->IsA(UInterpTrackInstAnimControl::StaticClass()) && Track->IsA(UInterpTrackAnimControl::StaticClass()))
+				{
+					USkeletalMeshComponent* SkeletalMeshComp = Cast<USkeletalMeshComponent>(Actor->GetComponentByClass(USkeletalMeshComponent::StaticClass()));
+					if(SkeletalMeshComp)
+					{
+						ExportAnimTrack(InMatineeActor, SkeletalMeshComp);
+					}
 				}
 			}
 		}
@@ -1368,7 +1371,7 @@ void FFbxExporter::ExportAnimatedVector(FbxAnimCurve* FbxCurve, const char* Chan
 
 	// Determine how many key frames we are exporting. If the user wants to export a key every 
 	// frame, calculate this number. Otherwise, use the number of keys the user created. 
-	int32 KeyCount = bBakeKeys ? (InterpLength * BakeTransformsFPS) + Curve->Points.Num() : Curve->Points.Num();
+	int32 KeyCount = bBakeKeys ? (InterpLength * BakeTransformsFPS) : Curve->Points.Num();
 
 	// Write out the key times from the curve to the FBX curve.
 	TArray<float> KeyTimes;
@@ -1497,24 +1500,21 @@ void FFbxExporter::ExportMoveSubTrack(FbxAnimCurve* FbxCurve, const ANSICHAR* Ch
 
 	// Determine how many key frames we are exporting. If the user wants to export a key every 
 	// frame, calculate this number. Otherwise, use the number of keys the user created. 
-	int32 KeyCount = bBakeKeys ? (InterpLength * BakeTransformsFPS) + Curve->Points.Num() : Curve->Points.Num();
+	int32 KeyCount = bBakeKeys ? (InterpLength * BakeTransformsFPS) : Curve->Points.Num();
 
 	// Write out the key times from the curve to the FBX curve.
 	TArray<float> KeyTimes;
-	for (int32 KeyIndex = 0; KeyIndex < KeyCount; ++KeyIndex)
+	for(int32 KeyIndex = 0; KeyIndex < KeyCount; ++KeyIndex)
 	{
-		FInterpCurvePoint<float>& Key = Curve->Points[KeyIndex];
-
 		// The Unreal engine allows you to place more than one key at one time value:
 		// displace the extra keys. This assumes that Unreal's keys are always ordered.
-		float KeyTime = bBakeKeys ? (KeyIndex * InterpLength) / KeyCount : Key.InVal;
-		if (KeyTimes.Num() && KeyTime < KeyTimes[KeyIndex-1] + FLT_TOLERANCE)
+		float KeyTime = bBakeKeys ? (KeyIndex * InterpLength) / KeyCount : Curve->Points[KeyIndex].InVal;
+		if(KeyTimes.Num() && KeyTime < KeyTimes[KeyIndex-1] + FLT_TOLERANCE)
 		{
 			KeyTime = KeyTimes[KeyIndex-1] + 0.01f; // Add 1 millisecond to the timing of this key.
 		}
 		KeyTimes.Add(KeyTime);
 	}
-
 	// Write out the key values from the curve to the FBX curve.
 	FbxCurve->KeyModifyBegin();
 	for (int32 KeyIndex = 0; KeyIndex < KeyCount; ++KeyIndex)
@@ -1552,8 +1552,6 @@ void FFbxExporter::ExportMoveSubTrack(FbxAnimCurve* FbxCurve, const ANSICHAR* Ch
 		float OutValue = (CurveIndex == 0) ? FinalOutVec.X : (CurveIndex == 1) ? FinalOutVec.Y : FinalOutVec.Z;
 		float FbxKeyValue = bNegative ? -OutValue : OutValue;
 
-		FInterpCurvePoint<float>& Key = Curve->Points[KeyIndex];
-
 		// Add a new key to the FBX curve
 		FbxTime Time;
 		FbxAnimCurveKey FbxKey;
@@ -1562,14 +1560,16 @@ void FFbxExporter::ExportMoveSubTrack(FbxAnimCurve* FbxCurve, const ANSICHAR* Ch
 
 		FbxAnimCurveDef::EInterpolationType Interpolation = FbxAnimCurveDef::eInterpolationConstant;
 		FbxAnimCurveDef::ETangentMode Tangent = FbxAnimCurveDef::eTangentAuto;
-		ConvertInterpToFBX(Key.InterpMode, Interpolation, Tangent);
-
+		
 		if (bBakeKeys || Interpolation != FbxAnimCurveDef::eInterpolationCubic)
 		{
 			FbxCurve->KeySet(FbxKeyIndex, Time, (float)FbxKeyValue, Interpolation, Tangent);
 		}
 		else
 		{
+			FInterpCurvePoint<float>& Key = Curve->Points[KeyIndex];
+			ConvertInterpToFBX(Key.InterpMode, Interpolation, Tangent);
+
 			// Setup tangents for bezier curves. Avoid this for keys created from baking 
 			// transforms since there is no tangent info created for these types of keys. 
 			if( Interpolation == FbxAnimCurveDef::eInterpolationCubic )
