#Matinee: Export

- Added Keep Hierarchy option when export

diff --git a/Engine/Source/Editor/Matinee/Private/Matinee.cpp b/Engine/Source/Editor/Matinee/Private/Matinee.cpp
--- a/Engine/Source/Editor/Matinee/Private/Matinee.cpp
+++ b/Engine/Source/Editor/Matinee/Private/Matinee.cpp
@@ -430,6 +430,7 @@ void FMatinee::InitMatinee(const EToolkitMode::Type Mode, const TSharedPtr< clas
 	bEditingCrosshairEnabled = false;
 	bEditingGridEnabled = false;
 	bBakeTransforms = false;
+	bKeepHierarchy = false;
 	bAllowKeyframeBarSelection = false;
 	bAllowKeyframeTextSelection = false;
 	bLockCameraPitch = true;
@@ -1075,6 +1076,11 @@ void FMatinee::BindCommands()
 		FCanExecuteAction(),
 		FIsActionChecked::CreateSP(this, &FMatinee::IsBakeTransformsToggled)
 		);
+	ToolkitCommands->MapAction(Commands.FileExportKeepHierarchy,
+		FExecuteAction::CreateSP(this, &FMatinee::OnToggleKeepHierarchy),
+		FCanExecuteAction(),
+		FIsActionChecked::CreateSP(this, &FMatinee::IsKeepHierarchyToggled)
+		);
 	
 	ToolkitCommands->MapAction( Commands.DeleteSelectedKeys, FExecuteAction::CreateSP(this, &FMatinee::OnDeleteSelectedKeys) );
 	ToolkitCommands->MapAction( Commands.DuplicateKeys, FExecuteAction::CreateSP(this, &FMatinee::OnMenuDuplicateSelectedKeys) );
diff --git a/Engine/Source/Editor/Matinee/Private/Matinee.h b/Engine/Source/Editor/Matinee/Private/Matinee.h
--- a/Engine/Source/Editor/Matinee/Private/Matinee.h
+++ b/Engine/Source/Editor/Matinee/Private/Matinee.h
@@ -403,7 +403,17 @@ public:
 	 * Updates the checked-menu item for baking transforms
 	 */
 	bool IsBakeTransformsToggled();
+	
+	/**
+	 * Called when the user toggles the ability to export with current hierarchy
+	 */
+	void OnToggleKeepHierarchy();
 
+	/**
+	 * Updates the checked-menu item for baking transforms
+	 */
+	bool IsKeepHierarchyToggled();
+	
 	/** Called when the 'Export Sound Cue Info' command is issued */
 	void OnExportSoundCueInfoCommand();
 	
@@ -1311,6 +1321,9 @@ protected:
 	/** When true, a key will be exported every frame instead of just the keys that user created. */
 	bool bBakeTransforms;
 
+	/** When true, the hierarchical relationship of objects will be kept when exported. */
+	bool bKeepHierarchy;
+
 	/** If true, clicking on a keyframe bar (such as one representing the duration of an audio cue, etc.) will cause a selection */
 	bool bAllowKeyframeBarSelection;
 
diff --git a/Engine/Source/Editor/Matinee/Private/MatineeActions.cpp b/Engine/Source/Editor/Matinee/Private/MatineeActions.cpp
--- a/Engine/Source/Editor/Matinee/Private/MatineeActions.cpp
+++ b/Engine/Source/Editor/Matinee/Private/MatineeActions.cpp
@@ -39,6 +39,7 @@ void FMatineeCommands::RegisterCommands()
 	UI_COMMAND(ExportSoundCueInfo, "Export Sound Cue Info", "Exports Sound Cue info", EUserInterfaceActionType::Button, FInputGesture());
 	UI_COMMAND(ExportAnimInfo, "Export Animation Track Info", "Exports Animation Track info", EUserInterfaceActionType::Button, FInputGesture());
 	UI_COMMAND(FileExportBakeTransforms, "Bake Transforms on Export", "Exports a key every frame, instead of just user created keys", EUserInterfaceActionType::ToggleButton, FInputGesture());
+	UI_COMMAND(FileExportKeepHierarchy, "Keep Hierarchy on Export", "Exports all objects to be in the same hierarchy in engine", EUserInterfaceActionType::ToggleButton, FInputGesture());
 
 	UI_COMMAND(DeleteSelectedKeys, "Delete Keys", "Deletes selected keys", EUserInterfaceActionType::Button, FInputGesture());
 	UI_COMMAND(DuplicateKeys, "Duplicate Keys", "Duplicates selected keys", EUserInterfaceActionType::Button, FInputGesture());
diff --git a/Engine/Source/Editor/Matinee/Private/MatineeActions.h b/Engine/Source/Editor/Matinee/Private/MatineeActions.h
--- a/Engine/Source/Editor/Matinee/Private/MatineeActions.h
+++ b/Engine/Source/Editor/Matinee/Private/MatineeActions.h
@@ -47,7 +47,7 @@ public:
 	TSharedPtr<FUICommandInfo> ExportSoundCueInfo;
 	TSharedPtr<FUICommandInfo> ExportAnimInfo;
 	TSharedPtr<FUICommandInfo> FileExportBakeTransforms;
-
+	TSharedPtr<FUICommandInfo> FileExportKeepHierarchy;
 	
 	TSharedPtr<FUICommandInfo> DeleteSelectedKeys;
 	TSharedPtr<FUICommandInfo> DuplicateKeys;
diff --git a/Engine/Source/Editor/Matinee/Private/MatineeMenus.cpp b/Engine/Source/Editor/Matinee/Private/MatineeMenus.cpp
--- a/Engine/Source/Editor/Matinee/Private/MatineeMenus.cpp
+++ b/Engine/Source/Editor/Matinee/Private/MatineeMenus.cpp
@@ -1580,6 +1580,7 @@ void FMatinee::OnContextTrackExportAnimFBX()
 					// Export the Matinee information to a COLLADA document.
 					Exporter->CreateDocument();
 					Exporter->SetTrasformBaking(bBakeTransforms);
+					Exporter->SetKeepHierarchy(bKeepHierarchy);
 
 					// Export the anim sequences
 					TArray<UAnimSequence*> AnimSequences;
@@ -1970,6 +1971,7 @@ void FMatinee::OnContextGroupExportAnimFBX()
 							// Export the Matinee information to an FBX document.
 							Exporter->CreateDocument();
 							Exporter->SetTrasformBaking(bBakeTransforms);
+							Exporter->SetKeepHierarchy(bKeepHierarchy);
 
 							// Export the animation sequences in the group by sampling the skeletal mesh over the
 							// duration of the matinee sequence
@@ -3840,6 +3842,7 @@ void FMatinee::OnMenuExport()
 				// Export the Matinee information to an FBX file
 				Exporter->CreateDocument();
 				Exporter->SetTrasformBaking(bBakeTransforms);
+				Exporter->SetKeepHierarchy(bKeepHierarchy);
 
 				const bool bSelectedOnly = false;
 				// Export the persistent level and all of it's actors
@@ -4218,6 +4221,23 @@ bool FMatinee::IsBakeTransformsToggled()
 	return bBakeTransforms;
 }
 
+/**
+ * Called when the user toggles the ability to export a key every frame. 
+ */
+void FMatinee::OnToggleKeepHierarchy()
+{
+	bKeepHierarchy = !bKeepHierarchy;
+}
+
+/**
+ * Updates the checked-menu item for baking transforms
+ */
+bool FMatinee::IsKeepHierarchyToggled()
+{
+	return bKeepHierarchy;
+}
+
+
 void FMatinee::OnMenuReduceKeys()
 {
 	ReduceKeys();
@@ -4990,9 +5010,10 @@ void FMatinee::ExtendDefaultToolbarMenu()
 			}
 			InMenuBarBuilder.EndSection();
 
-			InMenuBarBuilder.BeginSection("FileBake");
+			InMenuBarBuilder.BeginSection("Export");
 			{
 				InMenuBarBuilder.AddMenuEntry(FMatineeCommands::Get().FileExportBakeTransforms);
+				InMenuBarBuilder.AddMenuEntry(FMatineeCommands::Get().FileExportKeepHierarchy);
 			}
 			InMenuBarBuilder.EndSection();
 		}
diff --git a/Engine/Source/Editor/UnrealEd/Private/Fbx/FbxMainExport.cpp b/Engine/Source/Editor/UnrealEd/Private/Fbx/FbxMainExport.cpp
--- a/Engine/Source/Editor/UnrealEd/Private/Fbx/FbxMainExport.cpp
+++ b/Engine/Source/Editor/UnrealEd/Private/Fbx/FbxMainExport.cpp
@@ -1067,16 +1067,42 @@ FbxNode* FFbxExporter::ExportActor(AActor* Actor, AMatineeActor* InMatineeActor,
 		}
 
 		ActorNode = FbxNode::Create(Scene, TCHAR_TO_ANSI(*FbxNodeName));
-		Scene->GetRootNode()->AddChild(ActorNode);
+		
+
+		AActor* ParentActor = Actor->GetAttachParentActor();
+		FbxNode* ParentNode;
+		FVector ActorLocation, ActorRotation, ActorScale;
+		if (bKeepHierarchy && ParentActor)
+		{
+			// this doesn't work with skeletalmeshcomponent
+			ParentNode = FindActor(ParentActor);
+			check (ParentNode);
+
+			// Set the default position of the actor on the transforms
+			// The transformation is different from FBX's Z-up: invert the Y-axis for translations and the Y/Z angle values in rotations.
+			const FTransform RelativeTransform = Actor->GetTransform().GetRelativeTransform(ParentActor->GetTransform());
+			ActorLocation = RelativeTransform.GetTranslation();
+			ActorRotation = RelativeTransform.GetRotation().Euler();
+			ActorScale = RelativeTransform.GetScale3D();
+		}
+		else
+		{
+			ParentNode = Scene->GetRootNode();
+			// Set the default position of the actor on the transforms
+			// The transformation is different from FBX's Z-up: invert the Y-axis for translations and the Y/Z angle values in rotations.
+			ActorLocation = Actor->GetActorLocation();
+			ActorRotation = Actor->GetActorRotation().Euler();
+			ActorScale = Actor->GetRootComponent() ? Actor->GetRootComponent()->RelativeScale3D : FVector(1.f,1.f,1.f);
+		}
 
+		ParentNode->AddChild(ActorNode);
 		FbxActors.Add(Actor, ActorNode);
 
 		// Set the default position of the actor on the transforms
 		// The transformation is different from FBX's Z-up: invert the Y-axis for translations and the Y/Z angle values in rotations.
-		ActorNode->LclTranslation.Set(Converter.ConvertToFbxPos(Actor->GetActorLocation()));
-		ActorNode->LclRotation.Set(Converter.ConvertToFbxRot(Actor->GetActorRotation().Euler()));
-		const FVector DrawScale3D = Actor->GetRootComponent() ? Actor->GetRootComponent()->RelativeScale3D : FVector(1.f,1.f,1.f);
-		ActorNode->LclScaling.Set(Converter.ConvertToFbxScale(DrawScale3D));
+		ActorNode->LclTranslation.Set(Converter.ConvertToFbxPos(ActorLocation));
+		ActorNode->LclRotation.Set(Converter.ConvertToFbxRot(ActorRotation));
+		ActorNode->LclScaling.Set(Converter.ConvertToFbxScale(ActorScale));
 	
 		// For cameras and lights: always add a Y-pivot rotation to get the correct coordinate system.
 		if (Actor->IsA(ACameraActor::StaticClass()) || Actor->IsA(ALight::StaticClass()))
@@ -1411,22 +1437,36 @@ void FFbxExporter::ExportAnimatedVector(FbxAnimCurve* FbxCurve, const char* Chan
 				KeyRotation = FRotator( FQuat::MakeFromEuler(MoveTrack->EulerTrack.Points[KeyIndex].OutVal) );
 			}
 
-			FVector WorldSpacePos;
-			FRotator WorldSpaceRotator;
-			MoveTrack->ComputeWorldSpaceKeyTransform(
-				MoveTrackInst,
-				KeyPosition,
-				KeyRotation,
-				WorldSpacePos,			// Out
-				WorldSpaceRotator );	// Out
-
-			if( bPosCurve )
+			if (bKeepHierarchy)
 			{
-				FinalOutVec = WorldSpacePos;
+				if(bPosCurve)
+				{
+					FinalOutVec = KeyPosition;
+				}
+				else
+				{
+					FinalOutVec = KeyRotation.Euler();
+				}
 			}
 			else
 			{
-				FinalOutVec = WorldSpaceRotator.Euler();
+				FVector WorldSpacePos;
+				FRotator WorldSpaceRotator;
+				MoveTrack->ComputeWorldSpaceKeyTransform(
+					MoveTrackInst,
+					KeyPosition,
+					KeyRotation,
+					WorldSpacePos,			// Out
+					WorldSpaceRotator);	// Out
+
+				if(bPosCurve)
+				{
+					FinalOutVec = WorldSpacePos;
+				}
+				else
+				{
+					FinalOutVec = WorldSpaceRotator.Euler();
+				}
 			}
 		}
 
diff --git a/Engine/Source/Editor/UnrealEd/Public/MatineeExporter.h b/Engine/Source/Editor/UnrealEd/Public/MatineeExporter.h
--- a/Engine/Source/Editor/UnrealEd/Public/MatineeExporter.h
+++ b/Engine/Source/Editor/UnrealEd/Public/MatineeExporter.h
@@ -39,6 +39,11 @@ public:
 		bBakeKeys = bBakeTransforms;
 	}
 
+	void SetKeepHierarchy(bool bInKeepHierarchy)
+	{
+		bKeepHierarchy = bInKeepHierarchy;
+	}
+
 	/**
 	 * Exports the basic scene information to a file.
 	 */
@@ -106,6 +111,8 @@ protected:
 
 	/** When true, a key will exported per frame at the set frames-per-second (FPS). */
 	bool bBakeKeys;
+	/** When true, we'll export with hierarchical relation of attachment with relative transform */
+	bool bKeepHierarchy;
 };
 
 #endif // __MATINEEEXPORTER_H__
