Fix crash mirroring bsp

diff --git a/Engine/Source/Editor/UnrealEd/Private/UnrealEdSrv.cpp b/Engine/Source/Editor/UnrealEd/Private/UnrealEdSrv.cpp
--- a/Engine/Source/Editor/UnrealEd/Private/UnrealEdSrv.cpp
+++ b/Engine/Source/Editor/UnrealEd/Private/UnrealEdSrv.cpp
@@ -1736,7 +1736,7 @@ static void MirrorActors(const FVector& MirrorScale)
 
 		const FVector PivotLocation = GLevelEditorModeTools().PivotLocation;
 		ABrush* Brush = Cast< ABrush >( Actor );
-		if( Brush )
+		if( Brush && Brush->Brush )
 		{
 			Brush->Brush->Modify();
 
