Keeping materail export in sync with recent BaseColor visualization gamma space changes

diff --git a/Engine/Source/Editor/UnrealEd/Private/EditorExporters.cpp b/Engine/Source/Editor/UnrealEd/Private/EditorExporters.cpp
--- a/Engine/Source/Editor/UnrealEd/Private/EditorExporters.cpp
+++ b/Engine/Source/Editor/UnrealEd/Private/EditorExporters.cpp
@@ -2781,7 +2781,7 @@ namespace MaterialExportUtils
 			OutFlattenMaterial.DiffuseSize.Y > 0)
 		{
 			static const FName BaseColorName("BaseColor");
-			const float BaseColorGamma = 1.0f; // BaseColor material already do transformation to gamma space
+			const float BaseColorGamma = 2.2f; // BaseColor to gamma space
 			RenderSceneToTexture(Scene, BaseColorName, ViewOrigin, ViewRotationMatrix, ProjectionMatrix, HiddenPrimitives, 
 				OutFlattenMaterial.DiffuseSize, BaseColorGamma, OutFlattenMaterial.DiffuseSamples);
 		}
