Added UEditorEngine::UnsnappedClickLocation, and changed Play From Here to use it rather than the snap-adjusted click location.

diff --git a/Engine/Source/Editor/UnrealEd/Classes/Editor/EditorEngine.h b/Engine/Source/Editor/UnrealEd/Classes/Editor/EditorEngine.h
--- a/Engine/Source/Editor/UnrealEd/Classes/Editor/EditorEngine.h
+++ b/Engine/Source/Editor/UnrealEd/Classes/Editor/EditorEngine.h
@@ -336,6 +336,9 @@ class UNREALED_API UEditorEngine : public UEngine
 	UPROPERTY()
 	class UPackage* ParentContext;
 
+	UPROPERTY()
+	FVector UnsnappedClickLocation;
+
 	UPROPERTY()
 	FVector ClickLocation;
 
diff --git a/Engine/Source/Editor/UnrealEd/Private/Kismet2/DebuggerCommands.cpp b/Engine/Source/Editor/UnrealEd/Private/Kismet2/DebuggerCommands.cpp
--- a/Engine/Source/Editor/UnrealEd/Private/Kismet2/DebuggerCommands.cpp
+++ b/Engine/Source/Editor/UnrealEd/Private/Kismet2/DebuggerCommands.cpp
@@ -841,7 +841,7 @@ void FPlayWorldCommandCallbacks::StartPlayFromHere()
 	// Figure out the start location of the player
 	UCapsuleComponent*	DefaultCollisionComponent = CastChecked<UCapsuleComponent>(PlayerStartClass->GetDefaultObject<AActor>()->GetRootComponent());
 	FVector				CollisionExtent = FVector(DefaultCollisionComponent->GetScaledCapsuleRadius(),DefaultCollisionComponent->GetScaledCapsuleRadius(),DefaultCollisionComponent->GetScaledCapsuleHalfHeight());
-	FVector				StartLocation = GEditor->ClickLocation + GEditor->ClickPlane * (FVector::BoxPushOut(GEditor->ClickPlane,CollisionExtent) + 0.1f);
+	FVector				StartLocation = GEditor->UnsnappedClickLocation + GEditor->ClickPlane * (FVector::BoxPushOut(GEditor->ClickPlane,CollisionExtent) + 0.1f);
 
 	FRotator StartRotation = FRotator::ZeroRotator;
 
diff --git a/Engine/Source/Editor/UnrealEd/Private/LevelEditorViewport.cpp b/Engine/Source/Editor/UnrealEd/Private/LevelEditorViewport.cpp
--- a/Engine/Source/Editor/UnrealEd/Private/LevelEditorViewport.cpp
+++ b/Engine/Source/Editor/UnrealEd/Private/LevelEditorViewport.cpp
@@ -1038,6 +1038,7 @@ bool FLevelEditorViewportClient::UpdateDropPreviewActors(int32 MouseX, int32 Mou
 
 	const FActorPositionTraceResult TraceResult = FActorPositioning::TraceWorldForPositionWithDefault(Cursor, *View, &DraggingActors);
 
+	GEditor->UnsnappedClickLocation = TraceResult.Location;
 	GEditor->ClickLocation = TraceResult.Location;
 	GEditor->ClickPlane = FPlane(TraceResult.Location, TraceResult.SurfaceNormal);
 
@@ -1188,6 +1189,7 @@ bool FLevelEditorViewportClient::DropObjectsAtCoordinates(int32 MouseX, int32 Mo
 
 		const FActorPositionTraceResult TraceResult = FActorPositioning::TraceWorldForPositionWithDefault(Cursor, *View);
 		
+		GEditor->UnsnappedClickLocation = TraceResult.Location;
 		GEditor->ClickLocation = TraceResult.Location;
 		GEditor->ClickPlane = FPlane(TraceResult.Location, TraceResult.SurfaceNormal);
 
@@ -1860,7 +1862,8 @@ void FLevelEditorViewportClient::ProcessClick(FSceneView& View, HHitProxy* HitPr
 			bool bHit = GWorld->SweepSingle(CheckResult, Click.GetOrigin(), Click.GetOrigin() + Click.GetDirection() * HALF_WORLD_MAX, FQuat::Identity, FCollisionShape::MakeBox(FVector(1.f)), BoxParams, FCollisionObjectQueryParams(ECC_WorldStatic));
 
 			if( bHit )
-			{	
+			{
+				GEditor->UnsnappedClickLocation = CheckResult.Location;
 				GEditor->ClickLocation = CheckResult.Location;
 				GEditor->ClickPlane = FPlane(CheckResult.Location,CheckResult.Normal);
 			}
@@ -2353,6 +2356,7 @@ bool FLevelEditorViewportClient::InputKey(FViewport* Viewport, int32 ControllerI
 	{
 		const FViewportCursorLocation Cursor(View, this, HitX, HitY);
 		const FActorPositionTraceResult TraceResult = FActorPositioning::TraceWorldForPositionWithDefault(Cursor, *View);
+		GEditor->UnsnappedClickLocation = TraceResult.Location;
 		GEditor->ClickLocation = TraceResult.Location;
 		GEditor->ClickPlane = FPlane(TraceResult.Location, TraceResult.SurfaceNormal);
 
diff --git a/Engine/Source/Editor/UnrealEd/Private/UnrealEdSrv.cpp b/Engine/Source/Editor/UnrealEd/Private/UnrealEdSrv.cpp
--- a/Engine/Source/Editor/UnrealEd/Private/UnrealEdSrv.cpp
+++ b/Engine/Source/Editor/UnrealEd/Private/UnrealEdSrv.cpp
@@ -1709,7 +1709,9 @@ bool UUnrealEdEngine::Exec_Pivot( const TCHAR* Str, FOutputDevice& Ar )
 
 		if( Count > 0 )
 		{
-			ClickLocation = Center / Count;
+			FVector CenterLocation = Center / Count;
+			UnsnappedClickLocation = CenterLocation;
+			ClickLocation = CenterLocation;
 			ClickPlane = FPlane(0.f,0.f,0.f,0.f);
 
 			SetPivot( ClickLocation, false, false );
