If a transform property is marked as "Use 3D Widget", the widget rotation is now correctly represented when in local space mode.

diff --git a/Engine/Source/Editor/UnrealEd/Private/EdMode.cpp b/Engine/Source/Editor/UnrealEd/Private/EdMode.cpp
--- a/Engine/Source/Editor/UnrealEd/Private/EdMode.cpp
+++ b/Engine/Source/Editor/UnrealEd/Private/EdMode.cpp
@@ -426,6 +426,33 @@ FVector FEdMode::GetWidgetLocation() const
 	return Owner->PivotLocation;
 }
 
+bool FEdMode::GetCustomDrawingCoordinateSystem(FMatrix& InMatrix, void* InData)
+{
+	if (UsesPropertyWidgets())
+	{
+		AActor* SelectedActor = GetFirstSelectedActorInstance();
+		if (SelectedActor != NULL)
+		{
+			if (EditedPropertyName != TEXT(""))
+			{
+				if (bEditedPropertyIsTransform)
+				{
+					FTransform LocalTM = GetPropertyValueByName<FTransform>(SelectedActor, EditedPropertyName, EditedPropertyIndex);
+					InMatrix = FRotationMatrix::Make((LocalTM * SelectedActor->GetTransform()).GetRotation());
+					return true;
+				}
+				else
+				{
+					InMatrix = FRotationMatrix::Make(SelectedActor->GetTransform().GetRotation());
+					return true;
+				}
+			}
+		}
+	}
+
+	return false;
+}
+
 /**
  * Lets the mode determine if it wants to draw the widget or not.
  */
diff --git a/Engine/Source/Editor/UnrealEd/Public/EdMode.h b/Engine/Source/Editor/UnrealEd/Public/EdMode.h
--- a/Engine/Source/Editor/UnrealEd/Public/EdMode.h
+++ b/Engine/Source/Editor/UnrealEd/Public/EdMode.h
@@ -78,7 +78,7 @@ public:
 	virtual bool GetCursor(EMouseCursor::Type& OutCursor) const { return false; }
 
 	virtual bool ShouldDrawBrushWireframe( AActor* InActor ) const { return true; }
-	virtual bool GetCustomDrawingCoordinateSystem( FMatrix& InMatrix, void* InData ) { return 0; }
+	virtual bool GetCustomDrawingCoordinateSystem(FMatrix& InMatrix, void* InData);
 	virtual bool GetCustomInputCoordinateSystem( FMatrix& InMatrix, void* InData ) { return 0; }
 
 	/** If Rotation Snap should be enabled for this mode*/ 
