[INTEGRATE] Change 2433769 by bruce.nesbit@bnesbit_Releases on 2015/02/05 11:52:57

	Fixed non-unity build error caused by 2431516

diff --git a/Engine/Source/Editor/UnrealEd/Private/FeaturePack/FeaturePack.cpp b/Engine/Source/Editor/UnrealEd/Private/FeaturePack/FeaturePack.cpp
--- a/Engine/Source/Editor/UnrealEd/Private/FeaturePack/FeaturePack.cpp
+++ b/Engine/Source/Editor/UnrealEd/Private/FeaturePack/FeaturePack.cpp
@@ -1,7 +1,7 @@
 // Copyright 1998-2014 Epic Games, Inc. All Rights Reserved.
 
 #include "UnrealEd.h"
-#include "FeaturePack/FeaturePack.h"
+#include "FeaturePack.h"
 #include "AssetToolsModule.h"
 #include "AssetEditorManager.h"
 #include "AssetRegistryModule.h"
diff --git a/Engine/Source/Editor/UnrealEd/Private/UnrealEdMisc.cpp b/Engine/Source/Editor/UnrealEd/Private/UnrealEdMisc.cpp
--- a/Engine/Source/Editor/UnrealEd/Private/UnrealEdMisc.cpp
+++ b/Engine/Source/Editor/UnrealEd/Private/UnrealEdMisc.cpp
@@ -37,6 +37,7 @@
 #include "HotReloadInterface.h"
 #include "PerformanceMonitor.h"
 #include "Engine/WorldComposition.h"
+#include "FeaturePack.h"
 
 #define LOCTEXT_NAMESPACE "UnrealEd"
 
@@ -257,6 +258,7 @@ void FUnrealEdMisc::OnInit()
 		bool bMapLoaded = false;
 
 		// Insert any feature packs if required. We need to do this before we try and load a map since any pack may contain a map
+		FFeaturePack FeaturePackHandler;
 		FeaturePackHandler.ImportPendingPacks();
 
 		FString ParsedMapName;
diff --git a/Engine/Source/Editor/UnrealEd/Public/UnrealEd.h b/Engine/Source/Editor/UnrealEd/Public/UnrealEd.h
--- a/Engine/Source/Editor/UnrealEd/Public/UnrealEd.h
+++ b/Engine/Source/Editor/UnrealEd/Public/UnrealEd.h
@@ -30,7 +30,7 @@
 
 #include "MRUList.h"
 #include "Projects.h"
-#include "FeaturePack/FeaturePack.h"
+
 
 //#include "../Private/GeomFitUtils.h"
 
diff --git a/Engine/Source/Editor/UnrealEd/Public/UnrealEdMisc.h b/Engine/Source/Editor/UnrealEd/Public/UnrealEdMisc.h
--- a/Engine/Source/Editor/UnrealEd/Public/UnrealEdMisc.h
+++ b/Engine/Source/Editor/UnrealEd/Public/UnrealEdMisc.h
@@ -329,8 +329,5 @@ private:
 	/** Handle to the registered OnMapChanged delegate. */
 	FDelegateHandle OnMapChangedDelegateHandle;
 
-	FTimerHandle EditorAnalyticsHeartbeatTimerHandle;
-	
-	/* Handles insertion of feature packs (EG Starter content) at editor startup */
-	FFeaturePack FeaturePackHandler;
+	FTimerHandle EditorAnalyticsHeartbeatTimerHandle;	
 };
diff --git a/Engine/Source/Editor/UnrealEd/UnrealEd.Build.cs b/Engine/Source/Editor/UnrealEd/UnrealEd.Build.cs
--- a/Engine/Source/Editor/UnrealEd/UnrealEd.Build.cs
+++ b/Engine/Source/Editor/UnrealEd/UnrealEd.Build.cs
@@ -14,7 +14,8 @@ public class UnrealEd : ModuleRules
 			{
 				"Editor/UnrealEd/Private",
 				"Editor/UnrealEd/Private/Settings",
-				"Editor/PackagesDialog/Public",
+				"Editor/UnrealEd/Private/FeaturePack",
+				"Editor/PackagesDialog/Public",				
 				"Developer/DerivedDataCache/Public",
 				"Developer/TargetPlatform/Public",
 			}
