[INTEGRATE] Change 2433786 by [REDACTED_PERSON][REDACTED_HANDLE] on 2015/02/05 12:00:28

	Improve naming of scriptable components so they are more discoverable for those looking to add behaviors to actors.

diff --git a/Engine/Source/Editor/UnrealEd/Private/ComponentTypeRegistry.cpp b/Engine/Source/Editor/UnrealEd/Private/ComponentTypeRegistry.cpp
--- a/Engine/Source/Editor/UnrealEd/Private/ComponentTypeRegistry.cpp
+++ b/Engine/Source/Editor/UnrealEd/Private/ComponentTypeRegistry.cpp
@@ -143,7 +143,7 @@ void FComponentTypeRegistryData::RefreshComponentList()
 
 	if (GetDefault<UEditorExperimentalSettings>()->bBlueprintableComponents)
 	{
-		FString NewComponentsHeading = LOCTEXT("NewComponentsHeading", "New").ToString();
+		FString NewComponentsHeading = LOCTEXT("NewComponentsHeading", "Scripting").ToString();
 		// Add new C++ component class
 		FComponentClassComboEntryPtr NewClassHeader = MakeShareable(new FComponentClassComboEntry(NewComponentsHeading));
 		ComponentClassList.Add(NewClassHeader);
diff --git a/Engine/Source/Editor/UnrealEd/Private/SComponentClassCombo.cpp b/Engine/Source/Editor/UnrealEd/Private/SComponentClassCombo.cpp
--- a/Engine/Source/Editor/UnrealEd/Private/SComponentClassCombo.cpp
+++ b/Engine/Source/Editor/UnrealEd/Private/SComponentClassCombo.cpp
@@ -363,7 +363,7 @@ FText SComponentClassCombo::GetFriendlyComponentName(FComponentClassComboEntryPt
 	}
 	else if (Entry->GetComponentCreateAction() == EComponentCreateAction::CreateNewBlueprintClass )
 	{
-		FriendlyComponentName = LOCTEXT("NewBlueprintComponentFriendlyName", "New Blueprint Component...").ToString();
+		FriendlyComponentName = LOCTEXT("NewBlueprintComponentFriendlyName", "New Blueprint Script Component...").ToString();
 	}
 	else
 	{
