[INTEGRATE] Change 2433804 by [REDACTED_PERSON][REDACTED_HANDLE] on 2015/02/05 12:16:18

	empty character and empty pawn now say "Empty Character " and "Empty Pawn" in the placement browser UI.

diff --git a/Engine/Source/Editor/UnrealEd/Private/Factories/ActorFactory.cpp b/Engine/Source/Editor/UnrealEd/Private/Factories/ActorFactory.cpp
--- a/Engine/Source/Editor/UnrealEd/Private/Factories/ActorFactory.cpp
+++ b/Engine/Source/Editor/UnrealEd/Private/Factories/ActorFactory.cpp
@@ -1223,7 +1223,7 @@ UActorFactoryCharacter
 UActorFactoryCharacter::UActorFactoryCharacter(const FObjectInitializer& ObjectInitializer)
 	: Super(ObjectInitializer)
 {
-	DisplayName = LOCTEXT("ActorFactoryCharacterDisplayName", "Character");
+	DisplayName = LOCTEXT("ActorFactoryCharacterDisplayName", "Empty Character");
 	NewActorClass = ACharacter::StaticClass();
 }
 
@@ -1250,7 +1250,7 @@ UActorFactoryPawn
 UActorFactoryPawn::UActorFactoryPawn(const FObjectInitializer& ObjectInitializer)
 	: Super(ObjectInitializer)
 {
-	DisplayName = LOCTEXT("ActorFactoryPawnDisplayName", "Pawn");
+	DisplayName = LOCTEXT("ActorFactoryPawnDisplayName", "Empty Pawn");
 	NewActorClass = APawn::StaticClass();
 }
 
