[INTEGRATE] Change 2433843 by [REDACTED_PERSON][REDACTED_HANDLE] on 2015/02/05 12:37:31

	[UE-9104] Making sure we reset macro lib's super when they're left with bad REINST_ super classes.

diff --git a/Engine/Source/Editor/UnrealEd/Private/Kismet2/BlueprintEditorUtils.cpp b/Engine/Source/Editor/UnrealEd/Private/Kismet2/BlueprintEditorUtils.cpp
--- a/Engine/Source/Editor/UnrealEd/Private/Kismet2/BlueprintEditorUtils.cpp
+++ b/Engine/Source/Editor/UnrealEd/Private/Kismet2/BlueprintEditorUtils.cpp
@@ -1211,6 +1211,23 @@ UClass* FBlueprintEditorUtils::RegenerateBlueprintClass(UBlueprint* Blueprint, U
 		{
 			// Just refresh all nodes in macro blueprints, but don't recompil
 			FBlueprintEditorUtils::RefreshAllNodes(Blueprint);
+
+			if (ClassToRegenerate != nullptr)
+			{
+				UClass* OldSuperClass = ClassToRegenerate->GetSuperClass();
+				if ((OldSuperClass != nullptr) && OldSuperClass->HasAnyClassFlags(CLASS_NewerVersionExists))
+				{
+					UClass* NewSuperClass = OldSuperClass->GetAuthoritativeClass();
+					ensure(NewSuperClass == Blueprint->ParentClass);
+
+					// in case the macro's super class was re-instanced (it 
+					// would have re-parented this to a REINST_ class), for non-
+					// macro blueprints this would normally be reset in 
+					// CompileBlueprint (but since we don't compile macros, we 
+					// need to fix this up here)
+					ClassToRegenerate->SetSuperStruct(NewSuperClass);
+				}
+			}
 		}
 		else
 		{
