[INTEGRATE] Change 2433848 by [REDACTED_PERSON][REDACTED_HANDLE] on 2015/02/05 12:40:30

	[UE-8152] - Changed asset dropping on a selected actor the level viewport to apply to all selected actors even when only one actor is selected.
	The result is that even if only one actor is selected, dropping a material onto that selected actor will attempt to apply the material to every component.

diff --git a/Engine/Source/Editor/UnrealEd/Private/LevelEditorViewport.cpp b/Engine/Source/Editor/UnrealEd/Private/LevelEditorViewport.cpp
--- a/Engine/Source/Editor/UnrealEd/Private/LevelEditorViewport.cpp
+++ b/Engine/Source/Editor/UnrealEd/Private/LevelEditorViewport.cpp
@@ -1232,9 +1232,9 @@ bool FLevelEditorViewportClient::DropObjectsAtCoordinates(int32 MouseX, int32 Mo
 			{
 				FNavigationLockContext LockNavigationUpdates(TargetActor->GetWorld(), ENavigationLockReason::SpawnOnDragEnter, bCreateDropPreview);
 
-				// If more than one actor is selected and the target actor is one of them, we should drop onto all selected actors
+				// If the target actor is selected, we should drop onto all selected actors
 				// otherwise, we should drop only onto the target object
-				const bool bDropOntoSelectedActors = TargetActor->IsSelected() && GEditor->GetSelectedActorCount() > 1;
+				const bool bDropOntoSelectedActors = TargetActor->IsSelected();
 				const bool bCanApplyToComponent = AttemptApplyObjToComponent(DroppedObjects[0], TargetComponent, TargetMaterialSlot, true);
 				if (bOnlyDropOnTarget || !bDropOntoSelectedActors || !bCanApplyToComponent)
 				{
