Improvement on gc frequency for cook on the fly.

diff --git a/Engine/Source/Editor/UnrealEd/Classes/Commandlets/CookCommandlet.h b/Engine/Source/Editor/UnrealEd/Classes/Commandlets/CookCommandlet.h
--- a/Engine/Source/Editor/UnrealEd/Classes/Commandlets/CookCommandlet.h
+++ b/Engine/Source/Editor/UnrealEd/Classes/Commandlets/CookCommandlet.h
@@ -34,6 +34,8 @@ class UCookCommandlet
 	bool bUnversioned;
 	/** Generate manifests for building streaming install packages */
 	bool bGenerateStreamingInstallManifests;
+	/** Max memory the cooker should use before forcing a gc */
+	uint64 MaxMemoryAllowance;
 	/** All commandline tokens */
 	TArray<FString> Tokens;
 	/** All commandline switches */
@@ -116,6 +118,9 @@ private:
 
 	void MaybeMarkPackageAsAlreadyLoaded(UPackage *Package);
 
+	/** See if the cooker has exceeded max memory allowance in this case the cooker should force a garbage collection */
+	bool HasExceededMaxMemory() const;
+
 	/** Gets the output directory respecting any command line overrides */
 	FString GetOutputDirectoryOverride() const;
 
diff --git a/Engine/Source/Editor/UnrealEd/Classes/CookOnTheSide/CookOnTheFlyServer.h b/Engine/Source/Editor/UnrealEd/Classes/CookOnTheSide/CookOnTheFlyServer.h
--- a/Engine/Source/Editor/UnrealEd/Classes/CookOnTheSide/CookOnTheFlyServer.h
+++ b/Engine/Source/Editor/UnrealEd/Classes/CookOnTheSide/CookOnTheFlyServer.h
@@ -678,6 +678,7 @@ public:
 		COSR_CookedPackage		= 0x00000002,
 		COSR_ErrorLoadingPackage= 0x00000004,
 		COSR_RequiresGC			= 0x00000008,
+		COSR_WaitingOnCache		= 0x00000010,
 	};
 
 
diff --git a/Engine/Source/Editor/UnrealEd/Private/Commandlets/CookCommandlet.cpp b/Engine/Source/Editor/UnrealEd/Private/Commandlets/CookCommandlet.cpp
--- a/Engine/Source/Editor/UnrealEd/Private/Commandlets/CookCommandlet.cpp
+++ b/Engine/Source/Editor/UnrealEd/Private/Commandlets/CookCommandlet.cpp
@@ -24,6 +24,7 @@
 #include "PhysicsPublic.h"
 #include "CookerSettings.h"
 #include "ShaderCompiler.h"
+#include "MemoryMisc.h"
 
 DEFINE_LOG_CATEGORY_STATIC(LogCookCommandlet, Log, All);
 
@@ -134,50 +135,61 @@ bool UCookCommandlet::CookOnTheFly( FGuid InstanceId, int32 Timeout, bool bForce
 		static const float CookOnTheSideTimeSlice = 10.0f;
 		TickResults = CookOnTheFlyServer->TickCookOnTheSide(CookOnTheSideTimeSlice, NonMapPackageCountSinceLastGC);
 
-		bCookedAMapSinceLastGC |= TickResults & UCookOnTheFlyServer::COSR_CookedMap;
-		if ( TickResults & (UCookOnTheFlyServer::COSR_CookedMap | UCookOnTheFlyServer::COSR_CookedPackage))
+		bCookedAMapSinceLastGC |= ((TickResults & UCookOnTheFlyServer::COSR_RequiresGC) != 0);
+		if ( TickResults & (UCookOnTheFlyServer::COSR_CookedMap | UCookOnTheFlyServer::COSR_CookedPackage | UCookOnTheFlyServer::COSR_WaitingOnCache))
 		{
 			LastCookActionTime = FPlatformTime::Seconds();
 		}
 
-
-		while ( (CookOnTheFlyServer->HasCookRequests() == false) && !GIsRequestingExit)
+		if (NonMapPackageCountSinceLastGC > 0)
 		{
-				
+			if ( NonMapPackageCountSinceLastGC > PackagesPerGC)
 			{
-				if (NonMapPackageCountSinceLastGC > 0)
-				{
-					// We should GC if we have packages to collect and we've been idle for some time.
-					bShouldGC = (NonMapPackageCountSinceLastGC > PackagesPerGC) || 
-						((FPlatformTime::Seconds() - LastCookActionTime) >= IdleTimeToGC);
-				}
+				UE_LOG(LogCookCommandlet, Display, TEXT("Cooker has exceeded max number of non map packages since last gc"));
+				bShouldGC |= true;
+			}
 
-				// delay the gc until we process some unsolicited packages
-				if ( bCookedAMapSinceLastGC )
-				{
-					UE_LOG( LogCookCommandlet, Display, TEXT("Delaying map gc because we have unsolicited cook requests") );
-					bShouldGC |= bCookedAMapSinceLastGC;
-				}
+			if (((FPlatformTime::Seconds() - LastCookActionTime) >= IdleTimeToGC))
+			{
+				UE_LOG(LogCookCommandlet, Display, TEXT("Cooker has been idle for long time gc"));
+				bShouldGC |= true;
+			}
+		}
 
-				if (bShouldGC)
-				{
-					bShouldGC = false;
-					bCookedAMapSinceLastGC = false;
-					NonMapPackageCountSinceLastGC = 0;
+		if ( bCookedAMapSinceLastGC )
+		{
+			UE_LOG(LogCookCommandlet, Display, TEXT("Cooker cooked a map since last gc collecting garbage"));
+			bShouldGC |= true;
+		}
 
-					UE_LOG(LogCookCommandlet, Display, TEXT("GC..."));
+		if ( !bShouldGC && HasExceededMaxMemory() )
+		{
+			UE_LOG(LogCookCommandlet, Display, TEXT("Cooker has exceeded max memory usage collecting garbage"));
+			bShouldGC |= true;
+		}
 
-					CollectGarbage( RF_Native );
-				}
-				else
-				{
-					CookOnTheFlyServer->TickRecompileShaderRequests();
+		// we don't want to gc if we are waiting on cache of objects. this could clean up objects which we will need to reload next frame
+		if (bShouldGC && ((TickResults & UCookOnTheFlyServer::COSR_WaitingOnCache)==0) )
+		{
+			bShouldGC = false;
+			bCookedAMapSinceLastGC = false;
+			NonMapPackageCountSinceLastGC = 0;
 
-					FPlatformProcess::Sleep(0.0f);
-				}
-			}
+			UE_LOG(LogCookCommandlet, Display, TEXT("GC..."));
+
+			CollectGarbage( RF_Native );
+		}
 
 
+		// force at least a tick shader compilation even if we are requesting stuff
+		CookOnTheFlyServer->TickRecompileShaderRequests();
+		GShaderCompilingManager->ProcessAsyncResults(true, false);
+
+
+		while ( (CookOnTheFlyServer->HasCookRequests() == false) && !GIsRequestingExit)
+		{
+			CookOnTheFlyServer->TickRecompileShaderRequests();
+
 			// Shaders need to be updated
 			GShaderCompilingManager->ProcessAsyncResults(true, false);
 
@@ -510,6 +522,8 @@ int32 UCookCommandlet::Main(const FString& CmdLineParams)
 	bIterativeCooking = Switches.Contains(TEXT("ITERATE"));
 	bSkipEditorContent = Switches.Contains(TEXT("SKIPEDITORCONTENT")); // This won't save out any packages in Engine/COntent/Editor*
 
+	MaxMemoryAllowance = 8LL * 1024LL * 1024LL * 1024LL;
+
 	if (bLeakTest)
 	{
 		for (FObjectIterator It; It; ++It)
@@ -1213,7 +1227,11 @@ bool UCookCommandlet::NewCook( const TArray<ITargetPlatform*>& Platforms, TArray
 
 		bShouldGC |= (TickResults & UCookOnTheFlyServer::COSR_RequiresGC)!=0;
 
-		if (bShouldGC)
+		bShouldGC |= HasExceededMaxMemory();
+
+
+		// don't clean up if we are waiting on cache of cooked data
+		if (bShouldGC && ((TickResults & UCookOnTheFlyServer::COSR_WaitingOnCache) == 0) )
 		{
 			bShouldGC = false;
 			NonMapPackageCountSinceLastGC = 0;
@@ -1235,12 +1253,21 @@ bool UCookCommandlet::NewCook( const TArray<ITargetPlatform*>& Platforms, TArray
 	}
 
 	return true;
+}
 
+bool UCookCommandlet::HasExceededMaxMemory() const
+{
+	const FPlatformMemoryStats MemStats = FPlatformMemory::GetStats();
 
-
+	uint64 UsedMemory = MemStats.UsedPhysical + MemStats.UsedVirtual;
+	if ( (UsedMemory >= MaxMemoryAllowance) && 
+		(MaxMemoryAllowance > 0u) )
+	{
+		return true;
+	}
+	return false;
 }
 
-
 void UCookCommandlet::ProcessDeferredCommands()
 {
 #if PLATFORM_MAC
diff --git a/Engine/Source/Editor/UnrealEd/Private/CookOnTheFlyServer.cpp b/Engine/Source/Editor/UnrealEd/Private/CookOnTheFlyServer.cpp
--- a/Engine/Source/Editor/UnrealEd/Private/CookOnTheFlyServer.cpp
+++ b/Engine/Source/Editor/UnrealEd/Private/CookOnTheFlyServer.cpp
@@ -1189,6 +1189,7 @@ uint32 UCookOnTheFlyServer::TickCookOnTheSide( const float TimeSlice, uint32 &Co
 			{
 				// reque the current task and process it next tick
 				CookRequests.EnqueueUnique(ToBuild, true);
+				Result |= COSR_WaitingOnCache;
 				break; // break out of the package tick loop
 			}
 		}
@@ -2337,6 +2338,12 @@ void UCookOnTheFlyServer::AddFileToCook( TArray<FString>& InOutFilesToCook, cons
 { 
 	if (!FPackageName::IsScriptPackage(InFilename))
 	{
+		/*if ( !InOutFilesToCook.Contains(InFilename) )
+		{
+			int Index = FMath::RandRange(0,InOutFilesToCook.Num()-1);
+			InOutFilesToCook.Insert(InFilename, Index );
+		}*/
+
 		InOutFilesToCook.AddUnique(InFilename);
 	}
 }
