#PHYSICS: If smaller tha min, just assign min, not default prim size, which is huge

Merging

//depot/UE4-GDC2015/Engine/Source/Editor/UnrealEd/Private/PhysicsAssetUtils.cpp

to //depot/UE4/Engine/Source/Editor/UnrealEd/Private/PhysicsAssetUtils.cpp

diff --git a/Engine/Source/Editor/UnrealEd/Private/PhysicsAssetUtils.cpp b/Engine/Source/Editor/UnrealEd/Private/PhysicsAssetUtils.cpp
--- a/Engine/Source/Editor/UnrealEd/Private/PhysicsAssetUtils.cpp
+++ b/Engine/Source/Editor/UnrealEd/Private/PhysicsAssetUtils.cpp
@@ -402,7 +402,8 @@ bool CreateCollisionFromBone( UBodySetup* bs, USkeletalMesh* skelMesh, int32 Bon
 	// If the primitive is going to be too small - just use some default numbers and let the user tweak.
 	if( MinRad < MinAllowedSize )
 	{
-		BoxExtent = FVector(DefaultPrimSize, DefaultPrimSize, DefaultPrimSize);
+		// change min allowed size to be min, not DefaultPrimSize
+		BoxExtent = FVector(MinAllowedSize, MinAllowedSize, MinAllowedSize);
 	}
 
 	FVector BoneOrigin = ElementTransform.TransformPosition( BoxCenter );
