#Simplification: Added "regenerate" button so that artists can refresh LOD when they change mesh

Merging

//depot/UE4-GDC2015/Engine/Source/Editor/...

to //depot/UE4-GDC2015/Engine/Source/Editor/...

diff --git a/Engine/Source/Editor/Persona/Private/PersonaMeshDetails.cpp b/Engine/Source/Editor/Persona/Private/PersonaMeshDetails.cpp
--- a/Engine/Source/Editor/Persona/Private/PersonaMeshDetails.cpp
+++ b/Engine/Source/Editor/Persona/Private/PersonaMeshDetails.cpp
@@ -336,7 +336,7 @@ FReply FSkelMeshReductionSettingsLayout::OnApplyChanges()
 {
 	if (ParentLODSettings.IsValid())
 	{
-		ParentLODSettings.Pin()->ApplyChanges();
+		ParentLODSettings.Pin()->ApplyChanges(false);
 	}
 	return FReply::Handled();
 }
@@ -584,10 +584,10 @@ void FPersonaMeshDetails::CustomizeLODSettingsCategories(IDetailLayoutBuilder& D
 		[
 			SNew(SButton)
 			.OnClicked(this, &FPersonaMeshDetails::OnApplyChanges)
-			.IsEnabled(this, &FPersonaMeshDetails::IsApplyNeeded)
+			.IsEnabled(this, &FPersonaMeshDetails::IsGenerateAvailable)
 			[
 				SNew(STextBlock)
-				.Text(LOCTEXT("ApplyChanges", "Apply Changes"))
+				.Text(this, &FPersonaMeshDetails::GetApplyButtonText)
 				.Font(DetailLayout.GetDetailFont())
 			]
 		];
@@ -626,10 +626,24 @@ void FPersonaMeshDetails::OnLODCountCommitted(int32 InValue, ETextCommit::Type C
 
 FReply FPersonaMeshDetails::OnApplyChanges()
 {
-	ApplyChanges();
+	ApplyChanges(true);
 	return FReply::Handled();
 }
 
+FText FPersonaMeshDetails::GetApplyButtonText() const
+{
+	if (IsApplyNeeded())	
+	{
+		return LOCTEXT("ApplyChanges", "Apply Changes");
+	}
+	else if (IsGenerateAvailable())
+	{
+		return LOCTEXT("Regenerate", "Regenerate");
+	}
+
+	return LOCTEXT("ApplyChanges", "Apply Changes");
+}
+
 void FPersonaMeshDetails::RefreshBonesToRemove(TArray<FBoneReference>& InOutBonesToRemove, int32 LODIndex)
 {
 	if (LODIndex == 0)
@@ -696,11 +710,13 @@ void FPersonaMeshDetails::RefreshBonesToRemove(TArray<FBoneReference>& InOutBone
 	}
 }
 
-void FPersonaMeshDetails::ApplyChanges()
+void FPersonaMeshDetails::ApplyChanges(bool bForceUpdate)
 {
 	USkeletalMesh* SkelMesh = PersonaPtr->GetMesh();
 	check(SkelMesh);
 
+	const bool bForceUpdateLOD = (IsApplyNeeded() == false && bForceUpdate);
+
 	FSkeletalMeshUpdateContext UpdateContext;
 	UpdateContext.SkeletalMesh = SkelMesh;
 	UpdateContext.AssociatedComponents.Push(PersonaPtr->PreviewComponent);
@@ -747,7 +763,7 @@ void FPersonaMeshDetails::ApplyChanges()
 			SettingIndex = Settings.Add(Setting);
 		}
 
-		FLODUtilities::SimplifySkeletalMesh(UpdateContext, Settings);
+		FLODUtilities::SimplifySkeletalMesh(UpdateContext, Settings, bForceUpdateLOD);
 
 		// save reduction settings
 		for (int32 LODIdx = 1; LODIdx < LODCount; ++LODIdx)
@@ -774,6 +790,16 @@ void FPersonaMeshDetails::UpdateLODNames()
 	LODNames.Add(MakeShareable(new FString(FText::Format(NSLOCTEXT("LODSettingsLayout", "LODLevel_Import", "Import LOD Level {0}"), FText::AsNumber(LODCount)).ToString())));
 }
 
+bool FPersonaMeshDetails::IsGenerateAvailable() const
+{
+	if (IsApplyNeeded())
+	{
+		return true;
+	}
+
+	// if you have LOD
+	return (LODCount > 1);
+}
 bool FPersonaMeshDetails::IsApplyNeeded() const
 {
 	USkeletalMesh* SkelMesh = PersonaPtr->GetMesh();
diff --git a/Engine/Source/Editor/Persona/Private/PersonaMeshDetails.h b/Engine/Source/Editor/Persona/Private/PersonaMeshDetails.h
--- a/Engine/Source/Editor/Persona/Private/PersonaMeshDetails.h
+++ b/Engine/Source/Editor/Persona/Private/PersonaMeshDetails.h
@@ -209,7 +209,9 @@ private:
 public:
 
 	bool IsApplyNeeded() const;
-	void ApplyChanges();
+	bool IsGenerateAvailable() const;
+	void ApplyChanges(bool bForceUpdate);
+	FText GetApplyButtonText() const;
 
 	USkeletalMesh* GetMesh()
 	{ 
diff --git a/Engine/Source/Editor/UnrealEd/Private/LODUtilities.cpp b/Engine/Source/Editor/UnrealEd/Private/LODUtilities.cpp
--- a/Engine/Source/Editor/UnrealEd/Private/LODUtilities.cpp
+++ b/Engine/Source/Editor/UnrealEd/Private/LODUtilities.cpp
@@ -71,7 +71,40 @@ void FLODUtilities::RemoveLOD(FSkeletalMeshUpdateContext& UpdateContext, int32 D
 	}
 }
 
-void FLODUtilities::SimplifySkeletalMesh( FSkeletalMeshUpdateContext& UpdateContext, TArray<FSkeletalMeshOptimizationSettings> &InSettings )
+void FLODUtilities::SimplifySkeletalMeshLOD( USkeletalMesh* SkeletalMesh, const FSkeletalMeshOptimizationSettings& InSetting, int32 DesiredLOD )
+{
+	IMeshUtilities& MeshUtilities = FModuleManager::Get().LoadModuleChecked<IMeshUtilities>("MeshUtilities");
+	IMeshReduction* MeshReduction = MeshUtilities.GetMeshReductionInterface();
+
+	check (MeshReduction && MeshReduction->IsSupported());
+	{
+		FFormatNamedArguments Args;
+		Args.Add(TEXT("DesiredLOD"), DesiredLOD);
+		Args.Add(TEXT("SkeletalMeshName"), FText::FromString(SkeletalMesh->GetName()));
+		const FText StatusUpdate = FText::Format(NSLOCTEXT("UnrealEd", "MeshSimp_GeneratingLOD_F", "Generating LOD{DesiredLOD} for {SkeletalMeshName}..."), Args);
+		GWarn->BeginSlowTask(StatusUpdate, true);
+	}
+
+	if(MeshReduction->ReduceSkeletalMesh(SkeletalMesh, DesiredLOD, InSetting, true))
+	{
+		check(SkeletalMesh->LODInfo.Num() >= 2);
+		SkeletalMesh->MarkPackageDirty();
+#if WITH_APEX_CLOTHING
+		ApexClothingUtils::ReImportClothingSectionsFromClothingAsset(SkeletalMesh);
+#endif// #if WITH_APEX_CLOTHING
+	}
+	else
+	{
+		// Simplification failed! Warn the user.
+		FFormatNamedArguments Args;
+		Args.Add(TEXT("SkeletalMeshName"), FText::FromString(SkeletalMesh->GetName()));
+		const FText Message = FText::Format(NSLOCTEXT("UnrealEd", "MeshSimp_GenerateLODFailed_F", "An error occurred while simplifying the geometry for mesh '{SkeletalMeshName}'.  Consider adjusting simplification parameters and re-simplifying the mesh."), Args);
+		FMessageDialog::Open(EAppMsgType::Ok, Message);
+	}
+	GWarn->EndSlowTask();
+}
+
+void FLODUtilities::SimplifySkeletalMesh( FSkeletalMeshUpdateContext& UpdateContext, TArray<FSkeletalMeshOptimizationSettings> &InSettings, bool bForceRegenerate )
 {
 	USkeletalMesh* SkeletalMesh = UpdateContext.SkeletalMesh;
 	IMeshUtilities& MeshUtilities = FModuleManager::Get().LoadModuleChecked<IMeshUtilities>("MeshUtilities");
@@ -85,37 +118,13 @@ void FLODUtilities::SimplifySkeletalMesh( FSkeletalMeshUpdateContext& UpdateCont
 			uint32 DesiredLOD = SettingIndex + 1;
 
 			// check whether reduction settings are same or not
-			if (SkeletalMesh->LODInfo.IsValidIndex(DesiredLOD) 
+			if (!bForceRegenerate && SkeletalMesh->LODInfo.IsValidIndex(DesiredLOD) 
 			 && SkeletalMesh->LODInfo[DesiredLOD].ReductionSettings == InSettings[SettingIndex])
 			{
 				continue;
 			}
 
-			{
-				FFormatNamedArguments Args;
-				Args.Add( TEXT("DesiredLOD"), DesiredLOD );
-				Args.Add( TEXT("SkeletalMeshName"), FText::FromString(SkeletalMesh->GetName()) );
-				const FText StatusUpdate = FText::Format( NSLOCTEXT("UnrealEd", "MeshSimp_GeneratingLOD_F", "Generating LOD{DesiredLOD} for {SkeletalMeshName}..." ), Args );
-				GWarn->BeginSlowTask( StatusUpdate, true );
-			}
-
-			if (MeshReduction->ReduceSkeletalMesh(SkeletalMesh, DesiredLOD, InSettings[SettingIndex], true))
-			{
-				check( SkeletalMesh->LODInfo.Num() >= 2 );
-				SkeletalMesh->MarkPackageDirty();
-#if WITH_APEX_CLOTHING
-				ApexClothingUtils::ReImportClothingSectionsFromClothingAsset(SkeletalMesh);
-#endif// #if WITH_APEX_CLOTHING
-			}
-			else
-			{
-				// Simplification failed! Warn the user.
-				FFormatNamedArguments Args;
-				Args.Add( TEXT("SkeletalMeshName"), FText::FromString(SkeletalMesh->GetName()) );
-				const FText Message = FText::Format( NSLOCTEXT("UnrealEd", "MeshSimp_GenerateLODFailed_F", "An error occurred while simplifying the geometry for mesh '{SkeletalMeshName}'.  Consider adjusting simplification parameters and re-simplifying the mesh." ), Args );
-				FMessageDialog::Open( EAppMsgType::Ok, Message );
-			}
-			GWarn->EndSlowTask();
+			SimplifySkeletalMeshLOD( SkeletalMesh, InSettings[SettingIndex], DesiredLOD );
 		}
 
 		//Notify calling system of change
diff --git a/Engine/Source/Editor/UnrealEd/Public/LODUtilities.h b/Engine/Source/Editor/UnrealEd/Public/LODUtilities.h
--- a/Engine/Source/Editor/UnrealEd/Public/LODUtilities.h
+++ b/Engine/Source/Editor/UnrealEd/Public/LODUtilities.h
@@ -35,7 +35,7 @@ public:
 	 * @param UpdateContext - The skeletal mesh and actor components to operate on.
 	 * @param InSettings   - The user settings passed from the Simplification tool.
 	 */
-	static void SimplifySkeletalMesh( FSkeletalMeshUpdateContext& UpdateContext, TArray<FSkeletalMeshOptimizationSettings> &InSettings );
+	static void SimplifySkeletalMesh( FSkeletalMeshUpdateContext& UpdateContext, TArray<FSkeletalMeshOptimizationSettings> &InSettings, bool bForceRegenerate = false );
 
 	/**
 	 * Refresh LOD Change
@@ -49,4 +49,14 @@ public:
 
 private:
 	FLODUtilities() {}
+
+	/**
+	 *	Simplifies the static mesh based upon various user settings for DesiredLOD
+	 *  This is private function that gets called by SimplifySkeletalMesh
+	 *
+	 * @param SkeletalMesh - The skeletal mesh and actor components to operate on.
+	 * @param InSetting   - The user settings passed from the Simplification tool.
+	 * @param DesiredLOD - Desired LOD
+	 */
+	static void SimplifySkeletalMeshLOD( USkeletalMesh* SkeletalMesh, const FSkeletalMeshOptimizationSettings& InSetting, int32 DesiredLOD );
 };
\ No newline at end of file
