Fix measure tool not working when matinee is open

diff --git a/Engine/Source/Editor/UnrealEd/Private/EditorModeInterpolation.cpp b/Engine/Source/Editor/UnrealEd/Private/EditorModeInterpolation.cpp
--- a/Engine/Source/Editor/UnrealEd/Private/EditorModeInterpolation.cpp
+++ b/Engine/Source/Editor/UnrealEd/Private/EditorModeInterpolation.cpp
@@ -220,7 +220,7 @@ bool FModeTool_InterpEdit::MouseMove(FEditorViewportClient* ViewportClient,FView
 
 bool FModeTool_InterpEdit::InputAxis(FEditorViewportClient* InViewportClient, FViewport* Viewport, int32 ControllerId, FKey Key, float Delta, float DeltaTime)
 {
-	if ( InViewportClient->GetCurrentWidgetAxis() == EAxisList::None )
+	if ( InViewportClient->GetCurrentWidgetAxis() == EAxisList::None && !InViewportClient->Viewport->KeyState(EKeys::MiddleMouseButton) )
 	{
 		// We need to set up a widget axis here to prevent our drag operation being co-opted by box/frustum selection
 		ELevelViewportType ViewportType = InViewportClient->ViewportType;
