Fix spelling error

diff --git a/Engine/Source/Editor/UnrealEd/Private/ComponentAssetBroker.cpp b/Engine/Source/Editor/UnrealEd/Private/ComponentAssetBroker.cpp
--- a/Engine/Source/Editor/UnrealEd/Private/ComponentAssetBroker.cpp
+++ b/Engine/Source/Editor/UnrealEd/Private/ComponentAssetBroker.cpp
@@ -366,7 +366,7 @@ void FComponentAssetBrokerage::InitializeMap()
 	}
 }
 
-void FComponentAssetBrokerage::RegisterBroker(TSharedPtr<IComponentAssetBroker> Broker, TSubclassOf<UActorComponent> InComponentClass, bool bSetAsPrimary, bool bMapCompnentForAssets)
+void FComponentAssetBrokerage::RegisterBroker(TSharedPtr<IComponentAssetBroker> Broker, TSubclassOf<UActorComponent> InComponentClass, bool bSetAsPrimary, bool bMapComponentForAssets)
 {
 	InitializeMap();
 
@@ -387,7 +387,7 @@ void FComponentAssetBrokerage::RegisterBroker(TSharedPtr<IComponentAssetBroker>
 		AssetToBrokerMap.FindOrAdd(AssetClass).Add(Broker);
 	}
 
-	if (bMapCompnentForAssets)
+	if (bMapComponentForAssets)
 	{
 		FComponentClassList& ValidComponentTypes = AssetToComponentClassMap.FindOrAdd(AssetClass);
 		if (bSetAsPrimary)
diff --git a/Engine/Source/Editor/UnrealEd/Public/ComponentAssetBroker.h b/Engine/Source/Editor/UnrealEd/Public/ComponentAssetBroker.h
--- a/Engine/Source/Editor/UnrealEd/Public/ComponentAssetBroker.h
+++ b/Engine/Source/Editor/UnrealEd/Public/ComponentAssetBroker.h
@@ -73,7 +73,7 @@ public:
 
 
 	/** Register a component class for a specified asset class */
-	static void RegisterBroker(TSharedPtr<IComponentAssetBroker> Broker, TSubclassOf<UActorComponent> InComponentClass, bool bSetAsPrimary, bool bMapCompnentForAssets);
+	static void RegisterBroker(TSharedPtr<IComponentAssetBroker> Broker, TSubclassOf<UActorComponent> InComponentClass, bool bSetAsPrimary, bool bMapComponentForAssets);
 
 	/** Unregister a component type for a specified asset class */
 	static void UnregisterBroker(TSharedPtr<IComponentAssetBroker> Broker);
