UE-9121: Crash when project attempts to autosave

Couldn't reproduce it. Changed the behavior that it now reports an error in the log and informs the user that auto-save failed instead of crashing.

diff --git a/Engine/Source/Editor/UnrealEd/Private/Editor.cpp b/Engine/Source/Editor/UnrealEd/Private/Editor.cpp
--- a/Engine/Source/Editor/UnrealEd/Private/Editor.cpp
+++ b/Engine/Source/Editor/UnrealEd/Private/Editor.cpp
@@ -4324,7 +4324,15 @@ bool UEditorEngine::IsPackageOKToSave(UPackage* InPackage, const FString& InFile
 {
 	TArray<FString>	AllStartupPackageNames;
 	appGetAllPotentialStartupPackageNames(AllStartupPackageNames, GEngineIni, false);
-	bool bIsStartupPackage = AllStartupPackageNames.Contains(FPackageName::FilenameToLongPackageName(InFilename));
+
+	FString ConvertedPackageName, ConversionError;
+	if (!FPackageName::TryConvertFilenameToLongPackageName(InFilename, ConvertedPackageName, &ConversionError))
+	{
+		Error->Logf(ELogVerbosity::Error, *FText::Format(NSLOCTEXT("UnrealEd", "CannotConvertPackageName", "Cannot save asset '{0}' as conversion of long package name failed. Reason: '{1}'."), FText::FromString(InFilename), FText::FromString(ConversionError)).ToString());
+		return false;
+	}
+
+	bool bIsStartupPackage = AllStartupPackageNames.Contains(ConvertedPackageName);
 
 	// Make sure that if the package is a startup package, the user indeed wants to save changes
 	if( !IsRunningCommandlet()																		&& // Don't prompt about saving startup packages when running UCC
diff --git a/Engine/Source/Editor/UnrealEd/Private/FileHelpers.cpp b/Engine/Source/Editor/UnrealEd/Private/FileHelpers.cpp
--- a/Engine/Source/Editor/UnrealEd/Private/FileHelpers.cpp
+++ b/Engine/Source/Editor/UnrealEd/Private/FileHelpers.cpp
@@ -2147,6 +2147,15 @@ bool FEditorFileUtils::AutosaveMap(const FString& AbsoluteAutosaveDir, const int
 }
 
 bool FEditorFileUtils::AutosaveContentPackages(const FString& AbsoluteAutosaveDir, const int32 AutosaveIndex, const bool bForceIfNotInList, const TSet< TWeakObjectPtr<UPackage> >& DirtyPackagesForAutoSave)
+{
+	auto Result = AutosaveContentPackagesEx(AbsoluteAutosaveDir, AutosaveIndex, bForceIfNotInList, DirtyPackagesForAutoSave);
+
+	check(Result != EAutosaveContentPackagesResult::Failure);
+
+	return Result == EAutosaveContentPackagesResult::Success;
+}
+
+EAutosaveContentPackagesResult::Type FEditorFileUtils::AutosaveContentPackagesEx(const FString& AbsoluteAutosaveDir, const int32 AutosaveIndex, const bool bForceIfNotInList, const TSet< TWeakObjectPtr<UPackage> >& DirtyPackagesForAutoSave)
 {
 	const FScopedBusyCursor BusyCursor;
 	double SaveStartTime = FPlatformTime::Seconds();
@@ -2180,7 +2189,10 @@ bool FEditorFileUtils::AutosaveContentPackages(const FString& AbsoluteAutosaveDi
 				}
 
 				const FString AutosaveFilename = GetAutoSaveFilename(CurPackage, AbsoluteAutosaveDir, AutosaveIndex, FPackageName::GetAssetPackageExtension());
-				GUnrealEd->Exec( NULL, *FString::Printf( TEXT("OBJ SAVEPACKAGE PACKAGE=\"%s\" FILE=\"%s\" SILENT=false AUTOSAVING=true"), *CurPackage->GetName(), *AutosaveFilename ) );
+				if (!GUnrealEd->Exec(nullptr, *FString::Printf(TEXT("OBJ SAVEPACKAGE PACKAGE=\"%s\" FILE=\"%s\" SILENT=false AUTOSAVING=true"), *CurPackage->GetName(), *AutosaveFilename)))
+				{
+					return EAutosaveContentPackagesResult::Failure;
+				}
 
 				// Re-mark the package as dirty, because autosaving it will have cleared the dirty flag
 				CurPackage->MarkPackageDirty();
@@ -2194,7 +2206,7 @@ bool FEditorFileUtils::AutosaveContentPackages(const FString& AbsoluteAutosaveDi
 		UE_LOG(LogFileHelpers, Log, TEXT("Auto-saving content packages took %.3f"), FPlatformTime::Seconds() - SaveStartTime );
 	}
 
-	return bSavedPkgs;
+	return bSavedPkgs ? EAutosaveContentPackagesResult::Success : EAutosaveContentPackagesResult::NothingToDo;
 }
 /**
  * Actually save a package. Prompting for Save as if necessary
diff --git a/Engine/Source/Editor/UnrealEd/Private/PackageAutoSaver.cpp b/Engine/Source/Editor/UnrealEd/Private/PackageAutoSaver.cpp
--- a/Engine/Source/Editor/UnrealEd/Private/PackageAutoSaver.cpp
+++ b/Engine/Source/Editor/UnrealEd/Private/PackageAutoSaver.cpp
@@ -151,8 +151,6 @@ void FPackageAutoSaver::AttemptAutoSave()
 		const bool bCanAutosave = CanAutoSave();
 		if (bCanAutosave)
 		{
-			// Auto-saving, so close any warning notifications.
-			CloseAutoSaveNotification(true);
 			bAutosaveHandled = true;
 
 			bIsAutoSaving = true;
@@ -168,7 +166,7 @@ void FPackageAutoSaver::AttemptAutoSave()
 			const int32 NewAutoSaveIndex = (AutoSaveIndex + 1) % 10;
 
 			bool bLevelSaved = false;
-			bool bAssetsSaved = false;
+			auto AssetsSaveResults = EAutosaveContentPackagesResult::NothingToDo;
 
 			if (LoadingSavingSettings->bAutoSaveMaps)
 			{
@@ -183,10 +181,16 @@ void FPackageAutoSaver::AttemptAutoSave()
 				// If we weren't saving content packages when we last ran an auto-save, we need to forcibly save any that are dirty
 				// as DirtyPackagesForAutoSave gets cleared each time an auto-save is run, so may not be up-to-date for the new settings
 				const bool bForceSaveContentPackages = !bWasAutoSavingContent && LoadingSavingSettings->bAutoSaveContent;
-				bAssetsSaved = FEditorFileUtils::AutosaveContentPackages(AutoSaveDir, NewAutoSaveIndex, bForceSaveContentPackages, DirtyPackagesForAutoSave);
+				AssetsSaveResults = FEditorFileUtils::AutosaveContentPackagesEx(AutoSaveDir, NewAutoSaveIndex, bForceSaveContentPackages, DirtyPackagesForAutoSave);
 			}
 
-			if (bLevelSaved || bAssetsSaved)
+			// Auto-saved, so close any warning notifications.
+			CloseAutoSaveNotification(
+				AssetsSaveResults == EAutosaveContentPackagesResult::Success ?	ECloseNotification::Success	:
+				AssetsSaveResults == EAutosaveContentPackagesResult::Failure ?	ECloseNotification::Failed	:
+																				ECloseNotification::Postponed);
+
+			if (bLevelSaved || AssetsSaveResults == EAutosaveContentPackagesResult::Success)
 			{
 				// If a level was actually saved, update the auto-save index
 				AutoSaveIndex = NewAutoSaveIndex;
@@ -513,6 +517,11 @@ void FPackageAutoSaver::UpdateAutoSaveNotification()
 }
 
 void FPackageAutoSaver::CloseAutoSaveNotification(const bool Success)
+{
+	CloseAutoSaveNotification(Success ? ECloseNotification::Success : ECloseNotification::Postponed);
+}
+
+void FPackageAutoSaver::CloseAutoSaveNotification(ECloseNotification::Type Type)
 {
 	// If a notification is open close it
 	if(bAutoSaveNotificationLaunched)
@@ -526,16 +535,21 @@ void FPackageAutoSaver::CloseAutoSaveNotification(const bool Success)
 			FText CloseMessage;
 
 			// Set the test on the notification based on whether it was a successful launch
-			if(Success)
+			if (Type == ECloseNotification::Success)
 			{
 				CloseMessage = NSLOCTEXT("AutoSaveNotify", "AutoSaving", "Saving");
 				CloseState = SNotificationItem::CS_Success;
 			}
-			else
+			else if (Type == ECloseNotification::Postponed)
 			{
 				CloseMessage = NSLOCTEXT("AutoSaveNotify", "AutoSavePostponed", "Autosave postponed");
 				CloseState = SNotificationItem::CS_None; // Set back to none rather than failed, as this is too harsh
 			}
+			else
+			{
+				CloseMessage = NSLOCTEXT("AutoSaveNotify", "AutoSaveFailed", "Auto-save failed. Please check the log for the details.");
+				CloseState = SNotificationItem::CS_Fail;
+			}
 
 			// update notification
 			NotificationItem->SetText(CloseMessage);
diff --git a/Engine/Source/Editor/UnrealEd/Private/PackageAutoSaver.h b/Engine/Source/Editor/UnrealEd/Private/PackageAutoSaver.h
--- a/Engine/Source/Editor/UnrealEd/Private/PackageAutoSaver.h
+++ b/Engine/Source/Editor/UnrealEd/Private/PackageAutoSaver.h
@@ -6,6 +6,16 @@
 
 class SNotificationItem;
 
+namespace ECloseNotification
+{
+	enum Type
+	{
+		Success,
+		Postponed,
+		Failed
+	};
+}
+
 /** A class to handle the creation, destruction, and restoration of auto-saved packages */
 class FPackageAutoSaver : public IPackageAutoSaver
 {
@@ -85,6 +95,9 @@ private:
 	/** Closes the auto-save warning notification if open, with an appropriate message based on whether it was successful */
 	void CloseAutoSaveNotification(const bool Success);
 
+	/** Closes the auto-save warning notification if open, with an appropriate message based on type */
+	void CloseAutoSaveNotification(ECloseNotification::Type Type);
+
 	/** Used as a callback for auto-save warning buttons, called when auto-save is forced early */
 	void OnAutoSaveSave();
 
diff --git a/Engine/Source/Editor/UnrealEd/Public/FileHelpers.h b/Engine/Source/Editor/UnrealEd/Public/FileHelpers.h
--- a/Engine/Source/Editor/UnrealEd/Public/FileHelpers.h
+++ b/Engine/Source/Editor/UnrealEd/Public/FileHelpers.h
@@ -18,6 +18,16 @@ enum EFileInteraction
 	FI_Export
 };
 
+namespace EAutosaveContentPackagesResult
+{
+	enum Type
+	{
+		Success,
+		NothingToDo,
+		Failure
+	};
+}
+
 /**
  * For saving map files through the main editor frame.
  */
@@ -142,6 +152,18 @@ public:
 	 */
 	static bool AutosaveContentPackages(const FString& AbsoluteAutosaveDir, const int32 AutosaveIndex, const bool bForceIfNotInList, const TSet< TWeakObjectPtr<UPackage> >& DirtyPackagesForAutoSave);
 
+	/**
+	 * Saves all asset packages to the specified directory.
+	 *
+	 * @param	AbsoluteAutosaveDir			Autosave directory.
+	 * @param	AutosaveIndex				Integer prepended to autosave filenames.
+	 * @param	bForceIfNotInList			Should the save be forced if the package is dirty, but not in DirtyPackagesForAutoSave?
+	 * @param	DirtyPackagesForAutoSave	A set of packages that are considered by the auto-save system to be dirty, you should check this to see if a package needs saving
+	 *
+	 * @return	Success if saved at least one faile. NothingToDo if there was nothing to save. Failure on at least one auto-save failure.
+	 */
+	static EAutosaveContentPackagesResult::Type AutosaveContentPackagesEx(const FString& AbsoluteAutosaveDir, const int32 AutosaveIndex, const bool bForceIfNotInList, const TSet< TWeakObjectPtr<UPackage> >& DirtyPackagesForAutoSave);
+
 	/**
 	 * Looks at all currently loaded packages and saves them if their "bDirty" flag is set, optionally prompting the user to select which packages to save)
 	 * 
