[INTEGRATE] Change 2434215 by [REDACTED_PERSON][REDACTED_HANDLE] on 2015/02/05 16:42:23

	Changes to BP created components on instances now pushed back to blueprint when Apply Changes used.

diff --git a/Engine/Source/Editor/UnrealEd/Private/Editor.cpp b/Engine/Source/Editor/UnrealEd/Private/Editor.cpp
--- a/Engine/Source/Editor/UnrealEd/Private/Editor.cpp
+++ b/Engine/Source/Editor/UnrealEd/Private/Editor.cpp
@@ -34,6 +34,7 @@
 #include "Kismet2/KismetEditorUtilities.h"
 #include "Kismet2/KismetDebugUtilities.h"
 #include "Editor/Kismet/Public/BlueprintEditorModule.h"
+#include "Engine/InheritableComponentHandler.h"
 
 #include "BlueprintUtilities.h"
 
@@ -7076,56 +7077,87 @@ namespace EditorUtilities
 		// If the source and target components do not match (e.g. context-specific), attempt to find a match in the target's array elsewhere
 		const int32 NumTargetComponents = TargetComponents.Num();
 		if( (SourceComponent != NULL) 
-			&& (TargetComponent != NULL) 
-			&& (SourceComponent->GetFName() != TargetComponent->GetFName()) )
+			&& ((TargetComponent == NULL) 
+				|| (SourceComponent->GetFName() != TargetComponent->GetFName()) ))
 		{
 			// Reset the target component since it doesn't match the source
 			TargetComponent = NULL;
 
-			// Attempt to locate a match elsewhere in the target's component list
-			const int32 StartingIndex = StartIndex + 1;
-			int32 FindTargetComponentIndex = (StartingIndex >= NumTargetComponents) ? 0 : StartingIndex;
-			do
+			if (NumTargetComponents > 0)
 			{
-				// If we found a match, update the target component and adjust the target index to the matching position
-				UActorComponent* FindTargetComponent = TargetComponents[ FindTargetComponentIndex ];
-				if( FindTargetComponent != NULL && SourceComponent->GetFName() == FindTargetComponent->GetFName() )
+				const bool bSourceIsArchetype = SourceComponent->HasAnyFlags(RF_ArchetypeObject);
+				// Attempt to locate a match elsewhere in the target's component list
+				const int32 StartingIndex = (bSourceIsArchetype ? StartIndex : StartIndex + 1);
+				int32 FindTargetComponentIndex = (StartingIndex >= NumTargetComponents) ? 0 : StartingIndex;
+				do
 				{
-					TargetComponent = FindTargetComponent;
-					StartIndex = FindTargetComponentIndex;
+					UActorComponent* FindTargetComponent = TargetComponents[ FindTargetComponentIndex ];
 
-					break;
-				}
+					// In the case that the SourceComponent is an Archetype there is a better than even chance the name won't match due to the way the SCS
+					// is set up, so we're actually going to reverse search
+					if (bSourceIsArchetype)
+					{
+						if ( SourceComponent == FindTargetComponent->GetArchetype())
+						{
+							TargetComponent = FindTargetComponent;
+							StartIndex = FindTargetComponentIndex;
+							break;
+						}
+					}
+					else
+					{
+						// If we found a match, update the target component and adjust the target index to the matching position
+						UActorComponent* FindTargetComponent = TargetComponents[ FindTargetComponentIndex ];
+						if( FindTargetComponent != NULL && SourceComponent->GetFName() == FindTargetComponent->GetFName() )
+						{
+							TargetComponent = FindTargetComponent;
+							StartIndex = FindTargetComponentIndex;
+							break;
+						}
+					}
 
-				// Increment the index counter, and loop back to 0 if necessary
-				if( ++FindTargetComponentIndex >= NumTargetComponents )
-				{
-					FindTargetComponentIndex = 0;
-				}
-			} while( FindTargetComponentIndex != StartIndex );
+					// Increment the index counter, and loop back to 0 if necessary
+					if( ++FindTargetComponentIndex >= NumTargetComponents )
+					{
+						FindTargetComponentIndex = 0;
+					}
 
-			// If we still haven't found a match and we're targeting a class default object
+				} while( FindTargetComponentIndex != StartIndex );
+			}
+
+			// If we still haven't found a match and we're targeting a class default object what we're really looking
+			// for is the Archetype
 			if(TargetComponent == NULL && TargetActor->HasAnyFlags(RF_ClassDefaultObject|RF_ArchetypeObject))
 			{
-				// Check for a Blueprint inheritance hierarchy
-				TArray<UBlueprint*> ParentBPStack;
-				UBlueprint::GetBlueprintHierarchyFromClass(TargetActor->GetClass(), ParentBPStack);
-				for(int32 StackIndex = 0; StackIndex < ParentBPStack.Num() && TargetComponent == NULL; ++StackIndex)
+				TargetComponent = CastChecked<UActorComponent>(SourceComponent->GetArchetype(), ECastCheckedType::NullAllowed);
+
+				// If the returned target component is not from the direct class of the actor we're targeting, we need to insert an inheritable component
+				if (TargetComponent && (TargetComponent->GetOuter() != TargetActor->GetClass()))
 				{
-					// For each Blueprint in the hierarchy, look for a match in the SCS if valid
-					UBlueprint* Blueprint = ParentBPStack[StackIndex];
-					if(Blueprint != NULL && Blueprint->SimpleConstructionScript != NULL)
+					// This component doesn't exist in the hierarchy anywhere and we're not going to modify the CDO, so we'll drop it
+					if (TargetComponent->HasAnyFlags(RF_ClassDefaultObject))
 					{
-						const TArray<USCS_Node*> AllSCSNodes = Blueprint->SimpleConstructionScript->GetAllNodes();
-						for(int32 SCSNodeIndex = 0; SCSNodeIndex < AllSCSNodes.Num(); ++SCSNodeIndex)
+						TargetComponent = nullptr;
+					}
+					else
+					{
+						UBlueprintGeneratedClass* BPGC = CastChecked<UBlueprintGeneratedClass>(TargetActor->GetClass());
+						UBlueprint* Blueprint = CastChecked<UBlueprint>(BPGC->ClassGeneratedBy);
+						UInheritableComponentHandler* InheritableComponentHandler = Blueprint->GetInheritableComponentHandler(true);
+						if (InheritableComponentHandler)
 						{
-							USCS_Node* SCS_Node = AllSCSNodes[SCSNodeIndex];
-							if(SCS_Node != NULL && SourceComponent->GetFName() == SCS_Node->GetVariableName())
+							BPGC = Cast<UBlueprintGeneratedClass>(BPGC->GetSuperClass());
+							USCS_Node* SCSNode = nullptr;
+							while (BPGC)
 							{
-								// Found it!
-								TargetComponent = SCS_Node->ComponentTemplate;
-								break;
+								SCSNode = BPGC->SimpleConstructionScript->FindSCSNode(SourceComponent->GetFName());
+								BPGC = (SCSNode ? nullptr : Cast<UBlueprintGeneratedClass>(BPGC->GetSuperClass()));
 							}
+							check(SCSNode);
+
+							FComponentKey Key(SCSNode);
+							check(InheritableComponentHandler->GetOverridenComponentTemplate(Key) == nullptr);
+							TargetComponent = InheritableComponentHandler->CreateOverridenComponentTemplate(Key);
 						}
 					}
 				}
@@ -7492,6 +7524,10 @@ namespace EditorUtilities
 			}
 		}
 
+		if (!bIsPreviewing && CopiedPropertyCount > 0 && TargetActor->HasAnyFlags(RF_ClassDefaultObject|RF_ArchetypeObject) && TargetActor->GetClass()->HasAllClassFlags(CLASS_CompiledFromBlueprint))
+		{
+			FBlueprintEditorUtils::PostEditChangeBlueprintActors(CastChecked<UBlueprint>(TargetActor->GetClass()->ClassGeneratedBy));
+		}
 
 		// If one of the changed properties was part of the actor's transformation, then we'll call PostEditMove too.
 		if( !bIsPreviewing && bTransformChanged )
diff --git a/Engine/Source/Runtime/Engine/Classes/Engine/BlueprintGeneratedClass.h b/Engine/Source/Runtime/Engine/Classes/Engine/BlueprintGeneratedClass.h
--- a/Engine/Source/Runtime/Engine/Classes/Engine/BlueprintGeneratedClass.h
+++ b/Engine/Source/Runtime/Engine/Classes/Engine/BlueprintGeneratedClass.h
@@ -463,6 +463,8 @@ public:
 	 */
 	static bool GetGeneratedClassesHierarchy(const UClass* InClass, TArray<const UBlueprintGeneratedClass*>& OutBPGClasses);
 
+	UInheritableComponentHandler* GetInheritableComponentHandler(const bool bCreateIfNecessary = false);
+
 	/** Find the object in the TemplateObjects array with the supplied name */
 	UActorComponent* FindComponentTemplateByName(const FName& TemplateName) const;
 
diff --git a/Engine/Source/Runtime/Engine/Classes/Engine/InheritableComponentHandler.h b/Engine/Source/Runtime/Engine/Classes/Engine/InheritableComponentHandler.h
--- a/Engine/Source/Runtime/Engine/Classes/Engine/InheritableComponentHandler.h
+++ b/Engine/Source/Runtime/Engine/Classes/Engine/InheritableComponentHandler.h
@@ -92,6 +92,8 @@ public:
 #endif
 
 public:
+	FComponentKey FindKey(const FName VariableName) const;
+
 	UActorComponent* GetOverridenComponentTemplate(FComponentKey Key) const;
 
 private:
diff --git a/Engine/Source/Runtime/Engine/Private/Blueprint.cpp b/Engine/Source/Runtime/Engine/Private/Blueprint.cpp
--- a/Engine/Source/Runtime/Engine/Private/Blueprint.cpp
+++ b/Engine/Source/Runtime/Engine/Private/Blueprint.cpp
@@ -1233,15 +1233,14 @@ UInheritableComponentHandler* UBlueprint::GetInheritableComponentHandler(bool bC
 	static const FBoolConfigValueHelper EnableInheritableComponents(TEXT("Kismet"), TEXT("bEnableInheritableComponents"), GEngineIni);
 	if (!EnableInheritableComponents)
 	{
-		return NULL;
+		return nullptr;
 	}
 
 	if (!InheritableComponentHandler && bCreateIfNecessary)
 	{
-		auto BPGC = CastChecked<UBlueprintGeneratedClass>(GeneratedClass);
+		UBlueprintGeneratedClass* BPGC = CastChecked<UBlueprintGeneratedClass>(GeneratedClass);
 		ensure(!BPGC->InheritableComponentHandler);
-		InheritableComponentHandler = NewObject<UInheritableComponentHandler>(BPGC, FName(TEXT("InheritableComponentHandler")));
-		BPGC->InheritableComponentHandler = InheritableComponentHandler;
+		InheritableComponentHandler = BPGC->GetInheritableComponentHandler(true);
 	}
 	return InheritableComponentHandler;
 }
diff --git a/Engine/Source/Runtime/Engine/Private/BlueprintGeneratedClass.cpp b/Engine/Source/Runtime/Engine/Private/BlueprintGeneratedClass.cpp
--- a/Engine/Source/Runtime/Engine/Private/BlueprintGeneratedClass.cpp
+++ b/Engine/Source/Runtime/Engine/Private/BlueprintGeneratedClass.cpp
@@ -7,6 +7,7 @@
 #include "Engine/SimpleConstructionScript.h"
 #include "Engine/SCS_Node.h"
 #include "Engine/LevelScriptActor.h"
+#
 
 #if WITH_EDITOR
 #include "BlueprintEditorUtils.h"
@@ -263,6 +264,23 @@ bool UBlueprintGeneratedClass::IsFunctionImplementedInBlueprint(FName InFunction
 	return Function && Function->GetOuter() && Function->GetOuter()->IsA(UBlueprintGeneratedClass::StaticClass());
 }
 
+UInheritableComponentHandler* UBlueprintGeneratedClass::GetInheritableComponentHandler(const bool bCreateIfNecessary)
+{
+	static const FBoolConfigValueHelper EnableInheritableComponents(TEXT("Kismet"), TEXT("bEnableInheritableComponents"), GEngineIni);
+	if (!EnableInheritableComponents)
+	{
+		return nullptr;
+	}
+
+	if (!InheritableComponentHandler && bCreateIfNecessary)
+	{
+		InheritableComponentHandler = NewNamedObject<UInheritableComponentHandler>(this, FName(TEXT("InheritableComponentHandler")));
+	}
+
+	return InheritableComponentHandler;
+}
+
+
 UObject* UBlueprintGeneratedClass::FindArchetype(UClass* ArchetypeClass, const FName ArchetypeName) const
 {
 	UObject* Archetype = nullptr;
@@ -283,7 +301,10 @@ UObject* UBlueprintGeneratedClass::FindArchetype(UClass* ArchetypeClass, const F
 			if (SCSNode)
 			{
 				Archetype = SCSNode->ComponentTemplate;
-				Class = nullptr;
+			}
+			else if (UInheritableComponentHandler* ICH = Class->GetInheritableComponentHandler())
+			{
+				Archetype = ICH->GetOverridenComponentTemplate(ICH->FindKey(ArchetypeName));
 			}
 
 			if (Archetype == nullptr)
@@ -291,6 +312,10 @@ UObject* UBlueprintGeneratedClass::FindArchetype(UClass* ArchetypeClass, const F
 				Archetype = static_cast<UObject*>(FindObjectWithOuter(Class, ArchetypeClass, ArchetypeName));
 				Class = (Archetype ? nullptr : Cast<UBlueprintGeneratedClass>(Class->GetSuperClass()));
 			}
+			else
+			{
+				Class = nullptr;
+			}
 		}
 	}
 
diff --git a/Engine/Source/Runtime/Engine/Private/InheritableComponentHandler.cpp b/Engine/Source/Runtime/Engine/Private/InheritableComponentHandler.cpp
--- a/Engine/Source/Runtime/Engine/Private/InheritableComponentHandler.cpp
+++ b/Engine/Source/Runtime/Engine/Private/InheritableComponentHandler.cpp
@@ -260,6 +260,18 @@ FComponentKey UInheritableComponentHandler::FindKey(UActorComponent* ComponentTe
 
 #endif
 
+FComponentKey UInheritableComponentHandler::FindKey(const FName VariableName) const
+{
+	for (const FComponentOverrideRecord& Record : Records)
+	{
+		if (Record.ComponentKey.VariableName == VariableName)
+		{
+			return Record.ComponentKey;
+		}
+	}
+	return FComponentKey();
+}
+
 UActorComponent* UInheritableComponentHandler::GetOverridenComponentTemplate(FComponentKey Key) const
 {
 	auto Record = FindRecord(Key);
