[INTEGRATE] Change 2434216 by [REDACTED_PERSON][REDACTED_HANDLE] on 2015/02/05 16:42:44

	Fix FBX crash importing a null texture.

diff --git a/Engine/Source/Editor/UnrealEd/Private/Fbx/FbxMaterialImport.cpp b/Engine/Source/Editor/UnrealEd/Private/Fbx/FbxMaterialImport.cpp
--- a/Engine/Source/Editor/UnrealEd/Private/Fbx/FbxMaterialImport.cpp
+++ b/Engine/Source/Editor/UnrealEd/Private/Fbx/FbxMaterialImport.cpp
@@ -48,6 +48,11 @@ using namespace UnFbx;
 
 UTexture* UnFbx::FFbxImporter::ImportTexture( FbxFileTexture* FbxTexture, bool bSetupAsNormalMap )
 {
+	if( !FbxTexture )
+	{
+		return nullptr;
+	}
+	
 	// create an unreal texture asset
 	UTexture* UnrealTexture = NULL;
 	FString Filename1 = ANSI_TO_TCHAR(FbxTexture->GetFileName());
