[INTEGRATE] Change 2434245 by [REDACTED_PERSON][REDACTED_HANDLE] on 2015/02/05 16:59:55

	Fix crash during packaging trying to play a sound in the editor to indicate that packaging is complete.

diff --git a/Engine/Source/Editor/UnrealEd/Private/Editor.cpp b/Engine/Source/Editor/UnrealEd/Private/Editor.cpp
--- a/Engine/Source/Editor/UnrealEd/Private/Editor.cpp
+++ b/Engine/Source/Editor/UnrealEd/Private/Editor.cpp
@@ -1963,7 +1963,7 @@ void UEditorEngine::PlayPreviewSound( USoundBase* Sound,  USoundNode* SoundNode
 void UEditorEngine::PlayEditorSound( const FString& SoundAssetName )
 {
 	// Only play sounds if the user has that feature enabled
-	if( GetDefault<ULevelEditorMiscSettings>()->bEnableEditorSounds )
+	if( GetDefault<ULevelEditorMiscSettings>()->bEnableEditorSounds && !GIsSavingPackage )
 	{
 		USoundBase* Sound = Cast<USoundBase>( StaticFindObject( USoundBase::StaticClass(), NULL, *SoundAssetName ) );
 		if( Sound == NULL )
