[INTEGRATE] Change 2434469 by [REDACTED_PERSON][REDACTED_HANDLE] on 2015/02/05 18:47:38

	Made "For Distribution" in project packaging not advanced

diff --git a/Engine/Source/Editor/UnrealEd/Classes/Settings/ProjectPackagingSettings.h b/Engine/Source/Editor/UnrealEd/Classes/Settings/ProjectPackagingSettings.h
--- a/Engine/Source/Editor/UnrealEd/Classes/Settings/ProjectPackagingSettings.h
+++ b/Engine/Source/Editor/UnrealEd/Classes/Settings/ProjectPackagingSettings.h
@@ -75,7 +75,7 @@ public:
 	 * If disabled, a development build will be created
 	 * Distribution builds are for publishing to the App Store
 	 */
-	UPROPERTY(config, EditAnywhere, Category=Project, AdvancedDisplay)
+	UPROPERTY(config, EditAnywhere, Category=Project)
 	bool ForDistribution;
 
 	/** If enabled, all content will be put into a single .pak file instead of many individual files (default = enabled). */
