[INTEGRATE] Change 2434536 by [REDACTED_PERSON][REDACTED_HANDLE] on 2015/02/05 19:25:09

	Fix crash when BSP is selected

diff --git a/Engine/Source/Editor/UnrealEd/Private/AssetSelection.cpp b/Engine/Source/Editor/UnrealEd/Private/AssetSelection.cpp
--- a/Engine/Source/Editor/UnrealEd/Private/AssetSelection.cpp
+++ b/Engine/Source/Editor/UnrealEd/Private/AssetSelection.cpp
@@ -298,6 +298,12 @@ namespace AssetSelectionUtils
 			}
 		}
 
+		// hack when only BSP is selected
+		if( ActorInfo.SharedWorld == nullptr )
+		{
+			ActorInfo.SharedWorld = GWorld;
+		}
+
 		return ActorInfo;
 	}
 
