[INTEGRATE] Change 2434847 by [REDACTED_PERSON][REDACTED_HANDLE] on 2015/02/06 00:36:04

	Blueprints: Demote some IsValidLowLevel checks in pin watching code to once-per-session ensures to avoid crashing the editor when we can probably recover [UE-9120]

diff --git a/Engine/Source/Editor/UnrealEd/Private/Kismet2/KismetDebugUtilities.cpp b/Engine/Source/Editor/UnrealEd/Private/Kismet2/KismetDebugUtilities.cpp
--- a/Engine/Source/Editor/UnrealEd/Private/Kismet2/KismetDebugUtilities.cpp
+++ b/Engine/Source/Editor/UnrealEd/Private/Kismet2/KismetDebugUtilities.cpp
@@ -853,11 +853,31 @@ FKismetDebugUtilities::EWatchTextResult FKismetDebugUtilities::GetWatchText(FStr
 {
 	if (UProperty* Property = FKismetDebugUtilities::FindClassPropertyForPin(Blueprint, WatchPin))
 	{
-		if (ActiveObject != NULL)
+		if (!Property->IsValidLowLevel())
 		{
 			//@TODO: Temporary checks to attempt to determine intermittent unreproducable crashes in this function
-			check(ActiveObject->IsValidLowLevel());
-			check(Property->IsValidLowLevel());
+			static bool bErrorOnce = true;
+			if (bErrorOnce)
+			{
+				ensureMsg(false, TEXT("Error: Invalid (but non-null) property associated with pin; cannot get variable value"));
+				bErrorOnce = false;
+			}
+			return EWTR_NoProperty;
+		}
+
+		if (ActiveObject != nullptr)
+		{
+			if (!ActiveObject->IsValidLowLevel())
+			{
+				//@TODO: Temporary checks to attempt to determine intermittent unreproducable crashes in this function
+				static bool bErrorOnce = true;
+				if (bErrorOnce)
+				{
+					ensureMsgf(false, TEXT("Error: Invalid (but non-null) active object being debugged; cannot get variable value for property %s"), *Property->GetPathName());
+					bErrorOnce = false;
+				}
+				return EWTR_NoDebugObject;
+			}
 
 			void* PropertyBase = NULL;
 
@@ -873,13 +893,13 @@ FKismetDebugUtilities::EWatchTextResult FKismetDebugUtilities::GetWatchText(FStr
 
 			// Try at member scope if it wasn't part of a current function scope
 			UClass* PropertyClass = Cast<UClass>(Property->GetOuter());
-			if ((PropertyBase == NULL) && (PropertyClass != NULL) && ActiveObject->GetClass()->IsChildOf(PropertyClass))
+			if ((PropertyBase == nullptr) && (PropertyClass != nullptr) && ActiveObject->GetClass()->IsChildOf(PropertyClass))
 			{
 				PropertyBase = ActiveObject;
 			}
 
 			// Now either print out the variable value, or that it was out-of-scope
-			if (PropertyBase != NULL)
+			if (PropertyBase != nullptr)
 			{
 				Property->ExportText_InContainer(/*ArrayElement=*/ 0, /*inout*/ OutWatchText, PropertyBase, PropertyBase, /*Parent=*/ ActiveObject, PPF_PropertyWindow|PPF_BlueprintDebugView);
 				return EWTR_Valid;
@@ -904,7 +924,7 @@ FText FKismetDebugUtilities::GetAndClearLastExceptionMessage()
 {
 	const FString Result = LastExceptionMessage;
 	LastExceptionMessage.Empty();
-	return  FText::FromString( Result );
+	return FText::FromString(Result);
 }
 
 #undef LOCTEXT_NAMESPACE
