[INTEGRATE] Change 2435587 by [REDACTED_PERSON][REDACTED_HANDLE] on 2015/02/06 12:09:30

	Component combo box now scrolls to the top when opened. UE 9209

diff --git a/Engine/Source/Editor/UnrealEd/Private/SComponentClassCombo.cpp b/Engine/Source/Editor/UnrealEd/Private/SComponentClassCombo.cpp
--- a/Engine/Source/Editor/UnrealEd/Private/SComponentClassCombo.cpp
+++ b/Engine/Source/Editor/UnrealEd/Private/SComponentClassCombo.cpp
@@ -120,9 +120,15 @@ void SComponentClassCombo::ClearSelection()
 	SearchBox->SetText(FText::GetEmpty());
 
 	PrevSelectedIndex = INDEX_NONE;
-
+	
 	// Clear the selection in such a way as to also clear the keyboard selector
 	ComponentClassListView->SetSelection(NULL, ESelectInfo::OnNavigation);
+
+	// Make sure we scroll to the top
+	if (ComponentClassList->Num() > 0)
+	{
+		ComponentClassListView->RequestScrollIntoView((*ComponentClassList)[0]);
+	}
 }
 
 void SComponentClassCombo::GenerateFilteredComponentList(const FString& InSearchText)
