[INTEGRATE] Change 2435778 by [REDACTED_PERSON][REDACTED_HANDLE] on 2015/02/06 14:13:54

	Fixed crash undoing converting an actor to a blueprint

diff --git a/Engine/Source/Editor/UnrealEd/Private/Kismet2/Kismet2.cpp b/Engine/Source/Editor/UnrealEd/Private/Kismet2/Kismet2.cpp
--- a/Engine/Source/Editor/UnrealEd/Private/Kismet2/Kismet2.cpp
+++ b/Engine/Source/Editor/UnrealEd/Private/Kismet2/Kismet2.cpp
@@ -1394,6 +1394,8 @@ AActor* FKismetEditorUtilities::CreateBlueprintInstanceFromSelection(UBlueprint*
 	// Assume all selected actors are in the same world
 	UWorld* World = SelectedActors[0]->GetWorld();
 
+	GEditor->GetSelectedActors()->Modify();
+
 	for(auto It(SelectedActors.CreateIterator());It;++It)
 	{
 		if (AActor* Actor = Cast<AActor>(*It))
