[INTEGRATE] Change 2436561 by [REDACTED_PERSON][REDACTED_HANDLE] on 2015/02/06 22:16:17

	Fixed component blueprints considered data only

diff --git a/Engine/Source/Editor/UnrealEd/Private/Kismet2/BlueprintEditorUtils.cpp b/Engine/Source/Editor/UnrealEd/Private/Kismet2/BlueprintEditorUtils.cpp
--- a/Engine/Source/Editor/UnrealEd/Private/Kismet2/BlueprintEditorUtils.cpp
+++ b/Engine/Source/Editor/UnrealEd/Private/Kismet2/BlueprintEditorUtils.cpp
@@ -2378,6 +2378,12 @@ bool FBlueprintEditorUtils::IsDataOnlyBlueprint(const UBlueprint* Blueprint)
 		return false;
 	}
 
+
+	if( Blueprint->ParentClass->IsChildOf( UActorComponent::StaticClass() ) )
+	{
+		return false;
+	}
+
 	// No new variables defined
 	if (Blueprint->NewVariables.Num() > 0)
 	{
