[INTEGRATE] Change 2436669 by [REDACTED_PERSON][REDACTED_HANDLE] on 2015/02/07 02:41:53

	Blueprints: Rebuild harvest components functionality with proper support for instance-added components, etc... (WIP)

diff --git a/Engine/Source/Editor/UnrealEd/Private/Kismet2/Kismet2.cpp b/Engine/Source/Editor/UnrealEd/Private/Kismet2/Kismet2.cpp
--- a/Engine/Source/Editor/UnrealEd/Private/Kismet2/Kismet2.cpp
+++ b/Engine/Source/Editor/UnrealEd/Private/Kismet2/Kismet2.cpp
@@ -909,7 +909,7 @@ UBlueprint* FKismetEditorUtilities::CreateBlueprintFromActor(const FString& Path
 	return NewBlueprint;
 }
 
-void FKismetEditorUtilities::AddComponentsToBlueprint(UBlueprint* Blueprint, const TArray<UActorComponent*>& Components)
+void FKismetEditorUtilities::AddComponentsToBlueprint(UBlueprint* Blueprint, const TArray<UActorComponent*>& Components, bool bHarvesting, USCS_Node* OptionalNewRootNode)
 {
 	USimpleConstructionScript* SCS = Blueprint->SimpleConstructionScript;
 
@@ -936,7 +936,7 @@ void FKismetEditorUtilities::AddComponentsToBlueprint(UBlueprint* Blueprint, con
 			}
 
 			USCS_Node* SCSNode = SCS->CreateNode(ActorComponent->GetClass(), ActorComponent->GetFName());
-			UEditorEngine::CopyPropertiesForUnrelatedObjects(ActorComponent,SCSNode->ComponentTemplate);
+			UEditorEngine::CopyPropertiesForUnrelatedObjects(ActorComponent, SCSNode->ComponentTemplate);
 
 			// Clear the instance component flag
 			SCSNode->ComponentTemplate->CreationMethod = EComponentCreationMethod::Native;
@@ -954,7 +954,14 @@ void FKismetEditorUtilities::AddComponentsToBlueprint(UBlueprint* Blueprint, con
 
 				if (ActorComponent == Actor->GetRootComponent())
 				{
-					SCS->AddNode(SCSNode);
+					if (OptionalNewRootNode != nullptr)
+					{
+						OptionalNewRootNode->AddChildNode(SCSNode);
+					}
+					else
+					{
+						SCS->AddNode(SCSNode);
+					}
 				}
 				// If we're attached to a blueprint component look it up as the variable name is the component name
 				else if (SceneComponent->AttachParent->IsCreatedByConstructionScript())
@@ -980,9 +987,9 @@ void FKismetEditorUtilities::AddComponentsToBlueprint(UBlueprint* Blueprint, con
 						ParentSCSNode->AddChildNode(SCSNode);
 					}
 				}
-				// If we're attached to a native component
-				else if (SceneComponent->AttachParent->CreationMethod == EComponentCreationMethod::Native)
+				else if ((SceneComponent->AttachParent->CreationMethod == EComponentCreationMethod::Native) && !bHarvesting)
 				{
+					// If we're attached to a component that will be native in the new blueprint
 					SCS->AddNode(SCSNode);
 					SCSNode->SetParent(SceneComponent->AttachParent);
 				}
@@ -1083,7 +1090,7 @@ UBlueprint* FKismetEditorUtilities::CreateBlueprintFromActor(const FName Bluepri
 		}
 	}
 
-	if (NewBlueprint != nullptr)
+	if (NewBlueprint)
 	{
 		// Open the editor for the new blueprint
 		FAssetEditorManager::Get().OpenEditorForAsset(NewBlueprint);
@@ -1096,8 +1103,8 @@ class FCreateConstructionScriptFromSelectedActors
 {
 public:
 	FCreateConstructionScriptFromSelectedActors()
-		: Blueprint(NULL)
-		, SCS(NULL)
+		: Blueprint(nullptr)
+		, SCS(nullptr)
 	{
 	}
 
@@ -1105,38 +1112,6 @@ public:
 	{
 		if (SelectedActors.Num() > 0)
 		{
-			// See if we have a single actor with a single component selected
-			// If so, we don't bother creating the root SceneComponent.
-			USceneComponent* SingleSceneComp = NULL;
-			if(SelectedActors.Num() == 1)
-			{
-				TInlineComponentArray<UActorComponent*> Components;
-				SelectedActors[0]->GetComponents(Components);
-				for(UActorComponent* Component : Components)
-				{
-					USceneComponent* SceneComp = Cast<USceneComponent>(Component);
-					if( SceneComp != NULL && 
-						!SceneComp->bHiddenInGame && 
-						SceneComp->GetClass()->HasMetaData(FBlueprintMetadata::MD_BlueprintSpawnableComponent) )
-					{
-						// If we already found one, then there is more than one, so set to NULL and break and continue
-						if(SingleSceneComp != NULL)
-						{
-							SingleSceneComp = NULL;
-							break;
-						}
-						// This is the first valid scene component we have found, save it
-						else
-						{
-							SingleSceneComp = SceneComp;
-						}
-					}
-				}
-			}
-
-			// Determine the origin for all the actors, so it can be backed out when saving them in the blueprint
-			FTransform NewActorTransform = FTransform::Identity;
-
 			// Create a blueprint
 			FString PackageName = Path;
 			FString AssetName = FPackageName::GetLongPackageAssetName(Path);
@@ -1150,101 +1125,71 @@ public:
 				AssetToolsModule.Get().CreateUniqueAssetName(BasePath, TEXT(""), PackageName, AssetName);
 			}
 
-			UPackage* Package = CreatePackage(NULL, *PackageName);
+			UPackage* Package = CreatePackage(nullptr, *PackageName);
 			Blueprint = FKismetEditorUtilities::CreateBlueprint(AActor::StaticClass(), Package, *AssetName, BPTYPE_Normal, UBlueprint::StaticClass(), UBlueprintGeneratedClass::StaticClass(), FName("HarvestFromActors"));
 
 			check(Blueprint->SimpleConstructionScript != NULL);
 			SCS = Blueprint->SimpleConstructionScript;
 
-			// If we have a single component selected, make a BP with a single component
-			if(SingleSceneComp != NULL)
+			// Create a common root if necessary
+			USCS_Node* RootNodeOverride = nullptr;
+			if ((SelectedActors.Num() == 1) && (SelectedActors[0]->GetRootComponent() != nullptr))
 			{
-				USCS_Node* Node = CreateUSCSNode(SingleSceneComp);
-				Node->VariableName = FBlueprintEditorUtils::FindUniqueKismetName(Blueprint, SingleSceneComp->GetOwner()->GetActorLabel());
-				SCS->AddNode(Node);
-
-				NewActorTransform = SingleSceneComp->ComponentToWorld;
+				// We have a single actor that has a scene component, so no need to create a common root
 			}
-			// Multiple actors/component selected, so we create a root scenecomponent and attach things to it
 			else
+			{
+				// Add a new scene component to serve as the shared root node
+				RootNodeOverride = SCS->CreateNode(USceneComponent::StaticClass(), TEXT("SharedRoot"));
+				SCS->AddNode(RootNodeOverride);
+			}
+
+			// Harvest the components from each actor and clone them into the SCS
+			for (const AActor* Actor : SelectedActors)
+			{
+				TArray<UActorComponent*> Components;
+				Actor->GetComponents(Components);
+				FKismetEditorUtilities::AddComponentsToBlueprint(Blueprint, Components, /*bHarvesting=*/ true, RootNodeOverride);
+			}
+
+			// Compute the average origin for all the actors, so it can be backed out when saving them in the blueprint
+			FTransform NewActorTransform = FTransform::Identity;
 			{
 				// Find average location of all selected actors
 				FVector AverageLocation = FVector::ZeroVector;
-				for (auto It(SelectedActors.CreateConstIterator()); It; ++It)
+				for (const AActor* Actor : SelectedActors)
 				{
-					if (const AActor* Actor = Cast<AActor>(*It))
+					if (USceneComponent* RootComponent = Actor->GetRootComponent())
 					{
-						if (USceneComponent* RootComponent = Actor->GetRootComponent())
-						{
-							AverageLocation += Actor->GetActorLocation();
-						}
+						AverageLocation += Actor->GetActorLocation();
 					}
 				}
 				AverageLocation /= (float)SelectedActors.Num();
 
 				// Spawn the new BP at that location
 				NewActorTransform.SetTranslation(AverageLocation);
+			}
 
-				// Add a new scene component to serve as the root node
-				USCS_Node* RootNode = SCS->CreateNode(USceneComponent::StaticClass());
-				SCS->AddNode(RootNode);
-
-				TMap<const AActor*, USCS_Node*> ActorToUSCSNodeMap;
-
-				// Add each actor's root component to the mapping of Actors to USCS_Node map to be organized later
-				for (auto ActorIt = SelectedActors.CreateConstIterator(); ActorIt; ++ActorIt)
+			// Reposition all of the children of the root node to recenter them around the new pivot
+			for (USCS_Node* TopLevelNode : SCS->GetRootNodes())
+			{
+				if (USceneComponent* TestRoot = Cast<USceneComponent>(TopLevelNode->ComponentTemplate))
 				{
-					const AActor* Actor = *ActorIt;
-					USceneComponent* ActorRootComponent = Actor->GetRootComponent();
-					// Check we have a component and it is valid to add to a BP
-					if(ActorRootComponent != NULL && ActorRootComponent->GetClass()->HasMetaData(FBlueprintMetadata::MD_BlueprintSpawnableComponent))
+					for (USCS_Node* ChildNode : TopLevelNode->ChildNodes)
 					{
-						USCS_Node* Node = CreateUSCSNode(ActorRootComponent);
-						Node->VariableName = FBlueprintEditorUtils::FindUniqueKismetName(Blueprint, Actor->GetActorLabel());
-
-						// Add the node as a child to the root node, this will be rearranged later for correct parent-child relations but needs to be there now for name validation purposes
-						RootNode->AddChildNode(Node);
-
-						ActorToUSCSNodeMap.Add(Actor, Node);
-					}
-				}
-
-				// Clear the root node of all children we added, they will be re-added with correct parent-child relationships
-				RootNode->ChildNodes.Empty();
-
-				// Attach all immediate children to their parent, or the root if their parent is not being added
-				for (auto ActorIt = SelectedActors.CreateConstIterator(); ActorIt; ++ActorIt)
-				{
-					const AActor* Actor = *ActorIt;
-					USceneComponent* ActorRootComponent = Actor->GetRootComponent();
-					// Check we have a component and it is valid to add to a BP
-					if(ActorRootComponent != NULL && ActorRootComponent->GetClass()->HasMetaData(FBlueprintMetadata::MD_BlueprintSpawnableComponent))
-					{			
-						USCS_Node** CurrentNode = ActorToUSCSNodeMap.Find(Actor);
-						USCS_Node** ParentNode = ActorToUSCSNodeMap.Find(Actor->GetAttachParentActor());
-
-						check(CurrentNode);
-
-						FTransform ComponentToWorldSpace;
-						if(ParentNode)
-						{
-							ComponentToWorldSpace = TransformChildComponent(Actor->GetRootComponent(), *CurrentNode, Actor->GetAttachParentActor()->GetRootComponent()->GetComponentToWorld());
-							(*ParentNode)->AddChildNode(*CurrentNode);
-						}
-						else
+						if (USceneComponent* ChildComponent = Cast<USceneComponent>(ChildNode->ComponentTemplate))
 						{
-							ComponentToWorldSpace = TransformChildComponent(Actor->GetRootComponent(), *CurrentNode, NewActorTransform);
-							RootNode->AddChildNode(*CurrentNode);
+							const FTransform OldChildToWorld(ChildComponent->RelativeRotation, ChildComponent->RelativeLocation);
+							const FTransform NewRelativeTransform = OldChildToWorld.GetRelativeTransform(NewActorTransform);
+							ChildComponent->RelativeLocation = NewRelativeTransform.GetLocation();
+							ChildComponent->RelativeRotation = NewRelativeTransform.GetRotation().Rotator();
 						}
-
-						// Attach any children components as well, as long as their owner is this Actor
-						AddActorComponents(Actor, Actor->GetRootComponent(), *CurrentNode, ComponentToWorldSpace);
 					}
 				}
 			}
 
 			// Regenerate skeleton class as components have been added since initial generation
-			FKismetEditorUtilities::GenerateBlueprintSkeleton(Blueprint,true); 
+			FKismetEditorUtilities::GenerateBlueprintSkeleton(Blueprint, /*bForceRegeneration=*/ true); 
 			
 			// Notify the asset registry
 			FAssetRegistryModule::AssetCreated(Blueprint);
@@ -1266,112 +1211,9 @@ public:
 
 			return Blueprint;
 		}
-		return NULL;
+		return nullptr;
 	}
 
-protected:
-
-	void CopyComponentToTemplate(UActorComponent* Source, UActorComponent* Destination)
-	{
-		check(Source->GetClass() == Destination->GetClass());
-		UClass* CommonBaseClass = Source->GetClass();
-		UActorComponent* DefaultComponent = CastChecked<UActorComponent>(Destination->GetArchetype());
-
-		//@TODO: Copy of code from inside CopyActorComponentProperties
-
-		// Iterate through the properties, only copying those which are non-native, non-transient, non-component, and not identical
-		// to the values in the default component
-		for (UProperty* Property = CommonBaseClass->PropertyLink; Property != NULL; Property = Property->PropertyLinkNext)
-		{
-			const bool bIsTransient = !!( Property->PropertyFlags & CPF_Transient );
-			const bool bIsIdentical = Property->Identical_InContainer(Source, DefaultComponent);
-			const bool bIsComponent = !!( Property->PropertyFlags & (CPF_InstancedReference | CPF_ContainsInstancedReference) );
-
-			if (!bIsTransient && !bIsIdentical && !bIsComponent)
-			{
-				Property->CopyCompleteValue_InContainer(Destination, Source);
-			}
-		}
-
-		if (UStaticMeshComponent* TemplateMeshComponent = Cast<UStaticMeshComponent>(Destination))
-		{
-			UStaticMeshComponent* SourceMeshComponent = CastChecked<UStaticMeshComponent>(Source);
-			// vert coloring data lives in the mesh component's LODData member,
-			// which is transient, so it wouldn't be copied in the above loop
-			TemplateMeshComponent->CopyInstanceVertexColorsIfCompatible(SourceMeshComponent);
-		}
-	}
-
-	/** Creates a USCS Node for the passed in component */
-	USCS_Node* CreateUSCSNode(USceneComponent* Component)
-	{
-		// Create a new SCS node for this component
-		USCS_Node* NewNode = SCS->CreateNode(Component->GetClass());
-
-		// Copy properties across to the new template
-		CopyComponentToTemplate(Component, NewNode->ComponentTemplate);
-
-		return NewNode;
-	}
-
-	/**
-	 * Transforms a USCS_Node component based on the original Component and the passed in transform.
-	 *
-	 * @param InOriginalComponent		The original component the USCS_Node is based on
-	 * @param InUSCSNode				The USCS_Node based on the original component
-	 * @param InParentToWorldSpace		FTransform to place the USCS_Node in the correct spot.
-	 *
-	 * @return							The component to world transform of the USCS_Node component
-	 */
-	FTransform TransformChildComponent(USceneComponent* InOriginalComponent, USCS_Node* InUSCSNode, const FTransform& InParentToWorldSpace)
-	{
-		// Position this node relative to it's parent
-		USceneComponent* NewTemplate = CastChecked<USceneComponent>(InUSCSNode->ComponentTemplate);
-
-		FTransform ChildToWorld = InOriginalComponent->GetComponentToWorld();
-		FTransform RelativeTransform = ChildToWorld.GetRelativeTransform(InParentToWorldSpace);
-
-		NewTemplate->SetRelativeTransform(RelativeTransform);
-
-		return ChildToWorld;
-	}
-
-	/**
-	 * Adds all child components (but not child actors) to the RootNode as a USCS_Node
-	 *
-	 * @param InParentActor		The actor all components must be a child of
-	 * @param InComponent		The component to pull child components from
-	 * @param InRootNode		The root node to attach all new USCS_Nodes to
-	 * @param InParentToWorldSpace		FTransform to place the USCS_Node in the correct spot.
-	 */
-	void AddActorComponents(const AActor* InParentActor, USceneComponent* InComponent, USCS_Node* InRootNode, const FTransform& InParentToWorldSpace)
-	{
-		if (InComponent)
-		{
-			for (auto ChildIt = InComponent->AttachChildren.CreateConstIterator(); ChildIt; ++ChildIt)
-			{
-				USceneComponent* ChildComponent = *ChildIt;
-
-				if(ChildComponent->GetOwner() == InParentActor)
-				{				
-					// if component is 'editor only' we usually don't want to copy it
-					// Also check its a valid class to add
-					if(!ChildComponent->bHiddenInGame && ChildComponent->GetClass()->HasMetaData(FBlueprintMetadata::MD_BlueprintSpawnableComponent))
-					{
-						USCS_Node* NewNode = CreateUSCSNode(ChildComponent);
-						FTransform ComponentToWorldSpace = TransformChildComponent(ChildComponent, NewNode, InParentToWorldSpace);
-						InRootNode->AddChildNode(NewNode);
-
-						// Add any child components that may be owned by this Actor
-						AddActorComponents(InParentActor, ChildComponent, NewNode, ComponentToWorldSpace);
-					}
-				}
-
-			}
-		}
-	}
-
-
 protected:
 	UBlueprint* Blueprint;
 	USimpleConstructionScript* SCS;
diff --git a/Engine/Source/Editor/UnrealEd/Public/Kismet2/KismetEditorUtilities.h b/Engine/Source/Editor/UnrealEd/Public/Kismet2/KismetEditorUtilities.h
--- a/Engine/Source/Editor/UnrealEd/Public/Kismet2/KismetEditorUtilities.h
+++ b/Engine/Source/Editor/UnrealEd/Public/Kismet2/KismetEditorUtilities.h
@@ -87,7 +87,7 @@ public:
 	static bool CanCreateBlueprintOfClass(const UClass* Class);
 
 	/** Take a list of components that belong to a single Actor and add them to a blueprint as SCSNodes */
-	static void AddComponentsToBlueprint(UBlueprint* Blueprint, const TArray<UActorComponent*>& Components);
+	static void AddComponentsToBlueprint(UBlueprint* Blueprint, const TArray<UActorComponent*>& Components, bool bHarvesting = false, class USCS_Node* OptionalNewRootNode = nullptr);
 
 	/** 
 	 * Take an Actor and generate a blueprint based on it. Uses the Actors type as the parent class. 
