[INTEGRATE] Change 2436825 by [REDACTED_PERSON][REDACTED_HANDLE] on 2015/02/07 13:53:36

	Blueprints: Prevent harvest components from incorrectly gathering editor-only components

diff --git a/Engine/Source/Editor/UnrealEd/Private/Kismet2/Kismet2.cpp b/Engine/Source/Editor/UnrealEd/Private/Kismet2/Kismet2.cpp
--- a/Engine/Source/Editor/UnrealEd/Private/Kismet2/Kismet2.cpp
+++ b/Engine/Source/Editor/UnrealEd/Private/Kismet2/Kismet2.cpp
@@ -1149,6 +1149,17 @@ public:
 			{
 				TArray<UActorComponent*> Components;
 				Actor->GetComponents(Components);
+
+				// Exclude any components created by other components
+				for (int32 Index = Components.Num() - 1; Index >= 0; --Index)
+				{
+					UActorComponent* ComponentToConsider = Components[Index];
+					if (ComponentToConsider->IsEditorOnly())
+					{
+						Components.RemoveAt(Index, 1, /*bAllowShrinking=*/ false);
+					}
+				}
+
 				FKismetEditorUtilities::AddComponentsToBlueprint(Blueprint, Components, /*bHarvesting=*/ true, RootNodeOverride);
 			}
 
