[INTEGRATE] Change 2436905 by [REDACTED_PERSON][REDACTED_HANDLE] on 2015/02/07 15:36:17

	Added some of the basic shapes to common on the Add component dropdown.

	Moved collison shapes to a collision category

	Fixed sorting of basic shapes against other components

diff --git a/Engine/Source/Editor/UnrealEd/Private/ComponentTypeRegistry.cpp b/Engine/Source/Editor/UnrealEd/Private/ComponentTypeRegistry.cpp
--- a/Engine/Source/Editor/UnrealEd/Private/ComponentTypeRegistry.cpp
+++ b/Engine/Source/Editor/UnrealEd/Private/ComponentTypeRegistry.cpp
@@ -50,9 +50,18 @@ static void AddBasicShapeComponents(TArray<FComponentClassComboEntryPtr>& Sorted
 		Args.OnComponentCreated = FOnComponentCreated::CreateStatic(OnBasicShapeCreated);
 		Args.ComponentNameOverride = LOCTEXT("BasicCubeShapeDisplayName", "Cube").ToString();
 		Args.IconOverrideBrushName = FName("ClassIcon.Cube");
+		Args.SortPriority = 2;
 
-		FComponentClassComboEntryPtr NewShape = MakeShareable(new FComponentClassComboEntry(BasicShapesHeading, UStaticMeshComponent::StaticClass(), true, EComponentCreateAction::SpawnExistingClass, Args));
-		SortedClassList.Add(NewShape);
+		{
+			FComponentClassComboEntryPtr NewShape = MakeShareable(new FComponentClassComboEntry(BasicShapesHeading, UStaticMeshComponent::StaticClass(), true, EComponentCreateAction::SpawnExistingClass, Args));
+			SortedClassList.Add(NewShape);
+		}
+
+		{
+			//Cube also goes in the common group
+			FComponentClassComboEntryPtr NewShape = MakeShareable(new FComponentClassComboEntry(CommonClassGroup, UStaticMeshComponent::StaticClass(), false, EComponentCreateAction::SpawnExistingClass, Args));
+			SortedClassList.Add(NewShape);
+		}
 	}
 
 	{
@@ -61,9 +70,17 @@ static void AddBasicShapeComponents(TArray<FComponentClassComboEntryPtr>& Sorted
 		Args.OnComponentCreated = FOnComponentCreated::CreateStatic(OnBasicShapeCreated);
 		Args.ComponentNameOverride = LOCTEXT("BasicSphereShapeDisplayName", "Sphere").ToString();
 		Args.IconOverrideBrushName = FName("ClassIcon.Sphere");
+		Args.SortPriority = 2;
+		{
+			FComponentClassComboEntryPtr NewShape = MakeShareable(new FComponentClassComboEntry(BasicShapesHeading, UStaticMeshComponent::StaticClass(), true, EComponentCreateAction::SpawnExistingClass, Args));
+			SortedClassList.Add(NewShape);
+		}
 
-		FComponentClassComboEntryPtr NewShape = MakeShareable(new FComponentClassComboEntry(BasicShapesHeading, UStaticMeshComponent::StaticClass(), true, EComponentCreateAction::SpawnExistingClass, Args));
-		SortedClassList.Add(NewShape);
+		{
+			// Sphere also goes in the common group
+			FComponentClassComboEntryPtr NewShape = MakeShareable(new FComponentClassComboEntry(CommonClassGroup, UStaticMeshComponent::StaticClass(), false, EComponentCreateAction::SpawnExistingClass, Args));
+			SortedClassList.Add(NewShape);
+		}
 	}
 
 	{
@@ -72,7 +89,7 @@ static void AddBasicShapeComponents(TArray<FComponentClassComboEntryPtr>& Sorted
 		Args.OnComponentCreated = FOnComponentCreated::CreateStatic(OnBasicShapeCreated);
 		Args.ComponentNameOverride = LOCTEXT("BasicCylinderShapeDisplayName", "Cylinder").ToString();
 		Args.IconOverrideBrushName = FName("ClassIcon.Cylinder");
-
+		Args.SortPriority = 3;
 		FComponentClassComboEntryPtr NewShape = MakeShareable(new FComponentClassComboEntry(BasicShapesHeading, UStaticMeshComponent::StaticClass(), true, EComponentCreateAction::SpawnExistingClass, Args));
 		SortedClassList.Add(NewShape);
 	}
@@ -83,7 +100,7 @@ static void AddBasicShapeComponents(TArray<FComponentClassComboEntryPtr>& Sorted
 		Args.OnComponentCreated = FOnComponentCreated::CreateStatic(OnBasicShapeCreated);
 		Args.ComponentNameOverride = LOCTEXT("BasicConeShapeDisplayName", "Cone").ToString();
 		Args.IconOverrideBrushName = FName("ClassIcon.Cone");
-
+		Args.SortPriority = 4;
 		FComponentClassComboEntryPtr NewShape = MakeShareable(new FComponentClassComboEntry(BasicShapesHeading, UStaticMeshComponent::StaticClass(), true, EComponentCreateAction::SpawnExistingClass, Args));
 		SortedClassList.Add(NewShape);
 	}
@@ -122,7 +139,14 @@ void FComponentTypeRegistryData::RefreshComponentList()
 			int32 HeadingCompareResult = FCString::Stricmp(*A->GetHeadingText(), *B->GetHeadingText());
 			if (HeadingCompareResult == 0)
 			{
-				bResult = FCString::Stricmp(*A->GetClassName(), *B->GetClassName()) < 0;
+				if( A->GetSortPriority() == 0 && B->GetSortPriority() == 0 )
+				{
+					bResult = FCString::Stricmp(*A->GetClassName(), *B->GetClassName()) < 0;
+				}
+				else
+				{
+					bResult = A->GetSortPriority() < B->GetSortPriority();
+				}
 			}
 			else if (CommonClassGroup == A->GetHeadingText())
 			{
@@ -161,6 +185,9 @@ void FComponentTypeRegistryData::RefreshComponentList()
 	}
 
 	TArray<FComponentClassComboEntryPtr> SortedClassList;
+
+	AddBasicShapeComponents(SortedClassList);
+
 	TArray<FName> InMemoryClasses;
 	for (TObjectIterator<UClass> It; It; ++It)
 	{
@@ -195,7 +222,6 @@ void FComponentTypeRegistryData::RefreshComponentList()
 		}
 	}
 
-	AddBasicShapeComponents(SortedClassList);
 
 	{
 		// make sure that we add any user created classes immediately, generally this will not create anything (because assets have not been discovered yet), but asset discovery
diff --git a/Engine/Source/Editor/UnrealEd/Public/SComponentClassCombo.h b/Engine/Source/Editor/UnrealEd/Public/SComponentClassCombo.h
--- a/Engine/Source/Editor/UnrealEd/Public/SComponentClassCombo.h
+++ b/Engine/Source/Editor/UnrealEd/Public/SComponentClassCombo.h
@@ -34,12 +34,15 @@ struct FComponentEntryCustomizationArgs
 	FOnComponentCreated OnComponentCreated;
 	/** Brush icon to use instead of the class icon */
 	FName IconOverrideBrushName;
+	/** Custom sort priority to use (smaller means sorted first) */
+	int32 SortPriority;
 
 	FComponentEntryCustomizationArgs()
 		: AssetOverride( nullptr )
 		, ComponentNameOverride()
 		, OnComponentCreated()
 		, IconOverrideBrushName( NAME_None )
+		, SortPriority(0)
 	{
 	
 	}
@@ -125,6 +128,8 @@ public:
 	}
 
 	FName GetIconOverrideBrushName() const { return CustomizationArgs.IconOverrideBrushName; }
+
+	int32 GetSortPriority() const { return CustomizationArgs.SortPriority; }
 private:
 	TSubclassOf<UActorComponent> ComponentClass;
 	// For components that are not loaded we just keep the name of the component,
diff --git a/Engine/Source/Runtime/Engine/Classes/Components/BoxComponent.h b/Engine/Source/Runtime/Engine/Classes/Components/BoxComponent.h
--- a/Engine/Source/Runtime/Engine/Classes/Components/BoxComponent.h
+++ b/Engine/Source/Runtime/Engine/Classes/Components/BoxComponent.h
@@ -8,7 +8,7 @@
 /** 
  * A box generally used for simple collision. Bounds are rendered as lines in the editor.
  */
-UCLASS(ClassGroup="Physics", hidecategories=(Object,LOD,Lighting,TextureStreaming), editinlinenew, meta=(DisplayName="Box Collision", BlueprintSpawnableComponent))
+UCLASS(ClassGroup="Collision", hidecategories=(Object,LOD,Lighting,TextureStreaming), editinlinenew, meta=(DisplayName="Box Collision", BlueprintSpawnableComponent))
 class ENGINE_API UBoxComponent : public UShapeComponent
 {
 	GENERATED_UCLASS_BODY()
diff --git a/Engine/Source/Runtime/Engine/Classes/Components/CapsuleComponent.h b/Engine/Source/Runtime/Engine/Classes/Components/CapsuleComponent.h
--- a/Engine/Source/Runtime/Engine/Classes/Components/CapsuleComponent.h
+++ b/Engine/Source/Runtime/Engine/Classes/Components/CapsuleComponent.h
@@ -7,7 +7,7 @@
 /** 
  * A capsule generally used for simple collision. Bounds are rendered as lines in the editor.
  */
-UCLASS(ClassGroup="Physics", editinlinenew, hidecategories=(Object,LOD,Lighting,TextureStreaming), meta=(DisplayName="Capsule Collision", BlueprintSpawnableComponent))
+UCLASS(ClassGroup="Collision", editinlinenew, hidecategories=(Object,LOD,Lighting,TextureStreaming), meta=(DisplayName="Capsule Collision", BlueprintSpawnableComponent))
 class ENGINE_API UCapsuleComponent : public UShapeComponent
 {
 	GENERATED_UCLASS_BODY()
diff --git a/Engine/Source/Runtime/Engine/Classes/Components/SphereComponent.h b/Engine/Source/Runtime/Engine/Classes/Components/SphereComponent.h
--- a/Engine/Source/Runtime/Engine/Classes/Components/SphereComponent.h
+++ b/Engine/Source/Runtime/Engine/Classes/Components/SphereComponent.h
@@ -8,7 +8,7 @@
 /** 
  * A sphere generally used for simple collision. Bounds are rendered as lines in the editor.
  */
-UCLASS(ClassGroup="Physics", editinlinenew, hidecategories=(Object,LOD,Lighting,TextureStreaming), meta=(DisplayName="Sphere Collision", BlueprintSpawnableComponent))
+UCLASS(ClassGroup="Collision", editinlinenew, hidecategories=(Object,LOD,Lighting,TextureStreaming), meta=(DisplayName="Sphere Collision", BlueprintSpawnableComponent))
 class ENGINE_API USphereComponent : public UShapeComponent
 {
 	GENERATED_UCLASS_BODY()
