[INTEGRATE] Change 2436952 by [REDACTED_PERSON][REDACTED_HANDLE] on 2015/02/07 16:43:52

	You can now force delete blueprint components.
	UE-9325 fixed by this.

diff --git a/Engine/Source/Editor/UnrealEd/Private/AssetDeleteModel.cpp b/Engine/Source/Editor/UnrealEd/Private/AssetDeleteModel.cpp
--- a/Engine/Source/Editor/UnrealEd/Private/AssetDeleteModel.cpp
+++ b/Engine/Source/Editor/UnrealEd/Private/AssetDeleteModel.cpp
@@ -286,7 +286,7 @@ bool FAssetDeleteModel::ComputeCanReplaceReferences()
 		PendingDeletedObjects.Add(PendingDelete->GetObject());
 	}
 
-	return ObjectTools::AreObjectsOfEquivalantType( PendingDeletedObjects );
+	return ObjectTools::AreObjectsValidForReplace(PendingDeletedObjects) && ObjectTools::AreObjectsValidForReplace(PendingDeletedObjects);
 }
 
 bool FAssetDeleteModel::CanReplaceReferences() const
diff --git a/Engine/Source/Editor/UnrealEd/Private/ObjectTools.cpp b/Engine/Source/Editor/UnrealEd/Private/ObjectTools.cpp
--- a/Engine/Source/Editor/UnrealEd/Private/ObjectTools.cpp
+++ b/Engine/Source/Editor/UnrealEd/Private/ObjectTools.cpp
@@ -1756,6 +1756,14 @@ namespace ObjectTools
 
 		GWarn->BeginSlowTask( NSLOCTEXT("UnrealEd", "Deleting", "Deleting"), true );
 
+		struct FSCSNodeToDelete
+		{
+			USimpleConstructionScript* SimpleConstructionScript;
+			USCS_Node* SCS_Node;
+		};
+
+		TArray<FSCSNodeToDelete> SCSNodesToDelete;
+		TArray<UActorComponent*> ComponentsToDelete;
 		TArray<AActor*> ActorsToDelete;
 		TArray<UObject*> ObjectsToDelete;
 		bool bNeedsGarbageCollection = false;
@@ -1772,7 +1780,7 @@ namespace ObjectTools
 			ObjectsToDelete.Add( CurrentObject );
 
 			// If the object about to be deleted is a Blueprint asset, make sure that any instances of the Blueprint class get deleted as well
-			UBlueprint *BlueprintObject = Cast<UBlueprint>(CurrentObject);
+			UBlueprint* BlueprintObject = Cast<UBlueprint>(CurrentObject);
 			if ( BlueprintObject && BlueprintObject->GeneratedClass && BlueprintObject->GeneratedClass->ClassDefaultObject )
 			{
 				TArray<UObject*> InstancesToDelete;
@@ -1783,10 +1791,33 @@ namespace ObjectTools
 					UObject* CurrentInstance = *InstanceItr;
 
 					AActor* CurrentInstanceAsActor = Cast<AActor>( CurrentInstance );
+					UActorComponent* CurrentInstanceAsComponent = Cast<UActorComponent>(CurrentInstance);
 					if ( CurrentInstanceAsActor )
 					{
 						ActorsToDelete.Add( CurrentInstanceAsActor );
 					}
+					else if ( CurrentInstanceAsComponent )
+					{
+						ComponentsToDelete.Add( CurrentInstanceAsComponent );
+
+						// Find all the SCS_Node references that need to be destroyed before this component is destroyed.
+						UBlueprintGeneratedClass* UBGC = CurrentInstanceAsComponent->GetTypedOuter<UBlueprintGeneratedClass>();
+						if (UBGC && UBGC->SimpleConstructionScript)
+						{
+							TArray<USCS_Node*> SCSNodes = UBGC->SimpleConstructionScript->GetAllNodes();
+							for (int32 SCSNodeIndex = 0; SCSNodeIndex < SCSNodes.Num(); ++SCSNodeIndex)
+							{
+								USCS_Node* SCS_Node = SCSNodes[SCSNodeIndex];
+								if (SCS_Node && SCS_Node->ComponentTemplate == CurrentInstanceAsComponent)
+								{
+									FSCSNodeToDelete DeleteNode;
+									DeleteNode.SimpleConstructionScript = UBGC->SimpleConstructionScript;
+									DeleteNode.SCS_Node = SCS_Node;
+									SCSNodesToDelete.Add(DeleteNode);
+								}
+							}
+						}
+					}
 					else
 					{
 						ObjectsToDelete.Add( CurrentInstance );
@@ -1795,10 +1826,53 @@ namespace ObjectTools
 			}
 		}
 
+		// Destroy all SCSNodes
+		if (SCSNodesToDelete.Num() > 0)
+		{
+			for (TArray<FSCSNodeToDelete>::TConstIterator SCSNodeItr(SCSNodesToDelete); SCSNodeItr; ++SCSNodeItr)
+			{
+				FSCSNodeToDelete SCSNodeToDelete = *SCSNodeItr;
+				
+				SCSNodeToDelete.SimpleConstructionScript->RemoveNodeAndPromoteChildren(SCSNodeToDelete.SCS_Node);
+
+				GWarn->StatusUpdate(SCSNodeItr.GetIndex(), SCSNodesToDelete.Num(), NSLOCTEXT("UnrealEd", "ConsolidateAssetsUpdate_DeletingSCSNodes", "Deleting Blueprint Component references..."));
+			}
+		}
+
+		bool bSelectionChanged = false;
+
+		// Destroy all Components
+		if (ComponentsToDelete.Num() > 0)
+		{
+			for (TArray<UActorComponent*>::TConstIterator ComponentItr(ComponentsToDelete); ComponentItr; ++ComponentItr)
+			{
+				UActorComponent* CurComponent = *ComponentItr;
+
+				// Skip if already pending GC
+				if (!CurComponent->IsPendingKill())
+				{
+					// Deselect if active
+					USelection* SelectedComponents = GEditor->GetSelectedComponents();
+					if (SelectedComponents && CurComponent->IsSelected())
+					{
+						SelectedComponents->Deselect(CurComponent);
+
+						bSelectionChanged = true;
+					}
+
+					// Destroy the Component Instance
+					CurComponent->DestroyComponent(true);
+
+					bNeedsGarbageCollection = true;
+				}
+
+				GWarn->StatusUpdate(ComponentItr.GetIndex(), ComponentsToDelete.Num(), NSLOCTEXT("UnrealEd", "ConsolidateAssetsUpdate_DeletingComponentInstances", "Deleting Component Instances..."));
+			}
+		}
+
 		// Destroy all Actor instances
 		if ( ActorsToDelete.Num() > 0 )
 		{
-			bool bSelectionChanged = false;
 			for ( TArray<AActor*>::TConstIterator ActorItr( ActorsToDelete ); ActorItr; ++ActorItr )
 			{
 				AActor* CurActor = *ActorItr;
@@ -1822,13 +1896,13 @@ namespace ObjectTools
 					bNeedsGarbageCollection = true;
 				}
 
-				GWarn->StatusUpdate( ActorItr.GetIndex(), ActorsToDelete.Num(), NSLOCTEXT( "UnrealEd", "ConsolidateAssetsUpdate_DeletingInstances", "Deleting Instances..." ) );
+				GWarn->StatusUpdate( ActorItr.GetIndex(), ActorsToDelete.Num(), NSLOCTEXT( "UnrealEd", "ConsolidateAssetsUpdate_DeletingActorInstances", "Deleting Actor Instances..." ) );
 			}
+		}
 
-			if ( bSelectionChanged )
-			{
-				GEditor->NoteSelectionChange();
-			}
+		if (bSelectionChanged)
+		{
+			GEditor->NoteSelectionChange();
 		}
 
 		// If the current editor world is in this list, transition to a new map and reload the world to finish the delete
@@ -1842,7 +1916,7 @@ namespace ObjectTools
 				{
 					UObject* CurObject = *ObjectItr;
 					
-					UBlueprint *BlueprintObject = Cast<UBlueprint>(CurObject);
+					UBlueprint* BlueprintObject = Cast<UBlueprint>(CurObject);
 					if (BlueprintObject)
 					{
 						// If we're a blueprint add our generated class as well
@@ -1879,13 +1953,6 @@ namespace ObjectTools
 					}
 				}
 
-				// Handle deferred garbage collection
-				if(bNeedsGarbageCollection)
-				{
-					CollectGarbage(GARBAGE_COLLECTION_KEEPFLAGS);
-					bNeedsGarbageCollection = false;
-				}
-
 				// Replacing references inside already loaded objects could cause rendering issues, so globally detach all components from their scenes for now
 				FGlobalComponentReregisterContext ReregisterContext;
 
@@ -1920,6 +1987,13 @@ namespace ObjectTools
 				}
 			}
 
+			// Handle deferred garbage collection
+			if (bNeedsGarbageCollection)
+			{
+				CollectGarbage(GARBAGE_COLLECTION_KEEPFLAGS);
+				bNeedsGarbageCollection = false;
+			}
+
 			// Load the asset tools module to get access to the browser type maps
 			FAssetToolsModule& AssetToolsModule = FModuleManager::LoadModuleChecked<FAssetToolsModule>(TEXT("AssetTools"));
 
@@ -3338,6 +3412,27 @@ namespace ObjectTools
 		return bIsPlaceable && !bIsAbstractOrDeprecated && !bIsSkeletonClass;
 	}
 
+	bool AreObjectsValidForReplace(const TArray<UObject*>& InProposedObjects)
+	{
+		if (InProposedObjects.Num() > 0)
+		{
+			// Make sure we don't have any blueprinted components
+			for (TArray<UObject*>::TConstIterator ProposedObjIter(InProposedObjects); ProposedObjIter; ++ProposedObjIter)
+			{
+				UObject* CurProposedObj = *ProposedObjIter;
+				check(CurProposedObj);
+				
+				const UBlueprint* BlueprintObject = Cast<UBlueprint>(CurProposedObj);
+				if (BlueprintObject && BlueprintObject->GeneratedClass->IsChildOf(UActorComponent::StaticClass()))
+				{
+					return false;
+				}
+			}
+		}
+		// If we haven't found any classes we deem illegal just make sure they are all the same type.
+		return AreObjectsOfEquivalantType(InProposedObjects);
+	}
+
 	bool AreObjectsOfEquivalantType( const TArray<UObject*>& InProposedObjects )
 	{
 		if ( InProposedObjects.Num() > 0 )
diff --git a/Engine/Source/Editor/UnrealEd/Public/ObjectTools.h b/Engine/Source/Editor/UnrealEd/Public/ObjectTools.h
--- a/Engine/Source/Editor/UnrealEd/Public/ObjectTools.h
+++ b/Engine/Source/Editor/UnrealEd/Public/ObjectTools.h
@@ -504,6 +504,15 @@ namespace ObjectTools
 	 */
 	UNREALED_API bool IsClassRedirector( const UClass* Class );
 
+	/**
+	 * Determines if an array of objects are valid for replace
+	 *
+	 * @param InProposedObjects	The objects to check.
+	 * 
+	 * @return true if all objects are not a restricted type for replacement.
+	 */
+	UNREALED_API bool AreObjectsValidForReplace( const TArray<UObject*>& InProposedObjects );
+
 	/**
 	 * Determines if an array of objects are all of interchangeable types.
 	 *
diff --git a/Engine/Source/Runtime/Engine/Private/SimpleConstructionScript.cpp b/Engine/Source/Runtime/Engine/Private/SimpleConstructionScript.cpp
--- a/Engine/Source/Runtime/Engine/Private/SimpleConstructionScript.cpp
+++ b/Engine/Source/Runtime/Engine/Private/SimpleConstructionScript.cpp
@@ -151,17 +151,24 @@ void USimpleConstructionScript::PostLoad()
 
 	int32 NodeIndex;
 	TArray<USCS_Node*> Nodes = GetAllNodes();
-
-	if (GetLinkerUE4Version() < VER_UE4_REMOVE_INPUT_COMPONENTS_FROM_BLUEPRINTS)
+	
+	for (NodeIndex=0; NodeIndex < Nodes.Num(); ++NodeIndex)
 	{
-		for (NodeIndex=0; NodeIndex < Nodes.Num(); ++NodeIndex)
+		USCS_Node* Node = Nodes[NodeIndex];
+		
+		if (GetLinkerUE4Version() < VER_UE4_REMOVE_INPUT_COMPONENTS_FROM_BLUEPRINTS)
 		{
-			USCS_Node* Node = Nodes[NodeIndex];
 			if (!Node->bIsNative_DEPRECATED && Node->ComponentTemplate->IsA<UInputComponent>())
 			{
 				RemoveNodeAndPromoteChildren(Node);
 			}
 		}
+
+		// Couldn't find the template the Blueprint or C++ class must have been deleted out from under us
+		if (Node->ComponentTemplate == nullptr)
+		{
+			RemoveNodeAndPromoteChildren(Node);
+		}
 	}
 
 #if WITH_EDITOR
