[INTEGRATE] Change 2437837 by [REDACTED_PERSON][REDACTED_HANDLE].7 on 2015/02/09 07:39:39

	Added more information to non-power of 2 texture import message box

diff --git a/Engine/Source/Editor/UnrealEd/Private/Factories/EditorFactories.cpp b/Engine/Source/Editor/UnrealEd/Private/Factories/EditorFactories.cpp
--- a/Engine/Source/Editor/UnrealEd/Private/Factories/EditorFactories.cpp
+++ b/Engine/Source/Editor/UnrealEd/Private/Factories/EditorFactories.cpp
@@ -4653,7 +4653,18 @@ bool UTextureFactory::IsImportResolutionValid(int32 Width, int32 Height, bool bA
 	if( bAllowOneTimeWarningMessages && !bSuppressImportResolutionWarnings && bAllowNonPowerOfTwo && !bIsPowerOfTwo && bValid )
 	{
 		bAllowOneTimeWarningMessages = false;
-		if( EAppReturnType::Yes != FMessageDialog::Open( EAppMsgType::YesNo, NSLOCTEXT("UnrealEd", "Warning_NPTTexture", "The texture you are importing is not a power of two.  Non power of two textures are never streamed and have no mipmaps. Proceed?") ) )
+
+		FText MessageText;
+		if (GetCurrentFilename().IsEmpty())
+		{
+			MessageText = NSLOCTEXT("UnrealEd", "Warning_NPTTexture", "The texture you are importing is not a power of two.  Non power of two textures are never streamed and have no mipmaps. Proceed?");
+		}
+		else
+		{
+			MessageText = FText::Format(NSLOCTEXT("UnrealEd", "Warning_NPTTexture_Filename", "{0} is not a power of two.  Non power of two textures are never streamed and have no mipmaps. Proceed?"), FText::FromString(FPaths::GetCleanFilename(GetCurrentFilename())));
+		}
+
+		if( EAppReturnType::Yes != FMessageDialog::Open( EAppMsgType::YesNo, MessageText ) )
 		{
 			bValid = false;
 		}
