[INTEGRATE] Change 2438211 by [REDACTED_PERSON][REDACTED_HANDLE] on 2015/02/09 11:38:04

	Fix new c++ components not showing up in the "Add Components" list after a new class is added

diff --git a/Engine/Source/Editor/UnrealEd/Private/ComponentTypeRegistry.cpp b/Engine/Source/Editor/UnrealEd/Private/ComponentTypeRegistry.cpp
--- a/Engine/Source/Editor/UnrealEd/Private/ComponentTypeRegistry.cpp
+++ b/Engine/Source/Editor/UnrealEd/Private/ComponentTypeRegistry.cpp
@@ -6,6 +6,7 @@
 #include "ComponentTypeRegistry.h"
 #include "EdGraphSchema_K2.h"
 #include "KismetEditorUtilities.h"
+#include "HotReloadInterface.h"
 
 #define LOCTEXT_NAMESPACE "ComponentTypeRegistry"
 
@@ -20,7 +21,6 @@ struct FComponentTypeRegistryData
 	virtual bool IsTickable() const override { return true; }
 	virtual TStatId GetStatId() const { RETURN_QUICK_DECLARE_CYCLE_STAT(FTypeDatabaseUpdater, STATGROUP_Tickables); }
 	/** End implementation of FTickableEditorObject */
-
 	TArray<FComponentClassComboEntryPtr> ComponentClassList;
 	TArray<FComponentTypeEntry> ComponentTypeList;
 	TArray<FAssetData> PendingAssetData;
@@ -364,6 +364,24 @@ FComponentTypeRegistry::FComponentTypeRegistry()
 {
 	Data = new FComponentTypeRegistryData();
 	Data->RefreshComponentList();
+
+	IHotReloadInterface& HotReloadSupport = FModuleManager::LoadModuleChecked<IHotReloadInterface>("HotReload");
+	HotReloadSupport.OnHotReload().AddRaw(this, &FComponentTypeRegistry::OnProjectHotReloaded);
+}
+
+FComponentTypeRegistry::~FComponentTypeRegistry()
+{
+	if( FModuleManager::Get().IsModuleLoaded("HotReload") )
+	{
+		IHotReloadInterface& HotReloadSupport = FModuleManager::GetModuleChecked<IHotReloadInterface>("HotReload");
+		HotReloadSupport.OnHotReload().RemoveAll(this);
+	}
+}
+
+void FComponentTypeRegistry::OnProjectHotReloaded( bool bWasTriggeredAutomatically )
+{
+	Data->RefreshComponentList();
+	Data->ComponentListChanged.Broadcast();
 }
 
 #undef LOCTEXT_NAMESPACE
\ No newline at end of file
diff --git a/Engine/Source/Editor/UnrealEd/Private/SComponentClassCombo.cpp b/Engine/Source/Editor/UnrealEd/Private/SComponentClassCombo.cpp
--- a/Engine/Source/Editor/UnrealEd/Private/SComponentClassCombo.cpp
+++ b/Engine/Source/Editor/UnrealEd/Private/SComponentClassCombo.cpp
@@ -8,7 +8,6 @@
 #include "ComponentTypeRegistry.h"
 #include "EditorClassUtils.h"
 #include "Engine/Selection.h"
-#include "HotReloadInterface.h"
 #include "IDocumentation.h"
 #include "KismetEditorUtilities.h"
 #include "SComponentClassCombo.h"
@@ -32,9 +31,6 @@ void SComponentClassCombo::Construct(const FArguments& InArgs)
 
 	UpdateComponentClassList();
 
-	IHotReloadInterface& HotReloadSupport = FModuleManager::LoadModuleChecked<IHotReloadInterface>("HotReload");
-	HotReloadSupport.OnHotReload().AddSP( this, &SComponentClassCombo::OnProjectHotReloaded );
-
 	SAssignNew(ComponentClassListView, SListView<FComponentClassComboEntryPtr>)
 		.ListItemsSource(&FilteredComponentClassList)
 		.OnSelectionChanged( this, &SComponentClassCombo::OnAddComponentSelectionChanged )
@@ -350,14 +346,6 @@ void SComponentClassCombo::UpdateComponentClassList()
 	GenerateFilteredComponentList(CurrentSearchString.ToString());
 }
 
-void SComponentClassCombo::OnProjectHotReloaded( bool bWasTriggeredAutomatically )
-{
-	UpdateComponentClassList();
-
-	ComponentClassListView->RequestListRefresh();
-}
-
-
 FText SComponentClassCombo::GetFriendlyComponentName(FComponentClassComboEntryPtr Entry) const
 {
 	// Get a user friendly string from the component name
diff --git a/Engine/Source/Editor/UnrealEd/Public/ComponentTypeRegistry.h b/Engine/Source/Editor/UnrealEd/Public/ComponentTypeRegistry.h
--- a/Engine/Source/Editor/UnrealEd/Public/ComponentTypeRegistry.h
+++ b/Engine/Source/Editor/UnrealEd/Public/ComponentTypeRegistry.h
@@ -34,7 +34,11 @@ struct UNREALED_API FComponentTypeRegistry
 	FOnComponentTypeListChanged& SubscribeToComponentList(const TArray<FComponentTypeEntry>*& OutComponentList);
 	FOnComponentTypeListChanged& GetOnComponentTypeListChanged();
 
+private:
+	void OnProjectHotReloaded( bool bWasTriggeredAutomatically );
+
 private:
 	FComponentTypeRegistry();
+	~FComponentTypeRegistry();
 	struct FComponentTypeRegistryData* Data;
 };
diff --git a/Engine/Source/Editor/UnrealEd/Public/SComponentClassCombo.h b/Engine/Source/Editor/UnrealEd/Public/SComponentClassCombo.h
--- a/Engine/Source/Editor/UnrealEd/Public/SComponentClassCombo.h
+++ b/Engine/Source/Editor/UnrealEd/Public/SComponentClassCombo.h
@@ -191,11 +191,6 @@ public:
 	/** Returns a component name without the substring "Component" and sanitized for display */
 	static FString GetSanitizedComponentName(FComponentClassComboEntryPtr Entry);
 
-protected:
-
-	/** Called when a project is hot reloaded to refresh the components list */
-	void OnProjectHotReloaded( bool bWasTriggeredAutomatically );
-
 private:
 
 	FText GetFriendlyComponentName(FComponentClassComboEntryPtr Entry) const;
