[INTEGRATE] Change 2438257 by [REDACTED_PERSON][REDACTED_HANDLE] on 2015/02/09 11:52:28

	Fix subdued selection outline not initializing correctly

diff --git a/Engine/Source/Editor/UnrealEd/Private/Editor.cpp b/Engine/Source/Editor/UnrealEd/Private/Editor.cpp
--- a/Engine/Source/Editor/UnrealEd/Private/Editor.cpp
+++ b/Engine/Source/Editor/UnrealEd/Private/Editor.cpp
@@ -491,6 +491,7 @@ void UEditorEngine::InitEditor(IEngineLoop* InEngineLoop)
 	// Needs to be set early as materials can be cached with selected material color baked in
 	GEngine->SetSelectedMaterialColor(ViewportSettings->bHighlightWithBrackets ? FLinearColor::Black : StyleSettings->SelectionColor);
 	GEngine->SetSelectionOutlineColor(StyleSettings->SelectionColor);
+	GEngine->SetSubduedSelectionOutlineColor(StyleSettings->GetSubduedSelectionColor());
 	GEngine->SelectionHighlightIntensity = ViewportSettings->SelectionHighlightIntensity;
 	GEngine->BSPSelectionHighlightIntensity = ViewportSettings->BSPSelectionHighlightIntensity;
 	GEngine->HoverHighlightIntensity = ViewportSettings->HoverHighlightIntensity;
@@ -554,6 +555,7 @@ void UEditorEngine::HandleSettingChanged( FName Name )
 		// Selection outline color and material color use the same color but sometimes the selected material color can be overidden so these need to be set independently
 		GEngine->SetSelectedMaterialColor(GetDefault<UEditorStyleSettings>()->SelectionColor);
 		GEngine->SetSelectionOutlineColor(GetDefault<UEditorStyleSettings>()->SelectionColor);
+		GEngine->SetSubduedSelectionOutlineColor(GetDefault<UEditorStyleSettings>()->GetSubduedSelectionColor());
 	}
 }
 
