[INTEGRATE] Change 2438293 by [REDACTED_PERSON][REDACTED_HANDLE] on 2015/02/09 12:03:10

	Editor startup splash improvements
	- Splash screen now updates with loading activity more frequently
	- Cleaned up the game name and version text
	- Cleaned up loading activity animation look a bit
	- Removed 32-bit/64-bit from splash.  The editor is always 64-bit.
	- We now display the game name after the editor version

diff --git a/Engine/Source/Editor/UnrealEd/Private/FeedbackContextEditor.cpp b/Engine/Source/Editor/UnrealEd/Private/FeedbackContextEditor.cpp
--- a/Engine/Source/Editor/UnrealEd/Private/FeedbackContextEditor.cpp
+++ b/Engine/Source/Editor/UnrealEd/Private/FeedbackContextEditor.cpp
@@ -473,17 +473,32 @@ void FFeedbackContextEditor::ProgressReported( const float TotalProgressInterp,
 	{
 		if (!DisplayMessage.IsEmpty())
 		{
-			// Animate a dot to show progress
-			FFormatOrderedArguments Args;
-			Args.Add(DisplayMessage);
-			Args.Add(int(TotalProgressInterp * 100.f));
-
-			double CurrentTime = FPlatformTime::Seconds() / 5.0;
-			int32 Progress = static_cast<int32>((CurrentTime - FMath::FloorToDouble(CurrentTime)) * 5.0);
-			static const TCHAR Spaces[] = TEXT("     ");
-			FString Dots = FString(Progress, Spaces) + TEXT(".") + FString(5 - Progress, Spaces);
-			Args.Add(FText::FromString(Dots));
-			DisplayMessage = FText::Format(NSLOCTEXT("FeedbackContextEditor", "ProgressDisplayText", "{0}{2} ({1}%)"), Args);
+			const int32 DotCount = 4;
+			const float MinTimeBetweenUpdates = 0.2f;
+			static double LastUpdateTime = -100000.0;
+			static int32 DotProgress = 0;
+			const double CurrentTime = FPlatformTime::Seconds();
+			if( CurrentTime - LastUpdateTime >= MinTimeBetweenUpdates )
+			{
+				LastUpdateTime = CurrentTime;
+				DotProgress = ( DotProgress + 1 ) % DotCount;
+			}
+
+			FString NewDisplayMessage = DisplayMessage.ToString();
+			NewDisplayMessage.RemoveFromEnd( TEXT( "..." ) );
+			for( int32 DotIndex = 0; DotIndex <= DotCount; ++DotIndex )
+			{
+				if( DotIndex <= DotProgress )
+				{
+					NewDisplayMessage.AppendChar( TCHAR( '.' ) );
+				}
+				else
+				{
+					NewDisplayMessage.AppendChar( TCHAR( ' ' ) );
+				}				
+			}
+			NewDisplayMessage.Append( FString::Printf( TEXT( " %i%%" ), int(TotalProgressInterp * 100.f) ) );
+			DisplayMessage = FText::FromString( NewDisplayMessage );
 		}
 
 		FPlatformSplash::SetSplashText(SplashTextType::StartupProgress, *DisplayMessage.ToString());
diff --git a/Engine/Source/Runtime/Core/Private/Linux/LinuxPlatformSplash.cpp b/Engine/Source/Runtime/Core/Private/Linux/LinuxPlatformSplash.cpp
--- a/Engine/Source/Runtime/Core/Private/Linux/LinuxPlatformSplash.cpp
+++ b/Engine/Source/Runtime/Core/Private/Linux/LinuxPlatformSplash.cpp
@@ -907,8 +907,18 @@ void FLinuxPlatformSplash::Show( )
 			const FText GameName = FText::FromString( FApp::GetGameName() );
 			const FText Version = FText::FromString( GEngineVersion.ToString( FEngineBuildSettings::IsPerforceBuild() ? EVersionComponent::Branch : EVersionComponent::Patch ) );
 
-			FText VersionInfo = FText::Format( NSLOCTEXT("UnrealEd", "UnrealEdTitleWithVersion_F", "Unreal Editor - {0} [Version {1}]" ), GameName, Version );
-			FText AppName =     FText::Format( NSLOCTEXT("UnrealEd", "UnrealEdTitle_F",            "Unreal Editor - {0}" ), GameName );
+			FText VersionInfo;
+			FText AppName;
+			if( GameName.IsEmpty() )
+			{
+				VersionInfo = FText::Format( NSLOCTEXT( "UnrealEd", "UnrealEdTitleWithVersionNoGameName_F", "Unreal Editor {0}" ), Version );
+				AppName = NSLOCTEXT( "UnrealEd", "UnrealEdTitleNoGameName_F", "Unreal Editor" );
+			}
+			else
+			{
+				VersionInfo = FText::Format( NSLOCTEXT( "UnrealEd", "UnrealEdTitleWithVersion_F", "Unreal Editor {0}  -  {1}" ), Version, GameName );
+				AppName = FText::Format( NSLOCTEXT( "UnrealEd", "UnrealEdTitle_F", "Unreal Editor  -  {0}" ), GameName );
+			}
 
 			StartSetSplashText( SplashTextType::VersionInfo1, VersionInfo );
 
diff --git a/Engine/Source/Runtime/Core/Private/Mac/MacPlatformSplash.cpp b/Engine/Source/Runtime/Core/Private/Mac/MacPlatformSplash.cpp
--- a/Engine/Source/Runtime/Core/Private/Mac/MacPlatformSplash.cpp
+++ b/Engine/Source/Runtime/Core/Private/Mac/MacPlatformSplash.cpp
@@ -238,17 +238,20 @@ void FMacPlatformSplash::Show()
 
 				// Set version info
 				{
-#if PLATFORM_64BITS
-					//These are invariant strings so they don't need to be localized
-					const FText PlatformBits = FText::FromString( TEXT( "64" ) );
-#else	//PLATFORM_64BITS
-					const FText PlatformBits = FText::FromString( TEXT( "32" ) );
-#endif	//PLATFORM_64BITS
-
 					const FText Version = FText::FromString( GEngineVersion.ToString( FEngineBuildSettings::IsPerforceBuild() ? EVersionComponent::Branch : EVersionComponent::Patch ) );
 
-					FText VersionInfo = FText::Format( NSLOCTEXT("UnrealEd", "UnrealEdTitleWithVersion_F", "Unreal Editor - {0} ({1}-bit) [Version {2}]" ), GameName, PlatformBits, Version );
-					FText AppName =     FText::Format( NSLOCTEXT("UnrealEd", "UnrealEdTitle_F",            "Unreal Editor - {0} ({1}-bit)" ), GameName, PlatformBits );
+					FText VersionInfo;
+					FText AppName;
+					if( GameName.IsEmpty() )
+					{
+						VersionInfo = FText::Format( NSLOCTEXT( "UnrealEd", "UnrealEdTitleWithVersionNoGameName_F", "Unreal Editor {0}" ), Version );
+						AppName = NSLOCTEXT( "UnrealEd", "UnrealEdTitleNoGameName_F", "Unreal Editor" );
+					}
+					else
+					{
+						VersionInfo = FText::Format( NSLOCTEXT( "UnrealEd", "UnrealEdTitleWithVersion_F", "Unreal Editor {0}  -  {1}" ), Version, GameName );
+						AppName = FText::Format( NSLOCTEXT( "UnrealEd", "UnrealEdTitle_F", "Unreal Editor  -  {0}" ), GameName );
+					}
 
 					StartSetSplashText( SplashTextType::VersionInfo1, VersionInfo );
 
diff --git a/Engine/Source/Runtime/Core/Private/Windows/WindowsPlatformSplash.cpp b/Engine/Source/Runtime/Core/Private/Windows/WindowsPlatformSplash.cpp
--- a/Engine/Source/Runtime/Core/Private/Windows/WindowsPlatformSplash.cpp
+++ b/Engine/Source/Runtime/Core/Private/Windows/WindowsPlatformSplash.cpp
@@ -647,17 +647,20 @@ void FWindowsPlatformSplash::Show()
 
 				// Set version info
 				{
-#if PLATFORM_64BITS
-					//These are invariant strings so they don't need to be localized
-					const FText PlatformBits = FText::FromString( TEXT( "64" ) );
-#else	//PLATFORM_64BITS
-					const FText PlatformBits = FText::FromString( TEXT( "32" ) );
-#endif	//PLATFORM_64BITS
-					
 					const FText Version = FText::FromString( GEngineVersion.ToString( FEngineBuildSettings::IsPerforceBuild() ? EVersionComponent::Branch : EVersionComponent::Patch ) );
 
-					FText VersionInfo = FText::Format( NSLOCTEXT("UnrealEd", "UnrealEdTitleWithVersion_F", "Unreal Editor - {0} ({1}-bit) [Version {2}]" ), GameName, PlatformBits, Version );
-					FText AppName =     FText::Format( NSLOCTEXT("UnrealEd", "UnrealEdTitle_F",            "Unreal Editor - {0} ({1}-bit)" ), GameName, PlatformBits );
+					FText VersionInfo;
+					FText AppName;
+					if( GameName.IsEmpty() )
+					{
+						VersionInfo = FText::Format( NSLOCTEXT( "UnrealEd", "UnrealEdTitleWithVersionNoGameName_F", "Unreal Editor {0}" ), Version );
+						AppName = NSLOCTEXT( "UnrealEd", "UnrealEdTitleNoGameName_F", "Unreal Editor" );
+					}
+					else
+					{
+						VersionInfo = FText::Format( NSLOCTEXT( "UnrealEd", "UnrealEdTitleWithVersion_F", "Unreal Editor {0}  -  {1}" ), Version, GameName );
+						AppName = FText::Format( NSLOCTEXT( "UnrealEd", "UnrealEdTitle_F", "Unreal Editor  -  {0}" ), GameName );
+					}
 
 					StartSetSplashText( SplashTextType::VersionInfo1, *VersionInfo.ToString() );
 
