Default to not including starter content if the user has not chosen otherwise, as was the previous setting.

diff --git a/Engine/Source/Editor/UnrealEd/Private/Editor/EditorGameAgnosticSettings.cpp b/Engine/Source/Editor/UnrealEd/Private/Editor/EditorGameAgnosticSettings.cpp
--- a/Engine/Source/Editor/UnrealEd/Private/Editor/EditorGameAgnosticSettings.cpp
+++ b/Engine/Source/Editor/UnrealEd/Private/Editor/EditorGameAgnosticSettings.cpp
@@ -5,7 +5,7 @@
 UEditorGameAgnosticSettings::UEditorGameAgnosticSettings(const FObjectInitializer& ObjectInitializer)
 	: Super(ObjectInitializer)
 {
-	bCopyStarterContentPreference = true;
+	bCopyStarterContentPreference = false;
 	bEditorAnalyticsEnabled = true;
 	AutoScalabilityWorkScaleAmount = 1;
 }
