[INTEGRATE] Change 2438318 by [REDACTED_PERSON][REDACTED_HANDLE] on 2015/02/09 12:20:44

	[UE-8712] - Right-clicking a component in the world brings up a context menu with basic options
	- Fixed issue where duplicated components retained incorrect AttachChildren

diff --git a/Engine/Source/Editor/Kismet/Private/SSCSEditor.cpp b/Engine/Source/Editor/Kismet/Private/SSCSEditor.cpp
--- a/Engine/Source/Editor/Kismet/Private/SSCSEditor.cpp
+++ b/Engine/Source/Editor/Kismet/Private/SSCSEditor.cpp
@@ -3374,24 +3374,6 @@ void SSCSEditor::GetSelectedItemsForContextMenu(TArray<FComponentEventConstructi
 	}
 }
 
-void SSCSEditor::OnGoToAssetInBrowser(UObject* Asset) const
-{
-	TArray<UObject*> Objects;
-	Objects.Add(Asset);
-	GEditor->SyncBrowserToObjects( Objects );
-}
-
-void SSCSEditor::OnEditBlueprint(UObject* Asset) const
-{
-	FAssetEditorManager::Get().OpenEditorForAsset(Asset);
-}
-
-void SSCSEditor::OnOpenCodeFile(const FString CodeFileName) const
-{
-	const FString AbsoluteHeaderPath = IFileManager::Get().ConvertToAbsolutePathForExternalAppForRead(*CodeFileName);
-	FSourceCodeNavigation::OpenSourceFile( AbsoluteHeaderPath );
-}
-
 TSharedPtr< SWidget > SSCSEditor::CreateContextMenu()
 {
 	TArray<FSCSEditorTreeNodePtrType> SelectedNodes = SCSTreeWidget->GetSelectedItems();
@@ -3430,103 +3412,52 @@ TSharedPtr< SWidget > SSCSEditor::CreateContextMenu()
 				}
 				if (!bShowJustPasteOption)
 				{
-					MenuBuilder.BeginSection("ComponentActions", LOCTEXT("ComponentContextMenu", "Component Actions") );
+					TArray<UActorComponent*> SelectedComponents;
+					TArray<FSCSEditorTreeNodePtrType> SelectedNodes = GetSelectedNodes();
+					for (int32 i = 0; i < SelectedNodes.Num(); ++i)
 					{
-						MenuBuilder.AddMenuEntry( FGenericCommands::Get().Cut) ;
-						MenuBuilder.AddMenuEntry( FGenericCommands::Get().Copy );
-						MenuBuilder.AddMenuEntry( FGenericCommands::Get().Paste );
-						MenuBuilder.AddMenuEntry( FGenericCommands::Get().Duplicate );
-						MenuBuilder.AddMenuEntry( FGenericCommands::Get().Delete );
-						MenuBuilder.AddMenuEntry( FGenericCommands::Get().Rename );
-
-						if (EditorMode == EComponentEditorMode::BlueprintSCS)
-						{
-							// Collect the classes of all selected objects
-							TArray<UClass*> SelectionClasses;
-							for( auto NodeIter = SelectedNodes.CreateConstIterator(); NodeIter; ++NodeIter )
-							{
-								auto TreeNode = *NodeIter;
-								if( UActorComponent* ComponentTemplate = TreeNode->GetComponentTemplate() )
-								{
-									SelectionClasses.Add(ComponentTemplate->GetClass());
-								}
-							}
+						// Get the current selected node reference
+						FSCSEditorTreeNodePtrType SelectedNodePtr = SelectedNodes[i];
+						check(SelectedNodePtr.IsValid());
 
-							if ( SelectionClasses.Num() )
-							{
-								// Find the common base class of all selected classes
-								UClass* SelectedClass = UClass::FindCommonBase( SelectionClasses );
-								// Build an event submenu if we can generate events
-								if( FBlueprintEditorUtils::CanClassGenerateEvents( SelectedClass ))
-								{
-									MenuBuilder.AddSubMenu(	LOCTEXT("AddEventSubMenu", "Add Event"), 
-										LOCTEXT("ActtionsSubMenu_ToolTip", "Add Event"), 
-										FNewMenuDelegate::CreateStatic( &SSCSEditor::BuildMenuEventsSection,
-										GetBlueprint(), SelectedClass, FCanExecuteAction::CreateSP(this, &SSCSEditor::IsEditingAllowed),
-										FGetSelectedObjectsDelegate::CreateSP(this, &SSCSEditor::GetSelectedItemsForContextMenu)));
-								}
-							}
-						}					
+						// Get the component template associated with the selected node
+						UActorComponent* ComponentTemplate = SelectedNodePtr->GetComponentTemplate();
+						if (ComponentTemplate)
+						{
+							SelectedComponents.Add(ComponentTemplate);
+						}
 					}
-					MenuBuilder.EndSection();
 
-					if (SelectedNodes.Num() == 1)
+					if (EditorMode == EComponentEditorMode::BlueprintSCS)
 					{
-						UActorComponent* ComponentTemplate = SelectedNodes[0]->GetComponentTemplate();
-						if (ComponentTemplate->GetClass()->ClassGeneratedBy)
+						// Collect the classes of all selected objects
+						TArray<UClass*> SelectionClasses;
+						for( auto NodeIter = SelectedNodes.CreateConstIterator(); NodeIter; ++NodeIter )
 						{
-							MenuBuilder.BeginSection("ComponentAsset", LOCTEXT("ComponentAssetHeading", "Asset"));
+							auto TreeNode = *NodeIter;
+							if( UActorComponent* ComponentTemplate = TreeNode->GetComponentTemplate() )
 							{
-								MenuBuilder.AddMenuEntry(
-									FText::Format(LOCTEXT("GoToBlueprintForComponent", "Edit {0}"), FText::FromString(ComponentTemplate->GetClass()->ClassGeneratedBy->GetName())),
-									LOCTEXT("EditBlueprintForComponent_ToolTip", "Edits the Blueprint Class that defines this component."),
-									FSlateIcon(FEditorStyle::GetStyleSetName(), FClassIconFinder::FindIconNameForClass(ComponentTemplate->GetClass())),
-									FUIAction(
-										FExecuteAction::CreateSP(this, &SSCSEditor::OnEditBlueprint, ComponentTemplate->GetClass()->ClassGeneratedBy),
-										FCanExecuteAction()));
-
-								MenuBuilder.AddMenuEntry(
-									LOCTEXT("GoToAssetForComponent", "Find Class in Content Browser"),
-									LOCTEXT("GoToAssetForComponent_ToolTip", "Summons the content browser and goes to the class for this component."),
-									FSlateIcon(FEditorStyle::GetStyleSetName(), "SystemWideCommands.FindInContentBrowser"),
-									FUIAction(
-										FExecuteAction::CreateSP(this, &SSCSEditor::OnGoToAssetInBrowser, ComponentTemplate->GetClass()->ClassGeneratedBy),
-										FCanExecuteAction()));
+								SelectionClasses.Add(ComponentTemplate->GetClass());
 							}
-							MenuBuilder.EndSection();			
 						}
-						else
+
+						if ( SelectionClasses.Num() )
 						{
-							MenuBuilder.BeginSection("ComponentCode", LOCTEXT("ComponentCodeHeading", "C++"));
+							// Find the common base class of all selected classes
+							UClass* SelectedClass = UClass::FindCommonBase( SelectionClasses );
+							// Build an event submenu if we can generate events
+							if( FBlueprintEditorUtils::CanClassGenerateEvents( SelectedClass ))
 							{
-								if (FSourceCodeNavigation::IsCompilerAvailable())
-								{
-									FString ClassHeaderPath;
-									if (FSourceCodeNavigation::FindClassHeaderPath(ComponentTemplate->GetClass(), ClassHeaderPath) && IFileManager::Get().FileSize(*ClassHeaderPath) != INDEX_NONE)
-									{
-										const FString CodeFileName = FPaths::GetCleanFilename(*ClassHeaderPath);
-
-										MenuBuilder.AddMenuEntry(
-											FText::Format(LOCTEXT("GoToCodeForComponent", "Open {0}"), FText::FromString(CodeFileName)),
-											FText::Format(LOCTEXT("GoToCodeForComponent_ToolTip", "Opens the header file for this component ({0}) in a code editing program"), FText::FromString(CodeFileName)),
-											FSlateIcon(),
-											FUIAction(
-												FExecuteAction::CreateSP(this, &SSCSEditor::OnOpenCodeFile, ClassHeaderPath),
-												FCanExecuteAction()));
-									}
-
-									MenuBuilder.AddMenuEntry(
-										LOCTEXT("GoToAssetForComponent", "Find Class in Content Browser"),
-										LOCTEXT("GoToAssetForComponent_ToolTip", "Summons the content browser and goes to the class for this component."),
-										FSlateIcon(FEditorStyle::GetStyleSetName(), "SystemWideCommands.FindInContentBrowser"),
-										FUIAction(
-											FExecuteAction::CreateSP(this, &SSCSEditor::OnGoToAssetInBrowser, (UObject*)ComponentTemplate->GetClass()),
-											FCanExecuteAction()));
-								}
+								MenuBuilder.AddSubMenu(	LOCTEXT("AddEventSubMenu", "Add Event"), 
+									LOCTEXT("ActtionsSubMenu_ToolTip", "Add Event"), 
+									FNewMenuDelegate::CreateStatic( &SSCSEditor::BuildMenuEventsSection,
+									GetBlueprint(), SelectedClass, FCanExecuteAction::CreateSP(this, &SSCSEditor::IsEditingAllowed),
+									FGetSelectedObjectsDelegate::CreateSP(this, &SSCSEditor::GetSelectedItemsForContextMenu)));
 							}
-							MenuBuilder.EndSection();
 						}
-					}
+					}					
+
+					FComponentEditorUtils::FillComponentContextMenuOptions(MenuBuilder, SelectedComponents);
 				}
 			}
 		}
@@ -4561,28 +4492,24 @@ void SSCSEditor::CutSelectedNodes()
 
 bool SSCSEditor::CanCopyNodes() const
 {
-	TArray<FSCSEditorTreeNodePtrType> SelectedNodes = SCSTreeWidget->GetSelectedItems();
-	bool bCanCopy = SelectedNodes.Num() > 0;
-	if(bCanCopy)
+	TArray<UActorComponent*> ComponentsToCopy;
+	TArray<FSCSEditorTreeNodePtrType> SelectedNodes = GetSelectedNodes();
+	for (int32 i = 0; i < SelectedNodes.Num(); ++i)
 	{
-		for (int32 i = 0; i < SelectedNodes.Num() && bCanCopy; ++i)
-		{
-			// Check for the default scene root; that cannot be copied/duplicated
-			UActorComponent* ComponentTemplate = SelectedNodes[i]->GetComponentTemplate();
-			bCanCopy = ComponentTemplate != nullptr && !SelectedNodes[i]->IsDefaultSceneRoot();
-			if (bCanCopy)
-			{
-				UClass* ComponentTemplateClass = ComponentTemplate->GetClass();
-				check(ComponentTemplateClass != nullptr);
+		// Get the current selected node reference
+		FSCSEditorTreeNodePtrType SelectedNodePtr = SelectedNodes[i];
+		check(SelectedNodePtr.IsValid());
 
-				// Component class cannot be abstract and must also be tagged as BlueprintSpawnable
-				bCanCopy = !ComponentTemplateClass->HasAnyClassFlags(CLASS_Abstract)
-					&& ComponentTemplateClass->HasMetaData(FBlueprintMetadata::MD_BlueprintSpawnableComponent);
-			}
+		// Get the component template associated with the selected node
+		UActorComponent* ComponentTemplate = SelectedNodePtr->GetComponentTemplate();
+		if (ComponentTemplate)
+		{
+			ComponentsToCopy.Add(ComponentTemplate);
 		}
 	}
 
-	return bCanCopy;
+	// Copy the components to the clipboard
+	return FComponentEditorUtils::CanCopyComponents(ComponentsToCopy);
 }
 
 void SSCSEditor::CopySelectedNodes()
diff --git a/Engine/Source/Editor/Kismet/Public/SSCSEditor.h b/Engine/Source/Editor/Kismet/Public/SSCSEditor.h
--- a/Engine/Source/Editor/Kismet/Public/SSCSEditor.h
+++ b/Engine/Source/Editor/Kismet/Public/SSCSEditor.h
@@ -826,12 +826,6 @@ public:
 	 */
 	static void BuildMenuEventsSection( FMenuBuilder& Menu, UBlueprint* Blueprint, UClass* SelectedClass, FCanExecuteAction CanExecuteActionDelegate, FGetSelectedObjectsDelegate GetSelectedObjectsDelegate );
 
-	void OnEditBlueprint(UObject* Blueprint) const;
-
-	void OnGoToAssetInBrowser(UObject* Asset) const;
-
-	void OnOpenCodeFile(const FString CodeFileName) const;
-
 	/**
 	 * Given an actor component, attempts to find an associated tree node.
 	 *
diff --git a/Engine/Source/Editor/LevelEditor/Private/LevelEditor.cpp b/Engine/Source/Editor/LevelEditor/Private/LevelEditor.cpp
--- a/Engine/Source/Editor/LevelEditor/Private/LevelEditor.cpp
+++ b/Engine/Source/Editor/LevelEditor/Private/LevelEditor.cpp
@@ -625,7 +625,7 @@ void FLevelEditorModule::BindGlobalLevelEditorCommands()
 		);
 
 	ActionList.MapAction( 
-		Commands.SnapCameraToActor, 
+		Commands.SnapCameraToObject,
 		FExecuteAction::CreateStatic( &FLevelEditorActionCallbacks::ExecuteExecCommand, FString( TEXT("CAMERA SNAP") ) )
 		);
 
diff --git a/Engine/Source/Editor/LevelEditor/Private/LevelEditorActions.cpp b/Engine/Source/Editor/LevelEditor/Private/LevelEditorActions.cpp
--- a/Engine/Source/Editor/LevelEditor/Private/LevelEditorActions.cpp
+++ b/Engine/Source/Editor/LevelEditor/Private/LevelEditorActions.cpp
@@ -1365,7 +1365,25 @@ bool FLevelEditorActionCallbacks::Duplicate_CanExecute()
 			return false;
 		}
 	}
-	return GUnrealEd->CanCopySelectedActorsToClipboard( GetWorld() );	// If we can copy, we can duplicate
+	
+	// If we can copy, we can duplicate
+	bool bCanCopy = false;
+	if (GEditor->GetSelectedComponentCount() > 0)
+	{
+		TArray<UActorComponent*> SelectedComponents;
+		for (FSelectedEditableComponentIterator It(GEditor->GetSelectedEditableComponentIterator()); It; ++It)
+		{
+			SelectedComponents.Add(CastChecked<UActorComponent>(*It));
+		}
+
+		bCanCopy = FComponentEditorUtils::CanCopyComponents(SelectedComponents);
+	}
+	else
+	{
+		bCanCopy = GUnrealEd->CanCopySelectedActorsToClipboard(GetWorld());
+	}
+
+	return bCanCopy;
 }
 
 bool FLevelEditorActionCallbacks::Delete_CanExecute()
@@ -1384,7 +1402,24 @@ bool FLevelEditorActionCallbacks::Delete_CanExecute()
 			return false;
 		}
 	}
-	return GUnrealEd->CanDeleteSelectedActors( GetWorld(), true, false );
+
+	bool bCanDelete = false;
+	if (GEditor->GetSelectedComponentCount() > 0)
+	{
+		TArray<UActorComponent*> SelectedComponents;
+		for (FSelectedEditableComponentIterator It(GEditor->GetSelectedEditableComponentIterator()); It; ++It)
+		{
+			SelectedComponents.Add(CastChecked<UActorComponent>(*It));
+		}
+
+		bCanDelete = FComponentEditorUtils::CanDeleteComponents(SelectedComponents);
+	}
+	else
+	{
+		bCanDelete = GUnrealEd->CanDeleteSelectedActors(GetWorld(), true, false);
+	}
+
+	return bCanDelete;
 }
 
 void FLevelEditorActionCallbacks::Rename_Execute()
@@ -1406,7 +1441,21 @@ void FLevelEditorActionCallbacks::Rename_Execute()
 
 bool FLevelEditorActionCallbacks::Rename_CanExecute()
 {
-	return GEditor->GetSelectedActorCount() == 1;
+	bool bCanRename = false;
+	if (GEditor->GetSelectedComponentCount() == 1)
+	{
+		if (UActorComponent* ComponentToRename = GEditor->GetSelectedComponents()->GetTop<UActorComponent>())
+		{
+			// We can't edit non-instance components or the default scene root
+			bCanRename = ComponentToRename->CreationMethod == EComponentCreationMethod::Instance && ComponentToRename->GetFName() != USceneComponent::GetDefaultSceneRootVariableName();
+		}
+	}
+	else
+	{
+		bCanRename = GEditor->GetSelectedActorCount() == 1;
+	}
+
+	return bCanRename;
 }
 
 bool FLevelEditorActionCallbacks::Cut_CanExecute()
@@ -1425,7 +1474,26 @@ bool FLevelEditorActionCallbacks::Cut_CanExecute()
 			return false;
 		}
 	}
-	return GUnrealEd->CanCopySelectedActorsToClipboard( GetWorld() );	// If we can copy, we can cut
+
+	bool bCanCut = false;
+	if (GEditor->GetSelectedComponentCount() > 0)
+	{
+		// Make sure the components can be copied and deleted
+		TArray<UActorComponent*> SelectedComponents;
+		for (FSelectedEditableComponentIterator It(GEditor->GetSelectedEditableComponentIterator()); It; ++It)
+		{
+			SelectedComponents.Add(CastChecked<UActorComponent>(*It));
+		}
+
+		bCanCut = FComponentEditorUtils::CanCopyComponents(SelectedComponents) && FComponentEditorUtils::CanDeleteComponents(SelectedComponents);
+	}
+	else
+	{
+		// For actors, if we can copy, we can cut
+		bCanCut = GUnrealEd->CanCopySelectedActorsToClipboard(GetWorld());
+	}
+
+	return bCanCut;
 }
 
 bool FLevelEditorActionCallbacks::Copy_CanExecute()
@@ -1444,7 +1512,24 @@ bool FLevelEditorActionCallbacks::Copy_CanExecute()
 			return false;
 		}
 	}
-	return GUnrealEd->CanCopySelectedActorsToClipboard( GetWorld() );
+
+	bool bCanCopy = false;
+	if (GEditor->GetSelectedComponentCount() > 0)
+	{
+		TArray<UActorComponent*> SelectedComponents;
+		for (FSelectedEditableComponentIterator It(GEditor->GetSelectedEditableComponentIterator()); It; ++It)
+		{
+			SelectedComponents.Add(CastChecked<UActorComponent>(*It));
+		}
+
+		bCanCopy = FComponentEditorUtils::CanCopyComponents(SelectedComponents);
+	}
+	else
+	{
+		bCanCopy = GUnrealEd->CanCopySelectedActorsToClipboard(GetWorld());
+	}
+
+	return bCanCopy;
 }
 
 bool FLevelEditorActionCallbacks::Paste_CanExecute()
@@ -2773,7 +2858,7 @@ void FLevelEditorCommands::RegisterCommands()
 
 	UI_COMMAND( GoHere, "Go Here", "Moves the camera to the current mouse position", EUserInterfaceActionType::Button, FInputGesture() );
 
-	UI_COMMAND( SnapCameraToActor, "Snap View to Actor", "Snaps the view to the selected actors", EUserInterfaceActionType::Button, FInputGesture() );
+	UI_COMMAND( SnapCameraToObject, "Snap View to Object", "Snaps the view to the selected object", EUserInterfaceActionType::Button, FInputGesture() );
 
 	UI_COMMAND( GoToCodeForActor, "Go to C++ Code for Actor", "Opens a code editing IDE and navigates to the source file associated with the seleced actor", EUserInterfaceActionType::Button, FInputGesture() );
 	UI_COMMAND( GoToDocsForActor, "Go to Documentation for Actor", "Opens documentation for the Actor in the default web browser", EUserInterfaceActionType::Button, FInputGesture() );
diff --git a/Engine/Source/Editor/LevelEditor/Private/LevelEditorContextMenu.cpp b/Engine/Source/Editor/LevelEditor/Private/LevelEditorContextMenu.cpp
--- a/Engine/Source/Editor/LevelEditor/Private/LevelEditorContextMenu.cpp
+++ b/Engine/Source/Editor/LevelEditor/Private/LevelEditorContextMenu.cpp
@@ -31,6 +31,7 @@
 #include "GenericCommands.h"
 #include "Engine/Selection.h"
 #include "Kismet2/BlueprintEditorUtils.h"
+#include "ComponentEditorUtils.h"
 
 #define LOCTEXT_NAMESPACE "LevelViewportContextMenu"
 
@@ -155,10 +156,15 @@ void FLevelEditorContextMenu::FillMenu( FMenuBuilder& MenuBuilder, TWeakPtr<SLev
 
 	if (GEditor->GetSelectedComponentCount() > 0)
 	{
-		// Show component-specific context menu
-		MenuBuilder.BeginSection("Component", LOCTEXT("ComponentHeading", "Component"));
+		TArray<UActorComponent*> SelectedComponents;
+		for (FSelectedEditableComponentIterator It(GEditor->GetSelectedEditableComponentIterator()); It; ++It)
 		{
-			auto OwnerActor = Cast<AActor>(*(GEditor->GetSelectedActorIterator()));
+			SelectedComponents.Add(CastChecked<UActorComponent>(*It));
+		}
+
+		MenuBuilder.BeginSection("ComponentControl", LOCTEXT("ComponentControlHeading", "Component"));
+		{
+			auto OwnerActor = GEditor->GetSelectedActors()->GetTop<AActor>();
 			check(OwnerActor);
 
 			MenuBuilder.AddMenuEntry(
@@ -168,8 +174,20 @@ void FLevelEditorContextMenu::FillMenu( FMenuBuilder& MenuBuilder, TWeakPtr<SLev
 				TAttribute<FText>(),
 				FSlateIcon(FEditorStyle::GetStyleSetName(), FClassIconFinder::FindIconNameForClass(OwnerActor->GetClass()))
 				);
+
+			MenuBuilder.AddMenuEntry(FEditorViewportCommands::Get().FocusViewportToSelection);
+
+
+			const FVector* ClickLocation = &GEditor->ClickLocation;
+			FUIAction GoHereAction;
+			GoHereAction.ExecuteAction = FExecuteAction::CreateStatic(&FLevelEditorActionCallbacks::GoHere_Clicked, ClickLocation);
+
+			MenuBuilder.AddMenuEntry(FLevelEditorCommands::Get().GoHere);
+			MenuBuilder.AddMenuEntry(FLevelEditorCommands::Get().SnapCameraToObject);
 		}
 		MenuBuilder.EndSection();
+
+		FComponentEditorUtils::FillComponentContextMenuOptions(MenuBuilder, SelectedComponents);
 	}
 	else
 	{
@@ -256,7 +274,7 @@ void FLevelEditorContextMenu::FillMenu( FMenuBuilder& MenuBuilder, TWeakPtr<SLev
 			GoHereAction.ExecuteAction = FExecuteAction::CreateStatic(&FLevelEditorActionCallbacks::GoHere_Clicked, ClickLocation);
 
 			MenuBuilder.AddMenuEntry(FLevelEditorCommands::Get().GoHere);
-			MenuBuilder.AddMenuEntry(FLevelEditorCommands::Get().SnapCameraToActor);
+			MenuBuilder.AddMenuEntry(FLevelEditorCommands::Get().SnapCameraToObject);
 
 			if (SelectedActors.Num() == 1)
 			{
diff --git a/Engine/Source/Editor/LevelEditor/Public/LevelEditorActions.h b/Engine/Source/Editor/LevelEditor/Public/LevelEditorActions.h
--- a/Engine/Source/Editor/LevelEditor/Public/LevelEditorActions.h
+++ b/Engine/Source/Editor/LevelEditor/Public/LevelEditorActions.h
@@ -110,8 +110,8 @@ public:
 	/** Moves the camera to the current mouse position */
 	TSharedPtr< FUICommandInfo > GoHere;
 
-	/** Snaps the camera to the selected actors. */
-	TSharedPtr< FUICommandInfo > SnapCameraToActor;
+	/** Snaps the camera to the selected object. */
+	TSharedPtr< FUICommandInfo > SnapCameraToObject;
 
 	/** Goes to the source code for the selected actor's class. */
 	TSharedPtr< FUICommandInfo > GoToCodeForActor;
diff --git a/Engine/Source/Editor/UnrealEd/Classes/Editor/EditorEngine.h b/Engine/Source/Editor/UnrealEd/Classes/Editor/EditorEngine.h
--- a/Engine/Source/Editor/UnrealEd/Classes/Editor/EditorEngine.h
+++ b/Engine/Source/Editor/UnrealEd/Classes/Editor/EditorEngine.h
@@ -241,6 +241,11 @@ struct FActorOrComponent
 		return Actor ? Actor->GetActorLocation() : Component->GetComponentLocation();
 	}
 
+	FRotator GetWorldRotation() const
+	{
+		return Actor ? Actor->GetActorRotation() : Component->GetComponentRotation();
+	}
+
 	void SetWorldLocation( const FVector& NewLocation )
 	{
 		if( Actor )
@@ -957,7 +962,7 @@ public:
 	 *
 	 * @param	Actor	The actor the camera is going to be snapped to.
 	 */
-	void SnapViewToActor(const AActor &Actor);
+	void SnapViewTo(const FActorOrComponent& Object);
 
 	/**
 	 * Remove the roll, pitch and/or yaw from the perspective viewports' cameras.
diff --git a/Engine/Source/Editor/UnrealEd/Private/EditorServer.cpp b/Engine/Source/Editor/UnrealEd/Private/EditorServer.cpp
--- a/Engine/Source/Editor/UnrealEd/Private/EditorServer.cpp
+++ b/Engine/Source/Editor/UnrealEd/Private/EditorServer.cpp
@@ -4641,7 +4641,7 @@ void UEditorEngine::MoveActorInFrontOfCamera( AActor& InActor, const FVector& In
 }
 
 
-void UEditorEngine::SnapViewToActor(const AActor &Actor)
+void UEditorEngine::SnapViewTo(const FActorOrComponent& Object)
 {
 	for(int32 ViewportIndex = 0; ViewportIndex < LevelViewportClients.Num(); ++ViewportIndex)
 	{
@@ -4649,8 +4649,8 @@ void UEditorEngine::SnapViewToActor(const AActor &Actor)
 
 		if ( ViewportClient->IsPerspective()  )
 		{
-			ViewportClient->SetViewLocation( Actor.GetActorLocation() );
-			ViewportClient->SetViewRotation( Actor.GetActorRotation() );
+			ViewportClient->SetViewLocation( Object.GetWorldLocation() );
+			ViewportClient->SetViewRotation( Object.GetWorldRotation() );
 			ViewportClient->Invalidate();
 		}
 	}
@@ -4754,13 +4754,17 @@ bool UEditorEngine::Exec_Camera( const TCHAR* Str, FOutputDevice& Ar )
 	}
 	else if ( bSnap )
 	{
-		TargetSelectedActor = GEditor->GetSelectedActors()->GetTop<AActor>();
-
-		if (TargetSelectedActor)
+		FActorOrComponent SelectedObject(GEditor->GetSelectedComponents()->GetTop<USceneComponent>());
+		if (!SelectedObject.IsValid())
+		{
+			SelectedObject.Actor = GEditor->GetSelectedActors()->GetTop<AActor>();
+		}
+		
+		if (SelectedObject.IsValid())
 		{
 			// Set perspective viewport camera parameters to that of the selected camera.
-			SnapViewToActor(*TargetSelectedActor);
-			Ar.Log( TEXT("Snapped camera to the first selected actor.") );
+			SnapViewTo(SelectedObject);
+			Ar.Log( TEXT("Snapped camera to the first selected object.") );
 		}
 	}
 
diff --git a/Engine/Source/Editor/UnrealEd/Private/Kismet2/ComponentEditorUtils.cpp b/Engine/Source/Editor/UnrealEd/Private/Kismet2/ComponentEditorUtils.cpp
--- a/Engine/Source/Editor/UnrealEd/Private/Kismet2/ComponentEditorUtils.cpp
+++ b/Engine/Source/Editor/UnrealEd/Private/Kismet2/ComponentEditorUtils.cpp
@@ -10,6 +10,12 @@
 #include "BlueprintEditorUtils.h"
 #include "Factories.h"
 #include "UnrealExporter.h"
+#include "GenericCommands.h"
+#include "SourceCodeNavigation.h"
+#include "ClassIconFinder.h"
+#include "AssetEditorManager.h"
+
+#define LOCTEXT_NAMESPACE "ComponentEditorUtils"
 
 // Text object factory for pasting components
 struct FComponentObjectTextFactory : public FCustomizableTextObjectFactory
@@ -208,6 +214,31 @@ USceneComponent* FComponentEditorUtils::FindClosestParentInList(UActorComponent*
 	return ClosestParentComponent;
 }
 
+bool FComponentEditorUtils::CanCopyComponents(const TArray<UActorComponent*>& ComponentsToCopy)
+{
+	bool bCanCopy = ComponentsToCopy.Num() > 0;
+	if (bCanCopy)
+	{
+		for (int32 i = 0; i < ComponentsToCopy.Num() && bCanCopy; ++i)
+		{
+			// Check for the default scene root; that cannot be copied/duplicated
+			UActorComponent* Component = ComponentsToCopy[i];
+			bCanCopy = Component != nullptr && Component->GetFName() != USceneComponent::GetDefaultSceneRootVariableName();
+			if (bCanCopy)
+			{
+				UClass* ComponentClass = Component->GetClass();
+				check(ComponentClass != nullptr);
+
+				// Component class cannot be abstract and must also be tagged as BlueprintSpawnable
+				bCanCopy = !ComponentClass->HasAnyClassFlags(CLASS_Abstract)
+					&& ComponentClass->HasMetaData(FBlueprintMetadata::MD_BlueprintSpawnableComponent);
+			}
+		}
+	}
+
+	return bCanCopy;
+}
+
 void FComponentEditorUtils::CopyComponents(const TArray<UActorComponent*>& ComponentsToCopy)
 {
 	FStringOutputDevice Archive;
@@ -322,10 +353,7 @@ void FComponentEditorUtils::PasteComponents(TArray<UActorComponent*>& OutPastedC
 				}
 			}
 
-			//So, if we're pasting a component that was the root, we want to attach keeping the world location
-			// Otherwise, we want to attach keeping the relative transform
-			// So how can we tell if a pasted component was a root?
-
+			//@todo: Fix pasting when the pasted component was a root
 			//NewSceneComponent->UpdateComponentToWorld();
 			if (NewComponentParent)
 			{
@@ -360,6 +388,23 @@ void FComponentEditorUtils::GetComponentsFromClipboard(TMap<FName, FName>& OutPa
 	OutNewObjectMap = Factory->NewObjectMap;
 }
 
+bool FComponentEditorUtils::CanDeleteComponents(const TArray<UActorComponent*>& ComponentsToDelete)
+{
+	bool bCanDelete = true;
+	for (auto ComponentToDelete : ComponentsToDelete)
+	{
+		// We can't delete non-instance components or the default scene root
+		if (ComponentToDelete->CreationMethod != EComponentCreationMethod::Instance 
+			|| ComponentToDelete->GetFName() == USceneComponent::GetDefaultSceneRootVariableName())
+		{
+			bCanDelete = false;
+			break;
+		}
+	}
+
+	return bCanDelete;
+}
+
 int32 FComponentEditorUtils::DeleteComponents(const TArray<UActorComponent*>& ComponentsToDelete, UActorComponent*& OutComponentToSelect)
 {
 	int32 NumDeletedComponents = 0;
@@ -445,19 +490,23 @@ UActorComponent* FComponentEditorUtils::DuplicateComponent(UActorComponent* Temp
 		bool bKeepWorldLocationOnAttach = false;
 		NewCloneComponent = ConstructObject<UActorComponent>(ComponentClass, Actor, NewComponentName, RF_Transactional, TemplateComponent);
 		
-		// ComponentToWorld is not a UPROPERTY, so make sure the clone has calculated it properly
-		NewCloneComponent->UpdateComponentToWorld();
-
-		// If the clone is a scene component without an attach parent, attach it to the root (can happen when duplicating the root component)
 		auto NewSceneComponent = Cast<USceneComponent>(NewCloneComponent);
-		if (NewSceneComponent && !NewSceneComponent->GetAttachParent())
+		if (NewSceneComponent)
 		{
-			USceneComponent* RootComponent = Actor->GetRootComponent();
+			// Ensure the clone doesn't think it has children
+			NewSceneComponent->AttachChildren.Empty();
 
-			// There should be no situation in which a scene component is duplicated when a root doesn't exist
-			check(RootComponent);
+			// If the clone is a scene component without an attach parent, attach it to the root (can happen when duplicating the root component)
+			if (!NewSceneComponent->GetAttachParent())
+			{
+				USceneComponent* RootComponent = Actor->GetRootComponent();
+				check(RootComponent);
+
+				// ComponentToWorld is not a UPROPERTY, so make sure the clone has calculated it properly before attachment
+				NewSceneComponent->UpdateComponentToWorld();
 
-			NewSceneComponent->AttachTo(RootComponent, NAME_None, EAttachLocation::KeepWorldPosition);
+				NewSceneComponent->AttachTo(RootComponent, NAME_None, EAttachLocation::KeepWorldPosition);
+			}
 		}
 
 		// Add to SerializedComponents array so it gets saved
@@ -539,7 +588,7 @@ bool FComponentEditorUtils::AttemptApplyMaterialToComponent(USceneComponent* Sce
 	if (MeshComponent || DecalComponent)
 	{
 		bResult = true;
-		const FScopedTransaction Transaction(NSLOCTEXT("UnrealEd", "DropTarget_UndoSetComponentMaterial", "Assign Material to Component (Drag and Drop)"));
+		const FScopedTransaction Transaction(LOCTEXT("DropTarget_UndoSetComponentMaterial", "Assign Material to Component (Drag and Drop)"));
 		SceneComponent->Modify();
 
 		if (MeshComponent)
@@ -623,3 +672,97 @@ FName FComponentEditorUtils::FindVariableNameGivenComponentInstance(UActorCompon
 
 	return NAME_None;
 }
+
+void FComponentEditorUtils::FillComponentContextMenuOptions(FMenuBuilder& MenuBuilder, const TArray<UActorComponent*>& SelectedComponents)
+{
+	// Basic commands
+	MenuBuilder.BeginSection("EditComponent", LOCTEXT("EditComponentHeading", "Edit"));
+	{
+		MenuBuilder.AddMenuEntry(FGenericCommands::Get().Cut);
+		MenuBuilder.AddMenuEntry(FGenericCommands::Get().Copy);
+		MenuBuilder.AddMenuEntry(FGenericCommands::Get().Paste);
+		MenuBuilder.AddMenuEntry(FGenericCommands::Get().Duplicate);
+		MenuBuilder.AddMenuEntry(FGenericCommands::Get().Delete);
+		MenuBuilder.AddMenuEntry(FGenericCommands::Get().Rename);
+	}
+	MenuBuilder.EndSection();
+
+	if (SelectedComponents.Num() == 1)
+	{
+		UActorComponent* Component = SelectedComponents[0];
+
+		if (Component->GetClass()->ClassGeneratedBy)
+		{
+			MenuBuilder.BeginSection("ComponentAsset", LOCTEXT("ComponentAssetHeading", "Asset"));
+			{
+				MenuBuilder.AddMenuEntry(
+					FText::Format(LOCTEXT("GoToBlueprintForComponent", "Edit {0}"), FText::FromString(Component->GetClass()->ClassGeneratedBy->GetName())),
+					LOCTEXT("EditBlueprintForComponent_ToolTip", "Edits the Blueprint Class that defines this component."),
+					FSlateIcon(FEditorStyle::GetStyleSetName(), FClassIconFinder::FindIconNameForClass(Component->GetClass())),
+					FUIAction(
+					FExecuteAction::CreateStatic(&FComponentEditorUtils::OnEditBlueprintComponent, Component->GetClass()->ClassGeneratedBy),
+					FCanExecuteAction()));
+
+				MenuBuilder.AddMenuEntry(
+					LOCTEXT("GoToAssetForComponent", "Find Class in Content Browser"),
+					LOCTEXT("GoToAssetForComponent_ToolTip", "Summons the content browser and goes to the class for this component."),
+					FSlateIcon(FEditorStyle::GetStyleSetName(), "SystemWideCommands.FindInContentBrowser"),
+					FUIAction(
+					FExecuteAction::CreateStatic(&FComponentEditorUtils::OnGoToComponentAssetInBrowser, Component->GetClass()->ClassGeneratedBy),
+					FCanExecuteAction()));
+			}
+			MenuBuilder.EndSection();
+		}
+		else
+		{
+			MenuBuilder.BeginSection("ComponentCode", LOCTEXT("ComponentCodeHeading", "C++"));
+			{
+				if (FSourceCodeNavigation::IsCompilerAvailable())
+				{
+					FString ClassHeaderPath;
+					if (FSourceCodeNavigation::FindClassHeaderPath(Component->GetClass(), ClassHeaderPath) && IFileManager::Get().FileSize(*ClassHeaderPath) != INDEX_NONE)
+					{
+						const FString CodeFileName = FPaths::GetCleanFilename(*ClassHeaderPath);
+
+						MenuBuilder.AddMenuEntry(
+							FText::Format(LOCTEXT("GoToCodeForComponent", "Open {0}"), FText::FromString(CodeFileName)),
+							FText::Format(LOCTEXT("GoToCodeForComponent_ToolTip", "Opens the header file for this component ({0}) in a code editing program"), FText::FromString(CodeFileName)),
+							FSlateIcon(),
+							FUIAction(
+							FExecuteAction::CreateStatic(&FComponentEditorUtils::OnOpenComponentCodeFile, ClassHeaderPath),
+							FCanExecuteAction()));
+					}
+
+					MenuBuilder.AddMenuEntry(
+						LOCTEXT("GoToAssetForComponent", "Find Class in Content Browser"),
+						LOCTEXT("GoToAssetForComponent_ToolTip", "Summons the content browser and goes to the class for this component."),
+						FSlateIcon(FEditorStyle::GetStyleSetName(), "SystemWideCommands.FindInContentBrowser"),
+						FUIAction(
+						FExecuteAction::CreateStatic(&FComponentEditorUtils::OnGoToComponentAssetInBrowser, (UObject*)Component->GetClass()),
+						FCanExecuteAction()));
+				}
+			}
+			MenuBuilder.EndSection();
+		}
+	}
+}
+
+void FComponentEditorUtils::OnGoToComponentAssetInBrowser(UObject* Asset)
+{
+	TArray<UObject*> Objects;
+	Objects.Add(Asset);
+	GEditor->SyncBrowserToObjects(Objects);
+}
+
+void FComponentEditorUtils::OnOpenComponentCodeFile(const FString CodeFileName)
+{
+	const FString AbsoluteHeaderPath = IFileManager::Get().ConvertToAbsolutePathForExternalAppForRead(*CodeFileName);
+	FSourceCodeNavigation::OpenSourceFile(AbsoluteHeaderPath);
+}
+
+void FComponentEditorUtils::OnEditBlueprintComponent(UObject* Blueprint)
+{
+	FAssetEditorManager::Get().OpenEditorForAsset(Blueprint);
+}
+
+#undef LOCTEXT_NAMESPACE
\ No newline at end of file
diff --git a/Engine/Source/Editor/UnrealEd/Public/Kismet2/ComponentEditorUtils.h b/Engine/Source/Editor/UnrealEd/Public/Kismet2/ComponentEditorUtils.h
--- a/Engine/Source/Editor/UnrealEd/Public/Kismet2/ComponentEditorUtils.h
+++ b/Engine/Source/Editor/UnrealEd/Public/Kismet2/ComponentEditorUtils.h
@@ -24,6 +24,13 @@ public:
 	/** Generate a valid variable name string for the given component instance based on the name of the asset referenced by the component */
 	static FString GenerateValidVariableNameFromAsset(UObject* Asset, AActor* ComponentOwner);
 
+	/**
+	 * Checks whether it is valid to copy the indicated components
+	 * @param ComponentsToCopy The list of components to check
+	 * @return Whether the indicated components can be copied
+	 */
+	static bool CanCopyComponents(const TArray<UActorComponent*>& ComponentsToCopy);
+
 	/**
 	 * Copies the selected components to the clipboard
 	 * @param ComponentsToCopy The list of components to copy
@@ -53,6 +60,13 @@ public:
 	 */
 	static void GetComponentsFromClipboard(TMap<FName, FName>& OutParentMap, TMap<FName, UActorComponent*>& OutNewObjectMap, bool bGetComponentsAsArchetypes);
 
+	/**
+	 * Determines whether the indicated components can be deleted
+	 * @param ComponentsToDelete The list of components to determine can be deleted
+	 * @return Whether the indicated components can be deleted
+	 */
+	static bool CanDeleteComponents(const TArray<UActorComponent*>& ComponentsToDelete);
+
 	/**
 	 * Deletes the indicated components and identifies the component that should be selected following the operation.
 	 * Note: Does not take care of the actual selection of a new component. It only identifies which component should be selected.
@@ -179,5 +193,17 @@ public:
 	// Try to find the correct variable name for a given native component template or instance (which can have a mismatch)
 	static FName FindVariableNameGivenComponentInstance(UActorComponent* ComponentInstance);
 
+	/**
+	* Populates the given menu with basic options for operations on components in the world.
+	* @param MenuBuilder Used to create the menu options
+	* @param SelectedComponents The selected components to create menu options for
+	*/
+	static void FillComponentContextMenuOptions(FMenuBuilder& MenuBuilder, const TArray<UActorComponent*>& SelectedComponents);
+
+private:	
 	static USceneComponent* FindClosestParentInList(UActorComponent* ChildComponent, const TArray<UActorComponent*>& ComponentList);
+
+	static void OnGoToComponentAssetInBrowser(UObject* Asset);
+	static void OnOpenComponentCodeFile(const FString CodeFileName);
+	static void OnEditBlueprintComponent(UObject* Blueprint);
 };
\ No newline at end of file
