[INTEGRATE] Change 2438573 by [REDACTED_PERSON][REDACTED_HANDLE] on 2015/02/09 14:56:36

	The inner arc of the rotation gizmo/widget no longer draws a hit proxy, so clicks pass through as expected when translucent selection is allowed

diff --git a/Engine/Source/Editor/UnrealEd/Private/UnrealWidget.cpp b/Engine/Source/Editor/UnrealEd/Private/UnrealWidget.cpp
--- a/Engine/Source/Editor/UnrealEd/Private/UnrealWidget.cpp
+++ b/Engine/Source/Editor/UnrealEd/Private/UnrealWidget.cpp
@@ -1440,7 +1440,8 @@ void FWidget::DrawPartialRotationArc(const FSceneView* View, FPrimitiveDrawInter
 	}
 	PDI->SetHitProxy( NULL );
 
-	if (bIsPerspective)
+	const bool bIsHitProxyView = View->Family->EngineShowFlags.HitProxies;
+	if (bIsPerspective && !bIsHitProxyView && !PDI->IsHitTesting())
 	{
 		FThickArcParams InnerArcParams(PDI, InLocation, GridMaterial, 0.0f, InnerRadius);
 		FColor InnerColor = InColor;
