[INTEGRATE] Change 2438772 by [REDACTED_PERSON][REDACTED_HANDLE] on 2015/02/09 16:58:33

	Editor: Update 'create asset from object', 'delete', and 'parent class picker' dialogs to use colorized flat buttons
	Editor: Add FlatButton.DefaultTextStyle and change brightness on flat button to be brighter by default

diff --git a/Engine/Source/Editor/EditorStyle/Private/SlateEditorStyle.cpp b/Engine/Source/Editor/EditorStyle/Private/SlateEditorStyle.cpp
--- a/Engine/Source/Editor/EditorStyle/Private/SlateEditorStyle.cpp
+++ b/Engine/Source/Editor/EditorStyle/Private/SlateEditorStyle.cpp
@@ -341,11 +341,19 @@ void FSlateEditorStyle::FStyle::SetupGeneralStyles()
 			);
 
 		Set("FlatButton.Default", FButtonStyle(Button)
-			.SetNormal(BOX_BRUSH("Common/FlatButton", 2.0f / 8.0f, FLinearColor(1, 1, 1, 0.1f)))
+			.SetNormal(BOX_BRUSH("Common/FlatButton", 2.0f / 8.0f, FLinearColor(0.125f, 0.125f, 0.125f, 0.8f)))
 			.SetHovered(BOX_BRUSH("Common/FlatButton", 2.0f / 8.0f, SelectionColor))
 			.SetPressed(BOX_BRUSH("Common/FlatButton", 2.0f / 8.0f, SelectionColor_Pressed))
 			);
 
+		Set("FlatButton.DefaultTextStyle", FTextBlockStyle(NormalText)
+			.SetFont(TTF_CORE_FONT("Fonts/Roboto-Bold", 10))
+			.SetColorAndOpacity(FLinearColor(1.0f, 1.0f, 1.0f))
+			.SetHighlightColor(FLinearColor(1.0f, 1.0f, 1.0f))
+			.SetShadowOffset(FVector2D(1, 1))
+			.SetShadowColorAndOpacity(FLinearColor(0, 0, 0, 0.9f)));
+
+
 		struct ButtonColor
 		{
 		public:
@@ -363,9 +371,9 @@ void FSlateEditorStyle::FStyle::SetupGeneralStyles()
 		for ( const ButtonColor& Entry : FlatButtons )
 		{
 			Set(Entry.Name, FButtonStyle(Button)
-				.SetNormal(BOX_BRUSH("Common/FlatButton", 2.0f / 8.0f, FLinearColor(Entry.Color.X, Entry.Color.Y, Entry.Color.Z, 0.6f)))
-				.SetHovered(BOX_BRUSH("Common/FlatButton", 2.0f / 8.0f, FLinearColor(Entry.Color.X, Entry.Color.Y, Entry.Color.Z, 1.f)))
-				.SetPressed(BOX_BRUSH("Common/FlatButton", 2.0f / 8.0f, FLinearColor(Entry.Color.X, Entry.Color.Y, Entry.Color.Z, 0.8f)))
+				.SetNormal(BOX_BRUSH("Common/FlatButton", 2.0f / 8.0f, FLinearColor(Entry.Color.X, Entry.Color.Y, Entry.Color.Z, 0.8f)))
+				.SetHovered(BOX_BRUSH("Common/FlatButton", 2.0f / 8.0f, FLinearColor(Entry.Color.X, Entry.Color.Y, Entry.Color.Z, 1.0f)))
+				.SetPressed(BOX_BRUSH("Common/FlatButton", 2.0f / 8.0f, FLinearColor(Entry.Color.X, Entry.Color.Y, Entry.Color.Z, 0.6f)))
 				);
 		}
 
diff --git a/Engine/Source/Editor/UnrealEd/Private/Dialogs/SDeleteAssetsDialog.cpp b/Engine/Source/Editor/UnrealEd/Private/Dialogs/SDeleteAssetsDialog.cpp
--- a/Engine/Source/Editor/UnrealEd/Private/Dialogs/SDeleteAssetsDialog.cpp
+++ b/Engine/Source/Editor/UnrealEd/Private/Dialogs/SDeleteAssetsDialog.cpp
@@ -386,7 +386,9 @@ TSharedRef<SWidget> SDeleteAssetsDialog::BuildDeleteDialog()
 						.HAlign( HAlign_Center )
 						.Text( LOCTEXT( "Delete", "Delete" ) )
 						.ToolTipText( LOCTEXT( "DeleteTooltipText", "Perform the delete" ) )
-						.OnClicked( this, &SDeleteAssetsDialog::Delete )
+						.ButtonStyle(FEditorStyle::Get(), "FlatButton.Danger")
+						.TextStyle(FEditorStyle::Get(), "FlatButton.DefaultTextStyle")
+						.OnClicked(this, &SDeleteAssetsDialog::Delete)
 					]
 				]
 
@@ -402,6 +404,8 @@ TSharedRef<SWidget> SDeleteAssetsDialog::BuildDeleteDialog()
 						.HAlign( HAlign_Center )
 						.Text( LOCTEXT( "Cancel", "Cancel" ) )
 						.ToolTipText( LOCTEXT( "CancelDeleteTooltipText", "Cancel the delete" ) )
+						.ButtonStyle(FEditorStyle::Get(), "FlatButton.Default")
+						.TextStyle(FEditorStyle::Get(), "FlatButton.DefaultTextStyle")
 						.OnClicked( this, &SDeleteAssetsDialog::Cancel )
 					]
 				]
@@ -508,6 +512,8 @@ TSharedRef<SWidget> SDeleteAssetsDialog::BuildReplaceReferencesWidget()
 		.HAlign( HAlign_Center )
 		.Text( LOCTEXT( "Replace References", "Replace References" ) )
 		.OnClicked( this, &SDeleteAssetsDialog::ReplaceReferences )
+		.ButtonStyle(FEditorStyle::Get(), "FlatButton.Danger")
+		.TextStyle(FEditorStyle::Get(), "FlatButton.DefaultTextStyle")
 	];
 }
 
@@ -532,7 +538,9 @@ TSharedRef<SWidget> SDeleteAssetsDialog::BuildForceDeleteWidget()
 		.HAlign( HAlign_Center )
 		.Text( LOCTEXT( "ForceDelete", "Force Delete" ) )
 		.ToolTipText( LOCTEXT( "ForceDeleteTooltipText", "Force Delete will obliterate all references to this asset and is dangerous.\n\nUse as a last resort." ) )
-		.OnClicked( this, &SDeleteAssetsDialog::ForceDelete )
+		.ButtonStyle(FEditorStyle::Get(), "FlatButton.Danger")
+		.TextStyle(FEditorStyle::Get(), "FlatButton.DefaultTextStyle")
+		.OnClicked(this, &SDeleteAssetsDialog::ForceDelete)
 	];
 }
 
diff --git a/Engine/Source/Editor/UnrealEd/Private/SClassPickerDialog.cpp b/Engine/Source/Editor/UnrealEd/Private/SClassPickerDialog.cpp
--- a/Engine/Source/Editor/UnrealEd/Private/SClassPickerDialog.cpp
+++ b/Engine/Source/Editor/UnrealEd/Private/SClassPickerDialog.cpp
@@ -107,6 +107,8 @@ void SClassPickerDialog::Construct(const FArguments& InArgs)
 						.Visibility( this, &SClassPickerDialog::GetSelectButtonVisibility )
 						.ContentPadding(FEditorStyle::GetMargin("StandardDialog.ContentPadding"))
 						.OnClicked(this, &SClassPickerDialog::OnClassPickerConfirmed)
+						.ButtonStyle(FEditorStyle::Get(), "FlatButton.Success")
+						.TextStyle(FEditorStyle::Get(), "FlatButton.DefaultTextStyle")
 					]
 					+SUniformGridPanel::Slot(1,0)
 					[
@@ -115,6 +117,8 @@ void SClassPickerDialog::Construct(const FArguments& InArgs)
 						.HAlign(HAlign_Center)
 						.ContentPadding(FEditorStyle::GetMargin("StandardDialog.ContentPadding"))
 						.OnClicked(this, &SClassPickerDialog::OnClassPickerCanceled)
+						.ButtonStyle(FEditorStyle::Get(), "FlatButton.Default")
+						.TextStyle(FEditorStyle::Get(), "FlatButton.DefaultTextStyle")
 					]
 				]
 			]
diff --git a/Engine/Source/Editor/UnrealEd/Private/SCreateAssetFromObject.cpp b/Engine/Source/Editor/UnrealEd/Private/SCreateAssetFromObject.cpp
--- a/Engine/Source/Editor/UnrealEd/Private/SCreateAssetFromObject.cpp
+++ b/Engine/Source/Editor/UnrealEd/Private/SCreateAssetFromObject.cpp
@@ -101,10 +101,9 @@ void SCreateAssetFromObject::Construct(const FArguments& InArgs, TSharedPtr<SWin
 				.ContentPadding(FMargin(8, 2, 8, 2))
 				.OnClicked(this, &SCreateAssetFromObject::OnCreateAssetFromActorClicked)
 				.IsEnabled(this, &SCreateAssetFromObject::IsCreateAssetFromActorEnabled)
-				[
-					SNew(STextBlock)
-					.Text(CreateButtonText)
-				]
+				.ButtonStyle(FEditorStyle::Get(), "FlatButton.Success")
+				.TextStyle(FEditorStyle::Get(), "FlatButton.DefaultTextStyle")
+				.Text(CreateButtonText)
 			]
 			+ SHorizontalBox::Slot()
 			.Padding(0, 2, 0, 0)
@@ -114,10 +113,9 @@ void SCreateAssetFromObject::Construct(const FArguments& InArgs, TSharedPtr<SWin
 				.VAlign(VAlign_Bottom)
 				.ContentPadding(FMargin(8, 2, 8, 2))
 				.OnClicked(this, &SCreateAssetFromObject::OnCancelCreateAssetFromActor)
-				[
-					SNew(STextBlock)
-					.Text(LOCTEXT("CancelButtonText", "Cancel"))
-				]
+				.ButtonStyle(FEditorStyle::Get(), "FlatButton.Default")
+				.TextStyle(FEditorStyle::Get(), "FlatButton.DefaultTextStyle")
+				.Text(LOCTEXT("CancelButtonText", "Cancel"))
 			]
 		]
 	];
