[INTEGRATE] Change 2438892 by [REDACTED_PERSON][REDACTED_HANDLE] on 2015/02/09 17:50:38

	Fix annoying issue of not being able to move the gizmo afte dragging it, releasing, and not moving the mouse.

diff --git a/Engine/Source/Editor/UnrealEd/Private/EditorViewportClient.cpp b/Engine/Source/Editor/UnrealEd/Private/EditorViewportClient.cpp
--- a/Engine/Source/Editor/UnrealEd/Private/EditorViewportClient.cpp
+++ b/Engine/Source/Editor/UnrealEd/Private/EditorViewportClient.cpp
@@ -2041,15 +2041,17 @@ void FEditorViewportClient::StopTracking()
 		Widget->SetCurrentAxis( EAxisList::None );
 		Invalidate( true, true );
 
-	/*	if ( !MouseDeltaTracker->UsingDragTool() )
-		{
-			Widget->SetCurrentAxis( AXIS_None );
-			Invalidate( true, true );
-		}*/
-
 		SetRequiredCursorOverride( false );
 
 		bWidgetAxisControlledByDrag = false;
+
+		// Update the hovered hit proxy here.  If the user didnt move the mouse
+		// they still need to be able to pick up the gizmo without moving the mouse again
+		Viewport->InvalidateHitProxy();
+		HHitProxy* HitProxy = Viewport->GetHitProxy(CachedMouseX,CachedMouseY);
+
+		CheckHoveredHitProxy(HitProxy);
+
 		bIsTracking = false;	
 	}
 
