[INTEGRATE] Change 2438924 by [REDACTED_PERSON][REDACTED_HANDLE] on 2015/02/09 18:06:06

	Enabled BP components by default

diff --git a/Engine/Source/Editor/UnrealEd/Private/Settings/SettingsClasses.cpp b/Engine/Source/Editor/UnrealEd/Private/Settings/SettingsClasses.cpp
--- a/Engine/Source/Editor/UnrealEd/Private/Settings/SettingsClasses.cpp
+++ b/Engine/Source/Editor/UnrealEd/Private/Settings/SettingsClasses.cpp
@@ -55,6 +55,7 @@ UEditorExperimentalSettings::UEditorExperimentalSettings( const FObjectInitializ
 	: Super(ObjectInitializer)
 	, bInWorldBPEditing(true)
 	, bUnifiedBlueprintEditor(true)
+	, bBlueprintableComponents(true)
 
 {
 }
