[INTEGRATE] Change 2439005 by [REDACTED_PERSON][REDACTED_HANDLE] on 2015/02/09 18:45:39

	[UE-9358] REINST classes and skeleton classes were getting into the component type registry..

diff --git a/Engine/Source/Editor/UnrealEd/Private/ComponentTypeRegistry.cpp b/Engine/Source/Editor/UnrealEd/Private/ComponentTypeRegistry.cpp
--- a/Engine/Source/Editor/UnrealEd/Private/ComponentTypeRegistry.cpp
+++ b/Engine/Source/Editor/UnrealEd/Private/ComponentTypeRegistry.cpp
@@ -217,12 +217,15 @@ void FComponentTypeRegistryData::RefreshComponentList()
 				}
 			}
 
-			FComponentTypeEntry Entry = { Class->GetName(), FString(), Class };
-			ComponentTypeList.Add(Entry);
+			bool const bOutOfDateClass = Class->HasAnyClassFlags(CLASS_NewerVersionExists);
+			if (!bOutOfDateClass && !FKismetEditorUtilities::IsClassABlueprintSkeleton(Class))
+			{
+				FComponentTypeEntry Entry = { Class->GetName(), FString(), Class };
+				ComponentTypeList.Add(Entry);
+			}
 		}
 	}
 
-
 	{
 		// make sure that we add any user created classes immediately, generally this will not create anything (because assets have not been discovered yet), but asset discovery
 		// should be allowed to take place at any time:
