Removed non-power of 2 texture import warning
UE-8912 - Auto Asset Import: Non-power of two textures warning does not give sufficient context

diff --git a/Engine/Source/Editor/UnrealEd/Classes/Factories/Factory.h b/Engine/Source/Editor/UnrealEd/Classes/Factories/Factory.h
--- a/Engine/Source/Editor/UnrealEd/Classes/Factories/Factory.h
+++ b/Engine/Source/Editor/UnrealEd/Classes/Factories/Factory.h
@@ -66,6 +66,7 @@ public:
 
 	/** If this value is true, warning messages will be shown once for all objects being imported at the same time.  
 		This value will be reset to false each time a new import operation is started. */
+	DEPRECATED(4.8, "bAllowOneTimeWarningMessages is due to be removed in future.")
 	static bool bAllowOneTimeWarningMessages;
 
 	// Begin UObject interface.
diff --git a/Engine/Source/Editor/UnrealEd/Classes/Factories/TextureFactory.h b/Engine/Source/Editor/UnrealEd/Classes/Factories/TextureFactory.h
--- a/Engine/Source/Editor/UnrealEd/Classes/Factories/TextureFactory.h
+++ b/Engine/Source/Editor/UnrealEd/Classes/Factories/TextureFactory.h
@@ -112,9 +112,6 @@ public:
 	 */
 	static void SuppressImportOverwriteDialog();
 
-	/* Suppresses any warning dialogs about import resolution from popping up during operations that aren't concerned with the warnings. */
-	static void SuppressImportResolutionWarningDialog();
-
 	/**
 	 *	Initializes the given texture from the TextureData text block supplied.
 	 *	The TextureData text block is assumed to have been generated by the UTextureExporterT3D.
@@ -140,9 +137,6 @@ private:
 	/** This variable is static because in StaticImportObject() the type of the factory is not known. */
 	static bool bSuppressImportOverwriteDialog;
 
-	/** Whether to suppress one time warnings or not */
-	static bool bSuppressImportResolutionWarnings;
-
 	/**
 	*	Tests if the given height and width specify a supported texture resolution to import; Can optionally check if the height/width are powers of two
 	*
diff --git a/Engine/Source/Editor/UnrealEd/Private/Factories/EditorFactories.cpp b/Engine/Source/Editor/UnrealEd/Private/Factories/EditorFactories.cpp
--- a/Engine/Source/Editor/UnrealEd/Private/Factories/EditorFactories.cpp
+++ b/Engine/Source/Editor/UnrealEd/Private/Factories/EditorFactories.cpp
@@ -3285,7 +3285,6 @@ UTexture2D* DecompressTGA(
 }
 
 bool UTextureFactory::bSuppressImportOverwriteDialog = false;
-bool UTextureFactory::bSuppressImportResolutionWarnings = false;
 
 UTextureFactory::UTextureFactory(const FObjectInitializer& ObjectInitializer)
 	: Super(ObjectInitializer)
@@ -3335,11 +3334,6 @@ void UTextureFactory::SuppressImportOverwriteDialog()
 	bSuppressImportOverwriteDialog = true;
 }
 
-void UTextureFactory::SuppressImportResolutionWarningDialog()
-{
-	bSuppressImportResolutionWarnings = true;
-}
-
 /**
  * This fills any pixels of a texture with have an alpha value of zero,
  * with an RGB from the nearest neighboring pixel which has non-zero alpha.
@@ -4648,30 +4642,6 @@ bool UTextureFactory::IsImportResolutionValid(int32 Width, int32 Height, bool bA
 		Warn->Logf(ELogVerbosity::Error, *NSLOCTEXT("UnrealEd", "Warning_TextureNotAPowerOfTwo", "Cannot import texture with non-power of two dimensions").ToString() );
 		bValid = false;
 	}
-
-	// If we are allowed to warn about NPT textures and the texture is not a power of two, display a warning.
-	if( bAllowOneTimeWarningMessages && !bSuppressImportResolutionWarnings && bAllowNonPowerOfTwo && !bIsPowerOfTwo && bValid )
-	{
-		bAllowOneTimeWarningMessages = false;
-
-		FText MessageText;
-		if (GetCurrentFilename().IsEmpty())
-		{
-			MessageText = NSLOCTEXT("UnrealEd", "Warning_NPTTexture", "The texture you are importing is not a power of two.  Non power of two textures are never streamed and have no mipmaps. Proceed?");
-		}
-		else
-		{
-			MessageText = FText::Format(NSLOCTEXT("UnrealEd", "Warning_NPTTexture_Filename", "{0} is not a power of two.  Non power of two textures are never streamed and have no mipmaps. Proceed?"), FText::FromString(FPaths::GetCleanFilename(GetCurrentFilename())));
-		}
-
-		if( EAppReturnType::Yes != FMessageDialog::Open( EAppMsgType::YesNo, MessageText ) )
-		{
-			bValid = false;
-		}
-	}
-
-	// Reset the suppression so that future imports can still warn
-	bSuppressImportResolutionWarnings = false;
 	
 	return bValid;
 }
diff --git a/Engine/Source/Editor/UnrealEd/Private/Factories/Factory.cpp b/Engine/Source/Editor/UnrealEd/Private/Factories/Factory.cpp
--- a/Engine/Source/Editor/UnrealEd/Private/Factories/Factory.cpp
+++ b/Engine/Source/Editor/UnrealEd/Private/Factories/Factory.cpp
@@ -158,9 +158,6 @@ void UFactory::ResetState()
 	// Resets the state of the 'Yes To All / No To All' prompt for overwriting existing objects on import.
 	// After the reset, the next import collision will always display the prompt.
 	OverwriteYesOrNoToAllState = -1;
-
-	// Resets the state of one time warning messages. This will allow the warning message to be shown.
-	bAllowOneTimeWarningMessages = true;
 }
 
 void UFactory::DisplayOverwriteOptionsDialog(const FText& Message)
diff --git a/Engine/Source/Editor/UnrealEd/Private/Tests/AutomationEditorCommon.cpp b/Engine/Source/Editor/UnrealEd/Private/Tests/AutomationEditorCommon.cpp
--- a/Engine/Source/Editor/UnrealEd/Private/Tests/AutomationEditorCommon.cpp
+++ b/Engine/Source/Editor/UnrealEd/Private/Tests/AutomationEditorCommon.cpp
@@ -122,7 +122,6 @@ namespace AutomationEditorCommonUtils
 			//If we are a texture factory suppress some warning dialog that we don't want
 			if (ImportFactory->IsA(UTextureFactory::StaticClass()))
 			{
-				UTextureFactory::SuppressImportResolutionWarningDialog();
 				UTextureFactory::SuppressImportOverwriteDialog();
 			}
 
