Sequencer - Add a view to sequencer which uses a curve editor to show keyframes.

diff --git a/Engine/Source/Developer/MovieSceneTools/Private/PropertyTrackEditor.cpp b/Engine/Source/Developer/MovieSceneTools/Private/PropertyTrackEditor.cpp
--- a/Engine/Source/Developer/MovieSceneTools/Private/PropertyTrackEditor.cpp
+++ b/Engine/Source/Developer/MovieSceneTools/Private/PropertyTrackEditor.cpp
@@ -36,10 +36,10 @@ public:
 	{
 		UMovieSceneColorSection* ColorSection = Cast<UMovieSceneColorSection>(&SectionObject);
 
-		LayoutBuilder.AddKeyArea("R", NSLOCTEXT("FColorPropertySection", "RedArea", "Red"), MakeShareable(new FFloatCurveKeyArea(ColorSection->GetRedCurve())));
-		LayoutBuilder.AddKeyArea("G", NSLOCTEXT("FColorPropertySection", "GreenArea", "Green"), MakeShareable(new FFloatCurveKeyArea(ColorSection->GetGreenCurve())));
-		LayoutBuilder.AddKeyArea("B", NSLOCTEXT("FColorPropertySection", "BlueArea", "Blue"), MakeShareable(new FFloatCurveKeyArea(ColorSection->GetBlueCurve())));
-		LayoutBuilder.AddKeyArea("A", NSLOCTEXT("FColorPropertySection", "OpacityArea", "Opacity"), MakeShareable(new FFloatCurveKeyArea(ColorSection->GetAlphaCurve())));
+		LayoutBuilder.AddKeyArea("R", NSLOCTEXT("FColorPropertySection", "RedArea", "Red"), MakeShareable(new FFloatCurveKeyArea(&ColorSection->GetRedCurve(), ColorSection)));
+		LayoutBuilder.AddKeyArea("G", NSLOCTEXT("FColorPropertySection", "GreenArea", "Green"), MakeShareable(new FFloatCurveKeyArea(&ColorSection->GetGreenCurve(), ColorSection)));
+		LayoutBuilder.AddKeyArea("B", NSLOCTEXT("FColorPropertySection", "BlueArea", "Blue"), MakeShareable(new FFloatCurveKeyArea(&ColorSection->GetBlueCurve(), ColorSection)));
+		LayoutBuilder.AddKeyArea("A", NSLOCTEXT("FColorPropertySection", "OpacityArea", "Opacity"), MakeShareable(new FFloatCurveKeyArea(&ColorSection->GetAlphaCurve(), ColorSection)));
 	}
 
 	virtual int32 OnPaintSection( const FGeometry& AllottedGeometry, const FSlateRect& SectionClippingRect, FSlateWindowElementList& OutDrawElements, int32 LayerId, bool bParentEnabled ) const override 
@@ -166,7 +166,7 @@ public:
 	virtual void GenerateSectionLayout( class ISectionLayoutBuilder& LayoutBuilder ) const override
 	{
 		UMovieSceneFloatSection* FloatSection = Cast<UMovieSceneFloatSection>( &SectionObject );
-		LayoutBuilder.SetSectionAsKeyArea( MakeShareable( new FFloatCurveKeyArea( FloatSection->GetFloatCurve() ) ) );
+		LayoutBuilder.SetSectionAsKeyArea( MakeShareable( new FFloatCurveKeyArea( &FloatSection->GetFloatCurve(), FloatSection ) ) );
 	}
 };
 
@@ -185,15 +185,15 @@ public:
 		int32 ChannelsUsed = VectorSection->GetChannelsUsed();
 		check(ChannelsUsed >= 2 && ChannelsUsed <= 4);
 
-		LayoutBuilder.AddKeyArea("Vector.X", NSLOCTEXT("FVectorPropertySection", "XArea", "X"), MakeShareable(new FFloatCurveKeyArea(VectorSection->GetCurve(0))));
-		LayoutBuilder.AddKeyArea("Vector.Y", NSLOCTEXT("FVectorPropertySection", "YArea", "Y"), MakeShareable(new FFloatCurveKeyArea(VectorSection->GetCurve(1))));
+		LayoutBuilder.AddKeyArea("Vector.X", NSLOCTEXT("FVectorPropertySection", "XArea", "X"), MakeShareable(new FFloatCurveKeyArea(&VectorSection->GetCurve(0), VectorSection)));
+		LayoutBuilder.AddKeyArea("Vector.Y", NSLOCTEXT("FVectorPropertySection", "YArea", "Y"), MakeShareable(new FFloatCurveKeyArea(&VectorSection->GetCurve(1), VectorSection)));
 		if (ChannelsUsed >= 3)
 		{
-			LayoutBuilder.AddKeyArea("Vector.Z", NSLOCTEXT("FVectorPropertySection", "ZArea", "Z"), MakeShareable(new FFloatCurveKeyArea(VectorSection->GetCurve(2))));
+			LayoutBuilder.AddKeyArea("Vector.Z", NSLOCTEXT("FVectorPropertySection", "ZArea", "Z"), MakeShareable(new FFloatCurveKeyArea(&VectorSection->GetCurve(2), VectorSection)));
 		}
 		if (ChannelsUsed >= 4)
 		{
-			LayoutBuilder.AddKeyArea("Vector.W", NSLOCTEXT("FVectorPropertySection", "WArea", "W"), MakeShareable(new FFloatCurveKeyArea(VectorSection->GetCurve(3))));
+			LayoutBuilder.AddKeyArea("Vector.W", NSLOCTEXT("FVectorPropertySection", "WArea", "W"), MakeShareable(new FFloatCurveKeyArea(&VectorSection->GetCurve(3), VectorSection)));
 		}
 	}
 };
diff --git a/Engine/Source/Developer/MovieSceneTools/Private/TransformTrackEditor.cpp b/Engine/Source/Developer/MovieSceneTools/Private/TransformTrackEditor.cpp
--- a/Engine/Source/Developer/MovieSceneTools/Private/TransformTrackEditor.cpp
+++ b/Engine/Source/Developer/MovieSceneTools/Private/TransformTrackEditor.cpp
@@ -44,21 +44,21 @@ public:
 
 		// This generates the tree structure for the transform section
 		LayoutBuilder.PushCategory( "Location", NSLOCTEXT("FTransformSection", "LocationArea", "Location") );
-			LayoutBuilder.AddKeyArea("Location.X", NSLOCTEXT("FTransformSection", "LocXArea", "X"), MakeShareable( new FFloatCurveKeyArea( TransformSection->GetTranslationCurve( EAxis::X ) ) ) );
-			LayoutBuilder.AddKeyArea("Location.Y", NSLOCTEXT("FTransformSection", "LocYArea", "Y"), MakeShareable( new FFloatCurveKeyArea( TransformSection->GetTranslationCurve( EAxis::Y ) ) ) );
-			LayoutBuilder.AddKeyArea("Location.Z", NSLOCTEXT("FTransformSection", "LocZArea", "Z"), MakeShareable( new FFloatCurveKeyArea( TransformSection->GetTranslationCurve( EAxis::Z ) ) ) );
+			LayoutBuilder.AddKeyArea("Location.X", NSLOCTEXT("FTransformSection", "LocXArea", "X"), MakeShareable( new FFloatCurveKeyArea ( &TransformSection->GetTranslationCurve( EAxis::X ), TransformSection ) ) );
+			LayoutBuilder.AddKeyArea("Location.Y", NSLOCTEXT("FTransformSection", "LocYArea", "Y"), MakeShareable( new FFloatCurveKeyArea ( &TransformSection->GetTranslationCurve( EAxis::Y ), TransformSection ) ) );
+			LayoutBuilder.AddKeyArea("Location.Z", NSLOCTEXT("FTransformSection", "LocZArea", "Z"), MakeShareable( new FFloatCurveKeyArea ( &TransformSection->GetTranslationCurve( EAxis::Z ), TransformSection ) ) );
 		LayoutBuilder.PopCategory();
 
 		LayoutBuilder.PushCategory( "Rotation", NSLOCTEXT("FTransformSection", "RotationArea", "Rotation") );
-			LayoutBuilder.AddKeyArea("Rotation.X", NSLOCTEXT("FTransformSection", "RotXArea", "X"), MakeShareable( new FFloatCurveKeyArea( TransformSection->GetRotationCurve( EAxis::X ) ) ) );
-			LayoutBuilder.AddKeyArea("Rotation.Y", NSLOCTEXT("FTransformSection", "RotYArea", "Y"), MakeShareable( new FFloatCurveKeyArea( TransformSection->GetRotationCurve( EAxis::Y ) ) ) );
-			LayoutBuilder.AddKeyArea("Rotation.Z", NSLOCTEXT("FTransformSection", "RotZArea", "Z"), MakeShareable( new FFloatCurveKeyArea( TransformSection->GetRotationCurve( EAxis::Z ) ) ) );
+			LayoutBuilder.AddKeyArea("Rotation.X", NSLOCTEXT("FTransformSection", "RotXArea", "X"), MakeShareable( new FFloatCurveKeyArea ( &TransformSection->GetRotationCurve( EAxis::X ), TransformSection ) ) );
+			LayoutBuilder.AddKeyArea("Rotation.Y", NSLOCTEXT("FTransformSection", "RotYArea", "Y"), MakeShareable( new FFloatCurveKeyArea ( &TransformSection->GetRotationCurve( EAxis::Y ), TransformSection ) ) );
+			LayoutBuilder.AddKeyArea("Rotation.Z", NSLOCTEXT("FTransformSection", "RotZArea", "Z"), MakeShareable( new FFloatCurveKeyArea ( &TransformSection->GetRotationCurve( EAxis::Z ), TransformSection ) ) );
 		LayoutBuilder.PopCategory();
 
 		LayoutBuilder.PushCategory( "Scale", NSLOCTEXT("FTransformSection", "ScaleArea", "Scale") );
-			LayoutBuilder.AddKeyArea("Scale.X", NSLOCTEXT("FTransformSection", "ScaleXArea", "X"), MakeShareable( new FFloatCurveKeyArea( TransformSection->GetScaleCurve( EAxis::X ) ) ) );
-			LayoutBuilder.AddKeyArea("Scale.Y", NSLOCTEXT("FTransformSection", "ScaleYArea", "Y"), MakeShareable( new FFloatCurveKeyArea( TransformSection->GetScaleCurve( EAxis::Y ) ) ) );
-			LayoutBuilder.AddKeyArea("Scale.Z", NSLOCTEXT("FTransformSection", "ScaleZArea", "Z"), MakeShareable( new FFloatCurveKeyArea( TransformSection->GetScaleCurve( EAxis::Z ) ) ) );
+			LayoutBuilder.AddKeyArea("Scale.X", NSLOCTEXT("FTransformSection", "ScaleXArea", "X"), MakeShareable( new FFloatCurveKeyArea ( &TransformSection->GetScaleCurve( EAxis::X ), TransformSection ) ) );
+			LayoutBuilder.AddKeyArea("Scale.Y", NSLOCTEXT("FTransformSection", "ScaleYArea", "Y"), MakeShareable( new FFloatCurveKeyArea ( &TransformSection->GetScaleCurve( EAxis::Y ), TransformSection ) ) );
+			LayoutBuilder.AddKeyArea("Scale.Z", NSLOCTEXT("FTransformSection", "ScaleZArea", "Z"), MakeShareable( new FFloatCurveKeyArea ( &TransformSection->GetScaleCurve( EAxis::Z ), TransformSection ) ) );
 		LayoutBuilder.PopCategory();
 	}
 
diff --git a/Engine/Source/Developer/MovieSceneTools/Public/MovieSceneToolHelpers.h b/Engine/Source/Developer/MovieSceneTools/Public/MovieSceneToolHelpers.h
--- a/Engine/Source/Developer/MovieSceneTools/Public/MovieSceneToolHelpers.h
+++ b/Engine/Source/Developer/MovieSceneTools/Public/MovieSceneToolHelpers.h
@@ -11,15 +11,15 @@
 class FFloatCurveKeyArea : public IKeyArea
 {
 public:
-	FFloatCurveKeyArea( FRichCurve& InCurve )
-		: Curve( InCurve )
+	FFloatCurveKeyArea( FRichCurve* InCurve, UMovieSceneSection* InOwningSection )
+		: CurveInfo( InCurve, InOwningSection )
 	{}
 
 	/** IKeyArea interface */
 	virtual TArray<FKeyHandle> GetUnsortedKeyHandles() const override
 	{
 		TArray<FKeyHandle> OutKeyHandles;
-		for (auto It(Curve.GetKeyHandleIterator()); It; ++It)
+		for (auto It(CurveInfo.Curve->GetKeyHandleIterator()); It; ++It)
 		{
 			OutKeyHandles.Add(It.Key());
 		}
@@ -28,24 +28,25 @@ public:
 
 	virtual float GetKeyTime( FKeyHandle KeyHandle ) const override
 	{
-		return Curve.GetKeyTime( KeyHandle );
+		return CurveInfo.Curve->GetKeyTime( KeyHandle );
 	}
 
 	virtual FKeyHandle MoveKey( FKeyHandle KeyHandle, float DeltaPosition ) override
 	{
-		return Curve.SetKeyTime( KeyHandle, Curve.GetKeyTime( KeyHandle ) + DeltaPosition );
+		return CurveInfo.Curve->SetKeyTime( KeyHandle, CurveInfo.Curve->GetKeyTime( KeyHandle ) + DeltaPosition );
 	}
 
 	virtual void DeleteKey(FKeyHandle KeyHandle) override
 	{
-		if( Curve.IsKeyHandleValid(KeyHandle) )
+		if ( CurveInfo.Curve->IsKeyHandleValid( KeyHandle ) )
 		{
-			Curve.DeleteKey(KeyHandle);
+			CurveInfo.Curve->DeleteKey( KeyHandle );
 		}
 	}
+
+	virtual FCurveInfo* GetCurveInfo() override { return &CurveInfo; }
 private:
-	/** Curve with keys in this area */
-	FRichCurve& Curve;
+	IKeyArea::FCurveInfo CurveInfo;
 };
 
 
@@ -83,6 +84,8 @@ public:
 	{
 		Curve.DeleteKey(KeyHandle);
 	}
+
+	virtual IKeyArea::FCurveInfo* GetCurveInfo() override { return nullptr; };
 private:
 	/** Curve with keys in this area */
 	FIntegralCurve& Curve;
diff --git a/Engine/Source/Editor/CurveAssetEditor/Private/CurveAssetEditor.cpp b/Engine/Source/Editor/CurveAssetEditor/Private/CurveAssetEditor.cpp
--- a/Engine/Source/Editor/CurveAssetEditor/Private/CurveAssetEditor.cpp
+++ b/Engine/Source/Editor/CurveAssetEditor/Private/CurveAssetEditor.cpp
@@ -11,6 +11,7 @@
 #include "Curves/CurveVector.h"
 #include "Curves/CurveLinearColor.h"
 #include "Curves/CurveFloat.h"
+#include "SNumericDropDown.h"
 
 #define LOCTEXT_NAMESPACE "CurveAssetEditor"
 
@@ -168,19 +169,6 @@ void FCurveAssetEditor::SetInputSnap(float value)
 	InputSnap = value;
 }
 
-FText FCurveAssetEditor::GetInputSnapText() const
-{
-	return FText::AsNumber(InputSnap);
-}
-
-void FCurveAssetEditor::InputSnapTextComitted(const FText& InNewText, ETextCommit::Type InTextCommit)
-{
-	if (InNewText.IsNumeric())
-	{
-		TTypeFromString<float>::FromString(InputSnap, *InNewText.ToString());
-	}
-}
-
 float FCurveAssetEditor::GetOutputSnap() const
 {
 	return OutputSnap;
@@ -191,69 +179,6 @@ void FCurveAssetEditor::SetOutputSnap(float value)
 	OutputSnap = value;
 }
 
-FText FCurveAssetEditor::GetOutputSnapText() const
-{
-	return FText::AsNumber(OutputSnap);
-}
-
-void FCurveAssetEditor::OutputSnapTextComitted(const FText& InNewText, ETextCommit::Type InTextCommit)
-{
-	if (InNewText.IsNumeric())
-	{
-		TTypeFromString<float>::FromString(OutputSnap, *InNewText.ToString());
-	}
-}
-
-TSharedRef<SWidget> FCurveAssetEditor::BuildInputSnapMenu()
-{
-	FMenuBuilder MenuBuilder(true, NULL);
-
-	FUIAction OneThousandthAction(FExecuteAction::CreateSP(this, &FCurveAssetEditor::SetInputSnap, 0.001f));
-	MenuBuilder.AddMenuEntry(LOCTEXT("InputSnap_OneThousandth", "0.001"), LOCTEXT("InputSnap_OneThousandth", "Set snap to 1/1000th"), FSlateIcon(), OneThousandthAction);
-
-	FUIAction OneHundredthAction(FExecuteAction::CreateSP(this, &FCurveAssetEditor::SetInputSnap, 0.01f));
-	MenuBuilder.AddMenuEntry(LOCTEXT("InputSnap_OneHundredth", "0.01"), LOCTEXT("InputSnap_OneHundredth", "Set snap to 1/100th"), FSlateIcon(), OneHundredthAction);
-
-	FUIAction OneTenthAction(FExecuteAction::CreateSP(this, &FCurveAssetEditor::SetInputSnap, 0.1f));
-	MenuBuilder.AddMenuEntry(LOCTEXT("InputSnap_OneTenth", "0.1"), LOCTEXT("InputSnap_OneTenth", "Set snap to 1/10th"), FSlateIcon(), OneTenthAction);
-
-	FUIAction OneAction(FExecuteAction::CreateSP(this, &FCurveAssetEditor::SetInputSnap, 1.0f));
-	MenuBuilder.AddMenuEntry(LOCTEXT("InputSnap_One", "1"), LOCTEXT("InputSnap_One", "Set snap to 1h"), FSlateIcon(), OneAction);
-
-	FUIAction TenAction(FExecuteAction::CreateSP(this, &FCurveAssetEditor::SetInputSnap, 10.0f));
-	MenuBuilder.AddMenuEntry(LOCTEXT("InputSnap_Ten", "10"), LOCTEXT("InputSnap_Ten", "Set snap to 10"), FSlateIcon(), TenAction);
-
-	FUIAction HundredAction(FExecuteAction::CreateSP(this, &FCurveAssetEditor::SetInputSnap, 100.0f));
-	MenuBuilder.AddMenuEntry(LOCTEXT("InputSnap_OneHundred", "100"), LOCTEXT("InputSnap_OneHundred", "Set snap to 100"), FSlateIcon(), HundredAction);
-
-	return MenuBuilder.MakeWidget();
-}
-
-TSharedRef<SWidget> FCurveAssetEditor::BuildOutputSnapMenu()
-{
-	FMenuBuilder MenuBuilder(true, NULL);
-
-	FUIAction OneThousandthAction(FExecuteAction::CreateSP(this, &FCurveAssetEditor::SetOutputSnap, 0.001f));
-	MenuBuilder.AddMenuEntry(LOCTEXT("OutputSnap_OneThousandth", "0.001"), LOCTEXT("OutputSnap_OneThousandth", "Set snap to 1/1000th"), FSlateIcon(), OneThousandthAction);
-
-	FUIAction OneHundredthAction(FExecuteAction::CreateSP(this, &FCurveAssetEditor::SetOutputSnap, 0.01f));
-	MenuBuilder.AddMenuEntry(LOCTEXT("OutputSnap_OneHundredth", "0.01"), LOCTEXT("OutputSnap_OneHundredth", "Set snap to 1/100th"), FSlateIcon(), OneHundredthAction);
-
-	FUIAction OneTenthAction(FExecuteAction::CreateSP(this, &FCurveAssetEditor::SetOutputSnap, 0.1f));
-	MenuBuilder.AddMenuEntry(LOCTEXT("OutputSnap_OneTenth", "0.1"), LOCTEXT("OutputSnap_OneTenth", "Set snap to 1/10th"), FSlateIcon(), OneTenthAction);
-
-	FUIAction OneAction(FExecuteAction::CreateSP(this, &FCurveAssetEditor::SetOutputSnap, 1.0f));
-	MenuBuilder.AddMenuEntry(LOCTEXT("OutputSnap_One", "1"), LOCTEXT("OutputSnap_One", "Set snap to 1h"), FSlateIcon(), OneAction);
-
-	FUIAction TenAction(FExecuteAction::CreateSP(this, &FCurveAssetEditor::SetOutputSnap, 10.0f));
-	MenuBuilder.AddMenuEntry(LOCTEXT("OutputSnap_Ten", "10"), LOCTEXT("OutputSnap_Ten", "Set snap to 10"), FSlateIcon(), TenAction);
-
-	FUIAction HundredAction(FExecuteAction::CreateSP(this, &FCurveAssetEditor::SetOutputSnap, 100.0f));
-	MenuBuilder.AddMenuEntry(LOCTEXT("OutputSnap_OneHundred", "100"), LOCTEXT("OutputSnap_OneHundred", "Set snap to 100"), FSlateIcon(), HundredAction);
-
-	return MenuBuilder.MakeWidget();
-}
-
 float FCurveAssetEditor::GetTimelineLength() const
 {
 	return 0.f;
@@ -302,89 +227,29 @@ TSharedPtr<FExtender> FCurveAssetEditor::GetToolbarExtender()
 
 	TSharedPtr<FExtender> ToolbarExtender = MakeShareable(new FExtender);
 
+	TArray<SNumericDropDown<float>::FNamedValue> SnapValues;
+	SnapValues.Add( SNumericDropDown<float>::FNamedValue( 0.001f, LOCTEXT( "Snap_OneThousandth", "0.001" ), LOCTEXT( "SnapDescription_OneThousandth", "Set snap to 1/1000th" ) ) );
+	SnapValues.Add( SNumericDropDown<float>::FNamedValue( 0.01f, LOCTEXT( "Snap_OneHundredth", "0.01" ), LOCTEXT( "SnapDescription_OneHundredth", "Set snap to 1/100th" ) ) );
+	SnapValues.Add( SNumericDropDown<float>::FNamedValue( 0.1f, LOCTEXT( "Snap_OneTenth", "0.1" ), LOCTEXT( "SnapDescription_OneTenth", "Set snap to 1/10th" ) ) );
+	SnapValues.Add( SNumericDropDown<float>::FNamedValue( 1.0f, LOCTEXT( "Snap_One", "1" ), LOCTEXT( "SnapDescription_One", "Set snap to 1" ) ) );
+	SnapValues.Add( SNumericDropDown<float>::FNamedValue( 10.0f, LOCTEXT( "Snap_Ten", "10" ), LOCTEXT( "SnapDescription_Ten", "Set snap to 10" ) ) );
+	SnapValues.Add( SNumericDropDown<float>::FNamedValue( 100.0f, LOCTEXT( "Snap_OneHundred", "100" ), LOCTEXT( "SnapDescription_OneHundred", "Set snap to 100" ) ) );
+
 	TSharedRef<SWidget> InputSnapWidget =
-		SNew(SVerticalBox)
-		// Vertical Label
-		+ SVerticalBox::Slot()
-		.Padding(4)
-		.AutoHeight()
-		[
-			SNew(STextBlock)
-			.Text(LOCTEXT("InputSnapLabelVertical", "Time Snap"))
-			.Visibility(this, &FCurveAssetEditor::GetSnapLabelVisibilty, EOrientation::Orient_Vertical)
-		]
-		+ SVerticalBox::Slot()
-		.Padding(4)
-		.AutoHeight()
-		[
-			SNew(SHorizontalBox)
-			// Horizontal Icon Label
-			+ SHorizontalBox::Slot()
-			.VAlign(EVerticalAlignment::VAlign_Center)
-			.AutoWidth()
-			.Padding(FMargin(0, 0, 3, 0))
-			[
-				SNew( STextBlock )
-				.Text( LOCTEXT("InputSnapLabelHorizontal", "Time"))
-				.Visibility(this, &FCurveAssetEditor::GetSnapLabelVisibilty, EOrientation::Orient_Horizontal)
-			]
-			+ SHorizontalBox::Slot()
-			[
-				SNew(SComboButton)
-				.ContentPadding(1)
-				.OnGetMenuContent(this, &FCurveAssetEditor::BuildInputSnapMenu)
-				.ToolTipText(LOCTEXT("TimeSnapToolTip", "Time Snap"))
-				.ButtonContent()
-				[
-					SNew(SEditableTextBox)
-					.Text(this, &FCurveAssetEditor::GetInputSnapText)
-					.MinDesiredWidth(40)
-					.OnTextCommitted(this, &FCurveAssetEditor::InputSnapTextComitted)
-				]
-			]
-		];
+		SNew( SNumericDropDown<float> )
+		.DropDownValues( SnapValues )
+		.LabelText( LOCTEXT("InputSnapLabel", "Input Snap"))
+		.Value( this, &FCurveAssetEditor::GetInputSnap )
+		.OnValueChanged( this, &FCurveAssetEditor::SetInputSnap )
+		.Orientation( this, &FCurveAssetEditor::GetSnapLabelOrientation );
 
 	TSharedRef<SWidget> OutputSnapWidget =
-		SNew(SVerticalBox)
-		// Vertical Label
-		+ SVerticalBox::Slot()
-		.Padding(4)
-		.AutoHeight()
-		[
-			SNew(STextBlock)
-			.Text(LOCTEXT("OutputSnapLabelVertical", "Value Snap"))
-			.Visibility(this, &FCurveAssetEditor::GetSnapLabelVisibilty, EOrientation::Orient_Vertical)
-		]
-		+ SVerticalBox::Slot()
-		.Padding(4)
-		.AutoHeight()
-		[
-			SNew(SHorizontalBox)
-			// Horizontal Icon Label
-			+ SHorizontalBox::Slot()
-			.VAlign(EVerticalAlignment::VAlign_Center)
-			.AutoWidth()
-			.Padding(FMargin(0, 0, 3, 0))
-			[
-				SNew(STextBlock)
-				.Text(LOCTEXT("OutputSnapLabelHorizontal", "Value"))
-				.Visibility(this, &FCurveAssetEditor::GetSnapLabelVisibilty, EOrientation::Orient_Horizontal)
-			]
-			+ SHorizontalBox::Slot()
-			[
-				SNew(SComboButton)
-				.ContentPadding(1)
-				.OnGetMenuContent(this, &FCurveAssetEditor::BuildOutputSnapMenu)
-				.ToolTipText( LOCTEXT( "ValueSnapToolTip", "Value Snap" ) )
-				.ButtonContent()
-				[
-					SNew(SEditableTextBox)
-					.Text(this, &FCurveAssetEditor::GetOutputSnapText)
-					.MinDesiredWidth(40)
-					.OnTextCommitted(this, &FCurveAssetEditor::OutputSnapTextComitted)
-				]
-			]
-		];
+		SNew( SNumericDropDown<float> )
+		.DropDownValues( SnapValues )
+		.LabelText( LOCTEXT( "OutputSnapLabel", "Output Snap" ) )
+		.Value( this, &FCurveAssetEditor::GetOutputSnap )
+		.OnValueChanged( this, &FCurveAssetEditor::SetOutputSnap )
+		.Orientation( this, &FCurveAssetEditor::GetSnapLabelOrientation );
 
 	ToolbarExtender->AddToolBarExtension(
 		"Asset",
@@ -396,14 +261,11 @@ TSharedPtr<FExtender> FCurveAssetEditor::GetToolbarExtender()
 	return ToolbarExtender;
 }
 
-EVisibility FCurveAssetEditor::GetSnapLabelVisibilty(EOrientation LabelOrientation) const
+EOrientation FCurveAssetEditor::GetSnapLabelOrientation() const
 {
-	if ((FMultiBoxSettings::UseSmallToolBarIcons.Get() && LabelOrientation == EOrientation::Orient_Horizontal) ||
-		(FMultiBoxSettings::UseSmallToolBarIcons.Get() == false && LabelOrientation == EOrientation::Orient_Vertical))
-	{
-		return EVisibility::Visible;
-	}
-	return EVisibility::Collapsed;
+	return FMultiBoxSettings::UseSmallToolBarIcons.Get()
+		? EOrientation::Orient_Horizontal
+		: EOrientation::Orient_Vertical;
 }
 
 #undef LOCTEXT_NAMESPACE
\ No newline at end of file
diff --git a/Engine/Source/Editor/CurveAssetEditor/Private/CurveAssetEditor.h b/Engine/Source/Editor/CurveAssetEditor/Private/CurveAssetEditor.h
--- a/Engine/Source/Editor/CurveAssetEditor/Private/CurveAssetEditor.h
+++ b/Engine/Source/Editor/CurveAssetEditor/Private/CurveAssetEditor.h
@@ -50,27 +50,15 @@ private:
 	/** Get the snapping value for the input domain. */
 	float GetInputSnap() const;
 	/** Gets the snapping value for the input domain as text. */
-	FText GetInputSnapText() const;
-	/** Get set the snapping value for the input domain. */
 	void SetInputSnap(float value);
-	/** Callback for setting the snapping value for the input domain from a text value. */
-	void InputSnapTextComitted(const FText& InNewText, ETextCommit::Type InTextCommit);
-	/** Builds a menu widget with options for snapping in the input domain. */
-	TSharedRef<SWidget> BuildInputSnapMenu();
 
 	/** Get the snapping value for the output domain. */
 	float GetOutputSnap() const;
-	/** Gets the snapping value for the output domain as text. */
-	FText GetOutputSnapText() const;
 	/** Get set the snapping value for the output domain. */
 	void SetOutputSnap(float value);
-	/** Callback for setting the snapping value for the output domain from a text value. */
-	void OutputSnapTextComitted(const FText& InNewText, ETextCommit::Type InTextCommit);
-	/** Builds a menu widget with options for snapping in the output domain. */
-	TSharedRef<SWidget> BuildOutputSnapMenu();
 
-	/** Get the current snap label visibility for a given orientation. */
-	EVisibility GetSnapLabelVisibilty(EOrientation LabelOrientation) const;
+	/** Get the orientation for the snap value controls. */
+	EOrientation GetSnapLabelOrientation() const;
 
 	TSharedPtr<class SCurveEditor> TrackWidget;
 
diff --git a/Engine/Source/Editor/DetailCustomizations/Private/CurveColorCustomization.cpp b/Engine/Source/Editor/DetailCustomizations/Private/CurveColorCustomization.cpp
--- a/Engine/Source/Editor/DetailCustomizations/Private/CurveColorCustomization.cpp
+++ b/Engine/Source/Editor/DetailCustomizations/Private/CurveColorCustomization.cpp
@@ -192,11 +192,6 @@ TArray<FRichCurveEditInfo> FCurveColorCustomization::GetCurves()
 	return Curves;
 }
 
-UObject* FCurveColorCustomization::GetOwner()
-{
-	return Owner;
-}
-
 void FCurveColorCustomization::ModifyOwner()
 {
 	if (Owner)
@@ -213,7 +208,7 @@ void FCurveColorCustomization::MakeTransactional()
 	}
 }
 
-void FCurveColorCustomization::OnCurveChanged()
+void FCurveColorCustomization::OnCurveChanged(const TArray<FRichCurveEditInfo>& ChangedCurveEditInfos)
 {
 	StructPropertyHandle->NotifyPostChange();
 }
@@ -405,6 +400,7 @@ FReply FCurveColorCustomization::OnCurvePreviewDoubleClick(const FGeometry& InMy
 			TSharedRef<SMiniCurveEditor> MiniCurveEditor =
 				SNew(SMiniCurveEditor)
 				.CurveOwner(this)
+				.OwnerObject(Owner)
 				.ParentWindow(Window);
 
 			Window->SetContent( MiniCurveEditor );
diff --git a/Engine/Source/Editor/DetailCustomizations/Private/CurveColorCustomization.h b/Engine/Source/Editor/DetailCustomizations/Private/CurveColorCustomization.h
--- a/Engine/Source/Editor/DetailCustomizations/Private/CurveColorCustomization.h
+++ b/Engine/Source/Editor/DetailCustomizations/Private/CurveColorCustomization.h
@@ -30,10 +30,9 @@ public:
 	/** FCurveOwnerInterface interface */
 	virtual TArray<FRichCurveEditInfoConst> GetCurves() const override;
 	virtual TArray<FRichCurveEditInfo> GetCurves() override;
-	virtual UObject* GetOwner() override;
 	virtual void ModifyOwner() override;
 	virtual void MakeTransactional() override;
-	virtual void OnCurveChanged() override;
+	virtual void OnCurveChanged(const TArray<FRichCurveEditInfo>& ChangedCurveEditInfos) override;
 	virtual bool IsLinearColorCurve() const override { return true; }
 	virtual FLinearColor GetLinearColorValue(float InTime) const override;
 	virtual bool HasAnyAlphaKeys() const override;
diff --git a/Engine/Source/Editor/DetailCustomizations/Private/CurveStructCustomization.cpp b/Engine/Source/Editor/DetailCustomizations/Private/CurveStructCustomization.cpp
--- a/Engine/Source/Editor/DetailCustomizations/Private/CurveStructCustomization.cpp
+++ b/Engine/Source/Editor/DetailCustomizations/Private/CurveStructCustomization.cpp
@@ -203,11 +203,6 @@ TArray<FRichCurveEditInfo> FCurveStructCustomization::GetCurves()
 	return Curves;
 }
 
-UObject* FCurveStructCustomization::GetOwner()
-{
-	return Owner;
-}
-
 void FCurveStructCustomization::ModifyOwner()
 {
 	if (Owner)
@@ -224,7 +219,7 @@ void FCurveStructCustomization::MakeTransactional()
 	}
 }
 
-void FCurveStructCustomization::OnCurveChanged()
+void FCurveStructCustomization::OnCurveChanged(const TArray<FRichCurveEditInfo>& ChangedCurveEditInfos)
 {
 	StructPropertyHandle->NotifyPostChange();
 }
@@ -385,6 +380,7 @@ FReply FCurveStructCustomization::OnCurvePreviewDoubleClick( const FGeometry& In
 			TSharedRef<SMiniCurveEditor> MiniCurveEditor =
 				SNew(SMiniCurveEditor)
 				.CurveOwner(this)
+				.OwnerObject(Owner)
 				.ParentWindow(Window);
 
 			Window->SetContent( MiniCurveEditor );
diff --git a/Engine/Source/Editor/DetailCustomizations/Private/CurveStructCustomization.h b/Engine/Source/Editor/DetailCustomizations/Private/CurveStructCustomization.h
--- a/Engine/Source/Editor/DetailCustomizations/Private/CurveStructCustomization.h
+++ b/Engine/Source/Editor/DetailCustomizations/Private/CurveStructCustomization.h
@@ -30,10 +30,9 @@ public:
 	/** FCurveOwnerInterface interface */
 	virtual TArray<FRichCurveEditInfoConst> GetCurves() const override;
 	virtual TArray<FRichCurveEditInfo> GetCurves() override;
-	virtual UObject* GetOwner() override;
 	virtual void ModifyOwner() override;
 	virtual void MakeTransactional() override;
-	virtual void OnCurveChanged() override;
+	virtual void OnCurveChanged(const TArray<FRichCurveEditInfo>& ChangedCurveEditInfos) override;
 	virtual bool IsValidCurve( FRichCurveEditInfo CurveInfo );
 
 private:
diff --git a/Engine/Source/Editor/DetailCustomizations/Private/WheeledVehicleMovementComponent4WDetails.cpp b/Engine/Source/Editor/DetailCustomizations/Private/WheeledVehicleMovementComponent4WDetails.cpp
--- a/Engine/Source/Editor/DetailCustomizations/Private/WheeledVehicleMovementComponent4WDetails.cpp
+++ b/Engine/Source/Editor/DetailCustomizations/Private/WheeledVehicleMovementComponent4WDetails.cpp
@@ -109,11 +109,6 @@ TArray<FRichCurveEditInfo> FWheeledVehicleMovementComponent4WDetails::FSteeringC
 	return Curves;
 }
 
-UObject* FWheeledVehicleMovementComponent4WDetails::FSteeringCurveEditor::GetOwner()
-{
-	return Owner;
-}
-
 void FWheeledVehicleMovementComponent4WDetails::FSteeringCurveEditor::ModifyOwner()
 {
 	if (Owner)
@@ -158,11 +153,6 @@ TArray<FRichCurveEditInfo> FWheeledVehicleMovementComponent4WDetails::FTorqueCur
 	return Curves;
 }
 
-UObject* FWheeledVehicleMovementComponent4WDetails::FTorqueCurveEditor::GetOwner()
-{
-	return Owner;
-}
-
 void FWheeledVehicleMovementComponent4WDetails::FTorqueCurveEditor::ModifyOwner()
 {
 	if (Owner)
diff --git a/Engine/Source/Editor/DetailCustomizations/Private/WheeledVehicleMovementComponent4WDetails.h b/Engine/Source/Editor/DetailCustomizations/Private/WheeledVehicleMovementComponent4WDetails.h
--- a/Engine/Source/Editor/DetailCustomizations/Private/WheeledVehicleMovementComponent4WDetails.h
+++ b/Engine/Source/Editor/DetailCustomizations/Private/WheeledVehicleMovementComponent4WDetails.h
@@ -23,10 +23,9 @@ private:
 		/** FCurveOwnerInterface interface */
 		virtual TArray<FRichCurveEditInfoConst> GetCurves() const override;
 		virtual TArray<FRichCurveEditInfo> GetCurves() override;
-		virtual UObject* GetOwner() override;
 		virtual void ModifyOwner() override;
 		virtual void MakeTransactional() override;
-		virtual void OnCurveChanged() override { }
+		virtual void OnCurveChanged( const TArray<FRichCurveEditInfo>& ChangedCurveEditInfos ) override { }
 		virtual bool IsValidCurve(FRichCurveEditInfo CurveInfo) override;
 
 	private:
@@ -41,10 +40,9 @@ private:
 		/** FCurveOwnerInterface interface */
 		virtual TArray<FRichCurveEditInfoConst> GetCurves() const override;
 		virtual TArray<FRichCurveEditInfo> GetCurves() override;
-		virtual UObject* GetOwner() override;
 		virtual void ModifyOwner() override;
 		virtual void MakeTransactional() override;
-		virtual void OnCurveChanged() override { }
+		virtual void OnCurveChanged( const TArray<FRichCurveEditInfo>& ChangedCurveEditInfos ) override { }
 		virtual bool IsValidCurve( FRichCurveEditInfo CurveInfo ) override;
 
 	private:
diff --git a/Engine/Source/Editor/Persona/Private/SAnimCurvePanel.cpp b/Engine/Source/Editor/Persona/Private/SAnimCurvePanel.cpp
--- a/Engine/Source/Editor/Persona/Private/SAnimCurvePanel.cpp
+++ b/Engine/Source/Editor/Persona/Private/SAnimCurvePanel.cpp
@@ -102,7 +102,7 @@ public:
 		return FText::GetEmpty();
 	}
 
-	virtual void OnCurveChanged() override
+	virtual void OnCurveChanged(const TArray<FRichCurveEditInfo>& ChangedCurveEditInfos) override
 	{
 	}
 
diff --git a/Engine/Source/Editor/Persona/Private/SAnimTrackCurvePanel.cpp b/Engine/Source/Editor/Persona/Private/SAnimTrackCurvePanel.cpp
--- a/Engine/Source/Editor/Persona/Private/SAnimTrackCurvePanel.cpp
+++ b/Engine/Source/Editor/Persona/Private/SAnimTrackCurvePanel.cpp
@@ -163,7 +163,7 @@ public:
 		return TCHAR('S');
 	}
 
-	virtual void OnCurveChanged() override
+	virtual void OnCurveChanged(const TArray<FRichCurveEditInfo>& ChangedCurveEditInfos) override
 	{
 	}
 
diff --git a/Engine/Source/Editor/Sequencer/Private/SSection.cpp b/Engine/Source/Editor/Sequencer/Private/SSection.cpp
--- a/Engine/Source/Editor/Sequencer/Private/SSection.cpp
+++ b/Engine/Source/Editor/Sequencer/Private/SSection.cpp
@@ -339,7 +339,7 @@ void SSection::PaintKeys( const FGeometry& AllottedGeometry, const FSlateRect& M
 					// always draw selected keys on top of other keys
 					bSelected ? KeyLayer+1 : KeyLayer,
 					// Center the key along Y.  Ensure the middle of the key is at the actual key time
-					KeyAreaGeometry.ToPaintGeometry( FVector2D( KeyPosition - FMath::TruncToFloat(SequencerSectionConstants::KeySize.X/2.0f), ((KeyAreaGeometry.Size.Y*.5f)-(SequencerSectionConstants::KeySize.Y*.5f)) ), SequencerSectionConstants::KeySize ),
+					KeyAreaGeometry.ToPaintGeometry( FVector2D( KeyPosition - FMath::CeilToFloat(SequencerSectionConstants::KeySize.X/2.0f), ((KeyAreaGeometry.Size.Y*.5f)-(SequencerSectionConstants::KeySize.Y*.5f)) ), SequencerSectionConstants::KeySize ),
 					KeyBrush,
 					MyClippingRect,
 					ESlateDrawEffect::None,
diff --git a/Engine/Source/Editor/Sequencer/Private/SSequencer.cpp b/Engine/Source/Editor/Sequencer/Private/SSequencer.cpp
--- a/Engine/Source/Editor/Sequencer/Private/SSequencer.cpp
+++ b/Engine/Source/Editor/Sequencer/Private/SSequencer.cpp
@@ -23,7 +23,7 @@
 #include "MovieSceneShotSection.h"
 #include "CommonMovieSceneTools.h"
 #include "SSearchBox.h"
-#include "SNumericEntryBox.h"
+#include "SNumericDropDown.h"
 
 
 #define LOCTEXT_NAMESPACE "Sequencer"
@@ -148,9 +148,6 @@ void SSequencer::Construct( const FArguments& InArgs, TSharedRef< class FSequenc
 	// Create a node tree which contains a tree of movie scene data to display in the sequence
 	SequencerNodeTree = MakeShareable( new FSequencerNodeTree( InSequencer.Get() ) );
 
-	CleanViewEnabled = InArgs._CleanViewEnabled;
-	OnToggleCleanView = InArgs._OnToggleCleanView;
-
 	FSequencerWidgetsModule& SequencerWidgets = FModuleManager::Get().LoadModuleChecked<FSequencerWidgetsModule>( "SequencerWidgets" );
 
 	FTimeSliderArgs TimeSliderArgs;
@@ -168,6 +165,10 @@ void SSequencer::Construct( const FArguments& InArgs, TSharedRef< class FSequenc
 	bMirrorLabels = true;
 	TSharedRef<ITimeSlider> BottomTimeSlider = SequencerWidgets.CreateTimeSlider( TimeSliderController, bMirrorLabels);
 
+	SAssignNew( TrackArea, SSequencerTrackArea, Sequencer.Pin().ToSharedRef() )
+		.ViewRange( InArgs._ViewRange )
+		.OutlinerFillPercent( this, &SSequencer::GetAnimationOutlinerFillPercentage );
+
 	ChildSlot
 	[
 		SNew(SBorder)
@@ -180,37 +181,15 @@ void SSequencer::Construct( const FArguments& InArgs, TSharedRef< class FSequenc
 			[
 				SNew( SHorizontalBox )
 				+SHorizontalBox::Slot()
-				.Padding( FMargin( 0.0f, 0.0f, 0.0f, 0.0f ) )
-				.FillWidth( TAttribute<float>( this, &SSequencer::GetAnimationOutlinerFillPercentage ) )
-				.VAlign(VAlign_Center)
+				.AutoWidth()
 				[
-					SNew(SHorizontalBox)
-					+ SHorizontalBox::Slot()
-					.AutoWidth()
-					[
-						MakeToolBar()
-					]
-					+SHorizontalBox::Slot()
-					.HAlign(HAlign_Right)
-					.FillWidth(1.f)
-					[
-						// @todo Sequencer - Temp clean view button
-						SNew( SCheckBox )
-						.Visibility( this, &SSequencer::OnGetCleanViewVisibility )
-						.Style(FEditorStyle::Get(), "ToggleButtonCheckbox")
-						.IsChecked( this, &SSequencer::OnGetCleanViewCheckState )
-						.OnCheckStateChanged( this, &SSequencer::OnCleanViewChecked )
-						[
-							SNew( STextBlock )
-							.Text( LOCTEXT("ToggleCleanView", "Clean View") )
-						]
-					]
+					MakeToolBar()
 				]
 				+SHorizontalBox::Slot()
-				.FillWidth(1.0f)
+				.AutoWidth()
 				[
-					SAssignNew(BreadcrumbTrail, SBreadcrumbTrail<FSequencerBreadcrumb>)
-					.OnCrumbClicked(this, &SSequencer::OnCrumbClicked)
+					SNew( SSequencerCurveEditorToolBar, TrackArea->GetCurveEditor()->GetCommands() )
+					.Visibility( this, &SSequencer::GetCurveEditorToolBarVisibility )
 				]
 			]
 			+ SVerticalBox::Slot()
@@ -254,9 +233,7 @@ void SSequencer::Construct( const FArguments& InArgs, TSharedRef< class FSequenc
 				]
 				+ SOverlay::Slot()
 				[
-					SAssignNew( TrackArea, SSequencerTrackArea, Sequencer.Pin().ToSharedRef() )
-					.ViewRange( InArgs._ViewRange )
-					.OutlinerFillPercent( this, &SSequencer::GetAnimationOutlinerFillPercentage )
+					TrackArea.ToSharedRef()
 				]
 				+ SOverlay::Slot()
 				[
@@ -302,6 +279,23 @@ void SSequencer::Construct( const FArguments& InArgs, TSharedRef< class FSequenc
 					]
 				]
 			]
+			+ SVerticalBox::Slot()
+			.AutoHeight()
+			[
+				SNew( SHorizontalBox )
+				+ SHorizontalBox::Slot()
+				.FillWidth( TAttribute<float>( this, &SSequencer::GetAnimationOutlinerFillPercentage ) )
+				[
+					SNew( SSpacer )
+				]
+				+ SHorizontalBox::Slot()
+				.FillWidth( 1.0f )
+				[
+					SAssignNew( BreadcrumbTrail, SBreadcrumbTrail<FSequencerBreadcrumb> )
+					.Visibility(this, &SSequencer::GetBreadcrumbTrailVisibility)
+					.OnCrumbClicked( this, &SSequencer::OnCrumbClicked )
+				]
+			]
 		]
 	];
 
@@ -312,33 +306,75 @@ TSharedRef<SWidget> SSequencer::MakeToolBar()
 {
 	FToolBarBuilder ToolBarBuilder( Sequencer.Pin()->GetCommandBindings(), FMultiBoxCustomization::None, TSharedPtr<FExtender>(), Orient_Horizontal, true);
 	
-	ToolBarBuilder.SetLabelVisibility(EVisibility::Visible);
-	ToolBarBuilder.AddToolBarButton( FSequencerCommands::Get().ToggleAutoKeyEnabled );
-
-	ToolBarBuilder.SetLabelVisibility(EVisibility::Collapsed);
-	ToolBarBuilder.AddToolBarButton( FSequencerCommands::Get().ToggleIsSnapEnabled, NAME_None, TAttribute<FText>( FText::GetEmpty() ) );
-	ToolBarBuilder.AddComboButton(
-		FUIAction(),
-		FOnGetContent::CreateSP( this, &SSequencer::MakeSnapMenu ),
-		LOCTEXT( "SnapOptions", "Options" ),
-		LOCTEXT( "SnapOptionsToolTip", "Snapping Options" ),
-		TAttribute<FSlateIcon>(),
-		true );
-	ToolBarBuilder.AddWidget(
-		SNew( SBox )
-		.VAlign( VAlign_Center )
-		[
-			SNew( SNumericEntryBox<float> )
-			.AllowSpin( true )
-			.MinValue( 0 )
-			.MinSliderValue( 0 )
-			.MaxSliderValue( 1 )
-			.MinDesiredValueWidth( 50 )
-			.Delta( 0.05f )
-			.ToolTipText( LOCTEXT( "SnappingIntervalToolTip", "Snapping interval" ) )
-			.Value( this, &SSequencer::OnGetSnapInterval )
-			.OnValueChanged( this, &SSequencer::OnSnapIntervalChanged ) 
-		] );
+	ToolBarBuilder.BeginSection("Base Commands");
+	{
+		ToolBarBuilder.SetLabelVisibility(EVisibility::Visible);
+		ToolBarBuilder.AddToolBarButton( FSequencerCommands::Get().ToggleAutoKeyEnabled );
+
+		if ( Sequencer.Pin()->IsLevelEditorSequencer() )
+		{
+			ToolBarBuilder.AddToolBarButton( FSequencerCommands::Get().ToggleCleanView );
+		}
+	}
+	ToolBarBuilder.EndSection();
+
+	ToolBarBuilder.BeginSection("Snapping");
+	{
+		TArray<SNumericDropDown<float>::FNamedValue> SnapValues;
+		SnapValues.Add( SNumericDropDown<float>::FNamedValue( 0.001f, LOCTEXT( "Snap_OneThousandth", "0.001" ), LOCTEXT( "SnapDescription_OneThousandth", "Set snap to 1/1000th" ) ) );
+		SnapValues.Add( SNumericDropDown<float>::FNamedValue( 0.01f, LOCTEXT( "Snap_OneHundredth", "0.01" ), LOCTEXT( "SnapDescription_OneHundredth", "Set snap to 1/100th" ) ) );
+		SnapValues.Add( SNumericDropDown<float>::FNamedValue( 0.1f, LOCTEXT( "Snap_OneTenth", "0.1" ), LOCTEXT( "SnapDescription_OneTenth", "Set snap to 1/10th" ) ) );
+		SnapValues.Add( SNumericDropDown<float>::FNamedValue( 1.0f, LOCTEXT( "Snap_One", "1" ), LOCTEXT( "SnapDescription_One", "Set snap to 1" ) ) );
+		SnapValues.Add( SNumericDropDown<float>::FNamedValue( 10.0f, LOCTEXT( "Snap_Ten", "10" ), LOCTEXT( "SnapDescription_Ten", "Set snap to 10" ) ) );
+		SnapValues.Add( SNumericDropDown<float>::FNamedValue( 100.0f, LOCTEXT( "Snap_OneHundred", "100" ), LOCTEXT( "SnapDescription_OneHundred", "Set snap to 100" ) ) );
+
+		ToolBarBuilder.SetLabelVisibility( EVisibility::Collapsed );
+		ToolBarBuilder.AddToolBarButton( FSequencerCommands::Get().ToggleIsSnapEnabled, NAME_None, TAttribute<FText>( FText::GetEmpty() ) );
+		ToolBarBuilder.AddComboButton(
+			FUIAction(),
+			FOnGetContent::CreateSP( this, &SSequencer::MakeSnapMenu ),
+			LOCTEXT( "SnapOptions", "Options" ),
+			LOCTEXT( "SnapOptionsToolTip", "Snapping Options" ),
+			TAttribute<FSlateIcon>(),
+			true );
+		ToolBarBuilder.AddWidget(
+			SNew( SBox )
+			.VAlign( VAlign_Center )
+			[
+				SNew( SNumericDropDown<float> )
+				.DropDownValues( SnapValues )
+				.LabelText( LOCTEXT( "TimeSnapLabel", "Time" ) )
+				.ToolTipText( LOCTEXT( "TimeSnappingIntervalToolTip", "Time snapping interval" ) )
+				.Value( this, &SSequencer::OnGetTimeSnapInterval )
+				.OnValueChanged( this, &SSequencer::OnTimeSnapIntervalChanged )
+			] );
+		ToolBarBuilder.AddWidget(
+			SNew( SBox )
+			.VAlign( VAlign_Center )
+			[
+				SNew( SNumericDropDown<float> )
+				.DropDownValues( SnapValues )
+				.LabelText( LOCTEXT( "ValueSnapLabel", "Value" ) )
+				.ToolTipText( LOCTEXT( "ValueSnappingIntervalToolTip", "Curve value snapping interval" ) )
+				.Value( this, &SSequencer::OnGetValueSnapInterval )
+				.OnValueChanged( this, &SSequencer::OnValueSnapIntervalChanged )
+			] );
+	}
+	ToolBarBuilder.EndSection();
+
+	ToolBarBuilder.BeginSection("Curve Editor");
+	{
+		ToolBarBuilder.SetLabelVisibility( EVisibility::Visible );
+		ToolBarBuilder.AddToolBarButton( FSequencerCommands::Get().ToggleShowCurveEditor );
+		ToolBarBuilder.SetLabelVisibility( EVisibility::Collapsed );
+		ToolBarBuilder.AddComboButton(
+			FUIAction(),
+			FOnGetContent::CreateSP( this, &SSequencer::MakeCurveEditorMenu ),
+			LOCTEXT( "CurveEditorOptions", "Options" ),
+			LOCTEXT( "CurveEditorOptionsToolTip", "Curve Editor Options" ),
+			TAttribute<FSlateIcon>(),
+			true );
+	}
 
 	return ToolBarBuilder.MakeWidget();
 }
@@ -349,15 +385,15 @@ TSharedRef<SWidget> SSequencer::MakeSnapMenu()
 
 	MenuBuilder.BeginSection( "KeySnapping", LOCTEXT( "SnappingMenuKeyHeader", "Key Snapping" ) );
 	{
-		MenuBuilder.AddMenuEntry( FSequencerCommands::Get().ToggleSnapKeysToInterval );
-		MenuBuilder.AddMenuEntry( FSequencerCommands::Get().ToggleSnapKeysToKeys );
+		MenuBuilder.AddMenuEntry( FSequencerCommands::Get().ToggleSnapKeyTimesToInterval );
+		MenuBuilder.AddMenuEntry( FSequencerCommands::Get().ToggleSnapKeyTimesToKeys );
 	}
 	MenuBuilder.EndSection();
 
 	MenuBuilder.BeginSection( "SectionSnapping", LOCTEXT( "SnappingMenuSectionHeader", "Section Snapping" ) );
 	{
-		MenuBuilder.AddMenuEntry( FSequencerCommands::Get().ToggleSnapSectionsToInterval );
-		MenuBuilder.AddMenuEntry( FSequencerCommands::Get().ToggleSnapSectionsToSections );
+		MenuBuilder.AddMenuEntry( FSequencerCommands::Get().ToggleSnapSectionTimesToInterval );
+		MenuBuilder.AddMenuEntry( FSequencerCommands::Get().ToggleSnapSectionTimesToSections );
 	}
 	MenuBuilder.EndSection();
 
@@ -367,6 +403,33 @@ TSharedRef<SWidget> SSequencer::MakeSnapMenu()
 	}
 	MenuBuilder.EndSection();
 
+	MenuBuilder.BeginSection( "CurveSnapping", LOCTEXT( "SnappingMenuCurveHeader", "Curve Snapping" ) );
+	{
+		MenuBuilder.AddMenuEntry( FSequencerCommands::Get().ToggleSnapCurveValueToInterval );
+	}
+	MenuBuilder.EndSection();
+
+	return MenuBuilder.MakeWidget();
+}
+
+TSharedRef<SWidget> SSequencer::MakeCurveEditorMenu()
+{
+	FMenuBuilder MenuBuilder( true, Sequencer.Pin()->GetCommandBindings() );
+
+	MenuBuilder.BeginSection( "CurveVisibility", LOCTEXT( "CurveEditorMenuCurveVisibilityHeader", "Curve Visibility" ) );
+	{
+		MenuBuilder.AddMenuEntry( FSequencerCommands::Get().SetAllCurveVisibility );
+		MenuBuilder.AddMenuEntry( FSequencerCommands::Get().SetSelectedCurveVisibility );
+		MenuBuilder.AddMenuEntry( FSequencerCommands::Get().SetAnimatedCurveVisibility );
+	}
+	MenuBuilder.EndSection();
+
+	MenuBuilder.BeginSection( "Options", LOCTEXT( "CurveEditorMenuOptionsHeader", "Editor Options" ) );
+	{
+		MenuBuilder.AddMenuEntry( FSequencerCommands::Get().ToggleShowCurveEditorCurveToolTips );
+	}
+	MenuBuilder.EndSection();
+
 	return MenuBuilder.MakeWidget();
 }
 
@@ -403,7 +466,7 @@ void SSequencer::UpdateLayoutTree()
 	// Update the node tree
 	SequencerNodeTree->Update();
 	
-	TrackArea->Update( *SequencerNodeTree );
+	TrackArea->Update( SequencerNodeTree );
 
 	SequencerNodeTree->UpdateCachedVisibilityBasedOnShotFiltersChanged();
 }
@@ -472,29 +535,24 @@ TSharedRef<SWidget> SSequencer::MakeSectionOverlay( TSharedRef<FSequencerTimeSli
 		];
 }
 
-EVisibility SSequencer::OnGetCleanViewVisibility() const 
-{
-	return Sequencer.Pin()->IsLevelEditorSequencer() ? EVisibility::Visible : EVisibility::Collapsed;
-}
-
-ECheckBoxState SSequencer::OnGetCleanViewCheckState() const
+float SSequencer::OnGetTimeSnapInterval() const
 {
-	return CleanViewEnabled.Get() ? ECheckBoxState::Checked : ECheckBoxState::Unchecked;
+	return GetDefault<USequencerSettings>()->GetTimeSnapInterval();
 }
 
-void SSequencer::OnCleanViewChecked( ECheckBoxState InState ) 
+void SSequencer::OnTimeSnapIntervalChanged( float InInterval )
 {
-	OnToggleCleanView.ExecuteIfBound( InState == ECheckBoxState::Checked ? true : false );
+	GetMutableDefault<USequencerSettings>()->SetTimeSnapInterval(InInterval);
 }
 
-TOptional<float> SSequencer::OnGetSnapInterval() const
+float SSequencer::OnGetValueSnapInterval() const
 {
-	return GetDefault<USequencerSnapSettings>()->GetSnapInterval();
+	return GetDefault<USequencerSettings>()->GetCurveValueSnapInterval();
 }
 
-void SSequencer::OnSnapIntervalChanged( float InInterval )
+void SSequencer::OnValueSnapIntervalChanged( float InInterval )
 {
-	GetMutableDefault<USequencerSnapSettings>()->SetSnapInterval(InInterval);
+	GetMutableDefault<USequencerSettings>()->SetCurveValueSnapInterval( InInterval );
 }
 
 
@@ -805,4 +863,14 @@ void SSequencer::DeleteSelectedNodes()
 	}
 }
 
+EVisibility SSequencer::GetBreadcrumbTrailVisibility() const
+{
+	return Sequencer.Pin()->IsLevelEditorSequencer() ? EVisibility::Visible : EVisibility::Collapsed;
+}
+
+EVisibility SSequencer::GetCurveEditorToolBarVisibility() const
+{
+	return GetDefault<USequencerSettings>()->GetShowCurveEditor() ? EVisibility::Visible : EVisibility::Collapsed;
+}
+
 #undef LOCTEXT_NAMESPACE
diff --git a/Engine/Source/Editor/Sequencer/Private/SSequencer.h b/Engine/Source/Editor/Sequencer/Private/SSequencer.h
--- a/Engine/Source/Editor/Sequencer/Private/SSequencer.h
+++ b/Engine/Source/Editor/Sequencer/Private/SSequencer.h
@@ -57,14 +57,10 @@ public:
 		SLATE_ATTRIBUTE( TRange<float>, ViewRange )
 		/** The current scrub position in (seconds) */
 		SLATE_ATTRIBUTE( float, ScrubPosition )
-		/** Whether the clean sequencer view is enabled */
-		SLATE_ATTRIBUTE( bool, CleanViewEnabled )
 		/** Called when the user changes the view range */
 		SLATE_EVENT( FOnViewRangeChanged, OnViewRangeChanged )
 		/** Called when the user changes the scrub position */
 		SLATE_EVENT( FOnScrubPositionChanged, OnScrubPositionChanged )
-		/** Called when the user toggles auto-key */
-		SLATE_EVENT( FOnToggleBoolOption, OnToggleCleanView )
 	SLATE_END_ARGS()
 
 
@@ -100,33 +96,33 @@ private:
 
 	/** Makes the snapping menu for the toolbar. */
 	TSharedRef<SWidget> MakeSnapMenu();
-	
-	/** @return The visibility of the clean view button */
-	EVisibility OnGetCleanViewVisibility() const;
+
+	/** Makes the curve editor menu for the toolbar. */
+	TSharedRef<SWidget> MakeCurveEditorMenu();
 
 	/**
-	 * @return The state of the auto-key check box
-	 */
-	ECheckBoxState OnGetCleanViewCheckState() const;
-	
+	* @return The value of the current time snap interval.
+	*/
+	float OnGetTimeSnapInterval() const;
+
 	/**
-	 * Called when the auto-key checkbox state changes
-	 * 
-	 * @param InState	The new state
-	 */
-	void OnCleanViewChecked( ECheckBoxState InState );	
+	* Called when the time snap interval changes.
+	*
+	* @param InInterval	The new time snap interval
+	*/
+	void OnTimeSnapIntervalChanged(float InInterval);
 
 	/**
-	* @return The value of the current snap interval.
+	* @return The value of the current value snap interval.
 	*/
-	TOptional<float> OnGetSnapInterval() const;
+	float OnGetValueSnapInterval() const;
 
 	/**
-	* Called when the snap interval slider changes.
+	* Called when the value snap interval changes.
 	*
-	* @param InInterval	The new snap interval
+	* @param InInterval	The new value snap interval
 	*/
-	void OnSnapIntervalChanged(float InInterval);
+	void OnValueSnapIntervalChanged( float InInterval );
 
 	/**
 	 * Called when the outliner search terms change                                                              
@@ -194,6 +190,12 @@ private:
 	/** Gets the title of the passed in shot section */
 	FText GetShotSectionTitle(UMovieSceneSection* ShotSection) const;
 
+	/** Gets whether or not the breadcrumb trail should be visible. */
+	EVisibility GetBreadcrumbTrailVisibility() const;
+
+	/** Gets whether or not the curve editor toolbar should be visibe. */
+	EVisibility GetCurveEditorToolBarVisibility() const;
+
 private:
 	/** Section area widget */
 	TSharedPtr<class SSequencerTrackArea> TrackArea;
@@ -201,10 +203,8 @@ private:
 	TSharedPtr<class FSequencerNodeTree> SequencerNodeTree;
 	/** The breadcrumb trail widget for this sequencer */
 	TSharedPtr< class SBreadcrumbTrail<FSequencerBreadcrumb> > BreadcrumbTrail;
-	/** Delegate to call when clean view is toggled */
-	FOnToggleBoolOption OnToggleCleanView;
-	/** Whether the clean sequencer view is enabled */
-	TAttribute<bool> CleanViewEnabled;
+	/** Whether or not the breadcrumb trail should be shown. */
+	TAttribute<bool> bShowBreadcrumbTrail;
 	/** The main sequencer interface */
 	TWeakPtr<FSequencer> Sequencer;
 
diff --git a/Engine/Source/Editor/Sequencer/Private/SSequencerCurveEditor.cpp b/Engine/Source/Editor/Sequencer/Private/SSequencerCurveEditor.cpp
new file mode 100644
--- /dev/null
+++ b/Engine/Source/Editor/Sequencer/Private/SSequencerCurveEditor.cpp
@@ -0,0 +1,115 @@
+// Copyright 1998-2015 Epic Games, Inc. All Rights Reserved.
+
+#include "SequencerPrivatePCH.h"
+#include "SCurveEditor.h"
+
+#define LOCTEXT_NAMESPACE "SequencerTrackArea"
+
+void SSequencerCurveEditor::Construct( const FArguments& InArgs, TSharedRef<FSequencer> InSequencer )
+{
+	ViewRange = InArgs._ViewRange;
+	Sequencer = InSequencer;
+	GetMutableDefault<USequencerSettings>()->GetOnCurveVisibilityChanged()->AddSP( this, &SSequencerCurveEditor::SequencerCurveVisibilityChanged );
+
+	ChildSlot
+	[
+		SAssignNew( CurveEditor, SCurveEditor )
+		.ViewMinInput_Lambda( [this]{ return ViewRange.Get().GetLowerBoundValue(); } )
+		.ViewMaxInput_Lambda( [this]{ return ViewRange.Get().GetUpperBoundValue(); } )
+		.HideUI( false )
+		.ShowCurveSelector( false )
+		.ShowZoomButtons( false )
+		.ShowInputGridNumbers( false )
+		.SnappingEnabled( this, &SSequencerCurveEditor::GetCurveSnapEnabled )
+		.InputSnap( this, &SSequencerCurveEditor::GetCurveTimeSnapInterval )
+		.OutputSnap( this, &SSequencerCurveEditor::GetCurveValueSnapInterval )
+		.TimelineLength( 0.0f )
+		.ShowCurveToolTips( this, &SSequencerCurveEditor::GetShowCurveEditorCurveToolTips )
+		.GridColor( FLinearColor( 0.3f, 0.3f, 0.3f, 0.3f ) )
+	];
+}
+
+template<class ParentNodeType>
+TSharedPtr<ParentNodeType> GetParentOfType( TSharedRef<FSequencerDisplayNode> InNode, ESequencerNode::Type InNodeType )
+{
+	TSharedPtr<FSequencerDisplayNode> CurrentNode = InNode->GetParent();
+	while ( CurrentNode.IsValid() )
+	{
+		if ( CurrentNode->GetType() == InNodeType )
+		{
+			break;
+		}
+		CurrentNode = CurrentNode->GetParent();
+	}
+	return StaticCastSharedPtr<ParentNodeType>( CurrentNode );
+}
+
+void SSequencerCurveEditor::SetSequencerNodeTree( TSharedPtr<FSequencerNodeTree> InSequencerNodeTree )
+{
+	if ( SequencerNodeTree.IsValid() )
+	{
+		SequencerNodeTree->GetOnSelectionChanged()->Remove( NodeTreeSelectionChangedHandle );
+	}
+
+	SequencerNodeTree = InSequencerNodeTree;
+	NodeTreeSelectionChangedHandle = SequencerNodeTree->GetOnSelectionChanged()->AddSP( this, &SSequencerCurveEditor::NodeTreeSelectionChanged );
+
+	UpdateCurveOwner();
+}
+
+void SSequencerCurveEditor::UpdateCurveOwner()
+{
+	CurveOwner = MakeShareable( new FSequencerCurveOwner( SequencerNodeTree ) );
+	CurveEditor->SetCurveOwner( CurveOwner.Get() );
+}
+
+bool SSequencerCurveEditor::GetCurveSnapEnabled() const
+{
+	const USequencerSettings* Settings = GetDefault<USequencerSettings>();
+	return Settings->GetIsSnapEnabled();
+}
+
+float SSequencerCurveEditor::GetCurveTimeSnapInterval() const
+{
+	return GetDefault<USequencerSettings>()->GetSnapKeyTimesToInterval()
+		? GetDefault<USequencerSettings>()->GetTimeSnapInterval()
+		: 0;
+}
+
+float SSequencerCurveEditor::GetCurveValueSnapInterval() const
+{
+	return GetDefault<USequencerSettings>()->GetSnapCurveValueToInterval()
+		? GetDefault<USequencerSettings>()->GetCurveValueSnapInterval()
+		: 0;
+}
+
+bool SSequencerCurveEditor::GetShowCurveEditorCurveToolTips() const
+{
+	return GetDefault<USequencerSettings>()->GetShowCurveEditorCurveToolTips();
+}
+
+void SSequencerCurveEditor::NodeTreeSelectionChanged()
+{
+	if ( GetDefault<USequencerSettings>()->GetCurveVisibility() == ESequencerCurveVisibility::SelectedCurves )
+	{
+		UpdateCurveOwner();
+	}
+}
+
+void SSequencerCurveEditor::SequencerCurveVisibilityChanged()
+{
+	UpdateCurveOwner();
+}
+
+TSharedPtr<FUICommandList> SSequencerCurveEditor::GetCommands()
+{
+	return CurveEditor->GetCommands();
+}
+
+SSequencerCurveEditor::~SSequencerCurveEditor()
+{
+	if ( SequencerNodeTree.IsValid() )
+	{
+		SequencerNodeTree->GetOnSelectionChanged()->Remove( NodeTreeSelectionChangedHandle );
+	}
+}
\ No newline at end of file
diff --git a/Engine/Source/Editor/Sequencer/Private/SSequencerCurveEditor.h b/Engine/Source/Editor/Sequencer/Private/SSequencerCurveEditor.h
new file mode 100644
--- /dev/null
+++ b/Engine/Source/Editor/Sequencer/Private/SSequencerCurveEditor.h
@@ -0,0 +1,65 @@
+// Copyright 1998-2015 Epic Games, Inc. All Rights Reserved.
+
+#pragma once
+
+class SCurveEditor;
+
+/** A widget for displaying and managing an SCurveEditor in sequencer. */
+class SSequencerCurveEditor : public SCompoundWidget
+{
+public:
+	SLATE_BEGIN_ARGS( SSequencerCurveEditor )
+	{}
+
+		/** The range of time being viewed */
+		SLATE_ATTRIBUTE( TRange<float>, ViewRange )
+
+	SLATE_END_ARGS()
+
+	void SSequencerCurveEditor::Construct( const FArguments& InArgs, TSharedRef<FSequencer> InSequencer );
+
+	/** Sets the sequencer node tree which supplies the curves. */
+	void SetSequencerNodeTree( TSharedPtr<FSequencerNodeTree> InSequencerNodeTree );
+	
+	/** Gets the commands for building a toolbar for the curve editor. */
+	TSharedPtr<FUICommandList> GetCommands();
+
+	~SSequencerCurveEditor();
+
+private:
+	/** Builds and assigns a new curve owner to the FCurveEditor. */
+	void UpdateCurveOwner();
+
+	/** Gets whether or not snapping is enabled for the curve editor. */
+	bool GetCurveSnapEnabled() const;
+
+	/** Gets time snapping interval. */
+	float GetCurveTimeSnapInterval() const;
+
+	/** Gets the value snapping interval. */
+	float GetCurveValueSnapInterval() const;
+
+	/** Gets whether or not to display curve tooltips. */
+	bool GetShowCurveEditorCurveToolTips() const;
+
+	/** Run whenever the selection on the FSequencerNodeTree changes. */
+	void NodeTreeSelectionChanged();
+
+	/** Run whenever the selected curve visibility changes. */
+	void SequencerCurveVisibilityChanged();
+
+private:
+
+	/** The sequencer which owns this widget. */
+	TWeakPtr<FSequencer> Sequencer;
+	/** The curve editor widget contained by this widget. */
+	TSharedPtr<SCurveEditor> CurveEditor;
+	/** The visible time range displayed by the curve editor. */
+	TAttribute<TRange<float>> ViewRange;
+	/** The sequencer node tree which contains the key area nodes which supply the curves to edit. */
+	TSharedPtr<FSequencerNodeTree> SequencerNodeTree;
+	/** The sequencer curve owner implementation which is visualized by the SCurveEditor. */
+	TSharedPtr<FSequencerCurveOwner> CurveOwner;
+	/** A handle to remove the node tree selection changed delegate. */
+	FDelegateHandle NodeTreeSelectionChangedHandle;
+};
\ No newline at end of file
diff --git a/Engine/Source/Editor/Sequencer/Private/SSequencerCurveEditorToolBar.cpp b/Engine/Source/Editor/Sequencer/Private/SSequencerCurveEditorToolBar.cpp
new file mode 100644
--- /dev/null
+++ b/Engine/Source/Editor/Sequencer/Private/SSequencerCurveEditorToolBar.cpp
@@ -0,0 +1,32 @@
+// Copyright 1998-2015 Epic Games, Inc. All Rights Reserved.
+
+#include "SequencerPrivatePCH.h"
+#include "RichCurveEditorCommands.h"
+
+void SSequencerCurveEditorToolBar::Construct( const FArguments& InArgs, TSharedPtr<FUICommandList> CurveEditorCommandList )
+{
+	FToolBarBuilder ToolBarBuilder( CurveEditorCommandList, FMultiBoxCustomization::None, TSharedPtr<FExtender>(), Orient_Horizontal, true );
+	ToolBarBuilder.BeginSection( "Curve" );
+	{
+		ToolBarBuilder.AddToolBarButton( FRichCurveEditorCommands::Get().ZoomToFitHorizontal );
+		ToolBarBuilder.AddToolBarButton( FRichCurveEditorCommands::Get().ZoomToFitVertical );
+		ToolBarBuilder.AddToolBarButton( FRichCurveEditorCommands::Get().ZoomToFitAll );
+		ToolBarBuilder.AddToolBarButton( FRichCurveEditorCommands::Get().ZoomToFitSelected );
+	}
+	ToolBarBuilder.EndSection();
+
+	ToolBarBuilder.BeginSection( "Interpolation" );
+	{
+		ToolBarBuilder.AddToolBarButton( FRichCurveEditorCommands::Get().InterpolationCubicAuto );
+		ToolBarBuilder.AddToolBarButton( FRichCurveEditorCommands::Get().InterpolationCubicUser );
+		ToolBarBuilder.AddToolBarButton( FRichCurveEditorCommands::Get().InterpolationCubicBreak );
+		ToolBarBuilder.AddToolBarButton( FRichCurveEditorCommands::Get().InterpolationLinear );
+		ToolBarBuilder.AddToolBarButton( FRichCurveEditorCommands::Get().InterpolationConstant );
+	}
+	ToolBarBuilder.EndSection();
+
+	ChildSlot
+	[
+		ToolBarBuilder.MakeWidget()
+	];
+}
\ No newline at end of file
diff --git a/Engine/Source/Editor/Sequencer/Private/SSequencerCurveEditorToolBar.h b/Engine/Source/Editor/Sequencer/Private/SSequencerCurveEditorToolBar.h
new file mode 100644
--- /dev/null
+++ b/Engine/Source/Editor/Sequencer/Private/SSequencerCurveEditorToolBar.h
@@ -0,0 +1,15 @@
+// Copyright 1998-2015 Epic Games, Inc. All Rights Reserved.
+
+#pragma once
+
+/** Wraps and builds a toolbar which works with the SSequencerCurveEditor. */
+class SSequencerCurveEditorToolBar : public SCompoundWidget
+{
+public:
+	SLATE_BEGIN_ARGS( SSequencerCurveEditorToolBar )
+	{}
+
+	SLATE_END_ARGS()
+
+	void Construct( const FArguments& InArgs, TSharedPtr<FUICommandList> CurveEditorCommandList );
+};
\ No newline at end of file
diff --git a/Engine/Source/Editor/Sequencer/Private/SSequencerTrackArea.cpp b/Engine/Source/Editor/Sequencer/Private/SSequencerTrackArea.cpp
--- a/Engine/Source/Editor/Sequencer/Private/SSequencerTrackArea.cpp
+++ b/Engine/Source/Editor/Sequencer/Private/SSequencerTrackArea.cpp
@@ -27,18 +27,44 @@ void SSequencerTrackArea::Construct( const FArguments& InArgs, TSharedRef<FSeque
 			.ExternalScrollbar(ScrollBar)
 		]
 		+ SOverlay::Slot()
-		.HAlign(HAlign_Right)
 		[
-			ScrollBar
+			SNew(SHorizontalBox)
+			+ SHorizontalBox::Slot()
+			.FillWidth(OutlinerFillPercent)
+			[
+				SNew(SSpacer)
+			]
+			+ SHorizontalBox::Slot()
+			.FillWidth(1.0f)
+			[
+				SAssignNew(CurveEditor, SSequencerCurveEditor, InSequencer)
+				.Visibility( this, &SSequencerTrackArea::GetCurveEditorVisibility )
+				.ViewRange(ViewRange)
+			]
+		]
+		+ SOverlay::Slot()
+		[
+			SNew(SHorizontalBox)
+			+ SHorizontalBox::Slot()
+			.HAlign(HAlign_Right)
+			.FillWidth( TAttribute<float>( this, &SSequencerTrackArea::GetScrollBarSlotFill ) )
+			[
+				ScrollBar
+			]
+			+ SHorizontalBox::Slot()
+			.FillWidth( TAttribute<float>( this, &SSequencerTrackArea::GetScrollBarSpacerSlotFill ) )
+			[
+				SNew( SSpacer )
+			]
 		]
 	];
 }
 
-void SSequencerTrackArea::Update( const FSequencerNodeTree& InSequencerNodeTree )
+void SSequencerTrackArea::Update( TSharedPtr<FSequencerNodeTree> InSequencerNodeTree )
 {
 	ScrollBox->ClearChildren();
 
-	const TArray< TSharedRef<FSequencerDisplayNode> >& RootNodes = InSequencerNodeTree.GetRootNodes();
+	const TArray< TSharedRef<FSequencerDisplayNode> >& RootNodes = InSequencerNodeTree->GetRootNodes();
 
 	for( int32 RootNodeIndex = 0; RootNodeIndex < RootNodes.Num(); ++RootNodeIndex )
 	{
@@ -53,6 +79,8 @@ void SSequencerTrackArea::Update( const FSequencerNodeTree& InSequencerNodeTree
 
 	// @todo Sequencer - Remove this expensive operation
 	ScrollBox->SlatePrepass();
+
+	CurveEditor->SetSequencerNodeTree(InSequencerNodeTree);
 }
 
 void SSequencerTrackArea::GenerateLayoutNodeWidgetsRecursive( const TArray< TSharedRef<FSequencerDisplayNode> >& Nodes )
@@ -87,8 +115,32 @@ void SSequencerTrackArea::GenerateWidgetForNode( TSharedRef<FSequencerDisplayNod
 		.FillWidth( 1.0f )
 		.Padding( FMargin(0.0f, 2.0f) )
 		[
-			// Generate a widget for the section area portion of the node
-			Node->GenerateWidgetForSectionArea( ViewRange )
+			SNew(SBox)
+			.Visibility(this, &SSequencerTrackArea::GetSectionControlVisibility)
+			[
+				// Generate a widget for the section area portion of the node
+				Node->GenerateWidgetForSectionArea( ViewRange )
+			]
 		]
 	];
 }
+
+EVisibility SSequencerTrackArea::GetCurveEditorVisibility() const
+{
+	return GetDefault<USequencerSettings>()->GetShowCurveEditor() ? EVisibility::Visible : EVisibility::Hidden;
+}
+
+EVisibility SSequencerTrackArea::GetSectionControlVisibility() const
+{
+	return GetDefault<USequencerSettings>()->GetShowCurveEditor() ? EVisibility::Hidden : EVisibility::Visible;
+}
+
+float SSequencerTrackArea::GetScrollBarSlotFill() const
+{
+	return GetDefault<USequencerSettings>()->GetShowCurveEditor() ? OutlinerFillPercent.Get() : 1.0f;
+}
+
+float SSequencerTrackArea::GetScrollBarSpacerSlotFill() const
+{
+	return GetDefault<USequencerSettings>()->GetShowCurveEditor() ? 1.0f : 0.0f;
+}
diff --git a/Engine/Source/Editor/Sequencer/Private/SSequencerTrackArea.h b/Engine/Source/Editor/Sequencer/Private/SSequencerTrackArea.h
--- a/Engine/Source/Editor/Sequencer/Private/SSequencerTrackArea.h
+++ b/Engine/Source/Editor/Sequencer/Private/SSequencerTrackArea.h
@@ -21,7 +21,10 @@ public:
 	/**
 	 * Updates the track area with new nodes
 	 */
-	virtual void Update( const FSequencerNodeTree& InSequencerNodeTree );
+	virtual void Update( TSharedPtr<FSequencerNodeTree> InSequencerNodeTree );
+
+	TSharedPtr<SSequencerCurveEditor> GetCurveEditor() { return CurveEditor; }
+
 private:
 	/**
 	 * Generates a widget for the specified node and puts it into the scroll box
@@ -36,6 +39,19 @@ private:
 	 * @param Nodes to generate widgets for.  Will recurse into each node's children
 	 */
 	void GenerateLayoutNodeWidgetsRecursive( const TArray< TSharedRef<FSequencerDisplayNode> >& Nodes );
+
+	/** Gets the visibility of the curve editor based on the the current settings. */
+	EVisibility GetCurveEditorVisibility() const;
+	/** Gets the visibility of the default section controls based on the current settings. */
+	EVisibility GetSectionControlVisibility() const;
+
+	/** Gets the fill value for the slot which contains the scrollbar.  This allows the scroll bar to move depending
+	  * whether or not the curve editor is visible. */
+	float GetScrollBarSlotFill() const;
+	/** Gets the fill value for the spacer slot which is to the right of the scrollbar.  This allows the scroll bar 
+	  * to move depending whether or not the curve editor is visible. */
+	float GetScrollBarSpacerSlotFill() const;
+
 private:
 	/** Scrollable area to display widgets */
 	TSharedPtr<SScrollBox> ScrollBox;
@@ -45,4 +61,6 @@ private:
 	TAttribute<float> OutlinerFillPercent;
 	/** The main sequencer interface */
 	TWeakPtr<FSequencer> Sequencer;
+	/** The curve editor. */
+	TSharedPtr<SSequencerCurveEditor> CurveEditor;
 };
diff --git a/Engine/Source/Editor/Sequencer/Private/SectionDragOperations.cpp b/Engine/Source/Editor/Sequencer/Private/SectionDragOperations.cpp
--- a/Engine/Source/Editor/Sequencer/Private/SectionDragOperations.cpp
+++ b/Engine/Source/Editor/Sequencer/Private/SectionDragOperations.cpp
@@ -165,10 +165,10 @@ void FResizeSection::OnDrag( const FPointerEvent& MouseEvent, const FVector2D& L
 		TRange<float> SectionBoundaries = GetSectionBoundaries(Section.Get(), SequencerNode);
 		
 		// Snapping
-		if ( SnapSettings->GetIsSnapEnabled() )
+		if ( Settings->GetIsSnapEnabled() )
 		{
 			bool bSnappedToSection = false;
-			if ( SnapSettings->GetSnapSectionsToSections() )
+			if ( Settings->GetSnapSectionTimesToSections() )
 		{
 			TArray<float> TimesToSnapTo;
 			GetSectionSnapTimes(TimesToSnapTo, Section.Get(), SequencerNode, bIsDilating);
@@ -182,9 +182,9 @@ void FResizeSection::OnDrag( const FPointerEvent& MouseEvent, const FVector2D& L
 				}
 			}
 
-			if ( bSnappedToSection == false && SnapSettings->GetSnapSectionsToInterval() )
+			if ( bSnappedToSection == false && Settings->GetSnapSectionTimesToInterval() )
 			{
-				NewTime = SnapSettings->SnapToInterval(NewTime);
+				NewTime = Settings->SnapTimeToInterval(NewTime);
 			}
 		}
 
@@ -271,10 +271,10 @@ void FMoveSection::OnDrag( const FPointerEvent& MouseEvent, const FVector2D& Loc
 		
 		float DeltaTime = DistanceMoved;
 		
-		if ( SnapSettings->GetIsSnapEnabled() )
+		if ( Settings->GetIsSnapEnabled() )
 		{
 			bool bSnappedToSection = false;
-			if ( SnapSettings->GetSnapSectionsToSections() )
+			if ( Settings->GetSnapSectionTimesToSections() )
 		{
 			TArray<float> TimesToSnapTo;
 			GetSectionSnapTimes(TimesToSnapTo, Section.Get(), SequencerNode, true);
@@ -292,10 +292,10 @@ void FMoveSection::OnDrag( const FPointerEvent& MouseEvent, const FVector2D& Loc
 			}
 		}
 
-			if ( bSnappedToSection == false && SnapSettings->GetSnapSectionsToInterval() )
+			if ( bSnappedToSection == false && Settings->GetSnapSectionTimesToInterval() )
 			{
 				float NewStartTime = DistanceMoved + Section->GetStartTime();
-				DeltaTime = SnapSettings->SnapToInterval( NewStartTime ) - Section->GetStartTime();
+				DeltaTime = Settings->SnapTimeToInterval( NewStartTime ) - Section->GetStartTime();
 			}
 		}
 
@@ -429,22 +429,22 @@ void FMoveKeys::OnDrag( const FPointerEvent& MouseEvent, const FVector2D& LocalM
 	{
 		float TimeDelta = DistanceMoved;
 		// Snapping
-		if ( SnapSettings->GetIsSnapEnabled() )
+		if ( Settings->GetIsSnapEnabled() )
 		{
 			bool bSnappedToKeyTime = false;
-			if ( SnapSettings->GetSnapKeysToKeys() )
+			if ( Settings->GetSnapKeyTimesToKeys() )
 			{
-			TArray<float> OutSnapTimes;
-			GetKeySnapTimes(OutSnapTimes, SequencerNode);
+				TArray<float> OutSnapTimes;
+				GetKeySnapTimes(OutSnapTimes, SequencerNode);
 
-			TArray<float> InitialTimes;
-			for (TSet<FSelectedKey>::TConstIterator It(SelectedKeys); It; ++It)
-			{
-				const FSelectedKey& SelectedKey = *It;
-				InitialTimes.Add(SelectedKey.KeyArea->GetKeyTime(SelectedKey.KeyHandle.GetValue()) + DistanceMoved);
-			}
-			float OutInitialTime = 0.f;
-			float OutSnapTime = 0.f;
+				TArray<float> InitialTimes;
+				for (TSet<FSelectedKey>::TConstIterator It(SelectedKeys); It; ++It)
+				{
+					const FSelectedKey& SelectedKey = *It;
+					InitialTimes.Add(SelectedKey.KeyArea->GetKeyTime(SelectedKey.KeyHandle.GetValue()) + DistanceMoved);
+				}
+				float OutInitialTime = 0.f;
+				float OutSnapTime = 0.f;
 				if ( SnapToTimes( InitialTimes, OutSnapTimes, TimeToPixelConverter, OutInitialTime, OutSnapTime ) )
 				{
 					bSnappedToKeyTime = true;
@@ -452,9 +452,9 @@ void FMoveKeys::OnDrag( const FPointerEvent& MouseEvent, const FVector2D& LocalM
 				}
 			}
 
-			if ( bSnappedToKeyTime == false && SnapSettings->GetSnapKeysToInterval() )
+			if ( bSnappedToKeyTime == false && Settings->GetSnapKeyTimesToInterval() )
 			{
-				TimeDelta = SnapSettings->SnapToInterval( MouseTime ) - SelectedKeyTime;
+				TimeDelta = Settings->SnapTimeToInterval( MouseTime ) - SelectedKeyTime;
 			}
 		}
 
diff --git a/Engine/Source/Editor/Sequencer/Private/SectionDragOperations.h b/Engine/Source/Editor/Sequencer/Private/SectionDragOperations.h
--- a/Engine/Source/Editor/Sequencer/Private/SectionDragOperations.h
+++ b/Engine/Source/Editor/Sequencer/Private/SectionDragOperations.h
@@ -10,7 +10,7 @@ class FSequencerDragOperation
 public:
 	FSequencerDragOperation()
 	{
-		SnapSettings = GetDefault<USequencerSnapSettings>();
+		Settings = GetDefault<USequencerSettings>();
 	}
 
 	virtual ~FSequencerDragOperation(){}
@@ -75,8 +75,8 @@ protected:
 	 */
 	TOptional<float> SnapToTimes(float InitialTimes, const TArray<float>& SnapTimes, const FTimeToPixel& TimeToPixelConverter);
 
-	/** The current snap settings */
-	const USequencerSnapSettings* SnapSettings;
+	/** The current sequencer settings */
+	const USequencerSettings* Settings;
 };
 
 
diff --git a/Engine/Source/Editor/Sequencer/Private/Sequencer.cpp b/Engine/Source/Editor/Sequencer/Private/Sequencer.cpp
--- a/Engine/Source/Editor/Sequencer/Private/Sequencer.cpp
+++ b/Engine/Source/Editor/Sequencer/Private/Sequencer.cpp
@@ -75,10 +75,8 @@ void FSequencer::InitSequencer( const FSequencerInitParams& InitParams, const TA
 		SequencerWidget = SNew( SSequencer, SharedThis( this ) )
 			.ViewRange( this, &FSequencer::OnGetViewRange )
 			.ScrubPosition( this, &FSequencer::OnGetScrubPosition )
-			.CleanViewEnabled( this, &FSequencer::IsUsingCleanView )
 			.OnScrubPositionChanged( this, &FSequencer::OnScrubPositionChanged )
-			.OnViewRangeChanged( this, &FSequencer::OnViewRangeChanged, false )
-			.OnToggleCleanView( this, &FSequencer::OnToggleCleanView );
+			.OnViewRangeChanged( this, &FSequencer::OnViewRangeChanged, false );
 
 		// When undo occurs, get a notification so we can make sure our view is up to date
 		GEditor->RegisterForUndo(this);
@@ -133,7 +131,6 @@ FSequencer::FSequencer()
 	, LastViewRange(0.f, 5.f)
 	, PlaybackState( EMovieScenePlayerStatus::Stopped )
 	, ScrubPosition( 0.0f )
-	, bCleanViewEnabled( false )
 	, bLoopingEnabled( false )
 	, bAllowAutoKey( false )
 	, bPerspectiveViewportPossessionEnabled( true )
@@ -892,13 +889,6 @@ void FSequencer::OnToggleAutoKey()
 	bAllowAutoKey = !bAllowAutoKey;
 }
 
-void FSequencer::OnToggleCleanView( bool bInCleanViewEnabled )
-{
-	bCleanViewEnabled = bInCleanViewEnabled;
-}
-
-
-
 FGuid FSequencer::AddSpawnableForAssetOrClass( UObject* Object, UObject* CounterpartGamePreviewObject )
 {
 	FGuid NewSpawnableGuid;
@@ -1322,11 +1312,6 @@ bool FSequencer::IsShotFilteringOn() const
 	return FilteringShots.Num() > 0;
 }
 
-bool FSequencer::IsUsingCleanView() const
-{
-	return bCleanViewEnabled;
-}
-
 float FSequencer::GetOverlayFadeCurve() const
 {
 	return OverlayCurve.GetLerp();
@@ -1399,7 +1384,7 @@ void FSequencer::BindSequencerCommands()
 		Commands.SetKey,
 		FExecuteAction::CreateSP( this, &FSequencer::SetKey ) );
 
-	USequencerSnapSettings* SnapSettings = GetMutableDefault<USequencerSnapSettings>();
+	USequencerSettings* Settings = GetMutableDefault<USequencerSettings>();
 
 	SequencerCommandBindings->MapAction(
 		Commands.ToggleAutoKeyEnabled,
@@ -1407,41 +1392,83 @@ void FSequencer::BindSequencerCommands()
 		FCanExecuteAction::CreateLambda( []{ return true; } ),
 		FIsActionChecked::CreateSP( this, &FSequencer::OnGetAllowAutoKey ) );
 
+	SequencerCommandBindings->MapAction(
+		Commands.ToggleCleanView,
+		FExecuteAction::CreateLambda( [Settings]{ Settings->SetIsUsingCleanView( !Settings->GetIsUsingCleanView() ); } ),
+		FCanExecuteAction::CreateLambda( []{ return true; } ),
+		FIsActionChecked::CreateLambda( [Settings]{ return Settings->GetIsUsingCleanView(); } ) );
+
 	SequencerCommandBindings->MapAction(
 		Commands.ToggleIsSnapEnabled,
-		FExecuteAction::CreateLambda( [SnapSettings]{ SnapSettings->SetIsSnapEnabled( !SnapSettings->GetIsSnapEnabled() ); } ),
+		FExecuteAction::CreateLambda( [Settings]{ Settings->SetIsSnapEnabled( !Settings->GetIsSnapEnabled() ); } ),
 		FCanExecuteAction::CreateLambda( []{ return true; } ),
-		FIsActionChecked::CreateLambda( [SnapSettings]{ return SnapSettings->GetIsSnapEnabled(); } ) );
+		FIsActionChecked::CreateLambda( [Settings]{ return Settings->GetIsSnapEnabled(); } ) );
 
 	SequencerCommandBindings->MapAction(
-		Commands.ToggleSnapKeysToInterval,
-		FExecuteAction::CreateLambda( [SnapSettings]{ SnapSettings->SetSnapKeysToInterval( !SnapSettings->GetSnapKeysToInterval() ); } ),
+		Commands.ToggleSnapKeyTimesToInterval,
+		FExecuteAction::CreateLambda( [Settings]{ Settings->SetSnapKeyTimesToInterval( !Settings->GetSnapKeyTimesToInterval() ); } ),
 		FCanExecuteAction::CreateLambda( []{ return true; } ),
-		FIsActionChecked::CreateLambda( [SnapSettings]{ return SnapSettings->GetSnapKeysToInterval(); } ) );
+		FIsActionChecked::CreateLambda( [Settings]{ return Settings->GetSnapKeyTimesToInterval(); } ) );
 
 	SequencerCommandBindings->MapAction(
-		Commands.ToggleSnapKeysToKeys,
-		FExecuteAction::CreateLambda( [SnapSettings]{ SnapSettings->SetSnapKeysToKeys( !SnapSettings->GetSnapKeysToKeys() ); } ),
+		Commands.ToggleSnapKeyTimesToKeys,
+		FExecuteAction::CreateLambda( [Settings]{ Settings->SetSnapKeyTimesToKeys( !Settings->GetSnapKeyTimesToKeys() ); } ),
 		FCanExecuteAction::CreateLambda( []{ return true; } ),
-		FIsActionChecked::CreateLambda( [SnapSettings]{ return SnapSettings->GetSnapKeysToKeys(); } ) );
+		FIsActionChecked::CreateLambda( [Settings]{ return Settings->GetSnapKeyTimesToKeys(); } ) );
 
 	SequencerCommandBindings->MapAction(
-		Commands.ToggleSnapSectionsToInterval,
-		FExecuteAction::CreateLambda( [SnapSettings]{ SnapSettings->SetSnapSectionsToInterval( !SnapSettings->GetSnapSectionsToInterval() ); } ),
+		Commands.ToggleSnapSectionTimesToInterval,
+		FExecuteAction::CreateLambda( [Settings]{ Settings->SetSnapSectionTimesToInterval( !Settings->GetSnapSectionTimesToInterval() ); } ),
 		FCanExecuteAction::CreateLambda( []{ return true; } ),
-		FIsActionChecked::CreateLambda( [SnapSettings]{ return SnapSettings->GetSnapSectionsToInterval(); } ) );
+		FIsActionChecked::CreateLambda( [Settings]{ return Settings->GetSnapSectionTimesToInterval(); } ) );
 
 	SequencerCommandBindings->MapAction(
-		Commands.ToggleSnapSectionsToSections,
-		FExecuteAction::CreateLambda( [SnapSettings]{ SnapSettings->SetSnapSectionsToSections( !SnapSettings->GetSnapSectionsToSections() ); } ),
+		Commands.ToggleSnapSectionTimesToSections,
+		FExecuteAction::CreateLambda( [Settings]{ Settings->SetSnapSectionTimesToSections( !Settings->GetSnapSectionTimesToSections() ); } ),
 		FCanExecuteAction::CreateLambda( []{ return true; } ),
-		FIsActionChecked::CreateLambda( [SnapSettings]{ return SnapSettings->GetSnapSectionsToSections(); } ) );
+		FIsActionChecked::CreateLambda( [Settings]{ return Settings->GetSnapSectionTimesToSections(); } ) );
 
 	SequencerCommandBindings->MapAction(
 		Commands.ToggleSnapPlayTimeToInterval,
-		FExecuteAction::CreateLambda( [SnapSettings]{ SnapSettings->SetSnapPlayTimeToInterval( !SnapSettings->GetSnapPlayTimeToInterval() ); } ),
+		FExecuteAction::CreateLambda( [Settings]{ Settings->SetSnapPlayTimeToInterval( !Settings->GetSnapPlayTimeToInterval() ); } ),
+		FCanExecuteAction::CreateLambda( []{ return true; } ),
+		FIsActionChecked::CreateLambda( [Settings]{ return Settings->GetSnapPlayTimeToInterval(); } ) );
+
+	SequencerCommandBindings->MapAction(
+		Commands.ToggleSnapCurveValueToInterval,
+		FExecuteAction::CreateLambda( [Settings]{ Settings->SetSnapCurveValueToInterval( !Settings->GetSnapCurveValueToInterval() ); } ),
+		FCanExecuteAction::CreateLambda( []{ return true; } ),
+		FIsActionChecked::CreateLambda( [Settings]{ return Settings->GetSnapCurveValueToInterval(); } ) );
+
+	SequencerCommandBindings->MapAction(
+		Commands.ToggleShowCurveEditor,
+		FExecuteAction::CreateLambda( [Settings]{ Settings->SetShowCurveEditor(!Settings->GetShowCurveEditor()); } ),
+		FCanExecuteAction::CreateLambda( []{ return true; } ),
+		FIsActionChecked::CreateLambda( [Settings]{ return Settings->GetShowCurveEditor(); } ) );
+
+	SequencerCommandBindings->MapAction(
+		Commands.ToggleShowCurveEditorCurveToolTips,
+		FExecuteAction::CreateLambda( [Settings]{ Settings->SetShowCurveEditorCurveToolTips( !Settings->GetShowCurveEditorCurveToolTips() ); } ),
+		FCanExecuteAction::CreateLambda( []{ return true; } ),
+		FIsActionChecked::CreateLambda( [Settings]{ return Settings->GetShowCurveEditorCurveToolTips(); } ) );
+
+	SequencerCommandBindings->MapAction(
+		Commands.SetAllCurveVisibility,
+		FExecuteAction::CreateLambda( [Settings]{ Settings->SetCurveVisibility( ESequencerCurveVisibility::AllCurves ); } ),
+		FCanExecuteAction::CreateLambda( []{ return true; } ),
+		FIsActionChecked::CreateLambda( [Settings]{ return Settings->GetCurveVisibility() == ESequencerCurveVisibility::AllCurves; } ) );
+
+	SequencerCommandBindings->MapAction(
+		Commands.SetSelectedCurveVisibility,
+		FExecuteAction::CreateLambda( [Settings]{ Settings->SetCurveVisibility( ESequencerCurveVisibility::SelectedCurves ); } ),
+		FCanExecuteAction::CreateLambda( []{ return true; } ),
+		FIsActionChecked::CreateLambda( [Settings]{ return Settings->GetCurveVisibility() == ESequencerCurveVisibility::SelectedCurves; } ) );
+
+	SequencerCommandBindings->MapAction(
+		Commands.SetAnimatedCurveVisibility,
+		FExecuteAction::CreateLambda( [Settings]{ Settings->SetCurveVisibility( ESequencerCurveVisibility::AnimatedCurves ); } ),
 		FCanExecuteAction::CreateLambda( []{ return true; } ),
-		FIsActionChecked::CreateLambda( [SnapSettings]{ return SnapSettings->GetSnapPlayTimeToInterval(); } ) );
+		FIsActionChecked::CreateLambda( [Settings]{ return Settings->GetCurveVisibility() == ESequencerCurveVisibility::AnimatedCurves; } ) );
 }
 
 void FSequencer::BuildObjectBindingContextMenu(FMenuBuilder& MenuBuilder, const FGuid& ObjectBinding, const UClass* ObjectClass)
diff --git a/Engine/Source/Editor/Sequencer/Private/Sequencer.h b/Engine/Source/Editor/Sequencer/Private/Sequencer.h
--- a/Engine/Source/Editor/Sequencer/Private/Sequencer.h
+++ b/Engine/Source/Editor/Sequencer/Private/Sequencer.h
@@ -233,12 +233,6 @@ public:
 	 */
 	bool IsShotFilteringOn() const;
 
-	/**
-	 * Returns true if the sequencer is using the 'Clean View' mode
-	 * Clean View simply means no non-global tracks will appear if no shots are filtering
-	 */
-	bool IsUsingCleanView() const;
-
 	/**
 	 * Gets the overlay fading animation curve lerp
 	 */
@@ -355,13 +349,6 @@ protected:
 	 * @param bInAllowAutoKey	The new auto key state
 	 */
 	void OnToggleAutoKey();
-	
-	/**
-	 * Called when auto-key is toggled by a user
-	 *
-	 * @param bInAllowAutoKey	The new auto key state
-	 */
-	void OnToggleCleanView( bool bInCleanViewEnabled );
 
 	/** Called via UEditorEngine::GetActorRecordingStateEvent to check to see whether we need to record actor state */
 	void GetActorRecordingState( bool& bIsRecording /* In+Out */ ) const;
@@ -450,10 +437,6 @@ private:
 	/** The current scrub position */
 	// @todo sequencer: Should use FTimespan or "double" for Time Cursor Position! (cascades)
 	float ScrubPosition;
-	
-	/** Whether the clean sequencer view is enabled */
-	bool bCleanViewEnabled;
-
 
 	/** Whether looping while playing is enabled for this sequencer */
 	bool bLoopingEnabled;
diff --git a/Engine/Source/Editor/Sequencer/Private/SequencerCommands.cpp b/Engine/Source/Editor/Sequencer/Private/SequencerCommands.cpp
--- a/Engine/Source/Editor/Sequencer/Private/SequencerCommands.cpp
+++ b/Engine/Source/Editor/Sequencer/Private/SequencerCommands.cpp
@@ -11,12 +11,21 @@ void FSequencerCommands::RegisterCommands()
 
 	UI_COMMAND( ToggleIsSnapEnabled, "Enable Snapping", "Enables and disables snapping.", EUserInterfaceActionType::ToggleButton, FInputGesture() );
 
-	UI_COMMAND( ToggleSnapKeysToInterval, "Snap to the interval", "Snap keys to the interval.", EUserInterfaceActionType::ToggleButton, FInputGesture() );
-	UI_COMMAND( ToggleSnapKeysToKeys, "Snap to other keys", "Snap keys to other keys in the section.", EUserInterfaceActionType::ToggleButton, FInputGesture() );
+	UI_COMMAND( ToggleSnapKeyTimesToInterval, "Snap to the interval", "Snap keys to the interval.", EUserInterfaceActionType::ToggleButton, FInputGesture() );
+	UI_COMMAND( ToggleSnapKeyTimesToKeys, "Snap to other keys", "Snap keys to other keys in the section.", EUserInterfaceActionType::ToggleButton, FInputGesture() );
 
-	UI_COMMAND( ToggleSnapSectionsToInterval, "Snap to the interval", "Snap sections to the interval.", EUserInterfaceActionType::ToggleButton, FInputGesture() );
-	UI_COMMAND( ToggleSnapSectionsToSections, "Snap to other sections", "Snap sections to other sections.", EUserInterfaceActionType::ToggleButton, FInputGesture() );
+	UI_COMMAND( ToggleSnapSectionTimesToInterval, "Snap to the interval", "Snap sections to the interval.", EUserInterfaceActionType::ToggleButton, FInputGesture() );
+	UI_COMMAND( ToggleSnapSectionTimesToSections, "Snap to other sections", "Snap sections to other sections.", EUserInterfaceActionType::ToggleButton, FInputGesture() );
 
 	UI_COMMAND( ToggleSnapPlayTimeToInterval, "Snap to the interval while scrubbing", "Snap the play time to the interval while scrubbing.", EUserInterfaceActionType::ToggleButton, FInputGesture() );
+
+	UI_COMMAND(ToggleSnapCurveValueToInterval, "Snap curve key values", "Snap curve keys to the value snapping interval.", EUserInterfaceActionType::ToggleButton, FInputGesture() );
+
+	UI_COMMAND( ToggleCleanView, "Clean View", "Enable 'Clean View' mode which displays only global tracks when no filter is applied.", EUserInterfaceActionType::ToggleButton, FInputGesture() );
 	
+	UI_COMMAND( ToggleShowCurveEditor, "Curve Editor", "Show the animation keys in a curve editor.", EUserInterfaceActionType::ToggleButton, FInputGesture() );
+	UI_COMMAND( ToggleShowCurveEditorCurveToolTips, "Curve Tool Tips", "Show a tool tip with name and values when hoving over a curve.", EUserInterfaceActionType::ToggleButton, FInputGesture() );
+	UI_COMMAND( SetAllCurveVisibility, "All Curves", "Show all curves in the curve editor.", EUserInterfaceActionType::RadioButton, FInputGesture() );
+	UI_COMMAND( SetSelectedCurveVisibility, "Selected Curves", "Show only selected curves in the curve editor.", EUserInterfaceActionType::RadioButton, FInputGesture() );
+	UI_COMMAND( SetAnimatedCurveVisibility, "Animated Curves", "Show only curves with keyframes in the curve editor.", EUserInterfaceActionType::RadioButton, FInputGesture() );
 }
diff --git a/Engine/Source/Editor/Sequencer/Private/SequencerCommands.h b/Engine/Source/Editor/Sequencer/Private/SequencerCommands.h
--- a/Engine/Source/Editor/Sequencer/Private/SequencerCommands.h
+++ b/Engine/Source/Editor/Sequencer/Private/SequencerCommands.h
@@ -25,20 +25,41 @@ public:
 	TSharedPtr< FUICommandInfo > ToggleIsSnapEnabled;
 
 	/** Toggles whether or not keys should snap to the selected interval. */
-	TSharedPtr< FUICommandInfo > ToggleSnapKeysToInterval;
+	TSharedPtr< FUICommandInfo > ToggleSnapKeyTimesToInterval;
 
 	/** Toggles whether or not keys should snap to other keys in the section. */
-	TSharedPtr< FUICommandInfo > ToggleSnapKeysToKeys;
+	TSharedPtr< FUICommandInfo > ToggleSnapKeyTimesToKeys;
 
 	/** Toggles whether or not sections should snap to the selected interval. */
-	TSharedPtr< FUICommandInfo > ToggleSnapSectionsToInterval;
+	TSharedPtr< FUICommandInfo > ToggleSnapSectionTimesToInterval;
 
 	/** Toggles whether or not sections should snap to other sections. */
-	TSharedPtr< FUICommandInfo > ToggleSnapSectionsToSections;
+	TSharedPtr< FUICommandInfo > ToggleSnapSectionTimesToSections;
 
 	/** Toggles whether or not the play time should snap to the selected interval. */
 	TSharedPtr< FUICommandInfo > ToggleSnapPlayTimeToInterval;
 
+	/** Toggles whether or not to snap curve values to the interval. */
+	TSharedPtr< FUICommandInfo > ToggleSnapCurveValueToInterval;
+
+	/** Toggles whether the "clean view" is enabled in the level editor. */
+	TSharedPtr< FUICommandInfo > ToggleCleanView;
+
+	/** Toggles whether or not the curve editor should be shown. */
+	TSharedPtr< FUICommandInfo > ToggleShowCurveEditor;
+
+	/** Toggles whether or not to show tool tips for the curves in the curve editor. */
+	TSharedPtr< FUICommandInfo > ToggleShowCurveEditorCurveToolTips;
+
+	/** Sets the curve visibility to all curves. */
+	TSharedPtr< FUICommandInfo > SetAllCurveVisibility;
+
+	/** Sets the curve visibility to the selected curves. */
+	TSharedPtr< FUICommandInfo > SetSelectedCurveVisibility;
+
+	/** Sets the curve visibility to the animated curves. */
+	TSharedPtr< FUICommandInfo > SetAnimatedCurveVisibility;
+
 	/**
 	 * Initialize commands
 	 */
diff --git a/Engine/Source/Editor/Sequencer/Private/SequencerCurveOwner.cpp b/Engine/Source/Editor/Sequencer/Private/SequencerCurveOwner.cpp
new file mode 100644
--- /dev/null
+++ b/Engine/Source/Editor/Sequencer/Private/SequencerCurveOwner.cpp
@@ -0,0 +1,136 @@
+// Copyright 1998-2015 Epic Games, Inc. All Rights Reserved.
+
+#include "SequencerPrivatePCH.h"
+#include "SequencerCurveOwner.h"
+
+void GetAllKeyAreaNodes( TSharedPtr<FSequencerNodeTree> InSequencerNodeTree, TArray<TSharedRef<FSectionKeyAreaNode>>& KeyAreaNodes )
+{
+	TDoubleLinkedList<TSharedRef<FSequencerDisplayNode>> NodesToProcess;
+	for ( auto RootNode : InSequencerNodeTree->GetRootNodes() )
+	{
+		NodesToProcess.AddTail( RootNode );
+	}
+
+	while ( NodesToProcess.Num() > 0 )
+	{
+		auto Node = NodesToProcess.GetHead();
+		if ( Node->GetValue()->GetType() == ESequencerNode::KeyArea )
+		{
+			KeyAreaNodes.Add( StaticCastSharedRef<FSectionKeyAreaNode>( Node->GetValue() ) );
+		}
+		for ( auto ChildNode : Node->GetValue()->GetChildNodes() )
+		{
+			NodesToProcess.AddTail( ChildNode );
+		}
+		NodesToProcess.RemoveNode( Node );
+	}
+}
+
+FName BuildCurveName( TSharedPtr<FSectionKeyAreaNode> KeyAreaNode )
+{
+	FString CurveName;
+	TSharedPtr<FSequencerDisplayNode> CurrentNameNode = KeyAreaNode;
+	TArray<FString> NameParts;
+	while ( CurrentNameNode.IsValid() )
+	{
+		NameParts.Add( CurrentNameNode->GetDisplayName().ToString() );
+		CurrentNameNode = CurrentNameNode->GetParent();
+	}
+	return FName(*FString::Join(NameParts, TEXT(" - ")));
+}
+
+FSequencerCurveOwner::FSequencerCurveOwner( TSharedPtr<FSequencerNodeTree> InSequencerNodeTree )
+{
+	SequencerNodeTree = InSequencerNodeTree;
+
+	TArray<TSharedRef<FSectionKeyAreaNode>> KeyAreaNodes;
+	GetAllKeyAreaNodes( SequencerNodeTree, KeyAreaNodes );
+	for ( TSharedRef<FSectionKeyAreaNode> KeyAreaNode : KeyAreaNodes )
+	{
+		for ( TSharedRef<IKeyArea> KeyArea : KeyAreaNode->GetAllKeyAreas() )
+		{
+			IKeyArea::FCurveInfo* CurveInfo = KeyArea->GetCurveInfo();
+			if ( CurveInfo != nullptr )
+			{
+				bool bAddCurve = false;
+				switch ( GetDefault<USequencerSettings>()->GetCurveVisibility() )
+				{
+				case ESequencerCurveVisibility::AllCurves:
+					bAddCurve = true;
+					break;
+				case ESequencerCurveVisibility::SelectedCurves:
+					bAddCurve = KeyAreaNode->IsSelected();
+					break;
+				case ESequencerCurveVisibility::AnimatedCurves:
+					bAddCurve = CurveInfo->Curve->GetNumKeys() > 0;
+					break;
+				}
+				if ( bAddCurve )
+				{
+					FName CurveName = BuildCurveName(KeyAreaNode);
+					Curves.Add( FRichCurveEditInfo( CurveInfo->Curve, CurveName ) );
+					ConstCurves.Add( FRichCurveEditInfoConst( CurveInfo->Curve, CurveName ) );
+					EditInfoToSectionMap.Add( FRichCurveEditInfo( CurveInfo->Curve, CurveName ), CurveInfo->OwningSection );
+				}
+			}
+		}
+	}
+}
+
+TArray<FRichCurveEditInfoConst> FSequencerCurveOwner::GetCurves() const
+{
+	return ConstCurves;
+}
+
+TArray<FRichCurveEditInfo> FSequencerCurveOwner::GetCurves()
+{
+	return Curves;
+};
+
+void FSequencerCurveOwner::ModifyOwner()
+{
+	TArray<UMovieSceneSection*> Owners;
+	EditInfoToSectionMap.GenerateValueArray( Owners );
+	for ( auto Owner : Owners )
+	{
+		Owner->Modify();
+	}
+}
+
+void FSequencerCurveOwner::MakeTransactional()
+{
+	TArray<UMovieSceneSection*> Owners;
+	EditInfoToSectionMap.GenerateValueArray( Owners );
+	for ( auto Owner : Owners )
+	{
+		Owner->SetFlags( Owner->GetFlags() | RF_Transactional );
+	}
+}
+
+void FSequencerCurveOwner::OnCurveChanged( const TArray<FRichCurveEditInfo>& ChangedCurveEditInfos )
+{
+	// Whenever a curve changes make sure to resize it's section so that the curve fits.
+	for ( auto& ChangedCurveEditInfo : ChangedCurveEditInfos )
+	{
+		UMovieSceneSection** OwningSection = EditInfoToSectionMap.Find(ChangedCurveEditInfo);
+		if ( OwningSection != nullptr )
+		{
+			float CurveStart;
+			float CurveEnd;
+			ChangedCurveEditInfo.CurveToEdit->GetTimeRange(CurveStart, CurveEnd);
+			if ( (*OwningSection)->GetStartTime() > CurveStart )
+			{
+				(*OwningSection)->SetStartTime(CurveStart);
+			}			
+			if ( (*OwningSection)->GetEndTime() < CurveEnd )
+			{
+				(*OwningSection)->SetEndTime( CurveEnd );
+			}
+		}
+	}
+}
+
+bool FSequencerCurveOwner::IsValidCurve( FRichCurveEditInfo CurveInfo )
+{
+	return EditInfoToSectionMap.Contains(CurveInfo);
+}
\ No newline at end of file
diff --git a/Engine/Source/Editor/Sequencer/Private/SequencerCurveOwner.h b/Engine/Source/Editor/Sequencer/Private/SequencerCurveOwner.h
new file mode 100644
--- /dev/null
+++ b/Engine/Source/Editor/Sequencer/Private/SequencerCurveOwner.h
@@ -0,0 +1,31 @@
+// Copyright 1998-2015 Epic Games, Inc. All Rights Reserved.
+
+#pragma once
+
+/** A curve owner interface for displaying animation curves in sequencer. */
+class FSequencerCurveOwner : public FCurveOwnerInterface
+{
+public:
+	FSequencerCurveOwner(TSharedPtr<FSequencerNodeTree> InSequencerNodeTree);
+
+	/** FCurveOwnerInterface */
+	virtual TArray<FRichCurveEditInfoConst> GetCurves() const override;
+	virtual TArray<FRichCurveEditInfo> GetCurves() override;
+	virtual void ModifyOwner() override;
+	virtual void MakeTransactional() override;
+	virtual void OnCurveChanged( const TArray<FRichCurveEditInfo>& ChangedCurveEditInfos ) override;
+	virtual bool IsValidCurve( FRichCurveEditInfo CurveInfo );
+
+private:
+	/** The FSequencerNodeTree used to build the curve owner. */
+	TSharedPtr<FSequencerNodeTree> SequencerNodeTree;
+
+	/** The ordered array of const curves used to implement the curve owner interface. */
+	TArray<FRichCurveEditInfoConst> ConstCurves;
+
+	/** The ordered array or curves used to implement the curve owner interface. */
+	TArray<FRichCurveEditInfo> Curves;
+
+	/** A map of curve edit infos to their corresponding sections. */
+	TMap<FRichCurveEditInfo, UMovieSceneSection*> EditInfoToSectionMap;
+};
\ No newline at end of file
diff --git a/Engine/Source/Editor/Sequencer/Private/SequencerDisplayNode.cpp b/Engine/Source/Editor/Sequencer/Private/SequencerDisplayNode.cpp
--- a/Engine/Source/Editor/Sequencer/Private/SequencerDisplayNode.cpp
+++ b/Engine/Source/Editor/Sequencer/Private/SequencerDisplayNode.cpp
@@ -75,7 +75,7 @@ int32 SSequencerObjectTrack::OnPaint(const FPaintArgs& Args, const FGeometry& Al
 		FSlateDrawElement::MakeBox(
 			OutDrawElements,
 			LayerId+1,
-			AllottedGeometry.ToPaintGeometry(FVector2D(KeyPosition - FMath::TruncToFloat(KeyMarkSize.X/2.f), FMath::TruncToFloat(AllottedGeometry.Size.Y/2.f - KeyMarkSize.Y/2.f)), KeyMarkSize),
+			AllottedGeometry.ToPaintGeometry(FVector2D(KeyPosition - FMath::CeilToFloat(KeyMarkSize.X/2.f), FMath::CeilToFloat(AllottedGeometry.Size.Y/2.f - KeyMarkSize.Y/2.f)), KeyMarkSize),
 			FEditorStyle::GetBrush("Sequencer.KeyMark"),
 			MyClippingRect,
 			ESlateDrawEffect::None,
@@ -298,7 +298,7 @@ bool FSequencerDisplayNode::IsVisible() const
 	// If shot filtering is off on clean view is on, node must be pinned
 	return bCachedShotFilteredVisibility &&
 		(ParentTree.HasActiveFilter() ? ParentTree.IsNodeFiltered(AsShared()) : IsParentExpandedOrIsARootNode()) &&
-		(GetSequencer().IsShotFilteringOn() || !GetSequencer().IsUsingCleanView() || bNodeIsPinned);
+		(GetSequencer().IsShotFilteringOn() || !GetDefault<USequencerSettings>()->GetIsUsingCleanView() || bNodeIsPinned);
 }
 
 bool FSequencerDisplayNode::HasVisibleChildren() const
diff --git a/Engine/Source/Editor/Sequencer/Private/SequencerDisplayNode.h b/Engine/Source/Editor/Sequencer/Private/SequencerDisplayNode.h
--- a/Engine/Source/Editor/Sequencer/Private/SequencerDisplayNode.h
+++ b/Engine/Source/Editor/Sequencer/Private/SequencerDisplayNode.h
@@ -67,7 +67,7 @@ public:
 	virtual ESequencerNode::Type GetType() const = 0;
 
 	/** @return Whether or not this node can be selected */
-	virtual bool IsSelectable() const { return false; }
+	virtual bool IsSelectable() const { return true; }
 
 	/**
 	 * @return The desired height of the node when displayed
@@ -306,7 +306,6 @@ public:
 	virtual ESequencerNode::Type GetType() const override { return ESequencerNode::Track; }
 	virtual float GetNodeHeight() const override;
 	virtual FText GetDisplayName() const override;
-	virtual bool IsSelectable() const override { return true; }
 	virtual bool GetShotFilteredVisibilityToCache() const override;
 	virtual void GetChildKeyAreaNodesRecursively(TArray< TSharedRef<class FSectionKeyAreaNode> >& OutNodes) const override;
 
@@ -379,7 +378,6 @@ public:
 	virtual ESequencerNode::Type GetType() const override { return ESequencerNode::Object; }
 	virtual FText GetDisplayName() const override { return DisplayName; }
 	virtual float GetNodeHeight() const override;
-	virtual bool IsSelectable() const override { return true; }
 	virtual bool GetShotFilteredVisibilityToCache() const override;
 	
 	/** @return The object binding on this node */
diff --git a/Engine/Source/Editor/Sequencer/Private/SequencerNodeTree.cpp b/Engine/Source/Editor/Sequencer/Private/SequencerNodeTree.cpp
--- a/Engine/Source/Editor/Sequencer/Private/SequencerNodeTree.cpp
+++ b/Engine/Source/Editor/Sequencer/Private/SequencerNodeTree.cpp
@@ -197,6 +197,7 @@ void FSequencerNodeTree::SetSelectionState( TSharedRef<const FSequencerDisplayNo
 		// Not selecting so remove the node from the selection set
 		SelectedNodes.Remove( AffectedNode );
 	}
+	OnSelectionChanged.Broadcast();
 }
 
 bool FSequencerNodeTree::IsNodeSelected( TSharedRef<const FSequencerDisplayNode> Node ) const
@@ -295,3 +296,8 @@ void FSequencerNodeTree::UpdateCachedVisibilityBasedOnShotFiltersChanged()
 		RootNodes[i]->UpdateCachedShotFilteredVisibility();
 	}
 }
+
+FSequencerNodeTree::FOnSelectionChanged* FSequencerNodeTree::GetOnSelectionChanged()
+{
+	return &OnSelectionChanged;
+}
diff --git a/Engine/Source/Editor/Sequencer/Private/SequencerNodeTree.h b/Engine/Source/Editor/Sequencer/Private/SequencerNodeTree.h
--- a/Engine/Source/Editor/Sequencer/Private/SequencerNodeTree.h
+++ b/Engine/Source/Editor/Sequencer/Private/SequencerNodeTree.h
@@ -13,6 +13,8 @@ class UMovieSceneTrack;
 class FSequencerNodeTree : public TSharedFromThis<FSequencerNodeTree>
 {
 public:
+	DECLARE_MULTICAST_DELEGATE(FOnSelectionChanged)
+
 	FSequencerNodeTree( class FSequencer& InSequencer )
 		: Sequencer( InSequencer )
 	{}
@@ -92,6 +94,10 @@ public:
 	
 	/** Updates nodes visibility for use when shot filters change at all */
 	void UpdateCachedVisibilityBasedOnShotFiltersChanged();
+
+	/** Gets the multicast delegate which is run when the node selection changes. */
+	FOnSelectionChanged* GetOnSelectionChanged();
+
 private:
 	/**
 	 * Finds or adds a type editor for the track
@@ -130,4 +136,6 @@ private:
 	FString FilterString;
 	/** Sequencer interface */
 	FSequencer& Sequencer;
+	/** The multicast delegate which is run when the node selection changes */
+	FOnSelectionChanged OnSelectionChanged;
 };
diff --git a/Engine/Source/Editor/Sequencer/Private/SequencerPrivatePCH.h b/Engine/Source/Editor/Sequencer/Private/SequencerPrivatePCH.h
--- a/Engine/Source/Editor/Sequencer/Private/SequencerPrivatePCH.h
+++ b/Engine/Source/Editor/Sequencer/Private/SequencerPrivatePCH.h
@@ -20,7 +20,10 @@
 #include "SSequencerTrackViewPanel.h"
 #include "SAnimationOutlinerView.h"
 #include "MovieSceneInstance.h"
-#include "SequencerSnapSettings.h"
+#include "SequencerSettings.h"
+#include "SequencerCurveOwner.h"
+#include "SSequencerCurveEditor.h"
+#include "SSequencerCurveEditorToolBar.h"
 
 DECLARE_LOG_CATEGORY_EXTERN(LogSequencer, Log, All);
 
diff --git a/Engine/Source/Editor/Sequencer/Private/SequencerSettings.cpp b/Engine/Source/Editor/Sequencer/Private/SequencerSettings.cpp
new file mode 100644
--- /dev/null
+++ b/Engine/Source/Editor/Sequencer/Private/SequencerSettings.cpp
@@ -0,0 +1,216 @@
+// Copyright 1998-2015 Epic Games, Inc. All Rights Reserved.
+
+#include "SequencerPrivatePCH.h"
+
+USequencerSettings::USequencerSettings( const FObjectInitializer& ObjectInitializer )
+	: Super( ObjectInitializer )
+{
+	bIsSnapEnabled = true;
+	TimeSnapInterval = .05f;
+	bSnapKeyTimesToInterval = true;
+	bSnapKeyTimesToKeys = true;
+	bSnapSectionTimesToInterval = true;
+	bSnapSectionTimesToSections = true;
+	bSnapPlayTimeToInterval = true;
+	CurveValueSnapInterval = 10.0f;
+	bSnapCurveValueToInterval = true;
+	bIsUsingCleanView = false;
+	bShowCurveEditor = false;
+	bShowCurveEditorCurveToolTips = true;
+	CurveVisibility = ESequencerCurveVisibility::AllCurves;
+}
+
+bool USequencerSettings::GetIsSnapEnabled() const
+{
+	return bIsSnapEnabled;
+}
+
+void USequencerSettings::SetIsSnapEnabled(bool InbIsSnapEnabled)
+{
+	if ( bIsSnapEnabled != InbIsSnapEnabled )
+	{
+		bIsSnapEnabled = InbIsSnapEnabled;
+		SaveConfig();
+	}
+}
+
+float USequencerSettings::GetTimeSnapInterval() const
+{
+	return TimeSnapInterval;
+}
+
+void USequencerSettings::SetTimeSnapInterval(float InTimeSnapInterval)
+{
+	if ( TimeSnapInterval != InTimeSnapInterval )
+	{
+		TimeSnapInterval = InTimeSnapInterval;
+		SaveConfig();
+	}
+}
+
+bool USequencerSettings::GetSnapKeyTimesToInterval() const
+{
+	return bSnapKeyTimesToInterval;
+}
+
+void USequencerSettings::SetSnapKeyTimesToInterval(bool InbSnapKeyTimesToInterval)
+{
+	if ( bSnapKeyTimesToInterval != InbSnapKeyTimesToInterval )
+	{
+		bSnapKeyTimesToInterval = InbSnapKeyTimesToInterval;
+		SaveConfig();
+	}
+}
+
+bool USequencerSettings::GetSnapKeyTimesToKeys() const
+{
+	return bSnapKeyTimesToKeys;
+}
+
+void USequencerSettings::SetSnapKeyTimesToKeys(bool InbSnapKeyTimesToKeys)
+{
+	if ( bSnapKeyTimesToKeys != InbSnapKeyTimesToKeys )
+	{
+		bSnapKeyTimesToKeys = InbSnapKeyTimesToKeys;
+		SaveConfig();
+	}
+}
+
+bool USequencerSettings::GetSnapSectionTimesToInterval() const
+{
+	return bSnapSectionTimesToInterval;
+}
+
+void USequencerSettings::SetSnapSectionTimesToInterval(bool InbSnapSectionTimesToInterval)
+{
+	if ( bSnapSectionTimesToInterval != InbSnapSectionTimesToInterval )
+	{
+		bSnapSectionTimesToInterval = InbSnapSectionTimesToInterval;
+		SaveConfig();
+	}
+}
+
+bool USequencerSettings::GetSnapSectionTimesToSections() const
+{
+	return bSnapSectionTimesToSections;
+}
+
+void USequencerSettings::SetSnapSectionTimesToSections( bool InbSnapSectionTimesToSections )
+{
+	if ( bSnapSectionTimesToSections != InbSnapSectionTimesToSections )
+	{
+		bSnapSectionTimesToSections = InbSnapSectionTimesToSections;
+		SaveConfig();
+	}
+}
+
+bool USequencerSettings::GetSnapPlayTimeToInterval() const
+{
+	return bSnapPlayTimeToInterval;
+}
+
+void USequencerSettings::SetSnapPlayTimeToInterval(bool InbSnapPlayTimeToInterval)
+{
+	if ( bSnapPlayTimeToInterval != InbSnapPlayTimeToInterval )
+	{
+		bSnapPlayTimeToInterval = InbSnapPlayTimeToInterval;
+		SaveConfig();
+	}
+}
+
+float USequencerSettings::GetCurveValueSnapInterval() const
+{
+	return CurveValueSnapInterval;
+}
+
+void USequencerSettings::SetCurveValueSnapInterval( float InCurveValueSnapInterval )
+{
+	if ( CurveValueSnapInterval != InCurveValueSnapInterval )
+	{
+		CurveValueSnapInterval = InCurveValueSnapInterval;
+		SaveConfig();
+	}
+}
+
+bool USequencerSettings::GetSnapCurveValueToInterval() const
+{
+	return bSnapCurveValueToInterval;
+}
+
+void USequencerSettings::SetSnapCurveValueToInterval( bool InbSnapCurveValueToInterval )
+{
+	if ( bSnapCurveValueToInterval != InbSnapCurveValueToInterval )
+	{
+		bSnapCurveValueToInterval = InbSnapCurveValueToInterval;
+		SaveConfig();
+	}
+}
+
+bool USequencerSettings::GetIsUsingCleanView() const
+{
+	return bIsUsingCleanView;
+}
+
+void USequencerSettings::SetIsUsingCleanView(bool InbIsUsingCleanView )
+{
+	if (bIsUsingCleanView != InbIsUsingCleanView)
+	{
+		bIsUsingCleanView = InbIsUsingCleanView;
+		SaveConfig();
+	}
+}
+
+bool USequencerSettings::GetShowCurveEditor() const
+{
+	return bShowCurveEditor;
+}
+
+void USequencerSettings::SetShowCurveEditor(bool InbShowCurveEditor)
+{
+	if (bShowCurveEditor != InbShowCurveEditor)
+	{
+		bShowCurveEditor = InbShowCurveEditor;
+		SaveConfig();
+	}
+}
+
+bool USequencerSettings::GetShowCurveEditorCurveToolTips() const
+{
+	return bShowCurveEditorCurveToolTips;
+}
+
+void USequencerSettings::SetShowCurveEditorCurveToolTips(bool InbShowCurveEditorCurveToolTips)
+{
+	if (bShowCurveEditorCurveToolTips != InbShowCurveEditorCurveToolTips)
+	{
+		bShowCurveEditorCurveToolTips = InbShowCurveEditorCurveToolTips;
+		SaveConfig();
+	}
+}
+
+ESequencerCurveVisibility::Type USequencerSettings::GetCurveVisibility() const
+{
+	return CurveVisibility;
+}
+
+void USequencerSettings::SetCurveVisibility(ESequencerCurveVisibility::Type InCurveVisibility)
+{
+	if (CurveVisibility != InCurveVisibility)
+	{
+		CurveVisibility = InCurveVisibility;
+		OnCurveVisibilityChanged.Broadcast();
+		SaveConfig();
+	}
+}
+
+float USequencerSettings::SnapTimeToInterval( float InTimeValue ) const
+{
+	return TimeSnapInterval > 0
+		? FMath::RoundToInt( InTimeValue / TimeSnapInterval ) * TimeSnapInterval
+		: InTimeValue;
+}
+
+USequencerSettings::FOnCurveVisibilityChanged* USequencerSettings::GetOnCurveVisibilityChanged()
+{
+	return &OnCurveVisibilityChanged;
+}
diff --git a/Engine/Source/Editor/Sequencer/Private/SequencerSettings.h b/Engine/Source/Editor/Sequencer/Private/SequencerSettings.h
new file mode 100644
--- /dev/null
+++ b/Engine/Source/Editor/Sequencer/Private/SequencerSettings.h
@@ -0,0 +1,143 @@
+// Copyright 1998-2015 Epic Games, Inc. All Rights Reserved.
+
+#pragma once
+
+#include "SequencerSettings.generated.h"
+
+/** Defines visibility states for the curves in the curve editor. */
+UENUM()
+namespace ESequencerCurveVisibility
+{
+	enum Type
+	{
+		/** All curves should be visible. **/
+		AllCurves,
+		/** Only curves from selected nodes should be visible. **/
+		SelectedCurves,
+		/** Only curves which have keyframes should be visible. **/
+		AnimatedCurves
+	};
+}
+
+/** Serializable options for sequencer. */
+UCLASS(config=EditorUserSettings)
+class USequencerSettings : public UObject
+{
+public:
+	GENERATED_UCLASS_BODY()
+
+	DECLARE_MULTICAST_DELEGATE( FOnCurveVisibilityChanged );
+
+	/** Gets whether or not snapping is enabled. */
+	bool GetIsSnapEnabled() const;
+	/** Sets whether or not snapping is enabled. */
+	void SetIsSnapEnabled(bool InbIsSnapEnabled);
+
+	/** Gets the time in seconds used for interval snapping. */
+	float GetTimeSnapInterval() const;
+	/** Sets the time in seconds used for interval snapping. */
+	void SetTimeSnapInterval(float InTimeSnapInterval);
+
+	/** Gets whether or not to snap key times to the interval. */
+	bool GetSnapKeyTimesToInterval() const;
+	/** Sets whether or not to snap keys to the interval. */
+	void SetSnapKeyTimesToInterval(bool InbSnapKeyTimesToInterval);
+
+	/** Gets whether or not to snap keys to other keys. */
+	bool GetSnapKeyTimesToKeys() const;
+	/** Sets whether or not to snap keys to other keys. */
+	void SetSnapKeyTimesToKeys(bool InbSnapKeyTimesToKeys);
+
+	/** Gets whether or not to snap sections to the interval. */
+	bool GetSnapSectionTimesToInterval() const;
+	/** Sets whether or not to snap sections to the interval. */
+	void SetSnapSectionTimesToInterval(bool InbSnapSectionTimesToInterval);
+
+	/** Gets whether or not to snap sections to other sections. */
+	bool GetSnapSectionTimesToSections() const;
+	/** sets whether or not to snap sections to other sections. */
+	void SetSnapSectionTimesToSections( bool InbSnapSectionTimesToSections );
+
+	/** Gets whether or not to snap the play time to the interval while scrubbing. */
+	bool GetSnapPlayTimeToInterval() const;
+	/** Sets whether or not to snap the play time to the interval while scrubbing. */
+	void SetSnapPlayTimeToInterval(bool InbSnapPlayTimeToInterval);
+
+	/** Gets the snapping interval for curve values. */
+	float GetCurveValueSnapInterval() const;
+	/** Sets the snapping interval for curve values. */
+	void SetCurveValueSnapInterval(float InCurveValueSnapInterval);
+
+	/** Gets whether or not to snap curve values to the interval. */
+	bool GetSnapCurveValueToInterval() const;
+	/** Sets whether or not to snap curve values to the interval. */
+	void SetSnapCurveValueToInterval(bool InbSnapCurveValueToInterval);
+
+	/** Gets whether or not the 'Clean View' is enabled. In 'Clean View' mode only global tracks are displayed when no filter is applied. */
+	bool GetIsUsingCleanView() const;
+	/** Sets whether or not the 'Clean View' is enabled. In 'Clean View' mode only global tracks are displayed when no filter is applied. */
+	void SetIsUsingCleanView(bool InbIsUsingCleanView);
+
+	/** Gets whether or not the curve editor should be shown. */
+	bool GetShowCurveEditor() const;
+	/** Sets whether or not the curve editor should be shown. */
+	void SetShowCurveEditor(bool InbShowCurveEditor);
+
+	/** Gets whether or not to show curve tool tips in the curve editor. */
+	bool GetShowCurveEditorCurveToolTips() const;
+	/** Sets whether or not to show curve tool tips in the curve editor. */
+	void SetShowCurveEditorCurveToolTips(bool InbShowCurveEditorCurveToolTips);
+
+	/** Gets the current curve visibility. */
+	ESequencerCurveVisibility::Type GetCurveVisibility() const;
+	/** Sets the current curve visibility. */
+	void SetCurveVisibility(ESequencerCurveVisibility::Type InCurveVisibility);
+
+	/** Snaps a time value in seconds to the currently selected interval. */
+	float SnapTimeToInterval(float InTimeValue) const;
+
+	/** Gets the multicast delegate which is run whenever the CurveVisibility value changes. */
+	FOnCurveVisibilityChanged* GetOnCurveVisibilityChanged();
+
+protected:
+	UPROPERTY( config )
+	bool bIsSnapEnabled;
+
+	UPROPERTY( config )
+	float TimeSnapInterval;
+
+	UPROPERTY( config )
+	bool bSnapKeyTimesToInterval;
+
+	UPROPERTY( config )
+	bool bSnapKeyTimesToKeys;
+
+	UPROPERTY( config )
+	bool bSnapSectionTimesToInterval;
+
+	UPROPERTY( config )
+	bool bSnapSectionTimesToSections;
+
+	UPROPERTY( config )
+	bool bSnapPlayTimeToInterval;
+
+	UPROPERTY( config )
+	float CurveValueSnapInterval;
+
+	UPROPERTY( config )
+	bool bSnapCurveValueToInterval;
+
+	UPROPERTY( config )
+	bool bIsUsingCleanView;
+
+	UPROPERTY( config )
+	bool bShowCurveEditor;
+
+	UPROPERTY( config )
+	bool bShowCurveEditorCurveToolTips;
+
+	UPROPERTY( config )
+	TEnumAsByte<ESequencerCurveVisibility::Type> CurveVisibility;
+
+	FOnCurveVisibilityChanged OnCurveVisibilityChanged;
+};
\ No newline at end of file
diff --git a/Engine/Source/Editor/Sequencer/Private/SequencerSnapSettings.cpp b/Engine/Source/Editor/Sequencer/Private/SequencerSnapSettings.cpp
deleted file mode 100644
diff --git a/Engine/Source/Editor/Sequencer/Private/SequencerSnapSettings.h b/Engine/Source/Editor/Sequencer/Private/SequencerSnapSettings.h
deleted file mode 100644
diff --git a/Engine/Source/Editor/Sequencer/Private/TimeSliderController.cpp b/Engine/Source/Editor/Sequencer/Private/TimeSliderController.cpp
--- a/Engine/Source/Editor/Sequencer/Private/TimeSliderController.cpp
+++ b/Engine/Source/Editor/Sequencer/Private/TimeSliderController.cpp
@@ -153,9 +153,9 @@ void FSequencerTimeSliderController::DrawTicks( FSlateWindowElementList& OutDraw
 		if ( AbsOffsetNum % Divider == 0 )
 		{
 			FVector2D Offset( XPos, InArgs.TickOffset );
-			FVector2D TickSize( 1.0f, InArgs.MajorTickHeight );
+			FVector2D TickSize( 0.0f, InArgs.MajorTickHeight );
 
-			LinePoints[0] = FVector2D(1.0f,1.0f);
+			LinePoints[0] = FVector2D( 0.0f,1.0f);
 			LinePoints[1] = TickSize;
 
 			// lines should not need anti-aliasing
@@ -170,7 +170,7 @@ void FSequencerTimeSliderController::DrawTicks( FSlateWindowElementList& OutDraw
 				InArgs.ClippingRect,
 				InArgs.DrawEffects,
 				InArgs.TickColor,
-				false
+				bAntiAliasLines
 				);
 
 			if( !InArgs.bOnlyDrawMajorTicks )
@@ -200,9 +200,9 @@ void FSequencerTimeSliderController::DrawTicks( FSlateWindowElementList& OutDraw
 			const float MinorTickHeight = AbsOffsetNum % HalfDivider == 0 ? 7.0f : 4.0f;
 
 			FVector2D Offset(XPos, InArgs.bMirrorLabels ? 0.0f : FMath::Abs( InArgs.AllottedGeometry.Size.Y - MinorTickHeight ) );
-			FVector2D TickSize(1, MinorTickHeight);
+			FVector2D TickSize(0.0f, MinorTickHeight);
 
-			LinePoints[0] = FVector2D(1.0f,1.0f);
+			LinePoints[0] = FVector2D(0.0f,1.0f);
 			LinePoints[1] = TickSize;
 
 			const bool bAntiAlias = false;
@@ -255,7 +255,7 @@ int32 FSequencerTimeSliderController::OnPaintTimeSlider( bool bMirrorLabels, con
 		DrawTicks( OutDrawElements, RangeToScreen, Args );
 
 		const float HandleSize = 13.0f;
-		float HalfSize = FMath::TruncToFloat(HandleSize/2.0f);
+		float HalfSize = FMath::CeilToFloat(HandleSize/2.0f);
 
 		// Draw the scrub handle
 		const float XPos = RangeToScreen.InputToLocalX( TimeSliderArgs.ScrubPosition.Get() );
@@ -334,10 +334,10 @@ FReply FSequencerTimeSliderController::OnMouseButtonUp( TSharedRef<SWidget> Widg
 			FVector2D CursorPos = MyGeometry.AbsoluteToLocal(MouseEvent.GetLastScreenSpacePosition());
 			float NewValue = RangeToScreen.LocalXToInput(CursorPos.X);
 
-			const USequencerSnapSettings* SnapSettings = GetDefault<USequencerSnapSettings>();
-			if ( SnapSettings->GetIsSnapEnabled() && SnapSettings->GetSnapPlayTimeToInterval() )
+			const USequencerSettings* Settings = GetDefault<USequencerSettings>();
+			if ( Settings->GetIsSnapEnabled() && Settings->GetSnapPlayTimeToInterval() )
 			{
-				NewValue = SnapSettings->SnapToInterval( NewValue );
+				NewValue = Settings->SnapTimeToInterval( NewValue );
 			}
 
 			CommitScrubPosition( NewValue, /*bIsScrubbing=*/false );
@@ -419,10 +419,10 @@ FReply FSequencerTimeSliderController::OnMouseMove( TSharedRef<SWidget> WidgetOw
 				FVector2D CursorPos = MyGeometry.AbsoluteToLocal( MouseEvent.GetLastScreenSpacePosition() );
 				float NewValue = RangeToScreen.LocalXToInput( CursorPos.X );
 
-				const USequencerSnapSettings* SnapSettings = GetDefault<USequencerSnapSettings>();
-				if ( SnapSettings->GetIsSnapEnabled() && SnapSettings->GetSnapPlayTimeToInterval() )
+				const USequencerSettings* Settings = GetDefault<USequencerSettings>();
+				if ( Settings->GetIsSnapEnabled() && Settings->GetSnapPlayTimeToInterval() )
 				{
-					NewValue = SnapSettings->SnapToInterval(NewValue);
+					NewValue = Settings->SnapTimeToInterval(NewValue);
 				}
 
 				CommitScrubPosition( NewValue, /*bIsScrubbing=*/true );
@@ -530,8 +530,8 @@ int32 FSequencerTimeSliderController::OnPaintSectionView( const FGeometry& Allot
 		// Draw a line for the scrub position
 		TArray<FVector2D> LinePoints;
 		LinePoints.AddUninitialized(2);
-		LinePoints[0] = FVector2D( 1.0f, 0.0f );
-		LinePoints[1] = FVector2D( 1.0f, FMath::RoundToFloat( AllottedGeometry.Size.Y ) );
+		LinePoints[0] = FVector2D( 0.0f, 0.0f );
+		LinePoints[1] = FVector2D( 0.0f, FMath::RoundToFloat( AllottedGeometry.Size.Y ) );
 
 		FSlateDrawElement::MakeLines(
 			OutDrawElements,
diff --git a/Engine/Source/Editor/Sequencer/Public/IKeyArea.h b/Engine/Source/Editor/Sequencer/Public/IKeyArea.h
--- a/Engine/Source/Editor/Sequencer/Public/IKeyArea.h
+++ b/Engine/Source/Editor/Sequencer/Public/IKeyArea.h
@@ -21,6 +21,24 @@ struct FKeyDrawingInfo
  */
 class IKeyArea
 {
+public:
+	/** Defines curve information for this key area. */
+	class FCurveInfo
+	{
+	public:
+		FCurveInfo( FRichCurve* InCurve, UMovieSceneSection* InOwningSection )
+			: Curve(InCurve)
+			, OwningSection(InOwningSection)
+		{
+		}
+
+		/** The curve associated with the key area. */
+		FRichCurve* const Curve;
+
+		/** The movie section which owns curve. */
+		UMovieSceneSection* const OwningSection;
+	};
+
 public:
 	virtual ~IKeyArea() {}
 
@@ -50,6 +68,10 @@ public:
 	 * @param KeyHandle The key to delete
 	 */
 	virtual void DeleteKey(FKeyHandle KeyHandle) = 0;
+
+	/** Gets any curve info associated with this key area.  This can be null but it must be valid in order to be
+	  * edited by the curve editor. */
+	virtual FCurveInfo* GetCurveInfo() = 0;
 };
 
 
diff --git a/Engine/Source/Editor/UMGEditor/Private/Animation/MarginTrackEditor.cpp b/Engine/Source/Editor/UMGEditor/Private/Animation/MarginTrackEditor.cpp
--- a/Engine/Source/Editor/UMGEditor/Private/Animation/MarginTrackEditor.cpp
+++ b/Engine/Source/Editor/UMGEditor/Private/Animation/MarginTrackEditor.cpp
@@ -20,10 +20,10 @@ public:
 	{
 		UMovieSceneMarginSection* MarginSection = Cast<UMovieSceneMarginSection>(&SectionObject);
 
-		LayoutBuilder.AddKeyArea("Left", NSLOCTEXT("FMarginPropertySection", "MarginLeft", "Left"), MakeShareable(new FFloatCurveKeyArea(MarginSection->GetLeftCurve())));
-		LayoutBuilder.AddKeyArea("Top", NSLOCTEXT("FMarginPropertySection", "MarginTop", "Top"), MakeShareable(new FFloatCurveKeyArea(MarginSection->GetTopCurve())));
-		LayoutBuilder.AddKeyArea("Right", NSLOCTEXT("FMarginPropertySection", "MarginRight", "Right"), MakeShareable(new FFloatCurveKeyArea(MarginSection->GetRightCurve())));
-		LayoutBuilder.AddKeyArea("Bottom", NSLOCTEXT("FMarginPropertySection", "MarginBottom", "Bottom"), MakeShareable(new FFloatCurveKeyArea(MarginSection->GetBottomCurve())));
+		LayoutBuilder.AddKeyArea("Left", NSLOCTEXT("FMarginPropertySection", "MarginLeft", "Left"), MakeShareable(new FFloatCurveKeyArea(&MarginSection->GetLeftCurve(), MarginSection)));
+		LayoutBuilder.AddKeyArea("Top", NSLOCTEXT("FMarginPropertySection", "MarginTop", "Top"), MakeShareable(new FFloatCurveKeyArea(&MarginSection->GetTopCurve(), MarginSection)));
+		LayoutBuilder.AddKeyArea("Right", NSLOCTEXT("FMarginPropertySection", "MarginRight", "Right"), MakeShareable(new FFloatCurveKeyArea(&MarginSection->GetRightCurve(), MarginSection)));
+		LayoutBuilder.AddKeyArea("Bottom", NSLOCTEXT("FMarginPropertySection", "MarginBottom", "Bottom"), MakeShareable(new FFloatCurveKeyArea(&MarginSection->GetBottomCurve(), MarginSection)));
 	}
 };
 
diff --git a/Engine/Source/Editor/UMGEditor/Private/Animation/Sequencer2DTransformTrackEditor.cpp b/Engine/Source/Editor/UMGEditor/Private/Animation/Sequencer2DTransformTrackEditor.cpp
--- a/Engine/Source/Editor/UMGEditor/Private/Animation/Sequencer2DTransformTrackEditor.cpp
+++ b/Engine/Source/Editor/UMGEditor/Private/Animation/Sequencer2DTransformTrackEditor.cpp
@@ -23,22 +23,22 @@ public:
 
 		// This generates the tree structure for the transform section
 		LayoutBuilder.PushCategory("Location", NSLOCTEXT("F2DTransformSection", "LocationArea", "Location"));
-			LayoutBuilder.AddKeyArea("Location.X", NSLOCTEXT("F2DTransformSection", "LocXArea", "X"), MakeShareable(new FFloatCurveKeyArea(TransformSection->GetTranslationCurve(EAxis::X))));
-			LayoutBuilder.AddKeyArea("Location.Y", NSLOCTEXT("F2DTransformSection", "LocYArea", "Y"), MakeShareable(new FFloatCurveKeyArea(TransformSection->GetTranslationCurve(EAxis::Y))));
+			LayoutBuilder.AddKeyArea("Location.X", NSLOCTEXT("F2DTransformSection", "LocXArea", "X"), MakeShareable(new FFloatCurveKeyArea(&TransformSection->GetTranslationCurve(EAxis::X), TransformSection)));
+			LayoutBuilder.AddKeyArea("Location.Y", NSLOCTEXT("F2DTransformSection", "LocYArea", "Y"), MakeShareable(new FFloatCurveKeyArea(&TransformSection->GetTranslationCurve(EAxis::Y), TransformSection)));
 		LayoutBuilder.PopCategory();
 
 		LayoutBuilder.PushCategory("Rotation", NSLOCTEXT("F2DTransformSection", "RotationArea", "Rotation"));
-			LayoutBuilder.AddKeyArea("Rotation.Angle", NSLOCTEXT("F2DTransformSection", "AngleArea", "Angle"), MakeShareable(new FFloatCurveKeyArea(TransformSection->GetRotationCurve())));
+			LayoutBuilder.AddKeyArea("Rotation.Angle", NSLOCTEXT("F2DTransformSection", "AngleArea", "Angle"), MakeShareable(new FFloatCurveKeyArea(&TransformSection->GetRotationCurve(), TransformSection)));
 		LayoutBuilder.PopCategory();
 
 		LayoutBuilder.PushCategory("Scale", NSLOCTEXT("F2DTransformSection", "ScaleArea", "Scale"));
-			LayoutBuilder.AddKeyArea("Scale.X", NSLOCTEXT("F2DTransformSection", "ScaleXArea", "X"), MakeShareable(new FFloatCurveKeyArea(TransformSection->GetScaleCurve(EAxis::X))));
-			LayoutBuilder.AddKeyArea("Scale.Y", NSLOCTEXT("F2DTransformSection", "ScaleYArea", "Y"), MakeShareable(new FFloatCurveKeyArea(TransformSection->GetScaleCurve(EAxis::Y))));
+			LayoutBuilder.AddKeyArea("Scale.X", NSLOCTEXT("F2DTransformSection", "ScaleXArea", "X"), MakeShareable(new FFloatCurveKeyArea(&TransformSection->GetScaleCurve(EAxis::X), TransformSection)));
+			LayoutBuilder.AddKeyArea("Scale.Y", NSLOCTEXT("F2DTransformSection", "ScaleYArea", "Y"), MakeShareable(new FFloatCurveKeyArea(&TransformSection->GetScaleCurve(EAxis::Y), TransformSection)));
 		LayoutBuilder.PopCategory();
 
 		LayoutBuilder.PushCategory("Sheer", NSLOCTEXT("F2DTransformSection", "SheerArea", "Sheer"));
-			LayoutBuilder.AddKeyArea("Sheer.X", NSLOCTEXT("F2DTransformSection", "SheerXArea", "X"), MakeShareable(new FFloatCurveKeyArea(TransformSection->GetSheerCurve(EAxis::X))));
-			LayoutBuilder.AddKeyArea("Sheer.Y", NSLOCTEXT("F2DTransformSection", "SheerYArea", "Y"), MakeShareable(new FFloatCurveKeyArea(TransformSection->GetSheerCurve(EAxis::Y))));
+			LayoutBuilder.AddKeyArea("Sheer.X", NSLOCTEXT("F2DTransformSection", "SheerXArea", "X"), MakeShareable(new FFloatCurveKeyArea(&TransformSection->GetSheerCurve(EAxis::X), TransformSection)));
+			LayoutBuilder.AddKeyArea("Sheer.Y", NSLOCTEXT("F2DTransformSection", "SheerYArea", "Y"), MakeShareable(new FFloatCurveKeyArea(&TransformSection->GetSheerCurve(EAxis::Y), TransformSection)));
 		LayoutBuilder.PopCategory();
 	}
 };
diff --git a/Engine/Source/Editor/UnrealEd/Private/MiniCurveEditor.cpp b/Engine/Source/Editor/UnrealEd/Private/MiniCurveEditor.cpp
--- a/Engine/Source/Editor/UnrealEd/Private/MiniCurveEditor.cpp
+++ b/Engine/Source/Editor/UnrealEd/Private/MiniCurveEditor.cpp
@@ -32,8 +32,7 @@ void SMiniCurveEditor::Construct(const FArguments& InArgs)
 
 	WidgetWindow = InArgs._ParentWindow;
 
-	UObject* SelectedCurve = InArgs._CurveOwner->GetOwner();
-	FAssetEditorManager::Get().NotifyAssetOpened(SelectedCurve, this);	
+	FAssetEditorManager::Get().NotifyAssetOpened(InArgs._OwnerObject, this);
 }
 
 SMiniCurveEditor::~SMiniCurveEditor()
diff --git a/Engine/Source/Editor/UnrealEd/Private/SColorGradientEditor.cpp b/Engine/Source/Editor/UnrealEd/Private/SColorGradientEditor.cpp
--- a/Engine/Source/Editor/UnrealEd/Private/SColorGradientEditor.cpp
+++ b/Engine/Source/Editor/UnrealEd/Private/SColorGradientEditor.cpp
@@ -350,7 +350,7 @@ FReply SColorGradientEditor::OnMouseMove( const FGeometry& MyGeometry, const FPo
 					// Start a transaction, we just started dragging a stop
 					bDraggingStop = true;
 					GEditor->BeginTransaction( LOCTEXT("MoveGradientStop", "Move Gradient Stop") );
-					CurveOwner->GetOwner()->Modify();
+					CurveOwner->ModifyOwner();
 				}
 
 				return FReply::Handled();
@@ -591,7 +591,7 @@ void SColorGradientEditor::OnSelectedStopColorChanged( FLinearColor InNewColor )
 {
 	FScopedTransaction ColorChange( LOCTEXT("ChangeGradientStopColor", "Change Gradient Stop Color") );
 	SelectedStop.SetColor( InNewColor, *CurveOwner );
-	CurveOwner->GetOwner()->Modify();
+	CurveOwner->ModifyOwner();
 
 	// Set the the last edited color.  The next time a new stop is added we'll use this value
 	LastModifiedColor.R = InNewColor.R;
@@ -602,13 +602,13 @@ void SColorGradientEditor::OnSelectedStopColorChanged( FLinearColor InNewColor )
 void SColorGradientEditor::OnCancelSelectedStopColorChange( FLinearColor PreviousColor )
 {
 	SelectedStop.SetColor( PreviousColor.HSVToLinearRGB(), *CurveOwner );
-	CurveOwner->GetOwner()->Modify();
+	CurveOwner->ModifyOwner();
 }
 
 void SColorGradientEditor::OnBeginChangeAlphaValue()
 {
 	GEditor->BeginTransaction( LOCTEXT("ChangeGradientStopAlpha", "Change Gradient Stop Alpha") );
-	CurveOwner->GetOwner()->Modify();
+	CurveOwner->ModifyOwner();
 
 	bDraggingAlphaValue = true;
 }
@@ -639,7 +639,7 @@ void SColorGradientEditor::OnAlphaValueCommitted( float NewValue, ETextCommit::T
 	{
 		// Value was typed in, no transaction is active
 		FScopedTransaction ChangeAlphaTransaction( LOCTEXT("ChangeGradientStopAlpha", "Change Gradient Stop Alpha") );
-		CurveOwner->GetOwner()->Modify();
+		CurveOwner->ModifyOwner();
 
 		SelectedStop.SetColor( FLinearColor( 0,0,0, NewValue ), *CurveOwner );
 	}
@@ -664,7 +664,7 @@ void SColorGradientEditor::OnSetGradientStopTimeFromPopup( const FText& NewText,
 		float NewTime = FCString::Atof( *NewText.ToString() );
 
 		FScopedTransaction Transaction( LOCTEXT("ChangeGradientStopTime", "Change Gradient Stop Time" ) );
-		CurveOwner->GetOwner()->Modify();
+		CurveOwner->ModifyOwner();
 
 		SelectedStop.SetTime( NewTime, *CurveOwner );
 	}
@@ -819,7 +819,7 @@ void SColorGradientEditor::GetGradientStopMarks( TArray<FGradientStopMark>& OutC
 void SColorGradientEditor::DeleteStop( const FGradientStopMark& InMark )
 {
 	FScopedTransaction DeleteStopTrans( LOCTEXT("DeleteGradientStop", "Delete Gradient Stop") );
-	CurveOwner->GetOwner()->Modify();
+	CurveOwner->ModifyOwner();
 
 	TArray<FRichCurveEditInfo> Curves = CurveOwner->GetCurves();
 
@@ -879,7 +879,7 @@ FGradientStopMark SColorGradientEditor::AddStop( const FVector2D& Position, cons
 void SColorGradientEditor::MoveStop( FGradientStopMark& Mark, float NewTime )
 {
 	Mark.SetTime( NewTime, *CurveOwner );
-	CurveOwner->GetOwner()->Modify();
+	CurveOwner->ModifyOwner();
 }
 
 
diff --git a/Engine/Source/Editor/UnrealEd/Private/SCurveEditor.cpp b/Engine/Source/Editor/UnrealEd/Private/SCurveEditor.cpp
--- a/Engine/Source/Editor/UnrealEd/Private/SCurveEditor.cpp
+++ b/Engine/Source/Editor/UnrealEd/Private/SCurveEditor.cpp
@@ -48,6 +48,8 @@ void SCurveEditor::Construct(const FArguments& InArgs)
 	bZoomToFitHorizontal = InArgs._ZoomToFitHorizontal;
 	DesiredSize = InArgs._DesiredSize;
 
+	GridColor = InArgs._GridColor;
+
 	bIsUsingSlider = false;
 
 	// if editor size is set, use it, otherwise, use default value
@@ -64,6 +66,10 @@ void SCurveEditor::Construct(const FArguments& InArgs)
 	bAllowZoomOutput = InArgs._AllowZoomOutput;
 	bAlwaysDisplayColorCurves = InArgs._AlwaysDisplayColorCurves;
 	bShowZoomButtons = InArgs._ShowZoomButtons;
+	bShowCurveSelector = InArgs._ShowCurveSelector;
+	bDrawInputGridNumbers = InArgs._ShowInputGridNumbers;
+	bDrawOutputGridNumbers = InArgs._ShowOutputGridNumbers;
+	bShowCurveToolTips = InArgs._ShowCurveToolTips;
 
 	OnCreateAsset = InArgs._OnCreateAsset;
 
@@ -129,13 +135,16 @@ void SCurveEditor::Construct(const FArguments& InArgs)
 
 	TSharedRef<SBox> CurveSelector = SNew(SBox)
 		.VAlign(VAlign_Top)
-		.Visibility(this, &SCurveEditor::GetControlVisibility)
+		.Visibility(this, &SCurveEditor::GetCurveSelectorVisibility)
 		[
 			CreateCurveSelectionWidget()
 		];
 
 	CurveSelectionWidget = CurveSelector;
 
+	InputAxisName = InArgs._XAxisName.IsSet() ? FText::FromString(InArgs._XAxisName.GetValue()) : LOCTEXT("Time", "Time");
+	InputAxisName = InArgs._YAxisName.IsSet() ? FText::FromString(InArgs._YAxisName.GetValue()) : LOCTEXT("Value", "Value");
+
 	this->ChildSlot
 	[
 		SNew( SHorizontalBox )
@@ -166,7 +175,6 @@ void SCurveEditor::Construct(const FArguments& InArgs)
 				.HAlign(HAlign_Left)
 				.BorderImage( FEditorStyle::GetBrush("NoBorder") )
 				.DesiredSizeScale(FVector2D(256.0f,32.0f))
-				.Visibility(this, &SCurveEditor::GetControlVisibility)
 				.Padding(FMargin(2, 12, 0, 0))
 				[
 					SNew(SHorizontalBox)
@@ -229,7 +237,7 @@ void SCurveEditor::Construct(const FArguments& InArgs)
 							.Label()
 							[
 								SNew(STextBlock)
-								.Text(InArgs._XAxisName.IsSet() ? FText::FromString(InArgs._XAxisName.GetValue()) : LOCTEXT("Time", "Time"))
+								.Text(InputAxisName)
 							]
 						]
 					]
@@ -261,7 +269,7 @@ void SCurveEditor::Construct(const FArguments& InArgs)
 							.Label()
 							[
 								SNew(STextBlock)
-								.Text(InArgs._YAxisName.IsSet() ? FText::FromString(InArgs._YAxisName.GetValue()) : LOCTEXT("Value", "Value"))
+								.Text(InputAxisName)
 							]
 						]
 					]
@@ -361,6 +369,21 @@ void SCurveEditor::RemoveCurveKeysFromSelection(TSharedPtr<FCurveViewModel> Curv
 	}
 }
 
+FText SCurveEditor::GetCurveToolTipNameText() const
+{
+	return CurveToolTipNameText;
+}
+
+FText SCurveEditor::GetCurveToolTipInputText() const
+{
+	return CurveToolTipInputText;
+}
+
+FText SCurveEditor::GetCurveToolTipOutputText() const
+{
+	return CurveToolTipOutputText;
+}
+
 SCurveEditor::~SCurveEditor()
 {
 	if (GEditor != NULL)
@@ -498,14 +521,14 @@ EVisibility SCurveEditor::GetCurveAreaVisibility() const
 	return AreCurvesVisible() ? EVisibility::Visible : EVisibility::Collapsed;
 }
 
-EVisibility SCurveEditor::GetControlVisibility() const
+EVisibility SCurveEditor::GetCurveSelectorVisibility() const
 {
-	return (IsHovered() || (false == bHideUI)) ? EVisibility::Visible : EVisibility::Hidden;
+	return  (IsHovered() || (false == bHideUI)) && bShowCurveSelector ? EVisibility::Visible : EVisibility::Collapsed;
 }
 
 EVisibility SCurveEditor::GetEditVisibility() const
 {
-	return (SelectedKeys.Num() > 0) ? EVisibility::Visible : EVisibility::Collapsed;
+	return (SelectedKeys.Num() > 0) && (IsHovered() || (false == bHideUI)) ? EVisibility::Visible : EVisibility::Collapsed;
 }
 
 EVisibility SCurveEditor::GetColorGradientVisibility() const
@@ -515,7 +538,7 @@ EVisibility SCurveEditor::GetColorGradientVisibility() const
 
 EVisibility SCurveEditor::GetZoomButtonVisibility() const
 {
-	return bShowZoomButtons ? EVisibility::Visible : EVisibility::Collapsed;
+	return (IsHovered() || (false == bHideUI)) && bShowZoomButtons ? EVisibility::Visible : EVisibility::Collapsed;
 }
 
 bool SCurveEditor::GetInputEditEnabled() const
@@ -560,16 +583,18 @@ int32 SCurveEditor::OnPaint(const FPaintArgs& Args, const FGeometry& AllottedGeo
 
 	// time=0 line
 	int32 ZeroLineLayerId = GridLineLayerId + 1;
-		FSlateDrawElement::MakeBox
-		(
-			OutDrawElements,
+	TArray<FVector2D> ZeroLinePoints;
+	ZeroLinePoints.Add( FVector2D( ZeroInputX, 0 ) );
+	ZeroLinePoints.Add( FVector2D( ZeroInputX, CurveAreaGeometry.Size.Y ) );
+	FSlateDrawElement::MakeLines(
+		OutDrawElements,
 		ZeroLineLayerId,
-		CurveAreaGeometry.ToPaintGeometry(FVector2D(ZeroInputX, 0), FVector2D(1, CurveAreaGeometry.Size.Y)),
-			WhiteBrush,
-			MyClippingRect,
-			DrawEffects,
-		WhiteBrush->GetTint(InWidgetStyle) * InWidgetStyle.GetColorAndOpacityTint()
-		);
+		AllottedGeometry.ToPaintGeometry(),
+		ZeroLinePoints,
+		MyClippingRect,
+		DrawEffects,
+		FLinearColor::White,
+		false );
 
 	// value=0 line
 	if( AreCurvesVisible() )
@@ -868,7 +893,6 @@ void SCurveEditor::PaintGridLines(const FGeometry &AllottedGeometry, FTrackScale
 
 	const float GridPixelSpacing = FMath::Min(ScaleInfo.WidgetSize.GetMin()/1.5f, MaxGridPixelSpacing);
 
-	const FLinearColor GridColor = FLinearColor(0.0f,0.0f,0.0f, 0.30f) ;
 	const FLinearColor GridTextColor = FLinearColor(1.0f,1.0f,1.0f, 0.75f) ;
 
 	//Vertical grid(time)
@@ -891,9 +915,18 @@ void SCurveEditor::PaintGridLines(const FGeometry &AllottedGeometry, FTrackScale
 				{
 					LinePoints.Add(FVector2D(X, 0.0));
 					LinePoints.Add(FVector2D(X, AllottedGeometry.Size.Y));
-					FSlateDrawElement::MakeLines(OutDrawElements,LayerId,AllottedGeometry.ToPaintGeometry(),LinePoints,MyClippingRect,DrawEffects,GridColor);
+					FSlateDrawElement::MakeLines(
+						OutDrawElements,
+						LayerId,
+						AllottedGeometry.ToPaintGeometry(),
+						LinePoints,
+						MyClippingRect,
+						DrawEffects,
+						GridColor,
+						false);
 
 					//Show grid time
+					if (bDrawInputGridNumbers)
 					{
 						FString TimeStr = FString::Printf(TEXT("%.2f"), Time);
 						FSlateDrawElement::MakeText(OutDrawElements,LayerId,AllottedGeometry.MakeChild(FVector2D(X, 0.0), FVector2D(1.0f, ScaleX )).ToPaintGeometry(),TimeStr,
@@ -928,9 +961,18 @@ void SCurveEditor::PaintGridLines(const FGeometry &AllottedGeometry, FTrackScale
 				{
 					LinePoints.Add(FVector2D(0.0f, Y));
 					LinePoints.Add(FVector2D(AllottedGeometry.Size.X,Y));
-					FSlateDrawElement::MakeLines(OutDrawElements,LayerId,AllottedGeometry.ToPaintGeometry(),LinePoints,MyClippingRect,DrawEffects,GridColor);
+					FSlateDrawElement::MakeLines(
+						OutDrawElements,
+						LayerId,
+						AllottedGeometry.ToPaintGeometry(),
+						LinePoints,
+						MyClippingRect,
+						DrawEffects,
+						GridColor,
+						false);
 
 					//Show grid value
+					if (bDrawOutputGridNumbers)
 					{
 						FString ValueStr = FString::Printf(TEXT("%.2f"), Value);
 						FSlateFontInfo Font = FEditorStyle::GetFontStyle("CurveEd.InfoFont");
@@ -1040,22 +1082,33 @@ FRichCurve* SCurveEditor::GetCurve(int32 CurveIndex) const
 
 void SCurveEditor::DeleteSelectedKeys()
 {
+	const FScopedTransaction Transaction(LOCTEXT("CurveEditor_RemoveKeys", "Delete Key(s)"));
+	CurveOwner->ModifyOwner();
+	TSet<FRichCurve*> ChangedCurves;
+
 	// While there are still keys
 	while(SelectedKeys.Num() > 0)
 	{
-		const FScopedTransaction Transaction( LOCTEXT("CurveEditor_RemoveKeys", "Delete Key(s)") );
-		CurveOwner->ModifyOwner();
-
 		// Pull one out of the selected set
 		FSelectedCurveKey Key = SelectedKeys.Pop();
 		if(IsValidCurve(Key.Curve))
 		{
 			// Remove from the curve
 			Key.Curve->DeleteKey(Key.KeyHandle);
+			ChangedCurves.Add(Key.Curve);
 		}
 	}
 
-	CurveOwner->OnCurveChanged();
+	TArray<FRichCurveEditInfo> ChangedCurveEditInfos;
+	for (auto CurveViewModel : CurveViewModels)
+	{
+		if (ChangedCurves.Contains(CurveViewModel->CurveInfo.CurveToEdit))
+		{
+			ChangedCurveEditInfos.Add(CurveViewModel->CurveInfo);
+		}
+	}
+
+	CurveOwner->OnCurveChanged(ChangedCurveEditInfos);
 }
 
 FReply SCurveEditor::OnMouseButtonDown( const FGeometry& InMyGeometry, const FPointerEvent& InMouseEvent)
@@ -1083,39 +1136,36 @@ FReply SCurveEditor::OnMouseButtonDown( const FGeometry& InMyGeometry, const FPo
 
 void SCurveEditor::AddNewKey(FGeometry InMyGeometry, FVector2D ScreenPosition)
 {
-	bool bKeyAdded = false;
+	const FScopedTransaction Transaction(LOCTEXT("CurveEditor_AddKey", "Add Key(s)"));
+	CurveOwner->ModifyOwner();
+	TArray<FRichCurveEditInfo> ChangedCurveEditInfos;
 	for (auto CurveViewModel : CurveViewModels)
-		{
-		if (!CurveViewModel->bIsLocked)
 	{
-			FRichCurve* SelectedCurve = CurveViewModel->CurveInfo.CurveToEdit;
-		if ( IsValidCurve(SelectedCurve) )
+		if (!CurveViewModel->bIsLocked)
 		{
-			const FScopedTransaction Transaction(LOCTEXT("CurveEditor_AddKey", "Add Key(s)"));
-			CurveOwner->ModifyOwner();
-
-			FTrackScaleInfo ScaleInfo(ViewMinInput.Get(), ViewMaxInput.Get(), ViewMinOutput.Get(), ViewMaxOutput.Get(), InMyGeometry.Size);
+			FRichCurve* SelectedCurve = CurveViewModel->CurveInfo.CurveToEdit;
+			if (IsValidCurve(SelectedCurve))
+			{
+				FTrackScaleInfo ScaleInfo(ViewMinInput.Get(), ViewMaxInput.Get(), ViewMinOutput.Get(), ViewMaxOutput.Get(), InMyGeometry.Size);
 
-			FVector2D LocalClickPos = InMyGeometry.AbsoluteToLocal(ScreenPosition);
+				FVector2D LocalClickPos = InMyGeometry.AbsoluteToLocal(ScreenPosition);
 
-				FVector2D NewKeyLocation = SnapLocation(FVector2D
-				(
+				FVector2D NewKeyLocation = SnapLocation(FVector2D(
 					ScaleInfo.LocalXToInput(LocalClickPos.X),
-					ScaleInfo.LocalYToOutput(LocalClickPos.Y))
-				);
+					ScaleInfo.LocalYToOutput(LocalClickPos.Y)));
 				FKeyHandle NewKeyHandle = SelectedCurve->AddKey(NewKeyLocation.X, NewKeyLocation.Y);
 
-			EmptySelection();
-			AddToSelection(FSelectedCurveKey(SelectedCurve, NewKeyHandle));
-				bKeyAdded = true;
+				EmptySelection();
+				AddToSelection(FSelectedCurveKey(SelectedCurve, NewKeyHandle));
+				ChangedCurveEditInfos.Add(CurveViewModel->CurveInfo);
 			}
 		}
 	}
 
-	if (bKeyAdded)
+	if (ChangedCurveEditInfos.Num() > 0)
 	{
-	CurveOwner->OnCurveChanged();
-}
+		CurveOwner->OnCurveChanged(ChangedCurveEditInfos);
+	}
 }
 
 void SCurveEditor::OnMouseCaptureLost()
@@ -1167,6 +1217,8 @@ void ClampViewRangeToDataIfBound( float& NewViewMin, float& NewViewMax, const TA
 
 FReply SCurveEditor::OnMouseMove( const FGeometry& InMyGeometry, const FPointerEvent& InMouseEvent )
 {
+	UpdateCurveToolTip(InMyGeometry, InMouseEvent);
+
 	FRichCurve* Curve = GetCurve(0);
 	if( Curve != NULL && this->HasMouseCapture())
 	{
@@ -1184,7 +1236,64 @@ FReply SCurveEditor::OnMouseMove( const FGeometry& InMyGeometry, const FPointerE
 	return FReply::Unhandled();
 }
 
-FReply SCurveEditor::OnMouseWheel( const FGeometry& MyGeometry, const FPointerEvent& MouseEvent )
+void SCurveEditor::UpdateCurveToolTip(const FGeometry& InMyGeometry, const FPointerEvent& InMouseEvent)
+{
+	if (bShowCurveToolTips.Get())
+	{
+		TSharedPtr<FCurveViewModel> HoveredCurve = HitTestCurves(InMyGeometry, InMouseEvent);
+		if (HoveredCurve.IsValid())
+		{
+			FTrackScaleInfo ScaleInfo(ViewMinInput.Get(), ViewMaxInput.Get(), ViewMinOutput.Get(), ViewMaxOutput.Get(), InMyGeometry.Size);
+			const FVector2D HitPosition = InMyGeometry.AbsoluteToLocal(InMouseEvent.GetScreenSpacePosition());
+			float Time = ScaleInfo.LocalXToInput(HitPosition.X);
+			float Value = HoveredCurve->CurveInfo.CurveToEdit->Eval(Time);
+
+			FNumberFormattingOptions FormattingOptions;
+			FormattingOptions.MaximumFractionalDigits = 2;
+			CurveToolTipNameText = FText::FromName(HoveredCurve->CurveInfo.CurveName);
+			CurveToolTipInputText = FText::Format(LOCTEXT("CurveToolTipTimeFormat", "{0}: {1}"), InputAxisName, FText::AsNumber(Time, &FormattingOptions));
+			CurveToolTipOutputText = FText::Format(LOCTEXT("CurveToolTipValueFormat", "{0}: {1}"), OutputAxisName, FText::AsNumber(Value, &FormattingOptions));
+
+			if (CurveToolTip.IsValid() == false)
+			{
+				SetToolTip(
+					SAssignNew(CurveToolTip, SToolTip)
+					.BorderImage( FCoreStyle::Get().GetBrush( "ToolTip.BrightBackground" ) )
+					[
+						SNew(SVerticalBox)
+						+ SVerticalBox::Slot()
+						[
+							SNew(STextBlock)
+							.Text(this, &SCurveEditor::GetCurveToolTipNameText)
+							.Font(FCoreStyle::Get().GetFontStyle("ToolTip.LargerFont"))
+							.ColorAndOpacity( FLinearColor::Black)
+						]
+						+ SVerticalBox::Slot()
+						[
+							SNew(STextBlock)
+							.Text(this, &SCurveEditor::GetCurveToolTipInputText)
+							.Font(FCoreStyle::Get().GetFontStyle("ToolTip.LargerFont"))
+							.ColorAndOpacity(FLinearColor::Black)
+						]
+						+ SVerticalBox::Slot()
+						[
+							SNew(STextBlock)
+							.Text(this, &SCurveEditor::GetCurveToolTipOutputText)
+							.Font(FCoreStyle::Get().GetFontStyle("ToolTip.LargerFont"))
+							.ColorAndOpacity(FLinearColor::Black)
+						]
+					]);
+			}
+		}
+		else
+		{
+			CurveToolTip.Reset();
+			SetToolTip(CurveToolTip);
+		}
+	}
+}
+
+FReply SCurveEditor::OnMouseWheel(const FGeometry& MyGeometry, const FPointerEvent& MouseEvent)
 {
 	const float ZoomDelta = -0.1f * MouseEvent.GetWheelDelta();
 
@@ -1469,8 +1578,10 @@ void SCurveEditor::OnTimeComitted(float NewTime, ETextCommit::Type CommitType)
 				const FScopedTransaction Transaction(LOCTEXT("CurveEditor_NewTime", "New Time Entered"));
 				CurveOwner->ModifyOwner();
 				Key.Curve->SetKeyTime(Key.KeyHandle, NewTime);
+				TArray<FRichCurveEditInfo> ChangedCurveEditInfos;
+				ChangedCurveEditInfos.Add(GetViewModelForCurve(Key.Curve)->CurveInfo);
+				CurveOwner->OnCurveChanged(ChangedCurveEditInfos);
 			}
-			CurveOwner->OnCurveChanged();
 		}
 
 		FSlateApplication::Get().DismissAllMenus();
@@ -1486,10 +1597,13 @@ void SCurveEditor::OnTimeChanged(float NewTime)
 			auto Key = SelectedKeys[0];
 			if ( IsValidCurve(Key.Curve) )
 			{
+				const FScopedTransaction Transaction( LOCTEXT( "CurveEditor_NewTime", "New Time Entered" ) );
 				CurveOwner->ModifyOwner();
 				Key.Curve->SetKeyTime(Key.KeyHandle, NewTime);
+				TArray<FRichCurveEditInfo> ChangedCurveEditInfos;
+				ChangedCurveEditInfos.Add(GetViewModelForCurve(Key.Curve)->CurveInfo);
+				CurveOwner->OnCurveChanged(ChangedCurveEditInfos);
 			}
-			CurveOwner->OnCurveChanged();
 		}
 	}
 }
@@ -1521,21 +1635,31 @@ void SCurveEditor::OnValueComitted(float NewValue, ETextCommit::Type CommitType)
 	// Don't digest the number if we just clicked away from the popup
 	if ( !bIsUsingSlider && ((CommitType == ETextCommit::OnEnter) || ( CommitType == ETextCommit::OnUserMovedFocus )) )
 	{
+		const FScopedTransaction Transaction( LOCTEXT( "CurveEditor_NewValue", "New Value Entered" ) );
+		CurveOwner->ModifyOwner();
+		TSet<FRichCurve*> ChangedCurves;
+
 		// Iterate over selected set
 		for ( int32 i=0; i < SelectedKeys.Num(); i++ )
 		{
 			auto Key = SelectedKeys[i];
 			if ( IsValidCurve(Key.Curve) )
 			{
-				const FScopedTransaction Transaction(LOCTEXT("CurveEditor_NewValue", "New Value Entered"));
-				CurveOwner->ModifyOwner();
-
 				// Fill in each element of this key
 				Key.Curve->SetKeyValue(Key.KeyHandle, NewValue);
+				ChangedCurves.Add(Key.Curve);
 			}
 		}
 
-		CurveOwner->OnCurveChanged();
+		TArray<FRichCurveEditInfo> ChangedCurveEditInfos;
+		for (auto CurveViewModel : CurveViewModels)
+		{
+			if (ChangedCurves.Contains(CurveViewModel->CurveInfo.CurveToEdit))
+			{
+				ChangedCurveEditInfos.Add(CurveViewModel->CurveInfo);
+			}
+		}
+		CurveOwner->OnCurveChanged(ChangedCurveEditInfos);
 
 		FSlateApplication::Get().DismissAllMenus();
 	}
@@ -1546,6 +1670,7 @@ void SCurveEditor::OnValueChanged(float NewValue)
 	if ( bIsUsingSlider )
 	{
 		const FScopedTransaction Transaction(LOCTEXT("CurveEditor_NewValue", "New Value Entered"));
+		TSet<FRichCurve*> ChangedCurves;
 
 		// Iterate over selected set
 		for ( int32 i=0; i < SelectedKeys.Num(); i++ )
@@ -1557,10 +1682,19 @@ void SCurveEditor::OnValueChanged(float NewValue)
 
 				// Fill in each element of this key
 				Key.Curve->SetKeyValue(Key.KeyHandle, NewValue);
+				ChangedCurves.Add(Key.Curve);
 			}
 		}
 
-		CurveOwner->OnCurveChanged();
+		TArray<FRichCurveEditInfo> ChangedCurveEditInfos;
+		for (auto CurveViewModel : CurveViewModels)
+		{
+			if (ChangedCurves.Contains(CurveViewModel->CurveInfo.CurveToEdit))
+			{
+				ChangedCurveEditInfos.Add(CurveViewModel->CurveInfo);
+			}
+		}
+		CurveOwner->OnCurveChanged(ChangedCurveEditInfos);
 	}
 }
 
@@ -1873,12 +2007,15 @@ void SCurveEditor::ZoomToFitAll(bool OnlySelected)
 
 void SCurveEditor::ToggleSnapping()
 {
-	bSnappingEnabled = !bSnappingEnabled;
+	if (bSnappingEnabled.IsBound() == false)
+	{
+		bSnappingEnabled = !bSnappingEnabled.Get();
+	}
 }
 
 bool SCurveEditor::IsSnappingEnabled()
 {
-	return bSnappingEnabled;
+	return bSnappingEnabled.Get();
 }
 
 void SCurveEditor::CreateContextMenu(const FGeometry& InMyGeometry, const FVector2D& ScreenPosition)
@@ -2122,6 +2259,8 @@ void SCurveEditor::OnSelectInterpolationMode(ERichCurveInterpMode InterpMode, ER
 	{
 		const FScopedTransaction Transaction(LOCTEXT("CurveEditor_SetInterpolationMode", "Select Interpolation Mode"));
 		CurveOwner->ModifyOwner();
+		TSet<FRichCurve*> ChangedCurves;
+
 		for(auto It = SelectedKeys.CreateIterator();It;++It)
 		{
 			FSelectedCurveKey& Key = *It;
@@ -2129,7 +2268,16 @@ void SCurveEditor::OnSelectInterpolationMode(ERichCurveInterpMode InterpMode, ER
 			Key.Curve->SetKeyInterpMode(Key.KeyHandle,InterpMode );
 			Key.Curve->SetKeyTangentMode(Key.KeyHandle,TangentMode );
 		}
-		CurveOwner->OnCurveChanged();
+
+		TArray<FRichCurveEditInfo> ChangedCurveEditInfos;
+		for (auto CurveViewModel : CurveViewModels)
+		{
+			if (ChangedCurves.Contains(CurveViewModel->CurveInfo.CurveToEdit))
+			{
+				ChangedCurveEditInfos.Add(CurveViewModel->CurveInfo);
+			}
+		}
+		CurveOwner->OnCurveChanged(ChangedCurveEditInfos);
 	}
 }
 
@@ -2275,11 +2423,15 @@ void SCurveEditor::EndDragTransaction()
 {
 	if ( TransactionIndex >= 0 )
 	{
+		TArray<FRichCurveEditInfo> ChangedCurveEditInfos;
+		for (auto CurveViewModel : CurveViewModels)
+		{
+			ChangedCurveEditInfos.Add(CurveViewModel->CurveInfo);
+		}
+		CurveOwner->OnCurveChanged(ChangedCurveEditInfos);
 		GEditor->EndTransaction();
 		TransactionIndex = -1;
 	}
-
-	CurveOwner->OnCurveChanged();
 }
 
 bool SCurveEditor::FSelectedTangent::IsValid() const
@@ -2318,7 +2470,7 @@ bool SCurveEditor::IsLinearColorCurve() const
 
 FVector2D SCurveEditor::SnapLocation(FVector2D InLocation)
 {
-	if (bSnappingEnabled)
+	if (bSnappingEnabled.Get())
 	{
 		InLocation.X = InputSnap != 0 ? FMath::RoundToInt(InLocation.X / InputSnap.Get()) * InputSnap.Get() : InLocation.X;
 		InLocation.Y = OutputSnap != 0 ? FMath::RoundToInt(InLocation.Y / OutputSnap.Get()) * OutputSnap.Get() : InLocation.Y;
@@ -2326,4 +2478,16 @@ FVector2D SCurveEditor::SnapLocation(FVector2D InLocation)
 	return InLocation;
 }
 
+TSharedPtr<FCurveViewModel> SCurveEditor::GetViewModelForCurve(FRichCurve* InCurve)
+{
+	for (auto CurveViewModel : CurveViewModels)
+	{
+		if (InCurve == CurveViewModel->CurveInfo.CurveToEdit)
+		{
+			return CurveViewModel;
+		}
+	}
+	return TSharedPtr<FCurveViewModel>();
+}
+
 #undef LOCTEXT_NAMESPACE
diff --git a/Engine/Source/Editor/UnrealEd/Public/MiniCurveEditor.h b/Engine/Source/Editor/UnrealEd/Public/MiniCurveEditor.h
--- a/Engine/Source/Editor/UnrealEd/Public/MiniCurveEditor.h
+++ b/Engine/Source/Editor/UnrealEd/Public/MiniCurveEditor.h
@@ -9,10 +9,12 @@ class UNREALED_API SMiniCurveEditor :  public SCompoundWidget,public IAssetEdito
 {
 public:
 	SLATE_BEGIN_ARGS( SMiniCurveEditor )
-		: _CurveOwner(NULL)
+		: _CurveOwner(nullptr)
+		, _OwnerObject(nullptr)
 		{}
 
 	SLATE_ARGUMENT( FCurveOwnerInterface*, CurveOwner )
+	SLATE_ARGUMENT( UObject*, OwnerObject )
 	SLATE_ARGUMENT( TWeakPtr<SWindow>, ParentWindow )
 	SLATE_END_ARGS()
 
diff --git a/Engine/Source/Editor/UnrealEd/Public/SCurveEditor.h b/Engine/Source/Editor/UnrealEd/Public/SCurveEditor.h
--- a/Engine/Source/Editor/UnrealEd/Public/SCurveEditor.h
+++ b/Engine/Source/Editor/UnrealEd/Public/SCurveEditor.h
@@ -145,6 +145,11 @@ public:
 		, _ShowZoomButtons(true)
 		, _XAxisName()
 		, _YAxisName()
+		, _ShowInputGridNumbers(true)
+		, _ShowOutputGridNumbers(true)
+		, _ShowCurveSelector(true)
+		, _ShowCurveToolTips(false)
+		, _GridColor(FLinearColor(0.0f, 0.0f, 0.0f, 0.3f))
 		{}
 
 		SLATE_ATTRIBUTE( float, ViewMinInput )
@@ -155,9 +160,10 @@ public:
 		SLATE_ATTRIBUTE( float, ViewMaxOutput )
 		SLATE_ATTRIBUTE( float, InputSnap )
 		SLATE_ATTRIBUTE( float, OutputSnap )
-		SLATE_ARGUMENT( bool, SnappingEnabled )
+		SLATE_ATTRIBUTE( bool, SnappingEnabled )
 		SLATE_ATTRIBUTE( float, TimelineLength )
 		SLATE_ATTRIBUTE( FVector2D, DesiredSize )
+		SLATE_ATTRIBUTE( bool, ShowCurveToolTips )
 		SLATE_ARGUMENT( bool, DrawCurve )
 		SLATE_ARGUMENT( bool, HideUI )
 		SLATE_ARGUMENT( bool, AllowZoomOutput )
@@ -167,6 +173,10 @@ public:
 		SLATE_ARGUMENT( bool, ShowZoomButtons )
 		SLATE_ARGUMENT( TOptional<FString>, XAxisName )
 		SLATE_ARGUMENT( TOptional<FString>, YAxisName )
+		SLATE_ARGUMENT( bool, ShowInputGridNumbers )
+		SLATE_ARGUMENT( bool, ShowOutputGridNumbers )
+		SLATE_ARGUMENT( bool, ShowCurveSelector )
+		SLATE_ARGUMENT( FLinearColor, GridColor )
 		SLATE_EVENT( FOnSetInputViewRange, OnSetInputViewRange )
 		SLATE_EVENT( FOnSetOutputViewRange, OnSetOutputViewRange )
 		SLATE_EVENT( FSimpleDelegate, OnCreateAsset )
@@ -334,7 +344,7 @@ private:
 	void OnEndSliderMovement(float NewValue);
 
 	EVisibility GetCurveAreaVisibility() const;
-	EVisibility GetControlVisibility() const;
+	EVisibility GetCurveSelectorVisibility() const;
 	EVisibility GetEditVisibility() const;
 	EVisibility GetColorGradientVisibility() const;
 	EVisibility GetZoomButtonVisibility() const;
@@ -461,6 +471,14 @@ private:
 
 	void RemoveCurveKeysFromSelection(TSharedPtr<FCurveViewModel> CurveViewModel);
 
+	FText GetCurveToolTipNameText() const;
+	FText GetCurveToolTipInputText() const;
+	FText GetCurveToolTipOutputText() const;
+
+	void UpdateCurveToolTip( const FGeometry& MyGeometry, const FPointerEvent& MouseEvent );
+
+	TSharedPtr<FCurveViewModel> GetViewModelForCurve(FRichCurve* InCurve);
+
 protected:
 
 	/** Set Default output values when range is too small **/
@@ -498,6 +516,11 @@ private:
 	/** If we always show the color curves or allow the user to toggle this */
 	bool				bAlwaysDisplayColorCurves;
 
+	/** Whether or not to draw the numbers for the input grid. */
+	bool bDrawInputGridNumbers;
+	/** Whether or not to draw the numbers for the output grid. */
+	bool bDrawOutputGridNumbers;
+
 	/** Array of selected keys*/
 	TArray<FSelectedCurveKey> SelectedKeys;
 
@@ -560,7 +583,10 @@ protected:
 	TAttribute<float> OutputSnap;
 
 	/** Whether or not snapping is enabled. */
-	bool bSnappingEnabled;
+	TAttribute<bool> bSnappingEnabled;
+
+	/** Whether or not to show tool tips for curves. */
+	TAttribute<bool> bShowCurveToolTips;
 
 	/** True if you want the curve editor to fit to zoom **/
 	bool				bZoomToFitVertical;
@@ -574,6 +600,9 @@ protected:
 	/** True if the internal zoom buttons should be visible. */
 	bool bShowZoomButtons;
 
+	/** Whether or not to show the curve selector widgets. */
+	bool bShowCurveSelector;
+
 	/** The location of mouse during the last OnMouseButtonDown callback in widget local coordinates. */
 	FVector2D MouseDownLocation;
 
@@ -592,7 +621,26 @@ protected:
 	/** A map of selected key handle to their starting locations at the beginning of a drag operation. */
 	TMap<FKeyHandle, FVector2D> PreDragKeyLocations;
 
+	/** The text to display for the input axis. */
+	FText InputAxisName;
+	/** The text to display for the output axis. */
+	FText OutputAxisName;
+
+	/** The view models for the curves. */
 	TArray<TSharedPtr<FCurveViewModel>> CurveViewModels;
+
+	/** The tooltip control for the curves. */
+	TSharedPtr<SToolTip> CurveToolTip;
+
+	/** The text for the name portion of the tooltip. */
+	FText CurveToolTipNameText;
+	/** The text for the input portion of the tooltip. */
+	FText CurveToolTipInputText;
+	/** The text for the output portion of the tooltip. */
+	FText CurveToolTipOutputText;
+
+	/** The color used to draw the grid lines. */
+	FLinearColor GridColor;
 };
 
 #endif // __SCurveEditor_h__
diff --git a/Engine/Source/Runtime/Engine/Classes/Curves/CurveBase.h b/Engine/Source/Runtime/Engine/Classes/Curves/CurveBase.h
--- a/Engine/Source/Runtime/Engine/Classes/Curves/CurveBase.h
+++ b/Engine/Source/Runtime/Engine/Classes/Curves/CurveBase.h
@@ -394,8 +394,24 @@ template<class T> struct FRichCurveEditInfoTemplate
 	:	CurveName(InCurveName)
 	,	CurveToEdit(InCurveToEdit)
 	{}
+
+	inline bool operator==(const FRichCurveEditInfoTemplate<T>& Other) const
+	{
+		return Other.CurveName.IsEqual(CurveName) && Other.CurveToEdit == CurveToEdit;
+	}
+
+	uint32 GetTypeHash() const
+	{
+		return HashCombine(::GetTypeHash(CurveName), PointerHash(CurveToEdit));
+	}
 };
 
+template<class T>
+inline uint32 GetTypeHash( const FRichCurveEditInfoTemplate<T>& RichCurveEditInfo )
+{
+	return RichCurveEditInfo.GetTypeHash();
+}
+
 typedef FRichCurveEditInfoTemplate<FRichCurve*>			FRichCurveEditInfo;
 typedef FRichCurveEditInfoTemplate<const FRichCurve*>	FRichCurveEditInfoConst;
 
@@ -412,17 +428,14 @@ public:
 	/** Returns set of curves to query. Must not release the curves while being edited. */
 	virtual TArray<FRichCurveEditInfo> GetCurves() = 0;
 
-	/** Called to return the Curve UObject  **/
-	virtual UObject* GetOwner() = 0;
-
 	/** Called to modify the owner of the curve */
 	virtual void ModifyOwner() = 0;
 
 	/** Called to make curve owner transactional */
 	virtual void MakeTransactional() = 0;
 
-	/** Called when the curve has been changed */
-	virtual void OnCurveChanged() = 0;
+	/** Called when any of the curves have been changed */
+	virtual void OnCurveChanged(const TArray<FRichCurveEditInfo>& ChangedCurveEditInfos) = 0;
 
 	/** Whether the curve represents a linear color */
 	virtual bool IsLinearColorCurve() const { return false; }
@@ -474,14 +487,9 @@ public:
 		return Curves;
 	}
 
-	virtual UObject* GetOwner() override
-	{
-		return this;
-	}
-
 	virtual void ModifyOwner() override;
 	virtual void MakeTransactional() override;
-	virtual void OnCurveChanged() override;
+	virtual void OnCurveChanged(const TArray<FRichCurveEditInfo>& ChangedCurveEditInfos) override;
 
 	virtual bool IsValidCurve( FRichCurveEditInfo CurveInfo ) { return false; };
 	// End FCurveOwnerInterface
diff --git a/Engine/Source/Runtime/Engine/Private/CurveBase.cpp b/Engine/Source/Runtime/Engine/Private/CurveBase.cpp
--- a/Engine/Source/Runtime/Engine/Private/CurveBase.cpp
+++ b/Engine/Source/Runtime/Engine/Private/CurveBase.cpp
@@ -833,7 +833,7 @@ void UCurveBase::MakeTransactional()
 	SetFlags(GetFlags() | RF_Transactional);
 }
 
-void UCurveBase::OnCurveChanged()
+void UCurveBase::OnCurveChanged(const TArray<FRichCurveEditInfo>& ChangedCurveEditInfos)
 {
 
 }
diff --git a/Engine/Source/Runtime/Slate/Public/Widgets/Input/SNumericDropDown.h b/Engine/Source/Runtime/Slate/Public/Widgets/Input/SNumericDropDown.h
new file mode 100644
--- /dev/null
+++ b/Engine/Source/Runtime/Slate/Public/Widgets/Input/SNumericDropDown.h
@@ -0,0 +1,171 @@
+// Copyright 1998-2015 Epic Games, Inc. All Rights Reserved.
+
+#pragma once
+
+/** A widget which allows the user to enter a digit or choose a number from a drop down menu. */
+template<typename NumericType>
+class SNumericDropDown : public SCompoundWidget
+{
+public:
+	DECLARE_DELEGATE_OneParam( FOnValueChanged, NumericType );
+
+public:
+	/** Represents a named numeric value for display in the drop down menu. */
+	class FNamedValue
+	{
+	public:
+		/** Creates a new FNamedValue 
+		  * @InValue The numeric value to be assigned
+		  * @InName The display text for the value in the UI 
+		  * @InDescription The description of the value used in tooltips or wherever a longer description is needed. */
+		FNamedValue( NumericType InValue, FText InName, FText InDescription )
+		{
+			Value = InValue;
+			Name = InName;
+			Description = InDescription;
+		}
+
+		NumericType GetValue() const
+		{
+			return Value;
+		}
+
+		FText GetName() const
+		{
+			return Name;
+		}
+
+		FText GetDescription() const
+		{
+			return Description;
+		}
+
+	private:
+		NumericType Value;
+		FText Name;
+		FText Description;
+	};
+
+public:
+	SLATE_BEGIN_ARGS( SNumericDropDown<NumericType> )
+		: _DropDownValues()
+		, _LabelText()
+		, _Orientation( EOrientation::Orient_Horizontal )
+		, _MinDesiredValueWidth( 40 )
+		{}
+
+		/** The values which are used to populate the drop down menu. */
+		SLATE_ARGUMENT( TArray<FNamedValue>, DropDownValues )
+		/** The text which is displayed in the label next to the control. */
+		SLATE_ATTRIBUTE( FText, LabelText )
+		/** Controls the label placement for the control.  Vertical will place the label above, horizontal will place the label to the left. */
+		SLATE_ATTRIBUTE( EOrientation, Orientation )
+		/** Controls the minimum width for the text box portion of the control. */
+		SLATE_ATTRIBUTE( float, MinDesiredValueWidth )
+		/** The value displayed by the control. */
+		SLATE_ATTRIBUTE( NumericType, Value )
+		/** The callback for when the value changes. */
+		SLATE_EVENT( FOnValueChanged, OnValueChanged )
+
+	SLATE_END_ARGS()
+
+	void Construct( const FArguments& InArgs )
+	{
+		DropDownValues = InArgs._DropDownValues;
+		LabelText = InArgs._LabelText;
+		Orientation = InArgs._Orientation;
+		Value = InArgs._Value;
+		OnValueChanged = InArgs._OnValueChanged;
+
+		ChildSlot
+		[
+			SNew( SVerticalBox )
+			// Vertical Label
+			+ SVerticalBox::Slot()
+			.AutoHeight()
+			.Padding( FMargin( 0, 0, 0, 3 ) )
+			[
+				SNew( STextBlock )
+				.Text(LabelText)
+				.Visibility(this, &SNumericDropDown<NumericType>::GetLabelVisibility, EOrientation::Orient_Vertical)
+			]
+			+ SVerticalBox::Slot()
+			.AutoHeight()
+			[
+				SNew( SHorizontalBox )
+				// Horizontal Icon Label
+				+ SHorizontalBox::Slot()
+				.VAlign( EVerticalAlignment::VAlign_Center )
+				.AutoWidth()
+				.Padding( FMargin( 0, 0, 3, 0 ) )
+				[
+					
+					SNew( STextBlock )
+					.Text(LabelText)
+					.Visibility( this, &SNumericDropDown<NumericType>::GetLabelVisibility, EOrientation::Orient_Horizontal )
+				]
+				+ SHorizontalBox::Slot()
+				[
+					SNew( SComboButton )
+					.ContentPadding( 1 )
+					.OnGetMenuContent( this, &SNumericDropDown<NumericType>::BuildMenu )
+					.ButtonContent()
+					[
+						SNew( SEditableTextBox )
+						.Text( this, &SNumericDropDown<NumericType>::GetValueText )
+						.MinDesiredWidth( InArgs._MinDesiredValueWidth )
+						.OnTextCommitted( this, &SNumericDropDown<NumericType>::ValueTextComitted )
+					]
+				]
+			]
+		];
+	}
+
+private:
+	EVisibility GetLabelVisibility( EOrientation LabelOrientation ) const
+	{
+		return Orientation.Get() == LabelOrientation
+			? EVisibility::Visible
+			: EVisibility::Collapsed;
+	}
+
+	FText GetValueText() const
+	{
+		return FText::AsNumber( Value.Get() );
+	}
+
+	void ValueTextComitted( const FText& InNewText, ETextCommit::Type InTextCommit )
+	{
+		if ( InNewText.IsNumeric() )
+		{
+			NumericType NewValue;
+			TTypeFromString<NumericType>::FromString( NewValue, *InNewText.ToString() );
+			OnValueChanged.ExecuteIfBound(NewValue);
+		}
+	}
+
+	TSharedRef<SWidget> BuildMenu()
+	{
+		FMenuBuilder MenuBuilder( true, NULL );
+
+		for ( FNamedValue DropDownValue : DropDownValues )
+		{
+			FUIAction MenuAction(FExecuteAction::CreateSP(this, &SNumericDropDown::SetValue, DropDownValue.GetValue()));
+			MenuBuilder.AddMenuEntry(DropDownValue.GetName(), DropDownValue.GetDescription(), FSlateIcon(), MenuAction);
+		}
+
+		return MenuBuilder.MakeWidget();
+	}
+
+	void SetValue( NumericType InValue )
+	{
+		OnValueChanged.ExecuteIfBound(InValue);
+	}
+
+private:
+	TArray<FNamedValue> DropDownValues;
+	TAttribute<FText> LabelText;
+	TAttribute<EOrientation> Orientation;
+	TAttribute<NumericType> Value;
+	FOnValueChanged OnValueChanged;
+};
\ No newline at end of file
