Fixing ConstructObject deprecation warnings

diff --git a/Engine/Source/Editor/UnrealEd/Private/Kismet2/ComponentEditorUtils.cpp b/Engine/Source/Editor/UnrealEd/Private/Kismet2/ComponentEditorUtils.cpp
--- a/Engine/Source/Editor/UnrealEd/Private/Kismet2/ComponentEditorUtils.cpp
+++ b/Engine/Source/Editor/UnrealEd/Private/Kismet2/ComponentEditorUtils.cpp
@@ -488,7 +488,7 @@ UActorComponent* FComponentEditorUtils::DuplicateComponent(UActorComponent* Temp
 		FName NewComponentName = *FComponentEditorUtils::GenerateValidVariableName(ComponentClass, Actor);
 
 		bool bKeepWorldLocationOnAttach = false;
-		NewCloneComponent = ConstructObject<UActorComponent>(ComponentClass, Actor, NewComponentName, RF_Transactional, TemplateComponent);
+		NewCloneComponent = NewObject<UActorComponent>(Actor, ComponentClass, NewComponentName, RF_Transactional, TemplateComponent);
 		
 		auto NewSceneComponent = Cast<USceneComponent>(NewCloneComponent);
 		if (NewSceneComponent)
diff --git a/Engine/Source/Runtime/Engine/Private/BlueprintGeneratedClass.cpp b/Engine/Source/Runtime/Engine/Private/BlueprintGeneratedClass.cpp
--- a/Engine/Source/Runtime/Engine/Private/BlueprintGeneratedClass.cpp
+++ b/Engine/Source/Runtime/Engine/Private/BlueprintGeneratedClass.cpp
@@ -274,7 +274,7 @@ UInheritableComponentHandler* UBlueprintGeneratedClass::GetInheritableComponentH
 
 	if (!InheritableComponentHandler && bCreateIfNecessary)
 	{
-		InheritableComponentHandler = NewNamedObject<UInheritableComponentHandler>(this, FName(TEXT("InheritableComponentHandler")));
+		InheritableComponentHandler = NewObject<UInheritableComponentHandler>(this, FName(TEXT("InheritableComponentHandler")));
 	}
 
 	return InheritableComponentHandler;
