[INTEGRATE] Change 2440544 by [REDACTED_PERSON][REDACTED_HANDLE] on 2015/02/10 17:48:51

	Fixing (occassional?) crash on startup
	- NotifyToolsOfObjectReplacement can be called before SelectedComponents was initialized - added check for nullptr

diff --git a/Engine/Source/Editor/UnrealEd/Private/Editor.cpp b/Engine/Source/Editor/UnrealEd/Private/Editor.cpp
--- a/Engine/Source/Editor/UnrealEd/Private/Editor.cpp
+++ b/Engine/Source/Editor/UnrealEd/Private/Editor.cpp
@@ -6049,23 +6049,26 @@ void UEditorEngine::NotifyToolsOfObjectReplacement(const TMap<UObject*, UObject*
 
 	// Check to see if any selected components were reinstanced
 	USelection* ComponentSelection = GetSelectedComponents();
-	TArray<TWeakObjectPtr<UObject> > SelectedComponents;
-	ComponentSelection->GetSelectedObjects(SelectedComponents);
-
-	ComponentSelection->BeginBatchSelectOperation();
-	for (int32 i = 0; i < SelectedComponents.Num(); ++i)
+	if (ComponentSelection)
 	{
-		UObject* Component = SelectedComponents[i].GetEvenIfUnreachable();
+		TArray<TWeakObjectPtr<UObject> > SelectedComponents;
+		ComponentSelection->GetSelectedObjects(SelectedComponents);
 
-		// If the component corresponds to a new instance in the map, update the selection accordingly
-		if (OldToNewInstanceMap.Contains(Component))
+		ComponentSelection->BeginBatchSelectOperation();
+		for (int32 i = 0; i < SelectedComponents.Num(); ++i)
 		{
-			ComponentSelection->Deselect(Component);
-			SelectComponent(CastChecked<UActorComponent>(OldToNewInstanceMap[Component]), true, false);
+			UObject* Component = SelectedComponents[i].GetEvenIfUnreachable();
+
+			// If the component corresponds to a new instance in the map, update the selection accordingly
+			if (OldToNewInstanceMap.Contains(Component))
+			{
+				ComponentSelection->Deselect(Component);
+				SelectComponent(CastChecked<UActorComponent>(OldToNewInstanceMap[Component]), true, false);
+			}
 		}
+		ComponentSelection->EndBatchSelectOperation();
 	}
-	ComponentSelection->EndBatchSelectOperation();
-
+	
 	BroadcastObjectsReplaced(OldToNewInstanceMap);
 }
 
