[INTEGRATE] Change 2440563 by [REDACTED_PERSON][REDACTED_HANDLE] on 2015/02/10 17:58:19

	Fixed crash when converting an actor to a blueprint while a component was selected

diff --git a/Engine/Source/Editor/UnrealEd/Private/Kismet2/Kismet2.cpp b/Engine/Source/Editor/UnrealEd/Private/Kismet2/Kismet2.cpp
--- a/Engine/Source/Editor/UnrealEd/Private/Kismet2/Kismet2.cpp
+++ b/Engine/Source/Editor/UnrealEd/Private/Kismet2/Kismet2.cpp
@@ -1264,6 +1264,12 @@ AActor* FKismetEditorUtilities::CreateBlueprintInstanceFromSelection(UBlueprint*
 	AActor* NewActor = World->SpawnActor(Blueprint->GeneratedClass, &Location, &Rotator);
 	GEditor->Layers->InitializeNewActorLayers(NewActor);
 
+	// Quietly ensure that no components are selected
+	USelection* ComponentSelection = GEditor->GetSelectedComponents();
+	ComponentSelection->BeginBatchSelectOperation();
+	ComponentSelection->DeselectAll();
+	ComponentSelection->EndBatchSelectOperation(false);
+
 	// Update selection to new actor
 	GEditor->SelectActor( NewActor, /*bSelected=*/ true, /*bNotify=*/ true );
 
