[INTEGRATE] Change 2440738 by [REDACTED_PERSON][REDACTED_HANDLE] on 2015/02/10 19:33:20

	Blueprints: Add a way for component classes to invalidate the component type registry after they undergo changes via FComponentTypeRegistry::InvalidateClass [UE-9305]
	- This is now invoked during Blueprint compiles in case the BlueprintCategory variable has changed
	- Note: Unloaded Blueprint components will still show up in the Custom category until loaded

	Blueprints: Allow component category headers to be displayed even when the list is filtered

diff --git a/Engine/Source/Editor/KismetCompiler/Private/KismetCompiler.cpp b/Engine/Source/Editor/KismetCompiler/Private/KismetCompiler.cpp
--- a/Engine/Source/Editor/KismetCompiler/Private/KismetCompiler.cpp
+++ b/Engine/Source/Editor/KismetCompiler/Private/KismetCompiler.cpp
@@ -12,6 +12,7 @@
 #include "Editor/UnrealEd/Public/Kismet2/BlueprintEditorUtils.h"
 #include "Editor/UnrealEd/Public/Kismet2/KismetEditorUtilities.h"
 #include "Editor/UnrealEd/Public/ScriptDisassembler.h"
+#include "Editor/UnrealEd/Public/ComponentTypeRegistry.h"
 #include "K2Node_PlayMovieScene.h"
 #include "RuntimeMovieScenePlayer.h"
 #include "MovieSceneBindings.h"
@@ -1730,6 +1731,8 @@ void FKismetCompilerContext::FinishCompilingClass(UClass* Class)
 			}
 
 			Class->SetMetaData(NAME_ClassGroupNames, *ClassGroupCategory);
+
+			FComponentTypeRegistry::Get().InvalidateClass(Class);
 		}
 
 		// Add a category if one has been specified
diff --git a/Engine/Source/Editor/UnrealEd/Private/ComponentTypeRegistry.cpp b/Engine/Source/Editor/UnrealEd/Private/ComponentTypeRegistry.cpp
--- a/Engine/Source/Editor/UnrealEd/Private/ComponentTypeRegistry.cpp
+++ b/Engine/Source/Editor/UnrealEd/Private/ComponentTypeRegistry.cpp
@@ -7,31 +7,47 @@
 #include "EdGraphSchema_K2.h"
 #include "KismetEditorUtilities.h"
 #include "HotReloadInterface.h"
+#include "SComponentClassCombo.h"
 
 #define LOCTEXT_NAMESPACE "ComponentTypeRegistry"
 
+//////////////////////////////////////////////////////////////////////////
+// FComponentTypeRegistryData
+
 struct FComponentTypeRegistryData
 	: public FTickableEditorObject
 {
 	FComponentTypeRegistryData();
-	void RefreshComponentList();
+
+	// Force a refresh of the components list right now (also calls the ComponentListChanged delegate to notify watchers)
+	void ForceRefreshComponentList();
+
+	static void AddBasicShapeComponents(TArray<FComponentClassComboEntryPtr>& SortedClassList);
 
 	/** Implementation of FTickableEditorObject */
 	virtual void Tick(float) override;
 	virtual bool IsTickable() const override { return true; }
 	virtual TStatId GetStatId() const { RETURN_QUICK_DECLARE_CYCLE_STAT(FTypeDatabaseUpdater, STATGROUP_Tickables); }
+
+	// Request a refresh of the components list next frame
+	void Invalidate()
+	{
+		bNeedsRefreshNextTick = true;
+	}
+public:
 	/** End implementation of FTickableEditorObject */
 	TArray<FComponentClassComboEntryPtr> ComponentClassList;
 	TArray<FComponentTypeEntry> ComponentTypeList;
 	TArray<FAssetData> PendingAssetData;
 	FOnComponentTypeListChanged ComponentListChanged;
+	bool bNeedsRefreshNextTick;
 };
 
 static const FString CommonClassGroup(TEXT("Common"));
 // This has to stay in sync with logic in FKismetCompilerContext::FinishCompilingClass
 static const FString BlueprintComponents(TEXT("Custom"));
 
-static void AddBasicShapeComponents(TArray<FComponentClassComboEntryPtr>& SortedClassList)
+void FComponentTypeRegistryData::AddBasicShapeComponents(TArray<FComponentClassComboEntryPtr>& SortedClassList)
 {
 	FString BasicShapesHeading = LOCTEXT("BasicShapesHeading", "Basic Shapes").ToString();
 
@@ -107,6 +123,7 @@ static void AddBasicShapeComponents(TArray<FComponentClassComboEntryPtr>& Sorted
 }
 
 FComponentTypeRegistryData::FComponentTypeRegistryData()
+	: bNeedsRefreshNextTick(false)
 {
 	const auto HandleAdded = [](const FAssetData& Data, FComponentTypeRegistryData* Parent)
 	{
@@ -124,7 +141,7 @@ FComponentTypeRegistryData::FComponentTypeRegistryData()
 	AssetRegistryModule.Get().OnAssetRenamed().AddStatic(HandleRenamed, this);
 }
 
-void FComponentTypeRegistryData::RefreshComponentList()
+void FComponentTypeRegistryData::ForceRefreshComponentList()
 {
 	ComponentClassList.Empty();
 	ComponentTypeList.Empty();
@@ -295,11 +312,14 @@ void FComponentTypeRegistryData::RefreshComponentList()
 			ComponentClassList.Add(CurrentEntry);
 		}
 	}
+
+	ComponentListChanged.Broadcast();
 }
 
 void FComponentTypeRegistryData::Tick(float)
 {
-	bool bRequiresRefresh = false;
+	bool bRequiresRefresh = bNeedsRefreshNextTick;
+	bNeedsRefreshNextTick = false;
 
 	if (PendingAssetData.Num() != 0)
 	{
@@ -311,7 +331,7 @@ void FComponentTypeRegistryData::Tick(float)
 
 		for (auto Asset : PendingAssetData)
 		{
-			const FName BPParentClassName(TEXT("ParentClass"));
+			const FName BPParentClassName(GET_MEMBER_NAME_CHECKED(UBlueprint, ParentClass));
 
 			bool FoundBPNotify = false;
 			for (TMap<FName, FString>::TConstIterator TagIt(Asset.TagsAndValues); TagIt; ++TagIt)
@@ -331,11 +351,13 @@ void FComponentTypeRegistryData::Tick(float)
 
 	if (bRequiresRefresh)
 	{
-		RefreshComponentList();
-		ComponentListChanged.Broadcast();
+		ForceRefreshComponentList();
 	}
 }
 
+//////////////////////////////////////////////////////////////////////////
+// FComponentTypeRegistry
+
 FComponentTypeRegistry& FComponentTypeRegistry::Get()
 {
 	static FComponentTypeRegistry ComponentRegistry;
@@ -366,7 +388,7 @@ FComponentTypeRegistry::FComponentTypeRegistry()
 	: Data(nullptr)
 {
 	Data = new FComponentTypeRegistryData();
-	Data->RefreshComponentList();
+	Data->ForceRefreshComponentList();
 
 	IHotReloadInterface& HotReloadSupport = FModuleManager::LoadModuleChecked<IHotReloadInterface>("HotReload");
 	HotReloadSupport.OnHotReload().AddRaw(this, &FComponentTypeRegistry::OnProjectHotReloaded);
@@ -383,8 +405,13 @@ FComponentTypeRegistry::~FComponentTypeRegistry()
 
 void FComponentTypeRegistry::OnProjectHotReloaded( bool bWasTriggeredAutomatically )
 {
-	Data->RefreshComponentList();
-	Data->ComponentListChanged.Broadcast();
+	Data->ForceRefreshComponentList();
+}
+
+void FComponentTypeRegistry::InvalidateClass(TSubclassOf<UActorComponent> /*ClassToUpdate*/)
+{
+	Data->Invalidate();
 }
 
+
 #undef LOCTEXT_NAMESPACE
\ No newline at end of file
diff --git a/Engine/Source/Editor/UnrealEd/Private/SComponentClassCombo.cpp b/Engine/Source/Editor/UnrealEd/Private/SComponentClassCombo.cpp
--- a/Engine/Source/Editor/UnrealEd/Private/SComponentClassCombo.cpp
+++ b/Engine/Source/Editor/UnrealEd/Private/SComponentClassCombo.cpp
@@ -139,24 +139,47 @@ void SComponentClassCombo::GenerateFilteredComponentList(const FString& InSearch
 	else
 	{
 		FilteredComponentClassList.Empty();
+
+		int32 LastHeadingIndex = INDEX_NONE;
+		FComponentClassComboEntryPtr* LastHeadingPtr = nullptr;
+
 		for (int32 ComponentIndex = 0; ComponentIndex < ComponentClassList->Num(); ComponentIndex++)
 		{
 			FComponentClassComboEntryPtr& CurrentEntry = (*ComponentClassList)[ComponentIndex];
 
-			if (CurrentEntry->IsClass() && CurrentEntry->IsIncludedInFilter())
+			if (CurrentEntry->IsHeading())
+			{
+				LastHeadingIndex = FilteredComponentClassList.Num();
+				LastHeadingPtr = &CurrentEntry;
+			}
+			else if (CurrentEntry->IsClass() && CurrentEntry->IsIncludedInFilter())
 			{
 				FString FriendlyComponentName = GetSanitizedComponentName( CurrentEntry );
 
 				if ( FriendlyComponentName.Contains( InSearchText, ESearchCase::IgnoreCase ) )
 				{
+					// Add the heading first if it hasn't already been added
+					if (LastHeadingIndex != INDEX_NONE)
+					{
+						FilteredComponentClassList.Insert(*LastHeadingPtr, LastHeadingIndex);
+						LastHeadingIndex = INDEX_NONE;
+						LastHeadingPtr = nullptr;
+					}
+
+					// Add the class
 					FilteredComponentClassList.Add( CurrentEntry );
 				}
 			}
 		}
-		if (FilteredComponentClassList.Num() > 0)
+
+		// Select the first non-category item that passed the filter
+		for (FComponentClassComboEntryPtr& TestEntry : FilteredComponentClassList)
 		{
-			FComponentClassComboEntryPtr& FirstEntry = FilteredComponentClassList[0];
-			ComponentClassListView->SetSelection(FirstEntry, ESelectInfo::OnNavigation);
+			if (TestEntry->IsClass())
+			{
+				ComponentClassListView->SetSelection(TestEntry, ESelectInfo::OnNavigation);
+				break;
+			}
 		}
 	}
 }
@@ -242,17 +265,21 @@ void SComponentClassCombo::OnAddComponentSelectionChanged( FComponentClassComboE
 	else if ( InItem.IsValid() && SelectInfo != ESelectInfo::OnMouseClick )
 	{
 		int32 SelectedIdx = INDEX_NONE;
-		FilteredComponentClassList.Find(InItem, SelectedIdx);
-		
-		if(SelectedIdx != INDEX_NONE)
+		if (FilteredComponentClassList.Find(InItem, /*out*/ SelectedIdx))
 		{
-			if(!InItem->IsClass())
+			if (!InItem->IsClass())
 			{
 				int32 SelectionDirection = SelectedIdx - PrevSelectedIndex;
 
 				// Update the previous selected index
 				PrevSelectedIndex = SelectedIdx;
 
+				// Make sure we select past the category header if we started filtering with it selected somehow (avoiding the infinite loop selecting the same item forever)
+				if (SelectionDirection == 0)
+				{
+					SelectionDirection = 1;
+				}
+
 				if(SelectedIdx + SelectionDirection >= 0 && SelectedIdx + SelectionDirection < FilteredComponentClassList.Num())
 				{
 					ComponentClassListView->SetSelection(FilteredComponentClassList[SelectedIdx + SelectionDirection], ESelectInfo::OnNavigation);
diff --git a/Engine/Source/Editor/UnrealEd/Public/ComponentTypeRegistry.h b/Engine/Source/Editor/UnrealEd/Public/ComponentTypeRegistry.h
--- a/Engine/Source/Editor/UnrealEd/Public/ComponentTypeRegistry.h
+++ b/Engine/Source/Editor/UnrealEd/Public/ComponentTypeRegistry.h
@@ -2,8 +2,6 @@
 
 #pragma once
 
-#include "SComponentClassCombo.h"
-
 DECLARE_MULTICAST_DELEGATE(FOnComponentTypeListChanged);
 
 typedef TSharedPtr<class FComponentClassComboEntry> FComponentClassComboEntryPtr;
@@ -25,15 +23,19 @@ struct UNREALED_API FComponentTypeRegistry
 	static FComponentTypeRegistry& Get();
 
 	/**
-	* Called when the user changes the text in the search box.
-	* @OutComponentList Pointer that will be set to the (globally shared) component type list
-	* @return Deleate that can be used to handle change notifications. change notifications are raised when entries are 
-	*	added or removed from the component type list
-	*/
+	 * Called when the user changes the text in the search box.
+	 * @OutComponentList Pointer that will be set to the (globally shared) component type list
+	 * @return Deleate that can be used to handle change notifications. change notifications are raised when entries are 
+	 *	added or removed from the component type list
+	 */
 	FOnComponentTypeListChanged& SubscribeToComponentList(TArray<FComponentClassComboEntryPtr>*& OutComponentList);
 	FOnComponentTypeListChanged& SubscribeToComponentList(const TArray<FComponentTypeEntry>*& OutComponentList);
 	FOnComponentTypeListChanged& GetOnComponentTypeListChanged();
 
+	/**
+	 * Called when a specific class has been updated and should force the component type registry to update as well
+	 */
+	void InvalidateClass(TSubclassOf<UActorComponent> ClassToUpdate);
 private:
 	void OnProjectHotReloaded( bool bWasTriggeredAutomatically );
 
