[INTEGRATE] Change 2441461 by [REDACTED_PERSON][REDACTED_HANDLE] on 2015/02/11 08:28:18

	UE-9263 Fix crash in Project Settings when GameMode no longer exists/has been moved

diff --git a/Engine/Source/Editor/UnrealEd/Public/GameModeInfoCustomizer.h b/Engine/Source/Editor/UnrealEd/Public/GameModeInfoCustomizer.h
--- a/Engine/Source/Editor/UnrealEd/Public/GameModeInfoCustomizer.h
+++ b/Engine/Source/Editor/UnrealEd/Public/GameModeInfoCustomizer.h
@@ -176,11 +176,6 @@ public:
 		ConstructorHelpers::StripObjectClass(StrippedClassName);
 
 		const UClass* GameModeClass = FindObject<UClass>(ANY_PACKAGE, *StrippedClassName);
-		if (!GameModeClass)
-		{
-			GameModeClass = LoadObject<UClass>(NULL, *ClassName);
-		}
-
 		return GameModeClass;
 	}
 
