[INTEGRATE] Change 2441561 by [REDACTED_PERSON][REDACTED_HANDLE].7 on 2015/02/11 09:50:00

	Editor auto reimport improvements

	There are several tweaks here:

	    * Fixed a bug that caused the "also delete source files" checkbox to not appear when deleting assets
	    * Fixed "Processing outstanding changes" notification appearing after deleting an asset with its content file
	    * A few info message tweaks
	    * Newly created files that are already referenced by an asset(s) are now ignored
	    * Fixed an issue where the system kicked in before the asset registry had finished discovering assets

diff --git a/Engine/Source/Editor/UnrealEd/Private/AssetDeleteModel.cpp b/Engine/Source/Editor/UnrealEd/Private/AssetDeleteModel.cpp
--- a/Engine/Source/Editor/UnrealEd/Private/AssetDeleteModel.cpp
+++ b/Engine/Source/Editor/UnrealEd/Private/AssetDeleteModel.cpp
@@ -166,7 +166,8 @@ void FAssetDeleteModel::DiscoverSourceFileReferences(FPendingDelete& PendingDele
 		}
 		else
 		{
-			SourceFileToAssetCount.Add(SourcePath, Utils::FindAssetsPertainingToFile(AssetRegistry, SourcePath).Num());
+			const int32 NumReferencingAssets = Utils::FindAssetsPertainingToFile(AssetRegistry, SourcePath).Num();
+			SourceFileToAssetCount.Add(SourcePath, NumReferencingAssets - 1);
 		}
 	}
 }
@@ -198,6 +199,9 @@ void FAssetDeleteModel::DeleteSourceContentFiles()
 			continue;
 		}
 
+		// One way or another this file is going to be deleted, but we don't want the import manager to react to the deletion
+		GUnrealEd->AutoReimportManager->ReportExternalChange(Path, FFileChangeData::FCA_Removed);
+
 		if (ISourceControlModule::Get().IsEnabled())
 		{
 			const FSourceControlStatePtr SourceControlState = SourceControlProvider.GetState(Path, EStateCacheUsage::ForceUpdate);
diff --git a/Engine/Source/Editor/UnrealEd/Private/AutoReimport/AutoReimportManager.cpp b/Engine/Source/Editor/UnrealEd/Private/AutoReimport/AutoReimportManager.cpp
--- a/Engine/Source/Editor/UnrealEd/Private/AutoReimport/AutoReimportManager.cpp
+++ b/Engine/Source/Editor/UnrealEd/Private/AutoReimport/AutoReimportManager.cpp
@@ -98,7 +98,14 @@ class FAutoReimportManager : public FTickableEditorObject, public FGCObject, pub
 public:
 
 	FAutoReimportManager();
+
+	/** Get a list of currently monitored directories */
 	TArray<FPathAndMountPoint> GetMonitoredDirectories() const;
+
+	/** Report an external change to the manager, such that a subsequent equal change reported by the os be ignored */
+	void ReportExternalChange(const FString& Filename, FFileChangeData::EFileChangeAction Action);
+
+	/** Destroy this manager */
 	void Destroy();
 
 private:
@@ -177,9 +184,6 @@ private:
 	/** A list of packages to save when we've added a bunch of assets */
 	TArray<UPackage*> PackagesToSave;
 
-	/** Set when we are processing changes to avoid reentrant logic for asset deletions etc */
-	bool bIsProcessingChanges;
-
 	/** The time when the last change to the cache was reported */
 	FTimeLimit StartProcessingDelay;
 
@@ -189,7 +193,6 @@ private:
 
 FAutoReimportManager::FAutoReimportManager()
 	: StateMachine(ECurrentState::Idle)
-	, bIsProcessingChanges(false)
 	, StartProcessingDelay(0.5)
 {
 	// @todo: arodham: update this when modules are reloaded or new factory types are available?
@@ -231,6 +234,17 @@ TArray<FPathAndMountPoint> FAutoReimportManager::GetMonitoredDirectories() const
 	return Dirs;
 }
 
+void FAutoReimportManager::ReportExternalChange(const FString& Filename, FFileChangeData::EFileChangeAction Action)
+{
+	for (auto& Monitor : DirectoryMonitors)
+	{
+		if (Filename.StartsWith(Monitor.GetDirectory()))
+		{
+			Monitor.ReportExternalChange(Filename, Action);
+		}
+	}
+}
+
 void FAutoReimportManager::Destroy()
 {
 	if (auto* Settings = GetMutableDefault<UEditorLoadingSavingSettings>())
@@ -318,9 +332,11 @@ TOptional<ECurrentState> FAutoReimportManager::ProcessAdditions(const FTimeLimit
 		Pair.Value.Sort([](const UFactory& A, const UFactory& B) { return A.ImportPriority > B.ImportPriority; });
 	}
 
+	const IAssetRegistry& Registry = FModuleManager::LoadModuleChecked<FAssetRegistryModule>("AssetRegistry").Get();
+
 	for (auto& Monitor : DirectoryMonitors)
 	{
-		Monitor.ProcessAdditions(PackagesToSave, TimeLimit, Factories, *FeedbackContextOverride);
+		Monitor.ProcessAdditions(Registry, PackagesToSave, TimeLimit, Factories, *FeedbackContextOverride);
 		yield TOptional<ECurrentState>();
 	}
 
@@ -380,11 +396,12 @@ TOptional<ECurrentState> FAutoReimportManager::ProcessDeletions()
 
 	if (AssetsToDelete.Num() > 0)
 	{
-		ObjectTools::DeleteAssets(AssetsToDelete);
 		for (const auto& AssetData : AssetsToDelete)
 		{
-			FeedbackContextOverride->AddMessage(EMessageSeverity::Info, FText::Format(LOCTEXT("Success_DeletedAsset", "Deleted {0}."), FText::FromName(AssetData.AssetName)));
+			FeedbackContextOverride->AddMessage(EMessageSeverity::Info, FText::Format(LOCTEXT("Success_DeletedAsset", "Attempting to delete {0} (its source file has been removed)."), FText::FromName(AssetData.AssetName)));
 		}
+
+		ObjectTools::DeleteAssets(AssetsToDelete);
 	}
 
 	Cleanup();
@@ -416,7 +433,8 @@ TOptional<ECurrentState> FAutoReimportManager::Idle(const FTimeLimit& TimeLimit)
 		}
 		else
 		{
-			if (StartProcessingDelay.Exceeded())
+			const IAssetRegistry& AssetRegistry = FModuleManager::LoadModuleChecked<FAssetRegistryModule>("AssetRegistry").Get();
+			if (StartProcessingDelay.Exceeded() && !AssetRegistry.IsLoadingAssets())
 			{
 				// Deal with changes to the file system second
 				for (auto& Monitor : DirectoryMonitors)
@@ -616,6 +634,11 @@ void UAutoReimportManager::Initialize()
 	Implementation = MakeShareable(new FAutoReimportManager);
 }
 
+void UAutoReimportManager::ReportExternalChange(const FString& Filename, FFileChangeData::EFileChangeAction Action)
+{
+	Implementation->ReportExternalChange(Filename, Action);
+}
+
 TArray<FPathAndMountPoint> UAutoReimportManager::GetMonitoredDirectories() const
 {
 	return Implementation->GetMonitoredDirectories();
diff --git a/Engine/Source/Editor/UnrealEd/Private/AutoReimport/ContentDirectoryMonitor.cpp b/Engine/Source/Editor/UnrealEd/Private/AutoReimport/ContentDirectoryMonitor.cpp
--- a/Engine/Source/Editor/UnrealEd/Private/AutoReimport/ContentDirectoryMonitor.cpp
+++ b/Engine/Source/Editor/UnrealEd/Private/AutoReimport/ContentDirectoryMonitor.cpp
@@ -41,6 +41,11 @@ void FContentDirectoryMonitor::Destroy()
 	Cache.Destroy();
 }
 
+void FContentDirectoryMonitor::ReportExternalChange(const FString& Filename, FFileChangeData::EFileChangeAction Action)
+{
+	Cache.ReportExternalChange(Filename, Action);
+}
+
 void FContentDirectoryMonitor::Tick(const FTimeLimit& TimeLimit)
 {
 	Cache.Tick(TimeLimit);
@@ -98,7 +103,7 @@ void FContentDirectoryMonitor::StartProcessing()
 	}
 }
 
-void FContentDirectoryMonitor::ProcessAdditions(TArray<UPackage*>& OutPackagesToSave, const FTimeLimit& TimeLimit, const TMap<FString, TArray<UFactory*>>& InFactoriesByExtension, FReimportFeedbackContext& Context)
+void FContentDirectoryMonitor::ProcessAdditions(const IAssetRegistry& Registry, TArray<UPackage*>& OutPackagesToSave, const FTimeLimit& TimeLimit, const TMap<FString, TArray<UFactory*>>& InFactoriesByExtension, FReimportFeedbackContext& Context)
 {
 	bool bCancelled = false;
 	for (int32 Index = 0; Index < AddedFiles.Num(); ++Index)
@@ -116,6 +121,26 @@ void FContentDirectoryMonitor::ProcessAdditions(TArray<UPackage*>& OutPackagesTo
 
 		const FString FullFilename = Cache.GetDirectory() + Addition.Filename.Get();
 
+		// Don't create assets for new files if assets already exist for the filename
+		auto ExistingReferences = Utils::FindAssetsPertainingToFile(Registry, FullFilename);
+		if (ExistingReferences.Num() != 0)
+		{
+			FText Message;
+			if (ExistingReferences.Num() == 1)
+			{
+				Message = FText::Format(LOCTEXT("Info_AlreadyImported_Single", "Ignoring new file {0} as it's already imported by {1}."), FText::FromString(FPaths::GetCleanFilename(FullFilename)), FText::FromName(ExistingReferences[0].AssetName));
+			}
+			else
+			{
+				Message = FText::Format(LOCTEXT("Info_AlreadyImported_Multiple", "Ignoring new file {0} as it's already imported by {1} assets."), FText::FromString(FPaths::GetCleanFilename(FullFilename)), FText::AsNumber(ExistingReferences.Num()));
+			}
+			Context.AddMessage(EMessageSeverity::Info, Message);
+
+			// We don't try and import new files that are already associated with an asset
+			Cache.CompleteTransaction(MoveTemp(Addition));
+			continue;
+		}
+
 		FString NewAssetName = ObjectTools::SanitizeObjectName(FPaths::GetBaseFilename(FullFilename));
 		FString PackagePath = PackageTools::SanitizePackageName(MountedContentPath / FPaths::GetPath(Addition.Filename.Get()) / NewAssetName);
 
@@ -125,7 +150,7 @@ void FContentDirectoryMonitor::ProcessAdditions(TArray<UPackage*>& OutPackagesTo
 		}
 		else if (FindPackage(nullptr, *PackagePath))
 		{
-			Context.AddMessage(EMessageSeverity::Warning, FText::Format(LOCTEXT("Error_ExistingAsset", "Can't create a new asset at {0} - one already exists."), FText::FromString(PackagePath)));
+			Context.AddMessage(EMessageSeverity::Warning, FText::Format(LOCTEXT("Warning_ExistingAsset", "Can't create a new asset at {0} - one already exists."), FText::FromString(PackagePath)));
 		}
 		else
 		{
@@ -136,6 +161,8 @@ void FContentDirectoryMonitor::ProcessAdditions(TArray<UPackage*>& OutPackagesTo
 			}
 			else
 			{
+				Context.AddMessage(EMessageSeverity::Info, FText::Format(LOCTEXT("Info_CreatingNewAsset", "Creating new asset {0}."), FText::FromString(PackagePath)));
+
 				// Make sure the destination package is loaded
 				NewPackage->FullyLoad();
 
@@ -188,8 +215,6 @@ void FContentDirectoryMonitor::ProcessAdditions(TArray<UPackage*>& OutPackagesTo
 						continue;
 					}
 
-					Context.AddMessage(EMessageSeverity::Info, FText::Format(LOCTEXT("Success_CreatedNewAsset", "Created new asset {0}."), FText::FromString(PackagePath)));
-
 					FAssetRegistryModule::AssetCreated(NewAsset);
 					GEditor->BroadcastObjectReimported(NewAsset);
 
diff --git a/Engine/Source/Editor/UnrealEd/Private/AutoReimport/ContentDirectoryMonitor.h b/Engine/Source/Editor/UnrealEd/Private/AutoReimport/ContentDirectoryMonitor.h
--- a/Engine/Source/Editor/UnrealEd/Private/AutoReimport/ContentDirectoryMonitor.h
+++ b/Engine/Source/Editor/UnrealEd/Private/AutoReimport/ContentDirectoryMonitor.h
@@ -26,7 +26,7 @@ public:
 	void StartProcessing();
 
 	/** Extract the files we need to import from our outstanding changes (happens first)*/ 
-	void ProcessAdditions(TArray<UPackage*>& OutPackagesToSave, const FTimeLimit& TimeLimit, const TMap<FString, TArray<UFactory*>>& InFactoriesByExtension, class FReimportFeedbackContext& Context);
+	void ProcessAdditions(const IAssetRegistry& Registry, TArray<UPackage*>& OutPackagesToSave, const FTimeLimit& TimeLimit, const TMap<FString, TArray<UFactory*>>& InFactoriesByExtension, class FReimportFeedbackContext& Context);
 
 	/** Process the outstanding changes that we have cached */
 	void ProcessModifications(const IAssetRegistry& Registry, const FTimeLimit& TimeLimit, class FReimportFeedbackContext& Context);
@@ -34,6 +34,9 @@ public:
 	/** Extract the assets we need to delete from our outstanding changes (happens last) */ 
 	void ExtractAssetsToDelete(const IAssetRegistry& Registry, TArray<FAssetData>& OutAssetsToDelete);
 
+	/** Report an external change to the manager, such that a subsequent equal change reported by the os be ignored */
+	void ReportExternalChange(const FString& Filename, FFileChangeData::EFileChangeAction Action);
+
 	/** Destroy this monitor including its cache */
 	void Destroy();
 
diff --git a/Engine/Source/Editor/UnrealEd/Private/AutoReimport/FileCache.cpp b/Engine/Source/Editor/UnrealEd/Private/AutoReimport/FileCache.cpp
--- a/Engine/Source/Editor/UnrealEd/Private/AutoReimport/FileCache.cpp
+++ b/Engine/Source/Editor/UnrealEd/Private/AutoReimport/FileCache.cpp
@@ -205,6 +205,16 @@ TArray<FUpdateCacheTransaction> FFileCache::GetOutstandingChanges()
 	return Moved;
 }
 
+void FFileCache::ReportExternalChange(const FString& Filename, FFileChangeData::EFileChangeAction Action)
+{
+	auto Transaction = ProcessChange(Filename, Action, IFileManager::Get());
+	if (Transaction.IsSet())
+	{
+		// Just complete the transaction immediately without notifying the client
+		CompleteTransaction(MoveTemp(Transaction.GetValue()));
+	}
+}
+
 void FFileCache::CompleteTransaction(FUpdateCacheTransaction&& Transaction)
 {
 	auto* CachedData = CachedDirectoryState.Files.Find(Transaction.Filename);
@@ -307,82 +317,106 @@ void FFileCache::OnDirectoryChanged(const TArray<FFileChangeData>& FileChanges)
 
 	OutstandingChanges.Reserve(OutstandingChanges.Num() + FileChanges.Num());
 
-	const int32 RootPathLen = Config.Directory.Len();
 	for (const auto& ThisEntry : FileChanges)
 	{
-		// If it's a directory or is not applicable, ignore it
-		if (FileManager.DirectoryExists(*ThisEntry.Filename) || !IsFileApplicable(ThisEntry.Filename))
+		auto Transaction = ProcessChange(ThisEntry.Filename, ThisEntry.Action, FileManager);
+		if (Transaction.IsSet())
 		{
-			continue;
+			OutstandingChanges.Add(MoveTemp(Transaction.GetValue()));
 		}
+	}
+
+	Utils::RemoveDuplicates(OutstandingChanges, [](const FUpdateCacheTransaction& A, const FUpdateCacheTransaction& B){
+		return A.Action == B.Action && A.Filename == B.Filename;
+	});
+}
+
+TOptional<FUpdateCacheTransaction> FFileCache::ProcessChange(const FString& Filename, FFileChangeData::EFileChangeAction Action, IFileManager& FileManager)
+{
+	// If it's a directory or is not applicable, ignore it
+	if (FileManager.DirectoryExists(*Filename) || !IsFileApplicable(Filename))
+	{
+		return TOptional<FUpdateCacheTransaction>();
+	}
 
-		FImmutableString TransactionFilename;
+	FImmutableString TransactionFilename;
+	{
+		FString Temp = FPaths::ConvertRelativePathToFull(Filename);
+		if (Config.PathType == EPathType::Relative)
 		{
-			FString Temp = FPaths::ConvertRelativePathToFull(ThisEntry.Filename);
-			if (Config.PathType == EPathType::Relative)
-			{
-				TransactionFilename = FString(*Temp + RootPathLen);
-			}
-			else
-			{
-				TransactionFilename = MoveTemp(Temp);
-			}
+			TransactionFilename = FString(*Temp + Config.Directory.Len());
 		}
-		
-		switch (ThisEntry.Action)
+		else
 		{
-		case FFileChangeData::FCA_Added:
-			{
-				const int32 NumRemoved = OutstandingChanges.RemoveAll([&](const FUpdateCacheTransaction& X){
-					return X.Action == FFileChangeData::FCA_Removed && X.Filename == TransactionFilename;
-				});
+			TransactionFilename = MoveTemp(Temp);
+		}
+	}
+	
+	switch (Action)
+	{
+	case FFileChangeData::FCA_Added:
+		{
+			// Collapse Remove -> Add into a modification
+			const int32 NumRemoved = OutstandingChanges.RemoveAll([&](const FUpdateCacheTransaction& X){
+				return X.Action == FFileChangeData::FCA_Removed && X.Filename == TransactionFilename;
+			});
 
-				const auto Action = NumRemoved == 0 ? FFileChangeData::FCA_Added : FFileChangeData::FCA_Modified;
-				OutstandingChanges.Add(FUpdateCacheTransaction(TransactionFilename, Action, FileManager.GetTimeStamp(*ThisEntry.Filename)));
-			}
-			break;
+			const auto Action = NumRemoved == 0 ? FFileChangeData::FCA_Added : FFileChangeData::FCA_Modified;
 
-		case FFileChangeData::FCA_Removed:
+			// If it's a modification (remove then add) or the file doesn't already exist in the cache, we can create a transaction
+			if (Action == FFileChangeData::FCA_Modified || !CachedDirectoryState.Files.Find(TransactionFilename))
 			{
-				bool bPreviouslyAdded = false;
-				const int32 NumRemoved = OutstandingChanges.RemoveAll([&](const FUpdateCacheTransaction& X){
-					if (X.Filename == TransactionFilename)
-					{
-						bPreviouslyAdded = X.Action == FFileChangeData::FCA_Added || bPreviouslyAdded;
-						return true;
-					}
-					return false;
-				});
-
-				if (!bPreviouslyAdded)
+				return FUpdateCacheTransaction(TransactionFilename, Action, FileManager.GetTimeStamp(*Filename));
+			}
+		}
+		break;
+
+	case FFileChangeData::FCA_Removed:
+		{
+			bool bPreviouslyAdded = false;
+			const int32 NumRemoved = OutstandingChanges.RemoveAll([&](const FUpdateCacheTransaction& X){
+				if (X.Filename == TransactionFilename)
 				{
-					OutstandingChanges.Add(FUpdateCacheTransaction(TransactionFilename, FFileChangeData::FCA_Removed));
+					bPreviouslyAdded = X.Action == FFileChangeData::FCA_Added || bPreviouslyAdded;
+					return true;
 				}
-			}
-			
-			break;
+				return false;
+			});
 
-		case FFileChangeData::FCA_Modified:
+			// We only create the transaction if there is net change, and the file actually exists in the cache
+			if (!bPreviouslyAdded && CachedDirectoryState.Files.Find(TransactionFilename))
 			{
-				const bool bPreviouslyAdded = OutstandingChanges.ContainsByPredicate([&](const FUpdateCacheTransaction& X){
-					return X.Filename == TransactionFilename && X.Action == FFileChangeData::FCA_Added;
-				});
+				return FUpdateCacheTransaction(TransactionFilename, FFileChangeData::FCA_Removed);
+			}
+		}
+		
+		break;
+
+	case FFileChangeData::FCA_Modified:
+		{
+			const bool bPreviouslyAdded = OutstandingChanges.ContainsByPredicate([&](const FUpdateCacheTransaction& X){
+				return X.Filename == TransactionFilename && X.Action == FFileChangeData::FCA_Added;
+			});
 
-				if (!bPreviouslyAdded)
+			const FDateTime& Timestamp = FileManager.GetTimeStamp(*Filename);
+
+			// Don't report subsequent modifications if an addition is still waiting to be handled
+			if (!bPreviouslyAdded)
+			{
+				auto* ExistingFileData = CachedDirectoryState.Files.Find(TransactionFilename);
+				if (ExistingFileData && ExistingFileData->Timestamp < Timestamp)
 				{
-					OutstandingChanges.Add(FUpdateCacheTransaction(TransactionFilename, FFileChangeData::FCA_Modified, FileManager.GetTimeStamp(*ThisEntry.Filename)));
+					return FUpdateCacheTransaction(TransactionFilename, FFileChangeData::FCA_Modified, Timestamp);
 				}
 			}
-			break;
-
-		default:
-			break;
 		}
+		break;
+
+	default:
+		break;
 	}
 
-	Utils::RemoveDuplicates(OutstandingChanges, [](const FUpdateCacheTransaction& A, const FUpdateCacheTransaction& B){
-		return A.Action == B.Action && A.Filename == B.Filename;
-	});
+	return TOptional<FUpdateCacheTransaction>();
 }
 
 bool FFileCache::IsFileApplicable(const FString& Filename)
diff --git a/Engine/Source/Editor/UnrealEd/Private/AutoReimport/FileCache.h b/Engine/Source/Editor/UnrealEd/Private/AutoReimport/FileCache.h
--- a/Engine/Source/Editor/UnrealEd/Private/AutoReimport/FileCache.h
+++ b/Engine/Source/Editor/UnrealEd/Private/AutoReimport/FileCache.h
@@ -87,11 +87,15 @@ struct FUpdateCacheTransaction
 
 private:
 	friend class FFileCache;
-	
+
 	/** Construction responsibility is held by FFileCache */
 	FUpdateCacheTransaction(FImmutableString InFilename, FFileChangeData::EFileChangeAction InAction, const FDateTime& InTimestamp = 0)
 		: Filename(MoveTemp(InFilename)), Timestamp(InTimestamp), Action(InAction)
 	{}
+	
+	/** Not Copyable */
+	FUpdateCacheTransaction(const FUpdateCacheTransaction&) = delete;
+	FUpdateCacheTransaction& operator=(const FUpdateCacheTransaction&) = delete;
 };
 
 /** Enum specifying whether a path should be relative or absolute */
@@ -217,6 +221,9 @@ public:
 	 */
 	void CompleteTransaction(FUpdateCacheTransaction&& Transaction);
 
+	/** Report an external change to the manager, such that a subsequent equal change reported by the os be ignored */
+	void ReportExternalChange(const FString& Filename, FFileChangeData::EFileChangeAction Action);
+
 	/** Get the number of pending changes to the cache. */
 	int32 GetNumOutstandingChanges() const { return OutstandingChanges.Num(); }
 
@@ -228,6 +235,9 @@ private:
 	/** Called when the directory we are monitoring has been changed in some way */
 	void OnDirectoryChanged(const TArray<FFileChangeData>& FileChanges);
 
+	/** Process a change, creating an outstanding change transaction if necessary */
+	TOptional<FUpdateCacheTransaction> ProcessChange(const FString& Filename, FFileChangeData::EFileChangeAction Action, IFileManager& FileManager);
+
 	/** Check whether a file is applicable to our cache */
 	bool IsFileApplicable(const FString& InPath);
 
diff --git a/Engine/Source/Editor/UnrealEd/Public/AutoReimport/AutoReimportManager.h b/Engine/Source/Editor/UnrealEd/Public/AutoReimport/AutoReimportManager.h
--- a/Engine/Source/Editor/UnrealEd/Public/AutoReimport/AutoReimportManager.h
+++ b/Engine/Source/Editor/UnrealEd/Public/AutoReimport/AutoReimportManager.h
@@ -2,6 +2,7 @@
 
 #pragma once
 
+#include "IDirectoryWatcher.h"
 #include "AutoReimportManager.generated.h"
 
 /** Struct representing a path on disk, and its virtual mount point */
@@ -32,6 +33,9 @@ public:
 	/** Get a list of absolute directories that we are monitoring */
 	TArray<FPathAndMountPoint> GetMonitoredDirectories() const;
 
+	/** Report an external change to the manager, such that a subsequent equal change reported by the os be ignored */
+	void ReportExternalChange(const FString& Filename, FFileChangeData::EFileChangeAction Action);
+
 private:
 
 	/** UObject Interface */
