[INTEGRATE] Change 2442264 by [REDACTED_PERSON][REDACTED_HANDLE] on 2015/02/11 16:14:52

	- Disable adding ChildActorComponent via IWCE
	https://jira.ol.epicgames.net/browse/UE-9642

diff --git a/Engine/Source/Editor/Kismet/Private/SSCSEditor.cpp b/Engine/Source/Editor/Kismet/Private/SSCSEditor.cpp
--- a/Engine/Source/Editor/Kismet/Private/SSCSEditor.cpp
+++ b/Engine/Source/Editor/Kismet/Private/SSCSEditor.cpp
@@ -3138,8 +3138,8 @@ void SSCSEditor::Construct( const FArguments& InArgs )
 				.HAlign(HAlign_Left)
 				[
 					SNew(SComponentClassCombo)
-					//** Hack 4.7 Child actor components are not allowed to be added via the actor details panel */
-					//.AllowChildActorComponent(EditorMode==EComponentEditorMode::BlueprintSCS)
+					// Hack 4.7 Child actor components are not allowed to be added via the actor details panel 
+					.AllowChildActorComponent(EditorMode==EComponentEditorMode::BlueprintSCS)
 					.AddMetaData<FTagMetaData>(FTagMetaData(TEXT("Actor.AddComponent")))
 					.Visibility(bHideComponentClassCombo ? EVisibility::Hidden : EVisibility::Visible)
 					.OnComponentClassSelected(this, &SSCSEditor::PerformComboAddClass)
diff --git a/Engine/Source/Editor/UnrealEd/Private/SComponentClassCombo.cpp b/Engine/Source/Editor/UnrealEd/Private/SComponentClassCombo.cpp
--- a/Engine/Source/Editor/UnrealEd/Private/SComponentClassCombo.cpp
+++ b/Engine/Source/Editor/UnrealEd/Private/SComponentClassCombo.cpp
@@ -24,6 +24,8 @@ FString FComponentClassComboEntry::GetClassName() const
 
 void SComponentClassCombo::Construct(const FArguments& InArgs)
 {
+	bAllowChildActorComponent = InArgs._AllowChildActorComponent;
+
 	PrevSelectedIndex = INDEX_NONE;
 	OnComponentClassSelected = InArgs._OnComponentClassSelected;
 
@@ -132,14 +134,27 @@ void SComponentClassCombo::ClearSelection()
 
 void SComponentClassCombo::GenerateFilteredComponentList(const FString& InSearchText)
 {
+	FilteredComponentClassList.Empty();
 	if ( InSearchText.IsEmpty() )
 	{
-		FilteredComponentClassList = *ComponentClassList;
+		if( !bAllowChildActorComponent )
+		{
+			for(auto& ClassEntry : *ComponentClassList)
+			{
+				UClass* Class = ClassEntry->GetComponentClass();
+				if( !Class || !Class->IsChildOf(UChildActorComponent::StaticClass()) )
+				{
+					FilteredComponentClassList.Add( ClassEntry );
+				}
+			}
+		}
+		else
+		{
+			FilteredComponentClassList = *ComponentClassList;
+		}
 	}
 	else
 	{
-		FilteredComponentClassList.Empty();
-
 		int32 LastHeadingIndex = INDEX_NONE;
 		FComponentClassComboEntryPtr* LastHeadingPtr = nullptr;
 
diff --git a/Engine/Source/Editor/UnrealEd/Public/SComponentClassCombo.h b/Engine/Source/Editor/UnrealEd/Public/SComponentClassCombo.h
--- a/Engine/Source/Editor/UnrealEd/Public/SComponentClassCombo.h
+++ b/Engine/Source/Editor/UnrealEd/Public/SComponentClassCombo.h
@@ -150,9 +150,11 @@ class UNREALED_API SComponentClassCombo : public SComboButton
 public:
 	SLATE_BEGIN_ARGS( SComponentClassCombo )
 		: _IncludeText(true)
+		, _AllowChildActorComponent(true)
 	{}
 
 		SLATE_ATTRIBUTE(bool, IncludeText)
+		SLATE_ARGUMENT(bool, AllowChildActorComponent)
 		SLATE_EVENT( FComponentClassSelected, OnComponentClassSelected )
 
 	SLATE_END_ARGS()
@@ -214,4 +216,7 @@ private:
 
 	/** Cached selection index used to skip over unselectable items */
 	int32 PrevSelectedIndex;
+
+	/** Hack 4.7 Child actor components are not allowed to be added via the actor details panel */
+	bool bAllowChildActorComponent;
 };
