[INTEGRATE] Change 2442643 by [REDACTED_PERSON][REDACTED_HANDLE] on 2015/02/11 19:07:13

	Blueprints now start with basic nodes, in a disabled state that will not be compiled. These nodes will auto-enable if you connect another node to them.

diff --git a/Engine/Source/Editor/BlueprintGraph/Private/BlueprintEventNodeSpawner.cpp b/Engine/Source/Editor/BlueprintGraph/Private/BlueprintEventNodeSpawner.cpp
--- a/Engine/Source/Editor/BlueprintGraph/Private/BlueprintEventNodeSpawner.cpp
+++ b/Engine/Source/Editor/BlueprintGraph/Private/BlueprintEventNodeSpawner.cpp
@@ -21,6 +21,15 @@ namespace UBlueprintEventNodeSpawnerImpl
 	 * @return Null if no event was found, otherwise a pointer to the named event.
 	 */
 	static UK2Node_Event* FindCustomEventNode(UBlueprint* Blueprint, FName const CustomName);
+
+	/**
+	 * Helper function for removing all disabled nodes connected to a 
+	 * disabled source node
+	 *
+	 * @param InNode			The node to start with
+	 * @param InParentGraph		The graph to remove the nodes from
+	 */
+	static void RemoveAllDisabledNodes(UEdGraphNode* InNode, UEdGraph* InParentGraph);
 }
 
 //------------------------------------------------------------------------------
@@ -45,6 +54,28 @@ UK2Node_Event* UBlueprintEventNodeSpawnerImpl::FindCustomEventNode(UBlueprint* B
 	return FoundNode;
 }
 
+//------------------------------------------------------------------------------
+static void UBlueprintEventNodeSpawnerImpl::RemoveAllDisabledNodes(UEdGraphNode* InNode, UEdGraph* InParentGraph)
+{
+	if(InNode && !InNode->bIsNodeEnabled)
+	{
+		// Go through all pin connections and consume any disabled nodes so we do not leave garbage.
+		for (UEdGraphPin* Pin : InNode->Pins)
+		{
+			TArray<UEdGraphPin*> LinkedToCopy = Pin->LinkedTo;
+			for (UEdGraphPin* OtherPin : LinkedToCopy)
+			{
+				// Break the pin link back
+				OtherPin->BreakLinkTo(Pin);
+				RemoveAllDisabledNodes(OtherPin->GetOwningNode(), InParentGraph);
+			}
+		}
+
+		InNode->BreakAllNodeLinks();
+		InParentGraph->RemoveNode(InNode);
+	}
+};
+
 /*******************************************************************************
  * UBlueprintEventNodeSpawner
  ******************************************************************************/
@@ -149,6 +180,14 @@ UEdGraphNode* UBlueprintEventNodeSpawner::Invoke(UEdGraph* ParentGraph, FBinding
 		EventName  = EventFunc->GetFName();	
 	}
 
+	// This Event node might already be present in the Blueprint in a disabled state, 
+	// remove it and allow the user to successfully place the node where they want it.
+	if(EventNode && !EventNode->bIsNodeEnabled)
+	{
+		UBlueprintEventNodeSpawnerImpl::RemoveAllDisabledNodes(EventNode, ParentGraph);
+		EventNode = nullptr;
+	}
+
 	// if there is no existing node, then we can happily spawn one into the graph
 	if (EventNode == nullptr)
 	{
diff --git a/Engine/Source/Editor/GameplayAbilitiesEditor/Private/GameplayAbilitiesBlueprintFactory.cpp b/Engine/Source/Editor/GameplayAbilitiesEditor/Private/GameplayAbilitiesBlueprintFactory.cpp
--- a/Engine/Source/Editor/GameplayAbilitiesEditor/Private/GameplayAbilitiesBlueprintFactory.cpp
+++ b/Engine/Source/Editor/GameplayAbilitiesEditor/Private/GameplayAbilitiesBlueprintFactory.cpp
@@ -289,6 +289,9 @@ UObject* UGameplayAbilitiesBlueprintFactory::FactoryCreateNew(UClass* Class, UOb
 				FBlueprintEditorUtils::AddUbergraphPage(NewBP, NewGraph);
 				NewBP->LastEditedDocuments.Add(NewGraph);
 				NewGraph->bAllowDeletion = false;
+
+				UEdGraphNode* EventNode = FKismetEditorUtilities::AddDefaultEventNode(NewBP, NewGraph, FName(TEXT("K2_ActivateAbility")), UGameplayAbility::StaticClass());
+				FKismetEditorUtilities::AddDefaultEventNode(NewBP, NewGraph, FName(TEXT("K2_OnEndAbility")), UGameplayAbility::StaticClass(), EventNode->NodePosY + EventNode->NodeHeight + 200);
 			}
 		}
 
diff --git a/Engine/Source/Editor/GraphEditor/Private/SCommentBubble.cpp b/Engine/Source/Editor/GraphEditor/Private/SCommentBubble.cpp
--- a/Engine/Source/Editor/GraphEditor/Private/SCommentBubble.cpp
+++ b/Engine/Source/Editor/GraphEditor/Private/SCommentBubble.cpp
@@ -192,7 +192,7 @@ void SCommentBubble::UpdateBubble()
 					[
 						SNew(SImage)
 						.Image( FEditorStyle::GetBrush( TEXT("Graph.Node.CommentBubble")) )
-						.ColorAndOpacity( ColorAndOpacity )
+						.ColorAndOpacity( this, &SCommentBubble::GetBubbleColor )
 					]
 					+SOverlay::Slot()
 					.HAlign(HAlign_Left)
@@ -334,6 +334,17 @@ FSlateColor SCommentBubble::GetToggleButtonColor() const
 	return FLinearColor( 1.f, 1.f, 1.f, OpacityValue * OpacityValue );
 }
 
+FSlateColor SCommentBubble::GetBubbleColor() const
+{
+	FLinearColor ReturnColor = ColorAndOpacity.Get().GetSpecifiedColor();
+
+	if(!GraphNode->bIsNodeEnabled)
+	{
+		ReturnColor.A *= 0.6f;
+	}
+	return ReturnColor;
+}
+
 FSlateColor SCommentBubble::GetTextBackgroundColor() const
 {
 	return TextBlock->HasKeyboardFocus() ? FLinearColor::White : SCommentBubbleDefs::TextClearBackground;
diff --git a/Engine/Source/Editor/GraphEditor/Private/SGraphNode.cpp b/Engine/Source/Editor/GraphEditor/Private/SGraphNode.cpp
--- a/Engine/Source/Editor/GraphEditor/Private/SGraphNode.cpp
+++ b/Engine/Source/Editor/GraphEditor/Private/SGraphNode.cpp
@@ -502,9 +502,47 @@ UObject* SGraphNode::GetObjectBeingDisplayed() const
 
 FSlateColor SGraphNode::GetNodeTitleColor() const
 {
-	FLinearColor NodeTitleColor = GraphNode->IsDeprecated() ? FLinearColor::Red : GetNodeObj()->GetNodeTitleColor();
-	NodeTitleColor.A = FadeCurve.GetLerp();
-	return NodeTitleColor;
+	FLinearColor ReturnTitleColor = GraphNode->IsDeprecated() ? FLinearColor::Red : GetNodeObj()->GetNodeTitleColor();
+
+	if(!GraphNode->bIsNodeEnabled)
+	{
+		ReturnTitleColor *= FLinearColor(0.5f, 0.5f, 0.5f, 0.4f);
+	}
+	else
+	{
+		ReturnTitleColor.A = FadeCurve.GetLerp();
+	}
+	return ReturnTitleColor;
+}
+
+FSlateColor SGraphNode::GetNodeBodyColor() const
+{
+	FLinearColor ReturnBodyColor = FLinearColor::White;
+	if(!GraphNode->bIsNodeEnabled)
+	{
+		ReturnBodyColor *= FLinearColor(1.0f, 1.0f, 1.0f, 0.5f); 
+	}
+	return ReturnBodyColor;
+}
+
+FSlateColor SGraphNode::GetNodeTitleIconColor() const
+{
+	FLinearColor ReturnIconColor = IconColor;
+	if(!GraphNode->bIsNodeEnabled)
+	{
+		ReturnIconColor *= FLinearColor(1.0f, 1.0f, 1.0f, 0.3f); 
+	}
+	return ReturnIconColor;
+}
+
+FLinearColor SGraphNode::GetNodeTitleTextColor() const
+{
+	FLinearColor ReturnTextColor = FLinearColor::White;
+	if(!GraphNode->bIsNodeEnabled)
+	{
+		ReturnTextColor *= FLinearColor(1.0f, 1.0f, 1.0f, 0.3f); 
+	}
+	return ReturnTextColor;
 }
 
 FSlateColor SGraphNode::GetNodeCommentColor() const
@@ -615,13 +653,16 @@ void SGraphNode::SetupErrorReporting()
 
 TSharedRef<SWidget> SGraphNode::CreateTitleWidget(TSharedPtr<SNodeTitle> NodeTitle)
 {
-	return SAssignNew(InlineEditableText, SInlineEditableTextBlock)
+	SAssignNew(InlineEditableText, SInlineEditableTextBlock)
 		.Style(FEditorStyle::Get(), "Graph.Node.NodeTitleInlineEditableText")
 		.Text(NodeTitle.Get(), &SNodeTitle::GetHeadTitle)
 		.OnVerifyTextChanged(this, &SGraphNode::OnVerifyNameTextChanged)
 		.OnTextCommitted(this, &SGraphNode::OnNameTextCommited)
 		.IsReadOnly(this, &SGraphNode::IsNameReadOnly)
 		.IsSelected(this, &SGraphNode::IsSelectedExclusively);
+	InlineEditableText->SetColorAndOpacity(TAttribute<FLinearColor>::Create(TAttribute<FLinearColor>::FGetter::CreateSP(this, &SGraphNode::GetNodeTitleTextColor)));
+
+	return InlineEditableText.ToSharedRef();
 }
 
 /**
@@ -654,7 +695,7 @@ void SGraphNode::UpdateGraphNode()
 	TSharedPtr<SNodeTitle> NodeTitle = SNew(SNodeTitle, GraphNode);
 
 	// Get node icon
-	FLinearColor IconColor = FLinearColor::White;
+	IconColor = FLinearColor::White;
 	const FSlateBrush* IconBrush = NULL;
 	if (GraphNode != NULL && GraphNode->ShowPaletteIconOnNode())
 	{
@@ -667,6 +708,7 @@ void SGraphNode::UpdateGraphNode()
 		[
 			SNew(SImage)
 			.Image( FEditorStyle::GetBrush("Graph.Node.TitleGloss") )
+			.ColorAndOpacity( this, &SGraphNode::GetNodeTitleIconColor )
 		]
 		+SOverlay::Slot()
 		.HAlign(HAlign_Left)
@@ -687,7 +729,7 @@ void SGraphNode::UpdateGraphNode()
 				[
 					SNew(SImage)
 					.Image(IconBrush)
-					.ColorAndOpacity(IconColor)
+					.ColorAndOpacity(this, &SGraphNode::GetNodeTitleIconColor)
 				]
 				+ SHorizontalBox::Slot()
 				[
@@ -711,6 +753,7 @@ void SGraphNode::UpdateGraphNode()
 			SNew(SBorder)
 			.Visibility(EVisibility::HitTestInvisible)			
 			.BorderImage( FEditorStyle::GetBrush( "Graph.Node.TitleHighlight" ) )
+			.BorderBackgroundColor( this, &SGraphNode::GetNodeTitleIconColor )
 			[
 				SNew(SSpacer)
 				.Size(FVector2D(20,20))
@@ -763,6 +806,7 @@ void SGraphNode::UpdateGraphNode()
 				[
 					SNew(SImage)
 					.Image(FEditorStyle::GetBrush("Graph.Node.Body"))
+					.ColorAndOpacity(this, &SGraphNode::GetNodeBodyColor)
 				]
 				+SOverlay::Slot()
 				[
diff --git a/Engine/Source/Editor/GraphEditor/Private/SGraphPin.cpp b/Engine/Source/Editor/GraphEditor/Private/SGraphPin.cpp
--- a/Engine/Source/Editor/GraphEditor/Private/SGraphPin.cpp
+++ b/Engine/Source/Editor/GraphEditor/Private/SGraphPin.cpp
@@ -151,6 +151,7 @@ SGraphPin::SGraphPin()
 
 void SGraphPin::Construct(const FArguments& InArgs, UEdGraphPin* InPin)
 {
+	bUsePinColorForText = InArgs._UsePinColorForText;
 	this->SetCursor( EMouseCursor::Default );
 
 	Visibility = TAttribute<EVisibility>(this, &SGraphPin::GetPinVisiblity);
@@ -186,17 +187,11 @@ void SGraphPin::Construct(const FArguments& InArgs, UEdGraphPin* InPin)
 			.Image( this, &SGraphPin::GetPinStatusIcon )
 		];
 
-	TAttribute<FSlateColor> OverrideTextColor;
-	if (InArgs._UsePinColorForText)
-	{
-		OverrideTextColor = TAttribute<FSlateColor>::Create(TAttribute<FSlateColor>::FGetter::CreateSP(this, &SGraphPin::GetPinColor));
-	}
-
 	TSharedRef<SWidget> LabelWidget = SNew(STextBlock)
 		.Text(this, &SGraphPin::GetPinLabel)
 		.TextStyle( FEditorStyle::Get(), InArgs._PinLabelStyle )
 		.Visibility(this, &SGraphPin::GetPinLabelVisibility)
-		.ColorAndOpacity(OverrideTextColor);
+		.ColorAndOpacity(this, &SGraphPin::GetPinTextColor);
 
 	// Create the widget used for the pin body (status indicator, label, and value)
 	TSharedRef<SWrapBox> LabelAndValue = 
@@ -879,9 +874,27 @@ FSlateColor SGraphPin::GetPinColor() const
 		return FSlateColor(FLinearColor(0.9f,0.2f,0.15f));
 	}
 	const UEdGraphSchema* Schema = GraphPinObj->GetSchema();
+	if(!GetPinObj()->GetOwningNode()->bIsNodeEnabled)
+	{
+		return Schema->GetPinTypeColor(GraphPinObj->PinType) * FLinearColor(1.0f, 1.0f, 1.0f, 0.5f);
+	}
+
 	return Schema->GetPinTypeColor(GraphPinObj->PinType) * PinColorModifier;
 }
 
+FSlateColor SGraphPin::GetPinTextColor() const
+{
+	if(!GetPinObj()->GetOwningNode()->bIsNodeEnabled)
+	{
+		return FLinearColor(1.0f, 1.0f, 1.0f, 0.5f);
+	}
+	else if(bUsePinColorForText)
+	{
+		return GetPinColor();
+	}
+	return FLinearColor::White;
+}
+
 
 const FSlateBrush* SGraphPin::GetPinStatusIcon() const
 {
diff --git a/Engine/Source/Editor/GraphEditor/Public/SCommentBubble.h b/Engine/Source/Editor/GraphEditor/Public/SCommentBubble.h
--- a/Engine/Source/Editor/GraphEditor/Public/SCommentBubble.h
+++ b/Engine/Source/Editor/GraphEditor/Public/SCommentBubble.h
@@ -95,6 +95,9 @@ protected:
 	/** Returns the color for the toggle bubble including the opacity value */
 	FSlateColor GetToggleButtonColor() const;
 
+	/** Returns the color of the main bubble */
+	FSlateColor GetBubbleColor() const;
+
 	/** Returns the current background color for the textbox */
 	FSlateColor GetTextBackgroundColor() const;
 	
diff --git a/Engine/Source/Editor/GraphEditor/Public/SGraphNode.h b/Engine/Source/Editor/GraphEditor/Public/SGraphNode.h
--- a/Engine/Source/Editor/GraphEditor/Public/SGraphNode.h
+++ b/Engine/Source/Editor/GraphEditor/Public/SGraphNode.h
@@ -105,12 +105,21 @@ public:
 	/** @return the editable title for a node */
 	FText GetEditableNodeTitleAsText() const;
 
-	/** @return the tint for the node's title */
+	/** @return the tint for the node's title image */
 	FSlateColor GetNodeTitleColor() const;
 
 	/** @return the tint for the node's comment */
 	FSlateColor GetNodeCommentColor() const;
 
+	/** @return the tint for the node's main body */
+	FSlateColor GetNodeBodyColor() const;
+
+	/** @return the tint for the node's title icon */
+	FSlateColor GetNodeTitleIconColor() const;
+
+	/** @return the tint for the node's title text */
+	FLinearColor GetNodeTitleTextColor() const;
+
 	/** @return the tooltip to display when over the node */
 	FText GetNodeTooltip() const;
 	
@@ -357,6 +366,9 @@ protected:
 	/** Caches true position of node */
 	FVector2D CachedUnscaledPosition;
 
+	/** Cached icon color for the node */
+	FLinearColor IconColor;
+
 	/** Cached pointer to graph editor settings */
 	const class UGraphEditorSettings* Settings;
 };
diff --git a/Engine/Source/Editor/GraphEditor/Public/SGraphPin.h b/Engine/Source/Editor/GraphEditor/Public/SGraphPin.h
--- a/Engine/Source/Editor/GraphEditor/Public/SGraphPin.h
+++ b/Engine/Source/Editor/GraphEditor/Public/SGraphPin.h
@@ -139,6 +139,9 @@ protected:
 	/** @return The color that we should use to draw this pin */
 	virtual FSlateColor GetPinColor() const;
 
+	/** @return The color that we should use to draw this pin's text */
+	virtual FSlateColor GetPinTextColor() const;
+
 	/** @return The tooltip to display for this pin */
 	FText GetTooltip() const;
 
@@ -184,6 +187,9 @@ protected:
 
 	TSet< TWeakObjectPtr<UEdGraphPin> > HoverPinSet;
 
+	/** TRUE if the pin should use the Pin's color for the text */
+	bool bUsePinColorForText;
+
 	//@TODO: Want to cache these once for all SGraphPins, but still handle slate style updates
 	mutable const FSlateBrush* CachedImg_ArrayPin_ConnectedHovered;
 	mutable const FSlateBrush* CachedImg_ArrayPin_Connected;
diff --git a/Engine/Source/Editor/Kismet/Private/BlueprintEditor.cpp b/Engine/Source/Editor/Kismet/Private/BlueprintEditor.cpp
--- a/Engine/Source/Editor/Kismet/Private/BlueprintEditor.cpp
+++ b/Engine/Source/Editor/Kismet/Private/BlueprintEditor.cpp
@@ -291,7 +291,7 @@ static bool BlueprintEditorImpl::GraphHasDefaultNode(UEdGraph const* InGraph)
 			continue;
 		}
 
-		if (Node->GetOutermost()->GetMetaData()->HasValue(Node, FNodeMetadata::DefaultGraphNode))
+		if (Node->GetOutermost()->GetMetaData()->HasValue(Node, FNodeMetadata::DefaultGraphNode) && Node->bIsNodeEnabled)
 		{
 			bHasDefaultNodes = true;
 			break;
diff --git a/Engine/Source/Editor/KismetCompiler/Private/KismetCompiler.cpp b/Engine/Source/Editor/KismetCompiler/Private/KismetCompiler.cpp
--- a/Engine/Source/Editor/KismetCompiler/Private/KismetCompiler.cpp
+++ b/Engine/Source/Editor/KismetCompiler/Private/KismetCompiler.cpp
@@ -2516,7 +2516,7 @@ void FKismetCompilerContext::MergeUbergraphPagesIn(UEdGraph* Ubergraph)
 		if (CompileOptions.bSaveIntermediateProducts)
 		{
 			TArray<UEdGraphNode*> ClonedNodeList;
-			FEdGraphUtilities::CloneAndMergeGraphIn(Ubergraph, SourceGraph, MessageLog, /*bRequireSchemaMatch=*/ true, &ClonedNodeList);
+			FEdGraphUtilities::CloneAndMergeGraphIn(Ubergraph, SourceGraph, MessageLog, /*bRequireSchemaMatch=*/ true, /*bIsCompiling*/ true, &ClonedNodeList);
 
 			// Create a comment block around the ubergrapgh contents before anything else got started
 			int32 OffsetX = 0;
@@ -2534,7 +2534,7 @@ void FKismetCompilerContext::MergeUbergraphPagesIn(UEdGraph* Ubergraph)
 		}
 		else
 		{
-			FEdGraphUtilities::CloneAndMergeGraphIn(Ubergraph, SourceGraph, MessageLog, /*bRequireSchemaMatch=*/ true);
+			FEdGraphUtilities::CloneAndMergeGraphIn(Ubergraph, SourceGraph, MessageLog, /*bRequireSchemaMatch=*/ true, /*bIsCompiling*/ true);
 		}
 	}
 }
diff --git a/Engine/Source/Editor/UnrealEd/Private/EdGraphUtilities.cpp b/Engine/Source/Editor/UnrealEd/Private/EdGraphUtilities.cpp
--- a/Engine/Source/Editor/UnrealEd/Private/EdGraphUtilities.cpp
+++ b/Engine/Source/Editor/UnrealEd/Private/EdGraphUtilities.cpp
@@ -186,7 +186,7 @@ UEdGraph* FEdGraphUtilities::CloneGraph(UEdGraph* InSource, UObject* NewOuter, F
 }
 
 // Clones the content from SourceGraph and merges it into MergeTarget; including merging/flattening all of the children from the SourceGraph into MergeTarget
-void FEdGraphUtilities::CloneAndMergeGraphIn(UEdGraph* MergeTarget, UEdGraph* SourceGraph, FCompilerResultsLog& MessageLog, bool bRequireSchemaMatch, TArray<UEdGraphNode*>* OutClonedNodes)
+void FEdGraphUtilities::CloneAndMergeGraphIn(UEdGraph* MergeTarget, UEdGraph* SourceGraph, FCompilerResultsLog& MessageLog, bool bRequireSchemaMatch, bool bInIsCompiling/* = false*/, TArray<UEdGraphNode*>* OutClonedNodes)
 {
 	// Clone the graph, then move all of it's children
 	UEdGraph* ClonedGraph = CloneGraph(SourceGraph, NULL, &MessageLog, true);
@@ -203,11 +203,11 @@ void FEdGraphUtilities::CloneAndMergeGraphIn(UEdGraph* MergeTarget, UEdGraph* So
 	const bool bIsLoading = Blueprint ? Blueprint->bIsRegeneratingOnLoad : false;
 
 	// Move them all to the destination
-	ClonedGraph->MoveNodesToAnotherGraph(MergeTarget, GIsAsyncLoading || bIsLoading);
+	ClonedGraph->MoveNodesToAnotherGraph(MergeTarget, GIsAsyncLoading || bIsLoading, bInIsCompiling);
 }
 
 // Moves the contents of all of the children graphs (recursively) into the target graph.  This does not clone, it's destructive to the source
-void FEdGraphUtilities::MergeChildrenGraphsIn(UEdGraph* MergeTarget, UEdGraph* ParentGraph, bool bRequireSchemaMatch)
+void FEdGraphUtilities::MergeChildrenGraphsIn(UEdGraph* MergeTarget, UEdGraph* ParentGraph, bool bRequireSchemaMatch, bool bInIsCompiling/* = false*/)
 {
 	// Determine if we are regenerating a blueprint on load
 	UBlueprint* Blueprint = FBlueprintEditorUtils::FindBlueprintForGraph(MergeTarget);
@@ -225,10 +225,10 @@ void FEdGraphUtilities::MergeChildrenGraphsIn(UEdGraph* MergeTarget, UEdGraph* P
 			// Even if we don't require a match to recurse, we do to actually copy the nodes
 			if (bSchemaMatches)
 			{
-				ChildGraph->MoveNodesToAnotherGraph(MergeTarget, GIsAsyncLoading || bIsLoading);
+				ChildGraph->MoveNodesToAnotherGraph(MergeTarget, GIsAsyncLoading || bIsLoading, bInIsCompiling);
 			}
 
-			MergeChildrenGraphsIn(MergeTarget, ChildGraph, bRequireSchemaMatch);
+			MergeChildrenGraphsIn(MergeTarget, ChildGraph, bRequireSchemaMatch, bInIsCompiling);
 		}
 	}
 }
diff --git a/Engine/Source/Editor/UnrealEd/Private/Kismet2/BlueprintEditorUtils.cpp b/Engine/Source/Editor/UnrealEd/Private/Kismet2/BlueprintEditorUtils.cpp
--- a/Engine/Source/Editor/UnrealEd/Private/Kismet2/BlueprintEditorUtils.cpp
+++ b/Engine/Source/Editor/UnrealEd/Private/Kismet2/BlueprintEditorUtils.cpp
@@ -2440,7 +2440,14 @@ bool FBlueprintEditorUtils::IsDataOnlyBlueprint(const UBlueprint* Blueprint)
 	UEdGraph* EventGraph = (Blueprint->UbergraphPages.Num() == 1) ? Blueprint->UbergraphPages[0] : NULL;
 	if( EventGraph && EventGraph->Nodes.Num() > 0 )
 	{
-		return false;
+		for(UEdGraphNode* GraphNode : EventGraph->Nodes)
+		{
+			// If there is an enabled node in the event graph, the Blueprint is not data only
+			if (GraphNode->bIsNodeEnabled)
+			{
+				return false;
+			}
+		}
 	}
 
 	// No implemented interfaces
diff --git a/Engine/Source/Editor/UnrealEd/Private/Kismet2/Kismet2.cpp b/Engine/Source/Editor/UnrealEd/Private/Kismet2/Kismet2.cpp
--- a/Engine/Source/Editor/UnrealEd/Private/Kismet2/Kismet2.cpp
+++ b/Engine/Source/Editor/UnrealEd/Private/Kismet2/Kismet2.cpp
@@ -469,16 +469,7 @@ UBlueprint* FKismetEditorUtilities::CreateBlueprint(UClass* ParentClass, UObject
 		if(Graph->Nodes.Num() == 0)
 		{
 			// add update event graph
-			UK2Node_Event* EventNode = NewObject<UK2Node_Event>(Graph);
-			EventNode->EventReference.SetExternalMember(FName(TEXT("BlueprintUpdateAnimation")), UAnimInstance::StaticClass());
-			EventNode->bOverrideFunction=true;
-			EventNode->CreateNewGuid();
-			EventNode->PostPlacedNewNode();
-			EventNode->SetFlags(RF_Transactional);
-			EventNode->AllocateDefaultPins();
-			UEdGraphSchema_K2::SetNodeMetaData(EventNode, FNodeMetadata::DefaultGraphNode);
-
-			Graph->AddNode(EventNode);
+			UK2Node_Event* EventNode = FKismetEditorUtilities::AddDefaultEventNode(NewBP, Graph, FName(TEXT("BlueprintUpdateAnimation")), UAnimInstance::StaticClass());
 
 			// add try get owner node
 			UK2Node_CallFunction* GetOwnerNode = NewObject<UK2Node_CallFunction>(Graph);
@@ -491,10 +482,41 @@ UBlueprint* FKismetEditorUtilities::CreateBlueprint(UClass* ParentClass, UObject
 			GetOwnerNode->NodePosX = EventNode->NodePosX;
 			GetOwnerNode->NodePosY = EventNode->NodePosY + EventNode->NodeHeight + 100;
 			UEdGraphSchema_K2::SetNodeMetaData(GetOwnerNode, FNodeMetadata::DefaultGraphNode);
+			GetOwnerNode->bIsNodeEnabled = false;
 
 			Graph->AddNode(GetOwnerNode);
 		}
 	}
+	else
+	{
+		// Based on the Blueprint type we are constructing, place some starting events.
+		// Note, this cannot happen in the Factories for constructing these Blueprint types due to the fact that creating child BPs circumvent the factories
+		UClass* WidgetClass = FindObject<UClass>(ANY_PACKAGE, TEXT("UserWidget"));
+		UClass* GameplayAbilityClass = FindObject<UClass>(ANY_PACKAGE, TEXT("GameplayAbility"));
+
+		if(NewBP->GeneratedClass->IsChildOf(AActor::StaticClass()))
+		{
+			UEdGraphNode* BeginPlayNode = FKismetEditorUtilities::AddDefaultEventNode(NewBP, NewBP->UbergraphPages[0], FName(TEXT("ReceiveBeginPlay")), AActor::StaticClass());
+			UEdGraphNode* ActorBeginOverlapNode = FKismetEditorUtilities::AddDefaultEventNode(NewBP, NewBP->UbergraphPages[0], FName(TEXT("ReceiveActorBeginOverlap")), AActor::StaticClass(), BeginPlayNode->NodePosY + BeginPlayNode->NodeHeight + 200);
+			FKismetEditorUtilities::AddDefaultEventNode(NewBP, NewBP->UbergraphPages[0], FName(TEXT("ReceiveTick")), AActor::StaticClass(), ActorBeginOverlapNode->NodePosY + ActorBeginOverlapNode->NodeHeight + 200);
+		}
+		else if(NewBP->GeneratedClass->IsChildOf(UActorComponent::StaticClass()))
+		{
+			UEdGraphNode* ReceiveTickNode = FKismetEditorUtilities::AddDefaultEventNode(NewBP, NewBP->UbergraphPages[0], FName(TEXT("ReceiveTick")), UActorComponent::StaticClass());
+			FKismetEditorUtilities::AddDefaultEventNode(NewBP, NewBP->UbergraphPages[0], FName(TEXT("ReceiveInitializeComponent")), UActorComponent::StaticClass(), ReceiveTickNode->NodePosY + ReceiveTickNode->NodeHeight + 200);
+		}
+		else if(NewBP->GeneratedClass->IsChildOf(WidgetClass))
+		{
+			UEdGraphNode* EventNode = FKismetEditorUtilities::AddDefaultEventNode(NewBP, NewBP->UbergraphPages[0], FName(TEXT("Construct")), WidgetClass);
+			FKismetEditorUtilities::AddDefaultEventNode(NewBP, NewBP->UbergraphPages[0], FName(TEXT("Tick")), WidgetClass, EventNode->NodePosY + EventNode->NodeHeight + 200);
+		}
+		else if(NewBP->GeneratedClass->IsChildOf(GameplayAbilityClass))
+		{
+			UEdGraphNode* EventNode = FKismetEditorUtilities::AddDefaultEventNode(NewBP, NewBP->UbergraphPages[0], FName(TEXT("K2_ActivateAbility")), GameplayAbilityClass);
+			FKismetEditorUtilities::AddDefaultEventNode(NewBP, NewBP->UbergraphPages[0], FName(TEXT("K2_OnEndAbility")), GameplayAbilityClass, EventNode->NodePosY + EventNode->NodeHeight + 200);
+		}
+	}
+	
 
 	// Report blueprint creation to analytics
 	if (FEngineAnalytics::IsAvailable())
@@ -524,6 +546,58 @@ UBlueprint* FKismetEditorUtilities::CreateBlueprint(UClass* ParentClass, UObject
 	return NewBP;
 }
 
+UK2Node_Event* FKismetEditorUtilities::AddDefaultEventNode(UBlueprint* InBlueprint, UEdGraph* InGraph, FName InEventName, UClass* InEventClass, int32 InNodePosY/* = 0*/)
+{
+	const UEdGraphSchema_K2* Schema = GetDefault<UEdGraphSchema_K2>();
+
+	// Add a "Begin Play" event
+	UK2Node_Event* NewEventNode = NewObject<UK2Node_Event>(InGraph);
+	NewEventNode->EventReference.SetExternalMember(InEventName, InEventClass);
+
+	// add update event graph
+	NewEventNode->bOverrideFunction=true;
+	NewEventNode->CreateNewGuid();
+	NewEventNode->PostPlacedNewNode();
+	NewEventNode->SetFlags(RF_Transactional);
+	NewEventNode->AllocateDefaultPins();
+	NewEventNode->bIsNodeEnabled = false;
+	NewEventNode->NodeComment = LOCTEXT("DisabledNodeComment", "This node is disabled and will not be called.\nDrag off pins to build functionality.").ToString();
+	NewEventNode->NodePosY = InNodePosY;
+	UEdGraphSchema_K2::SetNodeMetaData(NewEventNode, FNodeMetadata::DefaultGraphNode);
+
+	InGraph->AddNode(NewEventNode);
+
+	// Rebuild the skeleton class so we can determine if we are overriding the function from the parent
+	FBlueprintEditorUtils::MarkBlueprintAsStructurallyModified(InBlueprint);
+
+	// Get the function that the event node or function entry represents
+	FFunctionFromNodeHelper FunctionFromNode(NewEventNode);
+	if (FunctionFromNode.Function && Schema->GetCallableParentFunction(FunctionFromNode.Function))
+	{
+		UFunction* ValidParent = Schema->GetCallableParentFunction(FunctionFromNode.Function);
+		FGraphNodeCreator<UK2Node_CallParentFunction> FunctionNodeCreator(*InGraph);
+		UK2Node_CallParentFunction* ParentFunctionNode = FunctionNodeCreator.CreateNode();
+		ParentFunctionNode->SetFromFunction(ValidParent);
+		ParentFunctionNode->AllocateDefaultPins();
+
+		ParentFunctionNode->GetExecPin()->MakeLinkTo(NewEventNode->FindPin(Schema->PN_Then));
+
+		ParentFunctionNode->NodePosX = FunctionFromNode.Node->NodePosX + FunctionFromNode.Node->NodeWidth + 200;
+		ParentFunctionNode->NodePosY = FunctionFromNode.Node->NodePosY;
+		UEdGraphSchema_K2::SetNodeMetaData(ParentFunctionNode, FNodeMetadata::DefaultGraphNode);
+		FunctionNodeCreator.Finalize();
+
+		ParentFunctionNode->bIsNodeEnabled = false;
+
+		// Adding the call to parent and connecting it will reset this value
+		NewEventNode->bIsNodeEnabled = false;
+		NewEventNode->NodeComment = LOCTEXT("DisabledNodeComment", "This node is disabled and will not be called.\nDrag off pins to build functionality.").ToString();
+	}
+
+
+	return NewEventNode;
+}
+
 UBlueprint* FKismetEditorUtilities::ReloadBlueprint(UBlueprint* StaleBlueprint)
 {
 	check(StaleBlueprint->IsAsset());
diff --git a/Engine/Source/Editor/UnrealEd/Public/EdGraphUtilities.h b/Engine/Source/Editor/UnrealEd/Public/EdGraphUtilities.h
--- a/Engine/Source/Editor/UnrealEd/Public/EdGraphUtilities.h
+++ b/Engine/Source/Editor/UnrealEd/Public/EdGraphUtilities.h
@@ -50,14 +50,21 @@ public:
 	  * including merging/flattening all of the children from the SourceGraph into MergeTarget
 	  *
 	  * Also optionally populates *OutClonedNodes with the list of cloned objects.
+	  *
+	  * @param MergeTarget				The graph to merge the source graph into
+	  * @param SourceGraph				Graph to merge from
+	  * @param MessageLog				MessageLog to report errors and warnings to
+	  * @param bRequireSchemaMatch		TRUE if the graphs must have the same schema for the merge to occur
+	  * @param OutClonedNodes			Will populate with a full list of cloned nodes if provided
+	  * @param bInIsCompiling			TRUE if the function is being called during compilation, this will eliminate some nodes that will not be compiled
 	  */
-	static void CloneAndMergeGraphIn(UEdGraph* MergeTarget, UEdGraph* SourceGraph, FCompilerResultsLog& MessageLog, bool bRequireSchemaMatch, TArray<UEdGraphNode*>* OutClonedNodes = NULL);
+	static void CloneAndMergeGraphIn(UEdGraph* MergeTarget, UEdGraph* SourceGraph, FCompilerResultsLog& MessageLog, bool bRequireSchemaMatch, bool bInIsCompiling = false, TArray<UEdGraphNode*>* OutClonedNodes = NULL);
 
 	/**
 	 * Moves the contents of all of the children graphs of ParentGraph (recursively) into the MergeTarget graph.
 	 * This does not clone, it's destructive to the ParentGraph
 	 */
-	static void MergeChildrenGraphsIn(UEdGraph* MergeTarget, UEdGraph* ParentGraph, bool bRequireSchemaMatch);
+	static void MergeChildrenGraphsIn(UEdGraph* MergeTarget, UEdGraph* ParentGraph, bool bRequireSchemaMatch, bool bInIsCompiling = false);
 
 	/** Tries to rename the graph to have a name similar to BaseName */
 	static void RenameGraphCloseToName(UEdGraph* Graph, const FString& BaseName, int32 StartIndex = 1);
diff --git a/Engine/Source/Editor/UnrealEd/Public/Kismet2/KismetEditorUtilities.h b/Engine/Source/Editor/UnrealEd/Public/Kismet2/KismetEditorUtilities.h
--- a/Engine/Source/Editor/UnrealEd/Public/Kismet2/KismetEditorUtilities.h
+++ b/Engine/Source/Editor/UnrealEd/Public/Kismet2/KismetEditorUtilities.h
@@ -212,6 +212,19 @@ public:
 	 */
 	static void AddInterfaceTags(const UBlueprint* Blueprint, TArray<UObject::FAssetRegistryTag>& OutTags);
 
+	/**
+	 * Add a default event node to the graph, this node will also be in a disabled state and will spawn
+	 * with a call to it's parent if available
+	 *
+	 * @param InBlueprint		Blueprint this event will be a part of
+	 * @param InGraph			The graph to spawn the event node in
+	 * @param InEventName		The name of the event function
+	 * @param InEventClass		The class this event can be found in
+	 * @param InNodePosY		Optional Y-position to spawn the node at to easily spawn in a line
+	 * @return					The K2Node_Event will be returned
+	 */
+	static class UK2Node_Event* AddDefaultEventNode(UBlueprint* InBlueprint, UEdGraph* InGraph, FName InEventName, UClass* InEventClass, int32 InNodePosY = 0);
+
 private:
 	/** Stores whether we are already listening for kismet clicks */
 	static bool bIsListeningForClicksOnKismetLog;
diff --git a/Engine/Source/Runtime/Engine/Classes/EdGraph/EdGraph.h b/Engine/Source/Runtime/Engine/Classes/EdGraph/EdGraph.h
--- a/Engine/Source/Runtime/Engine/Classes/EdGraph/EdGraph.h
+++ b/Engine/Source/Runtime/Engine/Classes/EdGraph/EdGraph.h
@@ -151,8 +151,9 @@ public:
 	 * Move all nodes from this graph to another graph
 	 * @param DestinationGraph	The graph to move the nodes too 
 	 * @param bIsLoading		If true, the node move is occurring during a blueprint load
+	 * @param bInIsCompiling	TRUE if the function is being called during compilation, this will eliminate some nodes that will not be compiled
 	 */
-	void MoveNodesToAnotherGraph(UEdGraph* DestinationGraph, bool bIsLoading);
+	void MoveNodesToAnotherGraph(UEdGraph* DestinationGraph, bool bIsLoading, bool bInIsCompiling = false);
 
 	/** Finds all the nodes of a given minimum type in the graph */
 	template<class MinRequiredType, class ArrayElementType>
diff --git a/Engine/Source/Runtime/Engine/Classes/EdGraph/EdGraphNode.h b/Engine/Source/Runtime/Engine/Classes/EdGraph/EdGraphNode.h
--- a/Engine/Source/Runtime/Engine/Classes/EdGraph/EdGraphNode.h
+++ b/Engine/Source/Runtime/Engine/Classes/EdGraph/EdGraphNode.h
@@ -151,6 +151,10 @@ class ENGINE_API UEdGraphNode : public UObject
 	UPROPERTY()
 	TEnumAsByte<ENodeAdvancedPins::Type> AdvancedPinDisplay;
 
+	/** FALSE if the node is a disabled, which eliminates it from being compiled */
+	UPROPERTY()
+	bool bIsNodeEnabled;
+
 #if WITH_EDITOR
 
 private:
diff --git a/Engine/Source/Runtime/Engine/Private/EdGraph/EdGraph.cpp b/Engine/Source/Runtime/Engine/Private/EdGraph/EdGraph.cpp
--- a/Engine/Source/Runtime/Engine/Private/EdGraph/EdGraph.cpp
+++ b/Engine/Source/Runtime/Engine/Private/EdGraph/EdGraph.cpp
@@ -187,13 +187,20 @@ void UEdGraph::NotifyGraphChanged(const FEdGraphEditAction& InAction)
 	OnGraphChanged.Broadcast(InAction);
 }
 
-void UEdGraph::MoveNodesToAnotherGraph(UEdGraph* DestinationGraph, bool bIsLoading)
+void UEdGraph::MoveNodesToAnotherGraph(UEdGraph* DestinationGraph, bool bIsLoading, bool bInIsCompiling/* = false*/)
 {
 	// Move one node over at a time
 	while (Nodes.Num())
 	{
 		if (UEdGraphNode* Node = Nodes.Pop())
 		{
+			// During compilation, do not move ghost nodes, they are not used during compilation.
+			if(bInIsCompiling && !Node->bIsNodeEnabled)
+			{
+				// Break all node links, if any exist, do not move the node
+				Node->BreakAllNodeLinks();
+				continue;
+			}
 			// Let the name be autogenerated, to automatically avoid naming conflicts
 			// Since this graph is always going to come from a cloned source graph, user readable names can come from the remap stored in a MessageLog.
 	
diff --git a/Engine/Source/Runtime/Engine/Private/EdGraph/EdGraphNode.cpp b/Engine/Source/Runtime/Engine/Private/EdGraph/EdGraphNode.cpp
--- a/Engine/Source/Runtime/Engine/Private/EdGraph/EdGraphNode.cpp
+++ b/Engine/Source/Runtime/Engine/Private/EdGraph/EdGraphNode.cpp
@@ -42,6 +42,7 @@ FGraphNodeContextMenuBuilder::FGraphNodeContextMenuBuilder(const UEdGraph* InGra
 UEdGraphNode::UEdGraphNode(const FObjectInitializer& ObjectInitializer)
 	: Super(ObjectInitializer)
 	, AdvancedPinDisplay(ENodeAdvancedPins::NoPins)
+	, bIsNodeEnabled(true)
 {
 
 #if WITH_EDITORONLY_DATA
@@ -357,4 +358,4 @@ void UEdGraphNode::OnUpdateCommentText( const FString& NewComment )
 
 /////////////////////////////////////////////////////
 
-#undef LOCTEXT_NAMESPACE
\ No newline at end of file
+#undef LOCTEXT_NAMESPACE
diff --git a/Engine/Source/Runtime/Engine/Private/EdGraph/EdGraphPin.cpp b/Engine/Source/Runtime/Engine/Private/EdGraph/EdGraphPin.cpp
--- a/Engine/Source/Runtime/Engine/Private/EdGraph/EdGraphPin.cpp
+++ b/Engine/Source/Runtime/Engine/Private/EdGraph/EdGraphPin.cpp
@@ -66,6 +66,31 @@ bool FEdGraphPinType::Serialize(FArchive& Ar)
 	return true;
 }
 
+/////////////////////////////////////////////////////
+// GraphPinHelpers
+namespace GraphPinHelpers
+{
+	void EnableAllConnectedNodes(UEdGraphNode* InNode)
+	{
+		if(InNode && !InNode->bIsNodeEnabled)
+		{
+			// Enable the node and clear the comment
+			InNode->Modify();
+			InNode->bIsNodeEnabled = true;
+			InNode->NodeComment.Empty();
+
+			// Go through all pin connections and enable the nodes. Enabled nodes will prevent further iteration
+			for (UEdGraphPin* Pin : InNode->Pins)
+			{
+				for (UEdGraphPin* OtherPin : Pin->LinkedTo)
+				{
+					EnableAllConnectedNodes(OtherPin->GetOwningNode());
+				}
+			}
+		}
+	};
+}
+
 /////////////////////////////////////////////////////
 // UEdGraphPin
 
@@ -99,6 +124,9 @@ void UEdGraphPin::MakeLinkTo(UEdGraphPin* ToPin)
 			// Add to both lists
 			LinkedTo.Add(ToPin);
 			ToPin->LinkedTo.Add(this);
+
+			GraphPinHelpers::EnableAllConnectedNodes(MyNode);
+			GraphPinHelpers::EnableAllConnectedNodes(ToPin->GetOwningNode());
 		}
 	}
 }
