[INTEGRATE] Change 2442802 by [REDACTED_PERSON][REDACTED_HANDLE] on 2015/02/11 21:27:18

	Persona: Prevent crash when right-clicking on variable nodes in transition rules (invalid assumption about graphs with type GT_Function)
	[UE-9559]

diff --git a/Engine/Source/Editor/UnrealEd/Private/Kismet2/BlueprintEditorUtils.cpp b/Engine/Source/Editor/UnrealEd/Private/Kismet2/BlueprintEditorUtils.cpp
--- a/Engine/Source/Editor/UnrealEd/Private/Kismet2/BlueprintEditorUtils.cpp
+++ b/Engine/Source/Editor/UnrealEd/Private/Kismet2/BlueprintEditorUtils.cpp
@@ -3401,21 +3401,22 @@ void FBlueprintEditorUtils::GetNewVariablesOfType( const UBlueprint* Blueprint,
 
 void FBlueprintEditorUtils::GetLocalVariablesOfType( const UEdGraph* Graph, const FEdGraphPinType& Type, TArray<FName>& OutVars)
 {
-	if(Graph && Graph->GetSchema()->GetGraphType(Graph) == GT_Function)
+	if ((Graph != nullptr) && (Graph->GetSchema()->GetGraphType(Graph) == GT_Function))
 	{
 		TArray<UK2Node_FunctionEntry*> GraphNodes;
 		Graph->GetNodesOfClass<UK2Node_FunctionEntry>(GraphNodes);
 
-		bool bFoundLocalVariable = false;
-
-		// There is only ever 1 function entry
-		check(GraphNodes.Num() == 1);
-
-		for( auto& LocalVar : GraphNodes[0]->LocalVariables )
+		if (GraphNodes.Num() > 0)
 		{
-			if(LocalVar.VarType == Type)
+			// If there is an entry node, there should only be one
+			check(GraphNodes.Num() == 1);
+
+			for (const FBPVariableDescription& LocalVar : GraphNodes[0]->LocalVariables)
 			{
-				OutVars.Add(LocalVar.VarName);
+				if (LocalVar.VarType == Type)
+				{
+					OutVars.Add(LocalVar.VarName);
+				}
 			}
 		}
 	}
