Added a ShouldGatherBlueprintPinNames setting to GatherTextFromAssets commandlet configuration. This allows the commandlet to ignore the PinFriendlyName FText UProperty in UEdGraphPin, which is often gathered unnecessarily due to the need to gather other data from UEdGraphPin.

diff --git a/Engine/Config/Localization/Editor.ini b/Engine/Config/Localization/Editor.ini
--- a/Engine/Config/Localization/Editor.ini
+++ b/Engine/Config/Localization/Editor.ini
@@ -46,6 +46,7 @@ PackageExtensions=*.umap
 PackageExtensions=*.upk
 PackageExtensions=*.uasset
 bFixBroken=false
+ShouldGatherBlueprintPinNames=true
 
 ;Write Manifest
 [GatherTextStep3]
diff --git a/Engine/Config/Localization/Engine.ini b/Engine/Config/Localization/Engine.ini
--- a/Engine/Config/Localization/Engine.ini
+++ b/Engine/Config/Localization/Engine.ini
@@ -46,6 +46,7 @@ PackageExtensions=*.umap
 PackageExtensions=*.upk
 PackageExtensions=*.uasset
 bFixBroken=false
+ShouldGatherBlueprintPinNames=true
 
 ;Write Manifest
 [GatherTextStep2]
diff --git a/Engine/Source/Editor/UnrealEd/Classes/Commandlets/GatherTextFromAssetsCommandlet.h b/Engine/Source/Editor/UnrealEd/Classes/Commandlets/GatherTextFromAssetsCommandlet.h
--- a/Engine/Source/Editor/UnrealEd/Classes/Commandlets/GatherTextFromAssetsCommandlet.h
+++ b/Engine/Source/Editor/UnrealEd/Classes/Commandlets/GatherTextFromAssetsCommandlet.h
@@ -109,4 +109,5 @@ private:
 	FConflictTracker ConflictTracker;
 
 	bool bFixBroken;
+	bool ShouldGatherBlueprintPinNames;
 };
diff --git a/Engine/Source/Editor/UnrealEd/Private/Commandlets/GatherTextFromAssetsCommandlet.cpp b/Engine/Source/Editor/UnrealEd/Private/Commandlets/GatherTextFromAssetsCommandlet.cpp
--- a/Engine/Source/Editor/UnrealEd/Private/Commandlets/GatherTextFromAssetsCommandlet.cpp
+++ b/Engine/Source/Editor/UnrealEd/Private/Commandlets/GatherTextFromAssetsCommandlet.cpp
@@ -558,6 +558,17 @@ void UGatherTextFromAssetsCommandlet::ProcessTextProperty(UTextProperty* InTextP
 {
 	OutFixed = false;
 
+	// Early out of gathering pin friendly names unless requested to gather them.
+	if (Object->IsA(UEdGraphPin::StaticClass()))
+	{
+		UProperty* const PinFriendlyNameProperty = Cast<UTextProperty>(UEdGraphPin::StaticClass()->FindPropertyByName(GET_MEMBER_NAME_CHECKED(UEdGraphPin, PinFriendlyName)));
+		const bool IsPinNameProperty = PinFriendlyNameProperty == InTextProp;
+		if (!ShouldGatherBlueprintPinNames && IsPinNameProperty)
+		{
+			return;
+		}
+	}
+
 	FConflictTracker::FEntry NewEntry;
 	NewEntry.ObjectPath = Object->GetPathName();
 	NewEntry.SourceString = Data->GetSourceString();
@@ -794,7 +805,16 @@ int32 UGatherTextFromAssetsCommandlet::Main(const FString& Params)
 	}
 
 	//Get whether we should fix broken properties that we find.
-	GetBoolFromConfig(*SectionName, TEXT("bFixBroken"), bFixBroken, GatherTextConfigPath);
+	if (!GetBoolFromConfig(*SectionName, TEXT("bFixBroken"), bFixBroken, GatherTextConfigPath))
+	{
+		bFixBroken = false;
+	}
+
+	// Get whether we should gather pin name properties that we find. Typically only useful for the localization of UE4 itself.
+	if (!GetBoolFromConfig(*SectionName, TEXT("ShouldGatherBlueprintPinNames"), ShouldGatherBlueprintPinNames, GatherTextConfigPath))
+	{
+		ShouldGatherBlueprintPinNames = false;
+	}
 
 	// Add any manifest dependencies if they were provided
 	TArray<FString> ManifestDependenciesList;
