[INTEGRATE] Change 2444702 by [REDACTED_PERSON][REDACTED_HANDLE] on 2015/02/13 09:57:20

	UE-9699 - Fixing all the places we were setting the asset size to 0 explictly, the new default reasonable size is .1, all asset pickers now just use the default, unless they're the content browser.

diff --git a/Engine/Source/Developer/Merge/Private/SMergeAssetPickerView.cpp b/Engine/Source/Developer/Merge/Private/SMergeAssetPickerView.cpp
--- a/Engine/Source/Developer/Merge/Private/SMergeAssetPickerView.cpp
+++ b/Engine/Source/Developer/Merge/Private/SMergeAssetPickerView.cpp
@@ -255,7 +255,6 @@ TSharedRef<SWidget> SMergeAssetPickerPanel::MakeAssetPickerMenu()
 	AssetPickerConfig.bAllowNullSelection = true;
 	AssetPickerConfig.Filter.bRecursiveClasses = true;
 	AssetPickerConfig.OnAssetSelected = FOnAssetSelected::CreateSP(this, &SMergeAssetPickerPanel::OnAssetSelected);
-	AssetPickerConfig.ThumbnailScale = 0;
 	AssetPickerConfig.InitialAssetViewType = EAssetViewType::List;
 	AssetPickerConfig.bAllowDragging = false;
 
diff --git a/Engine/Source/Editor/Blutility/Private/BlutilityShelf.cpp b/Engine/Source/Editor/Blutility/Private/BlutilityShelf.cpp
--- a/Engine/Source/Editor/Blutility/Private/BlutilityShelf.cpp
+++ b/Engine/Source/Editor/Blutility/Private/BlutilityShelf.cpp
@@ -29,7 +29,6 @@ void SBlutilityShelf::BuildShelf()
 	FAssetPickerConfig AssetPickerConfig;
 	AssetPickerConfig.OnAssetDoubleClicked = FOnAssetDoubleClicked::CreateStatic(&SBlutilityShelf::OnBlutilityDoubleClicked);
 	AssetPickerConfig.OnGetAssetContextMenu = FOnGetAssetContextMenu::CreateSP(this, &SBlutilityShelf::OnBlutilityGetContextMenu);
-	AssetPickerConfig.ThumbnailScale = 0;
 	AssetPickerConfig.InitialAssetViewType = EAssetViewType::Tile;
 	AssetPickerConfig.bAllowNullSelection = false;
 	AssetPickerConfig.bShowBottomToolbar = false;
diff --git a/Engine/Source/Editor/ContentBrowser/Private/SAssetDialog.cpp b/Engine/Source/Editor/ContentBrowser/Private/SAssetDialog.cpp
--- a/Engine/Source/Editor/ContentBrowser/Private/SAssetDialog.cpp
+++ b/Engine/Source/Editor/ContentBrowser/Private/SAssetDialog.cpp
@@ -44,7 +44,6 @@ void SAssetDialog::Construct(const FArguments& InArgs, const FSharedAssetDialogC
 	AssetPickerConfig.Filter.PackagePaths.Add(FName(*DefaultPath));
 	AssetPickerConfig.bAllowDragging = false;
 	AssetPickerConfig.InitialAssetViewType = EAssetViewType::Tile;
-	AssetPickerConfig.ThumbnailScale = 0.125;
 	AssetPickerConfig.OnAssetSelected = FOnAssetSelected::CreateSP(this, &SAssetDialog::OnAssetSelected);
 	AssetPickerConfig.OnAssetsActivated = FOnAssetsActivated::CreateSP(this, &SAssetDialog::OnAssetsActivated);
 	AssetPickerConfig.SetFilterDelegates.Add(&SetFilterDelegate);
diff --git a/Engine/Source/Editor/ContentBrowser/Private/SAssetView.h b/Engine/Source/Editor/ContentBrowser/Private/SAssetView.h
--- a/Engine/Source/Editor/ContentBrowser/Private/SAssetView.h
+++ b/Engine/Source/Editor/ContentBrowser/Private/SAssetView.h
@@ -18,7 +18,7 @@ public:
 		, _ThumbnailLabel( EThumbnailLabel::ClassName )
 		, _AllowThumbnailHintLabel(true)
 		, _InitialViewType(EAssetViewType::Tile)
-		, _ThumbnailScale(0.25f) // A reasonable scale
+		, _ThumbnailScale(0.1f)
 		, _ShowBottomToolbar(true)
 		, _AllowThumbnailEditMode(false)
 		, _CanShowClasses(true)
diff --git a/Engine/Source/Editor/ContentBrowser/Public/IContentBrowserSingleton.h b/Engine/Source/Editor/ContentBrowser/Public/IContentBrowserSingleton.h
--- a/Engine/Source/Editor/ContentBrowser/Public/IContentBrowserSingleton.h
+++ b/Engine/Source/Editor/ContentBrowser/Public/IContentBrowserSingleton.h
@@ -141,7 +141,7 @@ struct FAssetPickerConfig
 	FAssetPickerConfig()
 		: SelectionMode( ESelectionMode::Multi )
 		, ThumbnailLabel( EThumbnailLabel::ClassName )
-		, ThumbnailScale(0.25f) // A reasonable scale
+		, ThumbnailScale(0.1f)
 		, InitialAssetViewType(EAssetViewType::Tile)
 		, bFocusSearchBoxWhenOpened(true)
 		, bAllowNullSelection(false)
diff --git a/Engine/Source/Editor/EditorStyle/Private/SlateEditorStyle.cpp b/Engine/Source/Editor/EditorStyle/Private/SlateEditorStyle.cpp
--- a/Engine/Source/Editor/EditorStyle/Private/SlateEditorStyle.cpp
+++ b/Engine/Source/Editor/EditorStyle/Private/SlateEditorStyle.cpp
@@ -357,6 +357,12 @@ void FSlateEditorStyle::FStyle::SetupGeneralStyles()
 			.SetPressed(BOX_BRUSH("Common/RoundedSelection_16x", 4.0f / 16.0f, SelectionColor_Pressed))
 			);
 
+		Set("FlatButton", FButtonStyle(Button)
+			.SetNormal(FSlateNoResource())
+			.SetHovered(BOX_BRUSH("Common/FlatButton", 2.0f / 8.0f, SelectionColor))
+			.SetPressed(BOX_BRUSH("Common/FlatButton", 2.0f / 8.0f, SelectionColor_Pressed))
+			);
+
 		Set("FlatButton.Dark", FButtonStyle(Button)
 			.SetNormal(BOX_BRUSH("Common/FlatButton", 2.0f / 8.0f, FLinearColor(0.125f, 0.125f, 0.125f, 0.8f)))
 			.SetHovered(BOX_BRUSH("Common/FlatButton", 2.0f / 8.0f, SelectionColor))
@@ -386,12 +392,14 @@ void FSlateEditorStyle::FStyle::SetupGeneralStyles()
 			FLinearColor Hovered;
 			FLinearColor Pressed;
 
-			ButtonColor(const FName& Name, const FLinearColor& Normal) : Name(Name), Normal(Normal)
+			ButtonColor(const FName& Name, const FLinearColor& Color) : Name(Name)
 			{
-				Hovered = Normal * 0.8f;
-				Hovered.A = Normal.A;
-				Pressed = Normal * 0.6f;
-				Pressed.A = Normal.A;
+				Normal = Color * 0.8f;
+				Normal.A = Color.A;
+				Hovered = Color * 1.0f;
+				Hovered.A = Color.A;
+				Pressed = Color * 0.6f;
+				Pressed.A = Color.A;
 			}
 		};
 
@@ -416,6 +424,7 @@ void FSlateEditorStyle::FStyle::SetupGeneralStyles()
 		Set("FontAwesome.10", TTF_FONT("Fonts/FontAwesome", 10));
 		Set("FontAwesome.11", TTF_FONT("Fonts/FontAwesome", 11));
 		Set("FontAwesome.12", TTF_FONT("Fonts/FontAwesome", 12));
+		Set("FontAwesome.14", TTF_FONT("Fonts/FontAwesome", 14));
 
 		/* Create a checkbox style for "ToggleButton" ... */
 		const FCheckBoxStyle ToggleButtonStyle = FCheckBoxStyle()
diff --git a/Engine/Source/Editor/GraphEditor/Private/KismetPins/SGraphPinObject.cpp b/Engine/Source/Editor/GraphEditor/Private/KismetPins/SGraphPinObject.cpp
--- a/Engine/Source/Editor/GraphEditor/Private/KismetPins/SGraphPinObject.cpp
+++ b/Engine/Source/Editor/GraphEditor/Private/KismetPins/SGraphPinObject.cpp
@@ -181,7 +181,6 @@ TSharedRef<SWidget> SGraphPinObject::GenerateAssetPicker()
 	AssetPickerConfig.bAllowNullSelection = true;
 	AssetPickerConfig.Filter.bRecursiveClasses = true;
 	AssetPickerConfig.OnAssetSelected = FOnAssetSelected::CreateSP(this, &SGraphPinObject::OnAssetSelectedFromPicker);
-	AssetPickerConfig.ThumbnailScale = 0;
 	AssetPickerConfig.InitialAssetViewType = EAssetViewType::List;
 	AssetPickerConfig.bAllowDragging = false;
 
diff --git a/Engine/Source/Editor/Kismet/Private/SSCSEditor.cpp b/Engine/Source/Editor/Kismet/Private/SSCSEditor.cpp
--- a/Engine/Source/Editor/Kismet/Private/SSCSEditor.cpp
+++ b/Engine/Source/Editor/Kismet/Private/SSCSEditor.cpp
@@ -3164,15 +3164,6 @@ void SSCSEditor::Construct( const FArguments& InArgs )
 					[
 						SNew(SHorizontalBox)
 						
-						//+SHorizontalBox::Slot()
-						//.VAlign(VAlign_Center)
-						//.AutoWidth()
-						//.Padding(2.f, 1.f)
-						//[
-						//	SNew(SImage)
-						//	.Image( FEditorStyle::Get().GetBrush( "ClassIcon.BlueprintCore" ) )
-						//]
-
 						+ SHorizontalBox::Slot()
 						.VAlign(VAlign_Center)
 						.Padding(3.f)
@@ -3200,34 +3191,36 @@ void SSCSEditor::Construct( const FArguments& InArgs )
 				.FillWidth(1.0f)
 				.Padding( 3.0f, 3.0f )
 				.HAlign(HAlign_Right)
+				.Padding(3.0f, 3.0f)
 				[
 					SNew(SComboButton)
 					.AddMetaData<FTagMetaData>(FTagMetaData(TEXT("Actor.EditBlueprint")))
 					.Visibility(this, &SSCSEditor::GetEditBlueprintButtonVisibility)
-					.ContentPadding(FMargin(0))
+					.ContentPadding(FMargin(10, 0))
 					.ComboButtonStyle(FEditorStyle::Get(), "ToolbarComboButton")
 					.ButtonStyle(FEditorStyle::Get(), "FlatButton.Primary")
 					.ForegroundColor(FLinearColor::White)
 					.ButtonContent()
 					[
 						SNew( SHorizontalBox )
+
 						+ SHorizontalBox::Slot()
 						.AutoWidth()
-						.HAlign(HAlign_Center)
 						.VAlign(VAlign_Center)
-						.Padding( 2.0f, 0.0f )
+						.Padding(3.f)
 						[
-							SNew( SImage )
-							.Image( FEditorStyle::Get().GetBrush( "ClassIcon.BlueprintCore" ) )
+							SNew(STextBlock)
+							.TextStyle(FEditorStyle::Get(), "ContentBrowser.TopBar.Font")
+							.Font(FEditorStyle::Get().GetFontStyle("FontAwesome.10"))
+							.Text(FString(TEXT("\xf085")) /*fa-cogs*/)
 						]
-						.Padding( 2.0f, 0.0f )
+						
 						+ SHorizontalBox::Slot()
 						[
 							SNew(STextBlock)
 							.TextStyle(FEditorStyle::Get(), "ContentBrowser.TopBar.Font")
 							.Text(LOCTEXT("EditBlueprint", "Edit Blueprint"))
 						]
-
 					]
 					.MenuContent()
 					[
diff --git a/Engine/Source/Editor/LevelEditor/Private/LevelEditorToolBar.cpp b/Engine/Source/Editor/LevelEditor/Private/LevelEditorToolBar.cpp
--- a/Engine/Source/Editor/LevelEditor/Private/LevelEditorToolBar.cpp
+++ b/Engine/Source/Editor/LevelEditor/Private/LevelEditorToolBar.cpp
@@ -1869,7 +1869,6 @@ TSharedRef< SWidget > FLevelEditorToolBar::GenerateOpenBlueprintMenuContent( TSh
 			FAssetPickerConfig Config;
 			Config.Filter.ClassNames.Add(UBlueprint::StaticClass()->GetFName());
 			Config.InitialAssetViewType = EAssetViewType::List;
-			Config.ThumbnailScale = 0; // make thumbnails as small as possible
 			Config.OnAssetSelected = FOnAssetSelected::CreateStatic(&FBlueprintMenus::OnBPSelected);
 			Config.bAllowDragging = false;
 			// Don't show stuff in Engine
diff --git a/Engine/Source/Editor/Matinee/Private/MatineeTrackHelpers.cpp b/Engine/Source/Editor/Matinee/Private/MatineeTrackHelpers.cpp
--- a/Engine/Source/Editor/Matinee/Private/MatineeTrackHelpers.cpp
+++ b/Engine/Source/Editor/Matinee/Private/MatineeTrackHelpers.cpp
@@ -253,7 +253,6 @@ bool UMatineeTrackAnimControlHelper::PreCreateKeyframe( UInterpTrack *Track, flo
 			AssetPickerConfig.OnAssetSelected = FOnAssetSelected::CreateUObject( this, &UMatineeTrackAnimControlHelper::OnAddKeyTextEntry, Mode->InterpEd, Track );
 			AssetPickerConfig.bAllowNullSelection = false;
 			AssetPickerConfig.InitialAssetViewType = EAssetViewType::List;
-			AssetPickerConfig.ThumbnailScale = 0.0f;
 
 			// Filter config
 			AssetPickerConfig.Filter.ClassNames.Add(UAnimSequence::StaticClass()->GetFName());
diff --git a/Engine/Source/Editor/Persona/Private/SRetargetSourceWindow.cpp b/Engine/Source/Editor/Persona/Private/SRetargetSourceWindow.cpp
--- a/Engine/Source/Editor/Persona/Private/SRetargetSourceWindow.cpp
+++ b/Engine/Source/Editor/Persona/Private/SRetargetSourceWindow.cpp
@@ -540,7 +540,6 @@ void SRetargetSourceWindow::OnAddRetargetSource()
 	AssetPickerConfig.OnAssetSelected = FOnAssetSelected::CreateSP(this, &SRetargetSourceWindow::OnAssetSelectedFromMeshPicker);
 	AssetPickerConfig.bAllowNullSelection = false;
 	AssetPickerConfig.InitialAssetViewType = EAssetViewType::Tile;
-	AssetPickerConfig.ThumbnailScale = 0.0f;
 
 	if(Skeleton)
 	{
diff --git a/Engine/Source/Editor/Persona/Private/SRigPicker.cpp b/Engine/Source/Editor/Persona/Private/SRigPicker.cpp
--- a/Engine/Source/Editor/Persona/Private/SRigPicker.cpp
+++ b/Engine/Source/Editor/Persona/Private/SRigPicker.cpp
@@ -57,8 +57,6 @@ void SRigPicker::Construct(const FArguments& InArgs)
 		AssetPickerConfig.Filter.bRecursiveClasses = true;
 		// Set a delegate for setting the asset from the picker
 		AssetPickerConfig.OnAssetSelected = FOnAssetSelected::CreateSP(this, &SRigPicker::OnAssetSelected);
-		// Use the smallest size thumbnails
-		AssetPickerConfig.ThumbnailScale = 0;
 		// Use the list view by default
 		AssetPickerConfig.InitialAssetViewType = EAssetViewType::List;
 		// The initial selection should be the current value
diff --git a/Engine/Source/Editor/PhAT/Private/PhAT.cpp b/Engine/Source/Editor/PhAT/Private/PhAT.cpp
--- a/Engine/Source/Editor/PhAT/Private/PhAT.cpp
+++ b/Engine/Source/Editor/PhAT/Private/PhAT.cpp
@@ -3037,7 +3037,6 @@ TSharedRef<SWidget> FPhAT::BuildStaticMeshAssetPicker()
 	AssetPickerConfig.OnAssetDoubleClicked = FOnAssetDoubleClicked::CreateSP(this, &FPhAT::OnAssetSelectedFromStaticMeshAssetPicker);
 	AssetPickerConfig.bAllowNullSelection = true;
 	AssetPickerConfig.InitialAssetViewType = EAssetViewType::List;
-	AssetPickerConfig.ThumbnailScale = 0.0f;
 	AssetPickerConfig.bFocusSearchBoxWhenOpened = true;
 	AssetPickerConfig.bShowBottomToolbar = false;
 	AssetPickerConfig.SelectionMode = ESelectionMode::Single;
diff --git a/Engine/Source/Editor/PropertyEditor/Private/UserInterface/PropertyEditor/SPropertyAssetPicker.cpp b/Engine/Source/Editor/PropertyEditor/Private/UserInterface/PropertyEditor/SPropertyAssetPicker.cpp
--- a/Engine/Source/Editor/PropertyEditor/Private/UserInterface/PropertyEditor/SPropertyAssetPicker.cpp
+++ b/Engine/Source/Editor/PropertyEditor/Private/UserInterface/PropertyEditor/SPropertyAssetPicker.cpp
@@ -72,8 +72,6 @@ TSharedRef<SWidget> SPropertyAssetPicker::OnGenerateAssetPicker()
 	AssetPickerConfig.Filter.bRecursiveClasses = true;
 	// Set a delegate for setting the asset from the picker
 	AssetPickerConfig.OnAssetSelected = FOnAssetSelected::CreateSP(this, &SPropertyAssetPicker::OnAssetSelectedFromPicker);
-	// Use the smallest size thumbnails
-	AssetPickerConfig.ThumbnailScale = 0;
 	AssetPickerConfig.bAllowDragging = false;
 	// Use the list view by default
 	AssetPickerConfig.InitialAssetViewType = EAssetViewType::List;
diff --git a/Engine/Source/Editor/PropertyEditor/Private/UserInterface/PropertyEditor/SPropertyMenuAssetPicker.cpp b/Engine/Source/Editor/PropertyEditor/Private/UserInterface/PropertyEditor/SPropertyMenuAssetPicker.cpp
--- a/Engine/Source/Editor/PropertyEditor/Private/UserInterface/PropertyEditor/SPropertyMenuAssetPicker.cpp
+++ b/Engine/Source/Editor/PropertyEditor/Private/UserInterface/PropertyEditor/SPropertyMenuAssetPicker.cpp
@@ -97,8 +97,6 @@ void SPropertyMenuAssetPicker::Construct( const FArguments& InArgs )
 		AssetPickerConfig.Filter.bRecursiveClasses = true;
 		// Set a delegate for setting the asset from the picker
 		AssetPickerConfig.OnAssetSelected = FOnAssetSelected::CreateSP(this, &SPropertyMenuAssetPicker::OnAssetSelected);
-		// Use small size thumbnails by default
-		AssetPickerConfig.ThumbnailScale = 0.1;
 		// Use the list view by default
 		AssetPickerConfig.InitialAssetViewType = EAssetViewType::List;
 		// The initial selection should be the current value
diff --git a/Engine/Source/Editor/UnrealEd/Private/Dialogs/SDeleteAssetsDialog.cpp b/Engine/Source/Editor/UnrealEd/Private/Dialogs/SDeleteAssetsDialog.cpp
--- a/Engine/Source/Editor/UnrealEd/Private/Dialogs/SDeleteAssetsDialog.cpp
+++ b/Engine/Source/Editor/UnrealEd/Private/Dialogs/SDeleteAssetsDialog.cpp
@@ -724,7 +724,6 @@ TSharedRef<SWidget> SDeleteAssetsDialog::MakeAssetViewForReferencerAssets()
 
 	AssetPickerConfig.AssetShowWarningText = TAttribute< FText >( this, &SDeleteAssetsDialog::GetReferencingAssetsEmptyText );
 
-	AssetPickerConfig.ThumbnailScale = 0.0f;
 	AssetPickerConfig.InitialAssetViewType = EAssetViewType::Tile;
 	AssetPickerConfig.OnAssetsActivated = FOnAssetsActivated::CreateSP(this, &SDeleteAssetsDialog::OnAssetsActivated);
 	AssetPickerConfig.OnShouldFilterAsset = FOnShouldFilterAsset::CreateSP(this, &SDeleteAssetsDialog::OnShouldFilterAsset);
@@ -742,7 +741,6 @@ TSharedRef<SWidget> SDeleteAssetsDialog::MakeConsolidationAssetPicker()
 	AssetPickerConfig.OnShouldFilterAsset = FOnShouldFilterAsset::CreateSP( this, &SDeleteAssetsDialog::OnShouldConsolidationFilterAsset );
 	AssetPickerConfig.bAllowNullSelection = false;
 	AssetPickerConfig.InitialAssetViewType = EAssetViewType::List;
-	AssetPickerConfig.ThumbnailScale = 0.0f;
 	AssetPickerConfig.bFocusSearchBoxWhenOpened = true;
 	AssetPickerConfig.bShowBottomToolbar = true;
 	AssetPickerConfig.bAllowDragging = false;
diff --git a/Engine/Source/Editor/UnrealEd/Private/Factories/AnimCompositeFactory.cpp b/Engine/Source/Editor/UnrealEd/Private/Factories/AnimCompositeFactory.cpp
--- a/Engine/Source/Editor/UnrealEd/Private/Factories/AnimCompositeFactory.cpp
+++ b/Engine/Source/Editor/UnrealEd/Private/Factories/AnimCompositeFactory.cpp
@@ -40,9 +40,6 @@ bool UAnimCompositeFactory::ConfigureProperties()
 	/** The default view mode should be a list view */
 	AssetPickerConfig.InitialAssetViewType = EAssetViewType::List;
 
-	/** The default scale for thumbnails. [0-1] range */
-	AssetPickerConfig.ThumbnailScale = 0.25f;
-
 	PickerWindow = SNew(SWindow)
 	.Title(LOCTEXT("CreateAnimCompositeOptions", "Pick Skeleton"))
 	.ClientSize(FVector2D(500, 600))
diff --git a/Engine/Source/Editor/UnrealEd/Private/Factories/AnimMontageFactory.cpp b/Engine/Source/Editor/UnrealEd/Private/Factories/AnimMontageFactory.cpp
--- a/Engine/Source/Editor/UnrealEd/Private/Factories/AnimMontageFactory.cpp
+++ b/Engine/Source/Editor/UnrealEd/Private/Factories/AnimMontageFactory.cpp
@@ -39,9 +39,6 @@ bool UAnimMontageFactory::ConfigureProperties()
 	/** The default view mode should be a list view */
 	AssetPickerConfig.InitialAssetViewType = EAssetViewType::List;
 
-	/** The default scale for thumbnails. [0-1] range */
-	AssetPickerConfig.ThumbnailScale = 0.25f;
-
 	PickerWindow = SNew(SWindow)
 	.Title(LOCTEXT("CreateAnimMontageOptions", "Pick Skeleton"))
 	.ClientSize(FVector2D(500, 600))
diff --git a/Engine/Source/Editor/UnrealEd/Private/Factories/AnimSequenceFactory.cpp b/Engine/Source/Editor/UnrealEd/Private/Factories/AnimSequenceFactory.cpp
--- a/Engine/Source/Editor/UnrealEd/Private/Factories/AnimSequenceFactory.cpp
+++ b/Engine/Source/Editor/UnrealEd/Private/Factories/AnimSequenceFactory.cpp
@@ -38,9 +38,6 @@ bool UAnimSequenceFactory::ConfigureProperties()
 	/** The default view mode should be a list view */
 	AssetPickerConfig.InitialAssetViewType = EAssetViewType::List;
 
-	/** The default scale for thumbnails. [0-1] range */
-	AssetPickerConfig.ThumbnailScale = 0.25f;
-
 	PickerWindow = SNew(SWindow)
 	.Title(LOCTEXT("CreateAnimSequenceOptions", "Pick Skeleton"))
 	.ClientSize(FVector2D(500, 600))
diff --git a/Engine/Source/Editor/UnrealEd/Private/Factories/EditorFactories.cpp b/Engine/Source/Editor/UnrealEd/Private/Factories/EditorFactories.cpp
--- a/Engine/Source/Editor/UnrealEd/Private/Factories/EditorFactories.cpp
+++ b/Engine/Source/Editor/UnrealEd/Private/Factories/EditorFactories.cpp
@@ -6845,27 +6845,25 @@ UBlendSpaceFactoryNew::UBlendSpaceFactoryNew(const FObjectInitializer& ObjectIni
 		}
 
 bool UBlendSpaceFactoryNew::ConfigureProperties()
-	{
+{
 	// Null the parent class so we can check for selection later
 	TargetSkeleton = nullptr;
 
-		// Load the content browser module to display an asset picker
+	// Load the content browser module to display an asset picker
 	FContentBrowserModule& ContentBrowserModule = FModuleManager::LoadModuleChecked<FContentBrowserModule>("ContentBrowser");
 
-		FAssetPickerConfig AssetPickerConfig;
+	FAssetPickerConfig AssetPickerConfig;
 
 	/** The asset picker will only show skeletal meshes */
-		AssetPickerConfig.Filter.ClassNames.Add(USkeleton::StaticClass()->GetFName());
-		AssetPickerConfig.Filter.bRecursiveClasses = true;
+	AssetPickerConfig.Filter.ClassNames.Add(USkeleton::StaticClass()->GetFName());
+	AssetPickerConfig.Filter.bRecursiveClasses = true;
 
-		/** The delegate that fires when an asset was selected */
+	/** The delegate that fires when an asset was selected */
 	AssetPickerConfig.OnAssetSelected = FOnAssetSelected::CreateUObject(this, &UBlendSpaceFactoryNew::OnTargetSkeletonSelected);
 
-		/** The default view mode should be a list view */
-		AssetPickerConfig.InitialAssetViewType = EAssetViewType::List;
+	/** The default view mode should be a list view */
+	AssetPickerConfig.InitialAssetViewType = EAssetViewType::List;
 
-		/** The default scale for thumbnails. [0-1] range */
-		AssetPickerConfig.ThumbnailScale = 0.25f;
 
 	PickerWindow = SNew(SWindow)
 	.Title(LOCTEXT("CreateBlendSpaceOptions", "Pick Skeleton"))
@@ -6922,7 +6920,7 @@ bool UBlendSpaceFactory1D::ConfigureProperties()
 	// Null the parent class so we can check for selection later
 	TargetSkeleton = nullptr;
 
-		// Load the content browser module to display an asset picker
+	// Load the content browser module to display an asset picker
 	FContentBrowserModule& ContentBrowserModule = FModuleManager::LoadModuleChecked<FContentBrowserModule>("ContentBrowser");
 
 	FAssetPickerConfig AssetPickerConfig;
@@ -6937,9 +6935,6 @@ bool UBlendSpaceFactory1D::ConfigureProperties()
 	/** The default view mode should be a list view */
 	AssetPickerConfig.InitialAssetViewType = EAssetViewType::List;
 
-	/** The default scale for thumbnails. [0-1] range */
-	AssetPickerConfig.ThumbnailScale = 0.25f;
-
 	PickerWindow = SNew(SWindow)
 	.Title(LOCTEXT("CreateBlendSpaceOptions", "Pick Skeleton"))
 	.ClientSize(FVector2D(500, 600))
diff --git a/Engine/Source/Editor/UnrealEd/Private/Factories/SkeletonFactory.cpp b/Engine/Source/Editor/UnrealEd/Private/Factories/SkeletonFactory.cpp
--- a/Engine/Source/Editor/UnrealEd/Private/Factories/SkeletonFactory.cpp
+++ b/Engine/Source/Editor/UnrealEd/Private/Factories/SkeletonFactory.cpp
@@ -43,9 +43,6 @@ bool USkeletonFactory::ConfigureProperties()
 	/** The default view mode should be a list view */
 	AssetPickerConfig.InitialAssetViewType = EAssetViewType::List;
 
-	/** The default scale for thumbnails. [0-1] range */
-	AssetPickerConfig.ThumbnailScale = 0.25f;
-
 	PickerWindow = SNew(SWindow)
 	.Title( LOCTEXT("CreateSkeletonOptions", "Pick Skeletal Mesh") )
 	.ClientSize(FVector2D(500, 600))
diff --git a/Engine/Source/Editor/UnrealEd/Private/SSkeletonWidget.cpp b/Engine/Source/Editor/UnrealEd/Private/SSkeletonWidget.cpp
--- a/Engine/Source/Editor/UnrealEd/Private/SSkeletonWidget.cpp
+++ b/Engine/Source/Editor/UnrealEd/Private/SSkeletonWidget.cpp
@@ -383,7 +383,6 @@ void SAnimationRemapSkeleton::UpdateAssetPicker()
 	AssetPickerConfig.OnAssetSelected = FOnAssetSelected::CreateSP(this, &SAnimationRemapSkeleton::OnAssetSelectedFromPicker);
 	AssetPickerConfig.bAllowNullSelection = false;
 	AssetPickerConfig.InitialAssetViewType = EAssetViewType::Column;
-	AssetPickerConfig.ThumbnailScale = 0.0f;
 	AssetPickerConfig.OnShouldFilterAsset = FOnShouldFilterAsset::CreateSP(this, &SAnimationRemapSkeleton::OnShouldFilterAsset);
 
 	FContentBrowserModule& ContentBrowserModule = FModuleManager::Get().LoadModuleChecked<FContentBrowserModule>(TEXT("ContentBrowser"));
diff --git a/Engine/Source/Editor/UnrealEd/Private/Toolkits/GlobalEditorCommonCommands.cpp b/Engine/Source/Editor/UnrealEd/Private/Toolkits/GlobalEditorCommonCommands.cpp
--- a/Engine/Source/Editor/UnrealEd/Private/Toolkits/GlobalEditorCommonCommands.cpp
+++ b/Engine/Source/Editor/UnrealEd/Private/Toolkits/GlobalEditorCommonCommands.cpp
@@ -25,7 +25,6 @@ public:
 		FAssetPickerConfig AssetPickerConfig;
 		AssetPickerConfig.OnAssetDoubleClicked = FOnAssetSelected::CreateSP(this, &SGlobalOpenAssetDialog::OnAssetSelectedFromPicker);
 		AssetPickerConfig.OnAssetEnterPressed = FOnAssetEnterPressed::CreateSP(this, &SGlobalOpenAssetDialog::OnPressedEnterOnAssetsInPicker);
-		AssetPickerConfig.ThumbnailScale = 0.1f;
 		AssetPickerConfig.InitialAssetViewType = EAssetViewType::List;
 		AssetPickerConfig.bAllowNullSelection = false;
 		AssetPickerConfig.bShowBottomToolbar = true;
diff --git a/Engine/Source/Editor/WorldBrowser/Private/Tiles/STiledLandscapeImportDlg.cpp b/Engine/Source/Editor/WorldBrowser/Private/Tiles/STiledLandscapeImportDlg.cpp
--- a/Engine/Source/Editor/WorldBrowser/Private/Tiles/STiledLandscapeImportDlg.cpp
+++ b/Engine/Source/Editor/WorldBrowser/Private/Tiles/STiledLandscapeImportDlg.cpp
@@ -256,8 +256,7 @@ TSharedRef<SWidget> STiledLandcapeImportDlg::CreateLandscapeMaterialPicker()
 	AssetPickerConfig.OnAssetSelected = FOnAssetSelected::CreateSP(this, &STiledLandcapeImportDlg::OnLandscapeMaterialChanged);
 	AssetPickerConfig.InitialAssetViewType = EAssetViewType::List;
 	AssetPickerConfig.bAllowNullSelection = true;
-	AssetPickerConfig.ThumbnailScale = 0.0f;
-		
+	
 	return SNew(SBox)
 		.WidthOverride(250)
 		.HeightOverride(300)
