[INTEGRATE] Change 2445543 by [REDACTED_PERSON][REDACTED_HANDLE] on 2015/02/13 17:36:17

	Blueprints: Disable the reinstancing fast path (bChangeDefaultValueWithoutReinstancing) for now, as it is unsafe when a property outside of the ubergraph frame was modified
	- Add some more debug information / logging (still disabled by default) to reinstancing and CPFUO
	[UE-9500]

diff --git a/Engine/Config/BaseEngine.ini b/Engine/Config/BaseEngine.ini
--- a/Engine/Config/BaseEngine.ini
+++ b/Engine/Config/BaseEngine.ini
@@ -1502,7 +1502,7 @@ PrintStringDuration=2.0
 bUseLocalGraphVariables=false
 bPersistentUberGraphFrame=true
 bReinstanceOnlyWhenNecessary=true
-bChangeDefaultValueWithoutReinstancing=true
+bChangeDefaultValueWithoutReinstancing=false
 bDisplaysLayout=false
 bNativeCodeGenerationTool=false
 bSkeletonInheritSkeletonClasses=false
diff --git a/Engine/Source/Editor/UnrealEd/Private/Kismet2/KismetReinstanceUtilities.cpp b/Engine/Source/Editor/UnrealEd/Private/Kismet2/KismetReinstanceUtilities.cpp
--- a/Engine/Source/Editor/UnrealEd/Private/Kismet2/KismetReinstanceUtilities.cpp
+++ b/Engine/Source/Editor/UnrealEd/Private/Kismet2/KismetReinstanceUtilities.cpp
@@ -287,32 +287,41 @@ void FBlueprintCompileReinstancer::ReinstanceObjects(bool bAlwaysReinstance)
 	if (ClassToReinstance && DuplicatedClass)
 	{
 		bool bShouldReinstance = true;
+		
+		// See if we need to do a full reinstance or can do the faster refresh path (when enabled or no values were modified, and the structures match)
 		if (!bAlwaysReinstance)
 		{
 			BP_SCOPED_COMPILER_EVENT_STAT(EKismetReinstancerStats_ReplaceClassNoReinsancing);
 
-			auto BPClassA = Cast<const UBlueprintGeneratedClass>(DuplicatedClass);
-			auto BPClassB = Cast<const UBlueprintGeneratedClass>(ClassToReinstance);
-			auto BP = Cast<const UBlueprint>(ClassToReinstance->ClassGeneratedBy);
+			const UBlueprintGeneratedClass* BPClassA = Cast<const UBlueprintGeneratedClass>(DuplicatedClass);
+			const UBlueprintGeneratedClass* BPClassB = Cast<const UBlueprintGeneratedClass>(ClassToReinstance);
+			const UBlueprint* BP = Cast<const UBlueprint>(ClassToReinstance->ClassGeneratedBy);
 
 			static const FBoolConfigValueHelper ChangeDefaultValueWithoutReinstancing(TEXT("Kismet"), TEXT("bChangeDefaultValueWithoutReinstancing"), GEngineIni);
-			const bool bTheSameDefaultValues = BP && ClassToReinstanceDefaultValuesCRC && (BP->CrcPreviousCompiledCDO == ClassToReinstanceDefaultValuesCRC);
-			const bool bTheSame = (ChangeDefaultValueWithoutReinstancing || bTheSameDefaultValues) && BPClassA && BPClassB && FStructUtils::TheSameLayout(BPClassA, BPClassB, true);
-			if (bTheSame)
+			const bool bTheSameDefaultValues = (BP != nullptr) && (ClassToReinstanceDefaultValuesCRC != 0) && (BP->CrcPreviousCompiledCDO == ClassToReinstanceDefaultValuesCRC);
+
+			const bool bTheSameLayout = (BPClassA != nullptr) && (BPClassB != nullptr) && FStructUtils::TheSameLayout(BPClassA, BPClassB, true);
+
+			const bool bAllowedToDoFastPath = (ChangeDefaultValueWithoutReinstancing || bTheSameDefaultValues) && bTheSameLayout;
+			if (bAllowedToDoFastPath)
 			{
-				UE_LOG(LogBlueprint, Log, TEXT("BlueprintCompileReinstancer: class '%s' is replaced without reinstancing (the optimized way)."), *GetPathNameSafe(ClassToReinstance));
+
+				UE_LOG(LogBlueprint, Log, TEXT("BlueprintCompileReinstancer: Doing a fast path refresh on class '%s'."), *GetPathNameSafe(ClassToReinstance));
 
 				TArray<UObject*> ObjectsToReplace;
-				GetObjectsOfClass(DuplicatedClass, ObjectsToReplace, false);
+				GetObjectsOfClass(DuplicatedClass, ObjectsToReplace, /*bIncludeDerivedClasses=*/ false);
 
 				const bool bIsActor = ClassToReinstance->IsChildOf<AActor>();
 				const bool bIsAnimInstance = ClassToReinstance->IsChildOf<UAnimInstance>();
 				const bool bIsComponent = ClassToReinstance->IsChildOf<UActorComponent>();
 				for (auto Obj : ObjectsToReplace)
 				{
+					UE_LOG(LogBlueprint, Log, TEXT("  Fast path is refreshing (not replacing) %s"), *Obj->GetFullName());
+
 					if ((!Obj->IsTemplate() || bIsComponent) && !Obj->IsPendingKill())
 					{
 						Obj->SetClass(ClassToReinstance);
+
 						if (bIsActor)
 						{
 							auto Actor = CastChecked<AActor>(Obj);
@@ -344,6 +353,7 @@ void FBlueprintCompileReinstancer::ReinstanceObjects(bool bAlwaysReinstance)
 
 		if (bShouldReinstance)
 		{
+			UE_LOG(LogBlueprint, Log, TEXT("BlueprintCompileReinstancer: Doing a full reinstance on class '%s'"), *GetPathNameSafe(ClassToReinstance));
 			ReplaceInstancesOfClass(DuplicatedClass, ClassToReinstance, OriginalCDO, &ObjectsThatShouldUseOldStuff);
 		}
 		else if (ClassToReinstance->IsChildOf<UActorComponent>())
diff --git a/Engine/Source/Runtime/Engine/Private/UnrealEngine.cpp b/Engine/Source/Runtime/Engine/Private/UnrealEngine.cpp
--- a/Engine/Source/Runtime/Engine/Private/UnrealEngine.cpp
+++ b/Engine/Source/Runtime/Engine/Private/UnrealEngine.cpp
@@ -10310,7 +10310,7 @@ void UEngine::CopyPropertiesForUnrelatedObjects(UObject* OldObject, UObject* New
 
 	// Bad idea to write data to an actor while its components are registered
 	AActor* NewActor = Cast<AActor>(NewObject);
-	if(NewActor != NULL)
+	if (NewActor != nullptr)
 	{
 		TInlineComponentArray<UActorComponent*> Components;
 		NewActor->GetComponents(Components);
@@ -10322,12 +10322,12 @@ void UEngine::CopyPropertiesForUnrelatedObjects(UObject* OldObject, UObject* New
 	}
 
 	// If the new object is an Actor, save the root component reference, to be restored later
-	USceneComponent* SavedRootComponent = NULL;
-	UObjectProperty* RootComponentProperty = NULL;
-	if(NewActor != NULL)
+	USceneComponent* SavedRootComponent = nullptr;
+	UObjectProperty* RootComponentProperty = nullptr;
+	if (NewActor != nullptr)
 	{
 		RootComponentProperty = FindField<UObjectProperty>(NewActor->GetClass(), "RootComponent");
-		if(RootComponentProperty != NULL)
+		if (RootComponentProperty != nullptr)
 		{
 			SavedRootComponent = Cast<USceneComponent>(RootComponentProperty->GetObjectPropertyValue_InContainer(NewActor));
 		}
@@ -10346,7 +10346,7 @@ void UEngine::CopyPropertiesForUnrelatedObjects(UObject* OldObject, UObject* New
 	{
 		// Find all instanced objects of the old CDO, and save off their modified properties to be later applied to the newly instanced objects of the new CDO
 		TArray<UObject*> Components;
-		OldObject->CollectDefaultSubobjects(Components,true);
+		OldObject->CollectDefaultSubobjects(Components, true);
 
 		for (int32 Index = 0; Index < Components.Num(); Index++)
 		{
@@ -10469,9 +10469,9 @@ void UEngine::CopyPropertiesForUnrelatedObjects(UObject* OldObject, UObject* New
 	}
 
 	// Restore the root component reference
-	if(NewActor != NULL)
+	if (NewActor != nullptr)
 	{
-		if(RootComponentProperty != NULL)
+		if (RootComponentProperty != nullptr)
 		{
 			RootComponentProperty->SetObjectPropertyValue_InContainer(NewActor, SavedRootComponent);
 		}
@@ -10481,15 +10481,15 @@ void UEngine::CopyPropertiesForUnrelatedObjects(UObject* OldObject, UObject* New
 
 	bool bDumpProperties = CVarDumpCopyPropertiesForUnrelatedObjects.GetValueOnGameThread() != 0;
 	// Uncomment the next line to debug CPFUO for a specific object:
-	// bDumpProperties |= (NewObject->GetName().InStr(TEXT("Charm_Vim")) != INDEX_NONE);
+	// bDumpProperties |= (NewObject->GetName().Find(TEXT("SpinTree")) != INDEX_NONE);
 	if (bDumpProperties)
 	{
-		DumpObject(TEXT("CopyPropertiesForUnrelatedObjects: Old"), OldObject);
-		DumpObject(TEXT("CopyPropertiesForUnrelatedObjects: New"), NewObject);
+		DumpObject(*FString::Printf(TEXT("CopyPropertiesForUnrelatedObjects: Old (%s)"), *OldObject->GetFullName()), OldObject);
+		DumpObject(*FString::Printf(TEXT("CopyPropertiesForUnrelatedObjects: New (%s)"), *NewObject->GetFullName()), NewObject);
 	}
 
 	// Now notify any tools that aren't already updated via the FArchiveReplaceObjectRef path
-	if( GEngine != NULL )
+	if (GEngine != nullptr)
 	{
 		GEngine->NotifyToolsOfObjectReplacement(ReferenceReplacementMap);
 	}
