[INTEGRATE] Change 2445555 by [REDACTED_PERSON][REDACTED_HANDLE] on 2015/02/13 17:41:40

	[UE-9831] - Adding variables to BP should no longer crash and world component selection is retained across BP reinstancing

diff --git a/Engine/Source/Editor/LevelEditor/Private/SActorDetails.cpp b/Engine/Source/Editor/LevelEditor/Private/SActorDetails.cpp
--- a/Engine/Source/Editor/LevelEditor/Private/SActorDetails.cpp
+++ b/Engine/Source/Editor/LevelEditor/Private/SActorDetails.cpp
@@ -163,6 +163,8 @@ void SActorDetails::SetObjects(const TArray<UObject*>& InObjects)
 				// Update the tree if a new actor is selected
 				if(GEditor->GetSelectedComponentCount() == 0)
 				{
+					// Enable the selection guard to prevent OnTreeSelectionChanged() from altering the editor's component selection
+					TGuardValue<bool> SelectionGuard(bSelectionGuard, true);
 					SCSEditor->UpdateTree();
 				}
 			}
diff --git a/Engine/Source/Editor/UnrealEd/Private/Editor.cpp b/Engine/Source/Editor/UnrealEd/Private/Editor.cpp
--- a/Engine/Source/Editor/UnrealEd/Private/Editor.cpp
+++ b/Engine/Source/Editor/UnrealEd/Private/Editor.cpp
@@ -6062,8 +6062,11 @@ void UEditorEngine::NotifyToolsOfObjectReplacement(const TMap<UObject*, UObject*
 			// If the component corresponds to a new instance in the map, update the selection accordingly
 			if (OldToNewInstanceMap.Contains(Component))
 			{
-				ComponentSelection->Deselect(Component);
-				SelectComponent(CastChecked<UActorComponent>(OldToNewInstanceMap[Component]), true, false);
+				if (UActorComponent* NewComponent = CastChecked<UActorComponent>(OldToNewInstanceMap[Component], ECastCheckedType::NullAllowed))
+				{
+					ComponentSelection->Deselect(Component);
+					SelectComponent(NewComponent, true, false);
+				}
 			}
 		}
 		ComponentSelection->EndBatchSelectOperation();
diff --git a/Engine/Source/Editor/UnrealEd/Private/Kismet2/KismetReinstanceUtilities.cpp b/Engine/Source/Editor/UnrealEd/Private/Kismet2/KismetReinstanceUtilities.cpp
--- a/Engine/Source/Editor/UnrealEd/Private/Kismet2/KismetReinstanceUtilities.cpp
+++ b/Engine/Source/Editor/UnrealEd/Private/Kismet2/KismetReinstanceUtilities.cpp
@@ -432,6 +432,14 @@ struct FActorReplacementHelper
 	{
 		CachedActorData = StaticCastSharedPtr<AActor::FActorTransactionAnnotation>(OldActor->GetTransactionAnnotation());
 		CacheAttachInfo(OldActor);
+
+		for (UActorComponent* OldActorComponent : OldActor->GetComponents())
+		{
+			if (OldActorComponent)
+			{
+				OldActorComponentNameMap.Add(OldActorComponent->GetFName(), OldActorComponent);
+			}
+		}
 	}
 
 	/**
@@ -492,6 +500,8 @@ private:
 
 	/** Holds actor component data, etc. that we use to apply */
 	TSharedPtr<AActor::FActorTransactionAnnotation> CachedActorData;
+
+	TMap<FName, UActorComponent*> OldActorComponentNameMap;
 };
 
 void FActorReplacementHelper::Finalize(const TMap<UObject*, UObject*>& OldToNewInstanceMap)
@@ -533,8 +543,21 @@ void FActorReplacementHelper::Finalize(const TMap<UObject*, UObject*>& OldToNewI
 
 	if (bSelectNewActor)
 	{
-		GEditor->SelectActor(NewActor, /*bInSelected =*/true, /*bNotify =*/false);
+		GEditor->SelectActor(NewActor, /*bInSelected =*/true, /*bNotify =*/true);
+	}
+
+	TMap<UObject*, UObject*> ConstructedComponentReplacementMap;
+	for (UActorComponent* NewActorComponent : NewActor->GetComponents())
+	{
+		if (NewActorComponent)
+		{
+			if (UActorComponent** OldActorComponent = OldActorComponentNameMap.Find(NewActorComponent->GetFName()))
+			{
+				ConstructedComponentReplacementMap.Add(*OldActorComponent, NewActorComponent);
+			}
+		}
 	}
+	GEditor->NotifyToolsOfObjectReplacement(ConstructedComponentReplacementMap);
 	
 	// Destroy actor and clear references.
 	NewActor->Modify();
