[INTEGRATE] Change 2445929 by [REDACTED_PERSON][REDACTED_HANDLE] on 2015/02/13 23:12:29

	Fix garbage collection crash if any component is selected during PIE or simulate

diff --git a/Engine/Source/Editor/UnrealEd/Private/PlayLevel.cpp b/Engine/Source/Editor/UnrealEd/Private/PlayLevel.cpp
--- a/Engine/Source/Editor/UnrealEd/Private/PlayLevel.cpp
+++ b/Engine/Source/Editor/UnrealEd/Private/PlayLevel.cpp
@@ -85,8 +85,9 @@ void UEditorEngine::EndPlayMap()
 	// Enable screensavers when ending PIE.
 	EnableScreenSaver( true );
 
-	// Move SelectedActors object back to the transient package.
+	// Move SelectedActors and SelectedComponents object back to the transient package.
 	GetSelectedActors()->Rename(nullptr, GetTransientPackage(), REN_DoNotDirty | REN_DontCreateRedirectors | REN_NonTransactional);
+	GetSelectedComponents()->Rename(nullptr,GetTransientPackage(), REN_DoNotDirty | REN_DontCreateRedirectors | REN_NonTransactional);
 
 	// Make a list of all the actors that should be selected
 	TArray<UObject *> SelectedActors;
@@ -125,6 +126,7 @@ void UEditorEngine::EndPlayMap()
 	GEditor->SelectNone( true, true, false );
 	GetSelectedActors()->DeselectAll();
 	GetSelectedObjects()->DeselectAll();
+	GetSelectedComponents()->DeselectAll();
 
 	// For every actor that was selected previously, make sure it's editor equivalent is selected
 	for ( int32 ActorIndex = 0; ActorIndex < SelectedActors.Num(); ++ActorIndex )
@@ -2585,6 +2587,7 @@ UGameInstance* UEditorEngine::CreatePIEGameInstance(int32 PIEInstance, bool bInS
 
 	// Make a list of all the selected actors
 	TArray<UObject *> SelectedActors;
+	TArray<UObject*> SelectedComponents;
 	for ( FSelectionIterator It( GetSelectedActorIterator() ); It; ++It )
 	{
 		AActor* Actor = static_cast<AActor*>( *It );
@@ -2596,10 +2599,12 @@ UGameInstance* UEditorEngine::CreatePIEGameInstance(int32 PIEInstance, bool bInS
 		}
 	}
 
+
 	// Unselect everything
 	GEditor->SelectNone( true, true, false );
 	GetSelectedActors()->DeselectAll();
 	GetSelectedObjects()->DeselectAll();
+	GetSelectedComponents()->DeselectAll();
 
 	// For every actor that was selected previously, make sure it's sim equivalent is selected
 	for ( int32 ActorIndex = 0; ActorIndex < SelectedActors.Num(); ++ActorIndex )
@@ -2620,6 +2625,7 @@ UGameInstance* UEditorEngine::CreatePIEGameInstance(int32 PIEInstance, bool bInS
 	// Move SelectedActors global object to the PIE package for the duration of the PIE session.
 	// This will stop any transactions on it from being saved during PIE.
 	GetSelectedActors()->Rename(nullptr, GWorld->GetOutermost(), REN_DoNotDirty | REN_DontCreateRedirectors | REN_NonTransactional);
+	GetSelectedComponents()->Rename(nullptr, GWorld->GetOutermost(), REN_DoNotDirty | REN_DontCreateRedirectors | REN_NonTransactional);
 
 	// For play in editor, this is the viewport widget where the game is being displayed
 	TSharedPtr<SViewport> PieViewportWidget;
