[INTEGRATE] Change 2446219 by [REDACTED_PERSON][REDACTED_HANDLE] on 2015/02/14 21:11:35

	[UE-9497] Clear the blueprint's generated class's CDO before deleting the blueprint, otherwise it gets redirected to the replacing type and that causes a an assertion

diff --git a/Engine/Source/Editor/UnrealEd/Private/ObjectTools.cpp b/Engine/Source/Editor/UnrealEd/Private/ObjectTools.cpp
--- a/Engine/Source/Editor/UnrealEd/Private/ObjectTools.cpp
+++ b/Engine/Source/Editor/UnrealEd/Private/ObjectTools.cpp
@@ -802,6 +802,14 @@ namespace ObjectTools
 				UPackage* CurObjPackage = CurObjToConsolidate->GetOutermost();
 				FName CurObjName = CurObjToConsolidate->GetFName();
 
+				// null out the CDO of our current generated class so that DeleteSingleObject does not find it and set it's ClassGeneratedBy to the replacing type.
+				// That would trigger a type mismatch assertion...
+				UBlueprint* BlueprintToConsolidate = Cast<UBlueprint>(CurObjToConsolidate);
+				if (BlueprintToConsolidateTo != nullptr && BlueprintToConsolidate != nullptr && BlueprintToConsolidate->GeneratedClass != nullptr)
+				{
+					BlueprintToConsolidate->GeneratedClass->ClassDefaultObject = nullptr;
+				}
+
 				// Attempt to delete the object that was consolidated
 				if ( DeleteSingleObject( CurObjToConsolidate ) )
 				{
