[INTEGRATE] Change 2447279 by [REDACTED_PERSON][REDACTED_HANDLE] on 2015/02/16 12:02:36

	Tweak the generated names for components created from a class to strip _C when necessary

diff --git a/Engine/Source/Editor/Kismet/Private/SSCSEditor.cpp b/Engine/Source/Editor/Kismet/Private/SSCSEditor.cpp
--- a/Engine/Source/Editor/Kismet/Private/SSCSEditor.cpp
+++ b/Engine/Source/Editor/Kismet/Private/SSCSEditor.cpp
@@ -4322,7 +4322,7 @@ UActorComponent* SSCSEditor::AddNewComponent( UClass* NewComponentClass, UObject
 		Blueprint->Modify();
 		SaveSCSCurrentState(Blueprint->SimpleConstructionScript);
 
-		FName NewVariableName = Asset != nullptr ? Asset->GetFName() : NAME_None;
+		const FName NewVariableName = (Asset ? FName(*FComponentEditorUtils::GenerateValidVariableNameFromAsset(Asset, nullptr)) : NAME_None);
 		NewComponent = AddNewNode(Blueprint->SimpleConstructionScript->CreateNode(NewComponentClass, NewVariableName), Asset, true);
 
 		if (ComponentTemplate)
diff --git a/Engine/Source/Editor/UnrealEd/Private/Kismet2/ComponentEditorUtils.cpp b/Engine/Source/Editor/UnrealEd/Private/Kismet2/ComponentEditorUtils.cpp
--- a/Engine/Source/Editor/UnrealEd/Private/Kismet2/ComponentEditorUtils.cpp
+++ b/Engine/Source/Editor/UnrealEd/Private/Kismet2/ComponentEditorUtils.cpp
@@ -172,17 +172,32 @@ FString FComponentEditorUtils::GenerateValidVariableName(TSubclassOf<UActorCompo
 
 FString FComponentEditorUtils::GenerateValidVariableNameFromAsset(UObject* Asset, AActor* ComponentOwner)
 {
-	check(ComponentOwner);
-
 	int32 Counter = 1;
 	FString AssetName = Asset->GetName();
 
+	UClass* Class = Cast<UClass>(Asset);
+	if (Class)
+	{
+		if (!Class->HasAnyClassFlags(CLASS_CompiledFromBlueprint))
+		{
+			AssetName.RemoveFromEnd(TEXT("Component"));
+		}
+		else
+		{
+			AssetName.RemoveFromEnd("_C");
+		}
+	}
+
 	// Try to create a name without any numerical suffix first
 	FString ComponentInstanceName = AssetName;
-	while (!IsComponentNameAvailable(ComponentInstanceName, ComponentOwner))
+
+	if (ComponentOwner)
 	{
-		// Assign the lowest possible numerical suffix
-		ComponentInstanceName = FString::Printf(TEXT("%s%d"), *AssetName, Counter++);
+		while (!IsComponentNameAvailable(ComponentInstanceName, ComponentOwner))
+		{
+			// Assign the lowest possible numerical suffix
+			ComponentInstanceName = FString::Printf(TEXT("%s%d"), *AssetName, Counter++);
+		}
 	}
 
 	return ComponentInstanceName;
