[INTEGRATE] Change 2449494 by [REDACTED_PERSON][REDACTED_HANDLE] on 2015/02/17 19:56:33

	Add a separate option in the packaging settings for whether to include debug files, and let it default to off. PDBs in packaged games aren't usually necessary, and they bloat the packaged game a lot. Also move the Full Rebuild setting out of the advanced section so people can see it. UE-9983.

diff --git a/Engine/Source/Editor/MainFrame/Private/Frame/MainFrameActions.cpp b/Engine/Source/Editor/MainFrame/Private/Frame/MainFrameActions.cpp
--- a/Engine/Source/Editor/MainFrame/Private/Frame/MainFrameActions.cpp
+++ b/Engine/Source/Editor/MainFrame/Private/Frame/MainFrameActions.cpp
@@ -532,7 +532,12 @@ void FMainFrameActionCallbacks::PackageProject( const FName InPlatformInfoName )
 
 	if (PackagingSettings->ForDistribution)
 	{
-		OptionalParams += TEXT(" -distribution -nodebuginfo");
+		OptionalParams += TEXT(" -distribution");
+	}
+
+	if (!PackagingSettings->IncludeDebugFiles)
+	{
+		OptionalParams += TEXT(" -nodebuginfo");
 	}
 
 	if (PlatformInfo->TargetPlatformName == FName("WindowsNoEditor") && PlatformInfo->PlatformFlavor == TEXT("Win32"))
diff --git a/Engine/Source/Editor/UnrealEd/Classes/Settings/ProjectPackagingSettings.h b/Engine/Source/Editor/UnrealEd/Classes/Settings/ProjectPackagingSettings.h
--- a/Engine/Source/Editor/UnrealEd/Classes/Settings/ProjectPackagingSettings.h
+++ b/Engine/Source/Editor/UnrealEd/Classes/Settings/ProjectPackagingSettings.h
@@ -67,7 +67,7 @@ public:
 	 * If disabled, only modified files will be built, which can improve iteration time.
 	 * Unless you iterate on packaging, we recommend full rebuilds when packaging.
 	 */
-	UPROPERTY(config, EditAnywhere, Category=Project, AdvancedDisplay)
+	UPROPERTY(config, EditAnywhere, Category=Project)
 	bool FullRebuild;
 
 	/**
@@ -78,6 +78,10 @@ public:
 	UPROPERTY(config, EditAnywhere, Category=Project)
 	bool ForDistribution;
 
+	/** If enabled, debug files will be included in the packaged game */
+	UPROPERTY(config, EditAnywhere, Category=Project)
+	bool IncludeDebugFiles;
+
 	/** If enabled, all content will be put into a single .pak file instead of many individual files (default = enabled). */
 	UPROPERTY(config, EditAnywhere, Category=Packaging)
 	bool UsePakFile;
