[INTEGRATE] Change 2449713 by [REDACTED_PERSON][REDACTED_HANDLE] on 2015/02/18 00:07:14

	[UE-9990]
	 - nullptr check for ActorComponent that is briefly parented to a non-actor while 'applying instance changes to blueprint'
	 - emergencing stomping of CachedInstanceData (which appears to be stale) while undoing changes to child actor component and then reperforming changes
	 - nullptr check for null entry when undoing changes to ChildActorComponent made in level editor

diff --git a/Engine/Source/Editor/UnrealEd/Private/EditorSelectUtils.cpp b/Engine/Source/Editor/UnrealEd/Private/EditorSelectUtils.cpp
--- a/Engine/Source/Editor/UnrealEd/Private/EditorSelectUtils.cpp
+++ b/Engine/Source/Editor/UnrealEd/Private/EditorSelectUtils.cpp
@@ -281,23 +281,25 @@ void UUnrealEdEngine::UpdatePivotLocationForSelection( bool bOnChange )
 		{
 			UActorComponent* Component = CastChecked<UActorComponent>(*It);
 			AActor* ComponentOwner = Component->GetOwner();
-			check(ComponentOwner);
 
-			auto SelectedActors = GetSelectedActors();
-			const bool bIsOwnerSelected = SelectedActors->IsSelected(ComponentOwner);
-			check(bIsOwnerSelected);
-
-			if (ComponentOwner->GetWorld() == GWorld)
+			if (ComponentOwner != nullptr)
 			{
-				SingleActor = ComponentOwner;
-				if (Component->IsA<USceneComponent>())
+				auto SelectedActors = GetSelectedActors();
+				const bool bIsOwnerSelected = SelectedActors->IsSelected(ComponentOwner);
+				check(bIsOwnerSelected);
+
+				if (ComponentOwner->GetWorld() == GWorld)
 				{
-					SingleComponent = CastChecked<USceneComponent>(Component);
+					SingleActor = ComponentOwner;
+					if (Component->IsA<USceneComponent>())
+					{
+						SingleComponent = CastChecked<USceneComponent>(Component);
+					}
+
+					const bool IsTemplate = ComponentOwner->IsTemplate();
+					const bool LevelLocked = !FLevelUtils::IsLevelLocked(ComponentOwner->GetLevel());
+					check(IsTemplate || LevelLocked);
 				}
-
-				const bool IsTemplate = ComponentOwner->IsTemplate();
-				const bool LevelLocked = !FLevelUtils::IsLevelLocked(ComponentOwner->GetLevel());
-				check(IsTemplate || LevelLocked);
 			}
 		}
 	}
diff --git a/Engine/Source/Runtime/Engine/Private/ChildActorComponent.cpp b/Engine/Source/Runtime/Engine/Private/ChildActorComponent.cpp
--- a/Engine/Source/Runtime/Engine/Private/ChildActorComponent.cpp
+++ b/Engine/Source/Runtime/Engine/Private/ChildActorComponent.cpp
@@ -49,9 +49,12 @@ void UChildActorComponent::PostEditUndo()
 	// to do right now, so this will go away when that is done.
 	for (USceneComponent*& Component : AttachChildren)
 	{
-		if (Component->IsPendingKill() && Component->GetOwner() == ChildActor)
+		if (Component)
 		{
-			Component = ChildActor->GetRootComponent();
+			if (Component->IsPendingKill() && Component->GetOwner() == ChildActor)
+			{
+				Component = ChildActor->GetRootComponent();
+			}
 		}
 	}
 	
@@ -258,7 +261,14 @@ void UChildActorComponent::DestroyChildActor()
 		// if still alive, destroy, otherwise just clear the pointer
 		if(!ChildActor->IsPendingKill())
 		{
+#if WITH_EDITOR
+			if (CachedInstanceData)
+			{
+				delete CachedInstanceData;
+			}
+#else
 			check(!CachedInstanceData);
+#endif
 			CachedInstanceData = new FChildActorComponentInstanceData(this);
 
 			UWorld* World = ChildActor->GetWorld();
