Fixed BSP mirror operation so that undo works correctly.

diff --git a/Engine/Source/Editor/UnrealEd/Private/UnrealEdSrv.cpp b/Engine/Source/Editor/UnrealEd/Private/UnrealEdSrv.cpp
--- a/Engine/Source/Editor/UnrealEd/Private/UnrealEdSrv.cpp
+++ b/Engine/Source/Editor/UnrealEd/Private/UnrealEdSrv.cpp
@@ -1769,7 +1769,9 @@ static void MirrorActors(const FVector& MirrorScale)
 		ABrush* Brush = Cast< ABrush >( Actor );
 		if( Brush && Brush->Brush )
 		{
+			Brush->Modify();
 			Brush->Brush->Modify();
+			Brush->Brush->Polys->Modify();
 
 			const FVector LocalToWorldOffset = ( Brush->GetActorLocation() - PivotLocation );
 			const FVector LocationOffset = ( LocalToWorldOffset * MirrorScale ) - LocalToWorldOffset;
@@ -1800,8 +1802,6 @@ static void MirrorActors(const FVector& MirrorScale)
 				Poly->Reverse();
 				Poly->CalcNormal();
 			}
-
-			Brush->UnregisterAllComponents();
 		}
 		else
 		{
@@ -2323,7 +2323,9 @@ bool UUnrealEdEngine::Exec_Actor( UWorld* InWorld, const TCHAR* Str, FOutputDevi
 		if( !MirrorScale.X )		MirrorScale.X = 1;
 		if( !MirrorScale.Y )		MirrorScale.Y = 1;
 		if( !MirrorScale.Z )		MirrorScale.Z = 1;
-		MirrorActors( MirrorScale );
+
+		const FScopedTransaction Transaction(NSLOCTEXT("UnrealEd", "MirroringActors", "Mirroring Actors"));
+		MirrorActors(MirrorScale);
 		RebuildAlteredBSP(); // Update the Bsp of any levels containing a modified brush
 		return true;
 	}
