Changed UEditorEngine::SnapToSocket to use USceneComponent::SnapTo so that it works generically with any component which handles sockets.

diff --git a/Engine/Source/Editor/UnrealEd/Classes/Editor/EditorEngine.h b/Engine/Source/Editor/UnrealEd/Classes/Editor/EditorEngine.h
--- a/Engine/Source/Editor/UnrealEd/Classes/Editor/EditorEngine.h
+++ b/Engine/Source/Editor/UnrealEd/Classes/Editor/EditorEngine.h
@@ -2520,36 +2520,6 @@ public:
 
 
 private:
-	/**
-	* snap a child actor to a named socket in a specific SkeletalMeshComponent
-	*
-	* @param ParentActor	Actor to parent the child to
-	* @param ChildActor	Actor being parented
-	* @param SkeletalMeshComponent	Component which has the Socket
-	* @param SocketName	An optional socket name to attach to in the parent (use NAME_None when there is no socket)
-	* @return true if successful, false otherwise
-	*/
-	bool AttachActorToComponent(AActor* ParentActor, AActor* ChildActor, USkeletalMeshComponent* SkeletalMeshComponent, const FName SocketName);
-	/**
-	* snap a child actor to a named socket in a specific StaticMeshComponent
-	*
-	* @param ChildActor	Actor being parented
-	* @param StaticMeshComponent	Component which has the Socket
-	* @param SocketName	An optional socket name to attach to in the parent (use NAME_None when there is no socket)
-	* @return true if successful, false otherwise
-	*/
-	bool AttachActorToComponent(AActor* ChildActor, UStaticMeshComponent* StaticMeshComponent, const FName SocketName);
-	/**
-	 * Utility method to try and snap a child actor to a named socket in the parent
-	 *
-	 * @param ParentActor	Actor to parent the child to
-	 * @param ChildActor	Actor being parented
-	 * @param SocketName	An optional socket name to attach to in the parent (use NAME_None when there is no socket)
-	 * @param Component		Actual Component which has the Socket, this component will be used when ParentActor is a blueprint actor
-	 * @return true if successful, false otherwise
-	 */
-	bool SnapToSocket( AActor* ParentActor, AActor* ChildActor, const FName SocketName, USceneComponent* Component=NULL );
-
 	/**
 	 * Delegate definition for the execute function that follows
 	 */
diff --git a/Engine/Source/Editor/UnrealEd/Private/Editor.cpp b/Engine/Source/Editor/UnrealEd/Private/Editor.cpp
--- a/Engine/Source/Editor/UnrealEd/Private/Editor.cpp
+++ b/Engine/Source/Editor/UnrealEd/Private/Editor.cpp
@@ -4046,92 +4046,6 @@ UBrushBuilder* UEditorEngine::FindBrushBuilder( UClass* BrushBuilderClass )
 	return Builder;
 }
 
-bool UEditorEngine::AttachActorToComponent(AActor* ParentActor, AActor* ChildActor, USkeletalMeshComponent* SkeletalMeshComponent, const FName SocketName)
-{
-	bool bAttachedToSocket = false;
-
-	if(SkeletalMeshComponent && SkeletalMeshComponent->SkeletalMesh)
-	{
-		USkeletalMeshSocket const* const Socket = SkeletalMeshComponent->SkeletalMesh->FindSocket(SocketName);
-		if (Socket)
-		{
-			bAttachedToSocket = Socket->AttachActor(ChildActor, SkeletalMeshComponent);
-		}
-		else
-		{
-			// now search bone, if bone exists, snap to the bone
-			int32 BoneIndex = SkeletalMeshComponent->SkeletalMesh->RefSkeleton.FindBoneIndex(SocketName);
-			if (BoneIndex != INDEX_NONE)
-			{
-				ChildActor->GetRootComponent()->SnapTo(ParentActor->GetRootComponent(), SocketName);
-				bAttachedToSocket = true;
-	
-	#if WITH_EDITOR
-				ChildActor->PreEditChange(NULL);
-				ChildActor->PostEditChange();
-	#endif // WITH_EDITOR
-			}
-		}
-	}
-
-	return bAttachedToSocket;
-}
-
-bool UEditorEngine::AttachActorToComponent(AActor* ChildActor, UStaticMeshComponent* StaticMeshComponent, const FName SocketName)
-{
-	bool bAttachedToSocket = false;
-
-	if(StaticMeshComponent && StaticMeshComponent->StaticMesh)
-	{
-		UStaticMeshSocket const* const Socket = StaticMeshComponent->StaticMesh->FindSocket(SocketName);
-		if (Socket)
-		{
-			bAttachedToSocket = Socket->AttachActor(ChildActor, StaticMeshComponent);
-		}
-	}
-
-	return bAttachedToSocket;
-}
-
-bool UEditorEngine::SnapToSocket (AActor* ParentActor, AActor* ChildActor, const FName SocketName, USceneComponent* Component)
-{
-	bool bAttachedToSocket = false;
-
-	ASkeletalMeshActor* ParentSkelMeshActor = Cast<ASkeletalMeshActor>(ParentActor);
-	if (ParentSkelMeshActor != NULL &&
-		ParentSkelMeshActor->GetSkeletalMeshComponent() &&
-		ParentSkelMeshActor->GetSkeletalMeshComponent()->SkeletalMesh)
-	{
-		bAttachedToSocket = AttachActorToComponent(ParentActor, ChildActor, ParentSkelMeshActor->GetSkeletalMeshComponent(), SocketName);
-	}
-
-	AStaticMeshActor* ParentStaticMeshActor = Cast<AStaticMeshActor>(ParentActor);
-	if (ParentStaticMeshActor != NULL &&
-		ParentStaticMeshActor->GetStaticMeshComponent() &&
-		ParentStaticMeshActor->GetStaticMeshComponent()->StaticMesh)
-	{
-		bAttachedToSocket = AttachActorToComponent(ChildActor, ParentStaticMeshActor->GetStaticMeshComponent(), SocketName);
-	}
-
-	// if the component is in a blueprint actor
-	if (Component && ParentActor->OwnsComponent(Component))
-	{
-		USkeletalMeshComponent* const SkeletalMeshComponent = Cast<USkeletalMeshComponent>(Component);
-		if (SkeletalMeshComponent)
-		{
-			bAttachedToSocket = AttachActorToComponent(ParentActor, ChildActor, SkeletalMeshComponent, SocketName);
-		}
-
-		UStaticMeshComponent* const StaticMeshComponent = Cast<UStaticMeshComponent>(Component);		
-		if (StaticMeshComponent)
-		{
-			bAttachedToSocket = AttachActorToComponent(ChildActor, StaticMeshComponent, SocketName);
-		}
-	}
-
-	return bAttachedToSocket;
-}
-
 void UEditorEngine::ParentActors( AActor* ParentActor, AActor* ChildActor, const FName SocketName, USceneComponent* Component)
 {
 	if (CanParentActors(ParentActor, ChildActor))
@@ -4164,17 +4078,12 @@ void UEditorEngine::ParentActors( AActor* ParentActor, AActor* ChildActor, const
 			ParentRoot->DetachFromParent(true);
 		}
 
-		// Try snapping to socket if requested
-		if ((SocketName != NAME_None) && SnapToSocket(ParentActor, ChildActor, SocketName, Component))
-		{
-			// Refresh editor if snapping to socket was successful
-			RedrawLevelEditingViewports();
-		}
-		else
-		{
-			// Perform general attachment
-			ChildRoot->AttachTo( ParentRoot, SocketName, EAttachLocation::KeepWorldPosition );
-		}
+		// Snap to socket if a valid socket name was provided, otherwise attach without changing the relative transform
+		const bool bValidSocketName = !SocketName.IsNone() && ParentRoot->DoesSocketExist(SocketName);
+		ChildRoot->AttachTo(ParentRoot, SocketName, bValidSocketName ? EAttachLocation::SnapToTarget : EAttachLocation::KeepWorldPosition);
+
+		// Refresh editor in case child was translated after snapping to socket
+		RedrawLevelEditingViewports();
 	}
 }
 
