Fixes for crashing when undoing BSP transactions.  The core problem is that building BSP must be in a transaction so that when undoing from a transaction the build is also undone.  This keeps everything in the correct previous state.  Automatic building of BSP after Undo is also no longer allowed because it can change array sizes for the previous transaction.  We now send a non-fatal ensure to trap any future use of bsp rebuilding outside of a transaction

diff --git a/Engine/Source/Editor/UnrealEd/Private/Editor.cpp b/Engine/Source/Editor/UnrealEd/Private/Editor.cpp
--- a/Engine/Source/Editor/UnrealEd/Private/Editor.cpp
+++ b/Engine/Source/Editor/UnrealEd/Private/Editor.cpp
@@ -974,11 +974,6 @@ void UEditorEngine::Tick( float DeltaSeconds, bool bIdleMode )
 		}
 	}
 
-	// Rebuild the BSPs this frame, providing we're not in the middle of a transaction
-	if( ABrush::NeedsRebuild() && GUndo == NULL )
-	{
-		RebuildAlteredBSP();
-	}
 
 	// Potentially rebuilds the streaming data.
 	EditorContext.World()->ConditionallyBuildStreamingData();
diff --git a/Engine/Source/Editor/UnrealEd/Private/EditorServer.cpp b/Engine/Source/Editor/UnrealEd/Private/EditorServer.cpp
--- a/Engine/Source/Editor/UnrealEd/Private/EditorServer.cpp
+++ b/Engine/Source/Editor/UnrealEd/Private/EditorServer.cpp
@@ -1609,71 +1609,79 @@ void UEditorEngine::RebuildModelFromBrushes(UModel* Model, bool bSelectedBrushes
 
 void UEditorEngine::RebuildAlteredBSP()
 {
-	// Early out if BSP auto-updating is disabled
-	if (!GetDefault<ULevelEditorMiscSettings>()->bBSPAutoUpdate)
+	if( GUndo && !GIsTransacting )
 	{
-		return;
-	}
-
-	FlushRenderingCommands();
+		// Early out if BSP auto-updating is disabled
+		if (!GetDefault<ULevelEditorMiscSettings>()->bBSPAutoUpdate)
+		{
+			return;
+		}
 
-	// A list of all the levels that need to be rebuilt
-	TArray< TWeakObjectPtr< ULevel > > LevelsToRebuild;
-	ABrush::NeedsRebuild( &LevelsToRebuild );
+		FlushRenderingCommands();
 
-	// Determine which levels need to be rebuilt
-	for (FSelectionIterator It(GetSelectedActorIterator()); It; ++It)
-	{
-		AActor* Actor = static_cast<AActor*>(*It);
-		checkSlow( Actor->IsA(AActor::StaticClass()) );
+		// A list of all the levels that need to be rebuilt
+		TArray< TWeakObjectPtr< ULevel > > LevelsToRebuild;
+		ABrush::NeedsRebuild(&LevelsToRebuild);
 
-		ABrush* SelectedBrush = Cast< ABrush >( Actor );
-		if (SelectedBrush && !FActorEditorUtils::IsABuilderBrush(Actor))
+		// Determine which levels need to be rebuilt
+		for (FSelectionIterator It(GetSelectedActorIterator()); It; ++It)
 		{
-			ULevel* Level = SelectedBrush->GetLevel();
-			if (Level)
+			AActor* Actor = static_cast<AActor*>(*It);
+			checkSlow(Actor->IsA(AActor::StaticClass()));
+
+			ABrush* SelectedBrush = Cast< ABrush >(Actor);
+			if (SelectedBrush && !FActorEditorUtils::IsABuilderBrush(Actor))
 			{
-				LevelsToRebuild.AddUnique(Level);
+				ULevel* Level = SelectedBrush->GetLevel();
+				if (Level)
+				{
+					LevelsToRebuild.AddUnique(Level);
+				}
 			}
-		}
-		else
-		{
-			// In addition to any selected brushes, any brushes attached to a selected actor should be rebuilt
-			TArray<AActor*> AttachedActors;
-			Actor->GetAttachedActors( AttachedActors );
-
-			const bool bExactClass = true;
-			TArray<AActor*> AttachedBrushes;
-			// Get any brush actors attached to the selected actor
-			if( ContainsObjectOfClass( AttachedActors, ABrush::StaticClass(), bExactClass, &AttachedBrushes ) )
+			else
 			{
-				for( int32 BrushIndex = 0; BrushIndex < AttachedBrushes.Num(); ++BrushIndex )
+				// In addition to any selected brushes, any brushes attached to a selected actor should be rebuilt
+				TArray<AActor*> AttachedActors;
+				Actor->GetAttachedActors(AttachedActors);
+
+				const bool bExactClass = true;
+				TArray<AActor*> AttachedBrushes;
+				// Get any brush actors attached to the selected actor
+				if (ContainsObjectOfClass(AttachedActors, ABrush::StaticClass(), bExactClass, &AttachedBrushes))
 				{
-					ULevel* Level = CastChecked<ABrush>( AttachedBrushes[BrushIndex] )->GetLevel();
-					if (Level)
+					for (int32 BrushIndex = 0; BrushIndex < AttachedBrushes.Num(); ++BrushIndex)
 					{
-						LevelsToRebuild.AddUnique(Level);
+						ULevel* Level = CastChecked<ABrush>(AttachedBrushes[BrushIndex])->GetLevel();
+						if (Level)
+						{
+							LevelsToRebuild.AddUnique(Level);
+						}
 					}
 				}
+
 			}
 
 		}
 
-	}
-
-	// Rebuild the levels
-	for (int32 LevelIdx = 0; LevelIdx < LevelsToRebuild.Num(); ++LevelIdx)
-	{
-		TWeakObjectPtr< ULevel > LevelToRebuild = LevelsToRebuild[LevelIdx];
-		if ( LevelToRebuild.IsValid() )
+		// Rebuild the levels
+		for (int32 LevelIdx = 0; LevelIdx < LevelsToRebuild.Num(); ++LevelIdx)
 		{
-			RebuildLevel(*LevelToRebuild.Get());
+			TWeakObjectPtr< ULevel > LevelToRebuild = LevelsToRebuild[LevelIdx];
+			if (LevelToRebuild.IsValid())
+			{
+				RebuildLevel(*LevelToRebuild.Get());
+			}
 		}
-	}
 
-	RedrawLevelEditingViewports();
+		RedrawLevelEditingViewports();
 
-	ABrush::OnRebuildDone();
+		ABrush::OnRebuildDone();
+	}
+	else
+	{
+ 		ensureMsgf(0, TEXT("Rebuild BSP ignored. Not in a transaction") );
+		ABrush::OnRebuildDone();
+	}
 }
 
 void UEditorEngine::BSPIntersectionHelper(UWorld* InWorld, ECsgOper Operation)
diff --git a/Engine/Source/Editor/UnrealEd/Private/EditorTransaction.cpp b/Engine/Source/Editor/UnrealEd/Private/EditorTransaction.cpp
--- a/Engine/Source/Editor/UnrealEd/Private/EditorTransaction.cpp
+++ b/Engine/Source/Editor/UnrealEd/Private/EditorTransaction.cpp
@@ -292,13 +292,6 @@ void FTransaction::Apply()
 			ChangedObject->PostEditUndo();
 		}
 	}
-	
-	// Rebuild BSP here instead of waiting for the next tick since
-	// multiple transaction events can occur in a single tick
-	if (ABrush::NeedsRebuild())
-	{
-		GEditor->RebuildAlteredBSP();
-	}
 
 	// Flip it.
 	if (bFlip)
diff --git a/Engine/Source/Runtime/Engine/Private/Brush.cpp b/Engine/Source/Runtime/Engine/Private/Brush.cpp
--- a/Engine/Source/Runtime/Engine/Private/Brush.cpp
+++ b/Engine/Source/Runtime/Engine/Private/Brush.cpp
@@ -56,7 +56,7 @@ void ABrush::PostEditChangeProperty(FPropertyChangedEvent& PropertyChangedEvent)
 		Brush->BuildBound();
 	}
 
-	if(IsStaticBrush())
+	if(IsStaticBrush() && !GIsTransacting)
 	{
 		// Trigger a csg rebuild if we're in the editor.
 		SetNeedRebuild(GetLevel());
