Fixes a crash with creating Macro Blueprint Libraries or Function Blueprint Libraries and attempting to spawn default events.

diff --git a/Engine/Source/Editor/UnrealEd/Private/Kismet2/Kismet2.cpp b/Engine/Source/Editor/UnrealEd/Private/Kismet2/Kismet2.cpp
--- a/Engine/Source/Editor/UnrealEd/Private/Kismet2/Kismet2.cpp
+++ b/Engine/Source/Editor/UnrealEd/Private/Kismet2/Kismet2.cpp
@@ -489,31 +489,35 @@ UBlueprint* FKismetEditorUtilities::CreateBlueprint(UClass* ParentClass, UObject
 	}
 	else
 	{
-		// Based on the Blueprint type we are constructing, place some starting events.
-		// Note, this cannot happen in the Factories for constructing these Blueprint types due to the fact that creating child BPs circumvent the factories
-		UClass* WidgetClass = FindObject<UClass>(ANY_PACKAGE, TEXT("UserWidget"));
-		UClass* GameplayAbilityClass = FindObject<UClass>(ANY_PACKAGE, TEXT("GameplayAbility"));
-
-		if(NewBP->GeneratedClass->IsChildOf(AActor::StaticClass()))
-		{
-			UEdGraphNode* BeginPlayNode = FKismetEditorUtilities::AddDefaultEventNode(NewBP, NewBP->UbergraphPages[0], FName(TEXT("ReceiveBeginPlay")), AActor::StaticClass());
-			UEdGraphNode* ActorBeginOverlapNode = FKismetEditorUtilities::AddDefaultEventNode(NewBP, NewBP->UbergraphPages[0], FName(TEXT("ReceiveActorBeginOverlap")), AActor::StaticClass(), BeginPlayNode->NodePosY + BeginPlayNode->NodeHeight + 200);
-			FKismetEditorUtilities::AddDefaultEventNode(NewBP, NewBP->UbergraphPages[0], FName(TEXT("ReceiveTick")), AActor::StaticClass(), ActorBeginOverlapNode->NodePosY + ActorBeginOverlapNode->NodeHeight + 200);
-		}
-		else if(NewBP->GeneratedClass->IsChildOf(UActorComponent::StaticClass()))
-		{
-			UEdGraphNode* ReceiveTickNode = FKismetEditorUtilities::AddDefaultEventNode(NewBP, NewBP->UbergraphPages[0], FName(TEXT("ReceiveTick")), UActorComponent::StaticClass());
-			FKismetEditorUtilities::AddDefaultEventNode(NewBP, NewBP->UbergraphPages[0], FName(TEXT("ReceiveInitializeComponent")), UActorComponent::StaticClass(), ReceiveTickNode->NodePosY + ReceiveTickNode->NodeHeight + 200);
-		}
-		else if(NewBP->GeneratedClass->IsChildOf(WidgetClass))
-		{
-			UEdGraphNode* EventNode = FKismetEditorUtilities::AddDefaultEventNode(NewBP, NewBP->UbergraphPages[0], FName(TEXT("Construct")), WidgetClass);
-			FKismetEditorUtilities::AddDefaultEventNode(NewBP, NewBP->UbergraphPages[0], FName(TEXT("Tick")), WidgetClass, EventNode->NodePosY + EventNode->NodeHeight + 200);
-		}
-		else if(NewBP->GeneratedClass->IsChildOf(GameplayAbilityClass))
+		// Only add default events if there is an ubergraph and they are supported
+		if(NewBP->UbergraphPages.Num() && FBlueprintEditorUtils::DoesSupportEventGraphs(NewBP))
 		{
-			UEdGraphNode* EventNode = FKismetEditorUtilities::AddDefaultEventNode(NewBP, NewBP->UbergraphPages[0], FName(TEXT("K2_ActivateAbility")), GameplayAbilityClass);
-			FKismetEditorUtilities::AddDefaultEventNode(NewBP, NewBP->UbergraphPages[0], FName(TEXT("K2_OnEndAbility")), GameplayAbilityClass, EventNode->NodePosY + EventNode->NodeHeight + 200);
+			// Based on the Blueprint type we are constructing, place some starting events.
+			// Note, this cannot happen in the Factories for constructing these Blueprint types due to the fact that creating child BPs circumvent the factories
+			UClass* WidgetClass = FindObject<UClass>(ANY_PACKAGE, TEXT("UserWidget"));
+			UClass* GameplayAbilityClass = FindObject<UClass>(ANY_PACKAGE, TEXT("GameplayAbility"));
+
+			if(NewBP->GeneratedClass->IsChildOf(AActor::StaticClass()))
+			{
+				UEdGraphNode* BeginPlayNode = FKismetEditorUtilities::AddDefaultEventNode(NewBP, NewBP->UbergraphPages[0], FName(TEXT("ReceiveBeginPlay")), AActor::StaticClass());
+				UEdGraphNode* ActorBeginOverlapNode = FKismetEditorUtilities::AddDefaultEventNode(NewBP, NewBP->UbergraphPages[0], FName(TEXT("ReceiveActorBeginOverlap")), AActor::StaticClass(), BeginPlayNode->NodePosY + BeginPlayNode->NodeHeight + 200);
+				FKismetEditorUtilities::AddDefaultEventNode(NewBP, NewBP->UbergraphPages[0], FName(TEXT("ReceiveTick")), AActor::StaticClass(), ActorBeginOverlapNode->NodePosY + ActorBeginOverlapNode->NodeHeight + 200);
+			}
+			else if(NewBP->GeneratedClass->IsChildOf(UActorComponent::StaticClass()))
+			{
+				UEdGraphNode* ReceiveTickNode = FKismetEditorUtilities::AddDefaultEventNode(NewBP, NewBP->UbergraphPages[0], FName(TEXT("ReceiveTick")), UActorComponent::StaticClass());
+				FKismetEditorUtilities::AddDefaultEventNode(NewBP, NewBP->UbergraphPages[0], FName(TEXT("ReceiveInitializeComponent")), UActorComponent::StaticClass(), ReceiveTickNode->NodePosY + ReceiveTickNode->NodeHeight + 200);
+			}
+			else if(NewBP->GeneratedClass->IsChildOf(WidgetClass))
+			{
+				UEdGraphNode* EventNode = FKismetEditorUtilities::AddDefaultEventNode(NewBP, NewBP->UbergraphPages[0], FName(TEXT("Construct")), WidgetClass);
+				FKismetEditorUtilities::AddDefaultEventNode(NewBP, NewBP->UbergraphPages[0], FName(TEXT("Tick")), WidgetClass, EventNode->NodePosY + EventNode->NodeHeight + 200);
+			}
+			else if(NewBP->GeneratedClass->IsChildOf(GameplayAbilityClass))
+			{
+				UEdGraphNode* EventNode = FKismetEditorUtilities::AddDefaultEventNode(NewBP, NewBP->UbergraphPages[0], FName(TEXT("K2_ActivateAbility")), GameplayAbilityClass);
+				FKismetEditorUtilities::AddDefaultEventNode(NewBP, NewBP->UbergraphPages[0], FName(TEXT("K2_OnEndAbility")), GameplayAbilityClass, EventNode->NodePosY + EventNode->NodeHeight + 200);
+			}
 		}
 	}
 	
