fix UE-9447 Maps and assets are not discovered when clicking on folders

diff --git a/Engine/Source/Editor/UnrealEd/Private/FileHelpers.cpp b/Engine/Source/Editor/UnrealEd/Private/FileHelpers.cpp
--- a/Engine/Source/Editor/UnrealEd/Private/FileHelpers.cpp
+++ b/Engine/Source/Editor/UnrealEd/Private/FileHelpers.cpp
@@ -1644,6 +1644,9 @@ void FEditorFileUtils::OpenLevelPickingDialog(const FOnLevelsChosen& OnLevelsCho
 			// No saved path, just use a reasonable default
 			DefaultPath = TEXT("/Game/Maps");
 		}
+
+		//OpenAssetDialog expects no trailing "/" so remove if necessary
+		DefaultPath.RemoveFromEnd(TEXT("/"));
 	}
 
 	FOpenAssetDialogConfig OpenAssetDialogConfig;
