Fixed unloaded Blueprint components having the wrong icon

UE-9331 - Missing icons for BP component classes when restarting editor

Moved the code to work out the correct class icon for a Blueprint or asset data into FClassIconFinder so that it can be shared between the asset thumbnails and the components context menu.

Added an extra IconClass member to FComponentClassComboEntry so that unloaded Blueprints can set this appropriately despite having a null ComponentClass.

diff --git a/Engine/Source/Editor/UnrealEd/Private/AssetThumbnail.cpp b/Engine/Source/Editor/UnrealEd/Private/AssetThumbnail.cpp
--- a/Engine/Source/Editor/UnrealEd/Private/AssetThumbnail.cpp
+++ b/Engine/Source/Editor/UnrealEd/Private/AssetThumbnail.cpp
@@ -79,7 +79,7 @@ public:
 
 		TSharedRef<SOverlay> OverlayWidget = SNew(SOverlay);
 
-		UpdateThumbnailClass(AssetData.AssetName, Class);
+		UpdateThumbnailClass();
 
 		ClassThumbnailBrushOverride = InArgs._ClassThumbnailBrushOverride;
 
@@ -224,38 +224,10 @@ public:
 
 	}
 
-	void UpdateThumbnailClass(const FName& InAssetName, UClass* InAssetClass)
+	void UpdateThumbnailClass()
 	{
-		ThumbnailClass = InAssetClass;
-		bIsClassType = false;
-
-		if( InAssetClass == UClass::StaticClass() )
-		{
-			ThumbnailClass = FindObject<UClass>(ANY_PACKAGE, *InAssetName.ToString());
-			bIsClassType = true;
-		}
-		else if( InAssetClass == UBlueprint::StaticClass() )
-		{
-			static const FName NativeParentClassTag = "NativeParentClass";
-			static const FName ParentClassTag = "ParentClass";
-
-			// We need to use the asset data to get the parent class as the blueprint may not be loaded
-			const FAssetData& AssetData = AssetThumbnail->GetAssetData();
-			const FString* ParentClassNamePtr = AssetData.TagsAndValues.Find(NativeParentClassTag);
-			if(!ParentClassNamePtr)
-			{
-				ParentClassNamePtr = AssetData.TagsAndValues.Find(ParentClassTag);
-			}
-			if(ParentClassNamePtr && !ParentClassNamePtr->IsEmpty())
-			{
-				UObject* Outer = nullptr;
-				FString ParentClassName = *ParentClassNamePtr;
-				ResolveName(Outer, ParentClassName, false, false);
-				ThumbnailClass = FindObject<UClass>(ANY_PACKAGE, *ParentClassName);
-			}
-
-			bIsClassType = true;
-		}
+		const FAssetData& AssetData = AssetThumbnail->GetAssetData();
+		ThumbnailClass = FClassIconFinder::GetIconClassForAssetData(AssetData, &bIsClassType);
 	}
 
 	FSlateColor GetHintBackgroundColor() const
@@ -336,7 +308,7 @@ private:
 			AssetTypeActions = AssetToolsModule.Get().GetAssetTypeActionsForClass(Class);
 		}
 
-		UpdateThumbnailClass(AssetData.AssetName, Class);
+		UpdateThumbnailClass();
 
 		AssetColor = FLinearColor(1.f, 1.f, 1.f, 1.f);
 		if ( AssetTypeActions.IsValid() )
diff --git a/Engine/Source/Editor/UnrealEd/Private/ClassIconFinder.cpp b/Engine/Source/Editor/UnrealEd/Private/ClassIconFinder.cpp
--- a/Engine/Source/Editor/UnrealEd/Private/ClassIconFinder.cpp
+++ b/Engine/Source/Editor/UnrealEd/Private/ClassIconFinder.cpp
@@ -4,6 +4,7 @@
 #include "UnrealEd.h"
 #include "SlateBasics.h"
 #include "ClassIconFinder.h"
+#include "AssetData.h"
 
 TArray < const ISlateStyle* > FClassIconFinder::Styles;
 
@@ -19,7 +20,7 @@ void FClassIconFinder::UnregisterIconSource(const ISlateStyle* StyleSet)
 
 const FSlateBrush* FClassIconFinder::LookupBrush(FName IconName)
 {
-	const FSlateBrush* IconBrush = NULL;
+	const FSlateBrush* IconBrush = nullptr;
 
 	for ( const ISlateStyle* Style : Styles )
 	{
@@ -36,7 +37,7 @@ const FSlateBrush* FClassIconFinder::LookupBrush(FName IconName)
 const FSlateBrush* FClassIconFinder::FindIconForActors(const TArray< TWeakObjectPtr<AActor> >& InActors, UClass*& CommonBaseClass)
 {
 	// Get the common base class of the selected actors
-	const FSlateBrush* CommonIcon = NULL;
+	const FSlateBrush* CommonIcon = nullptr;
 	for( int32 ActorIdx = 0; ActorIdx < InActors.Num(); ++ActorIdx )
 	{
 		TWeakObjectPtr<AActor> Actor = InActors[ActorIdx];
@@ -77,7 +78,7 @@ FName FClassIconFinder::FindIconNameForActor( const TWeakObjectPtr<AActor>& InAc
 	AActor* Actor = InActor.Get();
 	FName BrushName = NAME_None;
 
-	if ( Actor != NULL )
+	if ( Actor )
 	{
 		ABrush* Brush = Cast< ABrush >( Actor );
 		if ( Brush )
@@ -93,7 +94,7 @@ FName FClassIconFinder::FindIconNameForActor( const TWeakObjectPtr<AActor>& InAc
 		}
 
 		// Actor didn't specify an icon - fallback on the class icon
-		if ( BrushName == NAME_None )
+		if ( BrushName.IsNone() )
 		{
 			BrushName = FindIconNameForClass( Actor->GetClass() );
 		}
@@ -127,16 +128,85 @@ FName FClassIconFinder::FindThumbnailNameForClass(const UClass* InClass, const F
 	return FindIconNameImpl( InClass, InDefaultName, TEXT("ClassThumbnail") );
 }
 
+const UClass* FClassIconFinder::GetIconClassForBlueprint(const UBlueprint* InBlueprint)
+{
+	if ( !InBlueprint )
+	{
+		return nullptr;
+	}
+
+	// If we're loaded and have a generated class, just use that
+	if ( const UClass* GeneratedClass = Cast<UClass>(InBlueprint->GeneratedClass) )
+	{
+		return GeneratedClass;
+	}
+
+	// We don't have a generated class, so instead try and use the parent class from our meta-data
+	return GetIconClassForAssetData(FAssetData(InBlueprint));
+}
+
+const UClass* FClassIconFinder::GetIconClassForAssetData(const FAssetData& InAssetData, bool* bOutIsClassType)
+{
+	if ( bOutIsClassType )
+	{
+		*bOutIsClassType = false;
+	}
+
+	UClass* AssetClass = FindObject<UClass>(ANY_PACKAGE, *InAssetData.AssetClass.ToString());
+	if ( !AssetClass )
+	{
+		return nullptr;
+	}
+
+	if ( AssetClass == UClass::StaticClass() )
+	{
+		if ( bOutIsClassType )
+		{
+			*bOutIsClassType = true;
+		}
+
+		return FindObject<UClass>(ANY_PACKAGE, *InAssetData.AssetName.ToString());
+	}
+	
+	if ( AssetClass == UBlueprint::StaticClass() )
+	{
+		static const FName NativeParentClassTag = "NativeParentClass";
+		static const FName ParentClassTag = "ParentClass";
+
+		if ( bOutIsClassType )
+		{
+			*bOutIsClassType = true;
+		}
+
+		// We need to use the asset data to get the parent class as the blueprint may not be loaded
+		const FString* ParentClassNamePtr = InAssetData.TagsAndValues.Find(NativeParentClassTag);
+		if ( !ParentClassNamePtr )
+		{
+			ParentClassNamePtr = InAssetData.TagsAndValues.Find(ParentClassTag);
+		}
+		if ( ParentClassNamePtr && !ParentClassNamePtr->IsEmpty() )
+		{
+			UObject* Outer = nullptr;
+			FString ParentClassName = *ParentClassNamePtr;
+			ResolveName(Outer, ParentClassName, false, false);
+			return FindObject<UClass>(ANY_PACKAGE, *ParentClassName);
+		}
+	}
+
+	// Default to using the class for the asset type
+	return AssetClass;
+}
+
 FName FClassIconFinder::FindIconNameImpl(const UClass* InClass, const FName& InDefaultName, const TCHAR* StyleRoot )
 {
 	FName BrushName;
-	const FSlateBrush* Brush = NULL;
+	const FSlateBrush* Brush = nullptr;
 
-	if ( InClass != NULL )
+	if ( InClass )
 	{
 		// walk up class hierarchy until we find an icon
 		const UClass* CurrentClass = InClass;
-		while( Brush == NULL && CurrentClass )
+		while( !Brush && CurrentClass )
 		{
 			BrushName = *FString::Printf( TEXT( "%s.%s" ), StyleRoot, *CurrentClass->GetName() );
 			Brush = FClassIconFinder::LookupBrush( BrushName );
@@ -144,10 +214,10 @@ FName FClassIconFinder::FindIconNameImpl(const UClass* InClass, const FName& InD
 		}
 	}
 
-	if( Brush == NULL )
+	if( !Brush )
 	{
 		// If we didn't supply an override name for the default icon use default class icon.
-		if( InDefaultName == NAME_None )
+		if( InDefaultName.IsNone() )
 		{
 			BrushName = *FString::Printf( TEXT( "%s.Default" ), StyleRoot);
 		}
diff --git a/Engine/Source/Editor/UnrealEd/Private/ComponentTypeRegistry.cpp b/Engine/Source/Editor/UnrealEd/Private/ComponentTypeRegistry.cpp
--- a/Engine/Source/Editor/UnrealEd/Private/ComponentTypeRegistry.cpp
+++ b/Engine/Source/Editor/UnrealEd/Private/ComponentTypeRegistry.cpp
@@ -3,6 +3,7 @@
 #include "UnrealEd.h"
 
 #include "AssetRegistryModule.h"
+#include "ClassIconFinder.h"
 #include "ComponentTypeRegistry.h"
 #include "EdGraphSchema_K2.h"
 #include "KismetEditorUtilities.h"
@@ -264,22 +265,25 @@ void FComponentTypeRegistryData::ForceRefreshComponentList()
 
 		for (auto OnDiskClass : OnDiskClasses)
 		{
-			FName AssetPath;
+			FAssetData AssetData;
 			FString FixedString = OnDiskClass.ToString();
 			FixedString.RemoveFromEnd(TEXT("_C"));
 			for (auto Blueprint : BlueprintAssetData)
 			{
 				if (Blueprint.AssetName.ToString() == FixedString)
 				{
-					AssetPath = Blueprint.ObjectPath;
+					AssetData = Blueprint;
 					break;
 				}
 			}
 
-			FComponentTypeEntry Entry = { FixedString, AssetPath.ToString(), nullptr };
+			FComponentTypeEntry Entry = { FixedString, AssetData.ObjectPath.ToString(), nullptr };
 			ComponentTypeList.Add(Entry);
 
-			FComponentClassComboEntryPtr NewEntry(new FComponentClassComboEntry(BlueprintComponents, FixedString, AssetPath, bIncludeInFilter));
+			// The blueprint is unloaded, so we need to work out which icon to use for it using its asset data
+			const UClass* BlueprintIconClass = FClassIconFinder::GetIconClassForAssetData(AssetData);
+
+			FComponentClassComboEntryPtr NewEntry(new FComponentClassComboEntry(BlueprintComponents, FixedString, AssetData.ObjectPath, BlueprintIconClass, bIncludeInFilter));
 			SortedClassList.Add(NewEntry);
 		}
 	}
diff --git a/Engine/Source/Editor/UnrealEd/Private/SComponentClassCombo.cpp b/Engine/Source/Editor/UnrealEd/Private/SComponentClassCombo.cpp
--- a/Engine/Source/Editor/UnrealEd/Private/SComponentClassCombo.cpp
+++ b/Engine/Source/Editor/UnrealEd/Private/SComponentClassCombo.cpp
@@ -350,7 +350,7 @@ TSharedRef<ITableRow> SComponentClassCombo::GenerateAddComponentRow( FComponentC
 				.AutoWidth()
 				[
 					SNew(SImage)
-					.Image( FClassIconFinder::FindIconForClass( Entry->GetIconOverrideBrushName() == NAME_None ? Entry->GetComponentClass() : nullptr, Entry->GetIconOverrideBrushName() ) )
+					.Image( FClassIconFinder::FindIconForClass( Entry->GetIconOverrideBrushName() == NAME_None ? Entry->GetIconClass() : nullptr, Entry->GetIconOverrideBrushName() ) )
 				]
 				+SHorizontalBox::Slot()
 				.AutoWidth()
diff --git a/Engine/Source/Editor/UnrealEd/Public/ClassIconFinder.h b/Engine/Source/Editor/UnrealEd/Public/ClassIconFinder.h
--- a/Engine/Source/Editor/UnrealEd/Public/ClassIconFinder.h
+++ b/Engine/Source/Editor/UnrealEd/Public/ClassIconFinder.h
@@ -1,11 +1,9 @@
 // Copyright 1998-2015 Epic Games, Inc. All Rights Reserved.
 
-
-#ifndef __ActorIconFinder_h__
-#define __ActorIconFinder_h__
-
 #pragma once
 
+class FAssetData;
+
 class FClassIconFinder
 {
 public:
@@ -36,6 +34,12 @@ public:
 	/** Find the large thumbnail name to use for the supplied class */
 	UNREALED_API static FName FindThumbnailNameForClass(const UClass* InClass, const FName& InDefaultName = FName() );
 
+	/** Utility function to convert a Blueprint into the most suitable class possible for use by the icon finder */
+	UNREALED_API static const UClass* GetIconClassForBlueprint(const UBlueprint* InBlueprint);
+
+	/** Utility function to convert an asset into the most suitable class possible for use by the icon finder */
+	UNREALED_API static const UClass* GetIconClassForAssetData(const FAssetData& InAssetData, bool* bOutIsClassType = nullptr);
+
 private:
 	/** Find a thumbnail/icon name to use for the supplied class */
 	UNREALED_API static FName FindIconNameImpl(const UClass* InClass, const FName& InDefaultName = FName(), const TCHAR* StyleRoot = TEXT("ClassIcon") );
@@ -46,5 +50,3 @@ private:
 	/** The available style sets to source icons from */
 	static TArray < const ISlateStyle* > Styles;
 };
-
-#endif // __ActorIconFinder_h__
\ No newline at end of file
diff --git a/Engine/Source/Editor/UnrealEd/Public/SComponentClassCombo.h b/Engine/Source/Editor/UnrealEd/Public/SComponentClassCombo.h
--- a/Engine/Source/Editor/UnrealEd/Public/SComponentClassCombo.h
+++ b/Engine/Source/Editor/UnrealEd/Public/SComponentClassCombo.h
@@ -52,6 +52,7 @@ class FComponentClassComboEntry: public TSharedFromThis<FComponentClassComboEntr
 public:
 	FComponentClassComboEntry( const FString& InHeadingText, TSubclassOf<UActorComponent> InComponentClass, bool InIncludedInFilter, EComponentCreateAction::Type InComponentCreateAction, FComponentEntryCustomizationArgs InCustomizationArgs = FComponentEntryCustomizationArgs() )
 		: ComponentClass(InComponentClass)
+		, IconClass(InComponentClass)
 		, ComponentName()
 		, ComponentPath()
 		, HeadingText(InHeadingText)
@@ -60,8 +61,9 @@ public:
 		, CustomizationArgs(InCustomizationArgs)
 	{}
 
-	FComponentClassComboEntry(const FString& InHeadingText, const FString& InComponentName, FName InComponentPath, bool InIncludedInFilter)
-		: ComponentClass()
+	FComponentClassComboEntry(const FString& InHeadingText, const FString& InComponentName, FName InComponentPath, const UClass* InIconClass, bool InIncludedInFilter)
+		: ComponentClass(nullptr)
+		, IconClass(InIconClass)
 		, ComponentName(InComponentName)
 		, ComponentPath(InComponentPath)
 		, HeadingText(InHeadingText)
@@ -70,7 +72,8 @@ public:
 	{}
 
 	FComponentClassComboEntry( const FString& InHeadingText )
-		: ComponentClass(NULL)
+		: ComponentClass(nullptr)
+		, IconClass(nullptr)
 		, ComponentName()
 		, ComponentPath()
 		, HeadingText(InHeadingText)
@@ -78,12 +81,15 @@ public:
 	{}
 
 	FComponentClassComboEntry()
-		: ComponentClass(NULL)
+		: ComponentClass(nullptr)
+		, IconClass(nullptr)
 	{}
 
 
 	TSubclassOf<UActorComponent> GetComponentClass() const { return ComponentClass; }
 
+	const UClass* GetIconClass() const { return IconClass; }
+
 	const FString& GetHeadingText() const { return HeadingText; }
 
 	bool IsHeading() const
@@ -132,6 +138,7 @@ public:
 	int32 GetSortPriority() const { return CustomizationArgs.SortPriority; }
 private:
 	TSubclassOf<UActorComponent> ComponentClass;
+	const UClass* IconClass;
 	// For components that are not loaded we just keep the name of the component,
 	// loading occurs when the blueprint is spawned, which should also trigger a refresh
 	// of the component list:
