Default placed events in Blueprints correctly have their comment bubbles showing.

diff --git a/Engine/Source/Editor/UnrealEd/Private/Kismet2/Kismet2.cpp b/Engine/Source/Editor/UnrealEd/Private/Kismet2/Kismet2.cpp
--- a/Engine/Source/Editor/UnrealEd/Private/Kismet2/Kismet2.cpp
+++ b/Engine/Source/Editor/UnrealEd/Private/Kismet2/Kismet2.cpp
@@ -566,6 +566,8 @@ UK2Node_Event* FKismetEditorUtilities::AddDefaultEventNode(UBlueprint* InBluepri
 	NewEventNode->AllocateDefaultPins();
 	NewEventNode->bIsNodeEnabled = false;
 	NewEventNode->NodeComment = LOCTEXT("DisabledNodeComment", "This node is disabled and will not be called.\nDrag off pins to build functionality.").ToString();
+	NewEventNode->bCommentBubblePinned = true;
+	NewEventNode->bCommentBubbleVisible = true;
 	NewEventNode->NodePosY = InNodePosY;
 	UEdGraphSchema_K2::SetNodeMetaData(NewEventNode, FNodeMetadata::DefaultGraphNode);
 
