UE-9030 Compiling a BP super class (with a child instance in the level) crashes with REINST class mismatch

All children classes are recompiled before any object is reinstanced.

diff --git a/Engine/Config/BaseEngine.ini b/Engine/Config/BaseEngine.ini
--- a/Engine/Config/BaseEngine.ini
+++ b/Engine/Config/BaseEngine.ini
@@ -1511,6 +1511,7 @@ bOptimizeAdjacentStates=true
 bEnableInheritableComponents=true
 bDeferDependencyLoads=true
 bForceDisableCookedDependencyDeferring=false
+bFirstCompileChildrenThenReinstance=true
 
 [/Script/Engine.Blueprint]
 bRecompileOnLoad=true
diff --git a/Engine/Plugins/ScriptPlugin/Source/ScriptEditorPlugin/Private/ScriptBlueprintCompiler.cpp b/Engine/Plugins/ScriptPlugin/Source/ScriptEditorPlugin/Private/ScriptBlueprintCompiler.cpp
--- a/Engine/Plugins/ScriptPlugin/Source/ScriptEditorPlugin/Private/ScriptBlueprintCompiler.cpp
+++ b/Engine/Plugins/ScriptPlugin/Source/ScriptEditorPlugin/Private/ScriptBlueprintCompiler.cpp
@@ -224,7 +224,7 @@ void FScriptBlueprintCompiler::SpawnNewClass(const FString& NewClassName)
 	else
 	{
 		// Already existed, but wasn't linked in the Blueprint yet due to load ordering issues
-		FBlueprintCompileReinstancer GeneratedClassReinstancer(NewScriptBlueprintClass);
+		FBlueprintCompileReinstancer::Create(NewScriptBlueprintClass);
 	}
 	NewClass = NewScriptBlueprintClass;
 }
diff --git a/Engine/Source/Developer/HotReload/Private/HotReloadClassReinstancer.cpp b/Engine/Source/Developer/HotReload/Private/HotReloadClassReinstancer.cpp
--- a/Engine/Source/Developer/HotReload/Private/HotReloadClassReinstancer.cpp
+++ b/Engine/Source/Developer/HotReload/Private/HotReloadClassReinstancer.cpp
@@ -428,7 +428,7 @@ void FHotReloadClassReinstancer::UpdateDefaultProperties()
 
 void FHotReloadClassReinstancer::ReinstanceObjectsAndUpdateDefaults()
 {
-	ReinstanceObjects();
+	ReinstanceObjects(true);
 	UpdateDefaultProperties();
 }
 
diff --git a/Engine/Source/Editor/KismetCompiler/Private/AnimBlueprintCompiler.cpp b/Engine/Source/Editor/KismetCompiler/Private/AnimBlueprintCompiler.cpp
--- a/Engine/Source/Editor/KismetCompiler/Private/AnimBlueprintCompiler.cpp
+++ b/Engine/Source/Editor/KismetCompiler/Private/AnimBlueprintCompiler.cpp
@@ -1122,7 +1122,7 @@ void FAnimBlueprintCompiler::SpawnNewClass(const FString& NewClassName)
 	else
 	{
 		// Already existed, but wasn't linked in the Blueprint yet due to load ordering issues
-		FBlueprintCompileReinstancer GeneratedClassReinstancer(NewAnimBlueprintClass);
+		FBlueprintCompileReinstancer::Create(NewAnimBlueprintClass);
 	}
 	NewClass = NewAnimBlueprintClass;
 }
diff --git a/Engine/Source/Editor/KismetCompiler/Private/KismetCompiler.cpp b/Engine/Source/Editor/KismetCompiler/Private/KismetCompiler.cpp
--- a/Engine/Source/Editor/KismetCompiler/Private/KismetCompiler.cpp
+++ b/Engine/Source/Editor/KismetCompiler/Private/KismetCompiler.cpp
@@ -119,7 +119,7 @@ void FKismetCompilerContext::SpawnNewClass(const FString& NewClassName)
 	else
 	{
 		// Already existed, but wasn't linked in the Blueprint yet due to load ordering issues
-		FBlueprintCompileReinstancer GeneratedClassReinstancer(NewClass);
+		FBlueprintCompileReinstancer::Create(NewClass);
 	}
 }
 
diff --git a/Engine/Source/Editor/KismetCompiler/Private/KismetCompilerMisc.cpp b/Engine/Source/Editor/KismetCompiler/Private/KismetCompilerMisc.cpp
--- a/Engine/Source/Editor/KismetCompiler/Private/KismetCompilerMisc.cpp
+++ b/Engine/Source/Editor/KismetCompiler/Private/KismetCompilerMisc.cpp
@@ -301,7 +301,7 @@ void FKismetCompilerUtilities::ConsignToOblivion(UClass* OldClass, bool bForceNo
 	if (OldClass != NULL)
 	{
 		// Use the Kismet class reinstancer to ensure that the CDO and any existing instances of this class are cleaned up!
-		FBlueprintCompileReinstancer CTOResinstancer(OldClass);
+		auto CTOResinstancer = FBlueprintCompileReinstancer::Create(OldClass);
 
 		UPackage* OwnerOutermost = OldClass->GetOutermost();
 		if( OldClass->ClassDefaultObject )
diff --git a/Engine/Source/Editor/KismetCompiler/Private/KismetCompilerModule.cpp b/Engine/Source/Editor/KismetCompiler/Private/KismetCompilerModule.cpp
--- a/Engine/Source/Editor/KismetCompiler/Private/KismetCompilerModule.cpp
+++ b/Engine/Source/Editor/KismetCompiler/Private/KismetCompilerModule.cpp
@@ -101,7 +101,7 @@ void FKismet2CompilerModule::CompileStructure(class UUserDefinedStruct* Struct,
 }
 
 // Compiles a blueprint.
-void FKismet2CompilerModule::CompileBlueprint(class UBlueprint* Blueprint, const FKismetCompilerOptions& CompileOptions, FCompilerResultsLog& Results, FBlueprintCompileReinstancer* ParentReinstancer, TArray<UObject*>* ObjLoaded)
+void FKismet2CompilerModule::CompileBlueprint(class UBlueprint* Blueprint, const FKismetCompilerOptions& CompileOptions, FCompilerResultsLog& Results, TSharedPtr<FBlueprintCompileReinstancer> ParentReinstancer, TArray<UObject*>* ObjLoaded)
 {
 	SCOPE_SECONDS_COUNTER(GBlueprintCompileTime);
 	BP_SCOPED_COMPILER_EVENT_STAT(EKismetCompilerStats_CompileTime);
@@ -119,7 +119,7 @@ void FKismet2CompilerModule::CompileBlueprint(class UBlueprint* Blueprint, const
 	{
 		BP_SCOPED_COMPILER_EVENT_STAT(EKismetCompilerStats_CompileSkeletonClass);
 
-		FBlueprintCompileReinstancer SkeletonReinstancer(Blueprint->SkeletonGeneratedClass);
+		auto SkeletonReinstancer = FBlueprintCompileReinstancer::Create(Blueprint->SkeletonGeneratedClass);
 
 		FCompilerResultsLog SkeletonResults;
 		SkeletonResults.bSilentMode = true;
@@ -158,20 +158,18 @@ void FKismet2CompilerModule::CompileBlueprint(class UBlueprint* Blueprint, const
 			// There were errors.  Compile the generated class to have function stubs
 			Blueprint->Status = BS_Error;
 
-			static const FBoolConfigValueHelper ReinstanceOnlyWhenNecessary(TEXT("Kismet"), TEXT("bReinstanceOnlyWhenNecessary"), GEngineIni);
-
 			// Reinstance objects here, so we can preserve their memory layouts to reinstance them again
-			if( ParentReinstancer != NULL )
+			if (ParentReinstancer.IsValid())
 			{
 				ParentReinstancer->UpdateBytecodeReferences();
 
 				if(!Blueprint->bIsRegeneratingOnLoad)
 				{
-					ParentReinstancer->ReinstanceObjects(!ReinstanceOnlyWhenNecessary);
+					ParentReinstancer->ReinstanceObjects();
 				}
 			}
 
-			FBlueprintCompileReinstancer StubReinstancer(Blueprint->GeneratedClass);
+			auto StubReinstancer = FBlueprintCompileReinstancer::Create(Blueprint->GeneratedClass);
 
 			// Toss the half-baked class and generate a stubbed out skeleton class that can be used
 			FCompilerResultsLog StubResults;
@@ -181,10 +179,10 @@ void FKismet2CompilerModule::CompileBlueprint(class UBlueprint* Blueprint, const
 
 			CompileBlueprintInner(Blueprint, StubCompileOptions, StubResults, ObjLoaded);
 
-			StubReinstancer.UpdateBytecodeReferences();
+			StubReinstancer->UpdateBytecodeReferences();
 			if( !Blueprint->bIsRegeneratingOnLoad )
 			{
-				StubReinstancer.ReinstanceObjects(!ReinstanceOnlyWhenNecessary);
+				StubReinstancer->ReinstanceObjects();
 			}
 		}
 	}
diff --git a/Engine/Source/Editor/KismetCompiler/Public/KismetCompilerModule.h b/Engine/Source/Editor/KismetCompiler/Public/KismetCompilerModule.h
--- a/Engine/Source/Editor/KismetCompiler/Public/KismetCompilerModule.h
+++ b/Engine/Source/Editor/KismetCompiler/Public/KismetCompilerModule.h
@@ -33,7 +33,7 @@ public:
 	 * @param	Results  	The results log for warnings and errors.
 	 * @param	Options		Compiler options.
 	 */
-	virtual void CompileBlueprint(class UBlueprint* Blueprint, const FKismetCompilerOptions& CompileOptions, FCompilerResultsLog& Results, class FBlueprintCompileReinstancer* ParentReinstancer = NULL, TArray<UObject*>* ObjLoaded = NULL)=0;
+	virtual void CompileBlueprint(class UBlueprint* Blueprint, const FKismetCompilerOptions& CompileOptions, FCompilerResultsLog& Results, TSharedPtr<class FBlueprintCompileReinstancer> ParentReinstancer = NULL, TArray<UObject*>* ObjLoaded = NULL)=0;
 
 	/**
 	 * Compiles a user defined structure.
@@ -73,7 +73,7 @@ class FKismet2CompilerModule : public IKismetCompilerInterface
 {
 public:
 	// Implementation of the IKismetCompilerInterface
-	virtual void CompileBlueprint(class UBlueprint* Blueprint, const FKismetCompilerOptions& CompileOptions, FCompilerResultsLog& Results, class FBlueprintCompileReinstancer* ParentReinstancer = NULL, TArray<UObject*>* ObjLoaded = NULL) override;
+	virtual void CompileBlueprint(class UBlueprint* Blueprint, const FKismetCompilerOptions& CompileOptions, FCompilerResultsLog& Results, TSharedPtr<class FBlueprintCompileReinstancer> ParentReinstancer = NULL, TArray<UObject*>* ObjLoaded = NULL) override;
 	virtual void CompileStructure(class UUserDefinedStruct* Struct, FCompilerResultsLog& Results) override;
 	virtual void RecoverCorruptedBlueprint(class UBlueprint* Blueprint) override;
 	virtual void RemoveBlueprintGeneratedClasses(class UBlueprint* Blueprint) override;
diff --git a/Engine/Source/Editor/UMGEditor/Private/WidgetBlueprintCompiler.cpp b/Engine/Source/Editor/UMGEditor/Private/WidgetBlueprintCompiler.cpp
--- a/Engine/Source/Editor/UMGEditor/Private/WidgetBlueprintCompiler.cpp
+++ b/Engine/Source/Editor/UMGEditor/Private/WidgetBlueprintCompiler.cpp
@@ -315,7 +315,7 @@ void FWidgetBlueprintCompiler::SpawnNewClass(const FString& NewClassName)
 	else
 	{
 		// Already existed, but wasn't linked in the Blueprint yet due to load ordering issues
-		FBlueprintCompileReinstancer GeneratedClassReinstancer(NewWidgetBlueprintClass);
+		FBlueprintCompileReinstancer::Create(NewWidgetBlueprintClass);
 	}
 	NewClass = NewWidgetBlueprintClass;
 }
diff --git a/Engine/Source/Editor/UnrealEd/Private/Kismet2/Kismet2.cpp b/Engine/Source/Editor/UnrealEd/Private/Kismet2/Kismet2.cpp
--- a/Engine/Source/Editor/UnrealEd/Private/Kismet2/Kismet2.cpp
+++ b/Engine/Source/Editor/UnrealEd/Private/Kismet2/Kismet2.cpp
@@ -700,14 +700,14 @@ void FKismetEditorUtilities::CompileBlueprint(UBlueprint* BlueprintObj, bool bIs
 	FCompilerResultsLog LocalResults;
 	FCompilerResultsLog& Results = (pResults != NULL) ? *pResults : LocalResults;
 
-	FBlueprintCompileReinstancer ReinstanceHelper(OldClass);
+	auto ReinstanceHelper = FBlueprintCompileReinstancer::Create(OldClass);
 
 	// Suppress errors/warnings in the log if we're recompiling on load on a build machine
 	Results.bLogInfoOnly = BlueprintObj->bIsRegeneratingOnLoad && GIsBuildMachine;
 
 	FKismetCompilerOptions CompileOptions;
 	CompileOptions.bSaveIntermediateProducts = bSaveIntermediateProducts;
-	Compiler.CompileBlueprint(BlueprintObj, CompileOptions, Results, &ReinstanceHelper);
+	Compiler.CompileBlueprint(BlueprintObj, CompileOptions, Results, ReinstanceHelper);
 
 	FBlueprintEditorUtils::UpdateDelegatesInBlueprint(BlueprintObj);
 
@@ -722,7 +722,7 @@ void FKismetEditorUtilities::CompileBlueprint(UBlueprint* BlueprintObj, bool bIs
 		}
 	}
 
-	ReinstanceHelper.UpdateBytecodeReferences();
+	ReinstanceHelper->UpdateBytecodeReferences();
 
 	if (!bIsRegeneratingOnLoad && (OldClass != NULL))
 	{
@@ -736,8 +736,7 @@ void FKismetEditorUtilities::CompileBlueprint(UBlueprint* BlueprintObj, bool bIs
 		}
 
 		// Replace instances of this class
-		static const FBoolConfigValueHelper ReinstanceOnlyWhenNecessary(TEXT("Kismet"), TEXT("bReinstanceOnlyWhenNecessary"), GEngineIni);
-		ReinstanceHelper.ReinstanceObjects(!ReinstanceOnlyWhenNecessary);
+		ReinstanceHelper->ReinstanceObjects();
 
 		// Notify everyone a blueprint has been compiled and reinstanced, but before GC so they can perform any final cleanup.
 		if ( GEditor )
@@ -873,13 +872,13 @@ void FKismetEditorUtilities::RecompileBlueprintBytecode(UBlueprint* BlueprintObj
 	TGuardValue<bool> GuardTemplateNameFlag(GCompilingBlueprint, true);
 	FCompilerResultsLog Results;
 
-	FBlueprintCompileReinstancer ReinstanceHelper(BlueprintObj->GeneratedClass, true);
+	auto ReinstanceHelper = FBlueprintCompileReinstancer::Create(BlueprintObj->GeneratedClass, true);
 
 	FKismetCompilerOptions CompileOptions;
 	CompileOptions.CompileType = EKismetCompileType::BytecodeOnly;
 	Compiler.CompileBlueprint(BlueprintObj, CompileOptions, Results, NULL, ObjLoaded);
 
-	ReinstanceHelper.UpdateBytecodeReferences();
+	ReinstanceHelper->UpdateBytecodeReferences();
 
 	if (BlueprintPackage != NULL)
 	{
diff --git a/Engine/Source/Editor/UnrealEd/Private/Kismet2/KismetReinstanceUtilities.cpp b/Engine/Source/Editor/UnrealEd/Private/Kismet2/KismetReinstanceUtilities.cpp
--- a/Engine/Source/Editor/UnrealEd/Private/Kismet2/KismetReinstanceUtilities.cpp
+++ b/Engine/Source/Editor/UnrealEd/Private/Kismet2/KismetReinstanceUtilities.cpp
@@ -273,23 +273,75 @@ FBlueprintCompileReinstancer::~FBlueprintCompileReinstancer()
 	}
 }
 
-void FBlueprintCompileReinstancer::ReinstanceObjects(bool bAlwaysReinstance)
+void FBlueprintCompileReinstancer::ReinstanceFast()
 {
-	BP_SCOPED_COMPILER_EVENT_STAT(EKismetCompilerStats_ReinstanceObjects);
+	UE_LOG(LogBlueprint, Log, TEXT("BlueprintCompileReinstancer: Doing a fast path refresh on class '%s'."), *GetPathNameSafe(ClassToReinstance));
+
+	TArray<UObject*> ObjectsToReplace;
+	GetObjectsOfClass(DuplicatedClass, ObjectsToReplace, /*bIncludeDerivedClasses=*/ false);
 
-	// Make sure we only reinstance classes once!
-	if( bHasReinstanced )
+	const bool bIsActor = ClassToReinstance->IsChildOf<AActor>();
+	const bool bIsAnimInstance = ClassToReinstance->IsChildOf<UAnimInstance>();
+	const bool bIsComponent = ClassToReinstance->IsChildOf<UActorComponent>();
+	for (auto Obj : ObjectsToReplace)
 	{
-		return;
+		UE_LOG(LogBlueprint, Log, TEXT("  Fast path is refreshing (not replacing) %s"), *Obj->GetFullName());
+
+		if ((!Obj->IsTemplate() || bIsComponent) && !Obj->IsPendingKill())
+		{
+			Obj->SetClass(ClassToReinstance);
+
+			if (bIsActor)
+			{
+				auto Actor = CastChecked<AActor>(Obj);
+				Actor->ReregisterAllComponents();
+				Actor->RerunConstructionScripts();
+			}
+
+			if (bIsAnimInstance)
+			{
+				// Initialising the anim instance isn't enough to correctly set up the skeletal mesh again in a
+				// paused world, need to initialise the skeletal mesh component that contains the anim instance.
+				if (USkeletalMeshComponent* SkelComponent = Cast<USkeletalMeshComponent>(Obj->GetOuter()))
+				{
+					SkelComponent->InitAnim(true);
+				}
+			}
+		}
 	}
-	bHasReinstanced = true;
 
+	TArray<UObject*> SourceObjects;
+	TMap<UObject*, UObject*> OldToNewInstanceMap;
+	TMap<FStringAssetReference, UObject*> ReinstancedObjectsWeakReferenceMap;
+	FReplaceReferenceHelper::IncludeCDO(DuplicatedClass, ClassToReinstance, OldToNewInstanceMap, SourceObjects, OriginalCDO);
+	FReplaceReferenceHelper::FindAndReplaceReferences(SourceObjects, &ObjectsThatShouldUseOldStuff, ObjectsToReplace, OldToNewInstanceMap, ReinstancedObjectsWeakReferenceMap);
+}
+
+void FBlueprintCompileReinstancer::CompileChildren()
+{
+	BP_SCOPED_COMPILER_EVENT_STAT(EKismetReinstancerStats_RecompileChildClasses);
+
+	// Reparent all dependent blueprints, and recompile to ensure that they get reinstanced with the new memory layout
+	for (auto ChildBP = Children.CreateIterator(); ChildBP; ++ChildBP)
+	{
+		UBlueprint* BP = *ChildBP;
+		if (BP->ParentClass == ClassToReinstance || BP->ParentClass == DuplicatedClass)
+		{
+			ReparentChild(BP);
+
+			FKismetEditorUtilities::CompileBlueprint(BP, false, bSkipGarbageCollection);
+		}
+	}
+}
+
+void FBlueprintCompileReinstancer::ReinstanceInner(bool bForceAlwaysReinstance)
+{
 	if (ClassToReinstance && DuplicatedClass)
 	{
+		static const FBoolConfigValueHelper ReinstanceOnlyWhenNecessary(TEXT("Kismet"), TEXT("bReinstanceOnlyWhenNecessary"), GEngineIni);
 		bool bShouldReinstance = true;
-		
 		// See if we need to do a full reinstance or can do the faster refresh path (when enabled or no values were modified, and the structures match)
-		if (!bAlwaysReinstance)
+		if (ReinstanceOnlyWhenNecessary && !bForceAlwaysReinstance)
 		{
 			BP_SCOPED_COMPILER_EVENT_STAT(EKismetReinstancerStats_ReplaceClassNoReinsancing);
 
@@ -305,48 +357,7 @@ void FBlueprintCompileReinstancer::ReinstanceObjects(bool bAlwaysReinstance)
 			const bool bAllowedToDoFastPath = (ChangeDefaultValueWithoutReinstancing || bTheSameDefaultValues) && bTheSameLayout;
 			if (bAllowedToDoFastPath)
 			{
-
-				UE_LOG(LogBlueprint, Log, TEXT("BlueprintCompileReinstancer: Doing a fast path refresh on class '%s'."), *GetPathNameSafe(ClassToReinstance));
-
-				TArray<UObject*> ObjectsToReplace;
-				GetObjectsOfClass(DuplicatedClass, ObjectsToReplace, /*bIncludeDerivedClasses=*/ false);
-
-				const bool bIsActor = ClassToReinstance->IsChildOf<AActor>();
-				const bool bIsAnimInstance = ClassToReinstance->IsChildOf<UAnimInstance>();
-				const bool bIsComponent = ClassToReinstance->IsChildOf<UActorComponent>();
-				for (auto Obj : ObjectsToReplace)
-				{
-					UE_LOG(LogBlueprint, Log, TEXT("  Fast path is refreshing (not replacing) %s"), *Obj->GetFullName());
-
-					if ((!Obj->IsTemplate() || bIsComponent) && !Obj->IsPendingKill())
-					{
-						Obj->SetClass(ClassToReinstance);
-
-						if (bIsActor)
-						{
-							auto Actor = CastChecked<AActor>(Obj);
-							Actor->ReregisterAllComponents();
-							Actor->RerunConstructionScripts();
-						}
-
-						if (bIsAnimInstance)
-						{
-							// Initialising the anim instance isn't enough to correctly set up the skeletal mesh again in a
-							// paused world, need to initialise the skeletal mesh component that contains the anim instance.
-							if (USkeletalMeshComponent* SkelComponent = Cast<USkeletalMeshComponent>(Obj->GetOuter()))
-							{
-								SkelComponent->InitAnim(true);
-							}
-						}
-					}
-				}
-
-				TArray<UObject*> SourceObjects;
-				TMap<UObject*, UObject*> OldToNewInstanceMap;
-				TMap<FStringAssetReference, UObject*> ReinstancedObjectsWeakReferenceMap;
-				FReplaceReferenceHelper::IncludeCDO(DuplicatedClass, ClassToReinstance, OldToNewInstanceMap, SourceObjects, OriginalCDO);
-				FReplaceReferenceHelper::FindAndReplaceReferences(SourceObjects, &ObjectsThatShouldUseOldStuff, ObjectsToReplace, OldToNewInstanceMap, ReinstancedObjectsWeakReferenceMap);
-
+				ReinstanceFast();
 				bShouldReinstance = false;
 			}
 		}
@@ -361,25 +372,61 @@ void FBlueprintCompileReinstancer::ReinstanceObjects(bool bAlwaysReinstance)
 			// ReplaceInstancesOfClass() handles this itself, if we had to re-instance
 			ReconstructOwnerInstances(ClassToReinstance);
 		}
+	}
+}
+
+void FBlueprintCompileReinstancer::ReinstanceObjects(bool bForceAlwaysReinstance)
+{
+	BP_SCOPED_COMPILER_EVENT_STAT(EKismetCompilerStats_ReinstanceObjects);
 	
-		{ 
-			BP_SCOPED_COMPILER_EVENT_STAT(EKismetReinstancerStats_RecompileChildClasses);
+	static const FBoolConfigValueHelper FirstCompileChildrenThenReinstance(TEXT("Kismet"), TEXT("bFirstCompileChildrenThenReinstance"), GEngineIni);
+
+	if (FirstCompileChildrenThenReinstance)
+	{
+		// Make sure we only reinstance classes once!
+		static TArray<TSharedPtr<FBlueprintCompileReinstancer>> QueueToReinstance;
+		TSharedPtr<FBlueprintCompileReinstancer> SharedThis = AsShared();
+		const bool bAlreadyQueued = QueueToReinstance.Contains(SharedThis);
+		ensure(!bAlreadyQueued);
+		if (bAlreadyQueued)
+		{
+			QueueToReinstance.Push(SharedThis);
 
-			// Reparent all dependent blueprints, and recompile to ensure that they get reinstanced with the new memory layout
-			for( auto ChildBP = Children.CreateIterator(); ChildBP; ++ChildBP)
+			if (ClassToReinstance && DuplicatedClass)
 			{
-				UBlueprint* BP = *ChildBP;
-				if( BP->ParentClass == ClassToReinstance || BP->ParentClass == DuplicatedClass)
-				{
-					ReparentChild(BP);
+				CompileChildren();
+			}
+		}
 
-					FKismetEditorUtilities::CompileBlueprint(BP, false, bSkipGarbageCollection);
-				}
+		if (QueueToReinstance.Num() && (QueueToReinstance[0] == SharedThis))
+		{
+			// All children were recompiled. It's safe to reinstance.
+			for (int32 Idx = 0; Idx < QueueToReinstance.Num(); ++Idx)
+			{
+				QueueToReinstance[Idx]->ReinstanceInner(bForceAlwaysReinstance);
 			}
+			QueueToReinstance.Empty();
+		}
+	}
+	else
+	{
+		//THE OLD WAY
+		if (bHasReinstanced)
+		{
+			return;
+		}
+		bHasReinstanced = true;
+
+		if (ClassToReinstance && DuplicatedClass)
+		{
+			ReinstanceInner(bForceAlwaysReinstance);
+			CompileChildren();
 		}
 	}
 }
 
+
+
 void FBlueprintCompileReinstancer::UpdateBytecodeReferences()
 {
 	BP_SCOPED_COMPILER_EVENT_STAT(EKismetReinstancerStats_UpdateBytecodeReferences);
diff --git a/Engine/Source/Editor/UnrealEd/Public/Kismet2/KismetReinstanceUtilities.h b/Engine/Source/Editor/UnrealEd/Public/Kismet2/KismetReinstanceUtilities.h
--- a/Engine/Source/Editor/UnrealEd/Public/Kismet2/KismetReinstanceUtilities.h
+++ b/Engine/Source/Editor/UnrealEd/Public/Kismet2/KismetReinstanceUtilities.h
@@ -10,7 +10,7 @@
 
 DECLARE_STATS_GROUP(TEXT("Kismet Reinstancer"), STATGROUP_KismetReinstancer, STATCAT_Advanced);
 
-class UNREALED_API FBlueprintCompileReinstancer
+class UNREALED_API FBlueprintCompileReinstancer : public TSharedFromThis<FBlueprintCompileReinstancer>
 {
 protected:
 
@@ -48,10 +48,12 @@ protected:
 	TSet<UObject*> ObjectsThatShouldUseOldStuff;
 
 public:
-	virtual ~FBlueprintCompileReinstancer();
+	static TSharedPtr<FBlueprintCompileReinstancer> Create(UClass* InClassToReinstance, bool bIsBytecodeOnly = false, bool bSkipGC = false)
+	{
+		return MakeShareable(new FBlueprintCompileReinstancer(InClassToReinstance, bIsBytecodeOnly, bSkipGC));
+	}
 
-	/** Sets the reinstancer up to work on every object of the specified class */
-	FBlueprintCompileReinstancer(UClass* InClassToReinstance, bool bIsBytecodeOnly = false, bool bSkipGC = false);
+	virtual ~FBlueprintCompileReinstancer();
 
 	/** Saves a mapping of field names to their UField equivalents, so we can remap any bytecode that references them later */
 	void SaveClassFieldMapping(UClass* InClassToReinstance);
@@ -60,7 +62,7 @@ public:
 	void GenerateFieldMappings(TMap<UObject*, UObject*>& FieldMapping);
 
 	/** Reinstances all objects in the ObjectReinstancingMap */
-	void ReinstanceObjects(bool bAlwaysReinstance = true);
+	void ReinstanceObjects(bool bForceAlwaysReinstance = false);
 
 	/** Updates references to properties and functions of the class that has in the bytecode of dependent blueprints */
 	void UpdateBytecodeReferences();
@@ -81,6 +83,11 @@ public:
 	void VerifyReplacement();
 
 protected:
+	void ReinstanceInner(bool bForceAlwaysReinstance);
+
+	void ReinstanceFast();
+
+	void CompileChildren();
 
 	/** Default constructor, can only be used by derived classes */
 	FBlueprintCompileReinstancer()
@@ -91,6 +98,9 @@ protected:
 		, bSkipGarbageCollection(false)
 	{}
 
+	/** Sets the reinstancer up to work on every object of the specified class */
+	FBlueprintCompileReinstancer(UClass* InClassToReinstance, bool bIsBytecodeOnly = false, bool bSkipGC = false);
+
 	/** Reparents the specified blueprint or class to be the duplicated class in order to allow properties to be copied from the previous CDO to the new one */
 	void ReparentChild(UBlueprint* ChildBP);
 	void ReparentChild(UClass* ChildClass);
