Fix bsp converted to a static mesh remaining in the world

diff --git a/Engine/Source/Editor/UnrealEd/Private/Editor.cpp b/Engine/Source/Editor/UnrealEd/Private/Editor.cpp
--- a/Engine/Source/Editor/UnrealEd/Private/Editor.cpp
+++ b/Engine/Source/Editor/UnrealEd/Private/Editor.cpp
@@ -5605,7 +5605,7 @@ AActor* UEditorEngine::ConvertBrushesToStaticMesh(const FString& InStaticMeshPac
 	}
 
 	ConversionTempModel->EmptyModel(1, 1);
-
+	GEditor->RebuildAlteredBSP();
 	GEditor->RedrawLevelEditingViewports();
 
 	return NewActor;
