PO files generated by exporting from manifest/archives now use text key(s) as reference-comments.

diff --git a/Engine/Source/Editor/UnrealEd/Private/Commandlets/InternationalizationExportCommandlet.cpp b/Engine/Source/Editor/UnrealEd/Private/Commandlets/InternationalizationExportCommandlet.cpp
--- a/Engine/Source/Editor/UnrealEd/Private/Commandlets/InternationalizationExportCommandlet.cpp
+++ b/Engine/Source/Editor/UnrealEd/Private/Commandlets/InternationalizationExportCommandlet.cpp
@@ -1139,8 +1139,8 @@ bool UInternationalizationExportCommandlet::DoExport( const FString& SourcePath,
 								PoEntry->MsgStr.Add( ConditionedArchiveTranslation );
 
 								FString PORefString = ConvertSrcLocationToPORef( ContextIter->SourceLocation );
-								PoEntry->AddReference( PORefString ); // The appropriate comment type.
-								PoEntry->AddExtractedComment( PORefString ); // Added to extracted comments to workaround what seems to be an issue with OneSky.
+								PoEntry->AddReference( PORefString ); // Source location.
+								PoEntry->AddExtractedComment( ContextIter->Key ); // "Notes from Programmer" in the form of the Key.
 								PortableObj.AddEntry( PoEntry );
 							}
 						}
