Events that are hidden in a Blueprint's parent class will not be generated as default event nodes.

Level Script Blueprints will not start off with a "Begin Actor Overlap" event node.

diff --git a/Engine/Source/Editor/GameplayAbilitiesEditor/Private/GameplayAbilitiesBlueprintFactory.cpp b/Engine/Source/Editor/GameplayAbilitiesEditor/Private/GameplayAbilitiesBlueprintFactory.cpp
--- a/Engine/Source/Editor/GameplayAbilitiesEditor/Private/GameplayAbilitiesBlueprintFactory.cpp
+++ b/Engine/Source/Editor/GameplayAbilitiesEditor/Private/GameplayAbilitiesBlueprintFactory.cpp
@@ -290,8 +290,9 @@ UObject* UGameplayAbilitiesBlueprintFactory::FactoryCreateNew(UClass* Class, UOb
 				NewBP->LastEditedDocuments.Add(NewGraph);
 				NewGraph->bAllowDeletion = false;
 
-				UEdGraphNode* EventNode = FKismetEditorUtilities::AddDefaultEventNode(NewBP, NewGraph, FName(TEXT("K2_ActivateAbility")), UGameplayAbility::StaticClass());
-				FKismetEditorUtilities::AddDefaultEventNode(NewBP, NewGraph, FName(TEXT("K2_OnEndAbility")), UGameplayAbility::StaticClass(), EventNode->NodePosY + EventNode->NodeHeight + 200);
+				int32 NodePositionY = 0;
+				FKismetEditorUtilities::AddDefaultEventNode(NewBP, NewGraph, FName(TEXT("K2_ActivateAbility")), UGameplayAbility::StaticClass(), NodePositionY);
+				FKismetEditorUtilities::AddDefaultEventNode(NewBP, NewGraph, FName(TEXT("K2_OnEndAbility")), UGameplayAbility::StaticClass(), NodePositionY);
 			}
 		}
 
diff --git a/Engine/Source/Editor/UnrealEd/Private/Kismet2/Kismet2.cpp b/Engine/Source/Editor/UnrealEd/Private/Kismet2/Kismet2.cpp
--- a/Engine/Source/Editor/UnrealEd/Private/Kismet2/Kismet2.cpp
+++ b/Engine/Source/Editor/UnrealEd/Private/Kismet2/Kismet2.cpp
@@ -469,7 +469,9 @@ UBlueprint* FKismetEditorUtilities::CreateBlueprint(UClass* ParentClass, UObject
 		if(Graph->Nodes.Num() == 0)
 		{
 			// add update event graph
-			UK2Node_Event* EventNode = FKismetEditorUtilities::AddDefaultEventNode(NewBP, Graph, FName(TEXT("BlueprintUpdateAnimation")), UAnimInstance::StaticClass());
+			int32 NodePositionY = 0;
+			UK2Node_Event* EventNode = FKismetEditorUtilities::AddDefaultEventNode(NewBP, Graph, FName(TEXT("BlueprintUpdateAnimation")), UAnimInstance::StaticClass(), NodePositionY);
+			check(EventNode);
 
 			// add try get owner node
 			UK2Node_CallFunction* GetOwnerNode = NewObject<UK2Node_CallFunction>(Graph);
@@ -497,26 +499,28 @@ UBlueprint* FKismetEditorUtilities::CreateBlueprint(UClass* ParentClass, UObject
 			UClass* WidgetClass = FindObject<UClass>(ANY_PACKAGE, TEXT("UserWidget"));
 			UClass* GameplayAbilityClass = FindObject<UClass>(ANY_PACKAGE, TEXT("GameplayAbility"));
 
+			int32 NodePositionY = 0;
+
 			if(NewBP->GeneratedClass->IsChildOf(AActor::StaticClass()))
 			{
-				UEdGraphNode* BeginPlayNode = FKismetEditorUtilities::AddDefaultEventNode(NewBP, NewBP->UbergraphPages[0], FName(TEXT("ReceiveBeginPlay")), AActor::StaticClass());
-				UEdGraphNode* ActorBeginOverlapNode = FKismetEditorUtilities::AddDefaultEventNode(NewBP, NewBP->UbergraphPages[0], FName(TEXT("ReceiveActorBeginOverlap")), AActor::StaticClass(), BeginPlayNode->NodePosY + BeginPlayNode->NodeHeight + 200);
-				FKismetEditorUtilities::AddDefaultEventNode(NewBP, NewBP->UbergraphPages[0], FName(TEXT("ReceiveTick")), AActor::StaticClass(), ActorBeginOverlapNode->NodePosY + ActorBeginOverlapNode->NodeHeight + 200);
+				FKismetEditorUtilities::AddDefaultEventNode(NewBP, NewBP->UbergraphPages[0], FName(TEXT("ReceiveBeginPlay")), AActor::StaticClass(), NodePositionY);
+				FKismetEditorUtilities::AddDefaultEventNode(NewBP, NewBP->UbergraphPages[0], FName(TEXT("ReceiveActorBeginOverlap")), AActor::StaticClass(), NodePositionY);
+				FKismetEditorUtilities::AddDefaultEventNode(NewBP, NewBP->UbergraphPages[0], FName(TEXT("ReceiveTick")), AActor::StaticClass(), NodePositionY);
 			}
 			else if(NewBP->GeneratedClass->IsChildOf(UActorComponent::StaticClass()))
 			{
-				UEdGraphNode* ReceiveTickNode = FKismetEditorUtilities::AddDefaultEventNode(NewBP, NewBP->UbergraphPages[0], FName(TEXT("ReceiveTick")), UActorComponent::StaticClass());
-				FKismetEditorUtilities::AddDefaultEventNode(NewBP, NewBP->UbergraphPages[0], FName(TEXT("ReceiveInitializeComponent")), UActorComponent::StaticClass(), ReceiveTickNode->NodePosY + ReceiveTickNode->NodeHeight + 200);
+				FKismetEditorUtilities::AddDefaultEventNode(NewBP, NewBP->UbergraphPages[0], FName(TEXT("ReceiveTick")), UActorComponent::StaticClass(), NodePositionY);
+				FKismetEditorUtilities::AddDefaultEventNode(NewBP, NewBP->UbergraphPages[0], FName(TEXT("ReceiveInitializeComponent")), UActorComponent::StaticClass(), NodePositionY);
 			}
 			else if(NewBP->GeneratedClass->IsChildOf(WidgetClass))
 			{
-				UEdGraphNode* EventNode = FKismetEditorUtilities::AddDefaultEventNode(NewBP, NewBP->UbergraphPages[0], FName(TEXT("Construct")), WidgetClass);
-				FKismetEditorUtilities::AddDefaultEventNode(NewBP, NewBP->UbergraphPages[0], FName(TEXT("Tick")), WidgetClass, EventNode->NodePosY + EventNode->NodeHeight + 200);
+				FKismetEditorUtilities::AddDefaultEventNode(NewBP, NewBP->UbergraphPages[0], FName(TEXT("Construct")), WidgetClass, NodePositionY);
+				FKismetEditorUtilities::AddDefaultEventNode(NewBP, NewBP->UbergraphPages[0], FName(TEXT("Tick")), WidgetClass, NodePositionY);
 			}
 			else if(NewBP->GeneratedClass->IsChildOf(GameplayAbilityClass))
 			{
-				UEdGraphNode* EventNode = FKismetEditorUtilities::AddDefaultEventNode(NewBP, NewBP->UbergraphPages[0], FName(TEXT("K2_ActivateAbility")), GameplayAbilityClass);
-				FKismetEditorUtilities::AddDefaultEventNode(NewBP, NewBP->UbergraphPages[0], FName(TEXT("K2_OnEndAbility")), GameplayAbilityClass, EventNode->NodePosY + EventNode->NodeHeight + 200);
+				FKismetEditorUtilities::AddDefaultEventNode(NewBP, NewBP->UbergraphPages[0], FName(TEXT("K2_ActivateAbility")), GameplayAbilityClass, NodePositionY);
+				FKismetEditorUtilities::AddDefaultEventNode(NewBP, NewBP->UbergraphPages[0], FName(TEXT("K2_OnEndAbility")), GameplayAbilityClass, NodePositionY);
 			}
 		}
 	}
@@ -550,56 +554,65 @@ UBlueprint* FKismetEditorUtilities::CreateBlueprint(UClass* ParentClass, UObject
 	return NewBP;
 }
 
-UK2Node_Event* FKismetEditorUtilities::AddDefaultEventNode(UBlueprint* InBlueprint, UEdGraph* InGraph, FName InEventName, UClass* InEventClass, int32 InNodePosY/* = 0*/)
+UK2Node_Event* FKismetEditorUtilities::AddDefaultEventNode(UBlueprint* InBlueprint, UEdGraph* InGraph, FName InEventName, UClass* InEventClass, int32& InOutNodePosY)
 {
-	const UEdGraphSchema_K2* Schema = GetDefault<UEdGraphSchema_K2>();
-
-	// Add a "Begin Play" event
-	UK2Node_Event* NewEventNode = NewObject<UK2Node_Event>(InGraph);
-	NewEventNode->EventReference.SetExternalMember(InEventName, InEventClass);
-
-	// add update event graph
-	NewEventNode->bOverrideFunction=true;
-	NewEventNode->CreateNewGuid();
-	NewEventNode->PostPlacedNewNode();
-	NewEventNode->SetFlags(RF_Transactional);
-	NewEventNode->AllocateDefaultPins();
-	NewEventNode->bIsNodeEnabled = false;
-	NewEventNode->NodeComment = LOCTEXT("DisabledNodeComment", "This node is disabled and will not be called.\nDrag off pins to build functionality.").ToString();
-	NewEventNode->bCommentBubblePinned = true;
-	NewEventNode->bCommentBubbleVisible = true;
-	NewEventNode->NodePosY = InNodePosY;
-	UEdGraphSchema_K2::SetNodeMetaData(NewEventNode, FNodeMetadata::DefaultGraphNode);
-
-	InGraph->AddNode(NewEventNode);
-
-	// Rebuild the skeleton class so we can determine if we are overriding the function from the parent
-	FBlueprintEditorUtils::MarkBlueprintAsStructurallyModified(InBlueprint);
-
-	// Get the function that the event node or function entry represents
-	FFunctionFromNodeHelper FunctionFromNode(NewEventNode);
-	if (FunctionFromNode.Function && Schema->GetCallableParentFunction(FunctionFromNode.Function))
-	{
-		UFunction* ValidParent = Schema->GetCallableParentFunction(FunctionFromNode.Function);
-		FGraphNodeCreator<UK2Node_CallParentFunction> FunctionNodeCreator(*InGraph);
-		UK2Node_CallParentFunction* ParentFunctionNode = FunctionNodeCreator.CreateNode();
-		ParentFunctionNode->SetFromFunction(ValidParent);
-		ParentFunctionNode->AllocateDefaultPins();
-
-		ParentFunctionNode->GetExecPin()->MakeLinkTo(NewEventNode->FindPin(Schema->PN_Then));
-
-		ParentFunctionNode->NodePosX = FunctionFromNode.Node->NodePosX + FunctionFromNode.Node->NodeWidth + 200;
-		ParentFunctionNode->NodePosY = FunctionFromNode.Node->NodePosY;
-		UEdGraphSchema_K2::SetNodeMetaData(ParentFunctionNode, FNodeMetadata::DefaultGraphNode);
-		FunctionNodeCreator.Finalize();
-
-		ParentFunctionNode->bIsNodeEnabled = false;
-
-		// Adding the call to parent and connecting it will reset this value
+	UK2Node_Event* NewEventNode = nullptr;
+
+	FMemberReference EventReference;
+	EventReference.SetExternalMember(InEventName, InEventClass);
+
+	// Prevent events that are hidden in the Blueprint's class from being auto-generated.
+	if(!FObjectEditorUtils::IsFunctionHiddenFromClass(EventReference.ResolveMember<UFunction>(InBlueprint), InBlueprint->ParentClass))
+	{
+		const UEdGraphSchema_K2* Schema = GetDefault<UEdGraphSchema_K2>();
+
+		// Add the event
+		NewEventNode = NewObject<UK2Node_Event>(InGraph);
+		NewEventNode->EventReference = EventReference;
+
+		// add update event graph
+		NewEventNode->bOverrideFunction=true;
+		NewEventNode->CreateNewGuid();
+		NewEventNode->PostPlacedNewNode();
+		NewEventNode->SetFlags(RF_Transactional);
+		NewEventNode->AllocateDefaultPins();
 		NewEventNode->bIsNodeEnabled = false;
 		NewEventNode->NodeComment = LOCTEXT("DisabledNodeComment", "This node is disabled and will not be called.\nDrag off pins to build functionality.").ToString();
-	}
+		NewEventNode->bCommentBubblePinned = true;
+		NewEventNode->bCommentBubbleVisible = true;
+		NewEventNode->NodePosY = InOutNodePosY;
+		UEdGraphSchema_K2::SetNodeMetaData(NewEventNode, FNodeMetadata::DefaultGraphNode);
+		InOutNodePosY = NewEventNode->NodePosY + NewEventNode->NodeHeight + 200;
+
+		InGraph->AddNode(NewEventNode);
 
+		// Rebuild the skeleton class so we can determine if we are overriding the function from the parent
+		FBlueprintEditorUtils::MarkBlueprintAsStructurallyModified(InBlueprint);
+
+		// Get the function that the event node or function entry represents
+		FFunctionFromNodeHelper FunctionFromNode(NewEventNode);
+		if (FunctionFromNode.Function && Schema->GetCallableParentFunction(FunctionFromNode.Function))
+		{
+			UFunction* ValidParent = Schema->GetCallableParentFunction(FunctionFromNode.Function);
+			FGraphNodeCreator<UK2Node_CallParentFunction> FunctionNodeCreator(*InGraph);
+			UK2Node_CallParentFunction* ParentFunctionNode = FunctionNodeCreator.CreateNode();
+			ParentFunctionNode->SetFromFunction(ValidParent);
+			ParentFunctionNode->AllocateDefaultPins();
+
+			ParentFunctionNode->GetExecPin()->MakeLinkTo(NewEventNode->FindPin(Schema->PN_Then));
+
+			ParentFunctionNode->NodePosX = FunctionFromNode.Node->NodePosX + FunctionFromNode.Node->NodeWidth + 200;
+			ParentFunctionNode->NodePosY = FunctionFromNode.Node->NodePosY;
+			UEdGraphSchema_K2::SetNodeMetaData(ParentFunctionNode, FNodeMetadata::DefaultGraphNode);
+			FunctionNodeCreator.Finalize();
+
+			ParentFunctionNode->bIsNodeEnabled = false;
+
+			// Adding the call to parent and connecting it will reset this value
+			NewEventNode->bIsNodeEnabled = false;
+			NewEventNode->NodeComment = LOCTEXT("DisabledNodeComment", "This node is disabled and will not be called.\nDrag off pins to build functionality.").ToString();
+		}
+	}
 
 	return NewEventNode;
 }
diff --git a/Engine/Source/Editor/UnrealEd/Public/Kismet2/KismetEditorUtilities.h b/Engine/Source/Editor/UnrealEd/Public/Kismet2/KismetEditorUtilities.h
--- a/Engine/Source/Editor/UnrealEd/Public/Kismet2/KismetEditorUtilities.h
+++ b/Engine/Source/Editor/UnrealEd/Public/Kismet2/KismetEditorUtilities.h
@@ -220,10 +220,10 @@ public:
 	 * @param InGraph			The graph to spawn the event node in
 	 * @param InEventName		The name of the event function
 	 * @param InEventClass		The class this event can be found in
-	 * @param InNodePosY		Optional Y-position to spawn the node at to easily spawn in a line
+	 * @param InOutNodePosY		Position to spawn the node at, will return with an offset more suitable to offset the next node
 	 * @return					The K2Node_Event will be returned
 	 */
-	static class UK2Node_Event* AddDefaultEventNode(UBlueprint* InBlueprint, UEdGraph* InGraph, FName InEventName, UClass* InEventClass, int32 InNodePosY = 0);
+	static class UK2Node_Event* AddDefaultEventNode(UBlueprint* InBlueprint, UEdGraph* InGraph, FName InEventName, UClass* InEventClass, int32& InOutNodePosY);
 
 private:
 	/** Stores whether we are already listening for kismet clicks */
