[UE-10169] - Added "/ChainingPackage EpicGames_UE4" param when invoking VS2013 installer

diff --git a/Engine/Source/Editor/UnrealEd/Private/SourceCodeNavigation.cpp b/Engine/Source/Editor/UnrealEd/Private/SourceCodeNavigation.cpp
--- a/Engine/Source/Editor/UnrealEd/Private/SourceCodeNavigation.cpp
+++ b/Engine/Source/Editor/UnrealEd/Private/SourceCodeNavigation.cpp
@@ -1725,7 +1725,7 @@ FString FSourceCodeNavigationImpl::GetSuggestedIDEInstallerFileName()
 void FSourceCodeNavigationImpl::LaunchIDEInstaller(const FString& Filepath)
 {
 #if PLATFORM_WINDOWS
-	auto Params = TEXT("/PromptRestart");
+	auto Params = TEXT("/PromptRestart /ChainingPackage EpicGames_UE4");
 	FPlatformProcess::CreateProc(*Filepath, Params, true, false, false, nullptr, 0, nullptr, nullptr);
 #endif
 }
