Fixed the tooltips used when adding new components, and also added a fallback for unloaded components

UE-9337 - IWCE: Fix busted tool-tips for "New C++ Component" and "New Blueprint Component"

diff --git a/Engine/Source/Editor/UnrealEd/Private/SComponentClassCombo.cpp b/Engine/Source/Editor/UnrealEd/Private/SComponentClassCombo.cpp
--- a/Engine/Source/Editor/UnrealEd/Private/SComponentClassCombo.cpp
+++ b/Engine/Source/Editor/UnrealEd/Private/SComponentClassCombo.cpp
@@ -335,7 +335,7 @@ TSharedRef<ITableRow> SComponentClassCombo::GenerateAddComponentRow( FComponentC
 		
 		return
 			SNew( SComboRow< TSharedPtr<FString> >, OwnerTable )
-			.ToolTip( FEditorClassUtils::GetTooltip(Entry->GetComponentClass()) )
+			.ToolTip( GetComponentToolTip(Entry) )
 			[
 				SNew(SHorizontalBox)
 				+SHorizontalBox::Slot()
@@ -471,4 +471,29 @@ FString SComponentClassCombo::GetSanitizedComponentName(FComponentClassComboEntr
 	return FName::NameToDisplayString(DisplayName, false);
 }
 
+TSharedRef<SToolTip> SComponentClassCombo::GetComponentToolTip(FComponentClassComboEntryPtr Entry) const
+{
+	// Special handling for the "New..." options
+	if (Entry->GetComponentCreateAction() == EComponentCreateAction::CreateNewCPPClass)
+	{
+		return SNew(SToolTip)
+			.Text(LOCTEXT("NewCPPComponentToolTip", "Create a custom actor component using C++"));
+	}
+	else if (Entry->GetComponentCreateAction() == EComponentCreateAction::CreateNewBlueprintClass)
+	{
+		return SNew(SToolTip)
+			.Text(LOCTEXT("NewBlueprintComponentToolTip", "Create a custom actor component using Blueprints"));
+	}
+	
+	// Handle components which have a currently loaded class
+	if (const UClass* ComponentClass = Entry->GetComponentClass())
+	{
+		return FEditorClassUtils::GetTooltip(ComponentClass);
+	}
+
+	// Fallback for components that don't currently have a loaded class
+	return SNew(SToolTip)
+		.Text(FText::FromString(Entry->GetClassName()));
+}
+
 #undef LOCTEXT_NAMESPACE
diff --git a/Engine/Source/Editor/UnrealEd/Public/SComponentClassCombo.h b/Engine/Source/Editor/UnrealEd/Public/SComponentClassCombo.h
--- a/Engine/Source/Editor/UnrealEd/Public/SComponentClassCombo.h
+++ b/Engine/Source/Editor/UnrealEd/Public/SComponentClassCombo.h
@@ -202,6 +202,8 @@ private:
 
 	FText GetFriendlyComponentName(FComponentClassComboEntryPtr Entry) const;
 
+	TSharedRef<SToolTip> GetComponentToolTip(FComponentClassComboEntryPtr Entry) const;
+
 	FComponentClassSelected OnComponentClassSelected;
 
 	/** List of component class names used by combo box */
