Added Blueprint Editor Setting for enabling/disabling of the default nodes being placing in new Blueprints. They are enabled by default.

diff --git a/Engine/Source/Editor/BlueprintGraph/Private/BlueprintEditorSettings.cpp b/Engine/Source/Editor/BlueprintGraph/Private/BlueprintEditorSettings.cpp
--- a/Engine/Source/Editor/BlueprintGraph/Private/BlueprintEditorSettings.cpp
+++ b/Engine/Source/Editor/BlueprintGraph/Private/BlueprintEditorSettings.cpp
@@ -21,6 +21,7 @@ UBlueprintEditorSettings::UBlueprintEditorSettings(const FObjectInitializer& Obj
 	, bShowViewportOnSimulate(false)
 	, bShowInheritedVariables(false)
 	, bShowEmptySections(true)
+	, bSpawnDefaultBlueprintNodes(true)
 	// Compiler Settings
 	, SaveOnCompile(SoC_Never)
 	, bJumpToNodeErrors(false)
diff --git a/Engine/Source/Editor/BlueprintGraph/Public/BlueprintEditorSettings.h b/Engine/Source/Editor/BlueprintGraph/Public/BlueprintEditorSettings.h
--- a/Engine/Source/Editor/BlueprintGraph/Public/BlueprintEditorSettings.h
+++ b/Engine/Source/Editor/BlueprintGraph/Public/BlueprintEditorSettings.h
@@ -71,6 +71,9 @@ public:
 	UPROPERTY(config)
 	bool bShowEmptySections;
 
+	/** If set will spawn default nodes in new Blueprints */
+	UPROPERTY(EditAnywhere, config, Category=Workflow)
+	bool bSpawnDefaultBlueprintNodes;
 // Compiler Settings
 public:
 	/** Determines when to save Blueprints post-compile */
diff --git a/Engine/Source/Editor/GameplayAbilitiesEditor/Private/GameplayAbilitiesBlueprintFactory.cpp b/Engine/Source/Editor/GameplayAbilitiesEditor/Private/GameplayAbilitiesBlueprintFactory.cpp
--- a/Engine/Source/Editor/GameplayAbilitiesEditor/Private/GameplayAbilitiesBlueprintFactory.cpp
+++ b/Engine/Source/Editor/GameplayAbilitiesEditor/Private/GameplayAbilitiesBlueprintFactory.cpp
@@ -4,6 +4,7 @@
 
 #include "Kismet2/BlueprintEditorUtils.h"
 #include "Kismet2/KismetEditorUtilities.h"
+#include "BlueprintEditorSettings.h"
 
 #include "Abilities/GameplayAbility.h"
 #include "GameplayAbilityBlueprint.h"
@@ -290,9 +291,13 @@ UObject* UGameplayAbilitiesBlueprintFactory::FactoryCreateNew(UClass* Class, UOb
 				NewBP->LastEditedDocuments.Add(NewGraph);
 				NewGraph->bAllowDeletion = false;
 
-				int32 NodePositionY = 0;
-				FKismetEditorUtilities::AddDefaultEventNode(NewBP, NewGraph, FName(TEXT("K2_ActivateAbility")), UGameplayAbility::StaticClass(), NodePositionY);
-				FKismetEditorUtilities::AddDefaultEventNode(NewBP, NewGraph, FName(TEXT("K2_OnEndAbility")), UGameplayAbility::StaticClass(), NodePositionY);
+				UBlueprintEditorSettings* Settings = GetMutableDefault<UBlueprintEditorSettings>();
+				if(Settings && Settings->bSpawnDefaultBlueprintNodes)
+				{
+					int32 NodePositionY = 0;
+					FKismetEditorUtilities::AddDefaultEventNode(NewBP, NewGraph, FName(TEXT("K2_ActivateAbility")), UGameplayAbility::StaticClass(), NodePositionY);
+					FKismetEditorUtilities::AddDefaultEventNode(NewBP, NewGraph, FName(TEXT("K2_OnEndAbility")), UGameplayAbility::StaticClass(), NodePositionY);
+				}
 			}
 		}
 
diff --git a/Engine/Source/Editor/UnrealEd/Private/Kismet2/Kismet2.cpp b/Engine/Source/Editor/UnrealEd/Private/Kismet2/Kismet2.cpp
--- a/Engine/Source/Editor/UnrealEd/Private/Kismet2/Kismet2.cpp
+++ b/Engine/Source/Editor/UnrealEd/Private/Kismet2/Kismet2.cpp
@@ -460,67 +460,71 @@ UBlueprint* FKismetEditorUtilities::CreateBlueprint(UClass* ParentClass, UObject
 	FKismetCompilerOptions CompileOptions;
 	Compiler.CompileBlueprint(NewBP, CompileOptions, Results);
 
-	//@TODO: ANIMREFACTOR 2: This kind of code should be on a per-blueprint basis; not centralized here
-	if(UAnimBlueprint* AnimBP = Cast<UAnimBlueprint>(NewBP))
+	UBlueprintEditorSettings* Settings = GetMutableDefault<UBlueprintEditorSettings>();
+	if(Settings && Settings->bSpawnDefaultBlueprintNodes)
 	{
-		// add default nodes to event graph
-		UEdGraph* Graph = NewBP->UbergraphPages[0];
-
-		if(Graph->Nodes.Num() == 0)
-		{
-			// add update event graph
-			int32 NodePositionY = 0;
-			UK2Node_Event* EventNode = FKismetEditorUtilities::AddDefaultEventNode(NewBP, Graph, FName(TEXT("BlueprintUpdateAnimation")), UAnimInstance::StaticClass(), NodePositionY);
-			check(EventNode);
-
-			// add try get owner node
-			UK2Node_CallFunction* GetOwnerNode = NewObject<UK2Node_CallFunction>(Graph);
-			UFunction* MakeNodeFunction = FindObject<UClass>(ANY_PACKAGE, TEXT("AnimInstance"))->FindFunctionByName(TEXT("TryGetPawnOwner"));
-			GetOwnerNode->CreateNewGuid();
-			GetOwnerNode->PostPlacedNewNode();
-			GetOwnerNode->SetFromFunction(MakeNodeFunction);
-			GetOwnerNode->SetFlags(RF_Transactional);
-			GetOwnerNode->AllocateDefaultPins();
-			GetOwnerNode->NodePosX = EventNode->NodePosX;
-			GetOwnerNode->NodePosY = EventNode->NodePosY + EventNode->NodeHeight + 100;
-			UEdGraphSchema_K2::SetNodeMetaData(GetOwnerNode, FNodeMetadata::DefaultGraphNode);
-			GetOwnerNode->bIsNodeEnabled = false;
-
-			Graph->AddNode(GetOwnerNode);
-		}
-	}
-	else
-	{
-		// Only add default events if there is an ubergraph and they are supported
-		if(NewBP->UbergraphPages.Num() && FBlueprintEditorUtils::DoesSupportEventGraphs(NewBP))
+		//@TODO: ANIMREFACTOR 2: This kind of code should be on a per-blueprint basis; not centralized here
+		if(UAnimBlueprint* AnimBP = Cast<UAnimBlueprint>(NewBP))
 		{
-			// Based on the Blueprint type we are constructing, place some starting events.
-			// Note, this cannot happen in the Factories for constructing these Blueprint types due to the fact that creating child BPs circumvent the factories
-			UClass* WidgetClass = FindObject<UClass>(ANY_PACKAGE, TEXT("UserWidget"));
-			UClass* GameplayAbilityClass = FindObject<UClass>(ANY_PACKAGE, TEXT("GameplayAbility"));
+			// add default nodes to event graph
+			UEdGraph* Graph = NewBP->UbergraphPages[0];
 
-			int32 NodePositionY = 0;
-
-			if(NewBP->GeneratedClass->IsChildOf(AActor::StaticClass()))
-			{
-				FKismetEditorUtilities::AddDefaultEventNode(NewBP, NewBP->UbergraphPages[0], FName(TEXT("ReceiveBeginPlay")), AActor::StaticClass(), NodePositionY);
-				FKismetEditorUtilities::AddDefaultEventNode(NewBP, NewBP->UbergraphPages[0], FName(TEXT("ReceiveActorBeginOverlap")), AActor::StaticClass(), NodePositionY);
-				FKismetEditorUtilities::AddDefaultEventNode(NewBP, NewBP->UbergraphPages[0], FName(TEXT("ReceiveTick")), AActor::StaticClass(), NodePositionY);
-			}
-			else if(NewBP->GeneratedClass->IsChildOf(UActorComponent::StaticClass()))
+			if(Graph->Nodes.Num() == 0)
 			{
-				FKismetEditorUtilities::AddDefaultEventNode(NewBP, NewBP->UbergraphPages[0], FName(TEXT("ReceiveTick")), UActorComponent::StaticClass(), NodePositionY);
-				FKismetEditorUtilities::AddDefaultEventNode(NewBP, NewBP->UbergraphPages[0], FName(TEXT("ReceiveInitializeComponent")), UActorComponent::StaticClass(), NodePositionY);
+				// add update event graph
+				int32 NodePositionY = 0;
+				UK2Node_Event* EventNode = FKismetEditorUtilities::AddDefaultEventNode(NewBP, Graph, FName(TEXT("BlueprintUpdateAnimation")), UAnimInstance::StaticClass(), NodePositionY);
+				check(EventNode);
+
+				// add try get owner node
+				UK2Node_CallFunction* GetOwnerNode = NewObject<UK2Node_CallFunction>(Graph);
+				UFunction* MakeNodeFunction = FindObject<UClass>(ANY_PACKAGE, TEXT("AnimInstance"))->FindFunctionByName(TEXT("TryGetPawnOwner"));
+				GetOwnerNode->CreateNewGuid();
+				GetOwnerNode->PostPlacedNewNode();
+				GetOwnerNode->SetFromFunction(MakeNodeFunction);
+				GetOwnerNode->SetFlags(RF_Transactional);
+				GetOwnerNode->AllocateDefaultPins();
+				GetOwnerNode->NodePosX = EventNode->NodePosX;
+				GetOwnerNode->NodePosY = EventNode->NodePosY + EventNode->NodeHeight + 100;
+				UEdGraphSchema_K2::SetNodeMetaData(GetOwnerNode, FNodeMetadata::DefaultGraphNode);
+				GetOwnerNode->bIsNodeEnabled = false;
+
+				Graph->AddNode(GetOwnerNode);
 			}
-			else if(NewBP->GeneratedClass->IsChildOf(WidgetClass))
-			{
-				FKismetEditorUtilities::AddDefaultEventNode(NewBP, NewBP->UbergraphPages[0], FName(TEXT("Construct")), WidgetClass, NodePositionY);
-				FKismetEditorUtilities::AddDefaultEventNode(NewBP, NewBP->UbergraphPages[0], FName(TEXT("Tick")), WidgetClass, NodePositionY);
-			}
-			else if(NewBP->GeneratedClass->IsChildOf(GameplayAbilityClass))
+		}
+		else
+		{
+			// Only add default events if there is an ubergraph and they are supported
+			if(NewBP->UbergraphPages.Num() && FBlueprintEditorUtils::DoesSupportEventGraphs(NewBP))
 			{
-				FKismetEditorUtilities::AddDefaultEventNode(NewBP, NewBP->UbergraphPages[0], FName(TEXT("K2_ActivateAbility")), GameplayAbilityClass, NodePositionY);
-				FKismetEditorUtilities::AddDefaultEventNode(NewBP, NewBP->UbergraphPages[0], FName(TEXT("K2_OnEndAbility")), GameplayAbilityClass, NodePositionY);
+				// Based on the Blueprint type we are constructing, place some starting events.
+				// Note, this cannot happen in the Factories for constructing these Blueprint types due to the fact that creating child BPs circumvent the factories
+				UClass* WidgetClass = FindObject<UClass>(ANY_PACKAGE, TEXT("UserWidget"));
+				UClass* GameplayAbilityClass = FindObject<UClass>(ANY_PACKAGE, TEXT("GameplayAbility"));
+
+				int32 NodePositionY = 0;
+
+				if(NewBP->GeneratedClass->IsChildOf(AActor::StaticClass()))
+				{
+					FKismetEditorUtilities::AddDefaultEventNode(NewBP, NewBP->UbergraphPages[0], FName(TEXT("ReceiveBeginPlay")), AActor::StaticClass(), NodePositionY);
+					FKismetEditorUtilities::AddDefaultEventNode(NewBP, NewBP->UbergraphPages[0], FName(TEXT("ReceiveActorBeginOverlap")), AActor::StaticClass(), NodePositionY);
+					FKismetEditorUtilities::AddDefaultEventNode(NewBP, NewBP->UbergraphPages[0], FName(TEXT("ReceiveTick")), AActor::StaticClass(), NodePositionY);
+				}
+				else if(NewBP->GeneratedClass->IsChildOf(UActorComponent::StaticClass()))
+				{
+					FKismetEditorUtilities::AddDefaultEventNode(NewBP, NewBP->UbergraphPages[0], FName(TEXT("ReceiveTick")), UActorComponent::StaticClass(), NodePositionY);
+					FKismetEditorUtilities::AddDefaultEventNode(NewBP, NewBP->UbergraphPages[0], FName(TEXT("ReceiveInitializeComponent")), UActorComponent::StaticClass(), NodePositionY);
+				}
+				else if(NewBP->GeneratedClass->IsChildOf(WidgetClass))
+				{
+					FKismetEditorUtilities::AddDefaultEventNode(NewBP, NewBP->UbergraphPages[0], FName(TEXT("Construct")), WidgetClass, NodePositionY);
+					FKismetEditorUtilities::AddDefaultEventNode(NewBP, NewBP->UbergraphPages[0], FName(TEXT("Tick")), WidgetClass, NodePositionY);
+				}
+				else if(NewBP->GeneratedClass->IsChildOf(GameplayAbilityClass))
+				{
+					FKismetEditorUtilities::AddDefaultEventNode(NewBP, NewBP->UbergraphPages[0], FName(TEXT("K2_ActivateAbility")), GameplayAbilityClass, NodePositionY);
+					FKismetEditorUtilities::AddDefaultEventNode(NewBP, NewBP->UbergraphPages[0], FName(TEXT("K2_OnEndAbility")), GameplayAbilityClass, NodePositionY);
+				}
 			}
 		}
 	}
