Added logging to check map names of future crashes when creating play world packages

diff --git a/Engine/Source/Editor/UnrealEd/Private/PlayLevel.cpp b/Engine/Source/Editor/UnrealEd/Private/PlayLevel.cpp
--- a/Engine/Source/Editor/UnrealEd/Private/PlayLevel.cpp
+++ b/Engine/Source/Editor/UnrealEd/Private/PlayLevel.cpp
@@ -3264,6 +3264,8 @@ UWorld* UEditorEngine::CreatePIEWorldByDuplication(FWorldContext &WorldContext,
 	UWorld::WorldTypePreLoadMap.FindOrAdd(PlayWorldMapFName) = EWorldType::PIE;
 
 	// Create a package for the PIE world
+	UE_LOG( LogPlayLevel, Log, TEXT("Creating play world package: %s"),  *PlayWorldMapName );	
+
 	UPackage* PlayWorldPackage = CastChecked<UPackage>(CreatePackage(NULL,*PlayWorldMapName));
 	PlayWorldPackage->PackageFlags |= PKG_PlayInEditor;
 	PlayWorldPackage->PIEInstanceID = WorldContext.PIEInstance;
