Can undo changes to drag and dropping variables to change category.

diff --git a/Engine/Source/Editor/UnrealEd/Private/Kismet2/BlueprintEditorUtils.cpp b/Engine/Source/Editor/UnrealEd/Private/Kismet2/BlueprintEditorUtils.cpp
--- a/Engine/Source/Editor/UnrealEd/Private/Kismet2/BlueprintEditorUtils.cpp
+++ b/Engine/Source/Editor/UnrealEd/Private/Kismet2/BlueprintEditorUtils.cpp
@@ -3103,6 +3103,9 @@ void FBlueprintEditorUtils::RemoveBlueprintVariableMetaData(UBlueprint* Blueprin
 
 void FBlueprintEditorUtils::SetBlueprintVariableCategory(UBlueprint* Blueprint, const FName& VarName, const UStruct* InLocalVarScope, const FName& NewCategory, bool bDontRecompile)
 {
+	const FScopedTransaction Transaction( LOCTEXT("ChangeVariableCategory", "Change Variable Category") );
+	Blueprint->Modify();
+
 	// Ensure we always set a category
 	const UEdGraphSchema_K2* K2Schema = GetDefault<UEdGraphSchema_K2>();
 	FName SetCategory = NewCategory;
@@ -3141,6 +3144,7 @@ void FBlueprintEditorUtils::SetBlueprintVariableCategory(UBlueprint* Blueprint,
 					
 					if(bIsCategoryChanged)
 					{
+						Blueprint->SimpleConstructionScript->GetAllNodes()[SCS_NodeIndex]->Modify();
 						Blueprint->SimpleConstructionScript->GetAllNodes()[SCS_NodeIndex]->CategoryName = SetCategory;
 					}
 				}
@@ -3154,13 +3158,15 @@ void FBlueprintEditorUtils::SetBlueprintVariableCategory(UBlueprint* Blueprint,
 	}
 	else if(InLocalVarScope)
 	{
-		if(FBPVariableDescription* LocalVariable = FindLocalVariable(Blueprint, InLocalVarScope, VarName))
+		UK2Node_FunctionEntry* OutFunctionEntryNode;
+		if(FBPVariableDescription* LocalVariable = FindLocalVariable(Blueprint, InLocalVarScope, VarName, &OutFunctionEntryNode))
 		{
 			// If the category does not change, we will not recompile the Blueprint
 			bool bIsCategoryChanged = LocalVariable->Category != SetCategory;
 
 			if(bIsCategoryChanged)
 			{
+				OutFunctionEntryNode->Modify();
 				LocalVariable->SetMetaData(TEXT("Category"), *SetCategory.ToString());
 				LocalVariable->Category = SetCategory;
 			}
