[UE-10118] Could not add calls to parent implementations for event handlers added after 4.8, we were always finding the version of the function at the top of the class hierarchy

diff --git a/Engine/Source/Editor/UnrealEd/Private/Kismet2/BlueprintEditorUtils.cpp b/Engine/Source/Editor/UnrealEd/Private/Kismet2/BlueprintEditorUtils.cpp
--- a/Engine/Source/Editor/UnrealEd/Private/Kismet2/BlueprintEditorUtils.cpp
+++ b/Engine/Source/Editor/UnrealEd/Private/Kismet2/BlueprintEditorUtils.cpp
@@ -4055,7 +4055,8 @@ UFunction* FFunctionFromNodeHelper::FunctionFromNode(UK2Node* Node)
 	{
 		if (UK2Node_Event* EventNode = Cast<UK2Node_Event>(Node))
 		{
-			Function = EventNode->EventReference.ResolveMember<UFunction>((UClass*)(SearchScope));
+			// We need to search up the class hierachy by name or functions like CanAddParentNode will fail:
+			Function = SearchScope->FindFunctionByName(EventNode->EventReference.GetMemberName());
 		}
 		else if (UK2Node_FunctionEntry* FunctionNode = Cast<UK2Node_FunctionEntry>(Node))
 		{
