[UE-8266] Fix lost actor attachment during post-compile reinstancing when parent and child actor are of the same Blueprint class type.

change summary:
- modified FActorReplacementHelper::AttachChildActors() to check for a reinstanced attach parent before attempting to reattach child actors

diff --git a/Engine/Source/Editor/UnrealEd/Private/Kismet2/KismetReinstanceUtilities.cpp b/Engine/Source/Editor/UnrealEd/Private/Kismet2/KismetReinstanceUtilities.cpp
--- a/Engine/Source/Editor/UnrealEd/Private/Kismet2/KismetReinstanceUtilities.cpp
+++ b/Engine/Source/Editor/UnrealEd/Private/Kismet2/KismetReinstanceUtilities.cpp
@@ -519,15 +519,18 @@ private:
 	 * Takes the cached child actors, as well as the old AttachParent, and sets
 	 * up the new actor so that its attachment hierarchy reflects the old actor
 	 * that it is replacing.
+	 *
+	 * @param OldToNewInstanceMap Mapping of reinstanced objects.
 	 */
-	void ApplyAttachments();
+	void ApplyAttachments(const TMap<UObject*, UObject*>& OldToNewInstanceMap);
 	
 	/**
 	 * Takes the cached child actors, and attaches them under the new actor.
 	 *
 	 * @param  RootComponent	The new actor's root, which the child actors should attach to.
+	 * @param  OldToNewInstanceMap	Mapping of reinstanced objects. Used for when child and parent actor are of the same type (and thus parent may have been reinstanced, so we can't reattach to the old instance).
 	 */
-	void AttachChildActors(USceneComponent* RootComponent);
+	void AttachChildActors(USceneComponent* RootComponent, const TMap<UObject*, UObject*>& OldToNewInstanceMap);
 
 private:
 	AActor*          NewActor;
@@ -586,7 +589,7 @@ void FActorReplacementHelper::Finalize(const TMap<UObject*, UObject*>& OldToNewI
 		}
 	}
 
-	ApplyAttachments();
+	ApplyAttachments(OldToNewInstanceMap);
 
 	if (bSelectNewActor)
 	{
@@ -666,7 +669,7 @@ void FActorReplacementHelper::CacheChildAttachments(const AActor* OldActor)
 	}
 }
 
-void FActorReplacementHelper::ApplyAttachments()
+void FActorReplacementHelper::ApplyAttachments(const TMap<UObject*, UObject*>& OldToNewInstanceMap)
 {
 	USceneComponent* NewRootComponent = NewActor->GetRootComponent();
 	if (NewRootComponent == nullptr)
@@ -686,14 +689,21 @@ void FActorReplacementHelper::ApplyAttachments()
 		}
 	}
 
-	AttachChildActors(NewRootComponent);
+	AttachChildActors(NewRootComponent, OldToNewInstanceMap);
 }
 
-void FActorReplacementHelper::AttachChildActors(USceneComponent* RootComponent)
+void FActorReplacementHelper::AttachChildActors(USceneComponent* RootComponent, const TMap<UObject*, UObject*>& OldToNewInstanceMap)
 {
 	// if we had attached children reattach them now - unless they are already attached
 	for (FAttachedActorInfo& Info : PendingChildAttachments)
 	{
+		// Check for a reinstanced attachment, and redirect to the new instance if found
+		AActor* NewAttachedActor = Cast<AActor>(OldToNewInstanceMap.FindRef(Info.AttachedActor));
+		if (NewAttachedActor)
+		{
+			Info.AttachedActor = NewAttachedActor;
+		}
+
 		// If this actor is no longer attached to anything, reattach
 		if (!Info.AttachedActor->IsPendingKill() && Info.AttachedActor->GetAttachParentActor() == nullptr)
 		{
