Gathering localizable text from metadata now uses a cleaner method for generating keys. Relevant code using those keys has also been updated.

diff --git a/Engine/Source/Editor/UnrealEd/Private/Commandlets/GatherTextFromMetadataCommandlet.cpp b/Engine/Source/Editor/UnrealEd/Private/Commandlets/GatherTextFromMetadataCommandlet.cpp
--- a/Engine/Source/Editor/UnrealEd/Private/Commandlets/GatherTextFromMetadataCommandlet.cpp
+++ b/Engine/Source/Editor/UnrealEd/Private/Commandlets/GatherTextFromMetadataCommandlet.cpp
@@ -1,4 +1,4 @@
-// Copyright 1998-2015 Epic Games, Inc. All Rights Reserved.
+﻿// Copyright 1998-2015 Epic Games, Inc. All Rights Reserved.
 
 #include "UnrealEd.h"
 #include "SourceCodeNavigation.h"
@@ -94,16 +94,6 @@ int32 UGatherTextFromMetaDataCommandlet::Main( const FString& Params )
 	TArray<FString> ExcludePaths;
 	GetStringArrayFromConfig(*SectionName, TEXT("ExcludePaths"), ExcludePaths, GatherTextConfigPath);
 
-	//Required module names
-	TArray<FString> RequiredModuleNames;
-	GetStringArrayFromConfig(*SectionName, TEXT("RequiredModuleNames"), RequiredModuleNames, GatherTextConfigPath);
-
-	// Pre-load all required modules so that UFields from those modules can be guaranteed a chance to have their metadata gathered.
-	for(const FString& RequiredModuleName : RequiredModuleNames)
-	{
-		FModuleManager::Get().LoadModule(*RequiredModuleName);
-	}
-
 	FGatherParameters Arguments;
 	GetStringArrayFromConfig(*SectionName, TEXT("InputKeys"), Arguments.InputKeys, GatherTextConfigPath);
 	GetStringArrayFromConfig(*SectionName, TEXT("OutputNamespaces"), Arguments.OutputNamespaces, GatherTextConfigPath);
@@ -139,6 +129,8 @@ void UGatherTextFromMetaDataCommandlet::GatherTextFromUObjects(const TArray<FStr
 		FSourceCodeNavigation::FindClassHeaderPath(*It, SourceFilePath);
 		SourceFilePath = FPaths::ConvertRelativePathToFull(SourceFilePath);
 
+		check(!SourceFilePath.IsEmpty());
+
 		// Returns true if in an include path. False otherwise.
 		auto IncludePathLogic = [&]() -> bool
 		{
@@ -189,7 +181,7 @@ void UGatherTextFromMetaDataCommandlet::GatherTextFromUObject(UField* const Fiel
 		for(int32 i = 0; i < Arguments.InputKeys.Num(); ++i)
 		{
 			FFormatNamedArguments PatternArguments;
-			PatternArguments.Add( TEXT("FieldPath"), FText::FromString( Field->GetFullGroupName(true) + TEXT(".") + Field->GetName() ) );
+			PatternArguments.Add( TEXT("FieldPath"), FText::FromString( Field->GetFullGroupName(false) ) );
 
 			if( Field->HasMetaData( *Arguments.InputKeys[i] ) )
 			{
@@ -225,7 +217,7 @@ void UGatherTextFromMetaDataCommandlet::GatherTextFromUObject(UField* const Fiel
 				for(int32 j = 0; j < Arguments.InputKeys.Num(); ++j)
 				{
 					FFormatNamedArguments PatternArguments;
-					PatternArguments.Add( TEXT("FieldPath"), FText::FromString( Enum->GetFullGroupName(true) + TEXT(".") + Enum->GetName() + TEXT(".") + Enum->GetEnumName(i) ) );
+					PatternArguments.Add( TEXT("FieldPath"), FText::FromString( Enum->GetFullGroupName(false) + TEXT(".") + Enum->GetEnumName(i) ) );
 
 					if( Enum->HasMetaData(*Arguments.InputKeys[j], i) )
 					{
diff --git a/Engine/Source/Runtime/CoreUObject/Private/UObject/Class.cpp b/Engine/Source/Runtime/CoreUObject/Private/UObject/Class.cpp
--- a/Engine/Source/Runtime/CoreUObject/Private/UObject/Class.cpp
+++ b/Engine/Source/Runtime/CoreUObject/Private/UObject/Class.cpp
@@ -186,7 +186,7 @@ FText UField::GetDisplayNameText() const
 	FText LocalizedDisplayName;
 
 	static const FString Namespace = TEXT("UObjectDisplayNames");
-	const FString Key = GetFullGroupName(true) + TEXT(".") + GetName();
+	const FString Key = GetFullGroupName(false);
 
 	FString NativeDisplayName;
 	if( HasMetaData( TEXT("DisplayName") ) )
@@ -237,7 +237,7 @@ FText UField::GetToolTipText(bool bShortTooltip) const
 	}
 
 	const FString Namespace = bFoundShortTooltip ? TEXT("UObjectShortTooltips") : TEXT("UObjectToolTips");
-	const FString Key = GetFullGroupName(true) + TEXT(".") + GetName();
+	const FString Key = GetFullGroupName(false);
 	if ( !(FText::FindText( Namespace, Key, /*OUT*/LocalizedToolTip )) || *FTextInspector::GetSourceString(LocalizedToolTip) != NativeToolTip)
 	{
 		if (NativeToolTip.IsEmpty())
diff --git a/Engine/Source/Runtime/CoreUObject/Private/UObject/Enum.cpp b/Engine/Source/Runtime/CoreUObject/Private/UObject/Enum.cpp
--- a/Engine/Source/Runtime/CoreUObject/Private/UObject/Enum.cpp
+++ b/Engine/Source/Runtime/CoreUObject/Private/UObject/Enum.cpp
@@ -289,8 +289,8 @@ FText UEnum::GetDisplayNameText(int32 NameIndex) const
 
 	static const FString Namespace = TEXT("UObjectDisplayNames");
 	const FString Key =	NameIndex == INDEX_NONE
-		?			GetFullGroupName(true) + TEXT(".") + GetName()
-		:			GetFullGroupName(true) + TEXT(".") + GetName() + TEXT(".") + GetEnumName(NameIndex);
+		?			GetFullGroupName(false)
+		:			GetFullGroupName(false) + TEXT(".") + GetEnumName(NameIndex);
 
 
 	FString NativeDisplayName;
@@ -325,8 +325,8 @@ FText UEnum::GetToolTipText(int32 NameIndex) const
 
 	static const FString Namespace = TEXT("UObjectToolTips");
 	FString Key =	NameIndex == INDEX_NONE
-		?			GetFullGroupName(true) + TEXT(".") + GetName()
-		:			GetFullGroupName(true) + TEXT(".") + GetName() + TEXT(".") + GetEnumName(NameIndex);
+		?			GetFullGroupName(false)
+		:			GetFullGroupName(false) + TEXT(".") + GetEnumName(NameIndex);
 		
 	if ( !(FText::FindText( Namespace, Key, /*OUT*/LocalizedToolTip )) || *FTextInspector::GetSourceString(LocalizedToolTip) != NativeToolTip)
 	{
