[UE-9724] Re-run construction scripts after changing properties/attachments on instanced components.

change summary:
- modified UActorComponent::ConsolidatedPostEditChange() to call RRCS after re-registering the component post-edit
- modified SSCS_RowWidget::PostDragDropAction() to call RRCS after attach/detach/make root within the scene hierarchy
- modified SSCSEditor::AddNewComponent() to call RRCS after creating and registering a new instanced component
- modified FComponentEditorUtils::PasteComponents() to call RRCS after all pasted components have been processed
- modified FComponentEditorUtils::DeleteComponents() to call RRCS after the given set of components has been removed
- modified FComponentEditorUtils::DuplicateComponent() to call RRCS after the given component has been duplicated

diff --git a/Engine/Source/Editor/Kismet/Private/SSCSEditor.cpp b/Engine/Source/Editor/Kismet/Private/SSCSEditor.cpp
--- a/Engine/Source/Editor/Kismet/Private/SSCSEditor.cpp
+++ b/Engine/Source/Editor/Kismet/Private/SSCSEditor.cpp
@@ -2548,6 +2548,14 @@ void SSCS_RowWidget::PostDragDropAction(bool bRegenerateTreeNodes)
 				}
 			}
 		}
+		else
+		{
+			AActor* ActorInstance = PinnedEditor->GetActorContext();
+			if(ActorInstance)
+			{
+				ActorInstance->RerunConstructionScripts();
+			}
+		}
 	}
 }
 
@@ -4379,6 +4387,9 @@ UActorComponent* SSCSEditor::AddNewComponent( UClass* NewComponentClass, UObject
 			NewInstanceComponent->OnComponentCreated();
 			NewInstanceComponent->RegisterComponent();
 
+			// Rerun construction scripts
+			ActorInstance->RerunConstructionScripts();
+
 			NewComponent = AddNewNodeForInstancedComponent(NewInstanceComponent, Asset);
 		}
 	}
diff --git a/Engine/Source/Editor/UnrealEd/Private/Kismet2/ComponentEditorUtils.cpp b/Engine/Source/Editor/UnrealEd/Private/Kismet2/ComponentEditorUtils.cpp
--- a/Engine/Source/Editor/UnrealEd/Private/Kismet2/ComponentEditorUtils.cpp
+++ b/Engine/Source/Editor/UnrealEd/Private/Kismet2/ComponentEditorUtils.cpp
@@ -394,6 +394,9 @@ void FComponentEditorUtils::PasteComponents(TArray<UActorComponent*>& OutPastedC
 
 		OutPastedComponents.Add(NewActorComponent);
 	}
+
+	// Rerun construction scripts
+	TargetActor->RerunConstructionScripts();
 }
 
 void FComponentEditorUtils::GetComponentsFromClipboard(TMap<FName, FName>& OutParentMap, TMap<FName, UActorComponent*>& OutNewObjectMap, bool bGetComponentsAsArchetypes)
@@ -431,6 +434,8 @@ int32 FComponentEditorUtils::DeleteComponents(const TArray<UActorComponent*>& Co
 {
 	int32 NumDeletedComponents = 0;
 
+	TArray<AActor*> ActorsToReconstruct;
+
 	for (auto ComponentToDelete : ComponentsToDelete)
 	{
 		if (ComponentToDelete->CreationMethod != EComponentCreationMethod::Instance)
@@ -440,10 +445,13 @@ int32 FComponentEditorUtils::DeleteComponents(const TArray<UActorComponent*>& Co
 			continue;
 		}
 
+		AActor* Owner = ComponentToDelete->GetOwner();
+		check(Owner != nullptr);
+
 		// If necessary, determine the component that should be selected following the deletion of the indicated component
 		if (!OutComponentToSelect || ComponentToDelete == OutComponentToSelect)
 		{
-			USceneComponent* RootComponent = ComponentToDelete->GetOwner()->GetRootComponent();
+			USceneComponent* RootComponent = Owner->GetRootComponent();
 			if (RootComponent != ComponentToDelete)
 			{
 				// Worst-case, the root can be selected
@@ -472,7 +480,7 @@ int32 FComponentEditorUtils::DeleteComponents(const TArray<UActorComponent*>& Co
 				{
 					// For a non-scene component, try to select the preceding non-scene component
 					TInlineComponentArray<UActorComponent*> ActorComponents;
-					ComponentToDelete->GetOwner()->GetComponents(ActorComponents);
+					Owner->GetComponents(ActorComponents);
 					for (int32 i = 0; i < ActorComponents.Num() && ComponentToDelete != ActorComponents[i]; ++i)
 					{
 						if (!ActorComponents[i]->IsA(USceneComponent::StaticClass()))
@@ -488,12 +496,21 @@ int32 FComponentEditorUtils::DeleteComponents(const TArray<UActorComponent*>& Co
 			}
 		}
 
+		// Defer reconstruction
+		ActorsToReconstruct.AddUnique(Owner);
+
 		// Actually delete the component
 		ComponentToDelete->Modify();
 		ComponentToDelete->DestroyComponent(true);
 		NumDeletedComponents++;
 	}
 
+	// Reconstruct owner instance(s) after deletion
+	for(auto ActorToReconstruct : ActorsToReconstruct)
+	{
+		ActorToReconstruct->RerunConstructionScripts();
+	}
+
 	return NumDeletedComponents;
 }
 
@@ -547,6 +564,9 @@ UActorComponent* FComponentEditorUtils::DuplicateComponent(UActorComponent* Temp
 		
 		// Register the new component
 		NewCloneComponent->RegisterComponent();
+
+		// Rerun construction scripts
+		Actor->RerunConstructionScripts();
 	}
 
 	return NewCloneComponent;
diff --git a/Engine/Source/Runtime/Engine/Private/ActorComponent.cpp b/Engine/Source/Runtime/Engine/Private/ActorComponent.cpp
--- a/Engine/Source/Runtime/Engine/Private/ActorComponent.cpp
+++ b/Engine/Source/Runtime/Engine/Private/ActorComponent.cpp
@@ -495,6 +495,12 @@ void UActorComponent::ConsolidatedPostEditChange()
 	{
 		delete ReregisterContext;
 		EditReregisterContexts.Remove(this);
+
+		AActor* Owner = GetOwner();
+		if (Owner && !Owner->IsTemplate())
+		{
+			Owner->RerunConstructionScripts();
+		}
 	}
 
 	// The component or its outer could be pending kill when calling PostEditChange when applying a transaction.
