Critical bug in FBlueprintCompileReinstancer::ReinstanceObjects - objects are reinstanced again

diff --git a/Engine/Source/Editor/UnrealEd/Private/Kismet2/KismetReinstanceUtilities.cpp b/Engine/Source/Editor/UnrealEd/Private/Kismet2/KismetReinstanceUtilities.cpp
--- a/Engine/Source/Editor/UnrealEd/Private/Kismet2/KismetReinstanceUtilities.cpp
+++ b/Engine/Source/Editor/UnrealEd/Private/Kismet2/KismetReinstanceUtilities.cpp
@@ -388,7 +388,7 @@ void FBlueprintCompileReinstancer::ReinstanceObjects(bool bForceAlwaysReinstance
 		TSharedPtr<FBlueprintCompileReinstancer> SharedThis = AsShared();
 		const bool bAlreadyQueued = QueueToReinstance.Contains(SharedThis);
 		ensure(!bAlreadyQueued);
-		if (bAlreadyQueued)
+		if (!bAlreadyQueued)
 		{
 			QueueToReinstance.Push(SharedThis);
 
