#FBXImport: Allow to clear up existing blendshape curves when you import animation
Merging

//depot/UE4-GDC2015/Engine/Source/Editor/UnrealEd/...

to //depot/UE4/Engine/Source/Editor/UnrealEd/...

diff --git a/Engine/Source/Editor/UnrealEd/Classes/Factories/FbxImportUI.h b/Engine/Source/Editor/UnrealEd/Classes/Factories/FbxImportUI.h
--- a/Engine/Source/Editor/UnrealEd/Classes/Factories/FbxImportUI.h
+++ b/Engine/Source/Editor/UnrealEd/Classes/Factories/FbxImportUI.h
@@ -111,6 +111,10 @@ class UFbxImportUI : public UObject
 	UPROPERTY(EditAnywhere, AdvancedDisplay, config, Category=Animation, meta=(editcondition="bImportAnimations", ImportType = "SkeletalMesh|Animation"))
 	bool bPreserveLocalTransform;
 
+	/** Type of asset to import from the FBX file */
+	UPROPERTY(EditAnywhere, AdvancedDisplay, config, Category=Animation, meta=(editcondition="bImportAnimations", ImportType = "SkeletalMesh|Animation"))
+	bool bDeleteExistingMorphTargetCurves;
+
 	/** Import data used when importing textures */
 	UPROPERTY(EditAnywhere, Instanced, Category=Material)
 	class UFbxTextureImportData* TextureImportData;
diff --git a/Engine/Source/Editor/UnrealEd/Private/Fbx/FbxMainImport.cpp b/Engine/Source/Editor/UnrealEd/Private/Fbx/FbxMainImport.cpp
--- a/Engine/Source/Editor/UnrealEd/Private/Fbx/FbxMainImport.cpp
+++ b/Engine/Source/Editor/UnrealEd/Private/Fbx/FbxMainImport.cpp
@@ -244,6 +244,7 @@ void ApplyImportUIToImportOptions(UFbxImportUI* ImportUI, FBXImportOptions& InOu
 	InOutImportOptions.AnimationRange.Y = ImportUI->AnimSequenceImportData->EndFrame;
 	InOutImportOptions.AnimationName = ImportUI->AnimationName;
 	InOutImportOptions.bPreserveLocalTransform = ImportUI->bPreserveLocalTransform;
+	InOutImportOptions.bDeleteExistingMorphTargetCurves = ImportUI->bDeleteExistingMorphTargetCurves;
 	InOutImportOptions.bImportCustomAttribute = ImportUI->AnimSequenceImportData->bImportCustomAttribute;
 }
 
diff --git a/Engine/Source/Editor/UnrealEd/Private/FbxImporter.h b/Engine/Source/Editor/UnrealEd/Private/FbxImporter.h
--- a/Engine/Source/Editor/UnrealEd/Private/FbxImporter.h
+++ b/Engine/Source/Editor/UnrealEd/Private/FbxImporter.h
@@ -145,6 +145,7 @@ struct FBXImportOptions
 	struct FIntPoint AnimationRange;
 	FString AnimationName;
 	bool	bPreserveLocalTransform;
+	bool	bDeleteExistingMorphTargetCurves;
 	bool	bImportCustomAttribute;
 
 	bool ShouldImportNormals()
diff --git a/Engine/Source/Editor/UnrealEd/Private/SkeletalMeshEdit.cpp b/Engine/Source/Editor/UnrealEd/Private/SkeletalMeshEdit.cpp
--- a/Engine/Source/Editor/UnrealEd/Private/SkeletalMeshEdit.cpp
+++ b/Engine/Source/Editor/UnrealEd/Private/SkeletalMeshEdit.cpp
@@ -996,6 +996,21 @@ bool UnFbx::FFbxImporter::ImportAnimation(USkeleton* Skeleton, UAnimSequence * D
 		}
 	}
 
+	if (ImportOptions->bDeleteExistingMorphTargetCurves)
+	{
+		for (int32 CurveIdx=0; CurveIdx<DestSeq->RawCurveData.FloatCurves.Num(); ++CurveIdx)
+		{
+			auto& Curve = DestSeq->RawCurveData.FloatCurves[CurveIdx];
+			if (Curve.GetCurveTypeFlag(ACF_DrivesMorphTarget))
+			{
+				DestSeq->RawCurveData.FloatCurves.RemoveAt(CurveIdx, 1, false);
+				--CurveIdx;
+			}
+		}
+
+		DestSeq->RawCurveData.FloatCurves.Shrink();
+	}
+
 	//
 	// import blend shape curves
 	//
