AssetRegistry module needs to be loaded before fixupredirects starts serialization

diff --git a/Engine/Source/Editor/UnrealEd/Private/Commandlets/FixupRedirectsCommandlet.cpp b/Engine/Source/Editor/UnrealEd/Private/Commandlets/FixupRedirectsCommandlet.cpp
--- a/Engine/Source/Editor/UnrealEd/Private/Commandlets/FixupRedirectsCommandlet.cpp
+++ b/Engine/Source/Editor/UnrealEd/Private/Commandlets/FixupRedirectsCommandlet.cpp
@@ -26,6 +26,9 @@ void UFixupRedirectsCommandlet::CreateCustomEngine(const FString& Params)
 
 int32 UFixupRedirectsCommandlet::Main( const FString& Params )
 {
+	// Loading asset registry during serialization below will assert
+	FAssetRegistryModule& AssetRegistryModule = FModuleManager::LoadModuleChecked<FAssetRegistryModule>(TEXT("AssetRegistry"));
+
 	// Retrieve list of all packages in .ini paths.
 	TArray<FString> PackageList;
 	FEditorFileUtils::FindAllPackageFiles(PackageList);
