UE-10298 Creating an exposed variable does not always update correctly to the bp's Details panel

When a selected actor is reinstanced, a details view (displaying details of the actor) needs to be refreshed.

An alternative solution: don't check SDetailsView::ShouldSetNewObjects (refresh always )
OR store an object's hash in SDetailsView (and compare it in the ShouldSetNewObjects)

diff --git a/Engine/Source/Editor/LevelEditor/Private/LevelEditor.cpp b/Engine/Source/Editor/LevelEditor/Private/LevelEditor.cpp
--- a/Engine/Source/Editor/LevelEditor/Private/LevelEditor.cpp
+++ b/Engine/Source/Editor/LevelEditor/Private/LevelEditor.cpp
@@ -372,9 +372,9 @@ void FLevelEditorModule::FocusViewport()
 	}
 }
 
-void FLevelEditorModule::BroadcastActorSelectionChanged( const TArray<UObject*>& NewSelection )
+void FLevelEditorModule::BroadcastActorSelectionChanged(const TArray<UObject*>& NewSelection, bool bForceRefresh)
 {
-	ActorSelectionChangedEvent.Broadcast( NewSelection );
+	ActorSelectionChangedEvent.Broadcast(NewSelection, bForceRefresh);
 }
 
 void FLevelEditorModule::BroadcastRedrawViewports( bool bInvalidateHitProxies )
diff --git a/Engine/Source/Editor/LevelEditor/Private/SActorDetails.cpp b/Engine/Source/Editor/LevelEditor/Private/SActorDetails.cpp
--- a/Engine/Source/Editor/LevelEditor/Private/SActorDetails.cpp
+++ b/Engine/Source/Editor/LevelEditor/Private/SActorDetails.cpp
@@ -144,11 +144,11 @@ SActorDetails::~SActorDetails()
 	}
 }
 
-void SActorDetails::SetObjects(const TArray<UObject*>& InObjects)
+void SActorDetails::SetObjects(const TArray<UObject*>& InObjects, bool bForceRefresh)
 {
 	if(!DetailsView->IsLocked())
 	{
-		DetailsView->SetObjects(InObjects);
+		DetailsView->SetObjects(InObjects, bForceRefresh);
 
 		bool bShowingComponents = false;
 
diff --git a/Engine/Source/Editor/LevelEditor/Private/SActorDetails.h b/Engine/Source/Editor/LevelEditor/Private/SActorDetails.h
--- a/Engine/Source/Editor/LevelEditor/Private/SActorDetails.h
+++ b/Engine/Source/Editor/LevelEditor/Private/SActorDetails.h
@@ -19,7 +19,7 @@ public:
 	 *
 	 * @param InObjects	The objects to set
 	 */
-	void SetObjects(const TArray<UObject*>& InObjects);
+	void SetObjects(const TArray<UObject*>& InObjects, bool bForceRefresh = false);
 
 	/** FEditorUndoClient Interface */
 	virtual void PostUndo(bool bSuccess) override;
diff --git a/Engine/Source/Editor/LevelEditor/Private/SLevelEditor.cpp b/Engine/Source/Editor/LevelEditor/Private/SLevelEditor.cpp
--- a/Engine/Source/Editor/LevelEditor/Private/SLevelEditor.cpp
+++ b/Engine/Source/Editor/LevelEditor/Private/SLevelEditor.cpp
@@ -1423,14 +1423,14 @@ void SLevelEditor::HandleEditorMapChange( uint32 MapChangeFlags )
 	}
 }
 
-void SLevelEditor::OnActorSelectionChanged( const TArray<UObject*>& NewSelection )
+void SLevelEditor::OnActorSelectionChanged(const TArray<UObject*>& NewSelection, bool bForceRefresh)
 {
 	for( auto It = AllActorDetailPanels.CreateIterator(); It; ++It )
 	{
 		TSharedPtr<SActorDetails> ActorDetails = It->Pin();
 		if( ActorDetails.IsValid() )
 		{
-			ActorDetails->SetObjects( NewSelection );
+			ActorDetails->SetObjects(NewSelection, bForceRefresh);
 		}
 		else
 		{
diff --git a/Engine/Source/Editor/LevelEditor/Private/SLevelEditor.h b/Engine/Source/Editor/LevelEditor/Private/SLevelEditor.h
--- a/Engine/Source/Editor/LevelEditor/Private/SLevelEditor.h
+++ b/Engine/Source/Editor/LevelEditor/Private/SLevelEditor.h
@@ -191,7 +191,7 @@ private:
 	void HandleEditorMapChange( uint32 MapChangeFlags );
 
 	/** Called when actors are selected or unselected */
-	void OnActorSelectionChanged( const TArray<UObject*>& NewSelection );
+	void OnActorSelectionChanged(const TArray<UObject*>& NewSelection, bool bForceRefresh = false);
 private:
 
 	// Tracking the active viewports in this level editor.
diff --git a/Engine/Source/Editor/LevelEditor/Private/SLevelViewport.cpp b/Engine/Source/Editor/LevelEditor/Private/SLevelViewport.cpp
--- a/Engine/Source/Editor/LevelEditor/Private/SLevelViewport.cpp
+++ b/Engine/Source/Editor/LevelEditor/Private/SLevelViewport.cpp
@@ -2380,7 +2380,7 @@ bool SLevelViewport::IsPlayInEditorViewportActive() const
 	return ActiveViewport->IsPlayInEditorViewport();
 }
 
-void SLevelViewport::OnActorSelectionChanged( const TArray<UObject*>& NewSelection )
+void SLevelViewport::OnActorSelectionChanged(const TArray<UObject*>& NewSelection, bool bForceRefresh)
 {
 	// On the first actor selection after entering Game View, enable the selection show flag
 	if (IsVisible() && IsInGameView() && NewSelection.Num() != 0)
diff --git a/Engine/Source/Editor/LevelEditor/Public/LevelEditor.h b/Engine/Source/Editor/LevelEditor/Public/LevelEditor.h
--- a/Engine/Source/Editor/LevelEditor/Public/LevelEditor.h
+++ b/Engine/Source/Editor/LevelEditor/Public/LevelEditor.h
@@ -161,7 +161,7 @@ public:
 	 * 
 	 * @param NewSelection	List of objects that are now selected
 	 */
-	virtual void BroadcastActorSelectionChanged( const TArray<UObject*>& NewSelection );
+	virtual void BroadcastActorSelectionChanged(const TArray<UObject*>& NewSelection, bool bForceRefresh=false);
 	
 	/**
 	 * Called by the engine when level editing viewports need to be redrawn
@@ -187,7 +187,7 @@ public:
 	virtual void BroadcastComponentsEdited();
 
 	/** Called when actor selection changes */
-	DECLARE_EVENT_OneParam( FLevelEditorModule, FActorSelectionChangedEvent, const TArray<UObject*>& );
+	DECLARE_EVENT_TwoParams(FLevelEditorModule, FActorSelectionChangedEvent, const TArray<UObject*>&, bool);
 	virtual FActorSelectionChangedEvent& OnActorSelectionChanged() { return ActorSelectionChangedEvent; }
 
 	/** Called when level editor viewports should be redrawn */
diff --git a/Engine/Source/Editor/LevelEditor/Public/SLevelViewport.h b/Engine/Source/Editor/LevelEditor/Public/SLevelViewport.h
--- a/Engine/Source/Editor/LevelEditor/Public/SLevelViewport.h
+++ b/Engine/Source/Editor/LevelEditor/Public/SLevelViewport.h
@@ -142,7 +142,7 @@ public:
 	 * 
 	 * @param NewSelection	List of objects that are now selected
 	 */
-	void OnActorSelectionChanged( const TArray<UObject*>& NewSelection );
+	void OnActorSelectionChanged(const TArray<UObject*>& NewSelection, bool bForceRefresh=false);
 
 	/**
 	 * Called when game view should be toggled
diff --git a/Engine/Source/Editor/UnrealEd/Classes/Editor/EditorEngine.h b/Engine/Source/Editor/UnrealEd/Classes/Editor/EditorEngine.h
--- a/Engine/Source/Editor/UnrealEd/Classes/Editor/EditorEngine.h
+++ b/Engine/Source/Editor/UnrealEd/Classes/Editor/EditorEngine.h
@@ -1325,7 +1325,7 @@ public:
 	virtual void polySetAndClearPolyFlags( UModel* Model, uint32 SetBits, uint32 ClearBits, bool SelectedOnly, bool UpdateMaster );
 
 	// Selection.
-	virtual void SelectActor(AActor* Actor, bool bInSelected, bool bNotify, bool bSelectEvenIfHidden=false) {}
+	virtual void SelectActor(AActor* Actor, bool bInSelected, bool bNotify, bool bSelectEvenIfHidden = false, bool bForceRefresh = false) {}
 	virtual bool CanSelectActor(AActor* Actor, bool bInSelected, bool bSelectEvenIfHidden=false, bool bWarnIfLevelLocked=false) const { return true; }
 	virtual void SelectGroup(class AGroupActor* InGroupActor, bool bForceSelection=false, bool bInSelected=true, bool bNotify=true) {}
 	virtual void SelectComponent(class UActorComponent* Component, bool bInSelected, bool bNotify, bool bSelectEvenIfHidden = false) {}
diff --git a/Engine/Source/Editor/UnrealEd/Classes/Editor/UnrealEdEngine.h b/Engine/Source/Editor/UnrealEd/Classes/Editor/UnrealEdEngine.h
--- a/Engine/Source/Editor/UnrealEd/Classes/Editor/UnrealEdEngine.h
+++ b/Engine/Source/Editor/UnrealEd/Classes/Editor/UnrealEdEngine.h
@@ -178,7 +178,7 @@ class UNREALED_API UUnrealEdEngine : public UEditorEngine, public FNotifyHook
 	// End FNotify interface.
 
 	// Begin UEditorEngine Interface
-	virtual void SelectActor(AActor* Actor, bool InSelected, bool bNotify, bool bSelectEvenIfHidden=false) override;
+	virtual void SelectActor(AActor* Actor, bool InSelected, bool bNotify, bool bSelectEvenIfHidden = false, bool bForceRefresh = false) override;
 	virtual bool CanSelectActor(AActor* Actor, bool InSelected, bool bSelectEvenIfHidden=false, bool bWarnIfLevelLocked=false) const override;
 	virtual void SelectGroup(AGroupActor* InGroupActor, bool bForceSelection=false, bool bInSelected=true, bool bNotify=true) override;
 	virtual void SelectComponent(class UActorComponent* Component, bool bInSelected, bool bNotify, bool bSelectEvenIfHidden = false) override;
@@ -246,7 +246,7 @@ class UNREALED_API UUnrealEdEngine : public UEditorEngine, public FNotifyHook
 	void DrawComponentVisualizersHUD(const FViewport* Viewport, const FSceneView* View, FCanvas* Canvas);
 
 	/** Updates the property windows of selected actors */
-	virtual void UpdateFloatingPropertyWindows();
+	virtual void UpdateFloatingPropertyWindows(bool bForceRefresh=false);
 
 	/**
 	*	Updates the property windows of the actors in the supplied ActorList
@@ -254,7 +254,7 @@ class UNREALED_API UUnrealEdEngine : public UEditorEngine, public FNotifyHook
 	*	@param	ActorList	The list of actors whose property windows should be updated
 	*
 	*/
-	virtual void UpdateFloatingPropertyWindowsFromActorList( const TArray< UObject *>& ActorList );
+	virtual void UpdateFloatingPropertyWindowsFromActorList(const TArray< UObject *>& ActorList, bool bForceRefresh=false);
 
 	/**
 	 * Fast track function to set render thread flags marking selection rather than reconnecting all components
diff --git a/Engine/Source/Editor/UnrealEd/Private/Dialogs/SMeshMergingDialog.cpp b/Engine/Source/Editor/UnrealEd/Private/Dialogs/SMeshMergingDialog.cpp
--- a/Engine/Source/Editor/UnrealEd/Private/Dialogs/SMeshMergingDialog.cpp
+++ b/Engine/Source/Editor/UnrealEd/Private/Dialogs/SMeshMergingDialog.cpp
@@ -309,7 +309,7 @@ void SMeshMergingDialog::SetImportVertexColors(ECheckBoxState NewValue)
 	MergingSettings.bImportVertexColors = (ECheckBoxState::Checked == NewValue);
 }
 
-void SMeshMergingDialog::OnActorSelectionChanged(const TArray<UObject*>& NewSelection)
+void SMeshMergingDialog::OnActorSelectionChanged(const TArray<UObject*>& NewSelection, bool bForceRefresh)
 {
 	GenerateNewPackageName();
 }
diff --git a/Engine/Source/Editor/UnrealEd/Private/Dialogs/SMeshProxyDialog.cpp b/Engine/Source/Editor/UnrealEd/Private/Dialogs/SMeshProxyDialog.cpp
--- a/Engine/Source/Editor/UnrealEd/Private/Dialogs/SMeshProxyDialog.cpp
+++ b/Engine/Source/Editor/UnrealEd/Private/Dialogs/SMeshProxyDialog.cpp
@@ -79,7 +79,7 @@ protected:
 	FReply OnMergeClicked();
 
 	/** Called when actors selection is changed */
-	void OnActorSelectionChanged(const TArray<UObject*>& NewSelection);
+	void OnActorSelectionChanged(const TArray<UObject*>& NewSelection, bool bForceRefresh = false);
 	
 	/** Generates destination package name using currently selected actors */
 	void GenerateNewProxyPackageName();
@@ -643,7 +643,7 @@ FReply SMeshProxyDialog::OnMergeClicked()
 	
 }
 
-void SMeshProxyDialog::OnActorSelectionChanged(const TArray<UObject*>& NewSelection)
+void SMeshProxyDialog::OnActorSelectionChanged(const TArray<UObject*>& NewSelection, bool bForceRefresh)
 {
 	GenerateNewProxyPackageName();
 }
diff --git a/Engine/Source/Editor/UnrealEd/Private/EditorHook.cpp b/Engine/Source/Editor/UnrealEd/Private/EditorHook.cpp
--- a/Engine/Source/Editor/UnrealEd/Private/EditorHook.cpp
+++ b/Engine/Source/Editor/UnrealEd/Private/EditorHook.cpp
@@ -39,7 +39,7 @@ void UUnrealEdEngine::NotifyPostChange( const FPropertyChangedEvent& PropertyCha
 	}
 }
 
-void UUnrealEdEngine::UpdateFloatingPropertyWindows()
+void UUnrealEdEngine::UpdateFloatingPropertyWindows(bool bForceRefresh)
 {
 	TArray<UObject*> SelectedObjects;
 
@@ -80,13 +80,13 @@ void UUnrealEdEngine::UpdateFloatingPropertyWindows()
 		}
 	}
 
-	UpdateFloatingPropertyWindowsFromActorList( SelectedObjects );
+	UpdateFloatingPropertyWindowsFromActorList(SelectedObjects, bForceRefresh);
 }
 
 
-void UUnrealEdEngine::UpdateFloatingPropertyWindowsFromActorList( const TArray< UObject *>& ActorList )
+void UUnrealEdEngine::UpdateFloatingPropertyWindowsFromActorList(const TArray< UObject *>& ActorList, bool bForceRefresh)
 {
 	FLevelEditorModule& LevelEditor = FModuleManager::LoadModuleChecked<FLevelEditorModule>(TEXT("LevelEditor"));
 
-	LevelEditor.BroadcastActorSelectionChanged( ActorList );
+	LevelEditor.BroadcastActorSelectionChanged(ActorList, bForceRefresh);
 }
diff --git a/Engine/Source/Editor/UnrealEd/Private/EditorSelectUtils.cpp b/Engine/Source/Editor/UnrealEd/Private/EditorSelectUtils.cpp
--- a/Engine/Source/Editor/UnrealEd/Private/EditorSelectUtils.cpp
+++ b/Engine/Source/Editor/UnrealEd/Private/EditorSelectUtils.cpp
@@ -494,7 +494,7 @@ bool UUnrealEdEngine::CanSelectActor(AActor* Actor, bool bInSelected, bool bSele
 	return bSelectionAllowed;
 }
 
-void UUnrealEdEngine::SelectActor(AActor* Actor, bool bInSelected, bool bNotify, bool bSelectEvenIfHidden)
+void UUnrealEdEngine::SelectActor(AActor* Actor, bool bInSelected, bool bNotify, bool bSelectEvenIfHidden, bool bForceRefresh)
 {
 	const bool bWarnIfLevelLocked = true;
 	if( !CanSelectActor( Actor, bInSelected, bSelectEvenIfHidden, bWarnIfLevelLocked ) )
@@ -599,10 +599,10 @@ void UUnrealEdEngine::SelectActor(AActor* Actor, bool bInSelected, bool bNotify,
 		}
 		else
 		{
-			if( bNotify )
+			if (bNotify || bForceRefresh)
 			{
 				//reset the property windows.  In case something has changed since previous selection
-				UpdateFloatingPropertyWindows();
+				UpdateFloatingPropertyWindows(bForceRefresh);
 			}
 		}
 	}
diff --git a/Engine/Source/Editor/UnrealEd/Private/Kismet2/KismetReinstanceUtilities.cpp b/Engine/Source/Editor/UnrealEd/Private/Kismet2/KismetReinstanceUtilities.cpp
--- a/Engine/Source/Editor/UnrealEd/Private/Kismet2/KismetReinstanceUtilities.cpp
+++ b/Engine/Source/Editor/UnrealEd/Private/Kismet2/KismetReinstanceUtilities.cpp
@@ -289,6 +289,8 @@ void FBlueprintCompileReinstancer::ReinstanceFast()
 
 		if ((!Obj->IsTemplate() || bIsComponent) && !Obj->IsPendingKill())
 		{
+			const bool bIsSelected = bIsActor ? Obj->IsSelected() : false;
+
 			Obj->SetClass(ClassToReinstance);
 
 			if (bIsActor)
@@ -296,6 +298,11 @@ void FBlueprintCompileReinstancer::ReinstanceFast()
 				auto Actor = CastChecked<AActor>(Obj);
 				Actor->ReregisterAllComponents();
 				Actor->RerunConstructionScripts();
+
+				if (bIsSelected)
+				{
+					GEditor->SelectActor(Actor, /*bInSelected =*/true, /*bNotify =*/true, false, true);
+				}
 			}
 
 			if (bIsAnimInstance)
diff --git a/Engine/Source/Editor/UnrealEd/Public/Dialogs/SMeshMergingDialog.h b/Engine/Source/Editor/UnrealEd/Public/Dialogs/SMeshMergingDialog.h
--- a/Engine/Source/Editor/UnrealEd/Public/Dialogs/SMeshMergingDialog.h
+++ b/Engine/Source/Editor/UnrealEd/Public/Dialogs/SMeshMergingDialog.h
@@ -54,7 +54,7 @@ private:
 	void SetPlaceInWorld(ECheckBoxState NewValue);
 	
 	/** Called when actors selection is changed */
-	void OnActorSelectionChanged(const TArray<UObject*>& NewSelection);
+	void OnActorSelectionChanged(const TArray<UObject*>& NewSelection, bool bForceRefresh=false);
 	
 	/** Generates destination package name using currently selected actors */
 	void GenerateNewPackageName();
