Fixing a build error.

diff --git a/Engine/Source/Editor/UnrealEd/Private/Commandlets/FixupRedirectsCommandlet.cpp b/Engine/Source/Editor/UnrealEd/Private/Commandlets/FixupRedirectsCommandlet.cpp
--- a/Engine/Source/Editor/UnrealEd/Private/Commandlets/FixupRedirectsCommandlet.cpp
+++ b/Engine/Source/Editor/UnrealEd/Private/Commandlets/FixupRedirectsCommandlet.cpp
@@ -8,6 +8,7 @@
 #include "ISourceControlModule.h"
 #include "Engine/UserDefinedEnum.h"
 #include "AutoSaveUtils.h"
+#include "AssetRegistryModule.h"
 
 DEFINE_LOG_CATEGORY_STATIC(LogFixupRedirectsCommandlet, Log, All);
 
