Fix deprecation warning

diff --git a/Engine/Source/Editor/UnrealEd/Private/Factories/TTFontImport.cpp b/Engine/Source/Editor/UnrealEd/Private/Factories/TTFontImport.cpp
--- a/Engine/Source/Editor/UnrealEd/Private/Factories/TTFontImport.cpp
+++ b/Engine/Source/Editor/UnrealEd/Private/Factories/TTFontImport.cpp
@@ -1671,7 +1671,7 @@ UTexture2D* UTrueTypeFontFactory::CreateTextureFromBitmap( UFont* Font, uint8* B
 	}
 
 	// Create texture for page.
-	UTexture2D* Texture = NewNamedObject<UTexture2D>(Font, *TextureString);
+	UTexture2D* Texture = NewObject<UTexture2D>(Font, *TextureString);
 
 	// note RF_Public because font textures can be referenced directly by material expressions
 	Texture->SetFlags(RF_Public);
