Changed editor feedback contexts to always delete the progress window if it exists.

There was a potential vulnerability where a slow task widget would not get cleaned up resulting in it outliving the pointer to the scope stack. The context no longer holds a strong reference to the slow task widget in a bid to prevent it from outliving the data to which it is tied.

diff --git a/Engine/Source/Editor/UnrealEd/Private/FeedbackContextEditor.cpp b/Engine/Source/Editor/UnrealEd/Private/FeedbackContextEditor.cpp
--- a/Engine/Source/Editor/UnrealEd/Private/FeedbackContextEditor.cpp
+++ b/Engine/Source/Editor/UnrealEd/Private/FeedbackContextEditor.cpp
@@ -381,10 +381,11 @@ void FFeedbackContextEditor::StartSlowTask( const FText& Task, bool bShowCancelB
 				.CreateTitleBar(true)
 				.ActivateWhenFirstShown(true);
 
-			SlowTaskWidget = SNew(SSlowTaskWidget)
+			SlowTaskWindowRef->SetContent(
+				SNew(SSlowTaskWidget)
 				.ScopeStack(&ScopeStack)
-				.OnCancelClickedDelegate( OnCancelClicked );
-			SlowTaskWindowRef->SetContent( SlowTaskWidget.ToSharedRef() );
+				.OnCancelClickedDelegate( OnCancelClicked )
+			);
 
 			SlowTaskWindow = SlowTaskWindowRef;
 
@@ -402,21 +403,12 @@ void FFeedbackContextEditor::StartSlowTask( const FText& Task, bool bShowCancelB
 
 void FFeedbackContextEditor::FinalizeSlowTask()
 {
-	TSharedPtr<SWindow> ParentWindow;
-	if( FModuleManager::Get().IsModuleLoaded( "MainFrame" ) )
-	{
-		IMainFrameModule& MainFrame = FModuleManager::LoadModuleChecked<IMainFrameModule>( "MainFrame" );
-		ParentWindow = MainFrame.GetParentWindow();
-	}
-
-	if( GIsEditor && ParentWindow.IsValid())
+	auto Window = SlowTaskWindow.Pin();
+	if (Window.IsValid())
 	{
-		if( SlowTaskWindow.IsValid() )
-		{
-			SlowTaskWindow.Pin()->RequestDestroyWindow();
-			SlowTaskWindow.Reset();
-			SlowTaskWidget.Reset();
-		}
+		Window->SetContent(SNullWidget::NullWidget);
+		Window->RequestDestroyWindow();
+		SlowTaskWindow.Reset();
 	}
 
 	FFeedbackContext::FinalizeSlowTask( );
@@ -450,7 +442,7 @@ void FFeedbackContextEditor::ProgressReported( const float TotalProgressInterp,
 			BuildProgressWidget->SetBuildProgressPercent(TotalProgressInterp * 100, 100);
 			TickSlate();
 		}
-		else if (SlowTaskWidget.IsValid())
+		else if (SlowTaskWindow.IsValid())
 		{
 			TickSlate();
 		}
diff --git a/Engine/Source/Editor/UnrealEd/Public/FeedbackContextEditor.h b/Engine/Source/Editor/UnrealEd/Public/FeedbackContextEditor.h
--- a/Engine/Source/Editor/UnrealEd/Public/FeedbackContextEditor.h
+++ b/Engine/Source/Editor/UnrealEd/Public/FeedbackContextEditor.h
@@ -14,7 +14,6 @@ class FFeedbackContextEditor : public FFeedbackContext
 {
 	/** Slate slow task widget */
 	TWeakPtr<class SWindow> SlowTaskWindow;
-	TSharedPtr<class SSlowTaskWidget> SlowTaskWidget;
 
 	/** Special Windows/Widget popup for building */
 	TWeakPtr<class SWindow> BuildProgressWindow;
