Merging 2456412

Fixed bBuildingNavigationFromUserRequest not being set correctly

//depot/UE4-UT/...

to //depot/UE4/...

diff --git a/Engine/Source/Editor/UnrealEd/Private/EditorBuildUtils.cpp b/Engine/Source/Editor/UnrealEd/Private/EditorBuildUtils.cpp
--- a/Engine/Source/Editor/UnrealEd/Private/EditorBuildUtils.cpp
+++ b/Engine/Source/Editor/UnrealEd/Private/EditorBuildUtils.cpp
@@ -776,7 +776,7 @@ void FEditorBuildUtils::TriggerNavigationBuilder(UWorld* InWorld, EBuildOptions:
 		case EBuildOptions::BuildAll:
 		case EBuildOptions::BuildAllSubmit:
 			{
-				bBuildingNavigationFromUserRequest = false;
+				bBuildingNavigationFromUserRequest = true;
 				break;
 			}
 		default:
