[UE-10537] - Only components with an attach parent transform the pivot relative to their parent

diff --git a/Engine/Source/Editor/UnrealEd/Private/LevelEditorViewport.cpp b/Engine/Source/Editor/UnrealEd/Private/LevelEditorViewport.cpp
--- a/Engine/Source/Editor/UnrealEd/Private/LevelEditorViewport.cpp
+++ b/Engine/Source/Editor/UnrealEd/Private/LevelEditorViewport.cpp
@@ -3133,7 +3133,7 @@ void FLevelEditorViewportClient::ApplyDeltaToComponent(USceneComponent* InCompon
 
 	// If necessary, transform the editor pivot location to be relative to the component's parent
 	const bool bIsRootComponent = InComponent->GetOwner()->GetRootComponent() == InComponent;
-	FVector RelativePivotLocation = bIsRootComponent ? EditorWorldPivotLocation : InComponent->GetAttachParent()->GetComponentToWorld().Inverse().TransformPosition(EditorWorldPivotLocation);
+	FVector RelativePivotLocation = bIsRootComponent || !InComponent->GetAttachParent() ? EditorWorldPivotLocation : InComponent->GetAttachParent()->GetComponentToWorld().Inverse().TransformPosition(EditorWorldPivotLocation);
 
 	GEditor->ApplyDeltaToComponent(
 		InComponent,
