[UE-10646] These casts should not be checked

diff --git a/Engine/Source/Editor/UnrealEd/Private/Kismet2/ComponentEditorUtils.cpp b/Engine/Source/Editor/UnrealEd/Private/Kismet2/ComponentEditorUtils.cpp
--- a/Engine/Source/Editor/UnrealEd/Private/Kismet2/ComponentEditorUtils.cpp
+++ b/Engine/Source/Editor/UnrealEd/Private/Kismet2/ComponentEditorUtils.cpp
@@ -208,8 +208,8 @@ USceneComponent* FComponentEditorUtils::FindClosestParentInList(UActorComponent*
 	USceneComponent* ClosestParentComponent = nullptr;
 	for (auto Component : ComponentList)
 	{
-		auto ChildAsScene = CastChecked<USceneComponent>(ChildComponent);
-		auto SceneComponent = CastChecked<USceneComponent>(Component);
+		auto ChildAsScene = Cast<USceneComponent>(ChildComponent);
+		auto SceneComponent = Cast<USceneComponent>(Component);
 		if (ChildAsScene && SceneComponent)
 		{
 			// Check to see if any parent is also in the list
