Fixed out-by-1 count error when removing scanned directories and files from the asynchronous directory reader

diff --git a/Engine/Source/Editor/UnrealEd/Private/AutoReimport/FileCache.cpp b/Engine/Source/Editor/UnrealEd/Private/AutoReimport/FileCache.cpp
--- a/Engine/Source/Editor/UnrealEd/Private/AutoReimport/FileCache.cpp
+++ b/Engine/Source/Editor/UnrealEd/Private/AutoReimport/FileCache.cpp
@@ -41,7 +41,7 @@ FAsyncDirectoryReader::EProgressResult FAsyncDirectoryReader::Tick(const FTimeLi
 		if (TimeLimit.Exceeded())
 		{
 			// We've spent too long, bail
-			PendingDirectories.RemoveAt(0, Index, false);
+			PendingDirectories.RemoveAt(0, Index + 1, false);
 			return EProgressResult::Pending;
 		}
 	}
@@ -58,7 +58,7 @@ FAsyncDirectoryReader::EProgressResult FAsyncDirectoryReader::Tick(const FTimeLi
 		if (Index % 100 == 0 && TimeLimit.Exceeded())
 		{
 			// We've spent too long, bail
-			PendingFiles.RemoveAt(0, Index, false);
+			PendingFiles.RemoveAt(0, Index + 1, false);
 			return EProgressResult::Pending;
 		}
 	}
