Fix staic mesh LODs not using base LOD materials by default when materials are not unique

diff --git a/Engine/Source/Editor/UnrealEd/Private/Fbx/FbxStaticMeshImport.cpp b/Engine/Source/Editor/UnrealEd/Private/Fbx/FbxStaticMeshImport.cpp
--- a/Engine/Source/Editor/UnrealEd/Private/Fbx/FbxStaticMeshImport.cpp
+++ b/Engine/Source/Editor/UnrealEd/Private/Fbx/FbxStaticMeshImport.cpp
@@ -1178,9 +1178,15 @@ UStaticMesh* UnFbx::FFbxImporter::ImportStaticMeshAsSingle(UObject* InParent, TA
 		for (int32 MaterialIndex = 0; MaterialIndex < NumMaterials; ++MaterialIndex)
 		{
 			FMeshSectionInfo Info = StaticMesh->SectionInfoMap.Get(LODIndex, MaterialIndex);
-			Info.MaterialIndex = StaticMesh->Materials.Num();
+			int32 Index = StaticMesh->Materials.Find( SortedMaterials[MaterialIndex] );
+			if( Index == INDEX_NONE )
+			{
+				Index = StaticMesh->Materials.Num();
+				StaticMesh->Materials.Add(SortedMaterials[MaterialIndex]);
+			}
+			Info.MaterialIndex = Index;
 			StaticMesh->SectionInfoMap.Set(LODIndex, MaterialIndex, Info);
-			StaticMesh->Materials.Add(SortedMaterials[MaterialIndex]);
+			
 		}
 
 	
