Remove source mip data after we populate ddc in cooker.

diff --git a/Engine/Source/Editor/UnrealEd/Private/Commandlets/CookCommandlet.cpp b/Engine/Source/Editor/UnrealEd/Private/Commandlets/CookCommandlet.cpp
--- a/Engine/Source/Editor/UnrealEd/Private/Commandlets/CookCommandlet.cpp
+++ b/Engine/Source/Editor/UnrealEd/Private/Commandlets/CookCommandlet.cpp
@@ -1221,7 +1221,7 @@ bool UCookCommandlet::NewCook( const TArray<ITargetPlatform*>& Platforms, TArray
 		if (NonMapPackageCountSinceLastGC > 0)
 		{
 			// We should GC if we have packages to collect and we've been idle for some time.
-			bShouldGC = (NonMapPackageCountSinceLastGC > PackagesPerGC) || 
+			bShouldGC |= (NonMapPackageCountSinceLastGC > PackagesPerGC) || 
 				((FPlatformTime::Seconds() - LastCookActionTime) >= IdleTimeToGC);
 		}
 
@@ -1397,7 +1397,8 @@ bool UCookCommandlet::Cook(const TArray<ITargetPlatform*>& Platforms, TArray<FSt
 						GetObjectsWithOuter(Pkg, ObjectsInPackage, true);
 						for( int32 IndexPackage = 0; IndexPackage < ObjectsInPackage.Num(); IndexPackage++ )
 						{
-							ObjectsInPackage[IndexPackage]->CookerWillNeverCookAgain();
+							ObjectsInPackage[IndexPackage]->WillNeverCacheCookedPlatformDataAgain();
+							ObjectsInPackage[IndexPackage]->ClearAllCachedCookedPlatformData();
 						}
 					}
 				}
diff --git a/Engine/Source/Editor/UnrealEd/Private/Commandlets/DerivedDataCacheCommandlet.cpp b/Engine/Source/Editor/UnrealEd/Private/Commandlets/DerivedDataCacheCommandlet.cpp
--- a/Engine/Source/Editor/UnrealEd/Private/Commandlets/DerivedDataCacheCommandlet.cpp
+++ b/Engine/Source/Editor/UnrealEd/Private/Commandlets/DerivedDataCacheCommandlet.cpp
@@ -139,7 +139,8 @@ int32 UDerivedDataCacheCommandlet::Main( const FString& Params )
 								GetObjectsWithOuter(Pkg, ObjectsInPackage, true);
 								for( int32 IndexPackage = 0; IndexPackage < ObjectsInPackage.Num(); IndexPackage++ )
 								{
-									ObjectsInPackage[IndexPackage]->CookerWillNeverCookAgain();
+									ObjectsInPackage[IndexPackage]->WillNeverCacheCookedPlatformDataAgain();
+									ObjectsInPackage[IndexPackage]->ClearAllCachedCookedPlatformData();
 								}
 							}
 						}
diff --git a/Engine/Source/Editor/UnrealEd/Private/CookOnTheFlyServer.cpp b/Engine/Source/Editor/UnrealEd/Private/CookOnTheFlyServer.cpp
--- a/Engine/Source/Editor/UnrealEd/Private/CookOnTheFlyServer.cpp
+++ b/Engine/Source/Editor/UnrealEd/Private/CookOnTheFlyServer.cpp
@@ -960,6 +960,14 @@ uint32 UCookOnTheFlyServer::TickCookOnTheSide( const float TimeSlice, uint32 &Co
 			AllTargetPlatformNames.AddUnique(PlatformName);
 		}
 
+		for ( const auto &PlatformName : AllTargetPlatformNames )
+		{
+			if ( ToBuild.HasPlatform(PlatformName) == false )
+			{
+				ToBuild.AddPlatform(PlatformName);
+			}
+		}
+
 		bool bWasUpToDate = false;
 
 		bool bLastLoadWasMap = false;
@@ -1121,69 +1129,87 @@ uint32 UCookOnTheFlyServer::TickCookOnTheSide( const float TimeSlice, uint32 &Co
 
 
 		bool bIsAllDataCached = true;
-		if ( PackagesToSave.Num() )
+
+		ITargetPlatformManagerModule& TPM = GetTargetPlatformManagerRef();
+		TArray<const ITargetPlatform*> TargetPlatforms;
+		for ( const auto& TargetPlatformName : AllTargetPlatformNames )
 		{
-			SCOPE_TIMER(CallBeginCacheForCookedPlatformData);
-			// cache the resources for this package for each platform
-			ITargetPlatformManagerModule& TPM = GetTargetPlatformManagerRef();
-			TArray<UObject*> ObjectsInPackage;
-			GetObjectsWithOuter( PackagesToSave[0], ObjectsInPackage );
-			for ( const auto& TargetPlatformName : AllTargetPlatformNames )
+			TargetPlatforms.Add( TPM.FindTargetPlatform( TargetPlatformName.ToString() ) );
+		}
+
+
+		auto BeginPackageCacheForCookedPlatformData = [&]( const TArray<UObject*>& ObjectsInPackage )
+		{
+			for ( int I = 0; I < ObjectsInPackage.Num(); ++I )
 			{
-				// avoid hitches in the editor by calling BeginCacheForCookedPlatformData on each object once using time slicing, 
-				// then check if the objects are finished caching by not using time slicing
-				const ITargetPlatform* TargetPlatform = TPM.FindTargetPlatform( TargetPlatformName.ToString() );
-				if ( CurrentReentryData.bBeginCacheFinished == false )
+				const auto& Obj = ObjectsInPackage[I];
+				for ( const auto& TargetPlatform : TargetPlatforms )
+				{
+					Obj->BeginCacheForCookedPlatformData( TargetPlatform );
+				}
+
+				if ( Timer.IsTimeUp() )
 				{
-					for ( int I = 0; I < ObjectsInPackage.Num(); ++I )
-					{
-						if ( CurrentReentryData.BeginCacheCount >= I )
-						{
-							const auto& Obj = ObjectsInPackage[I];
-							Obj->BeginCacheForCookedPlatformData( TargetPlatform );
-							
-							if ( Timer.IsTimeUp() )
-							{
 #if DEBUG_COOKONTHEFLY
-								UE_LOG(LogCook, Display, TEXT("Object %s took too long to cache"), *Obj->GetFullName());
+					UE_LOG(LogCook, Display, TEXT("Object %s took too long to cache"), *Obj->GetFullName());
 #endif
-								bIsAllDataCached = false;
-								break;
-							}
-						}
-					}
+					return false;
 				}
-				
-				if ( bIsAllDataCached )
+			}
+			return true;
+		};
+
+		auto FinishPackageCacheForCookedPlatformData = [&]( const TArray<UObject*>& ObjectsInPackage )
+		{
+			
+			// if we get here and bIsAllDataCached is true then BeginCacheFOrCookedPlatformData has been called once on every object
+			// so now 
+			CurrentReentryData.bBeginCacheFinished = true;
+			for ( const auto& Obj : ObjectsInPackage )
+			{
+				for ( const auto& TargetPlatform : TargetPlatforms )
 				{
-					// if we get here and bIsAllDataCached is true then BeginCacheFOrCookedPlatformData has been called once on every object
-					// so now 
-					CurrentReentryData.bBeginCacheFinished = true;
-					for ( const auto& Obj : ObjectsInPackage )
+					// These begin cache calls should be quick 
+					// because they will just be checking that the data is already cached and kicking off new multithreaded requests if not
+					// all sync requests should have been caught in the first begincache call above
+					Obj->BeginCacheForCookedPlatformData( TargetPlatform );
+					if ( Obj->IsCachedCookedPlatformDataLoaded(TargetPlatform) == false )
 					{
-						// These begin cache calls should be quick 
-						// because they will just be checking that the data is already cached and kicking off new multithreaded requests if not
-						// all sync requests should have been caught in the first begincache call above
-						Obj->BeginCacheForCookedPlatformData( TargetPlatform );
-						if ( Obj->IsCachedCookedPlatformDataLoaded(TargetPlatform) == false )
-						{
 #if DEBUG_COOKONTHEFLY
-							UE_LOG(LogCook, Display, TEXT("Object %s isn't cached yet"), *Obj->GetFullName());
+						UE_LOG(LogCook, Display, TEXT("Object %s isn't cached yet"), *Obj->GetFullName());
 #endif
-							bIsAllDataCached = false;
-						}
+						return false;
 					}
 				}
-				
 
-				if ( Timer.IsTimeUp() && (bIsAllDataCached == false) )
+				// if this objects data is cached then we can call FinishedCookedPLatformDataCache
+				// we can only safely call this when we are finished caching this object completely.
+				// this doesn't ever happen for cook in editor or cook on the fly mode
+				if ( bIsAllDataCached && (CurrentCookMode == ECookMode::CookByTheBook))
 				{
-					break; // break out of the target platform loop (can't break out of the main package tick loop here without requeing the request)
+					// this might be run multiple times for a single object
+					Obj->WillNeverCacheCookedPlatformDataAgain();
 				}
 			}
+			return true;
+		};
+
+		if ( PackagesToSave.Num() )
+		{
+			SCOPE_TIMER(CallBeginCacheForCookedPlatformData);
+			// cache the resources for this package for each platform
+			
+			TArray<UObject*> ObjectsInPackage;
+			GetObjectsWithOuter( PackagesToSave[0], ObjectsInPackage );
+
+			bIsAllDataCached &= BeginPackageCacheForCookedPlatformData(ObjectsInPackage);
+			if( bIsAllDataCached )
+			{
+				bIsAllDataCached &= FinishPackageCacheForCookedPlatformData(ObjectsInPackage);
+			}
 		}
 
-		if ( IsRealtimeMode() )
+		// if ( IsRealtimeMode() )
 		{
 			if ( bIsAllDataCached == false )
 			{
@@ -1225,7 +1251,7 @@ uint32 UCookOnTheFlyServer::TickCookOnTheSide( const float TimeSlice, uint32 &Co
 					{
 						Result |= COSR_RequiresGC;
 					}
-
+					
 					if (Package)
 					{
 						ACCUMULATE_TIMER_START(GetCachedName);
@@ -1270,11 +1296,6 @@ uint32 UCookOnTheFlyServer::TickCookOnTheSide( const float TimeSlice, uint32 &Co
 			}
 		}
 
-		{
-			SCOPE_TIMER(GenerateManifestInfo);
-			// update manifest with cooked package info
-			GenerateManifestInfo( PackagesToSave[0], AllTargetPlatformNames );
-		}
 
 		bool bFinishedSave = true;
 
@@ -1322,6 +1343,12 @@ uint32 UCookOnTheFlyServer::TickCookOnTheSide( const float TimeSlice, uint32 &Co
 					}
 				}
 
+				{
+					SCOPE_TIMER(GenerateManifestInfo);
+					// update manifest with cooked package info
+					GenerateManifestInfo( PackagesToSave[I], AllTargetPlatformNames );
+				}
+
 				UPackage *Package = PackagesToSave[I];
 				FName PackageFName = GetCachedStandardPackageFileFName(Package);
 				TArray<FName> SaveTargetPlatformNames = AllTargetPlatformNames;
@@ -1340,6 +1367,22 @@ uint32 UCookOnTheFlyServer::TickCookOnTheSide( const float TimeSlice, uint32 &Co
 					continue;
 				}
 
+				// precache platform data for next package 
+				UPackage *NextPackage = PackagesToSave[FMath::Min( PackagesToSave.Num()-1, I + 1 )];
+				UPackage *NextNextPackage = PackagesToSave[FMath::Min( PackagesToSave.Num()-1, I + 2 )];
+				if ( NextPackage != Package )
+				{
+					TArray<UObject*> ObjectsInPackage;
+					GetObjectsWithOuter( NextPackage, ObjectsInPackage );
+					BeginPackageCacheForCookedPlatformData(ObjectsInPackage);
+				}
+				if ( NextNextPackage != NextPackage )
+				{
+					TArray<UObject*> ObjectsInPackage;
+					GetObjectsWithOuter( NextNextPackage, ObjectsInPackage );
+					BeginPackageCacheForCookedPlatformData(ObjectsInPackage);
+				}
+
 				SCOPE_TIMER(SaveCookedPackage);
 				if( SaveCookedPackage(Package, SAVE_KeepGUID | SAVE_Async | (IsCookFlagSet(ECookInitializationFlags::Unversioned) ? SAVE_Unversioned : 0), bWasUpToDate, AllTargetPlatformNames ) )
 				{
@@ -1356,6 +1399,16 @@ uint32 UCookOnTheFlyServer::TickCookOnTheSide( const float TimeSlice, uint32 &Co
 				}
 				Timer.SavedPackage();
 
+				{
+					SCOPE_TIMER(ClearAllCachedCookedPlatformData);
+					TArray<UObject*> ObjectsInPackage;
+					GetObjectsWithOuter(Package, ObjectsInPackage);
+					for ( const auto& Object : ObjectsInPackage )
+					{
+						Object->ClearAllCachedCookedPlatformData();
+					}
+				}
+
 				//@todo ResetLoaders outside of this (ie when Package is NULL) causes problems w/ default materials
 				if (Package->HasAnyFlags(RF_RootSet) == false && ((CurrentCookMode==ECookMode::CookOnTheFly)) )
 				{
diff --git a/Engine/Source/Runtime/CoreUObject/Private/UObject/SavePackage.cpp b/Engine/Source/Runtime/CoreUObject/Private/UObject/SavePackage.cpp
--- a/Engine/Source/Runtime/CoreUObject/Private/UObject/SavePackage.cpp
+++ b/Engine/Source/Runtime/CoreUObject/Private/UObject/SavePackage.cpp
@@ -3713,11 +3713,12 @@ bool UPackage::SavePackage( UPackage* InOuter, UObject* Base, EObjectFlags TopLe
 
 				SlowTask.EnterProgressFrame();
 			
-
+#if WITH_EDITOR
 				for ( int CachedObjectIndex = 0; CachedObjectIndex < CachedObjects.Num(); ++CachedObjectIndex )
 				{
 					CachedObjects[CachedObjectIndex]->ClearCachedCookedPlatformData(TargetPlatform);
 				}
+#endif
 			}
 
 		}
diff --git a/Engine/Source/Runtime/CoreUObject/Public/UObject/UObject.h b/Engine/Source/Runtime/CoreUObject/Public/UObject/UObject.h
--- a/Engine/Source/Runtime/CoreUObject/Public/UObject/UObject.h
+++ b/Engine/Source/Runtime/CoreUObject/Public/UObject/UObject.h
@@ -181,11 +181,6 @@ public:
 	virtual bool Modify( bool bAlwaysMarkDirty=true );
 
 #if WITH_EDITOR
-	/**
-	 * The cooker (or DDC commandlet) have dealt with this object and will never deal with it again.
-	 * The object should try to free memory
-	 */
-	virtual void CookerWillNeverCookAgain() {}
 	/** 
 	 * Called when the object was loaded from another class via active class redirects.
 	 */
@@ -682,7 +677,7 @@ public:
 	 *								subobjects and components for a subobject root.
 	 */
 	void ConditionalPostLoadSubobjects( struct FObjectInstancingGraph* OuterInstanceGraph=NULL );
-
+#if WITH_EDITOR
 	/**
 	 * Starts caching of platform specific data for the target platform
 	 * Called when cooking before serialization so that object can prepare platform specific data
@@ -691,6 +686,18 @@ public:
 	 * @param	TargetPlatform	target platform to cache platform specific data for
 	 */
 	virtual void BeginCacheForCookedPlatformData( const ITargetPlatform* TargetPlatform ) {  }
+	
+	/**
+	 * Have we finished loading all the cooked platform data for the target platforms requested in BeginCacheForCookedPlatformData
+	 * 
+	 * @param	TargetPlatform target platform to check for cooked platform data
+	 */
+	virtual bool IsCachedCookedPlatformDataLoaded( const ITargetPlatform* TargetPlatform ) { return true; }
+
+	/**
+	 * All caching has finished for this object (all IsCachedCookedPlatformDataLoaded functions have finished for all platforms)
+	 */
+	virtual void WillNeverCacheCookedPlatformDataAgain() { }
 
 	/**
 	 * Clears cached cooked platform data for specific platform
@@ -705,15 +712,7 @@ public:
 	 * @param	TargetPlatform	target platform to cache platform specific data for
 	 */
 	virtual void ClearAllCachedCookedPlatformData() { }
-
-	/**
-	 * Have we finished loading all the cooked platform data for the target platforms requested in BeginCacheForCookedPlatformData
-	 * 
-	 * @param	TargetPlatform target platform to check for cooked platform data
-	 */
-	virtual bool IsCachedCookedPlatformDataLoaded( const ITargetPlatform* TargetPlatform ) { return true; }
-
-
+#endif
 	/**
 	 * Determine if this object has SomeObject in its archetype chain.
 	 */
diff --git a/Engine/Source/Runtime/Engine/Classes/Engine/Texture.h b/Engine/Source/Runtime/Engine/Classes/Engine/Texture.h
--- a/Engine/Source/Runtime/Engine/Classes/Engine/Texture.h
+++ b/Engine/Source/Runtime/Engine/Classes/Engine/Texture.h
@@ -198,6 +198,7 @@ struct FTextureSource
 	FORCEINLINE ETextureSourceFormat GetFormat() const { return Format; }
 	FORCEINLINE bool IsPNGCompressed() const { return bPNGCompressed; }
 	FORCEINLINE int32 GetSizeOnDisk() const { return BulkData.GetBulkDataSize(); }
+	FORCEINLINE bool IsBulkDataLoaded() const { return BulkData.IsBulkDataLoaded(); }
 #endif
 
 private:
@@ -212,19 +213,22 @@ private:
 	uint8* LockedMipData;
 	/** Which mips are locked, if any. */
 	uint32 LockedMips;
-
 #if WITH_EDITOR
 	/** Return true if the source art is not png compressed but could be. */
 	bool CanPNGCompress() const;
-
 	/** Removes source data. */
 	void RemoveSourceData();
-
 	/** Retrieve the size and offset for a source mip. The size includes all slices. */
 	int32 CalcMipOffset(int32 MipIndex) const;
 
 	/** Uses a hash as the GUID, useful to prevent creating new GUIDs on load for legacy assets. */
 	void UseHashAsGuid();
+public:
+	void ReleaseSourceMemory(); // release the memory from the mips (does almost the same as remove source data except doesn't rebuild the guid)
+	FORCEINLINE bool HasHadBulkDataCleared() const { return bHasHadBulkDataCleared; }
+private:
+	/** Used while cooking to clear out unneeded memory after compression */
+	bool bHasHadBulkDataCleared;
 #endif
 
 #if WITH_EDITORONLY_DATA
@@ -259,6 +263,8 @@ private:
 	/** Format in which the source data is stored. */
 	UPROPERTY(VisibleAnywhere, Category=TextureSource)
 	TEnumAsByte<enum ETextureSourceFormat> Format;
+
+
 #endif // WITH_EDITORONLY_DATA
 };
 
@@ -587,7 +593,6 @@ public:
 	virtual TMap<FString, FTexturePlatformData*>* GetCookedPlatformData() { return NULL; }
 
 	void CleanupCachedRunningPlatformData();
-	void CleanupCachedCookedPlatformData();
 
 	/**
 	 * Serializes cooked platform data.
@@ -612,6 +617,22 @@ public:
 	 */
 	ENGINE_API virtual bool IsCachedCookedPlatformDataLoaded( const ITargetPlatform* TargetPlatform ) override;
 
+
+	/**
+	 * Clears cached cooked platform data for specific platform
+	 * 
+	 * @param	TargetPlatform	target platform to cache platform specific data for
+	 */
+	ENGINE_API virtual void ClearCachedCookedPlatformData( const ITargetPlatform* TargetPlatform ) override;
+
+	/**
+	 * Clear all cached cooked platform data
+	 * 
+	 * @param	TargetPlatform	target platform to cache platform specific data for
+	 */
+	ENGINE_API virtual void ClearAllCachedCookedPlatformData() override;
+
+
 	/**
 	 * Begins caching platform data in the background.
 	 */
diff --git a/Engine/Source/Runtime/Engine/Classes/Engine/Texture2D.h b/Engine/Source/Runtime/Engine/Classes/Engine/Texture2D.h
--- a/Engine/Source/Runtime/Engine/Classes/Engine/Texture2D.h
+++ b/Engine/Source/Runtime/Engine/Classes/Engine/Texture2D.h
@@ -97,8 +97,10 @@ public:
 public:
 	/** The derived data for this texture on this platform. */
 	FTexturePlatformData *PlatformData;
+#if WITH_EDITOR
 	/* cooked platform data for this texture */
 	TMap<FString, FTexturePlatformData*> CookedPlatformData;
+#endif
 	/**
 	 * Thread-safe counter indicating the texture streaming state. The definitions below are mirrored in Texture.h.
 	 *
@@ -133,9 +135,9 @@ public:
 	// Begin UObject interface.
 	virtual void Serialize(FArchive& Ar) override;	
 #if WITH_EDITOR
-	virtual void CookerWillNeverCookAgain() override;
 	virtual void PostLinkerChange() override;
 	virtual void PostEditChangeProperty(FPropertyChangedEvent& PropertyChangedEvent) override;
+	virtual void WillNeverCacheCookedPlatformDataAgain() override;
 #endif // WITH_EDITOR
 	virtual void BeginDestroy() override;
 	virtual void PostLoad() override;
@@ -152,7 +154,9 @@ public:
 	virtual void UpdateResource() override;
 	virtual float GetAverageBrightness(bool bIgnoreTrueBlack, bool bUseGrayscale) override;
 	virtual FTexturePlatformData** GetRunningPlatformData() override { return &PlatformData; }
+#if WITH_EDITOR
 	virtual TMap<FString,FTexturePlatformData*>* GetCookedPlatformData() override { return &CookedPlatformData; }
+#endif
 	// End UTexture interface.
 
 	/** Trivial accessors. */
diff --git a/Engine/Source/Runtime/Engine/Classes/Materials/Material.h b/Engine/Source/Runtime/Engine/Classes/Materials/Material.h
--- a/Engine/Source/Runtime/Engine/Classes/Materials/Material.h
+++ b/Engine/Source/Runtime/Engine/Classes/Materials/Material.h
@@ -726,10 +726,10 @@ private:
 	 * These are always valid and non-null, but only the entries affected by CacheResourceShadersForRendering are actually valid for rendering.
 	 */
 	FMaterialResource* MaterialResources[EMaterialQualityLevel::Num][ERHIFeatureLevel::Num];
-
+#if WITH_EDITOR
 	/** Material resources being cached for cooking. */
 	TMap<const class ITargetPlatform*, TArray<FMaterialResource*>> CachedMaterialResourcesForCooking;
-
+#endif
 	/** Fence used to guarantee that the RT is finished using various resources in this UMaterial before cleanup. */
 	FRenderCommandFence ReleaseFence;
 
@@ -799,11 +799,12 @@ public:
 	ENGINE_API virtual void Serialize(FArchive& Ar) override;
 	ENGINE_API virtual void PostDuplicate(bool bDuplicateForPIE) override;
 	ENGINE_API virtual void PostLoad() override;
+#if WITH_EDITOR
 	ENGINE_API virtual void BeginCacheForCookedPlatformData( const ITargetPlatform *TargetPlatform ) override;
 	ENGINE_API virtual bool IsCachedCookedPlatformDataLoaded( const ITargetPlatform* TargetPlatform ) override;
 	ENGINE_API virtual void ClearCachedCookedPlatformData( const ITargetPlatform *TargetPlatform ) override;
 	ENGINE_API virtual void ClearAllCachedCookedPlatformData() override;
-
+#endif
 #if WITH_EDITOR
 	ENGINE_API virtual void PreEditChange(UProperty* PropertyAboutToChange) override;
 	ENGINE_API virtual void PostEditChangeProperty(FPropertyChangedEvent& PropertyChangedEvent) override;
diff --git a/Engine/Source/Runtime/Engine/Classes/Materials/MaterialInstance.h b/Engine/Source/Runtime/Engine/Classes/Materials/MaterialInstance.h
--- a/Engine/Source/Runtime/Engine/Classes/Materials/MaterialInstance.h
+++ b/Engine/Source/Runtime/Engine/Classes/Materials/MaterialInstance.h
@@ -231,10 +231,10 @@ private:
 	 * For example the material instance needs to support being rendered at different quality levels and feature levels within the same process.
 	 */
 	FMaterialResource* StaticPermutationMaterialResources[EMaterialQualityLevel::Num][ERHIFeatureLevel::Num];
-
+#if WITH_EDITOR
 	/** Material resources being cached for cooking. */
 	TMap<const class ITargetPlatform*, TArray<FMaterialResource*>> CachedMaterialResourcesForCooking;
-
+#endif
 	/** Fence used to guarantee that the RT is finished using various resources in this UMaterial before cleanup. */
 	FRenderCommandFence ReleaseFence;
 
@@ -300,10 +300,12 @@ public:
 	// Begin UObject interface.
 	virtual ENGINE_API SIZE_T GetResourceSize(EResourceSizeMode::Type Mode) override;
 	virtual ENGINE_API void PostInitProperties() override;	
+#if WITH_EDITOR
 	virtual ENGINE_API void BeginCacheForCookedPlatformData(const ITargetPlatform *TargetPlatform) override;
 	virtual ENGINE_API bool IsCachedCookedPlatformDataLoaded(const ITargetPlatform* TargetPlatform) override;
 	virtual ENGINE_API void ClearCachedCookedPlatformData(const ITargetPlatform *TargetPlatform) override;
 	virtual ENGINE_API void ClearAllCachedCookedPlatformData() override;
+#endif
 	virtual ENGINE_API void Serialize(FArchive& Ar) override;
 	virtual ENGINE_API void PostLoad() override;
 	virtual ENGINE_API void BeginDestroy() override;
diff --git a/Engine/Source/Runtime/Engine/Classes/Materials/MaterialInterface.h b/Engine/Source/Runtime/Engine/Classes/Materials/MaterialInterface.h
--- a/Engine/Source/Runtime/Engine/Classes/Materials/MaterialInterface.h
+++ b/Engine/Source/Runtime/Engine/Classes/Materials/MaterialInterface.h
@@ -622,4 +622,4 @@ private:
 };
 
 /** Helper function to serialize inline shader maps for the given material resources. */
-extern void SerializeInlineShaderMaps(const TMap<const class ITargetPlatform*,TArray<FMaterialResource*>>& PlatformMaterialResourcesToSave, FArchive& Ar, FMaterialResource* OutMaterialResourcesLoaded[][ERHIFeatureLevel::Num]);
\ No newline at end of file
+extern void SerializeInlineShaderMaps(const TMap<const class ITargetPlatform*,TArray<FMaterialResource*>>* PlatformMaterialResourcesToSave, FArchive& Ar, FMaterialResource* OutMaterialResourcesLoaded[][ERHIFeatureLevel::Num]);
\ No newline at end of file
diff --git a/Engine/Source/Runtime/Engine/Classes/Sound/SoundWave.h b/Engine/Source/Runtime/Engine/Classes/Sound/SoundWave.h
--- a/Engine/Source/Runtime/Engine/Classes/Sound/SoundWave.h
+++ b/Engine/Source/Runtime/Engine/Classes/Sound/SoundWave.h
@@ -267,7 +267,6 @@ public:
 	virtual void FinishDestroy() override;
 	virtual void PostLoad() override;
 #if WITH_EDITOR
-	virtual void CookerWillNeverCookAgain() override;
 	virtual void PostEditChangeProperty(FPropertyChangedEvent& PropertyChangedEvent) override;	
 #endif // WITH_EDITOR
 	virtual bool IsLocalizedResource() override;
@@ -372,7 +371,6 @@ public:
 	void UpdatePlatformData();
 
 	void CleanupCachedRunningPlatformData();
-	void CleanupCachedCookedPlatformData();
 
 	/**
 	 * Serializes cooked platform data.
@@ -380,11 +378,8 @@ public:
 	void SerializeCookedPlatformData(class FArchive& Ar);
 
 #if WITH_EDITORONLY_DATA
-	/**
-	 * Caches platform data for the sound.
-	 */
-	void CachePlatformData(bool bAsyncCache = false);
-
+	
+#if WITH_EDITOR
 	/**
 	 * Begins caching platform data in the background for the platform requested
 	 */
@@ -392,6 +387,22 @@ public:
 
 	virtual bool IsCachedCookedPlatformDataLoaded( const ITargetPlatform* TargetPlatform ) override;
 
+	/**
+	 * Clear all the cached cooked platform data which we have accumulated with BeginCacheForCookedPlatformData calls
+	 * The data can still be cached again using BeginCacheForCookedPlatformData again
+	 */
+	virtual void ClearAllCachedCookedPlatformData() override;
+
+	virtual void ClearCachedCookedPlatformData( const ITargetPlatform* TargetPlatform ) override;
+
+	virtual void WillNeverCacheCookedPlatformDataAgain() override;
+#endif
+	
+	/**
+	 * Caches platform data for the sound.
+	 */
+	void CachePlatformData(bool bAsyncCache = false);
+
 	/**
 	 * Begins caching platform data in the background.
 	 */
diff --git a/Engine/Source/Runtime/Engine/Private/AudioDerivedData.cpp b/Engine/Source/Runtime/Engine/Private/AudioDerivedData.cpp
--- a/Engine/Source/Runtime/Engine/Private/AudioDerivedData.cpp
+++ b/Engine/Source/Runtime/Engine/Private/AudioDerivedData.cpp
@@ -920,15 +920,6 @@ void USoundWave::CleanupCachedRunningPlatformData()
 	}
 }
 
-void USoundWave::CleanupCachedCookedPlatformData()
-{
-	for (auto It : CookedPlatformData)
-	{
-		delete It.Value;
-	}
-
-	CookedPlatformData.Empty();
-}
 
 void USoundWave::SerializeCookedPlatformData(FArchive& Ar)
 {
@@ -994,6 +985,7 @@ void USoundWave::BeginCachePlatformData()
 {
 	CachePlatformData(true);
 
+#if WITH_EDITOR
 	// enable caching in postload for derived data cache commandlet and cook by the book
 	ITargetPlatformManagerModule* TPM = GetTargetPlatformManager();
 	if (TPM && (TPM->RestrictFormatsToRuntimeOnly() == false))
@@ -1006,7 +998,9 @@ void USoundWave::BeginCachePlatformData()
 			BeginCacheForCookedPlatformData(Platforms[FormatIndex]);
 		}
 	}
+#endif
 }
+#if WITH_EDITOR
 
 void USoundWave::BeginCacheForCookedPlatformData(const ITargetPlatform *TargetPlatform)
 {
@@ -1067,6 +1061,55 @@ bool USoundWave::IsCachedCookedPlatformDataLoaded( const ITargetPlatform* Target
 	return true; 
 }
 
+
+/**
+* Clear all the cached cooked platform data which we have accumulated with BeginCacheForCookedPlatformData calls
+* The data can still be cached again using BeginCacheForCookedPlatformData again
+*/
+void USoundWave::ClearAllCachedCookedPlatformData()
+{
+	Super::ClearAllCachedCookedPlatformData();
+
+	for (auto It : CookedPlatformData)
+	{
+		delete It.Value;
+	}
+
+	CookedPlatformData.Empty();
+}
+
+void USoundWave::ClearCachedCookedPlatformData( const ITargetPlatform* TargetPlatform )
+{
+	Super::ClearCachedCookedPlatformData(TargetPlatform);
+
+	if (TargetPlatform->SupportsFeature(ETargetPlatformFeatures::AudioStreaming) && IsStreaming())
+	{
+		// Retrieve format to cache for targetplatform.
+		FName PlatformFormat = TargetPlatform->GetWaveFormat(this);
+
+		// find format data by comparing derived data keys.
+		FString DerivedDataKey;
+		GetStreamedAudioDerivedDataKeySuffix(*this, PlatformFormat, DerivedDataKey);
+
+		
+		if ( CookedPlatformData.Contains(DerivedDataKey) )
+		{
+			FStreamedAudioPlatformData *PlatformData = CookedPlatformData.FindAndRemoveChecked( DerivedDataKey );
+			delete PlatformData;
+		}	
+	}
+}
+
+void USoundWave::WillNeverCacheCookedPlatformDataAgain()
+{
+	// this is called after we have finished caching the platform data but before we have saved the data
+	// so need to keep the cached platform data around 
+	Super::WillNeverCacheCookedPlatformDataAgain();
+	RawData.RemoveBulkData();
+	CompressedFormatData.FlushData();
+}
+#endif
+
 void USoundWave::FinishCachePlatformData()
 {
 	if (RunningPlatformData == NULL)
diff --git a/Engine/Source/Runtime/Engine/Private/Materials/Material.cpp b/Engine/Source/Runtime/Engine/Private/Materials/Material.cpp
--- a/Engine/Source/Runtime/Engine/Private/Materials/Material.cpp
+++ b/Engine/Source/Runtime/Engine/Private/Materials/Material.cpp
@@ -409,7 +409,7 @@ void UMaterialInterface::AssertDefaultMaterialsPostLoaded()
 	}
 }
 
-void SerializeInlineShaderMaps(const TMap<const ITargetPlatform*,TArray<FMaterialResource*>>& PlatformMaterialResourcesToSave, FArchive& Ar, FMaterialResource* OutMaterialResourcesLoaded[][ERHIFeatureLevel::Num])
+void SerializeInlineShaderMaps(const TMap<const ITargetPlatform*,TArray<FMaterialResource*>>* PlatformMaterialResourcesToSavePtr, FArchive& Ar, FMaterialResource* OutMaterialResourcesLoaded[][ERHIFeatureLevel::Num])
 {
 	if (Ar.IsSaving())
 	{
@@ -417,6 +417,9 @@ void SerializeInlineShaderMaps(const TMap<const ITargetPlatform*,TArray<FMateria
 		const TArray<FMaterialResource*> *MaterialResourcesToSavePtr = NULL;
 		if (Ar.IsCooking())
 		{
+			check( PlatformMaterialResourcesToSavePtr );
+			auto& PlatformMaterialResourcesToSave = *PlatformMaterialResourcesToSavePtr;
+
 			MaterialResourcesToSavePtr = PlatformMaterialResourcesToSave.Find( Ar.CookingTarget() );
 			check( MaterialResourcesToSavePtr != NULL || (Ar.GetLinker()==NULL) );
 			if (MaterialResourcesToSavePtr!= NULL )
@@ -1729,7 +1732,11 @@ void UMaterial::Serialize(FArchive& Ar)
 
 	if (Ar.UE4Ver() >= VER_UE4_PURGED_FMATERIAL_COMPILE_OUTPUTS)
 	{
-		SerializeInlineShaderMaps( CachedMaterialResourcesForCooking, Ar, MaterialResources );
+#if WITH_EDITOR
+		SerializeInlineShaderMaps( &CachedMaterialResourcesForCooking, Ar, MaterialResources );
+#else
+		SerializeInlineShaderMaps( NULL, Ar, MaterialResources );
+#endif
 	}
 	else
 	{
@@ -2036,7 +2043,7 @@ void UMaterial::PostLoad()
 	STAT(double MaterialLoadTime = 0);
 	{
 		SCOPE_SECONDS_COUNTER(MaterialLoadTime);
-
+#if WITH_EDITOR
 		// enable caching in postload for derived data cache commandlet and cook by the book
 		ITargetPlatformManagerModule* TPM = GetTargetPlatformManager();
 		if (TPM && (TPM->RestrictFormatsToRuntimeOnly() == false))
@@ -2048,7 +2055,7 @@ void UMaterial::PostLoad()
 				BeginCacheForCookedPlatformData(Platforms[FormatIndex]);
 			}
 		}
-
+#endif
 		//Don't compile shaders in post load for dev overhead materials.
 		if (FApp::CanEverRender() && !bIsMaterialEditorStatsMaterial)
 		{
@@ -2111,7 +2118,7 @@ void UMaterial::PropagateDataToMaterialProxy()
 		}
 	}
 }
-
+#if WITH_EDITOR
 void UMaterial::BeginCacheForCookedPlatformData( const ITargetPlatform *TargetPlatform )
 {
 	TArray<FName> DesiredShaderFormats;
@@ -2186,7 +2193,7 @@ void UMaterial::ClearAllCachedCookedPlatformData()
 	}
 	CachedMaterialResourcesForCooking.Empty();
 }
-
+#endif
 #if WITH_EDITOR
 bool UMaterial::CanEditChange(const UProperty* InProperty) const
 {
@@ -2725,9 +2732,9 @@ void UMaterial::ReleaseResources()
 			CurrentResource = NULL;
 		}
 	}
-
+#if WITH_EDITOR
 	ClearAllCachedCookedPlatformData();
-
+#endif
 	for (int32 InstanceIndex = 0; InstanceIndex < 3; ++InstanceIndex)
 	{
 		if (DefaultMaterialInstances[InstanceIndex])
diff --git a/Engine/Source/Runtime/Engine/Private/Materials/MaterialInstance.cpp b/Engine/Source/Runtime/Engine/Private/Materials/MaterialInstance.cpp
--- a/Engine/Source/Runtime/Engine/Private/Materials/MaterialInstance.cpp
+++ b/Engine/Source/Runtime/Engine/Private/Materials/MaterialInstance.cpp
@@ -1513,6 +1513,8 @@ void TrimToOverriddenOnly(TArray<ParameterType>& Parameters)
 	}
 }
 
+#if WITH_EDITOR
+
 void UMaterialInstance::BeginCacheForCookedPlatformData( const ITargetPlatform *TargetPlatform )
 {
 	TArray<FMaterialResource*> *CachedMaterialResourcesForPlatform = CachedMaterialResourcesForCooking.Find( TargetPlatform );
@@ -1584,6 +1586,8 @@ void UMaterialInstance::ClearAllCachedCookedPlatformData()
 	CachedMaterialResourcesForCooking.Empty();
 }
 
+#endif 
+
 void UMaterialInstance::Serialize(FArchive& Ar)
 {
 	Super::Serialize(Ar);
@@ -1594,8 +1598,12 @@ void UMaterialInstance::Serialize(FArchive& Ar)
 		if (Ar.UE4Ver() >= VER_UE4_PURGED_FMATERIAL_COMPILE_OUTPUTS)
 		{
 			StaticParameters.Serialize(Ar);
+#if WITH_EDITOR
+			SerializeInlineShaderMaps( &CachedMaterialResourcesForCooking, Ar, StaticPermutationMaterialResources );
+#else
+			SerializeInlineShaderMaps( NULL, Ar, StaticPermutationMaterialResources );
+#endif
 
-			SerializeInlineShaderMaps( CachedMaterialResourcesForCooking, Ar, StaticPermutationMaterialResources );
 		}
 		else
 		{
@@ -1702,7 +1710,7 @@ void UMaterialInstance::PostLoad()
 
 		// Make sure static parameters are up to date and shaders are cached for the current platform
 		InitStaticPermutation();
-
+#if WITH_EDITOR
 		// enable caching in postload for derived data cache commandlet and cook by the book
 		ITargetPlatformManagerModule* TPM = GetTargetPlatformManager();
 		if (TPM && (TPM->RestrictFormatsToRuntimeOnly() == false)) 
@@ -1715,6 +1723,7 @@ void UMaterialInstance::PostLoad()
 				BeginCacheForCookedPlatformData(Platforms[FormatIndex]);
 			}
 		}
+#endif
 	}
 
 	INC_FLOAT_STAT_BY(STAT_ShaderCompiling_MaterialLoading,(float)MaterialLoadTime);
@@ -1779,9 +1788,9 @@ void UMaterialInstance::FinishDestroy()
 			CurrentResource = NULL;
 		}
 	}
-
+#if WITH_EDITOR
 	ClearAllCachedCookedPlatformData();
-
+#endif
 	Super::FinishDestroy();
 }
 
diff --git a/Engine/Source/Runtime/Engine/Private/SoundWave.cpp b/Engine/Source/Runtime/Engine/Private/SoundWave.cpp
--- a/Engine/Source/Runtime/Engine/Private/SoundWave.cpp
+++ b/Engine/Source/Runtime/Engine/Private/SoundWave.cpp
@@ -400,14 +400,6 @@ bool USoundWave::IsLocalizedResource()
 
 #if WITH_EDITOR
 
-void USoundWave::CookerWillNeverCookAgain()
-{
-	Super::CookerWillNeverCookAgain();
-	RawData.RemoveBulkData();
-	CompressedFormatData.FlushData();
-	CleanupCachedCookedPlatformData();
-}
-
 void USoundWave::PostEditChangeProperty(FPropertyChangedEvent& PropertyChangedEvent)
 {
 	Super::PostEditChangeProperty(PropertyChangedEvent);
@@ -527,7 +519,9 @@ void USoundWave::FinishDestroy()
 	Super::FinishDestroy();
 
 	CleanupCachedRunningPlatformData();
-	CleanupCachedCookedPlatformData();
+#if WITH_EDITOR
+	ClearAllCachedCookedPlatformData();
+#endif
 
 	IStreamingManager::Get().GetAudioStreamingManager().RemoveStreamingSoundWave(this);
 }
diff --git a/Engine/Source/Runtime/Engine/Private/Texture.cpp b/Engine/Source/Runtime/Engine/Private/Texture.cpp
--- a/Engine/Source/Runtime/Engine/Private/Texture.cpp
+++ b/Engine/Source/Runtime/Engine/Private/Texture.cpp
@@ -330,7 +330,9 @@ void UTexture::FinishDestroy()
 	}
 
 	CleanupCachedRunningPlatformData();
-	CleanupCachedCookedPlatformData();
+#if WITH_EDITOR
+	ClearAllCachedCookedPlatformData();
+#endif
 }
 
 void UTexture::PreSave()
@@ -462,6 +464,9 @@ bool UTexture::ForceUpdateTextureStreaming()
 FTextureSource::FTextureSource()
 	: LockedMipData(NULL)
 	, LockedMips(0)
+#if WITH_EDITOR
+	, bHasHadBulkDataCleared(false)
+#endif
 #if WITH_EDITORONLY_DATA
 	, SizeX(0)
 	, SizeY(0)
@@ -777,6 +782,16 @@ void FTextureSource::ForceGenerateGuid()
 	bGuidIsHash = false;
 }
 
+void FTextureSource::ReleaseSourceMemory()
+{
+	bHasHadBulkDataCleared = true;
+	if (BulkData.IsLocked())
+	{
+		BulkData.Unlock();
+	}
+	BulkData.RemoveBulkData();
+}
+
 void FTextureSource::RemoveSourceData()
 {
 	SizeX = 0;
diff --git a/Engine/Source/Runtime/Engine/Private/Texture2D.cpp b/Engine/Source/Runtime/Engine/Private/Texture2D.cpp
--- a/Engine/Source/Runtime/Engine/Private/Texture2D.cpp
+++ b/Engine/Source/Runtime/Engine/Private/Texture2D.cpp
@@ -405,12 +405,6 @@ void UTexture2D::UpdateResource()
 
 
 #if WITH_EDITOR
-void UTexture2D::CookerWillNeverCookAgain()
-{
-	Super::CookerWillNeverCookAgain();
-
-	CleanupCachedCookedPlatformData();
-}
 void UTexture2D::PostLinkerChange()
 {
 	// Changing the linker requires re-creating the resource to make sure streaming behavior is right.
diff --git a/Engine/Source/Runtime/Engine/Private/TextureDerivedData.cpp b/Engine/Source/Runtime/Engine/Private/TextureDerivedData.cpp
--- a/Engine/Source/Runtime/Engine/Private/TextureDerivedData.cpp
+++ b/Engine/Source/Runtime/Engine/Private/TextureDerivedData.cpp
@@ -488,6 +488,7 @@ namespace ETextureCacheFlags
 		InlineMips		= 0x08,
 		AllowAsyncBuild	= 0x10,
 		ForDDCBuild		= 0x20,
+		RemoveSourceMipDataAfterCache = 0x40,
 	};
 };
 
@@ -680,6 +681,13 @@ class FTextureCacheDerivedDataWorker : public FNonAbandonableTask
 	/** Gathers information needed to build a texture. */
 	void GetBuildInfo()
 	{
+		if ( Texture.Source.HasHadBulkDataCleared() )
+		{
+			// don't do any work we can't reload this
+			UE_LOG(LogTexture, Warning, TEXT("Unable to load texture source data could be because it has been cleared. %s"), *Texture.GetName())
+			return;
+		}
+
 		// Dump any existing mips.
 		DerivedData->Mips.Empty();
 		
@@ -909,6 +917,8 @@ public:
 		{
 			GetBuildInfo();
 			BuildTexture();
+
+			check( bSucceeded );
 		}
 	}
 
@@ -958,7 +968,9 @@ void FTexturePlatformData::Cache(
 	bool bForceRebuild = (Flags & ETextureCacheFlags::ForceRebuild) != 0;
 	bool bAsync = !bForDDC && (Flags & ETextureCacheFlags::Async) != 0;
 	GetTextureDerivedDataKey(InTexture, InSettings, DerivedDataKey);
-	
+
+	UE_LOG(LogTemp, Warning, TEXT("FTexturePlatformData::Cache for %s with ID %s"), *InTexture.GetName(), *DerivedDataKey );
+
 	ITextureCompressorModule* Compressor = &FModuleManager::LoadModuleChecked<ITextureCompressorModule>(TEXTURE_COMPRESSOR_MODULENAME);
 
 	if (bAsync && !bForceRebuild)
@@ -1363,25 +1375,6 @@ void UTexture2D::GetMipData(int32 FirstMipToLoad, void** OutMipData)
 	}
 }
 
-
-
-void UTexture::CleanupCachedCookedPlatformData()
-{
-	TMap<FString, FTexturePlatformData*> *CookedPlatformDataPtr = GetCookedPlatformData();
-
-	if ( CookedPlatformDataPtr )
-	{
-		TMap<FString, FTexturePlatformData*> &CookedPlatformData = *CookedPlatformDataPtr;
-
-		for ( auto It : CookedPlatformData )
-		{
-			delete It.Value;
-		}
-
-		CookedPlatformData.Empty();
-	}
-}
-
 void UTexture::UpdateCachedLODBias( bool bIncTextureMips )
 {
 	CachedCombinedLODBias = GSystemSettings.TextureLODSettings.CalculateLODBias( this, bIncTextureMips );
@@ -1430,6 +1423,7 @@ void UTexture::BeginCachePlatformData()
 {
 	CachePlatformData(true);
 
+#if 0 // don't cache in post load, this increases our peak memory usage, instead cache just before we save the package
 	// enable caching in postload for derived data cache commandlet and cook by the book
 	ITargetPlatformManagerModule* TPM = GetTargetPlatformManager();
 	if (TPM && (TPM->RestrictFormatsToRuntimeOnly() == false))
@@ -1442,6 +1436,7 @@ void UTexture::BeginCachePlatformData()
 			BeginCacheForCookedPlatformData(Platforms[FormatIndex]);
 		}
 	}
+#endif
 }
 
 void UTexture::BeginCacheForCookedPlatformData( const ITargetPlatform *TargetPlatform )
@@ -1467,10 +1462,7 @@ void UTexture::BeginCacheForCookedPlatformData( const ITargetPlatform *TargetPla
 		for (int32 FormatIndex = 0; FormatIndex < PlatformFormats.Num(); ++FormatIndex)
 		{
 			BuildSettings.TextureFormatName = PlatformFormats[FormatIndex];
-			if (BuildSettings.TextureFormatName != NAME_None)
-			{
-				BuildSettingsToCache.Add(BuildSettings);
-			}
+			BuildSettingsToCache.Add(BuildSettings);
 		}
 
 		uint32 CacheFlags = ETextureCacheFlags::Async | ETextureCacheFlags::InlineMips;
@@ -1495,6 +1487,7 @@ void UTexture::BeginCacheForCookedPlatformData( const ITargetPlatform *TargetPla
 
 			if ( PlatformData == NULL )
 			{
+				// UE_LOG(LogTemp, Warning, TEXT("Caching data for texture %s with id %s"), *GetName(), *DerivedDataKey);
 				uint32 CurrentCacheFlags = CacheFlags;
 				// if the cached data key exists already then we don't need to allowasync build
 				// if it doesn't then allow async builds
@@ -1520,6 +1513,73 @@ void UTexture::BeginCacheForCookedPlatformData( const ITargetPlatform *TargetPla
 	}
 }
 
+void UTexture::ClearCachedCookedPlatformData( const ITargetPlatform* TargetPlatform )
+{
+	// I feel like bobby fisher, always four moves ahead of
+	// my competition, listen they ain't gonna stop me ever
+	// I feel as large as Biggie, swear it could not get better, 
+	// I feel in charge like Biggie, wearing that Cosby sweater
+	TMap<FString, FTexturePlatformData*> *CookedPlatformDataPtr = GetCookedPlatformData();
+
+	if ( CookedPlatformDataPtr )
+	{
+
+		TMap<FString,FTexturePlatformData*>& CookedPlatformData = *CookedPlatformDataPtr;
+
+		// Make sure the pixel format enum has been cached.
+		UTexture::GetPixelFormatEnum();
+
+		// Retrieve formats to cache for targetplatform.
+
+		TArray<FName> PlatformFormats;
+
+
+		TArray<FTextureBuildSettings> BuildSettingsToCache;
+
+		FTextureBuildSettings BuildSettings;
+		GetTextureBuildSettings(*this, TargetPlatform->GetTextureLODSettings(), BuildSettings);
+		TargetPlatform->GetTextureFormats(this, PlatformFormats);
+		for (int32 FormatIndex = 0; FormatIndex < PlatformFormats.Num(); ++FormatIndex)
+		{
+			BuildSettings.TextureFormatName = PlatformFormats[FormatIndex];
+			// if (BuildSettings.TextureFormatName != NAME_None)
+			{
+				BuildSettingsToCache.Add(BuildSettings);
+			}
+		}
+
+		// Cull redundant settings by comparing derived data keys.
+		for (int32 SettingsIndex = 0; SettingsIndex < BuildSettingsToCache.Num(); ++SettingsIndex)
+		{
+			FString DerivedDataKey;
+			GetTextureDerivedDataKey(*this, BuildSettingsToCache[SettingsIndex], DerivedDataKey);
+
+			if ( CookedPlatformData.Contains( DerivedDataKey ) )
+			{
+				FTexturePlatformData *PlatformData = CookedPlatformData.FindAndRemoveChecked( DerivedDataKey );
+				delete PlatformData;
+			}
+		}
+	}
+}
+
+void UTexture::ClearAllCachedCookedPlatformData()
+{
+	TMap<FString, FTexturePlatformData*> *CookedPlatformDataPtr = GetCookedPlatformData();
+
+	if ( CookedPlatformDataPtr )
+	{
+		TMap<FString, FTexturePlatformData*> &CookedPlatformData = *CookedPlatformDataPtr;
+
+		for ( auto It : CookedPlatformData )
+		{
+			delete It.Value;
+		}
+
+		CookedPlatformData.Empty();
+	}
+}
+
 bool UTexture::IsCachedCookedPlatformDataLoaded( const ITargetPlatform* TargetPlatform ) 
 { 
 	const TMap<FString, FTexturePlatformData*> *CookedPlatformDataPtr = GetCookedPlatformData();
@@ -1541,16 +1601,52 @@ bool UTexture::IsCachedCookedPlatformDataLoaded( const ITargetPlatform* TargetPl
 
 		FTexturePlatformData *PlatformData= (*CookedPlatformDataPtr).FindRef(DerivedDataKey);
 
+		// begin cache hasn't been called
+		if ( !PlatformData )
+			return false;
+
 		if (PlatformData)
 		{
-			if ( (PlatformData->AsyncTask != NULL) && ( PlatformData->AsyncTask->IsWorkDone() == false ) )
+			if ( (PlatformData->AsyncTask != NULL) && ( PlatformData->AsyncTask->IsWorkDone() == true ) )
+			{
+				PlatformData->FinishCache();
+			}
+
+			if ( PlatformData->AsyncTask)
+			{
 				return false;
+			}
 		}
 	}
 	// if we get here all our stuff is cached :)
 	return true;
 }
 
+void UTexture2D::WillNeverCacheCookedPlatformDataAgain()
+{
+	Super::WillNeverCacheCookedPlatformDataAgain();
+#if WITH_EDITORONLY_DATA && 0 // enable this when I feel safe about it
+	// "Source" is only in WITH_EDITORONLY_DATA
+	UE_LOG(LogTemp, Warning, TEXT("Cleared source texture data for texture %s"), *GetName());
+	// clear source mips if we are in the editor then we don't have this luxury 
+	check( IsAsyncCacheComplete());
+
+	const TMap<FString, FTexturePlatformData*> *CookedPlatformDataPtr = GetCookedPlatformData();
+	if ( CookedPlatformDataPtr != NULL )
+	{
+		for ( const auto& TexturePlatformData : *CookedPlatformDataPtr )
+		{
+			check( TexturePlatformData.Value->AsyncTask == NULL);
+		}
+	}
+
+	if ( Source.IsBulkDataLoaded() )
+	{
+		Source.ReleaseSourceMemory();
+	}
+#endif
+}
+
 bool UTexture::IsAsyncCacheComplete()
 {
 	bool bComplete = true;
@@ -1744,8 +1840,6 @@ void UTexture::SerializeCookedPlatformData(FArchive& Ar)
 
 			for (int32 i = 0; i < PlatformDataToSerialize.Num(); ++i)
 			{
-				
-
 				FTexturePlatformData* PlatformDataToSave = PlatformDataToSerialize[i];
 				PlatformDataToSave->FinishCache();
 				FName PixelFormatName = PixelFormatEnum->GetEnum(PlatformDataToSave->PixelFormat);
