Framework: Various camera improvements
- Exposed near and far clip plane controls for orthographic cameras
- Fixed an issue where the view location was incorrectly affecting the ortho projection matrix
- Centralized projection matrix calculations for user cameras (editor viewports that aren't previewing a camera still do their own thing for now) into FMinimalViewInfo::CalculateProjectionMatrixGivenView
- Made UCameraComponent Blueprintable
[UE-2426]
[UE-10648]

diff --git a/Engine/Source/Editor/DetailCustomizations/Private/CameraDetails.cpp b/Engine/Source/Editor/DetailCustomizations/Private/CameraDetails.cpp
--- a/Engine/Source/Editor/DetailCustomizations/Private/CameraDetails.cpp
+++ b/Engine/Source/Editor/DetailCustomizations/Private/CameraDetails.cpp
@@ -37,56 +37,67 @@ void FCameraDetails::CustomizeDetails( IDetailLayoutBuilder& DetailLayout )
 	// Organize the properties
 	CameraCategory.AddProperty(ProjectionModeProperty);
 
+	// Perspective-specific properties
 	IDetailPropertyRow& FieldOfViewRow = CameraCategory.AddProperty(DetailLayout.GetProperty(GET_MEMBER_NAME_CHECKED(UCameraComponent, FieldOfView)));
-		FieldOfViewRow.Visibility( TAttribute<EVisibility>::Create(TAttribute<EVisibility>::FGetter::CreateSP( this, &FCameraDetails::ProjectionModeMatches, ProjectionModeProperty, ECameraProjectionMode::Perspective ) ) );
+	FieldOfViewRow.Visibility( TAttribute<EVisibility>::Create(TAttribute<EVisibility>::FGetter::CreateSP( this, &FCameraDetails::ProjectionModeMatches, ProjectionModeProperty, ECameraProjectionMode::Perspective ) ) );
+
+	// Orthographic-specific properties
+	TAttribute<EVisibility> OrthographicVisibility = TAttribute<EVisibility>::Create(TAttribute<EVisibility>::FGetter::CreateSP(this, &FCameraDetails::ProjectionModeMatches, ProjectionModeProperty, ECameraProjectionMode::Orthographic));
 
 	IDetailPropertyRow& OrthoWidthRow = CameraCategory.AddProperty(DetailLayout.GetProperty(GET_MEMBER_NAME_CHECKED(UCameraComponent, OrthoWidth)));
-		OrthoWidthRow.Visibility( TAttribute<EVisibility>::Create(TAttribute<EVisibility>::FGetter::CreateSP( this, &FCameraDetails::ProjectionModeMatches, ProjectionModeProperty, ECameraProjectionMode::Orthographic ) ) );
+	OrthoWidthRow.Visibility(OrthographicVisibility);
 
-		CameraCategory.AddProperty( bConstrainAspectRatioProperty );
-		IDetailPropertyRow& AspectRatioRow = CameraCategory.AddProperty(AspectRatioProperty);
+	IDetailPropertyRow& OrthoNearClipPlaneRow = CameraCategory.AddProperty(DetailLayout.GetProperty(GET_MEMBER_NAME_CHECKED(UCameraComponent, OrthoNearClipPlane)));
+	OrthoNearClipPlaneRow.Visibility(OrthographicVisibility);
 
-		// Provide the special aspect ratio row
-		AspectRatioRow.CustomWidget()
-			.NameContent()
-			[
-				AspectRatioProperty->CreatePropertyNameWidget()
-			]
+	IDetailPropertyRow& OrthoFarClipPlaneRow = CameraCategory.AddProperty(DetailLayout.GetProperty(GET_MEMBER_NAME_CHECKED(UCameraComponent, OrthoFarClipPlane)));
+	OrthoFarClipPlaneRow.Visibility(OrthographicVisibility);
+
+	// Aspect ratio
+	CameraCategory.AddProperty( bConstrainAspectRatioProperty );
+	IDetailPropertyRow& AspectRatioRow = CameraCategory.AddProperty(AspectRatioProperty);
+
+	// Provide the special aspect ratio row
+	AspectRatioRow.CustomWidget()
+		.NameContent()
+		[
+			AspectRatioProperty->CreatePropertyNameWidget()
+		]
 		.ValueContent()
+		[
+			SNew(SHorizontalBox)
+			+SHorizontalBox::Slot()
+			.Padding(0.0f, 2.0f, 5.0f, 2.0f)
 			[
-				SNew(SHorizontalBox)
-				+SHorizontalBox::Slot()
-				.Padding(0.0f, 2.0f, 5.0f, 2.0f)
+				SNew(SNumericEntryBox<float>)
+				.AllowSpin(true)
+				.Value(this, &FCameraDetails::GetAspectRatio)
+				.Font(FontStyle)
+				.MinValue(MinAspectRatio)
+				.MaxValue(MaxAspectRatio)
+				.MinSliderValue(LowestCommonAspectRatio)
+				.MaxSliderValue(HighestCommonAspectRatio)
+				.OnValueChanged(this, &FCameraDetails::OnAspectRatioSpinnerChanged)
+				.ToolTipText(LOCTEXT("AspectFloatTooltip", "Aspect Ratio (Width/Height)"))
+			]
+			+SHorizontalBox::Slot()
 				[
-					SNew(SNumericEntryBox<float>)
-					.AllowSpin(true)
-					.Value(this, &FCameraDetails::GetAspectRatio)
-					.Font(FontStyle)
-					.MinValue(MinAspectRatio)
-					.MaxValue(MaxAspectRatio)
-					.MinSliderValue(LowestCommonAspectRatio)
-					.MaxSliderValue(HighestCommonAspectRatio)
-					.OnValueChanged(this, &FCameraDetails::OnAspectRatioSpinnerChanged)
-					.ToolTipText(LOCTEXT("AspectFloatTooltip", "Aspect Ratio (Width/Height)"))
-				]
-				+SHorizontalBox::Slot()
+					SNew(SComboButton)
+					.OnGetMenuContent( this, &FCameraDetails::OnGetComboContent )
+					.ContentPadding(0.0f)
+					.ButtonStyle( FEditorStyle::Get(), "ToggleButton" )
+					.ForegroundColor(FSlateColor::UseForeground())
+					.VAlign(VAlign_Center)
+					.ButtonContent()
 					[
-						SNew(SComboButton)
-						.OnGetMenuContent( this, &FCameraDetails::OnGetComboContent )
-						.ContentPadding(0.0f)
-						.ButtonStyle( FEditorStyle::Get(), "ToggleButton" )
-						.ForegroundColor(FSlateColor::UseForeground())
-						.VAlign(VAlign_Center)
-						.ButtonContent()
-						[
-							SAssignNew(AspectTextBox, SEditableTextBox)
-							.HintText( LOCTEXT("AspectTextHint", "width x height") )
-							.ToolTipText( LOCTEXT("AspectTextTooltip", "Enter a ratio in the form \'width x height\' or \'width:height\'") )
-							.Font(FontStyle)
-							.OnTextCommitted(this, &FCameraDetails::OnCommitAspectRatioText)
-						]
+						SAssignNew(AspectTextBox, SEditableTextBox)
+						.HintText( LOCTEXT("AspectTextHint", "width x height") )
+						.ToolTipText( LOCTEXT("AspectTextTooltip", "Enter a ratio in the form \'width x height\' or \'width:height\'") )
+						.Font(FontStyle)
+						.OnTextCommitted(this, &FCameraDetails::OnCommitAspectRatioText)
 					]
-			];
+				]
+		];
 
 	CameraCategory.AddProperty(DetailLayout.GetProperty(GET_MEMBER_NAME_CHECKED(UCameraComponent, bUsePawnControlRotation)));
 	CameraCategory.AddProperty(DetailLayout.GetProperty(GET_MEMBER_NAME_CHECKED(UCameraComponent, PostProcessBlendWeight)));
diff --git a/Engine/Source/Editor/UnrealEd/Private/EditorViewportClient.cpp b/Engine/Source/Editor/UnrealEd/Private/EditorViewportClient.cpp
--- a/Engine/Source/Editor/UnrealEd/Private/EditorViewportClient.cpp
+++ b/Engine/Source/Editor/UnrealEd/Private/EditorViewportClient.cpp
@@ -549,141 +549,133 @@ FSceneView* FEditorViewportClient::CalcSceneView(FSceneViewFamily* ViewFamily)
 	TimeForForceRedraw = 0.0;
 
 	const bool bConstrainAspectRatio = bUseControllingActorViewInfo && ControllingActorViewInfo.bConstrainAspectRatio;
+	const EAspectRatioAxisConstraint AspectRatioAxisConstraint = GetDefault<ULevelEditorViewportSettings>()->AspectRatioAxisConstraint;
 
 	ViewInitOptions.ViewOrigin = ViewLocation;
-	if (EffectiveViewportType == LVT_Perspective)
+
+	if (bUseControllingActorViewInfo)
 	{
-		if (bUsingOrbitCamera)
-		{
-			ViewInitOptions.ViewRotationMatrix = ViewTransform.ComputeOrbitMatrix();
-			ViewInitOptions.ViewOrigin = FVector::ZeroVector;
-		}
-		else
+		ViewInitOptions.ViewRotationMatrix = FInverseRotationMatrix(ViewRotation) * FMatrix(
+			FPlane(0, 0, 1, 0),
+			FPlane(1, 0, 0, 0),
+			FPlane(0, 1, 0, 0),
+			FPlane(0, 0, 0, 1));
+
+		FMinimalViewInfo::CalculateProjectionMatrixGivenView(ControllingActorViewInfo, AspectRatioAxisConstraint, Viewport, /*inout*/ ViewInitOptions);
+	}
+	else
+	{
+		//
+		if (EffectiveViewportType == LVT_Perspective)
 		{
-			ViewInitOptions.ViewRotationMatrix = FInverseRotationMatrix(ViewRotation);
-		}
+			if (bUsingOrbitCamera)
+			{
+				ViewInitOptions.ViewRotationMatrix = ViewTransform.ComputeOrbitMatrix();
+				ViewInitOptions.ViewOrigin = FVector::ZeroVector;
+			}
+			else
+			{
+				ViewInitOptions.ViewRotationMatrix = FInverseRotationMatrix(ViewRotation);
+			}
 
-		ViewInitOptions.ViewRotationMatrix = ViewInitOptions.ViewRotationMatrix * FMatrix(
-			FPlane(0,	0,	1,	0),
-			FPlane(1,	0,	0,	0),
-			FPlane(0,	1,	0,	0),
-			FPlane(0,	0,	0,	1));
+			ViewInitOptions.ViewRotationMatrix = ViewInitOptions.ViewRotationMatrix * FMatrix(
+				FPlane(0, 0, 1, 0),
+				FPlane(1, 0, 0, 0),
+				FPlane(0, 1, 0, 0),
+				FPlane(0, 0, 0, 1));
 
-		float MinZ = GetNearClipPlane();
-		float MaxZ = MinZ;
-		// Avoid zero ViewFOV's which cause divide by zero's in projection matrix
-		float MatrixFOV = FMath::Max(0.001f, ViewFOV) * (float)PI / 360.0f;
+			float MinZ = GetNearClipPlane();
+			float MaxZ = MinZ;
+			// Avoid zero ViewFOV's which cause divide by zero's in projection matrix
+			float MatrixFOV = FMath::Max(0.001f, ViewFOV) * (float)PI / 360.0f;
 
-		if( bConstrainAspectRatio )
-		{
-			ViewInitOptions.ProjectionMatrix = FReversedZPerspectiveMatrix(
-				MatrixFOV,
-				MatrixFOV,
-				1.0f,
-				AspectRatio,
-				MinZ,
-				MaxZ
-				);
+			if (bConstrainAspectRatio)
+			{
+				ViewInitOptions.ProjectionMatrix = FReversedZPerspectiveMatrix(
+					MatrixFOV,
+					MatrixFOV,
+					1.0f,
+					AspectRatio,
+					MinZ,
+					MaxZ
+					);
+			}
+			else
+			{
+				float XAxisMultiplier;
+				float YAxisMultiplier;
 
-			ViewInitOptions.SetConstrainedViewRectangle(Viewport->CalculateViewExtents(AspectRatio, ViewRect));
+				if (((ViewportSizeXY.X > ViewportSizeXY.Y) && (AspectRatioAxisConstraint == AspectRatio_MajorAxisFOV)) || (AspectRatioAxisConstraint == AspectRatio_MaintainXFOV))
+				{
+					//if the viewport is wider than it is tall
+					XAxisMultiplier = 1.0f;
+					YAxisMultiplier = ViewportSizeXY.X / (float)ViewportSizeXY.Y;
+				}
+				else
+				{
+					//if the viewport is taller than it is wide
+					XAxisMultiplier = ViewportSizeXY.Y / (float)ViewportSizeXY.X;
+					YAxisMultiplier = 1.0f;
+				}
+
+				ViewInitOptions.ProjectionMatrix = FReversedZPerspectiveMatrix(
+					MatrixFOV,
+					MatrixFOV,
+					XAxisMultiplier,
+					YAxisMultiplier,
+					MinZ,
+					MaxZ
+					);
+			}
 		}
 		else
 		{
-			float XAxisMultiplier;
-			float YAxisMultiplier;
-			const EAspectRatioAxisConstraint AspectRatioAxisConstraint = GetDefault<ULevelEditorViewportSettings>()->AspectRatioAxisConstraint;
+			float ZScale = 0.5f / HALF_WORLD_MAX;
+			float ZOffset = HALF_WORLD_MAX;
+
+			//The divisor for the matrix needs to match the translation code.
+			const float Zoom = GetOrthoUnitsPerPixel(Viewport);
+
+			float OrthoWidth = Zoom * ViewportSizeXY.X / 2.0f;
+			float OrthoHeight = Zoom * ViewportSizeXY.Y / 2.0f;
 
-			if (((ViewportSizeXY.X > ViewportSizeXY.Y) && (AspectRatioAxisConstraint == AspectRatio_MajorAxisFOV)) || (AspectRatioAxisConstraint == AspectRatio_MaintainXFOV))
+			if (EffectiveViewportType == LVT_OrthoXY)
+			{
+				ViewInitOptions.ViewRotationMatrix = FMatrix(
+					FPlane(1, 0, 0, 0),
+					FPlane(0, -1, 0, 0),
+					FPlane(0, 0, -1, 0),
+					FPlane(0, 0, -ViewLocation.Z, 1));
+			}
+			else if (EffectiveViewportType == LVT_OrthoXZ)
+			{
+				ViewInitOptions.ViewRotationMatrix = FMatrix(
+					FPlane(1, 0, 0, 0),
+					FPlane(0, 0, -1, 0),
+					FPlane(0, 1, 0, 0),
+					FPlane(0, 0, -ViewLocation.Y, 1));
+			}
+			else if (EffectiveViewportType == LVT_OrthoYZ)
 			{
-				//if the viewport is wider than it is tall
-				XAxisMultiplier = 1.0f;
-				YAxisMultiplier = ViewportSizeXY.X / (float)ViewportSizeXY.Y;
+				ViewInitOptions.ViewRotationMatrix = FMatrix(
+					FPlane(0, 0, 1, 0),
+					FPlane(1, 0, 0, 0),
+					FPlane(0, 1, 0, 0),
+					FPlane(0, 0, ViewLocation.X, 1));
 			}
 			else
 			{
-				//if the viewport is taller than it is wide
-				XAxisMultiplier = ViewportSizeXY.Y / (float)ViewportSizeXY.X;
-				YAxisMultiplier = 1.0f;
+				// Unknown viewport type
+				check(false);
 			}
 
-			ViewInitOptions.ProjectionMatrix = FReversedZPerspectiveMatrix (
-				MatrixFOV,
-				MatrixFOV,
-				XAxisMultiplier,
-				YAxisMultiplier,
-				MinZ,
-				MaxZ
+			ViewInitOptions.ProjectionMatrix = FReversedZOrthoMatrix(
+				OrthoWidth,
+				OrthoHeight,
+				ZScale,
+				ZOffset
 				);
 		}
-	}
-	else
-	{
-		float ZScale = 0.5f / HALF_WORLD_MAX;
-		float ZOffset = HALF_WORLD_MAX;
-
-		//The divisor for the matrix needs to match the translation code.
-		const float Zoom = GetOrthoUnitsPerPixel(Viewport);
-
-		float OrthoWidth = Zoom * ViewportSizeXY.X / 2.0f;
-		float OrthoHeight = Zoom * ViewportSizeXY.Y / 2.0f;
-
-		if (EffectiveViewportType == LVT_OrthoXY)
-		{
-			ViewInitOptions.ViewRotationMatrix = FMatrix(
-				FPlane(1,	0,	0,					0),
-				FPlane(0,	-1,	0,					0),
-				FPlane(0,	0,	-1,					0),
-				FPlane(0,	0,	-ViewLocation.Z,	1));
-		}
-		else if (EffectiveViewportType == LVT_OrthoXZ)
-		{
-			ViewInitOptions.ViewRotationMatrix = FMatrix(
-				FPlane(1,	0,	0,					0),
-				FPlane(0,	0,	-1,					0),
-				FPlane(0,	1,	0,					0),
-				FPlane(0,	0,	-ViewLocation.Y,	1));
-		}
-		else if (EffectiveViewportType == LVT_OrthoYZ)
-		{
-			ViewInitOptions.ViewRotationMatrix = FMatrix(
-				FPlane(0,	0,	1,					0),
-				FPlane(1,	0,	0,					0),
-				FPlane(0,	1,	0,					0),
-				FPlane(0,	0,	ViewLocation.X,		1));
-		}
-		else if (EffectiveViewportType == LVT_OrthoFreelook)
-		{
-			ViewInitOptions.ViewRotationMatrix = FInverseRotationMatrix(ViewRotation);
-			ViewInitOptions.ViewRotationMatrix = ViewInitOptions.ViewRotationMatrix * FMatrix(
-				FPlane(0,	0,	 1,	0),
-				FPlane(1,	0,	 0,	0),
-				FPlane(0,	1,	 0,	0),
-				FPlane(0,	0,	 0,	1));
-
-			const float EffectiveAspectRatio = bConstrainAspectRatio ? AspectRatio : (ViewportSizeXY.X / (float)ViewportSizeXY.Y);
-			const float YScale = 1.0f / EffectiveAspectRatio;
-			OrthoWidth = ControllingActorViewInfo.OrthoWidth / 2.0f;
-			OrthoHeight = (ControllingActorViewInfo.OrthoWidth / 2.0f) * YScale;
-
-			const float NearViewPlane = -ViewInitOptions.ViewOrigin.X;
-			const float FarViewPlane = NearViewPlane - 2.0f*WORLD_MAX;
-
-			const float InverseRange = 1.0f / (FarViewPlane - NearViewPlane);
-			ZScale = -2.0f * InverseRange;
-			ZOffset = -(FarViewPlane + NearViewPlane) * InverseRange;
-		}
-		else
-		{
-			// Unknown viewport type
-			check(false);
-		}
-
-		ViewInitOptions.ProjectionMatrix = FReversedZOrthoMatrix(
-			OrthoWidth,
-			OrthoHeight,
-			ZScale,
-			ZOffset
-			);
 
 		if (bConstrainAspectRatio)
 		{
diff --git a/Engine/Source/Runtime/Engine/Classes/Camera/CameraComponent.h b/Engine/Source/Runtime/Engine/Classes/Camera/CameraComponent.h
--- a/Engine/Source/Runtime/Engine/Classes/Camera/CameraComponent.h
+++ b/Engine/Source/Runtime/Engine/Classes/Camera/CameraComponent.h
@@ -9,7 +9,7 @@
   * Represents a camera viewpoint and settings, such as projection type, field of view, and post-process overrides.
   * The default behavior for an actor used as the camera view target is to look for an attached camera component and use its location, rotation, and settings.
   */
-UCLASS(HideCategories=(Mobility, Rendering, LOD), ClassGroup=Camera, meta=(BlueprintSpawnableComponent), MinimalAPI)
+UCLASS(HideCategories=(Mobility, Rendering, LOD), Blueprintable, ClassGroup=Camera, meta=(BlueprintSpawnableComponent), MinimalAPI)
 class UCameraComponent : public USceneComponent
 {
 	GENERATED_UCLASS_BODY()
@@ -22,6 +22,14 @@ class UCameraComponent : public USceneComponent
 	UPROPERTY(Interp, EditAnywhere, BlueprintReadWrite, Category=CameraSettings)
 	float OrthoWidth;
 
+	/** The near plane distance of the orthographic view (in world units) */
+	UPROPERTY(Interp, EditAnywhere, BlueprintReadWrite, Category=CameraSettings)
+	float OrthoNearClipPlane;
+
+	/** The far plane distance of the orthographic view (in world units) */
+	UPROPERTY(Interp, EditAnywhere, BlueprintReadWrite, Category=CameraSettings)
+	float OrthoFarClipPlane;
+
 	// Aspect Ratio (Width/Height)
 	UPROPERTY(Interp, EditAnywhere, BlueprintReadWrite, Category=CameraSettings, meta=(ClampMin = "0.1", ClampMax = "100.0", EditCondition="bConstrainAspectRatio"))
 	float AspectRatio;
diff --git a/Engine/Source/Runtime/Engine/Classes/Camera/CameraTypes.h b/Engine/Source/Runtime/Engine/Classes/Camera/CameraTypes.h
--- a/Engine/Source/Runtime/Engine/Classes/Camera/CameraTypes.h
+++ b/Engine/Source/Runtime/Engine/Classes/Camera/CameraTypes.h
@@ -51,7 +51,15 @@ struct FMinimalViewInfo
 	UPROPERTY(EditAnywhere, BlueprintReadWrite, Category=Camera)
 	float OrthoWidth;
 
-	// Aspect Ratio (Width/Height); ignored unless bConstrianAspectRatio is true
+	/** The near plane distance of the orthographic view (in world units) */
+	UPROPERTY(Interp, EditAnywhere, BlueprintReadWrite, Category=Camera)
+	float OrthoNearClipPlane;
+
+	/** The far plane distance of the orthographic view (in world units) */
+	UPROPERTY(Interp, EditAnywhere, BlueprintReadWrite, Category=Camera)
+	float OrthoFarClipPlane;
+
+	// Aspect Ratio (Width/Height); ignored unless bConstrainAspectRatio is true
 	UPROPERTY(EditAnywhere, BlueprintReadWrite, Category=Camera)
 	float AspectRatio;
 
@@ -76,6 +84,8 @@ struct FMinimalViewInfo
 		, Rotation(ForceInit)
 		, FOV(90.0f)
 		, OrthoWidth(512.0f)
+		, OrthoNearClipPlane(0.0f)
+		, OrthoFarClipPlane(WORLD_MAX)
 		, AspectRatio(1.33333333f)
 		, bConstrainAspectRatio(false)
 		, ProjectionMode(ECameraProjectionMode::Perspective)
@@ -95,4 +105,7 @@ struct FMinimalViewInfo
 
 	/** Combines this view with another one which will be weighted. Equals to this += OtherView * Weight. */
 	ENGINE_API void AddWeightedViewInfo(const FMinimalViewInfo& OtherView, const float& Weight);
+
+	/** Calculates the projection matrix (and potentially a constrained view rectangle) given a FMinimalViewInfo and partially configured projection data (must have the view rect already set) */
+	ENGINE_API static void CalculateProjectionMatrixGivenView(const FMinimalViewInfo& ViewInfo, TEnumAsByte<enum EAspectRatioAxisConstraint> AspectRatioAxisConstraint, class FViewport* Viewport, struct FSceneViewProjectionData& InOutProjectionData);
 };
diff --git a/Engine/Source/Runtime/Engine/Private/Camera/CameraComponent.cpp b/Engine/Source/Runtime/Engine/Private/Camera/CameraComponent.cpp
--- a/Engine/Source/Runtime/Engine/Private/Camera/CameraComponent.cpp
+++ b/Engine/Source/Runtime/Engine/Private/Camera/CameraComponent.cpp
@@ -29,6 +29,8 @@ UCameraComponent::UCameraComponent(const FObjectInitializer& ObjectInitializer)
 	FieldOfView = 90.0f;
 	AspectRatio = 1.777778f;
 	OrthoWidth = 512.0f;
+	OrthoNearClipPlane = 0.0f;
+	OrthoFarClipPlane = WORLD_MAX;
 	bConstrainAspectRatio = false;
 	PostProcessBlendWeight = 1.0f;
 	bUseControllerViewRotation_DEPRECATED = true; // the previous default value before bUsePawnControlRotation replaced this var.
@@ -200,6 +202,8 @@ void UCameraComponent::GetCameraView(float DeltaTime, FMinimalViewInfo& DesiredV
 	DesiredView.bConstrainAspectRatio = bConstrainAspectRatio;
 	DesiredView.ProjectionMode = ProjectionMode;
 	DesiredView.OrthoWidth = OrthoWidth;
+	DesiredView.OrthoNearClipPlane = OrthoNearClipPlane;
+	DesiredView.OrthoFarClipPlane = OrthoFarClipPlane;
 
 	// See if the CameraActor wants to override the PostProcess settings used.
 	DesiredView.PostProcessBlendWeight = PostProcessBlendWeight;
diff --git a/Engine/Source/Runtime/Engine/Private/Camera/CameraStackTypes.cpp b/Engine/Source/Runtime/Engine/Private/Camera/CameraStackTypes.cpp
--- a/Engine/Source/Runtime/Engine/Private/Camera/CameraStackTypes.cpp
+++ b/Engine/Source/Runtime/Engine/Private/Camera/CameraStackTypes.cpp
@@ -2,6 +2,8 @@
 
 #include "EnginePrivate.h"
 #include "Camera/CameraStackTypes.h"
+#include "UnrealClient.h"
+#include "SceneView.h"
 
 //////////////////////////////////////////////////////////////////////////
 // FMinimalViewInfo
@@ -13,6 +15,8 @@ bool FMinimalViewInfo::Equals(const FMinimalViewInfo& OtherInfo) const
 		(Rotation == OtherInfo.Rotation) &&
 		(FOV == OtherInfo.FOV) &&
 		(OrthoWidth == OtherInfo.OrthoWidth) &&
+		(OrthoNearClipPlane == OtherInfo.OrthoNearClipPlane) &&
+		(OrthoFarClipPlane == OtherInfo.OrthoFarClipPlane) &&
 		(AspectRatio == OtherInfo.AspectRatio) &&
 		(bConstrainAspectRatio == OtherInfo.bConstrainAspectRatio) &&
 		(ProjectionMode == OtherInfo.ProjectionMode);
@@ -27,6 +31,8 @@ void FMinimalViewInfo::BlendViewInfo(FMinimalViewInfo& OtherInfo, float OtherWei
 
 	FOV = FMath::Lerp(FOV, OtherInfo.FOV, OtherWeight);
 	OrthoWidth = FMath::Lerp(OrthoWidth, OtherInfo.OrthoWidth, OtherWeight);
+	OrthoNearClipPlane = FMath::Lerp(OrthoNearClipPlane, OtherInfo.OrthoNearClipPlane, OtherWeight);
+	OrthoFarClipPlane = FMath::Lerp(OrthoFarClipPlane, OtherInfo.OrthoFarClipPlane, OtherWeight);
 
 	AspectRatio = FMath::Lerp(AspectRatio, OtherInfo.AspectRatio, OtherWeight);
 	bConstrainAspectRatio |= OtherInfo.bConstrainAspectRatio;
@@ -39,6 +45,8 @@ void FMinimalViewInfo::ApplyBlendWeight(const float& Weight)
 	Rotation *= Weight;
 	FOV *= Weight;
 	OrthoWidth *= Weight;
+	OrthoNearClipPlane *= Weight;
+	OrthoFarClipPlane *= Weight;
 	AspectRatio *= Weight;
 }
 
@@ -51,7 +59,106 @@ void FMinimalViewInfo::AddWeightedViewInfo(const FMinimalViewInfo& OtherView, co
 	Rotation += OtherViewWeighted.Rotation;
 	FOV += OtherViewWeighted.FOV;
 	OrthoWidth += OtherViewWeighted.OrthoWidth;
+	OrthoNearClipPlane += OtherViewWeighted.OrthoNearClipPlane;
+	OrthoFarClipPlane += OtherViewWeighted.OrthoFarClipPlane;
 	AspectRatio += OtherViewWeighted.AspectRatio;
 
 	bConstrainAspectRatio |= OtherViewWeighted.bConstrainAspectRatio;
-}
\ No newline at end of file
+}
+
+void FMinimalViewInfo::CalculateProjectionMatrixGivenView(const FMinimalViewInfo& ViewInfo, TEnumAsByte<enum EAspectRatioAxisConstraint> AspectRatioAxisConstraint, FViewport* Viewport, FSceneViewProjectionData& InOutProjectionData)
+{
+	// Create the projection matrix (and possibly constrain the view rectangle)
+	if (ViewInfo.bConstrainAspectRatio)
+	{			
+		// Enforce a particular aspect ratio for the render of the scene. 
+		// Results in black bars at top/bottom etc.
+		InOutProjectionData.SetConstrainedViewRectangle(Viewport->CalculateViewExtents(ViewInfo.AspectRatio, InOutProjectionData.GetViewRect()));
+
+		//
+		if (ViewInfo.ProjectionMode == ECameraProjectionMode::Orthographic)
+		{
+			const float YScale = 1.0f / ViewInfo.AspectRatio;
+	
+			const float OrthoWidth = ViewInfo.OrthoWidth / 2.0f;
+			const float OrthoHeight = ViewInfo.OrthoWidth / 2.0f * YScale;
+
+			const float NearPlane = ViewInfo.OrthoNearClipPlane;
+			const float FarPlane = ViewInfo.OrthoFarClipPlane;
+
+			const float ZScale = 1.0f / (FarPlane - NearPlane);
+			const float ZOffset = -NearPlane;
+
+			InOutProjectionData.ProjectionMatrix = FReversedZOrthoMatrix(
+				OrthoWidth,
+				OrthoHeight,
+				ZScale,
+				ZOffset
+				);
+		}
+		else
+		{
+			// Avoid divide by zero in the projection matrix calculation by clamping FOV
+			InOutProjectionData.ProjectionMatrix = FReversedZPerspectiveMatrix(
+				FMath::Max(0.001f, ViewInfo.FOV) * (float)PI / 360.0f,
+				ViewInfo.AspectRatio,
+				1.0f,
+				GNearClippingPlane );
+		}
+	}
+	else 
+	{
+		// Avoid divide by zero in the projection matrix calculation by clamping FOV
+		float MatrixFOV = FMath::Max(0.001f, ViewInfo.FOV) * (float)PI / 360.0f;
+		float XAxisMultiplier;
+		float YAxisMultiplier;
+
+		const FIntRect& ViewRect = InOutProjectionData.GetViewRect();
+		const int32 SizeX = ViewRect.Width();
+		const int32 SizeY = ViewRect.Height();
+
+		// if x is bigger, and we're respecting x or major axis, AND mobile isn't forcing us to be Y axis aligned
+		if (((SizeX > SizeY) && (AspectRatioAxisConstraint == AspectRatio_MajorAxisFOV)) || (AspectRatioAxisConstraint == AspectRatio_MaintainXFOV) || (ViewInfo.ProjectionMode == ECameraProjectionMode::Orthographic))
+		{
+			//if the viewport is wider than it is tall
+			XAxisMultiplier = 1.0f;
+			YAxisMultiplier = SizeX / (float)SizeY;
+		}
+		else
+		{
+			//if the viewport is taller than it is wide
+			XAxisMultiplier = SizeY / (float)SizeX;
+			YAxisMultiplier = 1.0f;
+		}
+	
+		if (ViewInfo.ProjectionMode == ECameraProjectionMode::Orthographic)
+		{
+			const float OrthoWidth = ViewInfo.OrthoWidth / 2.0f * XAxisMultiplier;
+			const float OrthoHeight = (ViewInfo.OrthoWidth / 2.0f) / YAxisMultiplier;
+
+			const float NearPlane = ViewInfo.OrthoNearClipPlane;
+			const float FarPlane = ViewInfo.OrthoFarClipPlane;
+
+			const float ZScale = 1.0f / (FarPlane - NearPlane);
+			const float ZOffset = -NearPlane;
+
+			InOutProjectionData.ProjectionMatrix = FReversedZOrthoMatrix(
+				OrthoWidth, 
+				OrthoHeight,
+				ZScale,
+				ZOffset
+				);		
+		}
+		else
+		{
+			InOutProjectionData.ProjectionMatrix = FReversedZPerspectiveMatrix(
+				MatrixFOV,
+				MatrixFOV,
+				XAxisMultiplier,
+				YAxisMultiplier,
+				GNearClippingPlane,
+				GNearClippingPlane
+				);
+		}
+	}
+}
diff --git a/Engine/Source/Runtime/Engine/Private/LocalPlayer.cpp b/Engine/Source/Runtime/Engine/Private/LocalPlayer.cpp
--- a/Engine/Source/Runtime/Engine/Private/LocalPlayer.cpp
+++ b/Engine/Source/Runtime/Engine/Private/LocalPlayer.cpp
@@ -950,7 +950,7 @@ bool ULocalPlayer::GetProjectionData(FViewport* Viewport, EStereoscopicPass Ster
 
 	// If stereo rendering is enabled, update the size and offset appropriately for this pass
 	const bool bNeedStereo = (StereoPass != eSSP_FULL) && GEngine->IsStereoscopic3D();
-	if( bNeedStereo )
+	if (bNeedStereo)
 	{
 		GEngine->StereoRenderingDevice->AdjustViewRect(StereoPass, X, Y, SizeX, SizeY);
 	}
@@ -972,96 +972,10 @@ bool ULocalPlayer::GetProjectionData(FViewport* Viewport, EStereoscopicPass Ster
 		FPlane(0,	1,	0,	0),
 		FPlane(0,	0,	0,	1));
 
-	if( !bNeedStereo )
+	if (!bNeedStereo)
 	{
 		// Create the projection matrix (and possibly constrain the view rectangle)
-		if (ViewInfo.bConstrainAspectRatio)
-		{			
-			// Enforce a particular aspect ratio for the render of the scene. 
-			// Results in black bars at top/bottom etc.
-			ProjectionData.SetConstrainedViewRectangle(ViewportClient->Viewport->CalculateViewExtents(ViewInfo.AspectRatio, UnconstrainedRectangle));
-	
-			if (ViewInfo.ProjectionMode == ECameraProjectionMode::Orthographic)
-			{
-				const float YScale = 1.0f / ViewInfo.AspectRatio;
-	
-				const float OrthoWidth = ViewInfo.OrthoWidth / 2.0f;
-				const float OrthoHeight = ViewInfo.OrthoWidth / 2.0f * YScale;
-				const float NearPlane = -ProjectionData.ViewOrigin.X;
-				const float FarPlane = NearPlane - 2.0f*WORLD_MAX;
-				const float InverseRange = 1.0f / (FarPlane - NearPlane);
-				const float ZScale = -2.0f * InverseRange;
-				const float ZOffset = -(FarPlane + NearPlane) * InverseRange;
-
-				ProjectionData.ProjectionMatrix = FReversedZOrthoMatrix(
-					OrthoWidth,
-					OrthoHeight,
-					ZScale,
-					ZOffset
-					);
-			}
-			else
-			{
-				// Avoid divide by zero in the projection matrix calculation by clamping FOV
-				ProjectionData.ProjectionMatrix = FReversedZPerspectiveMatrix( 
-					FMath::Max(0.001f, ViewInfo.FOV) * (float)PI / 360.0f,
-					ViewInfo.AspectRatio,
-					1.0f,
-					GNearClippingPlane );
-			}
-		}
-		else 
-		{
-			// Avoid divide by zero in the projection matrix calculation by clamping FOV
-			float MatrixFOV = FMath::Max(0.001f, ViewInfo.FOV) * (float)PI / 360.0f;
-			float XAxisMultiplier;
-			float YAxisMultiplier;
-
-			// if x is bigger, and we're respecting x or major axis, AND mobile isn't forcing us to be Y axis aligned
-			if (((SizeX > SizeY) && (AspectRatioAxisConstraint == AspectRatio_MajorAxisFOV)) || (AspectRatioAxisConstraint == AspectRatio_MaintainXFOV) || (ViewInfo.ProjectionMode == ECameraProjectionMode::Orthographic))
-			{
-				//if the viewport is wider than it is tall
-				XAxisMultiplier = 1.0f;
-				YAxisMultiplier = SizeX / (float)SizeY;
-			}
-			else
-			{
-				//if the viewport is taller than it is wide
-				XAxisMultiplier = SizeY / (float)SizeX;
-				YAxisMultiplier = 1.0f;
-			}
-	
-			if (ViewInfo.ProjectionMode == ECameraProjectionMode::Orthographic)
-			{
-				const float NearPlane = -ProjectionData.ViewOrigin.X;
-				const float FarPlane = NearPlane - 2.0f*WORLD_MAX;
-	
-				const float OrthoWidth = ViewInfo.OrthoWidth / 2.0f * XAxisMultiplier;
-				const float OrthoHeight = (ViewInfo.OrthoWidth / 2.0f) / YAxisMultiplier;
-	
-				const float InverseRange = 1.0f / (FarPlane - NearPlane);
-				const float ZScale = -2.0f * InverseRange;
-				const float ZOffset = -(FarPlane + NearPlane) * InverseRange;
-
-				ProjectionData.ProjectionMatrix = FReversedZOrthoMatrix(
-					OrthoWidth, 
-					OrthoHeight,
-					ZScale,
-					ZOffset
-					);		
-			}
-			else
-			{
-				ProjectionData.ProjectionMatrix = FReversedZPerspectiveMatrix(
-					MatrixFOV,
-					MatrixFOV,
-					XAxisMultiplier,
-					YAxisMultiplier,
-					GNearClippingPlane,
-					GNearClippingPlane
-					);
-			}
-		}
+		FMinimalViewInfo::CalculateProjectionMatrixGivenView(ViewInfo, AspectRatioAxisConstraint, ViewportClient->Viewport, /*inout*/ ProjectionData);
 	}
 	else
 	{
diff --git a/Engine/Source/Runtime/UMG/Private/Components/Viewport.cpp b/Engine/Source/Runtime/UMG/Private/Components/Viewport.cpp
--- a/Engine/Source/Runtime/UMG/Private/Components/Viewport.cpp
+++ b/Engine/Source/Runtime/UMG/Private/Components/Viewport.cpp
@@ -85,9 +85,6 @@ FMatrix FUMGViewportCameraTransform::ComputeOrbitMatrix() const
 FUMGViewportClient::FUMGViewportClient(FPreviewScene* InPreviewScene)
 	: PreviewScene(InPreviewScene)
 	, Viewport(nullptr)
-	, ViewFOV(90.0f)
-	, FOVAngle(90.0f)
-	, AspectRatio(1.777777f)
 	, EngineShowFlags(ESFIM_Game)
 {
 	ViewState.Allocate();
@@ -236,26 +233,11 @@ UWorld* FUMGViewportClient::GetWorld() const
 	return OutWorldPtr;
 }
 
-bool FUMGViewportClient::IsOrtho() const
-{
-	return ViewInfo.ProjectionMode == ECameraProjectionMode::Orthographic;
-}
-
-bool FUMGViewportClient::IsPerspective() const
-{
-	return ViewInfo.ProjectionMode == ECameraProjectionMode::Perspective;
-}
-
 bool FUMGViewportClient::IsAspectRatioConstrained() const
 {
 	return ViewInfo.bConstrainAspectRatio;
 }
 
-ECameraProjectionMode::Type FUMGViewportClient::GetProjectionType() const
-{
-	return ViewInfo.ProjectionMode;
-}
-
 void FUMGViewportClient::SetBackgroundColor(FLinearColor InBackgroundColor)
 {
 	BackgroundColor = InBackgroundColor;
@@ -266,26 +248,6 @@ FLinearColor FUMGViewportClient::GetBackgroundColor() const
 	return BackgroundColor;
 }
 
-float FUMGViewportClient::GetNearClipPlane() const
-{
-	return ( NearPlane < 0.0f ) ? GNearClippingPlane : NearPlane;
-}
-
-void FUMGViewportClient::OverrideNearClipPlane(float InNearPlane)
-{
-	NearPlane = InNearPlane;
-}
-
-float FUMGViewportClient::GetFarClipPlaneOverride() const
-{
-	return FarPlane;
-}
-
-void FUMGViewportClient::OverrideFarClipPlane(const float InFarPlane)
-{
-	FarPlane = InFarPlane;
-}
-
 float FUMGViewportClient::GetOrthoUnitsPerPixel(const FViewport* InViewport) const
 {
 	const float SizeX = InViewport->GetSizeXY().X;
@@ -307,111 +269,19 @@ FSceneView* FUMGViewportClient::CalcSceneView(FSceneViewFamily* ViewFamily)
 	FIntRect ViewRect = FIntRect(0, 0, ViewportSizeXY.X, ViewportSizeXY.Y);
 	ViewInitOptions.SetViewRectangle(ViewRect);
 
-	const ECameraProjectionMode::Type ProjectionType = GetProjectionType();
-	const bool bConstrainAspectRatio = ViewInfo.bConstrainAspectRatio;
-
 	ViewInitOptions.ViewOrigin = ViewLocation;
-	if ( ProjectionType == ECameraProjectionMode::Perspective )
-	{
-		ViewInitOptions.ViewRotationMatrix = FInverseRotationMatrix(ViewRotation);
 
-		ViewInitOptions.ViewRotationMatrix = ViewInitOptions.ViewRotationMatrix * FMatrix(
-			FPlane(0, 0, 1, 0),
-			FPlane(1, 0, 0, 0),
-			FPlane(0, 1, 0, 0),
-			FPlane(0, 0, 0, 1));
+	ViewInitOptions.ViewRotationMatrix = FInverseRotationMatrix(ViewRotation);
+	ViewInitOptions.ViewRotationMatrix = ViewInitOptions.ViewRotationMatrix * FMatrix(
+		FPlane(0, 0, 1, 0),
+		FPlane(1, 0, 0, 0),
+		FPlane(0, 1, 0, 0),
+		FPlane(0, 0, 0, 1));
 
-		float MinZ = GetNearClipPlane();
-		float MaxZ = MinZ;
-		float MatrixFOV = ViewFOV * (float)PI / 360.0f;
+	//@TODO: Should probably be locally configurable (or just made into a FMinimalViewInfo property)
+	const EAspectRatioAxisConstraint AspectRatioAxisConstraint = GetDefault<ULocalPlayer>()->AspectRatioAxisConstraint;
 
-		if ( bConstrainAspectRatio )
-		{
-			ViewInitOptions.ProjectionMatrix = FReversedZPerspectiveMatrix(
-				MatrixFOV,
-				MatrixFOV,
-				1.0f,
-				AspectRatio,
-				MinZ,
-				MaxZ
-				);
-
-			ViewInitOptions.SetConstrainedViewRectangle(Viewport->CalculateViewExtents(AspectRatio, ViewRect));
-		}
-		else
-		{
-			float XAxisMultiplier;
-			float YAxisMultiplier;
-
-			if ( ViewportSizeXY.X > ViewportSizeXY.Y )
-			{
-				//if the viewport is wider than it is tall
-				XAxisMultiplier = 1.0f;
-				YAxisMultiplier = ViewportSizeXY.X / (float)ViewportSizeXY.Y;
-			}
-			else
-			{
-				//if the viewport is taller than it is wide
-				XAxisMultiplier = ViewportSizeXY.Y / (float)ViewportSizeXY.X;
-				YAxisMultiplier = 1.0f;
-			}
-
-			ViewInitOptions.ProjectionMatrix = FReversedZPerspectiveMatrix(
-				MatrixFOV,
-				MatrixFOV,
-				XAxisMultiplier,
-				YAxisMultiplier,
-				MinZ,
-				MaxZ
-				);
-		}
-	}
-	else if ( ProjectionType == ECameraProjectionMode::Orthographic )
-	{
-		float ZScale = 0.5f / HALF_WORLD_MAX;
-		float ZOffset = HALF_WORLD_MAX;
-
-		//The divisor for the matrix needs to match the translation code.
-		const float Zoom = GetOrthoUnitsPerPixel(Viewport);
-
-		float OrthoWidth = Zoom * ViewportSizeXY.X / 2.0f;
-		float OrthoHeight = Zoom * ViewportSizeXY.Y / 2.0f;
-
-		ViewInitOptions.ViewRotationMatrix = FInverseRotationMatrix(ViewRotation);
-		ViewInitOptions.ViewRotationMatrix = ViewInitOptions.ViewRotationMatrix * FMatrix(
-			FPlane(0, 0, 1, 0),
-			FPlane(1, 0, 0, 0),
-			FPlane(0, 1, 0, 0),
-			FPlane(0, 0, 0, 1));
-
-		const float EffectiveAspectRatio = bConstrainAspectRatio ? AspectRatio : ( ViewportSizeXY.X / (float)ViewportSizeXY.Y );
-		const float YScale = 1.0f / EffectiveAspectRatio;
-		OrthoWidth = ViewInfo.OrthoWidth / 2.0f;
-		OrthoHeight = ( ViewInfo.OrthoWidth / 2.0f ) * YScale;
-
-		const float NearViewPlane = -ViewInitOptions.ViewOrigin.X;
-		const float FarViewPlane = NearViewPlane - 2.0f*WORLD_MAX;
-
-		const float InverseRange = 1.0f / ( FarViewPlane - NearViewPlane );
-		ZScale = -2.0f * InverseRange;
-		ZOffset = -( FarViewPlane + NearViewPlane ) * InverseRange;
-
-		ViewInitOptions.ProjectionMatrix = FReversedZOrthoMatrix(
-			OrthoWidth,
-			OrthoHeight,
-			ZScale,
-			ZOffset
-			);
-
-		if ( bConstrainAspectRatio )
-		{
-			ViewInitOptions.SetConstrainedViewRectangle(Viewport->CalculateViewExtents(AspectRatio, ViewRect));
-		}
-	}
-	else
-	{
-		check(false);
-	}
+	FMinimalViewInfo::CalculateProjectionMatrixGivenView(ViewInfo, AspectRatioAxisConstraint, Viewport, /*inout*/ ViewInitOptions);
 
 	ViewInitOptions.ViewFamily = ViewFamily;
 	ViewInitOptions.SceneViewStateInterface = ViewState.GetReference();
@@ -425,7 +295,6 @@ FSceneView* FUMGViewportClient::CalcSceneView(FSceneViewFamily* ViewFamily)
 	//ViewInitOptions.OverrideLODViewOrigin = FVector::ZeroVector;
 	//ViewInitOptions.bUseFauxOrthoViewPos = true;
 
-	ViewInitOptions.OverrideFarClippingPlaneDistance = FarPlane;
 	//ViewInitOptions.CursorPos = CurrentMousePos;
 
 	FSceneView* View = new FSceneView(ViewInitOptions);
diff --git a/Engine/Source/Runtime/UMG/Public/Components/Viewport.h b/Engine/Source/Runtime/UMG/Public/Components/Viewport.h
--- a/Engine/Source/Runtime/UMG/Public/Components/Viewport.h
+++ b/Engine/Source/Runtime/UMG/Public/Components/Viewport.h
@@ -93,42 +93,33 @@ public:
 	FUMGViewportClient(FPreviewScene* InPreviewScene = nullptr);
 	virtual ~FUMGViewportClient();
 
+	using FViewElementDrawer::Draw;
+
+	// FGCObject interface
 	virtual void AddReferencedObjects(FReferenceCollector & Collector) override;
 
+	// FViewportClient interface
+	virtual UWorld* GetWorld() const override;
 	virtual void Draw(FViewport* InViewport, FCanvas* Canvas) override;
 
+	// FUMGViewportClient
+
 	virtual void Tick(float InDeltaTime);
 
-	using FViewElementDrawer::Draw;
+	virtual FSceneView* CalcSceneView(FSceneViewFamily* ViewFamily);
 
-	virtual UWorld* GetWorld() const override;
 
 	/**
-	* @return The scene being rendered in this viewport
-	*/
+	 * @return The scene being rendered in this viewport
+	 */
 	virtual FSceneInterface* GetScene() const;
 
-	virtual FSceneView* CalcSceneView(FSceneViewFamily* ViewFamily);
-
-	bool IsOrtho() const;
-
-	bool IsPerspective() const;
 
 	bool IsAspectRatioConstrained() const;
 
 	void SetBackgroundColor(FLinearColor InBackgroundColor);
 	FLinearColor GetBackgroundColor() const;
 
-	float GetNearClipPlane() const;
-
-	void OverrideNearClipPlane(float InNearPlane);
-
-	float GetFarClipPlaneOverride() const;
-
-	void OverrideFarClipPlane(const float InFarPlane);
-
-	ECameraProjectionMode::Type GetProjectionType() const;
-
 	/** Sets the location of the viewport's camera */
 	void SetViewLocation(const FVector& NewLocation)
 	{
@@ -215,24 +206,11 @@ protected:
 
 	FViewport* Viewport;
 
-	/** near plane adjustable for each editor view, if < 0 GNearClippingPlane should be used. */
-	float NearPlane;
-
-	/** If > 0, overrides the view's far clipping plane with a plane at the specified distance. */
-	float FarPlane;
-
-	/** Viewport's current horizontal field of view (can be modified by locked cameras etc.) */
-	float ViewFOV;
-	/** Viewport's stored horizontal field of view (saved in ini files). */
-	float FOVAngle;
-
-	float AspectRatio;
-
 	/** The viewport's scene view state. */
 	FSceneViewStateReference ViewState;
 
 	/** A set of flags that determines visibility for various scene elements. */
-	FEngineShowFlags		EngineShowFlags;
+	FEngineShowFlags EngineShowFlags;
 };
 
 /**
