Fix convert BSP automation test causing an ensure

diff --git a/Engine/Source/Editor/UnrealEd/Private/Tests/EditorAutomationTests.cpp b/Engine/Source/Editor/UnrealEd/Private/Tests/EditorAutomationTests.cpp
--- a/Engine/Source/Editor/UnrealEd/Private/Tests/EditorAutomationTests.cpp
+++ b/Engine/Source/Editor/UnrealEd/Private/Tests/EditorAutomationTests.cpp
@@ -1134,6 +1134,9 @@ bool FConvertToValidation::RunTest(const FString& Parameters)
 
 	//BSP TO BLOCKING VOLUME
 	{
+		// Note: Rebuilding BSP requires a transaction.
+		FScopedTransaction Transaction(NSLOCTEXT("EditorAutomation", "ConvertBSPToBlocking", "Convert BSP to Blocking Volume"));
+
 		TArray<ABrush*> PreviousBrushes;
 		ConvertTestFindAllBrushes(PreviousBrushes);
 
@@ -1158,6 +1161,9 @@ bool FConvertToValidation::RunTest(const FString& Parameters)
 		GEditor->GetSelectedActors()->Select( NewBrush );	
 		GEditor->ConvertSelectedBrushesToVolumes( ABlockingVolume::StaticClass() );
 		GEditor->RebuildAlteredBSP();
+
+		// During automation we do not actually care about creating a transaction for the user to undo.  
+		Transaction.Cancel();
 	}
 
 	//convert to static mesh
