Fix shadowed variables

diff --git a/Engine/Source/Editor/UnrealEd/Private/AutoReimport/AutoReimportManager.cpp b/Engine/Source/Editor/UnrealEd/Private/AutoReimport/AutoReimportManager.cpp
--- a/Engine/Source/Editor/UnrealEd/Private/AutoReimport/AutoReimportManager.cpp
+++ b/Engine/Source/Editor/UnrealEd/Private/AutoReimport/AutoReimportManager.cpp
@@ -283,8 +283,8 @@ FText FAutoReimportManager::GetProgressText() const
 	}
 
 	{
-		const int32 Total = Utils::Reduce(DirectoryMonitors, [](const FContentDirectoryMonitor& Monitor, int32 Total){
-			return Total + Monitor.GetTotalWork();
+		const int32 Total = Utils::Reduce(DirectoryMonitors, [](const FContentDirectoryMonitor& Monitor, int32 InTotal){
+			return InTotal + Monitor.GetTotalWork();
 		}, 0);
 
 		Args.Add(Total);
@@ -558,8 +558,8 @@ void FAutoReimportManager::SetUpDirectoryMonitors()
 		if (FileManager.DirectoryExists(*ThisPath.Path))
 		{
 			// Set the mounted path if necessary
-			auto* Pair = MountedPaths.FindByPredicate([&](const TPair<FString, FString>& Pair){
-				return ThisPath.Path.StartsWith(Pair.Key);
+			auto* Pair = MountedPaths.FindByPredicate([&](const TPair<FString, FString>& InPair){
+				return ThisPath.Path.StartsWith(InPair.Key);
 			});
 
 			if (Pair)
diff --git a/Engine/Source/Editor/UnrealEd/Private/AutoReimport/FileCache.cpp b/Engine/Source/Editor/UnrealEd/Private/AutoReimport/FileCache.cpp
--- a/Engine/Source/Editor/UnrealEd/Private/AutoReimport/FileCache.cpp
+++ b/Engine/Source/Editor/UnrealEd/Private/AutoReimport/FileCache.cpp
@@ -361,12 +361,12 @@ TOptional<FUpdateCacheTransaction> FFileCache::ProcessChange(const FString& File
 				return X.Action == FFileChangeData::FCA_Removed && X.Filename == TransactionFilename;
 			});
 
-			const auto Action = NumRemoved == 0 ? FFileChangeData::FCA_Added : FFileChangeData::FCA_Modified;
+			const auto FileChangeAction = NumRemoved == 0 ? FFileChangeData::FCA_Added : FFileChangeData::FCA_Modified;
 
 			// If it's a modification (remove then add) or the file doesn't already exist in the cache, we can create a transaction
-			if (Action == FFileChangeData::FCA_Modified || !CachedDirectoryState.Files.Find(TransactionFilename))
+			if (FileChangeAction == FFileChangeData::FCA_Modified || !CachedDirectoryState.Files.Find(TransactionFilename))
 			{
-				return FUpdateCacheTransaction(TransactionFilename, Action, FileManager.GetTimeStamp(*Filename));
+				return FUpdateCacheTransaction(TransactionFilename, FileChangeAction, FileManager.GetTimeStamp(*Filename));
 			}
 		}
 		break;
diff --git a/Engine/Source/Editor/UnrealEd/Private/AutoReimport/ReimportFeedbackContext.cpp b/Engine/Source/Editor/UnrealEd/Private/AutoReimport/ReimportFeedbackContext.cpp
--- a/Engine/Source/Editor/UnrealEd/Private/AutoReimport/ReimportFeedbackContext.cpp
+++ b/Engine/Source/Editor/UnrealEd/Private/AutoReimport/ReimportFeedbackContext.cpp
@@ -137,8 +137,7 @@ class SWidgetStack : public SVerticalBox
 			const float Alpha = 1.f - SlideCurve.GetLerp();
 			float PositionSoFar = AllottedGeometry.GetLocalSize().Y + Alpha * StartSlideOffset;
 
-			int32 Index = 0;
-			for (; PositionSoFar > 0 && Index < NumSlots(); ++Index)
+			for (int32 Index = 0; PositionSoFar > 0 && Index < NumSlots(); ++Index)
 			{
 				const SBoxPanel::FSlot& CurChild = Children[Index];
 				const EVisibility ChildVisibility = CurChild.GetWidget()->GetVisibility();
diff --git a/Engine/Source/Editor/UnrealEd/Private/Commandlets/ChunkManifestGenerator.cpp b/Engine/Source/Editor/UnrealEd/Private/Commandlets/ChunkManifestGenerator.cpp
--- a/Engine/Source/Editor/UnrealEd/Private/Commandlets/ChunkManifestGenerator.cpp
+++ b/Engine/Source/Editor/UnrealEd/Private/Commandlets/ChunkManifestGenerator.cpp
@@ -198,7 +198,7 @@ void FChunkManifestGenerator::Initialize(bool InGenerateChunks)
 	}
 }
 
-void FChunkManifestGenerator::AddPackageToChunkManifest(UPackage* Package, const FString& SandboxFilename, const FString& LastLoadedMapName, FSandboxPlatformFile* SandboxFile)
+void FChunkManifestGenerator::AddPackageToChunkManifest(UPackage* Package, const FString& SandboxFilename, const FString& LastLoadedMapName, FSandboxPlatformFile* InSandboxFile)
 {
 	TArray<int32> TargetChunks;
 	TArray<int32> ExistingChunkIDs;
@@ -277,10 +277,10 @@ void FChunkManifestGenerator::CleanManifestDirectories()
 	}
 }
 
-bool FChunkManifestGenerator::SaveManifests(FSandboxPlatformFile* SandboxFile)
+bool FChunkManifestGenerator::SaveManifests(FSandboxPlatformFile* InSandboxFile)
 {
 	// Always do package dependency work, is required to modify asset registry
-	FixupPackageDependenciesForChunks(SandboxFile);
+	FixupPackageDependenciesForChunks(InSandboxFile);
 
 	if (bGenerateChunks)
 	{
@@ -726,7 +726,7 @@ void FChunkManifestGenerator::NotifyPackageWasCooked(const FString& PackageSandb
 	AllCookedPackages.Add(PackageName, PackageSandboxPath);
 }
 
-void FChunkManifestGenerator::FixupPackageDependenciesForChunks(FSandboxPlatformFile* SandboxFile)
+void FChunkManifestGenerator::FixupPackageDependenciesForChunks(FSandboxPlatformFile* InSandboxFile)
 {
 	for (int32 ChunkID = 0, MaxChunk = ChunkManifests.Num(); ChunkID < MaxChunk; ++ChunkID)
 	{
@@ -738,7 +738,7 @@ void FChunkManifestGenerator::FixupPackageDependenciesForChunks(FSandboxPlatform
 		FinalChunkManifests[ChunkID] = new FChunkPackageSet();
 		for (auto It = ChunkManifests[ChunkID]->CreateConstIterator(); It; ++It)
 		{
-			AddPackageAndDependenciesToChunk(FinalChunkManifests[ChunkID], It.Key(), It.Value(), ChunkID, SandboxFile);
+			AddPackageAndDependenciesToChunk(FinalChunkManifests[ChunkID], It.Key(), It.Value(), ChunkID, InSandboxFile);
 		}
 	}
 
@@ -755,7 +755,7 @@ void FChunkManifestGenerator::FixupPackageDependenciesForChunks(FSandboxPlatform
 	auto RemainingAssets = UnassignedPackageSet;
 	for (auto It = RemainingAssets.CreateConstIterator(); It; ++It)
 	{
-		AddPackageAndDependenciesToChunk(FinalChunkManifests[0], It.Key(), It.Value(), 0, SandboxFile);
+		AddPackageAndDependenciesToChunk(FinalChunkManifests[0], It.Key(), It.Value(), 0, InSandboxFile);
 	}
 
 	//Finally, if the previous step may added any extra packages to the 0 chunk. Pull them out of other chunks and save space
@@ -792,10 +792,10 @@ void FChunkManifestGenerator::FixupPackageDependenciesForChunks(FSandboxPlatform
 	}
 }
 
-void FChunkManifestGenerator::AddPackageAndDependenciesToChunk(FChunkPackageSet* ThisPackageSet, FName InPkgName, const FString& SandboxFile, int32 ChunkID, FSandboxPlatformFile* SandboxPlatformFile)
+void FChunkManifestGenerator::AddPackageAndDependenciesToChunk(FChunkPackageSet* ThisPackageSet, FName InPkgName, const FString& InSandboxFile, int32 ChunkID, FSandboxPlatformFile* SandboxPlatformFile)
 {
 	//Add this asset
-	ThisPackageSet->Add(InPkgName, SandboxFile);
+	ThisPackageSet->Add(InPkgName, InSandboxFile);
 
 	//Only gather dependencies if we're chunking
 	if (!bGenerateChunks)
diff --git a/Engine/Source/Editor/UnrealEd/Private/Commandlets/DumpBlueprintsInfoCommandlet.cpp b/Engine/Source/Editor/UnrealEd/Private/Commandlets/DumpBlueprintsInfoCommandlet.cpp
--- a/Engine/Source/Editor/UnrealEd/Private/Commandlets/DumpBlueprintsInfoCommandlet.cpp
+++ b/Engine/Source/Editor/UnrealEd/Private/Commandlets/DumpBlueprintsInfoCommandlet.cpp
@@ -1187,20 +1187,22 @@ static bool DumpBlueprintInfoUtils::DumpActionDatabaseInfo(uint32 Indent, FArchi
 			}
 		}
 
-		int32 SpawnerCount = 0;
-		for (TObjectIterator<UBlueprintNodeSpawner> NodeSpawnerIt; NodeSpawnerIt; ++NodeSpawnerIt)
-		{
-			++SpawnerCount;
-			// @TODO: doesn't account for any allocated memory (for delegates, text strings, etc.)
-			EstimatedSystemSize += sizeof(**NodeSpawnerIt);
-		}
-
 		FString const OriginalIndent = BuildIndentString(Indent);
 		FString const IndentedNewline = "\n" + BuildIndentString(Indent + 1);
 
-		FString DatabaseInfoHeading = FString::Printf(TEXT("%s\"ActionDatabaseInfo\" : {%s\"TotalNodeSpawnerCount\" : %d,"), 
-			*OriginalIndent, *IndentedNewline, SpawnerCount);
-		FileOutWriter->Serialize(TCHAR_TO_ANSI(*DatabaseInfoHeading), DatabaseInfoHeading.Len());
+		{
+			int32 SpawnerCount = 0;
+			for (TObjectIterator<UBlueprintNodeSpawner> NodeSpawnerIt; NodeSpawnerIt; ++NodeSpawnerIt)
+			{
+				++SpawnerCount;
+				// @TODO: doesn't account for any allocated memory (for delegates, text strings, etc.)
+				EstimatedSystemSize += sizeof(**NodeSpawnerIt);
+			}
+
+			FString DatabaseInfoHeading = FString::Printf(TEXT("%s\"ActionDatabaseInfo\" : {%s\"TotalNodeSpawnerCount\" : %d,"), 
+				*OriginalIndent, *IndentedNewline, SpawnerCount);
+			FileOutWriter->Serialize(TCHAR_TO_ANSI(*DatabaseInfoHeading), DatabaseInfoHeading.Len());
+		}
 
 		//--------------------------------------
 		// Dumping Database Stats
@@ -2449,19 +2451,19 @@ int32 UDumpBlueprintsInfoCommandlet::Main(FString const& Params)
 	// closing out the writer (and diffing the resultant file if the user deigns us to do so)
 	auto CloseFileStream = [bDiffGeneratedFile, &ActiveFilePath, &CommandOptions](FArchive** FileOutPtr)
 	{
-		FArchive*& FileOut = (*FileOutPtr);
-		if (FileOut != nullptr)
+		FArchive*& ArFileOut = (*FileOutPtr);
+		if (ArFileOut != nullptr)
 		{
-			FileOut->Serialize(TCHAR_TO_ANSI(TEXT("\n}")), 2);
-			FileOut->Close();
+			ArFileOut->Serialize(TCHAR_TO_ANSI(TEXT("\n}")), 2);
+			ArFileOut->Close();
 
 			if (bDiffGeneratedFile)
 			{
 				check(!ActiveFilePath.IsEmpty());
 				DumpBlueprintInfoUtils::DiffDumpFiles(ActiveFilePath, CommandOptions.DiffPath, CommandOptions.DiffCommand);
 			}
-			delete FileOut;
-			FileOut = nullptr;
+			delete ArFileOut;
+			ArFileOut = nullptr;
 		}
 	};
 
diff --git a/Engine/Source/Editor/UnrealEd/Private/EditorClassUtils.cpp b/Engine/Source/Editor/UnrealEd/Private/EditorClassUtils.cpp
--- a/Engine/Source/Editor/UnrealEd/Private/EditorClassUtils.cpp
+++ b/Engine/Source/Editor/UnrealEd/Private/EditorClassUtils.cpp
@@ -91,16 +91,16 @@ TSharedRef<SWidget> FEditorClassUtils::GetSourceLinkFormatted(const UClass* Clas
 		{
 			static void OnEditBlueprintClicked( TWeakObjectPtr<UBlueprint> InBlueprint, TWeakObjectPtr<UObject> InAsset )
 			{
-				if (UBlueprint* Blueprint = InBlueprint.Get())
+				if (UBlueprint* BlueprintToEdit = InBlueprint.Get())
 				{
 					// Set the object being debugged if given an actor reference (if we don't do this before we edit the object the editor wont know we are debugging something)
 					if (UObject* Asset = InAsset.Get())
 					{
-						check(Asset->GetClass()->ClassGeneratedBy == Blueprint);
-						Blueprint->SetObjectBeingDebugged(Asset);
+						check(Asset->GetClass()->ClassGeneratedBy == BlueprintToEdit);
+						BlueprintToEdit->SetObjectBeingDebugged(Asset);
 					}
 					// Open the blueprint
-					GEditor->EditObject( Blueprint );
+					GEditor->EditObject( BlueprintToEdit );
 				}
 			}
 		};
diff --git a/Engine/Source/Editor/UnrealEd/Private/EditorCommandLineUtils.cpp b/Engine/Source/Editor/UnrealEd/Private/EditorCommandLineUtils.cpp
--- a/Engine/Source/Editor/UnrealEd/Private/EditorCommandLineUtils.cpp
+++ b/Engine/Source/Editor/UnrealEd/Private/EditorCommandLineUtils.cpp
@@ -673,9 +673,9 @@ static void EditorCommandLineUtilsImpl::RunAssetMerge(FMergeAsset const& Base, F
 	// we use a lambda delegate to route the call into MergeHandler (we require 
 	// this intermediate to hold onto a MergeHandler ref, so it doesn't get 
 	// prematurely destroyed at the end of this function)
-	auto HandleEditorClose = [](TSharedRef<FMergeResolutionHandler> MergeHandler)
+	auto HandleEditorClose = [](TSharedRef<FMergeResolutionHandler> InMergeHandler)
 	{
-		MergeHandler->HandleEditorClose();
+		InMergeHandler->HandleEditorClose();
 	};	
 	// have to copy the file into the expected result file when we're done
 	FEditorDelegates::OnShutdownPostPackagesSaved.AddStatic(HandleEditorClose, MergeHandler);
