Editor: Added a 1 second visualization of the directional preview light after it is modified
[UE-10806]

diff --git a/Engine/Source/Editor/UnrealEd/Private/EditorViewportClient.cpp b/Engine/Source/Editor/UnrealEd/Private/EditorViewportClient.cpp
--- a/Engine/Source/Editor/UnrealEd/Private/EditorViewportClient.cpp
+++ b/Engine/Source/Editor/UnrealEd/Private/EditorViewportClient.cpp
@@ -77,6 +77,21 @@ namespace FocusConstants
 	const float TransitionTime = 0.25f;
 }
 
+namespace PreviewLightConstants
+{
+	const float MovingPreviewLightTimerDuration = 1.0f;
+
+	const float MinMouseRadius = 100.0f;
+	const float MinArrowLength = 10.0f;
+	const float ArrowLengthToSizeRatio = 0.1f;
+	const float MouseLengthToArrowLenghtRatio = 0.2f;
+
+	const float ArrowLengthToThicknessRatio = 0.05f;
+	const float MinArrowThickness = 2.0f;
+
+	// Note: MinMouseRadius must be greater than MinArrowLength
+}
+
 /**
  * Cached off joystick input state
  */
@@ -88,8 +103,6 @@ public:
 	TMap <FKey, EInputEvent> KeyEventValues;
 };
 
-
-
 FViewportCameraTransform::FViewportCameraTransform()
 	: TransitionCurve( new FCurveSequence( 0.0f, FocusConstants::TransitionTime, ECurveEaseFunction::CubicOut ) )
 	, ViewLocation( FVector::ZeroVector )
@@ -254,6 +267,8 @@ FEditorViewportClient::FEditorViewportClient(FEditorModeTools* InModeTools, FPre
 	, bCameraLock(false)
 	, EditorViewportWidget(InEditorViewportWidget)
 	, PreviewScene(InPreviewScene)
+	, MovingPreviewLightSavedScreenPos(ForceInitToZero)
+	, MovingPreviewLightTimer(0.0f)
 	, PerspViewModeIndex(DefaultPerspectiveViewMode)
 	, OrthoViewModeIndex(DefaultOrthoViewMode)
 	, NearPlane(-1.0f)
@@ -830,6 +845,17 @@ void FEditorViewportClient::Tick(float DeltaTime)
 		Invalidate();
 	}
 
+	// Update the fade out animation
+	if (MovingPreviewLightTimer > 0.0f)
+	{
+		MovingPreviewLightTimer = FMath::Max(MovingPreviewLightTimer - DeltaTime, 0.0f);
+
+		if (MovingPreviewLightTimer == 0.0f)
+		{
+			Invalidate();
+		}
+	}
+
 	// Tick the editor modes
 	ModeTools->Tick(this, DeltaTime);
 }
@@ -2904,7 +2930,7 @@ void FEditorViewportClient::Draw(FViewport* InViewport, FCanvas* Canvas)
 	Viewport = ViewportBackup;
 }
 
-void FEditorViewportClient::Draw(const FSceneView* View,FPrimitiveDrawInterface* PDI)
+void FEditorViewportClient::Draw(const FSceneView* View, FPrimitiveDrawInterface* PDI)
 {
 	// Draw the drag tool.
 	MouseDeltaTracker->Render3DDragTool( View, PDI );
@@ -2922,10 +2948,51 @@ void FEditorViewportClient::Draw(const FSceneView* View,FPrimitiveDrawInterface*
 	// Draw the current editor mode.
 	ModeTools->Render(View, Viewport, PDI);
 
+	// Draw the preview scene light visualization
+	DrawPreviewLightVisualization(View, PDI);
+
 	// This viewport was just rendered, reset this value.
 	FramesSinceLastDraw = 0;
 }
 
+void FEditorViewportClient::DrawPreviewLightVisualization(const FSceneView* View, FPrimitiveDrawInterface* PDI)
+{
+	// Draw the indicator of the current light direction if it was recently moved
+	if ((PreviewScene != nullptr) && (PreviewScene->DirectionalLight != nullptr) && (MovingPreviewLightTimer > 0.0f))
+	{
+		const float A = MovingPreviewLightTimer / PreviewLightConstants::MovingPreviewLightTimerDuration;
+
+		ULightComponent* Light = PreviewScene->DirectionalLight;
+
+		const FLinearColor ArrowColor = Light->LightColor;
+		const FTransform LightLocalToWorld = Light->GetComponentToWorld();
+
+		// Project the last mouse position during the click into world space
+		FVector LastMouseWorldPos;
+		FVector LastMouseWorldDir;
+		View->DeprojectFVector2D(MovingPreviewLightSavedScreenPos, /*out*/ LastMouseWorldPos, /*out*/ LastMouseWorldDir);
+
+		// The world pos may be nuts due to a super distant near plane for orthographic cameras, so find the closest
+		// point to the origin along the ray
+		LastMouseWorldPos = FMath::ClosestPointOnLine(LastMouseWorldPos, LastMouseWorldPos + LastMouseWorldDir * WORLD_MAX, FVector::ZeroVector);
+
+		// Figure out the radius to draw the light preview ray at
+		const FVector LightToMousePos = LastMouseWorldPos - LightLocalToWorld.GetTranslation();
+		const float LightToMouseRadius = FMath::Max(LightToMousePos.Size(), PreviewLightConstants::MinMouseRadius);
+
+		const float ArrowLength = FMath::Max(PreviewLightConstants::MinArrowLength, LightToMouseRadius * PreviewLightConstants::MouseLengthToArrowLenghtRatio);
+		const float ArrowSize = PreviewLightConstants::ArrowLengthToSizeRatio * ArrowLength;
+		const float ArrowThickness = FMath::Max(PreviewLightConstants::ArrowLengthToThicknessRatio * ArrowLength, PreviewLightConstants::MinArrowThickness);
+
+		const FVector ArrowOrigin = LightLocalToWorld.TransformPosition(FVector(-LightToMouseRadius - 0.5f * ArrowLength, 0.0f, 0.0f));
+		const FVector ArrowDirection = LightLocalToWorld.TransformVector(FVector(-1.0f, 0.0f, 0.0f));
+
+		const FQuatRotationTranslationMatrix ArrowToWorld(LightLocalToWorld.GetRotation(), ArrowOrigin);
+
+		DrawDirectionalArrow(PDI, ArrowToWorld, ArrowColor, ArrowLength, ArrowSize, SDPG_World, ArrowThickness);
+	}
+}
+
 void FEditorViewportClient::RenderDragTool(const FSceneView* View, FCanvas* Canvas)
 {
 	MouseDeltaTracker->RenderDragTool(View, Canvas);
@@ -3088,12 +3155,18 @@ bool FEditorViewportClient::InputAxis(FViewport* InViewport, int32 ControllerId,
 
 	if( bLightMoveDown && bMouseButtonDown && PreviewScene )
 	{
+		// Adjust the preview light direction
 		FRotator LightDir = PreviewScene->GetLightDirection();
 
 		LightDir.Yaw += -DragX * LightRotSpeed;
 		LightDir.Pitch += -DragY * LightRotSpeed;
 
 		PreviewScene->SetLightDirection( LightDir );
+		
+		// Remember that we adjusted it for the visualization
+		MovingPreviewLightTimer = PreviewLightConstants::MovingPreviewLightTimerDuration;
+		MovingPreviewLightSavedScreenPos = FVector2D(LastMouseX, LastMouseY);
+
 		Invalidate();
 	}
 	else
diff --git a/Engine/Source/Editor/UnrealEd/Public/EditorViewportClient.h b/Engine/Source/Editor/UnrealEd/Public/EditorViewportClient.h
--- a/Engine/Source/Editor/UnrealEd/Public/EditorViewportClient.h
+++ b/Engine/Source/Editor/UnrealEd/Public/EditorViewportClient.h
@@ -578,6 +578,9 @@ public:
 	 */
 	virtual void DrawCanvas(FViewport& InViewport, FSceneView& View, FCanvas& Canvas);
 
+	// Draws a visualization of the preview light if it was recently moved
+	virtual void DrawPreviewLightVisualization(const FSceneView* View, FPrimitiveDrawInterface* PDI);
+
 	/**
 	 * Render the drag tool in the viewport
 	 */
@@ -1419,7 +1422,12 @@ protected:
 	FVector DefaultOrbitLocation;
 	FVector DefaultOrbitZoom;
 	FVector DefaultOrbitLookAt;
-	
+
+protected:
+	// Used for the display of the current preview light after it has been adjusted
+	FVector2D MovingPreviewLightSavedScreenPos;
+	float MovingPreviewLightTimer;
+
 public:
 	/* Default view mode for perspective viewports */
 	static const EViewModeIndex DefaultPerspectiveViewMode;
