Fix a few crashes associated with deleting a child blueprint of an actor blueprint

diff --git a/Engine/Source/Editor/UnrealEd/Private/Layers/Layers.cpp b/Engine/Source/Editor/UnrealEd/Private/Layers/Layers.cpp
--- a/Engine/Source/Editor/UnrealEd/Private/Layers/Layers.cpp
+++ b/Engine/Source/Editor/UnrealEd/Private/Layers/Layers.cpp
@@ -92,7 +92,7 @@ void FLayers::RemoveLevelLayerInformation( const TWeakObjectPtr< ULevel >& Level
 
 bool FLayers::IsActorValidForLayer( const TWeakObjectPtr< AActor >& Actor )
 {
-	if( !Actor.IsValid() )
+	if( !Actor.IsValid() || Actor->GetClass() == nullptr || Actor->GetWorld() == nullptr)
 	{
 		return false;
 	}
diff --git a/Engine/Source/Runtime/Engine/Private/Actor.cpp b/Engine/Source/Runtime/Engine/Private/Actor.cpp
--- a/Engine/Source/Runtime/Engine/Private/Actor.cpp
+++ b/Engine/Source/Runtime/Engine/Private/Actor.cpp
@@ -1607,7 +1607,12 @@ void AActor::Destroyed()
 
 	ReceiveDestroyed();
 	OnDestroyed.Broadcast();
-	GetWorld()->RemoveNetworkActor(this);
+	UWorld* ActorWorld = GetWorld();
+
+	if( ActorWorld )
+	{
+		ActorWorld->RemoveNetworkActor(this);
+	}
 }
 
 void AActor::TornOff() {}
