Fixed another crash deleting an actor with no world

diff --git a/Engine/Source/Editor/UnrealEd/Private/ObjectTools.cpp b/Engine/Source/Editor/UnrealEd/Private/ObjectTools.cpp
--- a/Engine/Source/Editor/UnrealEd/Private/ObjectTools.cpp
+++ b/Engine/Source/Editor/UnrealEd/Private/ObjectTools.cpp
@@ -1898,7 +1898,10 @@ namespace ObjectTools
 
 					// Destroy the Actor instance. This is similar to edactDeleteSelected(), but we don't request user confirmation here.
 					GEditor->Layers->DisassociateActorFromLayers( CurActor );
-					GEditor->GetEditorWorldContext().World()->EditorDestroyActor( CurActor, false );
+					if( CurActor->GetWorld() )
+					{
+						CurActor->GetWorld()->EditorDestroyActor( CurActor, false );
+					}
 
 					bNeedsGarbageCollection = true;
 				}
