[UE-10822] We weren't handling trivial case of accepting the blueprint with no changes as the merge resolution, also we weren't closing existing blueprint editors

diff --git a/Engine/Source/Developer/Merge/Private/SBlueprintMerge.cpp b/Engine/Source/Developer/Merge/Private/SBlueprintMerge.cpp
--- a/Engine/Source/Developer/Merge/Private/SBlueprintMerge.cpp
+++ b/Engine/Source/Developer/Merge/Private/SBlueprintMerge.cpp
@@ -522,6 +522,7 @@ void SBlueprintMerge::ResolveMerge(UBlueprint* ResultantBlueprint)
 	// this scenario... we need to open the result to take its place)
 	if (ResultantBlueprint != GetTargetBlueprint())
 	{
+		FAssetEditorManager::Get().CloseAllEditorsForAsset(GetTargetBlueprint());
 		FAssetEditorManager::Get().OpenEditorForAsset(ResultantBlueprint);
 	}
 	
diff --git a/Engine/Source/Editor/UnrealEd/Private/Kismet2/Kismet2.cpp b/Engine/Source/Editor/UnrealEd/Private/Kismet2/Kismet2.cpp
--- a/Engine/Source/Editor/UnrealEd/Private/Kismet2/Kismet2.cpp
+++ b/Engine/Source/Editor/UnrealEd/Private/Kismet2/Kismet2.cpp
@@ -637,6 +637,11 @@ UBlueprint* FKismetEditorUtilities::ReloadBlueprint(UBlueprint* StaleBlueprint)
 
 UBlueprint* FKismetEditorUtilities::ReplaceBlueprint(UBlueprint* Target, UBlueprint const* ReplacementArchetype)
 {
+	if (Target == ReplacementArchetype)
+	{
+		return Target;
+	}
+
 	UPackage* BlueprintPackage = Target->GetOutermost();
 	check(BlueprintPackage != GetTransientPackage());
 	const FString BlueprintName = Target->GetName();
@@ -645,7 +650,7 @@ UBlueprint* FKismetEditorUtilities::ReplaceBlueprint(UBlueprint* Target, UBluepr
 	Unloader.UnloadBlueprint(/*bResetPackage =*/false);
 
 	UBlueprint* Replacement = Cast<UBlueprint>(StaticDuplicateObject(ReplacementArchetype, BlueprintPackage, *BlueprintName));
-
+	
 	Unloader.ReplaceStaleRefs(Replacement);
 	return Replacement;
 }
