Editor: Make lighting preview arrow ignore the directional light translation, to make it easier to see if the light was moved for some reason

diff --git a/Engine/Source/Editor/UnrealEd/Private/EditorViewportClient.cpp b/Engine/Source/Editor/UnrealEd/Private/EditorViewportClient.cpp
--- a/Engine/Source/Editor/UnrealEd/Private/EditorViewportClient.cpp
+++ b/Engine/Source/Editor/UnrealEd/Private/EditorViewportClient.cpp
@@ -2955,6 +2955,7 @@ void FEditorViewportClient::Draw(const FSceneView* View, FPrimitiveDrawInterface
 	FramesSinceLastDraw = 0;
 }
 
+PRAGMA_DISABLE_OPTIMIZATION
 void FEditorViewportClient::DrawPreviewLightVisualization(const FSceneView* View, FPrimitiveDrawInterface* PDI)
 {
 	// Draw the indicator of the current light direction if it was recently moved
@@ -2965,7 +2966,15 @@ void FEditorViewportClient::DrawPreviewLightVisualization(const FSceneView* View
 		ULightComponent* Light = PreviewScene->DirectionalLight;
 
 		const FLinearColor ArrowColor = Light->LightColor;
-		const FTransform LightLocalToWorld = Light->GetComponentToWorld();
+
+		// Figure out where the light is (ignoring position for directional lights)
+		const FTransform LightLocalToWorldRaw = Light->GetComponentToWorld();
+		FTransform LightLocalToWorld = LightLocalToWorldRaw;
+		if (Light->IsA(UDirectionalLightComponent::StaticClass()))
+		{
+			LightLocalToWorld.SetTranslation(FVector::ZeroVector);
+		}
+		LightLocalToWorld.SetScale3D(FVector(1.0f));
 
 		// Project the last mouse position during the click into world space
 		FVector LastMouseWorldPos;
@@ -2992,6 +3001,7 @@ void FEditorViewportClient::DrawPreviewLightVisualization(const FSceneView* View
 		DrawDirectionalArrow(PDI, ArrowToWorld, ArrowColor, ArrowLength, ArrowSize, SDPG_World, ArrowThickness);
 	}
 }
+PRAGMA_ENABLE_OPTIMIZATION
 
 void FEditorViewportClient::RenderDragTool(const FSceneView* View, FCanvas* Canvas)
 {
