Remove pragma disable opts accidentally checked in

diff --git a/Engine/Source/Editor/UnrealEd/Private/EditorViewportClient.cpp b/Engine/Source/Editor/UnrealEd/Private/EditorViewportClient.cpp
--- a/Engine/Source/Editor/UnrealEd/Private/EditorViewportClient.cpp
+++ b/Engine/Source/Editor/UnrealEd/Private/EditorViewportClient.cpp
@@ -2955,7 +2955,6 @@ void FEditorViewportClient::Draw(const FSceneView* View, FPrimitiveDrawInterface
 	FramesSinceLastDraw = 0;
 }
 
-PRAGMA_DISABLE_OPTIMIZATION
 void FEditorViewportClient::DrawPreviewLightVisualization(const FSceneView* View, FPrimitiveDrawInterface* PDI)
 {
 	// Draw the indicator of the current light direction if it was recently moved
@@ -3001,7 +3000,6 @@ void FEditorViewportClient::DrawPreviewLightVisualization(const FSceneView* View
 		DrawDirectionalArrow(PDI, ArrowToWorld, ArrowColor, ArrowLength, ArrowSize, SDPG_World, ArrowThickness);
 	}
 }
-PRAGMA_ENABLE_OPTIMIZATION
 
 void FEditorViewportClient::RenderDragTool(const FSceneView* View, FCanvas* Canvas)
 {
