Made EditorLevelUtils::GetWorlds a dll export

https://github.com/EpicGames/UnrealEngine/pull/789/

diff --git a/Engine/Source/Editor/UnrealEd/Public/EditorLevelUtils.h b/Engine/Source/Editor/UnrealEd/Public/EditorLevelUtils.h
--- a/Engine/Source/Editor/UnrealEd/Public/EditorLevelUtils.h
+++ b/Engine/Source/Editor/UnrealEd/Public/EditorLevelUtils.h
@@ -142,7 +142,7 @@ namespace EditorLevelUtils
 	 * @param	bIncludeInWorld		If true, include the InWorld in the output list.
 	 * @param	bOnlyEditorVisible	If true, only sub-levels that should be visible in-editor are included
 	 */
-	void GetWorlds(UWorld* InWorld, TArray<UWorld*>& OutWorlds, bool bIncludeInWorld, bool bOnlyEditorVisible = false);
+	UNREALED_API void GetWorlds(UWorld* InWorld, TArray<UWorld*>& OutWorlds, bool bIncludeInWorld, bool bOnlyEditorVisible = false);
 
 }
 
