UE-9391: GitHub #804 : Project Displayed Title from Project settings shows on PIE window

diff --git a/Engine/Source/Editor/UnrealEd/Private/PlayLevel.cpp b/Engine/Source/Editor/UnrealEd/Private/PlayLevel.cpp
--- a/Engine/Source/Editor/UnrealEd/Private/PlayLevel.cpp
+++ b/Engine/Source/Editor/UnrealEd/Private/PlayLevel.cpp
@@ -2573,8 +2573,11 @@ UGameInstance* UEditorEngine::CreatePIEGameInstance(int32 PIEInstance, bool bInS
 #else
 	const FString PlatformBitsString( TEXT( "32" ) );
 #endif
+
+	const FText WindowTitleOverride = GetDefault<UGeneralProjectSettings>()->ProjectDisplayedTitle;
+
 	FFormatNamedArguments Args;
-	Args.Add( TEXT("GameName"), FText::FromString( FString( FApp::GetGameName() ) ) );
+	Args.Add( TEXT("GameName"), FText::FromString( FString( WindowTitleOverride.IsEmpty() ? FApp::GetGameName() : WindowTitleOverride.ToString() ) ) );
 	Args.Add( TEXT("PlatformBits"), FText::FromString( PlatformBitsString ) );
 	Args.Add( TEXT("RHIName"), FText::FromName( LegacyShaderPlatformToShaderFormat( GShaderPlatformForFeatureLevel[GMaxRHIFeatureLevel] ) ) );
 
