UE-8484: GitHub #775 : Added the ability to load the latest opened level at editor startup.

diff --git a/Engine/Config/BaseEditorUserSettings.ini b/Engine/Config/BaseEditorUserSettings.ini
--- a/Engine/Config/BaseEditorUserSettings.ini
+++ b/Engine/Config/BaseEditorUserSettings.ini
@@ -132,8 +132,7 @@ bBehaviorTreeEditor=False
 
 [/Script/UnrealEd.EditorLoadingSavingSettings]
 ; True if we should automatically load a default level at start up
-bLoadDefaultLevelAtStartup=True
-
+LoadLevelAtStartup=ProjectDefault
 ; True if we should automatically reimport textures when a change to source content is detected
 bAutoReimportTextures=False
 bAutoReimportCSV=False
diff --git a/Engine/Source/Editor/UnrealEd/Classes/Settings/EditorLoadingSavingSettings.h b/Engine/Source/Editor/UnrealEd/Classes/Settings/EditorLoadingSavingSettings.h
--- a/Engine/Source/Editor/UnrealEd/Classes/Settings/EditorLoadingSavingSettings.h
+++ b/Engine/Source/Editor/UnrealEd/Classes/Settings/EditorLoadingSavingSettings.h
@@ -5,6 +5,16 @@
 
 #include "EditorLoadingSavingSettings.generated.h"
 
+UENUM()
+namespace ELoadLevelAtStartup
+{
+	enum Type
+	{
+		None,
+		ProjectDefault,
+		LastOpened
+	};
+}
 
 /**
  * Implements the Level Editor's loading and saving settings.
@@ -17,9 +27,9 @@ class UNREALED_API UEditorLoadingSavingSettings
 
 public:
 
-	/** Whether to load a default example map at startup */
+	/** Whether to load a default example map at startup  */
 	UPROPERTY(EditAnywhere, config, Category=Startup)
-	uint32 bLoadDefaultLevelAtStartup:1;
+	TEnumAsByte<ELoadLevelAtStartup::Type> LoadLevelAtStartup;
 
 	/** Whether to restore previously open assets at startup */
 	UPROPERTY(EditAnywhere, config, Category=Startup)
diff --git a/Engine/Source/Editor/UnrealEd/Private/FileHelpers.cpp b/Engine/Source/Editor/UnrealEd/Private/FileHelpers.cpp
--- a/Engine/Source/Editor/UnrealEd/Private/FileHelpers.cpp
+++ b/Engine/Source/Editor/UnrealEd/Private/FileHelpers.cpp
@@ -1962,6 +1962,9 @@ void FEditorFileUtils::LoadMap(const FString& InFilename, bool LoadAsTemplate, b
 		return;
 	}
 
+	// Save last opened level name.
+	GConfig->SetString(TEXT("EditorStartup"), TEXT("LastLevel"), *LongMapPackageName, GEditorUserSettingsIni);
+
 	// Deactivate any editor modes when loading a new map
 	GLevelEditorModeTools().DeactivateAllModes();
 
@@ -3132,7 +3135,16 @@ bool FEditorFileUtils::IsFilenameValidForSaving( const FString& Filename, FText&
 
 void FEditorFileUtils::LoadDefaultMapAtStartup()
 {
-	FString EditorStartupMap = GetDefault<UGameMapsSettings>()->EditorStartupMap;
+	FString EditorStartupMap;
+	// Last opened map.
+	if (GetDefault<UEditorLoadingSavingSettings>()->LoadLevelAtStartup == ELoadLevelAtStartup::LastOpened)
+	{
+		GConfig->GetString(TEXT("EditorStartup"), TEXT("LastLevel"), EditorStartupMap, GEditorUserSettingsIni);
+	}
+	// Default project map.
+	if (EditorStartupMap.IsEmpty()) {
+		EditorStartupMap = GetDefault<UGameMapsSettings>()->EditorStartupMap;
+	}
 	
 	const bool bIncludeReadOnlyRoots = true;
 	if ( FPackageName::IsValidLongPackageName(EditorStartupMap, bIncludeReadOnlyRoots) )
diff --git a/Engine/Source/Editor/UnrealEd/Private/UnrealEdMisc.cpp b/Engine/Source/Editor/UnrealEd/Private/UnrealEdMisc.cpp
--- a/Engine/Source/Editor/UnrealEd/Private/UnrealEdMisc.cpp
+++ b/Engine/Source/Editor/UnrealEd/Private/UnrealEdMisc.cpp
@@ -294,7 +294,7 @@ void FUnrealEdMisc::OnInit()
 			{
 				const FString& StartupMap = GetDefault<UGameMapsSettings>()->EditorStartupMap;
 
-				if ((StartupMap.Len() > 0) && GetDefault<UEditorLoadingSavingSettings>()->bLoadDefaultLevelAtStartup)
+				if ((StartupMap.Len() > 0) && (GetDefault<UEditorLoadingSavingSettings>()->LoadLevelAtStartup != ELoadLevelAtStartup::None))
 				{
 					FEditorFileUtils::LoadDefaultMapAtStartup();
 					BeginPerformanceSurvey();
