Can replace local variables nodes in collapsed graphs through the node's right click context menu.

diff --git a/Engine/Source/Editor/BlueprintGraph/Private/EdGraphSchema_K2.cpp b/Engine/Source/Editor/BlueprintGraph/Private/EdGraphSchema_K2.cpp
--- a/Engine/Source/Editor/BlueprintGraph/Private/EdGraphSchema_K2.cpp
+++ b/Engine/Source/Editor/BlueprintGraph/Private/EdGraphSchema_K2.cpp
@@ -1453,7 +1453,8 @@ void UEdGraphSchema_K2::OnReplaceVariableForVariableNode( UK2Node_Variable* Vari
 		}
 		else
 		{
-			Variable->VariableReference.SetLocalMember( FName(*VariableName), Variable->GetGraph()->GetName(), FBlueprintEditorUtils::FindLocalVariableGuidByName(OwnerBlueprint, Variable->GetGraph(), *VariableName));
+			UEdGraph* FunctionGraph = FBlueprintEditorUtils::GetTopLevelGraph(Variable->GetGraph());
+			Variable->VariableReference.SetLocalMember( FName(*VariableName), FunctionGraph->GetName(), FBlueprintEditorUtils::FindLocalVariableGuidByName(OwnerBlueprint, FunctionGraph, *VariableName));
 		}
 		Pin->PinName = VariableName;
 		Variable->ReconstructNode();
diff --git a/Engine/Source/Editor/Kismet/Private/BlueprintEditor.cpp b/Engine/Source/Editor/Kismet/Private/BlueprintEditor.cpp
--- a/Engine/Source/Editor/Kismet/Private/BlueprintEditor.cpp
+++ b/Engine/Source/Editor/Kismet/Private/BlueprintEditor.cpp
@@ -6590,9 +6590,7 @@ bool FBlueprintEditor::NewDocument_IsVisibleForType(ECreatedDocumentType GraphTy
 	case CGT_NewEventGraph:
 		return FBlueprintEditorUtils::DoesSupportEventGraphs(GetBlueprintObj());
 	case CGT_NewLocalVariable:
-		return (GetBlueprintObj()->BlueprintType != BPTYPE_Interface) 
-			&& (GetFocusedGraph() && GetFocusedGraph()->GetSchema()->GetGraphType(GetFocusedGraph()) == EGraphType::GT_Function)
-			&& !GetFocusedGraph()->IsA(UAnimationTransitionGraph::StaticClass()) 
+		return FBlueprintEditorUtils::DoesSupportLocalVariables(GetFocusedGraph()) 
 			&& IsFocusedGraphEditable();
 	}
 
diff --git a/Engine/Source/Editor/UnrealEd/Private/Kismet2/BlueprintEditorUtils.cpp b/Engine/Source/Editor/UnrealEd/Private/Kismet2/BlueprintEditorUtils.cpp
--- a/Engine/Source/Editor/UnrealEd/Private/Kismet2/BlueprintEditorUtils.cpp
+++ b/Engine/Source/Editor/UnrealEd/Private/Kismet2/BlueprintEditorUtils.cpp
@@ -18,6 +18,7 @@
 #include "AnimationGraph.h"
 #include "AnimationGraphSchema.h"
 #include "AnimationStateMachineGraph.h"
+#include "AnimationTransitionGraph.h"
 #include "AnimStateConduitNode.h"
 #include "AnimGraphNode_StateMachine.h"
 #include "Editor/UnrealEd/Public/Kismet2/KismetEditorUtilities.h"
@@ -2648,6 +2649,18 @@ bool FBlueprintEditorUtils::DoesSupportDefaults(UBlueprint const* Blueprint)
 		&& Blueprint->BlueprintType != BPTYPE_FunctionLibrary;
 }
 
+bool FBlueprintEditorUtils::DoesSupportLocalVariables(UEdGraph const* InGraph)
+{
+	if(InGraph)
+	{
+		UBlueprint* Blueprint = FBlueprintEditorUtils::FindBlueprintForGraph(InGraph);
+		return Blueprint->BlueprintType != BPTYPE_Interface
+			&& InGraph->GetSchema()->GetGraphType(InGraph) == EGraphType::GT_Function
+			&& !InGraph->IsA(UAnimationTransitionGraph::StaticClass());
+	}
+	return false;
+}
+
 // Returns a descriptive name of the type of blueprint passed in
 FString FBlueprintEditorUtils::GetBlueprintTypeDescription(const UBlueprint* Blueprint)
 {
@@ -3408,22 +3421,22 @@ void FBlueprintEditorUtils::GetNewVariablesOfType( const UBlueprint* Blueprint,
 
 void FBlueprintEditorUtils::GetLocalVariablesOfType( const UEdGraph* Graph, const FEdGraphPinType& Type, TArray<FName>& OutVars)
 {
-	if ((Graph != nullptr) && (Graph->GetSchema()->GetGraphType(Graph) == GT_Function))
+	if (DoesSupportLocalVariables(Graph))
 	{
+		// Grab the function graph, so we can find the function entry node for local variables
+		UEdGraph* FunctionGraph = FBlueprintEditorUtils::GetTopLevelGraph(Graph);
+
 		TArray<UK2Node_FunctionEntry*> GraphNodes;
-		Graph->GetNodesOfClass<UK2Node_FunctionEntry>(GraphNodes);
+		FunctionGraph->GetNodesOfClass<UK2Node_FunctionEntry>(GraphNodes);
 
-		if (GraphNodes.Num() > 0)
-		{
-			// If there is an entry node, there should only be one
-			check(GraphNodes.Num() == 1);
+		// There should only be one entry node
+		check(GraphNodes.Num() == 1);
 
-			for (const FBPVariableDescription& LocalVar : GraphNodes[0]->LocalVariables)
+		for (const FBPVariableDescription& LocalVar : GraphNodes[0]->LocalVariables)
+		{
+			if (LocalVar.VarType == Type)
 			{
-				if (LocalVar.VarType == Type)
-				{
-					OutVars.Add(LocalVar.VarName);
-				}
+				OutVars.Add(LocalVar.VarName);
 			}
 		}
 	}
@@ -4137,10 +4150,11 @@ FBPVariableDescription* FBlueprintEditorUtils::FindLocalVariable(UBlueprint* InB
 FBPVariableDescription* FBlueprintEditorUtils::FindLocalVariable(const UBlueprint* InBlueprint, const UEdGraph* InScopeGraph, const FName& InVariableName, class UK2Node_FunctionEntry** OutFunctionEntry)
 {
 	FBPVariableDescription* ReturnVariable = NULL;
-	if(InScopeGraph)
+	if(DoesSupportLocalVariables(InScopeGraph))
 	{
+		UEdGraph* FunctionGraph = GetTopLevelGraph(InScopeGraph);
 		TArray<UK2Node_FunctionEntry*> GraphNodes;
-		InScopeGraph->GetNodesOfClass<UK2Node_FunctionEntry>(GraphNodes);
+		FunctionGraph->GetNodesOfClass<UK2Node_FunctionEntry>(GraphNodes);
 
 		bool bFoundLocalVariable = false;
 
diff --git a/Engine/Source/Editor/UnrealEd/Public/Kismet2/BlueprintEditorUtils.h b/Engine/Source/Editor/UnrealEd/Public/Kismet2/BlueprintEditorUtils.h
--- a/Engine/Source/Editor/UnrealEd/Public/Kismet2/BlueprintEditorUtils.h
+++ b/Engine/Source/Editor/UnrealEd/Public/Kismet2/BlueprintEditorUtils.h
@@ -426,6 +426,9 @@ public:
 	/** Returns whether or not the blueprint supports default values (IE has a CDO) */
 	static bool DoesSupportDefaults(UBlueprint const* Blueprint);
 
+	/** Returns whether or not the blueprint graph supports local variables */
+	static bool DoesSupportLocalVariables(UEdGraph const* InGraph);
+
 	// Returns a descriptive name of the type of blueprint passed in
 	static FString GetBlueprintTypeDescription(const UBlueprint* Blueprint);
 
