Store references to blueprint constructied components in undo buffer by name to avoid maintaining references to TRASH_ components

diff --git a/Engine/Source/Editor/UnrealEd/Classes/Editor/Transactor.h b/Engine/Source/Editor/UnrealEd/Classes/Editor/Transactor.h
--- a/Engine/Source/Editor/UnrealEd/Classes/Editor/Transactor.h
+++ b/Engine/Source/Editor/UnrealEd/Classes/Editor/Transactor.h
@@ -97,11 +97,34 @@ protected:
 			}
 		};
 
+		// Structure to store information about a referenced object
+		// If ObjectName is set, it will represent the name of a blueprint constructed component 
+		// and Object will be the Outer of that Component, otherwise Object will be
+		// a direct reference to the object in question
+		struct FReferencedObject
+		{
+		private:
+			UObject* Object;
+			FName ComponentName;
+		public:
+			FReferencedObject() : Object(nullptr) { }
+			FReferencedObject(UObject* InObject);
+			UObject* GetObject() const;
+			void AddReferencedObjects( FReferenceCollector& Collector );
+
+			friend FArchive& operator<<( FArchive& Ar, FReferencedObject& ReferencedObject )
+			{
+				Ar << ReferencedObject.Object;
+				Ar << ReferencedObject.ComponentName;
+				return Ar;
+			}
+		};
+
 		// Variables.
 		/** The data stream used to serialize/deserialize record */
 		TArray<uint8>		Data;
 		/** External objects referenced in the transaction */
-		TArray<UObject*>	ReferencedObjects;
+		TArray<FReferencedObject>	ReferencedObjects;
 		/** FNames referenced in the object record */
 		TArray<FName>		ReferencedNames;
 		/** The object to track */
@@ -149,7 +172,7 @@ protected:
 			FReader(
 				FTransaction* InOwner,
 				const TArray<uint8>& InData,
-				const TArray<UObject*>& InReferencedObjects,
+				const TArray<FReferencedObject>& InReferencedObjects,
 				const TArray<FName>& InReferencedNames,
 				bool bWantBinarySerialization
 				):
@@ -187,7 +210,7 @@ protected:
 			{
 				int32 ObjectIndex = 0;
 				(FArchive&)*this << ObjectIndex;
-				Res = ReferencedObjects[ObjectIndex];
+				Res = ReferencedObjects[ObjectIndex].GetObject();
 				return *this;
 			}
 			void Preload( UObject* InObject )
@@ -205,7 +228,7 @@ protected:
 			}
 			FTransaction* Owner;
 			const TArray<uint8>& Data;
-			const TArray<UObject*>& ReferencedObjects;
+			const TArray<FReferencedObject>& ReferencedObjects;
 			const TArray<FName>& ReferencedNames;
 			int64 Offset;
 		};
@@ -218,7 +241,7 @@ protected:
 		public:
 			FWriter(
 				TArray<uint8>& InData,
-				TArray<UObject*>& InReferencedObjects,
+				TArray<FReferencedObject>& InReferencedObjects,
 				TArray<FName>& InReferencedNames,
 				bool bWantBinarySerialization
 				):
@@ -229,7 +252,7 @@ protected:
 			{
 				for(int32 ObjIndex = 0; ObjIndex < InReferencedObjects.Num(); ++ObjIndex)
 				{
-					ObjectMap.Add(InReferencedObjects[ObjIndex], ObjIndex);
+					ObjectMap.Add(InReferencedObjects[ObjIndex].GetObject(), ObjIndex);
 				}
 
 				ArWantBinaryPropertySerialization = bWantBinarySerialization;
@@ -271,14 +294,14 @@ protected:
 				}
 				else
 				{
-					ObjectIndex = ReferencedObjects.Add(Res);
+					ObjectIndex = ReferencedObjects.Add(FReferencedObject(Res));
 					ObjectMap.Add(Res, ObjectIndex);
 				}
 				return (FArchive&)*this << ObjectIndex;
 			}
 			TArray<uint8>& Data;
 			ObjectMapType ObjectMap;
-			TArray<UObject*>& ReferencedObjects;
+			TArray<FReferencedObject>& ReferencedObjects;
 			TArray<FName>& ReferencedNames;
 			int64 Offset;
 		};
diff --git a/Engine/Source/Editor/UnrealEd/Private/EditorTransaction.cpp b/Engine/Source/Editor/UnrealEd/Private/EditorTransaction.cpp
--- a/Engine/Source/Editor/UnrealEd/Private/EditorTransaction.cpp
+++ b/Engine/Source/Editor/UnrealEd/Private/EditorTransaction.cpp
@@ -97,7 +97,7 @@ void FTransaction::FObjectRecord::Restore( FTransaction* Owner )
 	{
 		bRestored = true;
 		TArray<uint8> FlipData;
-		TArray<UObject*> FlipReferencedObjects;
+		TArray<FReferencedObject> FlipReferencedObjects;
 		TArray<FName> FlipReferencedNames;
 		TSharedPtr<ITransactionObjectAnnotation> FlipObjectAnnotation;
 		if( Owner->bFlip )
@@ -164,22 +164,46 @@ FArchive& operator<<( FArchive& Ar, FTransaction::FObjectRecord& R )
 	return Ar;
 }
 
+FTransaction::FObjectRecord::FReferencedObject::FReferencedObject(UObject* InObject)
+{
+	UActorComponent* Component = Cast<UActorComponent>(InObject);
+	if (Component && Component->IsCreatedByConstructionScript())
+	{
+		Object = Component->GetOuter();
+		ComponentName = Component->GetFName();
+	}
+	else
+	{
+		Object = InObject;
+	}
+}
+
+UObject* FTransaction::FObjectRecord::FReferencedObject::GetObject() const
+{
+	return (ComponentName.IsNone() ? Object : FindObjectFast<UActorComponent>(Object, ComponentName));
+}
+
+void FTransaction::FObjectRecord::FReferencedObject::AddReferencedObjects(FReferenceCollector& Collector)
+{
+	Collector.AddReferencedObject(Object);
+}
+
 void FTransaction::FObjectRecord::AddReferencedObjects( FReferenceCollector& Collector )
 {
 	UObject* Obj = Object.Get();
 	Collector.AddReferencedObject(Obj);
 	Object = Obj;
-	for( int32 ObjIndex = 0; ObjIndex < ReferencedObjects.Num(); ObjIndex++ )
+	for( FReferencedObject& ReferencedObject : ReferencedObjects )
 	{
-		Collector.AddReferencedObject( ReferencedObjects[ ObjIndex ] );
+		ReferencedObject.AddReferencedObjects(Collector);
 	}
 }
 
 void FTransaction::AddReferencedObjects( FReferenceCollector& Collector )
 {
-	for( int32 Index = 0; Index < Records.Num(); Index++ )
+	for( FObjectRecord& ObjectRecord : Records )
 	{
-		Records[ Index ].AddReferencedObjects( Collector );
+		ObjectRecord.AddReferencedObjects( Collector );
 	}
 	Collector.AddReferencedObjects(ObjectMap);
 }
