[INTEGRATE] Change 2462896 by [REDACTED_PERSON][REDACTED_HANDLE] on 2015/02/26 16:27:04

	Support for specifying a list of maps to include in DDC.

diff --git a/Engine/Source/Editor/UnrealEd/Private/Commandlets/DerivedDataCacheCommandlet.cpp b/Engine/Source/Editor/UnrealEd/Private/Commandlets/DerivedDataCacheCommandlet.cpp
--- a/Engine/Source/Editor/UnrealEd/Private/Commandlets/DerivedDataCacheCommandlet.cpp
+++ b/Engine/Source/Editor/UnrealEd/Private/Commandlets/DerivedDataCacheCommandlet.cpp
@@ -55,7 +55,25 @@ int32 UDerivedDataCacheCommandlet::Main( const FString& Params )
 
 		Tokens.Empty(2);
 		Tokens.Add(FString("*") + FPackageName::GetAssetPackageExtension());
-		Tokens.Add(FString("*") + FPackageName::GetMapPackageExtension());
+
+		FString MapList;
+		if(FParse::Value(*Params, TEXT("Map="), MapList))
+		{
+			for(int StartIdx = 0; StartIdx < MapList.Len(); StartIdx++)
+			{
+				int EndIdx = StartIdx;
+				while(EndIdx < MapList.Len() && MapList[EndIdx] != '+')
+				{
+					EndIdx++;
+				}
+				Tokens.Add(MapList.Mid(StartIdx, EndIdx - StartIdx) + FPackageName::GetMapPackageExtension());
+				StartIdx = EndIdx + 1;
+			}
+		}
+		else
+		{
+			Tokens.Add(FString("*") + FPackageName::GetMapPackageExtension());
+		}
 		
 		uint8 PackageFilter = NORMALIZE_DefaultFlags;
 		if ( Switches.Contains(TEXT("MAPSONLY")) )
