Added option to specify whether a save dialog can be declined as opposed to cancelled.
Some dialogs, e.g. from save menu actions, don't need the two options.

diff --git a/Engine/Source/Editor/ContentBrowser/Private/ContentBrowserUtils.cpp b/Engine/Source/Editor/ContentBrowser/Private/ContentBrowserUtils.cpp
--- a/Engine/Source/Editor/ContentBrowser/Private/ContentBrowserUtils.cpp
+++ b/Engine/Source/Editor/ContentBrowser/Private/ContentBrowserUtils.cpp
@@ -632,7 +632,10 @@ bool ContentBrowserUtils::SaveDirtyPackages()
 	const bool bPromptUserToSave = true;
 	const bool bSaveMapPackages = true;
 	const bool bSaveContentPackages = true;
-	return FEditorFileUtils::SaveDirtyPackages( bPromptUserToSave, bSaveMapPackages, bSaveContentPackages );
+	const bool bFastSave = false;
+	const bool bNotifyNoPackagesSaved = false;
+	const bool bCanBeDeclined = false;
+	return FEditorFileUtils::SaveDirtyPackages( bPromptUserToSave, bSaveMapPackages, bSaveContentPackages, bFastSave, bNotifyNoPackagesSaved, bCanBeDeclined );
 }
 
 TArray<UPackage*> ContentBrowserUtils::LoadPackages(const TArray<FString>& PackageNames)
diff --git a/Engine/Source/Editor/LocalizationDashboard/Private/LocalizationDashboard.cpp b/Engine/Source/Editor/LocalizationDashboard/Private/LocalizationDashboard.cpp
--- a/Engine/Source/Editor/LocalizationDashboard/Private/LocalizationDashboard.cpp
+++ b/Engine/Source/Editor/LocalizationDashboard/Private/LocalizationDashboard.cpp
@@ -356,8 +356,14 @@ namespace
 		for (ULocalizationTarget* const LocalizationTarget : Settings->TargetObjects)
 		{
 			// Save unsaved packages.
+			const bool bPromptUserToSave = true;
+			const bool bSaveMapPackages = true;
+			const bool bSaveContentPackages = true;
+			const bool bFastSave = false;
+			const bool bNotifyNoPackagesSaved = false;
+			const bool bCanBeDeclined = true;
 			bool DidPackagesNeedSaving;
-			const bool WerePackagesSaved = FEditorFileUtils::SaveDirtyPackages(true, true, true, false, false, &DidPackagesNeedSaving);
+			const bool WerePackagesSaved = FEditorFileUtils::SaveDirtyPackages(bPromptUserToSave, bSaveMapPackages, bSaveContentPackages, bFastSave, bNotifyNoPackagesSaved, bCanBeDeclined, &DidPackagesNeedSaving);
 
 			if (DidPackagesNeedSaving && !WerePackagesSaved)
 			{
diff --git a/Engine/Source/Editor/LocalizationDashboard/Private/SLocalizationDashboardTargetRow.cpp b/Engine/Source/Editor/LocalizationDashboard/Private/SLocalizationDashboardTargetRow.cpp
--- a/Engine/Source/Editor/LocalizationDashboard/Private/SLocalizationDashboardTargetRow.cpp
+++ b/Engine/Source/Editor/LocalizationDashboard/Private/SLocalizationDashboardTargetRow.cpp
@@ -328,8 +328,14 @@ FReply SLocalizationDashboardTargetRow::Gather()
 	if (LocalizationTarget)
 	{
 		// Save unsaved packages.
+		const bool bPromptUserToSave = true;
+		const bool bSaveMapPackages = true;
+		const bool bSaveContentPackages = true;
+		const bool bFastSave = false;
+		const bool bNotifyNoPackagesSaved = false;
+		const bool bCanBeDeclined = true;
 		bool DidPackagesNeedSaving;
-		const bool WerePackagesSaved = FEditorFileUtils::SaveDirtyPackages(true, true, true, false, false, &DidPackagesNeedSaving);
+		const bool WerePackagesSaved = FEditorFileUtils::SaveDirtyPackages(bPromptUserToSave, bSaveMapPackages, bSaveContentPackages, bFastSave, bNotifyNoPackagesSaved, bCanBeDeclined, &DidPackagesNeedSaving);
 
 		if (DidPackagesNeedSaving && !WerePackagesSaved)
 		{
diff --git a/Engine/Source/Editor/LocalizationDashboard/Private/SLocalizationTargetEditor.cpp b/Engine/Source/Editor/LocalizationDashboard/Private/SLocalizationTargetEditor.cpp
--- a/Engine/Source/Editor/LocalizationDashboard/Private/SLocalizationTargetEditor.cpp
+++ b/Engine/Source/Editor/LocalizationDashboard/Private/SLocalizationTargetEditor.cpp
@@ -651,8 +651,14 @@ namespace
 		if (LocalizationTarget)
 		{
 			// Save unsaved packages.
+			const bool bPromptUserToSave = true;
+			const bool bSaveMapPackages = true;
+			const bool bSaveContentPackages = true;
+			const bool bFastSave = false;
+			const bool bNotifyNoPackagesSaved = false;
+			const bool bCanBeDeclined = true;
 			bool DidPackagesNeedSaving;
-			const bool WerePackagesSaved = FEditorFileUtils::SaveDirtyPackages(true, true, true, false, false, &DidPackagesNeedSaving);
+			const bool WerePackagesSaved = FEditorFileUtils::SaveDirtyPackages(bPromptUserToSave, bSaveMapPackages, bSaveContentPackages, bFastSave, bNotifyNoPackagesSaved, bCanBeDeclined, &DidPackagesNeedSaving);
 
 			if (DidPackagesNeedSaving && !WerePackagesSaved)
 			{
diff --git a/Engine/Source/Editor/MainFrame/Private/Frame/MainFrameActions.cpp b/Engine/Source/Editor/MainFrame/Private/Frame/MainFrameActions.cpp
--- a/Engine/Source/Editor/MainFrame/Private/Frame/MainFrameActions.cpp
+++ b/Engine/Source/Editor/MainFrame/Private/Frame/MainFrameActions.cpp
@@ -184,7 +184,8 @@ void FMainFrameActionCallbacks::ChoosePackagesToSave()
 	const bool bFastSave = false;
 	const bool bClosingEditor = false;
 	const bool bNotifyNoPackagesSaved = true;
-	FEditorFileUtils::SaveDirtyPackages( bPromptUserToSave, bSaveMapPackages, bSaveContentPackages, bFastSave, bNotifyNoPackagesSaved );
+	const bool bCanBeDeclined = false;
+	FEditorFileUtils::SaveDirtyPackages( bPromptUserToSave, bSaveMapPackages, bSaveContentPackages, bFastSave, bNotifyNoPackagesSaved, bCanBeDeclined );
 }
 
 void FMainFrameActionCallbacks::ChoosePackagesToCheckIn()
@@ -214,7 +215,9 @@ void FMainFrameActionCallbacks::SaveAll()
 	const bool bSaveMapPackages = true;
 	const bool bSaveContentPackages = true;
 	const bool bFastSave = false;
-	FEditorFileUtils::SaveDirtyPackages( bPromptUserToSave, bSaveMapPackages, bSaveContentPackages, bFastSave );
+	const bool bNotifyNoPackagesSaved = false;
+	const bool bCanBeDeclined = false;
+	FEditorFileUtils::SaveDirtyPackages( bPromptUserToSave, bSaveMapPackages, bSaveContentPackages, bFastSave, bNotifyNoPackagesSaved, bCanBeDeclined );
 }
 
 TArray<FString> FMainFrameActionCallbacks::ProjectNames;
diff --git a/Engine/Source/Editor/MainFrame/Private/Frame/MainFrameHandler.h b/Engine/Source/Editor/MainFrame/Private/Frame/MainFrameHandler.h
--- a/Engine/Source/Editor/MainFrame/Private/Frame/MainFrameHandler.h
+++ b/Engine/Source/Editor/MainFrame/Private/Frame/MainFrameHandler.h
@@ -121,7 +121,8 @@ public:
 					const bool bSaveContentPackages = true;
 					const bool bFastSave = false;
 					const bool bNotifyNoPackagesSaved = false;
-					bOkToExit = bOkToExit && FEditorFileUtils::SaveDirtyPackages(bPromptUserToSave, bSaveMapPackages, bSaveContentPackages, bFastSave, bNotifyNoPackagesSaved, &bHadPackagesToSave);
+					const bool bCanBeDeclined = true;
+					bOkToExit = bOkToExit && FEditorFileUtils::SaveDirtyPackages(bPromptUserToSave, bSaveMapPackages, bSaveContentPackages, bFastSave, bNotifyNoPackagesSaved, bCanBeDeclined, &bHadPackagesToSave);
 				}
 
 				// If there were packages to save, or switching project, then the user already had a chance to bail out of exiting.
diff --git a/Engine/Source/Editor/UnrealEd/Private/FileHelpers.cpp b/Engine/Source/Editor/UnrealEd/Private/FileHelpers.cpp
--- a/Engine/Source/Editor/UnrealEd/Private/FileHelpers.cpp
+++ b/Engine/Source/Editor/UnrealEd/Private/FileHelpers.cpp
@@ -2518,7 +2518,7 @@ static void WarnUserAboutFailedSave( const TArray<UPackage*>& InFailedPackages )
 	}
 }
 
-static bool InternalSavePackages(TArray<UPackage*>& PackagesToSave, int32 NumPackagesNotIgnored, bool bPromptUserToSave, bool bFastSave, bool bNotifyNoPackagesSaved, bool* bOutPackagesNeededSaving)
+static bool InternalSavePackages(TArray<UPackage*>& PackagesToSave, int32 NumPackagesNotIgnored, bool bPromptUserToSave, bool bFastSave, bool bNotifyNoPackagesSaved, bool bCanBeDeclined, bool* bOutPackagesNeededSaving)
 {
 	bool bReturnCode = true;
 
@@ -2532,7 +2532,9 @@ static bool InternalSavePackages(TArray<UPackage*>& PackagesToSave, int32 NumPac
 
 		if (!bFastSave)
 		{
-			const FEditorFileUtils::EPromptReturnCode Return = FEditorFileUtils::PromptForCheckoutAndSave(PackagesToSave, true, bPromptUserToSave);
+			const bool bCheckDirty = true;
+			const bool bAlreadyCheckedOut = false;
+			const FEditorFileUtils::EPromptReturnCode Return = FEditorFileUtils::PromptForCheckoutAndSave(PackagesToSave, bCheckDirty, bPromptUserToSave, nullptr, bAlreadyCheckedOut, bCanBeDeclined);
 			if (Return == FEditorFileUtils::EPromptReturnCode::PR_Cancelled)
 			{
 				// Only cancel should return false and stop whatever we were doing before.(like closing the editor)
@@ -2612,7 +2614,7 @@ static bool InternalSavePackages(TArray<UPackage*>& PackagesToSave, int32 NumPac
 	return bReturnCode;
 }
 
-bool FEditorFileUtils::SaveDirtyPackages(const bool bPromptUserToSave, const bool bSaveMapPackages, const bool bSaveContentPackages, const bool bFastSave, const bool bNotifyNoPackagesSaved,  bool* bOutPackagesNeededSaving )
+bool FEditorFileUtils::SaveDirtyPackages(const bool bPromptUserToSave, const bool bSaveMapPackages, const bool bSaveContentPackages, const bool bFastSave, const bool bNotifyNoPackagesSaved, const bool bCanBeDeclined, bool* bOutPackagesNeededSaving )
 {
 	if (bOutPackagesNeededSaving != NULL)
 	{
@@ -2647,10 +2649,10 @@ bool FEditorFileUtils::SaveDirtyPackages(const bool bPromptUserToSave, const boo
 		NumPackagesNotIgnored += (PackagesNotSavedDuringSaveAll.Find(Package->GetName()) == NULL) ? 1 : 0;
 	}
 
-	return InternalSavePackages(PackagesToSave, NumPackagesNotIgnored, bPromptUserToSave, bFastSave, bNotifyNoPackagesSaved, bOutPackagesNeededSaving);
+	return InternalSavePackages(PackagesToSave, NumPackagesNotIgnored, bPromptUserToSave, bFastSave, bNotifyNoPackagesSaved, bCanBeDeclined, bOutPackagesNeededSaving);
 }
 
-bool FEditorFileUtils::SaveDirtyContentPackages(TArray<UClass*>& SaveContentClasses, const bool bPromptUserToSave, const bool bFastSave, const bool bNotifyNoPackagesSaved)
+bool FEditorFileUtils::SaveDirtyContentPackages(TArray<UClass*>& SaveContentClasses, const bool bPromptUserToSave, const bool bFastSave, const bool bNotifyNoPackagesSaved, const bool bCanBeDeclined)
 {
 	bool bReturnCode = true;
 
@@ -2705,7 +2707,7 @@ bool FEditorFileUtils::SaveDirtyContentPackages(TArray<UClass*>& SaveContentClas
 		}
 	}
 
-	return InternalSavePackages(PackagesToSave, PackagesToSave.Num(), bPromptUserToSave, bFastSave, bNotifyNoPackagesSaved, NULL);
+	return InternalSavePackages(PackagesToSave, PackagesToSave.Num(), bPromptUserToSave, bFastSave, bNotifyNoPackagesSaved, bCanBeDeclined, NULL);
 }
 
 /**
@@ -2739,6 +2741,7 @@ bool FEditorFileUtils::SaveCurrentLevel()
  * @param		bPromptToSave				If true the user will be prompted with a list of packages to save, otherwise all passed in packages are saved
  * @param		OutFailedPackages			[out] If specified, will be filled in with all of the packages that failed to save successfully
  * @param		bAlreadyCheckedOut			If true, the user will not be prompted with the source control dialog
+ * @param		bCanBeDeclined				If true, offer a "Don't Save" option in addition to "Cancel", which will not result in a cancellation return code.
  *
  * @return		An enum value signifying success, failure, user declined, or cancellation. If any packages at all failed to save during execution, the return code will be 
  *				failure, even if other packages successfully saved. If the user cancels at any point during any prompt, the return code will be cancellation, even though it
@@ -2746,7 +2749,7 @@ bool FEditorFileUtils::SaveCurrentLevel()
  *				Save" option on the dialog, the return code will indicate the user has declined out of the prompt. This way calling code can distinguish between a decline and a cancel
  *				and then proceed as planned, or abort its operation accordingly.
  */
-FEditorFileUtils::EPromptReturnCode FEditorFileUtils::PromptForCheckoutAndSave( const TArray<UPackage*>& InPackages, bool bCheckDirty, bool bPromptToSave, TArray<UPackage*>* OutFailedPackages, bool bAlreadyCheckedOut )
+FEditorFileUtils::EPromptReturnCode FEditorFileUtils::PromptForCheckoutAndSave( const TArray<UPackage*>& InPackages, bool bCheckDirty, bool bPromptToSave, TArray<UPackage*>* OutFailedPackages, bool bAlreadyCheckedOut, bool bCanBeDeclined )
 {
 	// Check for re-entrance into this function
 	if ( bIsPromptingForCheckoutAndSave )
@@ -2774,7 +2777,10 @@ FEditorFileUtils::EPromptReturnCode FEditorFileUtils::PromptForCheckoutAndSave(
 		FPackagesDialogModule& PackagesDialogModule = FModuleManager::LoadModuleChecked<FPackagesDialogModule>( TEXT("PackagesDialog") );
 		PackagesDialogModule.CreatePackagesDialog(NSLOCTEXT("PackagesDialogModule", "PackagesDialogTitle", "Save Content"), NSLOCTEXT("PackagesDialogModule", "PackagesDialogMessage", "Select content to save."));
 		PackagesDialogModule.AddButton(DRT_Save, NSLOCTEXT("PackagesDialogModule", "SaveSelectedButton", "Save Selected"), NSLOCTEXT("PackagesDialogModule", "SaveSelectedButtonTip", "Attempt to save the selected content"));
-		PackagesDialogModule.AddButton(DRT_DontSave, NSLOCTEXT("PackagesDialogModule", "DontSaveSelectedButton", "Don't Save"), NSLOCTEXT("PackagesDialogModule", "DontSaveSelectedButtonTip", "Do not save any content"));
+		if (bCanBeDeclined)
+		{
+			PackagesDialogModule.AddButton(DRT_DontSave, NSLOCTEXT("PackagesDialogModule", "DontSaveSelectedButton", "Don't Save"), NSLOCTEXT("PackagesDialogModule", "DontSaveSelectedButtonTip", "Do not save any content"));
+		}
 		PackagesDialogModule.AddButton(DRT_Cancel, NSLOCTEXT("PackagesDialogModule", "CancelButton", "Cancel"), NSLOCTEXT("PackagesDialogModule", "CancelButtonTip", "Do not save any content and cancel the current operation"));
 
 		TArray<UPackage*> AddPackageItemsChecked;
diff --git a/Engine/Source/Editor/UnrealEd/Public/FileHelpers.h b/Engine/Source/Editor/UnrealEd/Public/FileHelpers.h
--- a/Engine/Source/Editor/UnrealEd/Public/FileHelpers.h
+++ b/Engine/Source/Editor/UnrealEd/Public/FileHelpers.h
@@ -172,10 +172,11 @@ public:
 	 * @param	bSaveContentPackages		true if we should save content packages. 
 	 * @param	bFastSave					true if we should do a fast save. (I.E dont prompt the user to save, dont prompt for checkout, and only save packages that are currently writable).  Note: Still prompts for SaveAs if a package needs a filename
 	 * @param	bNotifyNoPackagesSaved		true if a notification should be displayed when no packages need to be saved.
+	 * @param	bCanBeDeclined				true if the user prompt should contain a "Don't Save" button in addition to "Cancel", which won't result in a failure return code.
 	 * @param	bOutPackagesNeededSaving	when not NULL, will be set to true if there was any work to be done, and false otherwise.
 	 * @return								true on success, false on fail.
 	 */
-	UNREALED_API static bool SaveDirtyPackages(const bool bPromptUserToSave, const bool bSaveMapPackages, const bool bSaveContentPackages, const bool bFastSave = false, const bool bNotifyNoPackagesSaved = false, bool* bOutPackagesNeededSaving = NULL);
+	UNREALED_API static bool SaveDirtyPackages(const bool bPromptUserToSave, const bool bSaveMapPackages, const bool bSaveContentPackages, const bool bFastSave = false, const bool bNotifyNoPackagesSaved = false, const bool bCanBeDeclined = true, bool* bOutPackagesNeededSaving = NULL);
 
 	/**
 	* Looks at all currently loaded packages and saves them if their "bDirty" flag is set and they include specified clasees, optionally prompting the user to select which packages to save)
@@ -184,9 +185,10 @@ public:
 	* @param	bPromptUserToSave			true if we should prompt the user to save dirty packages we found. false to assume all dirty packages should be saved.  Regardless of this setting the user will be prompted for checkout(if needed) unless bFastSave is set
 	* @param	bFastSave					true if we should do a fast save. (I.E dont prompt the user to save, dont prompt for checkout, and only save packages that are currently writable).  Note: Still prompts for SaveAs if a package needs a filename
 	* @param	bNotifyNoPackagesSaved		true if a notification should be displayed when no packages need to be saved.
+	* @param	bCanBeDeclined				true if the user prompt should contain a "Don't Save" button in addition to "Cancel", which won't result in a failure return code.
 	* @return								true on success, false on fail.
 	*/
-	UNREALED_API static bool SaveDirtyContentPackages(TArray<UClass*>& SaveContentClasses, const bool bPromptUserToSave, const bool bFastSave = false, const bool bNotifyNoPackagesSaved = false);
+	UNREALED_API static bool SaveDirtyContentPackages(TArray<UClass*>& SaveContentClasses, const bool bPromptUserToSave, const bool bFastSave = false, const bool bNotifyNoPackagesSaved = false, const bool bCanBeDeclined = true);
 
 	/**
 	 * Appends array with all currently dirty world packages.
@@ -232,6 +234,7 @@ public:
 	 * @param		bPromptToSave				If true the user will be prompted with a list of packages to save, otherwise all passed in packages are saved
 	 * @param		OutFailedPackages			[out] If specified, will be filled in with all of the packages that failed to save successfully
 	 * @param		bAlreadyCheckedOut			If true, the user will not be prompted with the source control dialog
+	 * @param		bCanBeDeclined				If true, offer a "Don't Save" option in addition to "Cancel", which will not result in a cancellation return code.
 	 *
 	 * @return		An enum value signifying success, failure, user declined, or cancellation. If any packages at all failed to save during execution, the return code will be 
 	 *				failure, even if other packages successfully saved. If the user cancels at any point during any prompt, the return code will be cancellation, even though it
@@ -239,7 +242,7 @@ public:
 	 *				Save" option on the dialog, the return code will indicate the user has declined out of the prompt. This way calling code can distinguish between a decline and a cancel
 	 *				and then proceed as planned, or abort its operation accordingly.
 	 */
-	UNREALED_API static EPromptReturnCode PromptForCheckoutAndSave( const TArray<UPackage*>& PackagesToSave, bool bCheckDirty, bool bPromptToSave, TArray<UPackage*>* OutFailedPackages = NULL, bool bAlreadyCheckedOut = false );
+	UNREALED_API static EPromptReturnCode PromptForCheckoutAndSave( const TArray<UPackage*>& PackagesToSave, bool bCheckDirty, bool bPromptToSave, TArray<UPackage*>* OutFailedPackages = NULL, bool bAlreadyCheckedOut = false, bool bCanBeDeclined = true );
 
 	////////////////////////////////////////////////////////////////////////////
 	// Import/Export
