Let GC destroy sub-level actors when sub-level is removed/unloaded from world, instead of explicitly destroying actors one by one

diff --git a/Engine/Source/Editor/UnrealEd/Private/EditorLevelUtils.cpp b/Engine/Source/Editor/UnrealEd/Private/EditorLevelUtils.cpp
--- a/Engine/Source/Editor/UnrealEd/Private/EditorLevelUtils.cpp
+++ b/Engine/Source/Editor/UnrealEd/Private/EditorLevelUtils.cpp
@@ -489,9 +489,6 @@ namespace EditorLevelUtils
 		World->RemoveLevel(InLevel);
 		InLevel->ClearLevelComponents();
 
-		int32 NumFailedDestroyedAttempts = 0;
-		bool bDestroyedActor = false;
-
 		// remove all group actors from the world in the level we are removing
 		// otherwise, this will cause group actors to not be garbage collected
 		for (int32 GroupIndex = World->ActiveGroupActors.Num()-1; GroupIndex >= 0; --GroupIndex)
@@ -502,25 +499,24 @@ namespace EditorLevelUtils
 				World->ActiveGroupActors.RemoveAt(GroupIndex);
 			}
 		}
-
-		for(int32 ActorIndex = 0; ActorIndex < InLevel->Actors.Num(); ++ActorIndex)
+		
+		// Mark all model components as pending kill so GC deletes references to them.
+		for (UModelComponent* ModelComponent : InLevel->ModelComponents)
 		{
-			AActor* ActorToRemove = InLevel->Actors[ActorIndex];
-			if (ActorToRemove)
+			if (ModelComponent != nullptr)
 			{
-				bDestroyedActor = World->EditorDestroyActor(ActorToRemove, false);
-
-				// Keep track of how many actors were not destroyed because all actors need to be destroyed
-				if(!bDestroyedActor)
-				{
-					NumFailedDestroyedAttempts++;
-				}
+				ModelComponent->MarkPendingKill();
 			}
 		}
 
-		if(NumFailedDestroyedAttempts > 0)
+		// Mark all actors and their components as pending kill so GC will delete references to them.
+		for (AActor* Actor : InLevel->Actors)
 		{
-			UE_LOG(LogLevelTools, Log, TEXT("Failed to destroy %d actors after attempting to destroy level!"), NumFailedDestroyedAttempts);
+			if (Actor != nullptr)
+			{
+				Actor->MarkComponentsAsPendingKill();
+				Actor->MarkPendingKill();
+			}
 		}
 
 		World->MarkPackageDirty();
