Fixed up incorrect paths which was causing iterative cooking to fail to find cooked packages.

diff --git a/Engine/Source/Editor/UnrealEd/Private/CookOnTheFlyServer.cpp b/Engine/Source/Editor/UnrealEd/Private/CookOnTheFlyServer.cpp
--- a/Engine/Source/Editor/UnrealEd/Private/CookOnTheFlyServer.cpp
+++ b/Engine/Source/Editor/UnrealEd/Private/CookOnTheFlyServer.cpp
@@ -2203,6 +2203,7 @@ void UCookOnTheFlyServer::PopulateCookedPackagesFromDisk( const TArray<ITargetPl
 		{
 			FString CookedFilename = TimestampIt.Key();
 			FDateTime CookedTimestamp = TimestampIt.Value();
+			CookedFilename = FPaths::ConvertRelativePathToFull(CookedFilename);
 			FString StandardCookedFilename = CookedFilename.Replace(*SandboxDirectory, *(FPaths::GetRelativePathToRoot()));
 			FDateTime DependentTimestamp;
 
