Fixed name of root folder for starter content import.

diff --git a/Engine/Source/Editor/UnrealEd/Private/FeaturePack/FeaturePack.cpp b/Engine/Source/Editor/UnrealEd/Private/FeaturePack/FeaturePack.cpp
--- a/Engine/Source/Editor/UnrealEd/Private/FeaturePack/FeaturePack.cpp
+++ b/Engine/Source/Editor/UnrealEd/Private/FeaturePack/FeaturePack.cpp
@@ -70,7 +70,7 @@ void FFeaturePack::ParsePacks()
 			TArray<FString> EachImport;
 			FString FullPath = FPaths::FeaturePackDir() + EachPackData.PackSource;
 			EachImport.Add(FullPath);
-			EachPackData.ImportedObjects = AssetToolsModule.Get().ImportAssets(EachImport, EachPackData.PackName, nullptr, false);
+			EachPackData.ImportedObjects = AssetToolsModule.Get().ImportAssets(EachImport, TEXT("/Game"), nullptr, false);
 						
 			if (EachPackData.ImportedObjects.Num() == 0)
 			{			
