Fix variable shadowing

diff --git a/Engine/Source/Editor/UnrealEd/Private/Fbx/FbxSkeletalMeshImport.cpp b/Engine/Source/Editor/UnrealEd/Private/Fbx/FbxSkeletalMeshImport.cpp
--- a/Engine/Source/Editor/UnrealEd/Private/Fbx/FbxSkeletalMeshImport.cpp
+++ b/Engine/Source/Editor/UnrealEd/Private/Fbx/FbxSkeletalMeshImport.cpp
@@ -2195,8 +2195,12 @@ bool UnFbx::FFbxImporter::FillSkelMeshImporterFromFbx( FSkeletalMeshImportData&
 			{
 				// Set all UV's to zero.  If we are here the mesh had no UV sets so we only need to do this for the
 				// first UV set which always exists.
-				TmpWedges[VertexIndex].UVs[ UVLayerIndex ].X = 0.0f;
-				TmpWedges[VertexIndex].UVs[ UVLayerIndex ].Y = 0.0f;
+
+				for(VertexIndex=0; VertexIndex<3; VertexIndex++)
+				{
+					TmpWedges[VertexIndex].UVs[UVLayerIndex].X = 0.0f;
+					TmpWedges[VertexIndex].UVs[UVLayerIndex].Y = 0.0f;
+				}
 			}
 		}
 
