Close anim compression dialog once compression is complete.

diff --git a/Engine/Source/Editor/UnrealEd/Private/AnimationCompressionPanel.cpp b/Engine/Source/Editor/UnrealEd/Private/AnimationCompressionPanel.cpp
--- a/Engine/Source/Editor/UnrealEd/Private/AnimationCompressionPanel.cpp
+++ b/Engine/Source/Editor/UnrealEd/Private/AnimationCompressionPanel.cpp
@@ -25,6 +25,7 @@ FDlgAnimCompression::FDlgAnimCompression( TArray<TWeakObjectPtr<UAnimSequence>>&
 			[
 				SAssignNew(DialogWidget, SAnimationCompressionPanel)
 				.AnimSequences(AnimSequences)
+				.ParentWindow(DialogWindow)
 			];
 
 		DialogWindow->SetContent(DialogWrapper.ToSharedRef());
@@ -38,6 +39,8 @@ void FDlgAnimCompression::ShowModal()
 
 void SAnimationCompressionPanel::Construct(const FArguments& InArgs)
 {
+	ParentWindow = InArgs._ParentWindow.Get();
+
 	AnimSequences = InArgs._AnimSequences;
 	CurrentCompressionChoice = 0;
 
@@ -143,6 +146,7 @@ void SAnimationCompressionPanel::ApplyAlgorithm(class UAnimCompress* Algorithm)
 
 			GWarn->EndSlowTask( );
 
+			ParentWindow->RequestDestroyWindow();
 		}
 	}
 }
diff --git a/Engine/Source/Editor/UnrealEd/Public/AnimationCompressionPanel.h b/Engine/Source/Editor/UnrealEd/Public/AnimationCompressionPanel.h
--- a/Engine/Source/Editor/UnrealEd/Public/AnimationCompressionPanel.h
+++ b/Engine/Source/Editor/UnrealEd/Public/AnimationCompressionPanel.h
@@ -38,6 +38,9 @@ public:
 		/** The anim sequences to compress */
 		SLATE_ARGUMENT( TArray<TWeakObjectPtr<UAnimSequence>>, AnimSequences )
 
+		/** Window in which this widget resides */
+		SLATE_ATTRIBUTE(TSharedPtr<SWindow>, ParentWindow)
+
 	SLATE_END_ARGS()	
 
 	/**
@@ -112,4 +115,7 @@ private:
 	TArray< TWeakObjectPtr<UAnimSequence> >	AnimSequences;
 
 	int32 CurrentCompressionChoice;
+
+	/** Pointer to the window which holds this Widget, required for modal control */
+	TSharedPtr<SWindow> ParentWindow;
 };
\ No newline at end of file
