fix for UE-10021 Open Level window paths to a /Game/Maps folder, even when it doesn't exists

diff --git a/Engine/Source/Editor/UnrealEd/Private/FileHelpers.cpp b/Engine/Source/Editor/UnrealEd/Private/FileHelpers.cpp
--- a/Engine/Source/Editor/UnrealEd/Private/FileHelpers.cpp
+++ b/Engine/Source/Editor/UnrealEd/Private/FileHelpers.cpp
@@ -1638,7 +1638,14 @@ void FEditorFileUtils::OpenLevelPickingDialog(const FOnLevelsChosen& OnLevelsCho
 	// Determine the starting path. Try to use the most recently used directory
 	FString DefaultPath;
 	{
-		const FString DefaultFilesystemDirectory = FEditorDirectories::Get().GetLastDirectory(ELastDirectory::LEVEL);
+		FString DefaultFilesystemDirectory = FEditorDirectories::Get().GetLastDirectory(ELastDirectory::LEVEL);
+
+		//ensure trailing "/" for directory name since TryConvertFilenameToLongPackageName expects one
+		if(!DefaultFilesystemDirectory.IsEmpty() && !DefaultFilesystemDirectory.EndsWith("/"))
+		{
+			DefaultFilesystemDirectory.AppendChar(TEXT('/'));
+		}
+
 		if (DefaultFilesystemDirectory.IsEmpty() || !FPackageName::TryConvertFilenameToLongPackageName(DefaultFilesystemDirectory, DefaultPath))
 		{
 			// No saved path, just use a reasonable default
