UE-7921 In Simulate Mode, selecting and moving Movable actors using the gizmo doesn't work

diff --git a/Engine/Source/Editor/UnrealEd/Classes/Editor/EditorEngine.h b/Engine/Source/Editor/UnrealEd/Classes/Editor/EditorEngine.h
--- a/Engine/Source/Editor/UnrealEd/Classes/Editor/EditorEngine.h
+++ b/Engine/Source/Editor/UnrealEd/Classes/Editor/EditorEngine.h
@@ -2262,6 +2262,9 @@ public:
 	 */
 	FWorldContext &GetEditorWorldContext(bool bEnsureIsGWorld = false);
 
+	/** Returns the WorldContext for the PIE world.
+	*/
+	FWorldContext* GetPIEWorldContext();
 
 	/** 
 	 * mostly done to check if PIE is being set up, go GWorld is going to change, and it's not really _the_G_World_
diff --git a/Engine/Source/Editor/UnrealEd/Private/Editor.cpp b/Engine/Source/Editor/UnrealEd/Private/Editor.cpp
--- a/Engine/Source/Editor/UnrealEd/Private/Editor.cpp
+++ b/Engine/Source/Editor/UnrealEd/Private/Editor.cpp
@@ -6945,6 +6945,19 @@ FWorldContext& UEditorEngine::GetEditorWorldContext(bool bEnsureIsGWorld)
 	return CreateNewWorldContext(EWorldType::Editor);
 }
 
+FWorldContext* UEditorEngine::GetPIEWorldContext()
+{
+	for(auto& WorldContext : WorldList)
+	{
+		if(WorldContext.WorldType == EWorldType::PIE)
+		{
+			return &WorldContext;
+		}
+	}
+
+	return nullptr;
+}
+
 namespace EditorUtilities
 {
 	AActor* GetEditorWorldCounterpartActor( AActor* Actor )
diff --git a/Engine/Source/Editor/UnrealEd/Private/PhysicsManipulationMode.cpp b/Engine/Source/Editor/UnrealEd/Private/PhysicsManipulationMode.cpp
--- a/Engine/Source/Editor/UnrealEd/Private/PhysicsManipulationMode.cpp
+++ b/Engine/Source/Editor/UnrealEd/Private/PhysicsManipulationMode.cpp
@@ -51,7 +51,9 @@ FPhysicsManipulationEdMode::~FPhysicsManipulationEdMode()
 
 void FPhysicsManipulationEdMode::Enter()
 {
-	HandleComp->RegisterComponentWithWorld(GetWorld());
+	FWorldContext* PIEWorldContext = GEditor->GetPIEWorldContext();
+	check(PIEWorldContext && PIEWorldContext->World());
+	HandleComp->RegisterComponentWithWorld(PIEWorldContext->World());
 }
 
 void FPhysicsManipulationEdMode::Exit()
