Fix shadowed variables

diff --git a/Engine/Source/Editor/UMGEditor/Private/Designer/SDesignerView.cpp b/Engine/Source/Editor/UMGEditor/Private/Designer/SDesignerView.cpp
--- a/Engine/Source/Editor/UMGEditor/Private/Designer/SDesignerView.cpp
+++ b/Engine/Source/Editor/UMGEditor/Private/Designer/SDesignerView.cpp
@@ -785,17 +785,17 @@ bool SDesignerView::GetWidgetParentGeometry(const FWidgetReference& Widget, FGeo
 
 bool SDesignerView::GetWidgetGeometry(const FWidgetReference& Widget, FGeometry& Geometry) const
 {
-	if ( const UWidget* PreviewWidget = Widget.GetPreview() )
+	if ( const UWidget* WidgetPreview = Widget.GetPreview() )
 	{
-		return GetWidgetGeometry(PreviewWidget, Geometry);
+		return GetWidgetGeometry(WidgetPreview, Geometry);
 	}
 
 	return false;
 }
 
-bool SDesignerView::GetWidgetGeometry(const UWidget* PreviewWidget, FGeometry& Geometry) const
+bool SDesignerView::GetWidgetGeometry(const UWidget* InPreviewWidget, FGeometry& Geometry) const
 {
-	TSharedPtr<SWidget> CachedPreviewWidget = PreviewWidget->GetCachedWidget();
+	TSharedPtr<SWidget> CachedPreviewWidget = InPreviewWidget->GetCachedWidget();
 	if ( CachedPreviewWidget.IsValid() )
 	{
 		const FArrangedWidget* ArrangedWidget = CachedWidgetGeometry.Find(CachedPreviewWidget.ToSharedRef());
@@ -1389,9 +1389,11 @@ void SDesignerView::Tick( const FGeometry& AllottedGeometry, const double InCurr
 	UpdatePreviewWidget(bForceUpdate);
 
 	// Perform an arrange children pass to cache the geometry of all widgets so that we can query it later.
-	CachedWidgetGeometry.Reset();
-	FArrangedWidget WindowWidgetGeometry(PreviewHitTestRoot.ToSharedRef(), AllottedGeometry);
-	PopulateWidgetGeometryCache(WindowWidgetGeometry);
+	{
+		CachedWidgetGeometry.Reset();
+		FArrangedWidget WindowWidgetGeometry(PreviewHitTestRoot.ToSharedRef(), AllottedGeometry);
+		PopulateWidgetGeometryCache(WindowWidgetGeometry);
+	}
 
 	TArray< TFunction<void()> >& QueuedActions = BlueprintEditor.Pin()->GetQueuedDesignerActions();
 	for ( TFunction<void()>& Action : QueuedActions )
diff --git a/Engine/Source/Editor/UMGEditor/Private/Extensions/CanvasSlotExtension.cpp b/Engine/Source/Editor/UMGEditor/Private/Extensions/CanvasSlotExtension.cpp
--- a/Engine/Source/Editor/UMGEditor/Private/Extensions/CanvasSlotExtension.cpp
+++ b/Engine/Source/Editor/UMGEditor/Private/Extensions/CanvasSlotExtension.cpp
@@ -485,7 +485,7 @@ FReply FCanvasSlotExtension::HandleAnchorDragging(const FGeometry& Geometry, con
 	return FReply::Unhandled();
 }
 
-void FCanvasSlotExtension::PaintDragPercentages(const TSet< FWidgetReference >& Selection, const FGeometry& AllottedGeometry, const FSlateRect& MyClippingRect, FSlateWindowElementList& OutDrawElements, int32 LayerId) const
+void FCanvasSlotExtension::PaintDragPercentages(const TSet< FWidgetReference >& InSelection, const FGeometry& AllottedGeometry, const FSlateRect& MyClippingRect, FSlateWindowElementList& OutDrawElements, int32 LayerId) const
 {
 	// Just show the percentage lines when we're moving the anchor gizmo
 	if ( !(bMovingAnchor || bHoveringAnchor) )
diff --git a/Engine/Source/Editor/UnrealEd/Private/CookOnTheFlyServer.cpp b/Engine/Source/Editor/UnrealEd/Private/CookOnTheFlyServer.cpp
--- a/Engine/Source/Editor/UnrealEd/Private/CookOnTheFlyServer.cpp
+++ b/Engine/Source/Editor/UnrealEd/Private/CookOnTheFlyServer.cpp
@@ -354,7 +354,7 @@ bool UCookOnTheFlyServer::ClearPackageFilenameCacheForPackage( const UPackage* P
 	return PackageFilenameCache.Remove( Package->GetFName() ) >= 1;
 }
 
-const FString& UCookOnTheFlyServer::GetCachedSandboxFilename( const UPackage* Package, TAutoPtr<class FSandboxPlatformFile>& SandboxFile ) const 
+const FString& UCookOnTheFlyServer::GetCachedSandboxFilename( const UPackage* Package, TAutoPtr<class FSandboxPlatformFile>& InSandboxFile ) const 
 {
 	FName PackageFName = Package->GetFName();
 	static TMap<FName, FString> CachedSandboxFilenames;
@@ -766,12 +766,12 @@ void UCookOnTheFlyServer::GenerateManifestInfo( UPackage* Package, const TArray<
 
 		for ( const auto& PackageName : NewPackagesToCook)
 		{
-			UPackage* Package = LoadPackage( NULL, *PackageName, LOAD_None );
+			UPackage* PackageToCook = LoadPackage( NULL, *PackageName, LOAD_None );
 
-			RootPackages.Add( Package );
-			//GetDependencies( Package, Dependencies );
+			RootPackages.Add( PackageToCook );
+			//GetDependencies( PackageToCook, Dependencies );
 
-			// Dependencies.Add(Package);
+			// Dependencies.Add(PackageToCook);
 		}
 
 		LastLoadedMapName = Package->GetName();
@@ -994,11 +994,13 @@ uint32 UCookOnTheFlyServer::TickCookOnTheSide( const float TimeSlice, uint32 &Co
 		{
 			SCOPE_TIMER(AllOfLoadPackage);
 
-			FString PackageName;
 			UPackage *Package = NULL;
-			if ( FPackageName::TryConvertFilenameToLongPackageName(BuildFilename, PackageName) )
 			{
-				Package = FindObject<UPackage>( ANY_PACKAGE, *PackageName );
+				FString PackageName;
+				if ( FPackageName::TryConvertFilenameToLongPackageName(BuildFilename, PackageName) )
+				{
+					Package = FindObject<UPackage>( ANY_PACKAGE, *PackageName );
+				}
 			}
 
 #if DEBUG_COOKONTHEFLY
diff --git a/Engine/Source/Editor/UnrealEd/Private/Factories/EditorFactories.cpp b/Engine/Source/Editor/UnrealEd/Private/Factories/EditorFactories.cpp
--- a/Engine/Source/Editor/UnrealEd/Private/Factories/EditorFactories.cpp
+++ b/Engine/Source/Editor/UnrealEd/Private/Factories/EditorFactories.cpp
@@ -7181,15 +7181,15 @@ bool UDataTableFactory::ConfigureProperties()
 	public:
 		FDataTableFactoryUI() : ResultStruct(NULL) {}
 
-		TSharedRef<SWidget> MakeRowStructItemWidget(class UScriptStruct* Struct) const
+		TSharedRef<SWidget> MakeRowStructItemWidget(class UScriptStruct* InStruct) const
 		{
-			return SNew(STextBlock).Text(Struct ? Struct->GetDisplayNameText() : FText::GetEmpty());
+			return SNew(STextBlock).Text(InStruct ? InStruct->GetDisplayNameText() : FText::GetEmpty());
 		}
 
 		FText GetSelectedRowOptionText() const
 		{
-			UScriptStruct* Struct = RowStructCombo.IsValid() ? RowStructCombo->GetSelectedItem() : NULL;
-			return Struct ? Struct->GetDisplayNameText() : FText::GetEmpty();
+			UScriptStruct* RowStruct = RowStructCombo.IsValid() ? RowStructCombo->GetSelectedItem() : NULL;
+			return RowStruct ? RowStruct->GetDisplayNameText() : FText::GetEmpty();
 		}
 
 		FReply OnCreate()
diff --git a/Engine/Source/Editor/UnrealEd/Private/Kismet2/BlueprintEditorUtils.cpp b/Engine/Source/Editor/UnrealEd/Private/Kismet2/BlueprintEditorUtils.cpp
--- a/Engine/Source/Editor/UnrealEd/Private/Kismet2/BlueprintEditorUtils.cpp
+++ b/Engine/Source/Editor/UnrealEd/Private/Kismet2/BlueprintEditorUtils.cpp
@@ -2263,7 +2263,7 @@ UK2Node_Event* FBlueprintEditorUtils::FindOverrideForFunction(const UBlueprint*
 	return NULL;
 }
 
-void FBlueprintEditorUtils::GatherDependencies(const UBlueprint* Blueprint, TSet<TWeakObjectPtr<UBlueprint>>& Dependencies)
+void FBlueprintEditorUtils::GatherDependencies(const UBlueprint* InBlueprint, TSet<TWeakObjectPtr<UBlueprint>>& Dependencies)
 {
 	struct FGatherDependenciesHelper
 	{
@@ -2279,14 +2279,14 @@ void FBlueprintEditorUtils::GatherDependencies(const UBlueprint* Blueprint, TSet
 			return UBlueprint::GetBlueprintFromClass(BPGC);
 		}
 
-		static void ProcessHierarchy(const UStruct* Struct, TSet<TWeakObjectPtr<UBlueprint>>& Dependencies)
+		static void ProcessHierarchy(const UStruct* Struct, TSet<TWeakObjectPtr<UBlueprint>>& InDependencies)
 		{
 			for (UBlueprint* Blueprint = GetGeneratingBlueprint(Struct);
 				Blueprint;
 				Blueprint = UBlueprint::GetBlueprintFromClass(Cast<UBlueprintGeneratedClass>(Blueprint->ParentClass)))
 			{
 				bool bAlreadyProcessed = false;
-				Dependencies.Add(Blueprint, &bAlreadyProcessed);
+				InDependencies.Add(Blueprint, &bAlreadyProcessed);
 				if (bAlreadyProcessed)
 				{
 					return;
@@ -2295,9 +2295,9 @@ void FBlueprintEditorUtils::GatherDependencies(const UBlueprint* Blueprint, TSet
 		}
 	};
 
-	check(Blueprint);
+	check(InBlueprint);
 	Dependencies.Empty();
-	for (const auto& InterfaceDesc : Blueprint->ImplementedInterfaces)
+	for (const auto& InterfaceDesc : InBlueprint->ImplementedInterfaces)
 	{
 		UBlueprint* InterfaceBP = InterfaceDesc.Interface ? Cast<UBlueprint>(InterfaceDesc.Interface->ClassGeneratedBy) : NULL;
 		if (InterfaceBP)
@@ -2307,7 +2307,7 @@ void FBlueprintEditorUtils::GatherDependencies(const UBlueprint* Blueprint, TSet
 	}
 
 	TArray<UEdGraph*> Graphs;
-	Blueprint->GetAllGraphs(Graphs);
+	InBlueprint->GetAllGraphs(Graphs);
 	for (auto Graph : Graphs)
 	{
 		if (Graph && !FBlueprintEditorUtils::IsGraphIntermediate(Graph))
@@ -7119,11 +7119,11 @@ bool FBlueprintEditorUtils::CheckIfGraphHasLatentFunctions(UEdGraph* InGraph)
 {
 	struct Local
 	{
-		static bool CheckIfGraphHasLatentFunctions(UEdGraph* InGraph, TArray<UEdGraph*>& InspectedGraphList)
+		static bool CheckIfGraphHasLatentFunctions(UEdGraph* InGraphToCheck, TArray<UEdGraph*>& InspectedGraphList)
 		{
 			TWeakObjectPtr<UK2Node_EditablePinBase> EntryNode;
 			TWeakObjectPtr<UK2Node_EditablePinBase> ResultNode;
-			GetEntryAndResultNodes(InGraph, EntryNode, ResultNode);
+			GetEntryAndResultNodes(InGraphToCheck, EntryNode, ResultNode);
 
 			UK2Node_Tunnel* TunnelNode = ExactCast<UK2Node_Tunnel>(EntryNode.Get());
 			if(!TunnelNode)
@@ -7139,9 +7139,9 @@ bool FBlueprintEditorUtils::CheckIfGraphHasLatentFunctions(UEdGraph* InGraph)
 			else
 			{
 				// Add all graphs to the list of already inspected, this prevents circular inclusion issues.
-				InspectedGraphList.Add(InGraph);
+				InspectedGraphList.Add(InGraphToCheck);
 
-				for( const UEdGraphNode* Node : InGraph->Nodes )
+				for( const UEdGraphNode* Node : InGraphToCheck->Nodes )
 				{
 					if(const UK2Node_CallFunction* CallFunctionNode = Cast<UK2Node_CallFunction>(Node))
 					{
diff --git a/Engine/Source/Editor/UnrealEd/Private/SCurveEditor.cpp b/Engine/Source/Editor/UnrealEd/Private/SCurveEditor.cpp
--- a/Engine/Source/Editor/UnrealEd/Private/SCurveEditor.cpp
+++ b/Engine/Source/Editor/UnrealEd/Private/SCurveEditor.cpp
@@ -1752,7 +1752,7 @@ SCurveEditor::FSelectedCurveKey SCurveEditor::HitTestKeys(const FGeometry& InMyG
 	return SelectedKey;
 }
 
-void SCurveEditor::MoveSelectedKeys(FVector2D NewLocation)
+void SCurveEditor::MoveSelectedKeys(FVector2D InNewLocation)
 {
 	const FScopedTransaction Transaction( LOCTEXT("CurveEditor_MoveKeys", "Move Keys") );
 	CurveOwner->ModifyOwner();
@@ -1760,7 +1760,7 @@ void SCurveEditor::MoveSelectedKeys(FVector2D NewLocation)
 	// track all unique curves encountered so their tangents can be updated later
 	TSet<FRichCurve*> UniqueCurves;
 
-	FVector2D SnappedNewLocation = SnapLocation(NewLocation);
+	FVector2D SnappedNewLocation = SnapLocation(InNewLocation);
 
 	// The total move distance for all keys is the difference between the current snapped location
 	// and the start location of the key which was actually dragged.
diff --git a/Engine/Source/Editor/UnrealEd/Private/UnrealEdSrv.cpp b/Engine/Source/Editor/UnrealEd/Private/UnrealEdSrv.cpp
--- a/Engine/Source/Editor/UnrealEd/Private/UnrealEdSrv.cpp
+++ b/Engine/Source/Editor/UnrealEd/Private/UnrealEdSrv.cpp
@@ -1453,9 +1453,6 @@ void UUnrealEdEngine::AttemptModifiedPackageNotification()
 
 		if( bCanPrompt )
 		{
-			check(GUnrealEd != nullptr);
-			const TMap<TWeakObjectPtr<UPackage>, uint8>& PackageToNotifyState = GUnrealEd->PackageToNotifyState;
-
 			bool bNeedWarningDialog = false;
 			for (const auto& Entry : PackageToNotifyState)
 			{
diff --git a/Engine/Source/Editor/UnrealEd/Public/Kismet2/CompilerResultsLog.h b/Engine/Source/Editor/UnrealEd/Public/Kismet2/CompilerResultsLog.h
--- a/Engine/Source/Editor/UnrealEd/Public/Kismet2/CompilerResultsLog.h
+++ b/Engine/Source/Editor/UnrealEd/Public/Kismet2/CompilerResultsLog.h
@@ -224,12 +224,14 @@ public:
 	}
 };
 
+#define SCOPED_COMPILER_EVENT_TOKENPASTE_INNER(x,y) x##y
+#define SCOPED_COMPILER_EVENT_TOKENPASTE(x,y) SCOPED_COMPILER_EVENT_TOKENPASTE_INNER(x,y)
 #if STATS
 #define BP_SCOPED_COMPILER_EVENT_STAT(Stat) \
 	SCOPE_CYCLE_COUNTER(Stat); \
-	FScopedCompilerEvent ScopedCompilerEvent(GET_STATDESCRIPTION(Stat))
+	FScopedCompilerEvent SCOPED_COMPILER_EVENT_TOKENPASTE(ScopedCompilerEvent,__LINE__)(GET_STATDESCRIPTION(Stat))
 #else
 #define BP_SCOPED_COMPILER_EVENT_STAT(Stat) \
-	FScopedCompilerEvent ScopedCompilerEvent(ANSI_TO_TCHAR(#Stat))
+	FScopedCompilerEvent SCOPED_COMPILER_EVENT_TOKENPASTE(ScopedCompilerEvent,__LINE__)(ANSI_TO_TCHAR(#Stat))
 #endif
 #endif	//#if WITH_EDITOR
