Add a setting to open new tabs in a window or docked. Defaults to new window

diff --git a/Engine/Source/Editor/EditorStyle/Private/EditorStyleClasses.cpp b/Engine/Source/Editor/EditorStyle/Private/EditorStyleClasses.cpp
--- a/Engine/Source/Editor/EditorStyle/Private/EditorStyleClasses.cpp
+++ b/Engine/Source/Editor/EditorStyle/Private/EditorStyleClasses.cpp
@@ -13,6 +13,8 @@ UEditorStyleSettings::UEditorStyleSettings( const FObjectInitializer& ObjectInit
 	InactiveSelectionColor = FLinearColor(0.25f, 0.25f, 0.25f);
 	PressedSelectionColor = FLinearColor(0.701f, 0.225f, 0.003f);
 
+	bOpenTabsInNewWindow = true;
+
 	bDisplayPropertyUnits = true;
 	bShowFriendlyNames = true;
 	LogTimestampMode = ELogTimes::None;
diff --git a/Engine/Source/Editor/EditorStyle/Public/Classes/EditorStyleSettings.h b/Engine/Source/Editor/EditorStyle/Public/Classes/EditorStyleSettings.h
--- a/Engine/Source/Editor/EditorStyle/Public/Classes/EditorStyleSettings.h
+++ b/Engine/Source/Editor/EditorStyle/Public/Classes/EditorStyleSettings.h
@@ -83,9 +83,13 @@ public:
 	uint32 bShowProjectMenus : 1;
 
 	/** The display mode for timestamps in the output log */
-	UPROPERTY(EditAnywhere, config, Category=UserInterface)
+	UPROPERTY(EditAnywhere, config, AdvancedDisplay, Category=UserInterface)
 	TEnumAsByte<ELogTimes::Type> LogTimestampMode;
 
+	/** If checked, new asset editor tabs will open in a new window instead of docked in the tab from which they were opened */
+	UPROPERTY(EditAnywhere, config, Category=UserInterface)
+	bool bOpenAssetEditorTabsInNewWindow;
+
 public:
 
 	/**
diff --git a/Engine/Source/Editor/UnrealEd/Private/Toolkits/AssetEditorToolkit.cpp b/Engine/Source/Editor/UnrealEd/Private/Toolkits/AssetEditorToolkit.cpp
--- a/Engine/Source/Editor/UnrealEd/Private/Toolkits/AssetEditorToolkit.cpp
+++ b/Engine/Source/Editor/UnrealEd/Private/Toolkits/AssetEditorToolkit.cpp
@@ -92,8 +92,10 @@ void FAssetEditorToolkit::InitAssetEditor( const EToolkitMode::Type Mode, const
 		{
 			static_assert(sizeof(EAssetEditorToolkitTabLocation) == sizeof(int32), "EAssetEditorToolkitTabLocation is the incorrect size");
 
+			const UEditorStyleSettings* StyleSettings = GetDefault<UEditorStyleSettings>();
+
 			// Work out where we should create this asset editor
-			EAssetEditorToolkitTabLocation SavedAssetEditorToolkitTabLocation = EAssetEditorToolkitTabLocation::Docked;
+			EAssetEditorToolkitTabLocation SavedAssetEditorToolkitTabLocation = StyleSettings->bOpenTabsInNewWindow ? EAssetEditorToolkitTabLocation::Standalone : EAssetEditorToolkitTabLocation::Docked;
 			GConfig->GetInt(
 				TEXT("AssetEditorToolkitTabLocation"), 
 				*ObjectsToEdit[0]->GetPathName(), 
