Fixed splitting of source text files into line for text localization gathering to work with different newlines.

diff --git a/Engine/Source/Editor/UnrealEd/Private/Commandlets/GatherTextFromSourceCommandlet.cpp b/Engine/Source/Editor/UnrealEd/Private/Commandlets/GatherTextFromSourceCommandlet.cpp
--- a/Engine/Source/Editor/UnrealEd/Private/Commandlets/GatherTextFromSourceCommandlet.cpp
+++ b/Engine/Source/Editor/UnrealEd/Private/Commandlets/GatherTextFromSourceCommandlet.cpp
@@ -296,7 +296,7 @@ bool UGatherTextFromSourceCommandlet::ParseSourceText(const FString& Text, const
 
 	// Split the file into lines of 
 	TArray<FString> TextLines;
-	Text.ParseIntoArray(&TextLines, LINE_TERMINATOR, false);
+	Text.ParseIntoArrayLines(TextLines, false);
 
 	// Move through the text lines looking for the tokens that denote the items in the Parsables list
 	for (int32 LineIdx = 0; LineIdx < TextLines.Num(); LineIdx++)
