Variable name compile fix.

diff --git a/Engine/Source/Editor/EditorStyle/Private/EditorStyleClasses.cpp b/Engine/Source/Editor/EditorStyle/Private/EditorStyleClasses.cpp
--- a/Engine/Source/Editor/EditorStyle/Private/EditorStyleClasses.cpp
+++ b/Engine/Source/Editor/EditorStyle/Private/EditorStyleClasses.cpp
@@ -13,7 +13,7 @@ UEditorStyleSettings::UEditorStyleSettings( const FObjectInitializer& ObjectInit
 	InactiveSelectionColor = FLinearColor(0.25f, 0.25f, 0.25f);
 	PressedSelectionColor = FLinearColor(0.701f, 0.225f, 0.003f);
 
-	bOpenTabsInNewWindow = true;
+	bOpenAssetEditorTabsInNewWindow = true;
 
 	bDisplayPropertyUnits = true;
 	bShowFriendlyNames = true;
diff --git a/Engine/Source/Editor/UnrealEd/Private/Toolkits/AssetEditorToolkit.cpp b/Engine/Source/Editor/UnrealEd/Private/Toolkits/AssetEditorToolkit.cpp
--- a/Engine/Source/Editor/UnrealEd/Private/Toolkits/AssetEditorToolkit.cpp
+++ b/Engine/Source/Editor/UnrealEd/Private/Toolkits/AssetEditorToolkit.cpp
@@ -95,7 +95,7 @@ void FAssetEditorToolkit::InitAssetEditor( const EToolkitMode::Type Mode, const
 			const UEditorStyleSettings* StyleSettings = GetDefault<UEditorStyleSettings>();
 
 			// Work out where we should create this asset editor
-			EAssetEditorToolkitTabLocation SavedAssetEditorToolkitTabLocation = StyleSettings->bOpenTabsInNewWindow ? EAssetEditorToolkitTabLocation::Standalone : EAssetEditorToolkitTabLocation::Docked;
+			EAssetEditorToolkitTabLocation SavedAssetEditorToolkitTabLocation = StyleSettings->bOpenAssetEditorTabsInNewWindow ? EAssetEditorToolkitTabLocation::Standalone : EAssetEditorToolkitTabLocation::Docked;
 			GConfig->GetInt(
 				TEXT("AssetEditorToolkitTabLocation"), 
 				*ObjectsToEdit[0]->GetPathName(), 
