Fix shadowed variables

diff --git a/Engine/Source/Developer/AssetTools/Private/AssetRenameManager.cpp b/Engine/Source/Developer/AssetTools/Private/AssetRenameManager.cpp
--- a/Engine/Source/Developer/AssetTools/Private/AssetRenameManager.cpp
+++ b/Engine/Source/Developer/AssetTools/Private/AssetRenameManager.cpp
@@ -508,8 +508,8 @@ void FAssetRenameManager::RenameReferencingStringAssetReferences(const TArray<UP
 {
 	struct FStringAssetReferenceRenameSerializer : public FArchiveUObject
 	{
-		FStringAssetReferenceRenameSerializer(const FString& OldAssetPath, const FString& NewAssetPath)
-			: OldAssetPath(OldAssetPath), NewAssetPath(NewAssetPath)
+		FStringAssetReferenceRenameSerializer(const FString& InOldAssetPath, const FString& InNewAssetPath)
+			: OldAssetPath(InOldAssetPath), NewAssetPath(InNewAssetPath)
 		{ }
 
 		FArchive& operator<<(FStringAssetReference& Reference) override
diff --git a/Engine/Source/Developer/Profiler/Private/ProfilerManager.cpp b/Engine/Source/Developer/Profiler/Private/ProfilerManager.cpp
--- a/Engine/Source/Developer/Profiler/Private/ProfilerManager.cpp
+++ b/Engine/Source/Developer/Profiler/Private/ProfilerManager.cpp
@@ -439,10 +439,10 @@ void FProfilerManager::SessionManager_OnInstanceSelectionChanged()
 
 				ProfilerSessionInstances.Add( ProfilerSession->GetInstanceID(), ProfilerSession );
 				ProfilerClient->Track( ProfilerInstanceID );
-				TSharedPtr<SProfilerWindow> ProfilerWindow = GetProfilerWindow();
-				if (ProfilerWindow.IsValid())
+				TSharedPtr<SProfilerWindow> ProfilerWindowPtr = GetProfilerWindow();
+				if (ProfilerWindowPtr.IsValid())
 				{
-					ProfilerWindow->ManageEventGraphTab(ProfilerInstanceID, true, ProfilerSession->GetName());
+					ProfilerWindowPtr->ManageEventGraphTab(ProfilerInstanceID, true, ProfilerSession->GetName());
 				}
 			}
 		}
diff --git a/Engine/Source/Editor/BlueprintGraph/Private/BlueprintActionDatabase.cpp b/Engine/Source/Editor/BlueprintGraph/Private/BlueprintActionDatabase.cpp
--- a/Engine/Source/Editor/BlueprintGraph/Private/BlueprintActionDatabase.cpp
+++ b/Engine/Source/Editor/BlueprintGraph/Private/BlueprintActionDatabase.cpp
@@ -134,12 +134,12 @@ static UBlueprintNodeSpawner* FBlueprintNodeSpawnerFactory::MakeMacroNodeSpawner
 	UBlueprintNodeSpawner* NodeSpawner = UBlueprintNodeSpawner::Create(UK2Node_MacroInstance::StaticClass());
 	check(NodeSpawner != nullptr);
 
-	auto CustomizeMacroNodeLambda = [](UEdGraphNode* NewNode, bool bIsTemplateNode, TWeakObjectPtr<UEdGraph> MacroGraph)
+	auto CustomizeMacroNodeLambda = [](UEdGraphNode* NewNode, bool bIsTemplateNode, TWeakObjectPtr<UEdGraph> InMacroGraph)
 	{
 		UK2Node_MacroInstance* MacroNode = CastChecked<UK2Node_MacroInstance>(NewNode);
-		if (MacroGraph.IsValid())
+		if (InMacroGraph.IsValid())
 		{
-			MacroNode->SetMacroGraph(MacroGraph.Get());
+			MacroNode->SetMacroGraph(InMacroGraph.Get());
 		}
 	};
 
diff --git a/Engine/Source/Editor/BlueprintGraph/Private/BlueprintActionFilter.cpp b/Engine/Source/Editor/BlueprintGraph/Private/BlueprintActionFilter.cpp
--- a/Engine/Source/Editor/BlueprintGraph/Private/BlueprintActionFilter.cpp
+++ b/Engine/Source/Editor/BlueprintGraph/Private/BlueprintActionFilter.cpp
@@ -749,18 +749,18 @@ static bool BlueprintActionFilterImpl::IsFieldCategoryHidden(FBlueprintActionFil
 
 	if (UFunction const* Function = BlueprintAction.GetAssociatedFunction())
 	{
-		auto IsFunctionHiddenLambda = [](UClass* Class, UFunction const* Function)->bool
+		auto IsFunctionHiddenLambda = [](UClass* Class, UFunction const* InFunction)->bool
 		{
 			// Only hide functions that are not static
-			return !Function->HasAnyFunctionFlags(FUNC_Static) && FObjectEditorUtils::IsFunctionHiddenFromClass(Function, Class->GetAuthoritativeClass());
+			return !InFunction->HasAnyFunctionFlags(FUNC_Static) && FObjectEditorUtils::IsFunctionHiddenFromClass(InFunction, Class->GetAuthoritativeClass());
 		};
 		IsFieldHiddenDelegate = FIsFieldHiddenDelegate::CreateStatic(IsFunctionHiddenLambda, Function);
 	}
 	else if (UProperty const* Property = BlueprintAction.GetAssociatedProperty())
 	{
-		auto IsPropertyHiddenLambda = [](UClass* Class, UProperty const* Property)->bool
+		auto IsPropertyHiddenLambda = [](UClass* Class, UProperty const* InProperty)->bool
 		{
-			return FObjectEditorUtils::IsVariableCategoryHiddenFromClass(Property, Class->GetAuthoritativeClass());
+			return FObjectEditorUtils::IsVariableCategoryHiddenFromClass(InProperty, Class->GetAuthoritativeClass());
 		};
 		IsFieldHiddenDelegate = FIsFieldHiddenDelegate::CreateStatic(IsPropertyHiddenLambda, Property);
 	}
diff --git a/Engine/Source/Editor/BlueprintGraph/Private/BlueprintBoundEventNodeSpawner.cpp b/Engine/Source/Editor/BlueprintGraph/Private/BlueprintBoundEventNodeSpawner.cpp
--- a/Engine/Source/Editor/BlueprintGraph/Private/BlueprintBoundEventNodeSpawner.cpp
+++ b/Engine/Source/Editor/BlueprintGraph/Private/BlueprintBoundEventNodeSpawner.cpp
@@ -145,20 +145,20 @@ bool UBlueprintBoundEventNodeSpawner::IsBindingCompatible(UObject const* Binding
 bool UBlueprintBoundEventNodeSpawner::BindToNode(UEdGraphNode* Node, UObject* Binding) const
 {
 	bool bWasBound = false;
-	if (UK2Node_ComponentBoundEvent* EventNode = Cast<UK2Node_ComponentBoundEvent>(Node))
+	if (UK2Node_ComponentBoundEvent* ComponentEventNode = Cast<UK2Node_ComponentBoundEvent>(Node))
 	{
 		if (UObjectProperty const* BoundProperty = Cast<UObjectProperty>(Binding))
 		{
-			EventNode->InitializeComponentBoundEventParams(BoundProperty, EventDelegate);
+			ComponentEventNode->InitializeComponentBoundEventParams(BoundProperty, EventDelegate);
 			bWasBound = true;
 			Node->ReconstructNode();
 		}
 	}
-	else if (UK2Node_ActorBoundEvent* EventNode = CastChecked<UK2Node_ActorBoundEvent>(Node))
+	else if (UK2Node_ActorBoundEvent* ActorEventNode = CastChecked<UK2Node_ActorBoundEvent>(Node))
 	{
 		if (AActor* BoundActor = Cast<AActor>(Binding))
 		{
-			EventNode->InitializeActorBoundEventParams(BoundActor, EventDelegate);
+			ActorEventNode->InitializeActorBoundEventParams(BoundActor, EventDelegate);
 			bWasBound = true;
 			Node->ReconstructNode();
 		}
diff --git a/Engine/Source/Editor/BlueprintGraph/Private/BlueprintDelegateNodeSpawner.cpp b/Engine/Source/Editor/BlueprintGraph/Private/BlueprintDelegateNodeSpawner.cpp
--- a/Engine/Source/Editor/BlueprintGraph/Private/BlueprintDelegateNodeSpawner.cpp
+++ b/Engine/Source/Editor/BlueprintGraph/Private/BlueprintDelegateNodeSpawner.cpp
@@ -87,16 +87,16 @@ UBlueprintDelegateNodeSpawner* UBlueprintDelegateNodeSpawner::Create(TSubclassOf
 
 	auto SetDelegateLambda = [](UEdGraphNode* NewNode, UField const* Field)
 	{
-		UMulticastDelegateProperty const* Property = Cast<UMulticastDelegateProperty>(Field);
+		UMulticastDelegateProperty const* MCDProperty = Cast<UMulticastDelegateProperty>(Field);
 
 		UK2Node_BaseMCDelegate* DelegateNode = Cast<UK2Node_BaseMCDelegate>(NewNode);
-		if ((DelegateNode != nullptr) && (Property != nullptr))
+		if ((DelegateNode != nullptr) && (MCDProperty != nullptr))
 		{
 			UBlueprint* Blueprint = FBlueprintEditorUtils::FindBlueprintForNodeChecked(NewNode);
-			UClass* OwnerClass = Property->GetOwnerClass();
+			UClass* OwnerClass = MCDProperty->GetOwnerClass();
 			bool const bIsSelfContext = Blueprint->SkeletonGeneratedClass->IsChildOf(OwnerClass);
 
-			DelegateNode->SetFromProperty(Property, bIsSelfContext);
+			DelegateNode->SetFromProperty(MCDProperty, bIsSelfContext);
 		}
 	};
 	NodeSpawner->SetNodeFieldDelegate = FSetNodeFieldDelegate::CreateStatic(SetDelegateLambda);
diff --git a/Engine/Source/Editor/BlueprintGraph/Private/BlueprintEventNodeSpawner.cpp b/Engine/Source/Editor/BlueprintGraph/Private/BlueprintEventNodeSpawner.cpp
--- a/Engine/Source/Editor/BlueprintGraph/Private/BlueprintEventNodeSpawner.cpp
+++ b/Engine/Source/Editor/BlueprintGraph/Private/BlueprintEventNodeSpawner.cpp
@@ -191,17 +191,17 @@ UEdGraphNode* UBlueprintEventNodeSpawner::Invoke(UEdGraph* ParentGraph, FBinding
 	// if there is no existing node, then we can happily spawn one into the graph
 	if (EventNode == nullptr)
 	{
-		auto PostSpawnLambda = [](UEdGraphNode* NewNode, bool bIsTemplateNode, UFunction const* EventFunc, FName EventName, FCustomizeNodeDelegate UserDelegate)
+		auto PostSpawnLambda = [](UEdGraphNode* NewNode, bool bIsTemplateNode, UFunction const* InEventFunc, FName InEventName, FCustomizeNodeDelegate UserDelegate)
 		{
-			UK2Node_Event* EventNode = CastChecked<UK2Node_Event>(NewNode);
-			if (EventFunc != nullptr)
+			UK2Node_Event* K2EventNode = CastChecked<UK2Node_Event>(NewNode);
+			if (InEventFunc != nullptr)
 			{
-				EventNode->EventReference.SetFromField<UFunction>(EventFunc, false);
-				EventNode->bOverrideFunction   = true;
+				K2EventNode->EventReference.SetFromField<UFunction>(InEventFunc, false);
+				K2EventNode->bOverrideFunction   = true;
 			}
 			else if (!bIsTemplateNode)
 			{
-				EventNode->CustomFunctionName = EventName;
+				K2EventNode->CustomFunctionName = InEventName;
 			}
 
 			UserDelegate.ExecuteIfBound(NewNode, bIsTemplateNode);
diff --git a/Engine/Source/Editor/BlueprintGraph/Private/BlueprintFieldNodeSpawner.cpp b/Engine/Source/Editor/BlueprintGraph/Private/BlueprintFieldNodeSpawner.cpp
--- a/Engine/Source/Editor/BlueprintGraph/Private/BlueprintFieldNodeSpawner.cpp
+++ b/Engine/Source/Editor/BlueprintGraph/Private/BlueprintFieldNodeSpawner.cpp
@@ -38,9 +38,9 @@ FBlueprintNodeSignature UBlueprintFieldNodeSpawner::GetSpawnerSignature() const
 //------------------------------------------------------------------------------
 UEdGraphNode* UBlueprintFieldNodeSpawner::Invoke(UEdGraph* ParentGraph, FBindingSet const& Bindings, FVector2D const Location) const
 {
-	auto PostSpawnSetupLambda = [](UEdGraphNode* NewNode, bool bIsTemplateNode, UField const* Field, FSetNodeFieldDelegate SetFieldDelegate, FCustomizeNodeDelegate UserDelegate)
+	auto PostSpawnSetupLambda = [](UEdGraphNode* NewNode, bool bIsTemplateNode, UField const* InField, FSetNodeFieldDelegate SetFieldDelegate, FCustomizeNodeDelegate UserDelegate)
 	{
-		SetFieldDelegate.ExecuteIfBound(NewNode, Field);
+		SetFieldDelegate.ExecuteIfBound(NewNode, InField);
 		UserDelegate.ExecuteIfBound(NewNode, bIsTemplateNode);
 	};
 
diff --git a/Engine/Source/Editor/BlueprintGraph/Private/BlueprintFunctionNodeSpawner.cpp b/Engine/Source/Editor/BlueprintGraph/Private/BlueprintFunctionNodeSpawner.cpp
--- a/Engine/Source/Editor/BlueprintGraph/Private/BlueprintFunctionNodeSpawner.cpp
+++ b/Engine/Source/Editor/BlueprintGraph/Private/BlueprintFunctionNodeSpawner.cpp
@@ -95,9 +95,9 @@ static bool BlueprintFunctionNodeSpawnerImpl::BindFunctionNode(UK2Node_CallFunct
 		}
 		else if (AActor* BoundActor = Cast<AActor>(BoundObject))
 		{
-			auto PostSpawnSetupLambda = [](UEdGraphNode* NewNode, bool /*bIsTemplateNode*/, AActor* ActorInst)
+			auto PostSpawnSetupLambda = [](UEdGraphNode* InNewNode, bool /*bIsTemplateNode*/, AActor* ActorInst)
 			{
-				UK2Node_Literal* ActorRefNode = CastChecked<UK2Node_Literal>(NewNode);
+				UK2Node_Literal* ActorRefNode = CastChecked<UK2Node_Literal>(InNewNode);
 				ActorRefNode->SetObjectRef(ActorInst);
 			};
 			UBlueprintNodeSpawner::FCustomizeNodeDelegate PostSpawnDelegate = UBlueprintNodeSpawner::FCustomizeNodeDelegate::CreateStatic(PostSpawnSetupLambda, BoundActor);
diff --git a/Engine/Source/Editor/BlueprintGraph/Private/BlueprintVariableNodeSpawner.cpp b/Engine/Source/Editor/BlueprintGraph/Private/BlueprintVariableNodeSpawner.cpp
--- a/Engine/Source/Editor/BlueprintGraph/Private/BlueprintVariableNodeSpawner.cpp
+++ b/Engine/Source/Editor/BlueprintGraph/Private/BlueprintVariableNodeSpawner.cpp
@@ -203,17 +203,17 @@ UEdGraphNode* UBlueprintVariableNodeSpawner::Invoke(UEdGraph* ParentGraph, FBind
 	//        conform to UBlueprintFieldNodeSpawner
 	if (IsLocalVariable())
 	{
-		auto LocalVarSetupLambda = [](UEdGraphNode* NewNode, bool bIsTemplateNode, FName VarName, UObject const* VarOuter, FGuid VarGuid, FCustomizeNodeDelegate UserDelegate)
+		auto LocalVarSetupLambda = [](UEdGraphNode* InNewNode, bool bIsTemplateNode, FName VarName, UObject const* VarOuter, FGuid VarGuid, FCustomizeNodeDelegate UserDelegate)
 		{
-			UK2Node_Variable* VarNode = CastChecked<UK2Node_Variable>(NewNode);
+			UK2Node_Variable* VarNode = CastChecked<UK2Node_Variable>(InNewNode);
 			VarNode->VariableReference.SetLocalMember(VarName, VarOuter->GetName(), VarGuid);
-			UserDelegate.ExecuteIfBound(NewNode, bIsTemplateNode);
+			UserDelegate.ExecuteIfBound(InNewNode, bIsTemplateNode);
 		};
 
 		FCustomizeNodeDelegate PostSpawnDelegate = CustomizeNodeDelegate;
-		if (UObject const* LocalVarOuter = GetVarOuter())
+		if (UObject const* LocalVariableOuter = GetVarOuter())
 		{
-			PostSpawnDelegate = FCustomizeNodeDelegate::CreateStatic(LocalVarSetupLambda, LocalVarDesc.VarName, LocalVarOuter, LocalVarDesc.VarGuid, CustomizeNodeDelegate);
+			PostSpawnDelegate = FCustomizeNodeDelegate::CreateStatic(LocalVarSetupLambda, LocalVarDesc.VarName, LocalVariableOuter, LocalVarDesc.VarGuid, CustomizeNodeDelegate);
 		}
 
 		NewNode = UBlueprintNodeSpawner::SpawnNode<UEdGraphNode>(NodeClass, ParentGraph, Bindings, Location, PostSpawnDelegate);
diff --git a/Engine/Source/Editor/BlueprintGraph/Private/K2Node.cpp b/Engine/Source/Editor/BlueprintGraph/Private/K2Node.cpp
--- a/Engine/Source/Editor/BlueprintGraph/Private/K2Node.cpp
+++ b/Engine/Source/Editor/BlueprintGraph/Private/K2Node.cpp
@@ -498,22 +498,24 @@ UK2Node::ERedirectType UK2Node::DoPinsMatchForReconstruction(const UEdGraphPin*
 				};
 
 				TArray<FPropertyDetails> ParentHierarchy;
-				const UEdGraphPin* CurPin = OldPin;
-				do 
+				FString NewPinNameStr;
 				{
-					ParentHierarchy.Add(FPropertyDetails(CurPin, CurPin->PinName.RightChop(CurPin->ParentPin->PinName.Len() + 1)));
-					CurPin = CurPin->ParentPin;
-				} while (CurPin->ParentPin);
+					const UEdGraphPin* CurPin = OldPin;
+					do 
+					{
+						ParentHierarchy.Add(FPropertyDetails(CurPin, CurPin->PinName.RightChop(CurPin->ParentPin->PinName.Len() + 1)));
+						CurPin = CurPin->ParentPin;
+					} while (CurPin->ParentPin);
 
-				// if you don't have matching pin, now check if there is any redirect param set
-				TArray<FString> OldPinNames;
-				GetRedirectPinNames(*CurPin, OldPinNames);
+					// if you don't have matching pin, now check if there is any redirect param set
+					TArray<FString> OldPinNames;
+					GetRedirectPinNames(*CurPin, OldPinNames);
 
-				FString NewPinNameStr;
-				FName NewPinName;
-				RedirectType = ShouldRedirectParam(OldPinNames, /*out*/ NewPinName, Node);
+					FName NewPinName;
+					RedirectType = ShouldRedirectParam(OldPinNames, /*out*/ NewPinName, Node);
 
-				NewPinNameStr = (RedirectType == ERedirectType_None ? CurPin->PinName : NewPinName.ToString());
+					NewPinNameStr = (RedirectType == ERedirectType_None ? CurPin->PinName : NewPinName.ToString());
+				}
 
 				for (int32 ParentIndex = ParentHierarchy.Num() - 1; ParentIndex >= 0; --ParentIndex)
 				{
diff --git a/Engine/Source/Editor/BlueprintGraph/Private/K2Node_AddComponent.cpp b/Engine/Source/Editor/BlueprintGraph/Private/K2Node_AddComponent.cpp
--- a/Engine/Source/Editor/BlueprintGraph/Private/K2Node_AddComponent.cpp
+++ b/Engine/Source/Editor/BlueprintGraph/Private/K2Node_AddComponent.cpp
@@ -299,11 +299,11 @@ FText UK2Node_AddComponent::GetNodeTitle(ENodeTitleType::Type TitleType) const
 
 		if (UActorComponent* SourceTemplate = Blueprint->FindTemplateByName(FName(*TemplateName)))
 		{
-			CachedNodeTitle = SourceTemplate->GetClass()->GetDisplayNameText();
-
-			FFormatNamedArguments Args;
-			Args.Add(TEXT("ComponentType"), SourceTemplate->GetClass()->GetDisplayNameText());
-			CachedNodeTitle = FText::Format(LOCTEXT("AddClass", "Add {ComponentType}"), Args);
+			{
+				FFormatNamedArguments Args;
+				Args.Add(TEXT("ComponentType"), SourceTemplate->GetClass()->GetDisplayNameText());
+				CachedNodeTitle = FText::Format(LOCTEXT("AddClass", "Add {ComponentType}"), Args);
+			}
 
 			UChildActorComponent* SubActorComp = Cast<UChildActorComponent>(SourceTemplate);
 
diff --git a/Engine/Source/Editor/BlueprintGraph/Private/K2Node_CastByteToEnum.cpp b/Engine/Source/Editor/BlueprintGraph/Private/K2Node_CastByteToEnum.cpp
--- a/Engine/Source/Editor/BlueprintGraph/Private/K2Node_CastByteToEnum.cpp
+++ b/Engine/Source/Editor/BlueprintGraph/Private/K2Node_CastByteToEnum.cpp
@@ -182,8 +182,8 @@ void UK2Node_CastByteToEnum::GetMenuActions(FBlueprintActionDatabaseRegistrar& A
 
 	for (TObjectIterator<UEnum> EnumIt; EnumIt; ++EnumIt)
 	{
-		UEnum const* Enum = (*EnumIt);
-		if (!UEdGraphSchema_K2::IsAllowableBlueprintVariableType(Enum))
+		UEnum const* EnumToConsider = (*EnumIt);
+		if (!UEdGraphSchema_K2::IsAllowableBlueprintVariableType(EnumToConsider))
 		{
 			continue;
 		}
@@ -192,20 +192,20 @@ void UK2Node_CastByteToEnum::GetMenuActions(FBlueprintActionDatabaseRegistrar& A
 		// check to make sure that the registrar is looking for actions of this type
 		// (could be regenerating actions for a specific asset, and therefore the 
 		// registrar would only accept actions corresponding to that asset)
-		if (!ActionRegistrar.IsOpenForRegistration(Enum))
+		if (!ActionRegistrar.IsOpenForRegistration(EnumToConsider))
 		{
 			continue;
 		}
 
-		UBlueprintFieldNodeSpawner* NodeSpawner = UBlueprintFieldNodeSpawner::Create(GetClass(), Enum);
+		UBlueprintFieldNodeSpawner* NodeSpawner = UBlueprintFieldNodeSpawner::Create(GetClass(), EnumToConsider);
 		check(NodeSpawner != nullptr);
-		TWeakObjectPtr<UEnum> NonConstEnumPtr = Enum;
+		TWeakObjectPtr<UEnum> NonConstEnumPtr = EnumToConsider;
 		NodeSpawner->SetNodeFieldDelegate = UBlueprintFieldNodeSpawner::FSetNodeFieldDelegate::CreateStatic(SetNodeEnumLambda, NonConstEnumPtr);
 
 		// this enum could belong to a class, or is a user defined enum (asset), 
 		// that's why we want to make sure to register it along with the action 
 		// (so the action can be refreshed when the class/asset is).
-		ActionRegistrar.AddBlueprintAction(Enum, NodeSpawner);
+		ActionRegistrar.AddBlueprintAction(EnumToConsider, NodeSpawner);
 	}
 }
 
diff --git a/Engine/Source/Editor/BlueprintGraph/Private/K2Node_DynamicCast.cpp b/Engine/Source/Editor/BlueprintGraph/Private/K2Node_DynamicCast.cpp
--- a/Engine/Source/Editor/BlueprintGraph/Private/K2Node_DynamicCast.cpp
+++ b/Engine/Source/Editor/BlueprintGraph/Private/K2Node_DynamicCast.cpp
@@ -110,9 +110,9 @@ void UK2Node_DynamicCast::GetContextMenuActions(const FGraphNodeContextMenuBuild
 		FText MenuEntryTooltip = LOCTEXT("MakePureTooltip", "Removes the execution pins to make the node more versitile (NOTE: the cast could still, resulting in an invalid output).");
 
 		bool bCanTogglePurity = true;
-		auto CanExecutePurityToggle = [](bool const bCanTogglePurity)->bool
+		auto CanExecutePurityToggle = [](bool const bInCanTogglePurity)->bool
 		{
-			return bCanTogglePurity;
+			return bInCanTogglePurity;
 		};
 
 		if (bIsPureCast)
diff --git a/Engine/Source/Editor/BlueprintGraph/Private/K2Node_EnumLiteral.cpp b/Engine/Source/Editor/BlueprintGraph/Private/K2Node_EnumLiteral.cpp
--- a/Engine/Source/Editor/BlueprintGraph/Private/K2Node_EnumLiteral.cpp
+++ b/Engine/Source/Editor/BlueprintGraph/Private/K2Node_EnumLiteral.cpp
@@ -117,8 +117,8 @@ void UK2Node_EnumLiteral::GetMenuActions(FBlueprintActionDatabaseRegistrar& Acti
 
 	for (TObjectIterator<UEnum> EnumIt; EnumIt; ++EnumIt)
 	{
-		UEnum const* Enum = (*EnumIt);
-		if (!UEdGraphSchema_K2::IsAllowableBlueprintVariableType(Enum))
+		UEnum const* EnumToConsider = (*EnumIt);
+		if (!UEdGraphSchema_K2::IsAllowableBlueprintVariableType(EnumToConsider))
 		{
 			continue;
 		}
@@ -127,20 +127,20 @@ void UK2Node_EnumLiteral::GetMenuActions(FBlueprintActionDatabaseRegistrar& Acti
 		// check to make sure that the registrar is looking for actions of this type
 		// (could be regenerating actions for a specific asset, and therefore the 
 		// registrar would only accept actions corresponding to that asset)
-		if (!ActionRegistrar.IsOpenForRegistration(Enum))
+		if (!ActionRegistrar.IsOpenForRegistration(EnumToConsider))
 		{
 			continue;
 		}
 
-		UBlueprintFieldNodeSpawner* NodeSpawner = UBlueprintFieldNodeSpawner::Create(GetClass(), Enum);
+		UBlueprintFieldNodeSpawner* NodeSpawner = UBlueprintFieldNodeSpawner::Create(GetClass(), EnumToConsider);
 		check(NodeSpawner != nullptr);
-		TWeakObjectPtr<UEnum> NonConstEnumPtr = Enum;
+		TWeakObjectPtr<UEnum> NonConstEnumPtr = EnumToConsider;
 		NodeSpawner->SetNodeFieldDelegate = UBlueprintFieldNodeSpawner::FSetNodeFieldDelegate::CreateStatic(SetNodeEnumLambda, NonConstEnumPtr);
 
 		// this enum could belong to a class, or is a user defined enum (asset), 
 		// that's why we want to make sure to register it along with the action 
 		// (so the action can be refreshed when the class/asset is).
-		ActionRegistrar.AddBlueprintAction(Enum, NodeSpawner);
+		ActionRegistrar.AddBlueprintAction(EnumToConsider, NodeSpawner);
 	}
 }
 
diff --git a/Engine/Source/Editor/BlueprintGraph/Private/K2Node_ForEachElementInEnum.cpp b/Engine/Source/Editor/BlueprintGraph/Private/K2Node_ForEachElementInEnum.cpp
--- a/Engine/Source/Editor/BlueprintGraph/Private/K2Node_ForEachElementInEnum.cpp
+++ b/Engine/Source/Editor/BlueprintGraph/Private/K2Node_ForEachElementInEnum.cpp
@@ -210,8 +210,8 @@ void UK2Node_ForEachElementInEnum::GetMenuActions(FBlueprintActionDatabaseRegist
 
 	for (TObjectIterator<UEnum> EnumIt; EnumIt; ++EnumIt)
 	{
-		UEnum const* Enum = (*EnumIt);
-		if (!UEdGraphSchema_K2::IsAllowableBlueprintVariableType(Enum))
+		UEnum const* EnumToConsider = (*EnumIt);
+		if (!UEdGraphSchema_K2::IsAllowableBlueprintVariableType(EnumToConsider))
 		{
 			continue;
 		}
@@ -220,20 +220,20 @@ void UK2Node_ForEachElementInEnum::GetMenuActions(FBlueprintActionDatabaseRegist
 		// check to make sure that the registrar is looking for actions of this type
 		// (could be regenerating actions for a specific asset, and therefore the 
 		// registrar would only accept actions corresponding to that asset)
-		if (!ActionRegistrar.IsOpenForRegistration(Enum))
+		if (!ActionRegistrar.IsOpenForRegistration(EnumToConsider))
 		{
 			continue;
 		}
 
-		UBlueprintFieldNodeSpawner* NodeSpawner = UBlueprintFieldNodeSpawner::Create(GetClass(), Enum);
+		UBlueprintFieldNodeSpawner* NodeSpawner = UBlueprintFieldNodeSpawner::Create(GetClass(), EnumToConsider);
 		check(NodeSpawner != nullptr);
-		TWeakObjectPtr<UEnum> NonConstEnumPtr = Enum;
+		TWeakObjectPtr<UEnum> NonConstEnumPtr = EnumToConsider;
 		NodeSpawner->SetNodeFieldDelegate = UBlueprintFieldNodeSpawner::FSetNodeFieldDelegate::CreateStatic(SetNodeEnumLambda, NonConstEnumPtr);
 
 		// this enum could belong to a class, or is a user defined enum (asset), 
 		// that's why we want to make sure to register it along with the action 
 		// (so the action can be refreshed when the class/asset is).
-		ActionRegistrar.AddBlueprintAction(Enum, NodeSpawner);
+		ActionRegistrar.AddBlueprintAction(EnumToConsider, NodeSpawner);
 	}
 }
 
diff --git a/Engine/Source/Editor/BlueprintGraph/Private/K2Node_GetDataTableRow.cpp b/Engine/Source/Editor/BlueprintGraph/Private/K2Node_GetDataTableRow.cpp
--- a/Engine/Source/Editor/BlueprintGraph/Private/K2Node_GetDataTableRow.cpp
+++ b/Engine/Source/Editor/BlueprintGraph/Private/K2Node_GetDataTableRow.cpp
@@ -90,7 +90,6 @@ void UK2Node_GetDataTableRow::SetReturnTypeForStruct(UScriptStruct* RowStruct)
 		ResultPin->BreakAllPinLinks();
 
 		// Change class of output pin
-		UEdGraphPin* ResultPin = GetResultPin();
 		ResultPin->PinType.PinSubCategoryObject = RowStruct;
 
 		// When the DataTable pin gets a new value assigned, we need to update the Slate UI so that SGraphNodeCallParameterCollectionFunction will update the ParameterName drop down
diff --git a/Engine/Source/Editor/BlueprintGraph/Private/K2Node_GetInputAxisValue.cpp b/Engine/Source/Editor/BlueprintGraph/Private/K2Node_GetInputAxisValue.cpp
--- a/Engine/Source/Editor/BlueprintGraph/Private/K2Node_GetInputAxisValue.cpp
+++ b/Engine/Source/Editor/BlueprintGraph/Private/K2Node_GetInputAxisValue.cpp
@@ -140,12 +140,12 @@ void UK2Node_GetInputAxisValue::GetMenuActions(FBlueprintActionDatabaseRegistrar
 			FEditorDelegates::OnActionAxisMappingsChanged.AddStatic(RefreshClassActions);
 		}
 
-		for (FName const InputAxisName : AxisNames)
+		for (FName const AxisName : AxisNames)
 		{
 			UBlueprintNodeSpawner* NodeSpawner = UBlueprintNodeSpawner::Create(GetClass());
 			check(NodeSpawner != nullptr);
 
-			NodeSpawner->CustomizeNodeDelegate = UBlueprintNodeSpawner::FCustomizeNodeDelegate::CreateStatic(CustomizeInputNodeLambda, InputAxisName);
+			NodeSpawner->CustomizeNodeDelegate = UBlueprintNodeSpawner::FCustomizeNodeDelegate::CreateStatic(CustomizeInputNodeLambda, AxisName);
 			ActionRegistrar.AddBlueprintAction(ActionKey, NodeSpawner);
 		}
 	}
diff --git a/Engine/Source/Editor/BlueprintGraph/Private/K2Node_GetNumEnumEntries.cpp b/Engine/Source/Editor/BlueprintGraph/Private/K2Node_GetNumEnumEntries.cpp
--- a/Engine/Source/Editor/BlueprintGraph/Private/K2Node_GetNumEnumEntries.cpp
+++ b/Engine/Source/Editor/BlueprintGraph/Private/K2Node_GetNumEnumEntries.cpp
@@ -103,8 +103,8 @@ void UK2Node_GetNumEnumEntries::GetMenuActions(FBlueprintActionDatabaseRegistrar
 
 	for (TObjectIterator<UEnum> EnumIt; EnumIt; ++EnumIt)
 	{
-		UEnum const* Enum = (*EnumIt);
-		if (!UEdGraphSchema_K2::IsAllowableBlueprintVariableType(Enum))
+		UEnum const* EnumToConsider = (*EnumIt);
+		if (!UEdGraphSchema_K2::IsAllowableBlueprintVariableType(EnumToConsider))
 		{
 			continue;
 		}
@@ -113,20 +113,20 @@ void UK2Node_GetNumEnumEntries::GetMenuActions(FBlueprintActionDatabaseRegistrar
 		// check to make sure that the registrar is looking for actions of this type
 		// (could be regenerating actions for a specific asset, and therefore the 
 		// registrar would only accept actions corresponding to that asset)
-		if (!ActionRegistrar.IsOpenForRegistration(Enum))
+		if (!ActionRegistrar.IsOpenForRegistration(EnumToConsider))
 		{
 			continue;
 		}
 
-		UBlueprintFieldNodeSpawner* NodeSpawner = UBlueprintFieldNodeSpawner::Create(GetClass(), Enum);
+		UBlueprintFieldNodeSpawner* NodeSpawner = UBlueprintFieldNodeSpawner::Create(GetClass(), EnumToConsider);
 		check(NodeSpawner != nullptr);
-		TWeakObjectPtr<UEnum> NonConstEnumPtr = Enum;
+		TWeakObjectPtr<UEnum> NonConstEnumPtr = EnumToConsider;
 		NodeSpawner->SetNodeFieldDelegate = UBlueprintFieldNodeSpawner::FSetNodeFieldDelegate::CreateStatic(SetNodeEnumLambda, NonConstEnumPtr);
 
 		// this enum could belong to a class, or is a user defined enum (asset), 
 		// that's why we want to make sure to register it along with the action 
 		// (so the action can be refreshed when the class/asset is).
-		ActionRegistrar.AddBlueprintAction(Enum, NodeSpawner);
+		ActionRegistrar.AddBlueprintAction(EnumToConsider, NodeSpawner);
 	}
 }
 
diff --git a/Engine/Source/Editor/BlueprintGraph/Private/K2Node_InputAction.cpp b/Engine/Source/Editor/BlueprintGraph/Private/K2Node_InputAction.cpp
--- a/Engine/Source/Editor/BlueprintGraph/Private/K2Node_InputAction.cpp
+++ b/Engine/Source/Editor/BlueprintGraph/Private/K2Node_InputAction.cpp
@@ -176,12 +176,12 @@ void UK2Node_InputAction::GetMenuActions(FBlueprintActionDatabaseRegistrar& Acti
 			FEditorDelegates::OnActionAxisMappingsChanged.AddStatic(RefreshClassActions);
 		}
 
-		for (FName const InputActionName : ActionNames)
+		for (FName const ActionName : ActionNames)
 		{
 			UBlueprintNodeSpawner* NodeSpawner = UBlueprintNodeSpawner::Create(GetClass());
 			check(NodeSpawner != nullptr);
 
-			NodeSpawner->CustomizeNodeDelegate = UBlueprintNodeSpawner::FCustomizeNodeDelegate::CreateStatic(CustomizeInputNodeLambda, InputActionName);
+			NodeSpawner->CustomizeNodeDelegate = UBlueprintNodeSpawner::FCustomizeNodeDelegate::CreateStatic(CustomizeInputNodeLambda, ActionName);
 			ActionRegistrar.AddBlueprintAction(ActionKey, NodeSpawner);
 		}
 	}
diff --git a/Engine/Source/Editor/BlueprintGraph/Private/K2Node_InputAxisEvent.cpp b/Engine/Source/Editor/BlueprintGraph/Private/K2Node_InputAxisEvent.cpp
--- a/Engine/Source/Editor/BlueprintGraph/Private/K2Node_InputAxisEvent.cpp
+++ b/Engine/Source/Editor/BlueprintGraph/Private/K2Node_InputAxisEvent.cpp
@@ -160,12 +160,12 @@ void UK2Node_InputAxisEvent::GetMenuActions(FBlueprintActionDatabaseRegistrar& A
 			FEditorDelegates::OnActionAxisMappingsChanged.AddStatic(RefreshClassActions);
 		}
 
-		for (FName const InputAxisName : AxisNames)
+		for (FName const AxisName : AxisNames)
 		{
 			UBlueprintNodeSpawner* NodeSpawner = UBlueprintNodeSpawner::Create(GetClass());
 			check(NodeSpawner != nullptr);
 
-			NodeSpawner->CustomizeNodeDelegate = UBlueprintNodeSpawner::FCustomizeNodeDelegate::CreateStatic(CustomizeInputNodeLambda, InputAxisName);
+			NodeSpawner->CustomizeNodeDelegate = UBlueprintNodeSpawner::FCustomizeNodeDelegate::CreateStatic(CustomizeInputNodeLambda, AxisName);
 			ActionRegistrar.AddBlueprintAction(ActionKey, NodeSpawner);
 		}
 	}
diff --git a/Engine/Source/Editor/BlueprintGraph/Private/K2Node_MakeStruct.cpp b/Engine/Source/Editor/BlueprintGraph/Private/K2Node_MakeStruct.cpp
--- a/Engine/Source/Editor/BlueprintGraph/Private/K2Node_MakeStruct.cpp
+++ b/Engine/Source/Editor/BlueprintGraph/Private/K2Node_MakeStruct.cpp
@@ -248,7 +248,6 @@ UK2Node::ERedirectType UK2Node_MakeStruct::DoPinsMatchForReconstruction(const UE
 	}
 	else if ((ERedirectType_None == Result) && NewPin && OldPin)
 	{
-		const UEdGraphSchema_K2* K2Schema = GetDefault<UEdGraphSchema_K2>();
 		if ((EGPD_Output == NewPin->Direction) && (EGPD_Output == OldPin->Direction))
 		{
 			const UEdGraphSchema_K2* K2Schema = GetDefault<UEdGraphSchema_K2>();
diff --git a/Engine/Source/Editor/BlueprintGraph/Private/K2Node_SwitchEnum.cpp b/Engine/Source/Editor/BlueprintGraph/Private/K2Node_SwitchEnum.cpp
--- a/Engine/Source/Editor/BlueprintGraph/Private/K2Node_SwitchEnum.cpp
+++ b/Engine/Source/Editor/BlueprintGraph/Private/K2Node_SwitchEnum.cpp
@@ -93,8 +93,8 @@ void UK2Node_SwitchEnum::GetMenuActions(FBlueprintActionDatabaseRegistrar& Actio
 
 	for (TObjectIterator<UEnum> EnumIt; EnumIt; ++EnumIt)
 	{
-		UEnum const* Enum = (*EnumIt);
-		if (!UEdGraphSchema_K2::IsAllowableBlueprintVariableType(Enum))
+		UEnum const* EnumToConsider = (*EnumIt);
+		if (!UEdGraphSchema_K2::IsAllowableBlueprintVariableType(EnumToConsider))
 		{
 			continue;
 		}
@@ -103,20 +103,20 @@ void UK2Node_SwitchEnum::GetMenuActions(FBlueprintActionDatabaseRegistrar& Actio
 		// check to make sure that the registrar is looking for actions of this type
 		// (could be regenerating actions for a specific asset, and therefore the 
 		// registrar would only accept actions corresponding to that asset)
-		if (!ActionRegistrar.IsOpenForRegistration(Enum))
+		if (!ActionRegistrar.IsOpenForRegistration(EnumToConsider))
 		{
 			continue;
 		}
 
-		UBlueprintFieldNodeSpawner* NodeSpawner = UBlueprintFieldNodeSpawner::Create(GetClass(), Enum);
+		UBlueprintFieldNodeSpawner* NodeSpawner = UBlueprintFieldNodeSpawner::Create(GetClass(), EnumToConsider);
 		check(NodeSpawner != nullptr);
-		TWeakObjectPtr<UEnum> NonConstEnumPtr = Enum;
+		TWeakObjectPtr<UEnum> NonConstEnumPtr = EnumToConsider;
 		NodeSpawner->SetNodeFieldDelegate = UBlueprintFieldNodeSpawner::FSetNodeFieldDelegate::CreateStatic(SetNodeEnumLambda, NonConstEnumPtr);
 
 		// this enum could belong to a class, or is a user defined enum (asset), 
 		// that's why we want to make sure to register it along with the action 
 		// (so the action can be refreshed when the class/asset is).
-		ActionRegistrar.AddBlueprintAction(Enum, NodeSpawner);
+		ActionRegistrar.AddBlueprintAction(EnumToConsider, NodeSpawner);
 	}
 }
 
diff --git a/Engine/Source/Editor/BlueprintGraph/Private/K2Node_TemporaryVariable.cpp b/Engine/Source/Editor/BlueprintGraph/Private/K2Node_TemporaryVariable.cpp
--- a/Engine/Source/Editor/BlueprintGraph/Private/K2Node_TemporaryVariable.cpp
+++ b/Engine/Source/Editor/BlueprintGraph/Private/K2Node_TemporaryVariable.cpp
@@ -151,19 +151,19 @@ void UK2Node_TemporaryVariable::GetMenuActions(FBlueprintActionDatabaseRegistrar
 		return;
 	}
 
-	auto MakeTempVarNodeSpawner = [](FEdGraphPinType const& VarType, bool bIsPersistent)
+	auto MakeTempVarNodeSpawner = [](FEdGraphPinType const& VarType, bool bVarIsPersistent)
 	{
 		UBlueprintNodeSpawner* NodeSpawner = UBlueprintNodeSpawner::Create(UK2Node_TemporaryVariable::StaticClass());
 		check(NodeSpawner != nullptr);
 
-		auto PostSpawnLambda = [](UEdGraphNode* NewNode, bool bIsTemplateNode, FEdGraphPinType VarType, bool bIsPersistent)
+		auto PostSpawnLambda = [](UEdGraphNode* NewNode, bool bIsTemplateNode, FEdGraphPinType InVarType, bool bInIsPersistent)
 		{
 			UK2Node_TemporaryVariable* TempVarNode = CastChecked<UK2Node_TemporaryVariable>(NewNode);
-			TempVarNode->VariableType  = VarType;
-			TempVarNode->bIsPersistent = bIsPersistent;
+			TempVarNode->VariableType  = InVarType;
+			TempVarNode->bIsPersistent = bInIsPersistent;
 		};
 
-		NodeSpawner->CustomizeNodeDelegate = UBlueprintNodeSpawner::FCustomizeNodeDelegate::CreateStatic(PostSpawnLambda, VarType, bIsPersistent);
+		NodeSpawner->CustomizeNodeDelegate = UBlueprintNodeSpawner::FCustomizeNodeDelegate::CreateStatic(PostSpawnLambda, VarType, bVarIsPersistent);
 		return NodeSpawner;
 	};
 
diff --git a/Engine/Source/Editor/GameplayAbilitiesEditor/Private/K2Node_LatentAbilityCall.cpp b/Engine/Source/Editor/GameplayAbilitiesEditor/Private/K2Node_LatentAbilityCall.cpp
--- a/Engine/Source/Editor/GameplayAbilitiesEditor/Private/K2Node_LatentAbilityCall.cpp
+++ b/Engine/Source/Editor/GameplayAbilitiesEditor/Private/K2Node_LatentAbilityCall.cpp
@@ -128,12 +128,12 @@ void UK2Node_LatentAbilityCall::GetMenuActions(FBlueprintActionDatabaseRegistrar
 					UK2Node_LatentAbilityCall* AsyncTaskNode = CastChecked<UK2Node_LatentAbilityCall>(NewNode);
 					if (FunctionPtr.IsValid())
 					{
-						UFunction* Function = FunctionPtr.Get();
-						UObjectProperty* ReturnProperty = CastChecked<UObjectProperty>(Function->GetReturnProperty());
+						UFunction* Func = FunctionPtr.Get();
+						UObjectProperty* ReturnProp = CastChecked<UObjectProperty>(Func->GetReturnProperty());
 						
-						AsyncTaskNode->ProxyFactoryFunctionName = Function->GetFName();
-						AsyncTaskNode->ProxyFactoryClass        = Function->GetOuterUClass();
-						AsyncTaskNode->ProxyClass               = ReturnProperty->PropertyClass;
+						AsyncTaskNode->ProxyFactoryFunctionName = Func->GetFName();
+						AsyncTaskNode->ProxyFactoryClass        = Func->GetOuterUClass();
+						AsyncTaskNode->ProxyClass               = ReturnProp->PropertyClass;
 					}
 				};
 				
diff --git a/Engine/Source/Editor/Persona/Private/AnimNotifyDetails.cpp b/Engine/Source/Editor/Persona/Private/AnimNotifyDetails.cpp
--- a/Engine/Source/Editor/Persona/Private/AnimNotifyDetails.cpp
+++ b/Engine/Source/Editor/Persona/Private/AnimNotifyDetails.cpp
@@ -127,9 +127,9 @@ void FAnimNotifyDetails::CustomizeDetails( IDetailLayoutBuilder& DetailBuilder )
 
 	for (FPropVisPair& NotifyPair : TriggerSettingNames)
 	{
-		TSharedRef<IPropertyHandle> NotifyPropHandle = DetailBuilder.GetProperty(NotifyPair.NotifyName);
-		DetailBuilder.HideProperty(NotifyPropHandle);
-		TriggerSettingCategory.AddProperty(NotifyPropHandle).Visibility(NotifyPair.Visibility);
+		TSharedRef<IPropertyHandle> NotifyPropertyHandle = DetailBuilder.GetProperty(NotifyPair.NotifyName);
+		DetailBuilder.HideProperty(NotifyPropertyHandle);
+		TriggerSettingCategory.AddProperty(NotifyPropertyHandle).Visibility(NotifyPair.Visibility);
 	}
 }
 
diff --git a/Engine/Source/Editor/UnrealEd/Private/Fbx/FbxSkeletalMeshImport.cpp b/Engine/Source/Editor/UnrealEd/Private/Fbx/FbxSkeletalMeshImport.cpp
--- a/Engine/Source/Editor/UnrealEd/Private/Fbx/FbxSkeletalMeshImport.cpp
+++ b/Engine/Source/Editor/UnrealEd/Private/Fbx/FbxSkeletalMeshImport.cpp
@@ -2024,7 +2024,6 @@ bool UnFbx::FFbxImporter::FillSkelMeshImporterFromFbx( FSkeletalMeshImportData&
 	
 	bool OddNegativeScale = IsOddNegativeScale(TotalMatrix);
 	
-	int32 VertexIndex;
 	int32 TriangleCount = Mesh->GetPolygonCount();
 	int32 ExistFaceNum = ImportData.Faces.Num();
 	ImportData.Faces.AddUninitialized( TriangleCount );
@@ -2057,7 +2056,7 @@ bool UnFbx::FFbxImporter::FillSkelMeshImporterFromFbx( FSkeletalMeshImportData&
 			}
 		}
 
-		for (VertexIndex=0; VertexIndex<3; VertexIndex++)
+		for (int32 VertexIndex=0; VertexIndex<3; VertexIndex++)
 		{
 			// If there are odd number negative scale, invert the vertex order for triangles
 			int32 UnrealVertexIndex = OddNegativeScale ? 2 - VertexIndex : VertexIndex;
@@ -2153,7 +2152,7 @@ bool UnFbx::FFbxImporter::FillSkelMeshImporterFromFbx( FSkeletalMeshImportData&
 		}
 	
 		Triangle.AuxMatIndex = 0;
-		for (VertexIndex=0; VertexIndex<3; VertexIndex++)
+		for (int32 VertexIndex=0; VertexIndex<3; VertexIndex++)
 		{
 			// If there are odd number negative scale, invert the vertex order for triangles
 			int32 UnrealVertexIndex = OddNegativeScale ? 2 - VertexIndex : VertexIndex;
@@ -2175,7 +2174,7 @@ bool UnFbx::FFbxImporter::FillSkelMeshImporterFromFbx( FSkeletalMeshImportData&
 			if (LayerElementUV[UVLayerIndex] != NULL) 
 			{
 				// Get each UV from the layer
-				for (VertexIndex=0;VertexIndex<3;VertexIndex++)
+				for (int32 VertexIndex=0;VertexIndex<3;VertexIndex++)
 				{
 					// If there are odd number negative scale, invert the vertex order for triangles
 					int32 UnrealVertexIndex = OddNegativeScale ? 2 - VertexIndex : VertexIndex;
@@ -2196,7 +2195,7 @@ bool UnFbx::FFbxImporter::FillSkelMeshImporterFromFbx( FSkeletalMeshImportData&
 				// Set all UV's to zero.  If we are here the mesh had no UV sets so we only need to do this for the
 				// first UV set which always exists.
 
-				for(VertexIndex=0; VertexIndex<3; VertexIndex++)
+				for (int32 VertexIndex=0; VertexIndex<3; VertexIndex++)
 				{
 					TmpWedges[VertexIndex].UVs[UVLayerIndex].X = 0.0f;
 					TmpWedges[VertexIndex].UVs[UVLayerIndex].Y = 0.0f;
@@ -2211,8 +2210,7 @@ bool UnFbx::FFbxImporter::FillSkelMeshImporterFromFbx( FSkeletalMeshImportData&
 			{
 			case FbxLayerElement::eByControlPoint:
 				{
-					int32 VertexIndex;
-					for (VertexIndex=0;VertexIndex<3;VertexIndex++)
+					for (int32 VertexIndex=0;VertexIndex<3;VertexIndex++)
 					{
 						int32 UnrealVertexIndex = OddNegativeScale ? 2 - VertexIndex : VertexIndex;
 
@@ -2229,8 +2227,7 @@ bool UnFbx::FFbxImporter::FillSkelMeshImporterFromFbx( FSkeletalMeshImportData&
 				break;
 			case FbxLayerElement::eByPolygonVertex:
 				{	
-					int32 VertexIndex;
-					for (VertexIndex=0;VertexIndex<3;VertexIndex++)
+					for (int32 VertexIndex=0;VertexIndex<3;VertexIndex++)
 					{
 						int32 UnrealVertexIndex = OddNegativeScale ? 2 - VertexIndex : VertexIndex;
 
@@ -2251,7 +2248,7 @@ bool UnFbx::FFbxImporter::FillSkelMeshImporterFromFbx( FSkeletalMeshImportData&
 		//
 		// basic wedges matching : 3 unique per face. TODO Can we do better ?
 		//
-		for (VertexIndex=0; VertexIndex<3; VertexIndex++)
+		for (int32 VertexIndex=0; VertexIndex<3; VertexIndex++)
 		{
 			int32 w;
 			
diff --git a/Engine/Source/Editor/UnrealEd/Private/Kismet2/Kismet2.cpp b/Engine/Source/Editor/UnrealEd/Private/Kismet2/Kismet2.cpp
--- a/Engine/Source/Editor/UnrealEd/Private/Kismet2/Kismet2.cpp
+++ b/Engine/Source/Editor/UnrealEd/Private/Kismet2/Kismet2.cpp
@@ -1068,9 +1068,9 @@ void FKismetEditorUtilities::AddComponentsToBlueprint(UBlueprint* Blueprint, con
 				else if (SceneComponent->AttachParent->IsCreatedByConstructionScript())
 				{
 					USCS_Node* ParentSCSNode = nullptr;
-					for (UBlueprint* Blueprint : ParentBPStack)
+					for (UBlueprint* ParentBlueprint : ParentBPStack)
 					{
-						ParentSCSNode = Blueprint->SimpleConstructionScript->FindSCSNode(SceneComponent->AttachParent->GetFName());
+						ParentSCSNode = ParentBlueprint->SimpleConstructionScript->FindSCSNode(SceneComponent->AttachParent->GetFName());
 						if (ParentSCSNode)
 						{
 							break;
diff --git a/Engine/Source/Editor/UnrealEd/Private/Kismet2/KismetReinstanceUtilities.cpp b/Engine/Source/Editor/UnrealEd/Private/Kismet2/KismetReinstanceUtilities.cpp
--- a/Engine/Source/Editor/UnrealEd/Private/Kismet2/KismetReinstanceUtilities.cpp
+++ b/Engine/Source/Editor/UnrealEd/Private/Kismet2/KismetReinstanceUtilities.cpp
@@ -75,8 +75,8 @@ struct FReplaceReferenceHelper
 					class ReferenceReplace : public FArchiveReplaceObjectRef < UObject >
 					{
 					public:
-						ReferenceReplace(UObject* InSearchObject, const TMap<UObject*, UObject*>& InReplacementMap, TMap<FStringAssetReference, UObject*> WeakReferencesMap)
-							: FArchiveReplaceObjectRef<UObject>(InSearchObject, InReplacementMap, false, false, false, true), WeakReferencesMap(WeakReferencesMap)
+						ReferenceReplace(UObject* InSearchObject, const TMap<UObject*, UObject*>& InReplacementMap, TMap<FStringAssetReference, UObject*> InWeakReferencesMap)
+							: FArchiveReplaceObjectRef<UObject>(InSearchObject, InReplacementMap, false, false, false, true), WeakReferencesMap(InWeakReferencesMap)
 						{
 							SerializeSearchObject();
 						}
diff --git a/Engine/Source/Runtime/Core/Public/Stats/StatsMisc.h b/Engine/Source/Runtime/Core/Public/Stats/StatsMisc.h
--- a/Engine/Source/Runtime/Core/Public/Stats/StatsMisc.h
+++ b/Engine/Source/Runtime/Core/Public/Stats/StatsMisc.h
@@ -74,14 +74,17 @@ protected:
 	EScopeLogTimeUnits Units;
 };
 
+#define SCOPE_LOG_TOKENPASTE_INNER(x,y) x##y
+#define SCOPE_LOG_TOKENPASTE(x,y) SCOPE_LOG_TOKENPASTE_INNER(x,y)
+
 #define SCOPE_LOG_TIME(Name,CumulativePtr) \
-	FScopeLogTime ScopeLogTime(Name,CumulativePtr);
+	FScopeLogTime SCOPE_LOG_TOKENPASTE(ScopeLogTime,__LINE__)(Name,CumulativePtr);
 
 #define SCOPE_LOG_TIME_IN_SECONDS(Name,CumulativePtr) \
-	FScopeLogTime ScopeLogTime(Name, CumulativePtr, FScopeLogTime::ScopeLog_Seconds);
+	FScopeLogTime SCOPE_LOG_TOKENPASTE(ScopeLogTime,__LINE__)(Name, CumulativePtr, FScopeLogTime::ScopeLog_Seconds);
 
 #define SCOPE_LOG_TIME_FUNC() \
-	FScopeLogTime ScopeLogTime(__FUNCTION__);
+	FScopeLogTime SCOPE_LOG_TOKENPASTE(ScopeLogTime,__LINE__)(__FUNCTION__);
 
 #define SCOPE_LOG_TIME_FUNC_WITH_GLOBAL(CumulativePtr) \
-	FScopeLogTime ScopeLogTime(__FUNCTION__,CumulativePtr);
+	FScopeLogTime SCOPE_LOG_TOKENPASTE(ScopeLogTime,__LINE__)(__FUNCTION__,CumulativePtr);
diff --git a/Engine/Source/Runtime/Engine/Private/Particles/ParticleSystemRender.cpp b/Engine/Source/Runtime/Engine/Private/Particles/ParticleSystemRender.cpp
--- a/Engine/Source/Runtime/Engine/Private/Particles/ParticleSystemRender.cpp
+++ b/Engine/Source/Runtime/Engine/Private/Particles/ParticleSystemRender.cpp
@@ -885,10 +885,10 @@ void GatherParticleLightData(const FDynamicSpriteEmitterReplayDataBase& Source,
 							//If this is the first time we've needed to add these then we need to fill data up to this light.
 							int32 NumLights = OutParticleLights.InstanceData.Num();
 							OutParticleLights.InstancePerViewDataIndices.AddUninitialized(OutParticleLights.InstanceData.Num());
-							for (int32 i = 0; i < NumLights; ++i)
+							for (int32 LightIndex = 0; LightIndex < NumLights; ++LightIndex)
 							{
-								OutParticleLights.InstancePerViewDataIndices[i].PerViewIndex = i;
-								OutParticleLights.InstancePerViewDataIndices[i].bHasPerViewData = false;
+								OutParticleLights.InstancePerViewDataIndices[LightIndex].PerViewIndex = LightIndex;
+								OutParticleLights.InstancePerViewDataIndices[LightIndex].bHasPerViewData = false;
 							}
 						}
 
@@ -969,7 +969,7 @@ void FDynamicSpriteEmitterData::GetDynamicMeshElementsEmitter(const FParticleSys
 		if (SourceData->EmitterRenderMode == ERM_Normal)
 		{
 			// Determine how many vertices and indices are needed to render.
-			const int32 ParticleCount = SourceData->ActiveParticleCount;
+			int32 ParticleCount = SourceData->ActiveParticleCount;
 			const int32 VertexSize = GetDynamicVertexStride(FeatureLevel);
 			const int32 DynamicParameterVertexSize = sizeof(FParticleVertexDynamicParameter);
 			const int32 VertexPerParticle = bInstanced ? 1 : 4;
@@ -1034,16 +1034,13 @@ void FDynamicSpriteEmitterData::GetDynamicMeshElementsEmitter(const FParticleSys
 				PerViewUniformParameters.MacroUVParameters = FVector4(ObjectNDCPosition.X, ObjectNDCPosition.Y, ObjectMacroUVScales.X, ObjectMacroUVScales.Y);
 				CollectorResources.UniformBuffer = FParticleSpriteUniformBufferRef::CreateUniformBufferImmediate(PerViewUniformParameters, UniformBuffer_SingleFrame);
 
-				const FDynamicSpriteEmitterReplayDataBase* SourceData = GetSourceData();
-				check(SourceData);
-
 				FParticleSpriteVertexFactory* SpriteVertexFactory = &CollectorResources.VertexFactory;
 
 				// Don't render if the material will be ignored
 				const bool bIsWireframe = ViewFamily.EngineShowFlags.Wireframe;
 
 				// Calculate the number of particles that must be drawn.
-				int32 ParticleCount = Source.ActiveParticleCount;
+				ParticleCount = Source.ActiveParticleCount;
 				if ((Source.MaxDrawCount >= 0) && (ParticleCount > Source.MaxDrawCount))
 				{
 					ParticleCount = Source.MaxDrawCount;
