UE-11058 Crash when renaming blueprint

Original CDO is no longer deleted by GC, before the reinstanceng is done.
Fixed a crash when blueprint was reinstanced after BP compiler error.

diff --git a/Engine/Source/Editor/UnrealEd/Private/Kismet2/KismetReinstanceUtilities.cpp b/Engine/Source/Editor/UnrealEd/Private/Kismet2/KismetReinstanceUtilities.cpp
--- a/Engine/Source/Editor/UnrealEd/Private/Kismet2/KismetReinstanceUtilities.cpp
+++ b/Engine/Source/Editor/UnrealEd/Private/Kismet2/KismetReinstanceUtilities.cpp
@@ -134,8 +134,7 @@ FBlueprintCompileReinstancer::FBlueprintCompileReinstancer(UClass* InClassToRein
 		ClassToReinstance->ClassFlags |= CLASS_NewerVersionExists;
 		const FName RenistanceName = MakeUniqueObjectName(GetTransientPackage(), ClassToReinstance->GetClass(), *FString::Printf(TEXT("REINST_%s"), *ClassToReinstance->GetName()));
 		DuplicatedClass = (UClass*)StaticDuplicateObject(ClassToReinstance, GetTransientPackage(), *RenistanceName.ToString(), ~RF_Transactional); 
-		// Add the class to the root set, so that it doesn't get prematurely garbage collected
-		DuplicatedClass->AddToRoot();
+
 		ClassToReinstance->ClassFlags &= ~CLASS_NewerVersionExists;
 		GIsDuplicatingClassForReinstancing = false;
 
@@ -264,13 +263,14 @@ void FBlueprintCompileReinstancer::GenerateFieldMappings(TMap<UObject*, UObject*
 	FieldMapping.Add(OriginalCDO, NewCDO);
 }
 
+void FBlueprintCompileReinstancer::AddReferencedObjects(FReferenceCollector& Collector)
+{
+	Collector.AddReferencedObject(OriginalCDO);
+	Collector.AddReferencedObject(DuplicatedClass);
+}
+
 FBlueprintCompileReinstancer::~FBlueprintCompileReinstancer()
 {
-	// Remove the duplicated class from the root set, because it is no longer useful
-	if( DuplicatedClass )
-	{
-		DuplicatedClass->RemoveFromRoot();
-	}
 }
 
 void FBlueprintCompileReinstancer::ReinstanceFast()
@@ -391,10 +391,9 @@ void FBlueprintCompileReinstancer::ReinstanceObjects(bool bForceAlwaysReinstance
 	if (FirstCompileChildrenThenReinstance)
 	{
 		// Make sure we only reinstance classes once!
-		static TArray<TSharedPtr<FBlueprintCompileReinstancer>> QueueToReinstance;
-		TSharedPtr<FBlueprintCompileReinstancer> SharedThis = AsShared();
+		static TArray<TSharedRef<FBlueprintCompileReinstancer>> QueueToReinstance;
+		TSharedRef<FBlueprintCompileReinstancer> SharedThis = AsShared();
 		const bool bAlreadyQueued = QueueToReinstance.Contains(SharedThis);
-		ensure(!bAlreadyQueued);
 		if (!bAlreadyQueued)
 		{
 			QueueToReinstance.Push(SharedThis);
@@ -403,16 +402,16 @@ void FBlueprintCompileReinstancer::ReinstanceObjects(bool bForceAlwaysReinstance
 			{
 				CompileChildren();
 			}
-		}
 
-		if (QueueToReinstance.Num() && (QueueToReinstance[0] == SharedThis))
-		{
-			// All children were recompiled. It's safe to reinstance.
-			for (int32 Idx = 0; Idx < QueueToReinstance.Num(); ++Idx)
+			if (QueueToReinstance.Num() && (QueueToReinstance[0] == SharedThis))
 			{
-				QueueToReinstance[Idx]->ReinstanceInner(bForceAlwaysReinstance);
+				// All children were recompiled. It's safe to reinstance.
+				for (int32 Idx = 0; Idx < QueueToReinstance.Num(); ++Idx)
+				{
+					QueueToReinstance[Idx]->ReinstanceInner(bForceAlwaysReinstance);
+				}
+				QueueToReinstance.Empty();
 			}
-			QueueToReinstance.Empty();
 		}
 	}
 	else
diff --git a/Engine/Source/Editor/UnrealEd/Public/Kismet2/KismetReinstanceUtilities.h b/Engine/Source/Editor/UnrealEd/Public/Kismet2/KismetReinstanceUtilities.h
--- a/Engine/Source/Editor/UnrealEd/Public/Kismet2/KismetReinstanceUtilities.h
+++ b/Engine/Source/Editor/UnrealEd/Public/Kismet2/KismetReinstanceUtilities.h
@@ -10,7 +10,7 @@
 
 DECLARE_STATS_GROUP(TEXT("Kismet Reinstancer"), STATGROUP_KismetReinstancer, STATCAT_Advanced);
 
-class UNREALED_API FBlueprintCompileReinstancer : public TSharedFromThis<FBlueprintCompileReinstancer>
+class UNREALED_API FBlueprintCompileReinstancer : public TSharedFromThis<FBlueprintCompileReinstancer>, public FGCObject
 {
 protected:
 
@@ -48,6 +48,10 @@ protected:
 	TSet<UObject*> ObjectsThatShouldUseOldStuff;
 
 public:
+	// FSerializableObject interface
+	virtual void AddReferencedObjects(FReferenceCollector& Collector) override;
+	// End of FSerializableObject interface
+
 	static TSharedPtr<FBlueprintCompileReinstancer> Create(UClass* InClassToReinstance, bool bIsBytecodeOnly = false, bool bSkipGC = false)
 	{
 		return MakeShareable(new FBlueprintCompileReinstancer(InClassToReinstance, bIsBytecodeOnly, bSkipGC));
