Stopped FBlueprintEditorUtils::IsBlutility from requiring the Blutility module to be opened, this resolves issues with using the DumpBlueprintsInfoCommandlet.

diff --git a/Engine/Source/Editor/UnrealEd/Private/Kismet2/BlueprintEditorUtils.cpp b/Engine/Source/Editor/UnrealEd/Private/Kismet2/BlueprintEditorUtils.cpp
--- a/Engine/Source/Editor/UnrealEd/Private/Kismet2/BlueprintEditorUtils.cpp
+++ b/Engine/Source/Editor/UnrealEd/Private/Kismet2/BlueprintEditorUtils.cpp
@@ -2479,8 +2479,13 @@ bool FBlueprintEditorUtils::IsBlueprintConst(const UBlueprint* Blueprint)
 
 bool FBlueprintEditorUtils::IsBlutility(const UBlueprint* Blueprint)
 {
-	IBlutilityModule& BlutilityModule = FModuleManager::GetModuleChecked<IBlutilityModule>("Blutility");
-	return BlutilityModule.IsBlutility( Blueprint );
+	IBlutilityModule* BlutilityModule = FModuleManager::GetModulePtr<IBlutilityModule>("Blutility");
+
+	if (BlutilityModule)
+	{
+		return BlutilityModule->IsBlutility( Blueprint );
+	}
+	return false;
 }
 
 bool FBlueprintEditorUtils::IsActorBased(const UBlueprint* Blueprint)
