Shadow variable fixes

diff --git a/Engine/Plugins/Runtime/OculusRift/Source/OculusRift/Private/OculusRiftHMD.cpp b/Engine/Plugins/Runtime/OculusRift/Source/OculusRift/Private/OculusRiftHMD.cpp
--- a/Engine/Plugins/Runtime/OculusRift/Source/OculusRift/Private/OculusRiftHMD.cpp
+++ b/Engine/Plugins/Runtime/OculusRift/Source/OculusRift/Private/OculusRiftHMD.cpp
@@ -1268,12 +1268,13 @@ bool FOculusRiftHMD::DoEnableStereo(bool bStereo, bool bApplyToHmd)
 				TSharedPtr<SWindow> Window = SceneVP->FindWindow();
 				if (Window.IsValid())
 				{
-					FVector2D size = Window->GetSizeInScreen();
-
 					if (bApplyToHmd && IsFullscreenAllowed())
 					{
-						SceneVP->SetViewportSize(size.X, size.Y);
-						Window->SetViewportSizeDrivenByWindow(true);
+						{
+							FVector2D size = Window->GetSizeInScreen();
+							SceneVP->SetViewportSize(size.X, size.Y);
+							Window->SetViewportSizeDrivenByWindow(true);
+						}
 
 						if (stereoEnabled)
 						{
@@ -1305,6 +1306,7 @@ bool FOculusRiftHMD::DoEnableStereo(bool bStereo, bool bApplyToHmd)
 					{
 						// a special case when 'stereo on' or 'stereo hmd' is used in Direct mode. We must set proper window mode, otherwise
 						// it will be lost once window loses and regains the focus.
+						FVector2D size = Window->GetSizeInScreen();
 						FSystemResolution::RequestResolutionChange(size.X, size.Y, (stereoEnabled) ? EWindowMode::WindowedMirror : EWindowMode::Windowed);
 					}
 				}
diff --git a/Engine/Plugins/Runtime/OculusRift/Source/OculusRift/Private/OculusRiftRender.cpp b/Engine/Plugins/Runtime/OculusRift/Source/OculusRift/Private/OculusRiftRender.cpp
--- a/Engine/Plugins/Runtime/OculusRift/Source/OculusRift/Private/OculusRiftRender.cpp
+++ b/Engine/Plugins/Runtime/OculusRift/Source/OculusRift/Private/OculusRiftRender.cpp
@@ -90,7 +90,7 @@ void FOculusRiftHMD::PrecalculatePostProcess_NoLock()
 #endif
 }
 
-void FOculusRiftHMD::DrawDistortionMesh_RenderThread(FRenderingCompositePassContext& Context, const FSceneView& View, const FIntPoint& TextureSize)
+void FOculusRiftHMD::DrawDistortionMesh_RenderThread(FRenderingCompositePassContext& Context, const FSceneView& View, const FIntPoint& InTextureSize)
 {
 #ifndef OVR_SDK_RENDERING
 	float ClipSpaceQuadZ = 0.0f;
diff --git a/Engine/Source/Editor/LevelEditor/Private/LevelEditorContextMenu.cpp b/Engine/Source/Editor/LevelEditor/Private/LevelEditorContextMenu.cpp
--- a/Engine/Source/Editor/LevelEditor/Private/LevelEditorContextMenu.cpp
+++ b/Engine/Source/Editor/LevelEditor/Private/LevelEditorContextMenu.cpp
@@ -152,8 +152,8 @@ TSharedPtr< SWidget > FLevelEditorContextMenu::BuildMenuWidget( TWeakPtr< SLevel
 
 void FLevelEditorContextMenu::FillMenu( FMenuBuilder& MenuBuilder, TWeakPtr<SLevelEditor> LevelEditor, LevelEditorMenuContext ContextType, TSharedPtr<FExtender> Extender )
 {
-	auto LevelEditorActions = LevelEditor.Pin()->GetLevelEditorActions().ToSharedRef();
-	MenuBuilder.PushCommandList(LevelEditorActions);
+	auto LevelEditorActionsList = LevelEditor.Pin()->GetLevelEditorActions().ToSharedRef();
+	MenuBuilder.PushCommandList(LevelEditorActionsList);
 
 	if (GEditor->GetSelectedComponentCount() > 0)
 	{
@@ -213,7 +213,7 @@ void FLevelEditorContextMenu::FillMenu( FMenuBuilder& MenuBuilder, TWeakPtr<SLev
 		{
 			if (MenuExtenderDelegates[i].IsBound())
 			{
-				Extenders.Add(MenuExtenderDelegates[i].Execute(LevelEditorActions, SelectedActors));
+				Extenders.Add(MenuExtenderDelegates[i].Execute(LevelEditorActionsList, SelectedActors));
 			}
 		}
 		MenuBuilder.PushExtender(FExtender::Combine(Extenders).ToSharedRef());
diff --git a/Engine/Source/Editor/LevelEditor/Private/LevelEditorMenu.cpp b/Engine/Source/Editor/LevelEditor/Private/LevelEditorMenu.cpp
--- a/Engine/Source/Editor/LevelEditor/Private/LevelEditorMenu.cpp
+++ b/Engine/Source/Editor/LevelEditor/Private/LevelEditorMenu.cpp
@@ -194,18 +194,17 @@ TSharedRef< SWidget > FLevelEditorMenu::MakeLevelEditorMenu( const TSharedPtr<FU
 					{
 						struct Local
 						{
-							static void MakeFavoriteLevelMenu(FMenuBuilder& MenuBuilder)
+							static void MakeFavoriteLevelMenu(FMenuBuilder& InMenuBuilder)
 							{
-								IMainFrameModule& MainFrameModule = FModuleManager::LoadModuleChecked<IMainFrameModule>("MainFrame");
-								const FMainMRUFavoritesList& RecentsAndFavorites = *MainFrameModule.GetMRUFavoritesList();
-								const int32 NumFavorites = RecentsAndFavorites.GetNumFavorites();
+								const FMainMRUFavoritesList& MRUFavorites = *FModuleManager::LoadModuleChecked<IMainFrameModule>("MainFrame").GetMRUFavoritesList();
+								const int32 NumMRUFavorites = MRUFavorites.GetNumFavorites();
 
-								for (int32 CurFavoriteIndex = FavoritesToDisplayInMainMenu; CurFavoriteIndex < NumFavorites; ++CurFavoriteIndex)
+								for (int32 CurFavoriteIndex = FavoritesToDisplayInMainMenu; CurFavoriteIndex < NumMRUFavorites; ++CurFavoriteIndex)
 								{
-									const FString CurFavorite = FPaths::GetBaseFilename(RecentsAndFavorites.GetFavoritesItem(CurFavoriteIndex));
+									const FString CurFavorite = FPaths::GetBaseFilename(MRUFavorites.GetFavoritesItem(CurFavoriteIndex));
 									const bool bNoIndent = true;
 
-									MenuBuilder.AddWidget(
+									InMenuBuilder.AddWidget(
 										SNew(SFavouriteMenuEntry)
 										.LabelOverride(FText::FromString(FPaths::GetBaseFilename(CurFavorite)))
 										.OnOpenClickedDelegate(FExecuteAction::CreateStatic(&FLevelEditorActionCallbacks::OpenFavoriteFile, CurFavoriteIndex))
diff --git a/Engine/Source/Editor/LevelEditor/Private/LevelEditorToolBar.cpp b/Engine/Source/Editor/LevelEditor/Private/LevelEditorToolBar.cpp
--- a/Engine/Source/Editor/LevelEditor/Private/LevelEditorToolBar.cpp
+++ b/Engine/Source/Editor/LevelEditor/Private/LevelEditorToolBar.cpp
@@ -504,10 +504,10 @@ namespace LevelEditorActionHelpers
 		GameModeMenuSettings.LevelEditor = InLevelEditor;
 		GameModeMenuSettings.bIsProjectSettings = bInProjectSettings;
 
-		auto IsGameModeActive = [](TWeakPtr< SLevelEditor > InLevelEditor, bool bInProjectSettings)->bool
+		auto IsGameModeActive = [](TWeakPtr< SLevelEditor > InLevelEditorPtr, bool bInProjSettings)->bool
 		{
-			UClass* WorldSettingsGameMode = LevelEditorActionHelpers::GetGameModeClass(InLevelEditor, false);
-			if((WorldSettingsGameMode == nullptr) ^ bInProjectSettings ) //(WorldSettingsGameMode && !bInProjectSettings) || (!WorldSettingsGameMode && bInProjectSettings) )
+			UClass* WorldSettingsGameMode = LevelEditorActionHelpers::GetGameModeClass(InLevelEditorPtr, false);
+			if((WorldSettingsGameMode == nullptr) ^ bInProjSettings ) //(WorldSettingsGameMode && !bInProjectSettings) || (!WorldSettingsGameMode && bInProjectSettings) )
 			{
 				return false;
 			}
diff --git a/Engine/Source/Editor/LevelEditor/Private/SLevelViewport.cpp b/Engine/Source/Editor/LevelEditor/Private/SLevelViewport.cpp
--- a/Engine/Source/Editor/LevelEditor/Private/SLevelViewport.cpp
+++ b/Engine/Source/Editor/LevelEditor/Private/SLevelViewport.cpp
@@ -1944,7 +1944,7 @@ void SLevelViewport::SaveConfig(const FString& ConfigName)
 }
 
 
-FLevelEditorViewportInstanceSettings SLevelViewport::LoadLegacyConfigFromIni(const FString& ConfigKey, const FLevelEditorViewportInstanceSettings& InDefaultSettings)
+FLevelEditorViewportInstanceSettings SLevelViewport::LoadLegacyConfigFromIni(const FString& InConfigKey, const FLevelEditorViewportInstanceSettings& InDefaultSettings)
 {
 	FLevelEditorViewportInstanceSettings ViewportInstanceSettings = InDefaultSettings;
 
@@ -1952,7 +1952,7 @@ FLevelEditorViewportInstanceSettings SLevelViewport::LoadLegacyConfigFromIni(con
 
 	{
 		int32 ViewportTypeAsInt = ViewportInstanceSettings.ViewportType;
-		GConfig->GetInt(*IniSection, *(ConfigKey + TEXT(".Type")), ViewportTypeAsInt, GEditorUserSettingsIni);
+		GConfig->GetInt(*IniSection, *(InConfigKey + TEXT(".Type")), ViewportTypeAsInt, GEditorUserSettingsIni);
 		ViewportInstanceSettings.ViewportType = (ViewportTypeAsInt == -1 || ViewportTypeAsInt == 255) ? LVT_None : static_cast<ELevelViewportType>(ViewportTypeAsInt); // LVT_None used to be -1 or 255
 
 		if(ViewportInstanceSettings.ViewportType == LVT_None)
@@ -1961,18 +1961,18 @@ FLevelEditorViewportInstanceSettings SLevelViewport::LoadLegacyConfigFromIni(con
 		}
 	}
 	
-	GConfig->GetString(*IniSection, *(ConfigKey + TEXT(".EditorShowFlags")), ViewportInstanceSettings.EditorShowFlagsString, GEditorUserSettingsIni);
-	GConfig->GetString(*IniSection, *(ConfigKey + TEXT(".GameShowFlags")), ViewportInstanceSettings.GameShowFlagsString, GEditorUserSettingsIni);
+	GConfig->GetString(*IniSection, *(InConfigKey + TEXT(".EditorShowFlags")), ViewportInstanceSettings.EditorShowFlagsString, GEditorUserSettingsIni);
+	GConfig->GetString(*IniSection, *(InConfigKey + TEXT(".GameShowFlags")), ViewportInstanceSettings.GameShowFlagsString, GEditorUserSettingsIni);
 		
 	// A single view mode index has been deprecated in favor of separate perspective and orthographic settings
 	EViewModeIndex LegacyViewModeIndex = VMI_Unknown;
 	{
 		int32 LegacyVMIAsInt = VMI_Unknown;
-		GConfig->GetInt(*IniSection, *(ConfigKey+TEXT(".ViewModeIndex")), LegacyVMIAsInt, GEditorUserSettingsIni);
+		GConfig->GetInt(*IniSection, *(InConfigKey+TEXT(".ViewModeIndex")), LegacyVMIAsInt, GEditorUserSettingsIni);
 		LegacyViewModeIndex = (LegacyVMIAsInt == -1) ? VMI_Unknown : static_cast<EViewModeIndex>(LegacyVMIAsInt); // VMI_Unknown used to be -1
 	}
 
-	if(!GConfig->GetInt(*IniSection, *(ConfigKey+TEXT(".PerspViewModeIndex")), (int32&)ViewportInstanceSettings.PerspViewModeIndex, GEditorUserSettingsIni))
+	if(!GConfig->GetInt(*IniSection, *(InConfigKey+TEXT(".PerspViewModeIndex")), (int32&)ViewportInstanceSettings.PerspViewModeIndex, GEditorUserSettingsIni))
 	{
 		if(ViewportInstanceSettings.ViewportType == LVT_Perspective)
 		{
@@ -1986,7 +1986,7 @@ FLevelEditorViewportInstanceSettings SLevelViewport::LoadLegacyConfigFromIni(con
 		}
 	}
 
-	if(!GConfig->GetInt(*IniSection, *(ConfigKey+TEXT(".OrthoViewModeIndex")), (int32&)ViewportInstanceSettings.OrthoViewModeIndex, GEditorUserSettingsIni))
+	if(!GConfig->GetInt(*IniSection, *(InConfigKey+TEXT(".OrthoViewModeIndex")), (int32&)ViewportInstanceSettings.OrthoViewModeIndex, GEditorUserSettingsIni))
 	{
 		// Default to Brush Wireframe for an orthographic viewport
 		ViewportInstanceSettings.OrthoViewModeIndex = (ViewportInstanceSettings.ViewportType != LVT_Perspective && LegacyViewModeIndex != VMI_Unknown) ? LegacyViewModeIndex : VMI_BrushWireframe;
@@ -1994,7 +1994,7 @@ FLevelEditorViewportInstanceSettings SLevelViewport::LoadLegacyConfigFromIni(con
 
 	{
 		FString BufferVisualizationModeString;
-		if(GConfig->GetString(*IniSection, *(ConfigKey+TEXT(".BufferVisualizationMode")), BufferVisualizationModeString, GEditorUserSettingsIni))
+		if(GConfig->GetString(*IniSection, *(InConfigKey+TEXT(".BufferVisualizationMode")), BufferVisualizationModeString, GEditorUserSettingsIni))
 		{
 			ViewportInstanceSettings.BufferVisualizationMode = *BufferVisualizationModeString;
 		}
@@ -2002,16 +2002,16 @@ FLevelEditorViewportInstanceSettings SLevelViewport::LoadLegacyConfigFromIni(con
 	
 	{
 		FString ExposureSettingsString;
-		if(GConfig->GetString(*IniSection, *(ConfigKey + TEXT(".ExposureSettings")), ExposureSettingsString, GEditorUserSettingsIni))
+		if(GConfig->GetString(*IniSection, *(InConfigKey + TEXT(".ExposureSettings")), ExposureSettingsString, GEditorUserSettingsIni))
 		{
 			ViewportInstanceSettings.ExposureSettings.SetFromString(*ExposureSettingsString);
 		}
 	}
 
-	GConfig->GetBool(*IniSection, *(ConfigKey + TEXT(".bIsRealtime")), ViewportInstanceSettings.bIsRealtime, GEditorUserSettingsIni);
-	GConfig->GetBool(*IniSection, *(ConfigKey + TEXT(".bWantStats")), ViewportInstanceSettings.bShowStats, GEditorUserSettingsIni);
-	GConfig->GetBool(*IniSection, *(ConfigKey + TEXT(".bWantFPS")), ViewportInstanceSettings.bShowFPS_DEPRECATED, GEditorUserSettingsIni);
-	GConfig->GetFloat(*IniSection, *(ConfigKey + TEXT(".FOVAngle")), ViewportInstanceSettings.FOVAngle, GEditorUserSettingsIni);
+	GConfig->GetBool(*IniSection, *(InConfigKey + TEXT(".bIsRealtime")), ViewportInstanceSettings.bIsRealtime, GEditorUserSettingsIni);
+	GConfig->GetBool(*IniSection, *(InConfigKey + TEXT(".bWantStats")), ViewportInstanceSettings.bShowStats, GEditorUserSettingsIni);
+	GConfig->GetBool(*IniSection, *(InConfigKey + TEXT(".bWantFPS")), ViewportInstanceSettings.bShowFPS_DEPRECATED, GEditorUserSettingsIni);
+	GConfig->GetFloat(*IniSection, *(InConfigKey + TEXT(".FOVAngle")), ViewportInstanceSettings.FOVAngle, GEditorUserSettingsIni);
 
 	return ViewportInstanceSettings;
 }
diff --git a/Engine/Source/Editor/Matinee/Private/SMatineeRecorder.cpp b/Engine/Source/Editor/Matinee/Private/SMatineeRecorder.cpp
--- a/Engine/Source/Editor/Matinee/Private/SMatineeRecorder.cpp
+++ b/Engine/Source/Editor/Matinee/Private/SMatineeRecorder.cpp
@@ -130,8 +130,6 @@ void SMatineeRecorder::Construct(const FArguments& InArgs)
 
 	// The viewport widget needs an interface so it knows what should render
 	ViewportWidget->SetViewportInterface( Viewport.ToSharedRef() );
-
-	TWeakPtr<FLevelEditorViewportClient>(LevelViewportClient);
 }
 
 /** Gets the current image that should be displayed for the Record/Stop button based on the status of the InterpEditor. */
diff --git a/Engine/Source/Editor/OnlineBlueprintSupport/Private/K2Node_LatentOnlineCall.cpp b/Engine/Source/Editor/OnlineBlueprintSupport/Private/K2Node_LatentOnlineCall.cpp
--- a/Engine/Source/Editor/OnlineBlueprintSupport/Private/K2Node_LatentOnlineCall.cpp
+++ b/Engine/Source/Editor/OnlineBlueprintSupport/Private/K2Node_LatentOnlineCall.cpp
@@ -97,12 +97,12 @@ void UK2Node_LatentOnlineCall::GetMenuActions(FBlueprintActionDatabaseRegistrar&
 					UK2Node_LatentOnlineCall* AsyncTaskNode = CastChecked<UK2Node_LatentOnlineCall>(NewNode);
 					if (FunctionPtr.IsValid())
 					{
-						UFunction* Function = FunctionPtr.Get();
-						UObjectProperty* ReturnProperty = CastChecked<UObjectProperty>(Function->GetReturnProperty());
+						UFunction* Func = FunctionPtr.Get();
+						UObjectProperty* ReturnProp = CastChecked<UObjectProperty>(Func->GetReturnProperty());
 						
-						AsyncTaskNode->ProxyFactoryFunctionName = Function->GetFName();
-						AsyncTaskNode->ProxyFactoryClass        = Function->GetOuterUClass();
-						AsyncTaskNode->ProxyClass               = ReturnProperty->PropertyClass;
+						AsyncTaskNode->ProxyFactoryFunctionName = Func->GetFName();
+						AsyncTaskNode->ProxyFactoryClass        = Func->GetOuterUClass();
+						AsyncTaskNode->ProxyClass               = ReturnProp->PropertyClass;
 					}
 				};
 				
diff --git a/Engine/Source/Editor/UnrealEd/Private/Editor.cpp b/Engine/Source/Editor/UnrealEd/Private/Editor.cpp
--- a/Engine/Source/Editor/UnrealEd/Private/Editor.cpp
+++ b/Engine/Source/Editor/UnrealEd/Private/Editor.cpp
@@ -676,10 +676,10 @@ void UEditorEngine::Init(IEngineLoop* InEngineLoop)
 			TEXT("LocalizationDashboard")
 		};
 
-		FScopedSlowTask SlowTask(ARRAY_COUNT(ModuleNames));
+		FScopedSlowTask ModuleSlowTask(ARRAY_COUNT(ModuleNames));
 		for (const TCHAR* ModuleName : ModuleNames)
 		{
-			SlowTask.EnterProgressFrame(1);
+			ModuleSlowTask.EnterProgressFrame(1);
 			FModuleManager::Get().LoadModule(ModuleName);
 		}
 
@@ -7062,7 +7062,6 @@ namespace EditorUtilities
 					else
 					{
 						// If we found a match, update the target component and adjust the target index to the matching position
-						UActorComponent* FindTargetComponent = TargetComponents[ FindTargetComponentIndex ];
 						if( FindTargetComponent != NULL && SourceComponent->GetFName() == FindTargetComponent->GetFName() )
 						{
 							TargetComponent = FindTargetComponent;
diff --git a/Engine/Source/Editor/UnrealEd/Private/EditorServer.cpp b/Engine/Source/Editor/UnrealEd/Private/EditorServer.cpp
--- a/Engine/Source/Editor/UnrealEd/Private/EditorServer.cpp
+++ b/Engine/Source/Editor/UnrealEd/Private/EditorServer.cpp
@@ -1124,54 +1124,54 @@ void UEditorEngine::PostUndo (bool bSuccess)
 {
 	//Update the actor selection followed by the component selection if needed (note: order is important)
 		
-		//Get the list of all selected actors after the operation
-		TArray<AActor*> SelectedActors;
-		for (FSelectionIterator It(GEditor->GetSelectedActorIterator()); It; ++It)
+	//Get the list of all selected actors after the operation
+	TArray<AActor*> SelectedActors;
+	for (FSelectionIterator It(GEditor->GetSelectedActorIterator()); It; ++It)
+	{
+		AActor* Actor = CastChecked<AActor>(*It);
+		//if this actor is NOT in a hidden level add it to the list - otherwise de-select it
+		if (FLevelUtils::IsLevelLocked(Actor) == false)
 		{
-			AActor* Actor = CastChecked<AActor>(*It);
-			//if this actor is NOT in a hidden level add it to the list - otherwise de-select it
-			if (FLevelUtils::IsLevelLocked(Actor) == false)
-			{
-				SelectedActors.Add(Actor);
-			}
-			else
-			{
-				GetSelectedActors()->Select(Actor, false);
-			}
+			SelectedActors.Add(Actor);
 		}
+		else
+		{
+			GetSelectedActors()->Select(Actor, false);
+		}
+	}
 
-		USelection* Selection = GetSelectedActors();
-		Selection->BeginBatchSelectOperation();
+	USelection* ActorSelection = GetSelectedActors();
+	ActorSelection->BeginBatchSelectOperation();
 
-		//Deselect all of the actors that were selected prior to the operation
-		for (int32 OldSelectedActorIndex = OldSelectedActors.Num() - 1; OldSelectedActorIndex >= 0; --OldSelectedActorIndex)
-		{
-			AActor* Actor = OldSelectedActors[OldSelectedActorIndex];
+	//Deselect all of the actors that were selected prior to the operation
+	for (int32 OldSelectedActorIndex = OldSelectedActors.Num() - 1; OldSelectedActorIndex >= 0; --OldSelectedActorIndex)
+	{
+		AActor* Actor = OldSelectedActors[OldSelectedActorIndex];
 
-			//To stop us from unselecting and then reselecting again (causing two force update components, we will remove (from both lists) any object that was selected and should continue to be selected
-			int32 FoundIndex;
-			if (SelectedActors.Find(Actor, FoundIndex))
-			{
-				OldSelectedActors.RemoveAt(OldSelectedActorIndex);
-				SelectedActors.RemoveAt(FoundIndex);
-			}
-			else
-			{
-				SelectActor(Actor, false, false);//First false is to deselect, 2nd is to notify
-				Actor->UpdateComponentTransforms();
-			}
+		//To stop us from unselecting and then reselecting again (causing two force update components, we will remove (from both lists) any object that was selected and should continue to be selected
+		int32 FoundIndex;
+		if (SelectedActors.Find(Actor, FoundIndex))
+		{
+			OldSelectedActors.RemoveAt(OldSelectedActorIndex);
+			SelectedActors.RemoveAt(FoundIndex);
 		}
-
-		//Select all of the actors in SelectedActors
-		for (int32 SelectedActorIndex = 0; SelectedActorIndex < SelectedActors.Num(); ++SelectedActorIndex)
+		else
 		{
-			AActor* Actor = SelectedActors[SelectedActorIndex];
-			SelectActor(Actor, true, false);	//false is to stop notify which is done below if bOpWasSuccessful
+			SelectActor(Actor, false, false);//First false is to deselect, 2nd is to notify
 			Actor->UpdateComponentTransforms();
 		}
+	}
 
-		OldSelectedActors.Empty();
-		Selection->EndBatchSelectOperation();
+	//Select all of the actors in SelectedActors
+	for (int32 SelectedActorIndex = 0; SelectedActorIndex < SelectedActors.Num(); ++SelectedActorIndex)
+	{
+		AActor* Actor = SelectedActors[SelectedActorIndex];
+		SelectActor(Actor, true, false);	//false is to stop notify which is done below if bOpWasSuccessful
+		Actor->UpdateComponentTransforms();
+	}
+
+	OldSelectedActors.Empty();
+	ActorSelection->EndBatchSelectOperation();
 	
 	if (GetSelectedComponentCount() > 0)
 	{
@@ -1184,8 +1184,8 @@ void UEditorEngine::PostUndo (bool bSuccess)
 			SelectedComponents.Add(CastChecked<UActorComponent>(*It));
 		}
 		
-		USelection* Selection = GetSelectedComponents();
-		Selection->BeginBatchSelectOperation();
+		USelection* ComponentSelection = GetSelectedComponents();
+		ComponentSelection->BeginBatchSelectOperation();
 
 		//Deselect all of the actors that were selected prior to the operation
 		for (int32 OldSelectedComponentIndex = OldSelectedComponents.Num() - 1; OldSelectedComponentIndex >= 0; --OldSelectedComponentIndex)
@@ -1218,8 +1218,8 @@ void UEditorEngine::PostUndo (bool bSuccess)
 		OldSelectedComponents.Empty();
 
 		// We want to broadcast the component SelectionChangedEvent even if the selection didn't actually change
-		Selection->MarkBatchDirty();
-		Selection->EndBatchSelectOperation();
+		ComponentSelection->MarkBatchDirty();
+		ComponentSelection->EndBatchSelectOperation();
 	}
 }
 
@@ -1861,11 +1861,11 @@ bool UEditorEngine::ShouldAbortBecauseOfUnsavedWorld() const
 	// If an unsaved world exists that would be lost in a map transition, give the user the option to cancel a map load.
 
 	// First check if we have a world and it is dirty
-	UWorld* EditorWorld = GEditor->GetEditorWorldContext().World();
-	if (EditorWorld && EditorWorld->GetOutermost()->IsDirty())
+	UWorld* LevelEditorWorld = GEditor->GetEditorWorldContext().World();
+	if (LevelEditorWorld && LevelEditorWorld->GetOutermost()->IsDirty())
 	{
 		// Now check if the world is in a path that can be saved (otherwise it is in something like the transient package or temp)
-		const FString PackageName = EditorWorld->GetOutermost()->GetName();
+		const FString PackageName = LevelEditorWorld->GetOutermost()->GetName();
 		const bool bIncludeReadOnlyRoots = false;
 		if ( FPackageName::IsValidLongPackageName(PackageName, bIncludeReadOnlyRoots) )
 		{
@@ -1873,7 +1873,7 @@ bool UEditorEngine::ShouldAbortBecauseOfUnsavedWorld() const
 			if ( !FPackageName::DoesPackageExist(PackageName) )
 			{
 				// This world will be completely lost if a map transition happens. Warn the user that this is happening and ask him/her how to proceed.
-				if (EAppReturnType::Yes != FMessageDialog::Open(EAppMsgType::YesNo, FText::Format(NSLOCTEXT("UnrealEd", "Prompt_ThisActionWillDiscardWorldContinue", "The unsaved level {0} will be lost.  Continue?"), FText::FromString(EditorWorld->GetName()))))
+				if (EAppReturnType::Yes != FMessageDialog::Open(EAppMsgType::YesNo, FText::Format(NSLOCTEXT("UnrealEd", "Prompt_ThisActionWillDiscardWorldContinue", "The unsaved level {0} will be lost.  Continue?"), FText::FromString(LevelEditorWorld->GetName()))))
 				{
 					// User doesn't want to lose the world -- abort the load.
 					return true;
@@ -4383,18 +4383,17 @@ void UEditorEngine::MoveViewportCamerasToActor(const TArray<AActor*> &Actors, co
 				}
 				else
 				{
-					TInlineComponentArray<UPrimitiveComponent*> Components;
-					Actor->GetComponents(Components);
+					TInlineComponentArray<UPrimitiveComponent*> PrimitiveComponents(Actor);
 
-					for(int32 ComponentIndex = 0; ComponentIndex < Components.Num(); ++ComponentIndex)
+					for(int32 ComponentIndex = 0; ComponentIndex < PrimitiveComponents.Num(); ++ComponentIndex)
 					{
-						UPrimitiveComponent* PrimitiveComponent = Components[ComponentIndex];
+						UPrimitiveComponent* PrimitiveComponent = PrimitiveComponents[ComponentIndex];
 
 						if(PrimitiveComponent->IsRegistered())
 						{
 
 							// Some components can have huge bounds but are not visible.  Ignore these components unless it is the only component on the actor 
-							const bool bIgnore = Components.Num() > 1 && PrimitiveComponentTypesToIgnore.IndexOfByPredicate(ComponentTypeMatcher(PrimitiveComponent)) != INDEX_NONE;
+							const bool bIgnore = PrimitiveComponents.Num() > 1 && PrimitiveComponentTypesToIgnore.IndexOfByPredicate(ComponentTypeMatcher(PrimitiveComponent)) != INDEX_NONE;
 
 							if(!bIgnore)
 							{
diff --git a/Engine/Source/Editor/UnrealEd/Private/EditorViewportClient.cpp b/Engine/Source/Editor/UnrealEd/Private/EditorViewportClient.cpp
--- a/Engine/Source/Editor/UnrealEd/Private/EditorViewportClient.cpp
+++ b/Engine/Source/Editor/UnrealEd/Private/EditorViewportClient.cpp
@@ -401,11 +401,11 @@ void FEditorViewportClient::SetViewLocationForOrbiting(const FVector& LookAtPoin
 	SetLookAtLocation( LookAtPoint );
 }
 
-void FEditorViewportClient::SetInitialViewTransform(ELevelViewportType ViewportType, const FVector& ViewLocation, const FRotator& ViewRotation, float InOrthoZoom )
+void FEditorViewportClient::SetInitialViewTransform(ELevelViewportType InViewportType, const FVector& ViewLocation, const FRotator& ViewRotation, float InOrthoZoom )
 {
-	check(ViewportType < LVT_MAX);
+	check(InViewportType < LVT_MAX);
 
-	FViewportCameraTransform& ViewTransform = (ViewportType == LVT_Perspective) ? ViewTransformPerspective : ViewTransformOrthographic;
+	FViewportCameraTransform& ViewTransform = (InViewportType == LVT_Perspective) ? ViewTransformPerspective : ViewTransformOrthographic;
 
 	ViewTransform.SetLocation(ViewLocation);
 	ViewTransform.SetRotation(ViewRotation);
@@ -3229,6 +3229,7 @@ bool FEditorViewportClient::InputGesture(FViewport* InViewport, EGestureEvent::T
 	const FRotator& ViewRotation = GetViewRotation();
 
 	const bool LeftMouseButtonDown = InViewport->KeyState(EKeys::LeftMouseButton);
+	const bool RightMouseButtonDown = Viewport->KeyState(EKeys::RightMouseButton);
 
 	const ELevelViewportType LevelViewportType = GetViewportType();
 
@@ -3240,8 +3241,6 @@ bool FEditorViewportClient::InputGesture(FViewport* InViewport, EGestureEvent::T
 	case LVT_OrthoXZ:
 	case LVT_OrthoYZ:
 		{
-			const bool LeftMouseButtonDown = Viewport->KeyState(EKeys::LeftMouseButton);
-			const bool RightMouseButtonDown = Viewport->KeyState(EKeys::RightMouseButton);
 			if (GestureType == EGestureEvent::Scroll && !LeftMouseButtonDown && !RightMouseButtonDown)
 			{
 				const float UnitsPerPixel = GetOrthoUnitsPerPixel(Viewport);
diff --git a/Engine/Source/Editor/UnrealEd/Private/Fbx/FbxStaticMeshImport.cpp b/Engine/Source/Editor/UnrealEd/Private/Fbx/FbxStaticMeshImport.cpp
--- a/Engine/Source/Editor/UnrealEd/Private/Fbx/FbxStaticMeshImport.cpp
+++ b/Engine/Source/Editor/UnrealEd/Private/Fbx/FbxStaticMeshImport.cpp
@@ -1237,16 +1237,16 @@ UStaticMesh* UnFbx::FFbxImporter::ImportStaticMeshAsSingle(UObject* InParent, TA
 		FMeshSectionInfoMap OldSectionInfoMap = StaticMesh->SectionInfoMap;
 		StaticMesh->SectionInfoMap.Clear();
 		// fix up section data
-		for (int32 LODIndex = 0; LODIndex<StaticMesh->RenderData->LODResources.Num(); ++LODIndex)
+		for (int32 LODResoureceIndex = 0; LODResoureceIndex<StaticMesh->RenderData->LODResources.Num(); ++LODResoureceIndex)
 		{
-			FStaticMeshLODResources& LOD = StaticMesh->RenderData->LODResources[LODIndex];
+			FStaticMeshLODResources& LOD = StaticMesh->RenderData->LODResources[LODResoureceIndex];
 			int32 NumSections = LOD.Sections.Num();
 			for(int32 SectionIndex = 0; SectionIndex < NumSections; ++SectionIndex)
 			{
-				FMeshSectionInfo Info = OldSectionInfoMap.Get(LODIndex, SectionIndex);
+				FMeshSectionInfo Info = OldSectionInfoMap.Get(LODResoureceIndex, SectionIndex);
 				if (StaticMesh->Materials.IsValidIndex(Info.MaterialIndex))
 				{
-					StaticMesh->SectionInfoMap.Set(LODIndex, SectionIndex, Info);
+					StaticMesh->SectionInfoMap.Set(LODResoureceIndex, SectionIndex, Info);
 				}
 			}
 		}
diff --git a/Engine/Source/Runtime/Engine/Private/Actor.cpp b/Engine/Source/Runtime/Engine/Private/Actor.cpp
--- a/Engine/Source/Runtime/Engine/Private/Actor.cpp
+++ b/Engine/Source/Runtime/Engine/Private/Actor.cpp
@@ -191,12 +191,12 @@ bool AActor::CheckActorComponents()
 
 void AActor::ResetOwnedComponents()
 {
-	TArray<UObject*> Children;
+	TArray<UObject*> ActorChildren;
 	OwnedComponents.Empty();
 	ReplicatedComponents.Empty();
-	GetObjectsWithOuter(this, Children, true, RF_PendingKill);
+	GetObjectsWithOuter(this, ActorChildren, true, RF_PendingKill);
 
-	for (UObject* Child : Children)
+	for (UObject* Child : ActorChildren)
 	{
 		UActorComponent* Component = Cast<UActorComponent>(Child);
 		if (Component)
diff --git a/Engine/Source/Runtime/Engine/Private/ActorEditor.cpp b/Engine/Source/Runtime/Engine/Private/ActorEditor.cpp
--- a/Engine/Source/Runtime/Engine/Private/ActorEditor.cpp
+++ b/Engine/Source/Runtime/Engine/Private/ActorEditor.cpp
@@ -263,23 +263,23 @@ void AActor::DebugShowOneComponentHierarchy( USceneComponent* SceneComp, int32&
 AActor::FActorTransactionAnnotation::FActorTransactionAnnotation(const AActor* Actor)
 	: ComponentInstanceData(Actor)
 {
-	USceneComponent* RootComponent = Actor->GetRootComponent();
-	if (RootComponent && RootComponent->IsCreatedByConstructionScript())
+	USceneComponent* ActorRootComponent = Actor->GetRootComponent();
+	if (ActorRootComponent && ActorRootComponent->IsCreatedByConstructionScript())
 	{
 		bRootComponentDataCached = true;
-		RootComponentData.Transform = RootComponent->ComponentToWorld;
-		RootComponentData.Transform.SetTranslation(RootComponent->GetComponentLocation()); // take into account any custom location
+		RootComponentData.Transform = ActorRootComponent->ComponentToWorld;
+		RootComponentData.Transform.SetTranslation(ActorRootComponent->GetComponentLocation()); // take into account any custom location
 
-		if (RootComponent->AttachParent)
+		if (ActorRootComponent->AttachParent)
 		{
-			RootComponentData.AttachedParentInfo.Actor = RootComponent->AttachParent->GetOwner();
-			RootComponentData.AttachedParentInfo.AttachParent = RootComponent->AttachParent;
-			RootComponentData.AttachedParentInfo.AttachParentName = RootComponent->AttachParent->GetFName();
-			RootComponentData.AttachedParentInfo.SocketName = RootComponent->AttachSocketName;
-			RootComponentData.AttachedParentInfo.RelativeTransform = RootComponent->GetRelativeTransform();
+			RootComponentData.AttachedParentInfo.Actor = ActorRootComponent->AttachParent->GetOwner();
+			RootComponentData.AttachedParentInfo.AttachParent = ActorRootComponent->AttachParent;
+			RootComponentData.AttachedParentInfo.AttachParentName = ActorRootComponent->AttachParent->GetFName();
+			RootComponentData.AttachedParentInfo.SocketName = ActorRootComponent->AttachSocketName;
+			RootComponentData.AttachedParentInfo.RelativeTransform = ActorRootComponent->GetRelativeTransform();
 		}
 
-		for (USceneComponent* AttachChild : RootComponent->AttachChildren)
+		for (USceneComponent* AttachChild : ActorRootComponent->AttachChildren)
 		{
 			AActor* ChildOwner = (AttachChild ? AttachChild->GetOwner() : NULL);
 			if (ChildOwner && ChildOwner != Actor)
diff --git a/Engine/Source/Runtime/Engine/Private/Animation/AnimSequence.cpp b/Engine/Source/Runtime/Engine/Private/Animation/AnimSequence.cpp
--- a/Engine/Source/Runtime/Engine/Private/Animation/AnimSequence.cpp
+++ b/Engine/Source/Runtime/Engine/Private/Animation/AnimSequence.cpp
@@ -861,9 +861,9 @@ void UAnimSequence::GetAnimationPose(FTransformArrayA2& OutAtoms, const FBoneCon
 	}
 }
 
-void UAnimSequence::ResetRootBoneForRootMotion(FTransformArrayA2& BoneTransforms, const FBoneContainer& RequiredBones, ERootMotionRootLock::Type RootMotionRootLock) const
+void UAnimSequence::ResetRootBoneForRootMotion(FTransformArrayA2& BoneTransforms, const FBoneContainer& RequiredBones, ERootMotionRootLock::Type InRootMotionRootLock) const
 {
-	switch (RootMotionRootLock)
+	switch (InRootMotionRootLock)
 	{
 		case ERootMotionRootLock::AnimFirstFrame: BoneTransforms[0] = ExtractRootTrackTransform(0.f, &RequiredBones); break;
 		case ERootMotionRootLock::Zero: BoneTransforms[0] = FTransform::Identity; break;
diff --git a/Engine/Source/Runtime/Engine/Private/AudioDerivedData.cpp b/Engine/Source/Runtime/Engine/Private/AudioDerivedData.cpp
--- a/Engine/Source/Runtime/Engine/Private/AudioDerivedData.cpp
+++ b/Engine/Source/Runtime/Engine/Private/AudioDerivedData.cpp
@@ -990,7 +990,6 @@ void USoundWave::BeginCachePlatformData()
 	ITargetPlatformManagerModule* TPM = GetTargetPlatformManager();
 	if (TPM && (TPM->RestrictFormatsToRuntimeOnly() == false))
 	{
-		ITargetPlatformManagerModule* TPM = GetTargetPlatformManager();
 		TArray<ITargetPlatform*> Platforms = TPM->GetActiveTargetPlatforms();
 		// Cache for all the audio formats that the cooking target requires
 		for (int32 FormatIndex = 0; FormatIndex < Platforms.Num(); FormatIndex++)
diff --git a/Engine/Source/Runtime/Engine/Private/Brush.cpp b/Engine/Source/Runtime/Engine/Private/Brush.cpp
--- a/Engine/Source/Runtime/Engine/Private/Brush.cpp
+++ b/Engine/Source/Runtime/Engine/Private/Brush.cpp
@@ -140,7 +140,6 @@ void ABrush::SetIsTemporarilyHiddenInEditor( bool bIsHidden )
 
 			if (bAnySurfaceWasFound)
 			{
-				auto* Level = GetLevel();
 				Level->UpdateModelComponents();
 				Level->InvalidateModelSurface();
 			}
diff --git a/Engine/Source/Runtime/Engine/Private/LevelActor.cpp b/Engine/Source/Runtime/Engine/Private/LevelActor.cpp
--- a/Engine/Source/Runtime/Engine/Private/LevelActor.cpp
+++ b/Engine/Source/Runtime/Engine/Private/LevelActor.cpp
@@ -778,7 +778,7 @@ void UWorld::LoadSecondaryLevels(bool bForce, TSet<FString>* CookedPackages)
 
 
 				bool bAlreadyLoaded = false;
-				UPackage* const LevelPackage = FindObject<UPackage>(NULL, *StreamingLevelWorldAssetPackageName,true);
+				UPackage* LevelPackage = FindObject<UPackage>(NULL, *StreamingLevelWorldAssetPackageName,true);
 				// don't need to do any extra work if the level is already loaded
 				if ( LevelPackage && LevelPackage->IsFullyLoaded() ) 
 				{
@@ -793,7 +793,7 @@ void UWorld::LoadSecondaryLevels(bool bForce, TSet<FString>* CookedPackages)
 					if( FPackageName::IsShortPackageName(StreamingLevelWorldAssetPackageFName) == false )
 					{
 						ULevel::StreamedLevelsOwningWorld.Add(StreamingLevelWorldAssetPackageFName, this);
-						UPackage* const LevelPackage = LoadPackage( NULL, *StreamingLevelWorldAssetPackageName, LOAD_None );
+						LevelPackage = LoadPackage( NULL, *StreamingLevelWorldAssetPackageName, LOAD_None );
 						ULevel::StreamedLevelsOwningWorld.Remove(StreamingLevelWorldAssetPackageFName);
 
 						if( LevelPackage )
diff --git a/Engine/Source/Runtime/Engine/Private/SoundBase.cpp b/Engine/Source/Runtime/Engine/Private/SoundBase.cpp
--- a/Engine/Source/Runtime/Engine/Private/SoundBase.cpp
+++ b/Engine/Source/Runtime/Engine/Private/SoundBase.cpp
@@ -41,7 +41,7 @@ float USoundBase::GetMaxAudibleDistance()
 	return 0.f;
 }
 
-bool USoundBase::IsAudibleSimple( const FVector Location, USoundAttenuation* AttenuationSettings )
+bool USoundBase::IsAudibleSimple( const FVector Location, USoundAttenuation* InAttenuationSettings )
 {
 	// No audio device means no listeners to check against
 	if( !GEngine || !GEngine->GetAudioDevice() )
@@ -56,7 +56,7 @@ bool USoundBase::IsAudibleSimple( const FVector Location, USoundAttenuation* Att
 	}
 
 	// Is this SourceActor within the MaxAudibleDistance of any of the listeners?
-	float MaxAudibleDistance = AttenuationSettings != nullptr ? AttenuationSettings->Attenuation.GetMaxDimension() : GetMaxAudibleDistance();
+	float MaxAudibleDistance = InAttenuationSettings  != nullptr ? InAttenuationSettings ->Attenuation.GetMaxDimension() : GetMaxAudibleDistance();
 	return GEngine->GetAudioDevice()->LocationIsAudible( Location, MaxAudibleDistance );
 }
 
diff --git a/Engine/Source/Runtime/GameplayAbilities/Private/Abilities/Tasks/AbilityTask_WaitGameplayTagBase.cpp b/Engine/Source/Runtime/GameplayAbilities/Private/Abilities/Tasks/AbilityTask_WaitGameplayTagBase.cpp
--- a/Engine/Source/Runtime/GameplayAbilities/Private/Abilities/Tasks/AbilityTask_WaitGameplayTagBase.cpp
+++ b/Engine/Source/Runtime/GameplayAbilities/Private/Abilities/Tasks/AbilityTask_WaitGameplayTagBase.cpp
@@ -22,7 +22,7 @@ void UAbilityTask_WaitGameplayTag::Activate()
 	}
 }
 
-void UAbilityTask_WaitGameplayTag::GameplayTagCallback(const FGameplayTag Tag, int32 NewCount)
+void UAbilityTask_WaitGameplayTag::GameplayTagCallback(const FGameplayTag InTag, int32 NewCount)
 {
 }
 
diff --git a/Engine/Source/Runtime/MovieSceneCoreTypes/Private/Instances/MovieSceneAudioTrackInstance.cpp b/Engine/Source/Runtime/MovieSceneCoreTypes/Private/Instances/MovieSceneAudioTrackInstance.cpp
--- a/Engine/Source/Runtime/MovieSceneCoreTypes/Private/Instances/MovieSceneAudioTrackInstance.cpp
+++ b/Engine/Source/Runtime/MovieSceneCoreTypes/Private/Instances/MovieSceneAudioTrackInstance.cpp
@@ -28,23 +28,22 @@ void FMovieSceneAudioTrackInstance::Update( float Position, float LastPosition,
 			{
 				UMovieSceneAudioSection* AudioSection = Cast<UMovieSceneAudioSection>(AudioSections[i]);
 				int32 SectionIndex = AudioSection->GetRowIndex();
-				TArray<UMovieSceneAudioSection*>& AudioSections = AudioSectionsBySectionIndex.FindOrAdd(SectionIndex);
-				AudioSections.Add(AudioSection);
+				AudioSectionsBySectionIndex.FindOrAdd(SectionIndex).Add(AudioSection);
 			}
 
 			for (TMap<int32, TArray<UMovieSceneAudioSection*> >::TIterator It(AudioSectionsBySectionIndex); It; ++It)
 			{
 				int32 RowIndex = It.Key();
-				TArray<UMovieSceneAudioSection*>& AudioSections = It.Value();
+				TArray<UMovieSceneAudioSection*>& MovieSceneAudioSections = It.Value();
 
 				for (int32 ActorIndex = 0; ActorIndex < Actors.Num(); ++ActorIndex)
 				{
 					UAudioComponent* Component = GetAudioComponent(Actors[ActorIndex], RowIndex);
 
 					bool bComponentIsPlaying = false;
-					for (int32 i = 0; i < AudioSections.Num(); ++i)
+					for (int32 i = 0; i < MovieSceneAudioSections.Num(); ++i)
 					{
-						UMovieSceneAudioSection* AudioSection = AudioSections[i];
+						UMovieSceneAudioSection* AudioSection = MovieSceneAudioSections[i];
 						if (AudioSection->IsTimeWithinAudioRange(Position))
 						{
 							if (!AudioSection->IsTimeWithinAudioRange(LastPosition) || !Component->IsPlaying())
diff --git a/Engine/Source/Runtime/Renderer/Private/PostProcess/PostProcessing.cpp b/Engine/Source/Runtime/Renderer/Private/PostProcess/PostProcessing.cpp
--- a/Engine/Source/Runtime/Renderer/Private/PostProcess/PostProcessing.cpp
+++ b/Engine/Source/Runtime/Renderer/Private/PostProcess/PostProcessing.cpp
@@ -1938,7 +1938,6 @@ void FPostProcessing::ProcessES2(FRHICommandListImmediate& RHICmdList, FViewInfo
 			}
 
 			// Mobile temporal AA is done after tonemapping.
-			FIntPoint PrePostSourceViewportSize = View.ViewRect.Size();
 			FRenderingCompositePass* PostProcessAa = Context.Graph.RegisterPass(new(FMemStack::Get()) FRCPassPostProcessAaES2());
 			PostProcessAa->SetInput(ePId_Input0, Context.FinalOutput);
 			PostProcessAa->SetInput(ePId_Input1, PostProcessPrior);
diff --git a/Engine/Source/Runtime/UMG/Private/Components/BorderSlot.cpp b/Engine/Source/Runtime/UMG/Private/Components/BorderSlot.cpp
--- a/Engine/Source/Runtime/UMG/Private/Components/BorderSlot.cpp
+++ b/Engine/Source/Runtime/UMG/Private/Components/BorderSlot.cpp
@@ -74,19 +74,19 @@ void UBorderSlot::PostEditChangeProperty(FPropertyChangedEvent& PropertyChangedE
 		{
 			FName PropertyName = PropertyChangedEvent.Property->GetFName();
 
-			if ( UBorder* Border = CastChecked<UBorder>(Parent) )
+			if ( UBorder* ParentBorder = CastChecked<UBorder>(Parent) )
 			{
 				if ( PropertyName == "Padding" )
 				{
-					FObjectEditorUtils::MigratePropertyValue(this, "Padding", Border, "Padding");
+					FObjectEditorUtils::MigratePropertyValue(this, "Padding", ParentBorder, "Padding");
 				}
 				else if ( PropertyName == "HorizontalAlignment" )
 				{
-					FObjectEditorUtils::MigratePropertyValue(this, "HorizontalAlignment", Border, "HorizontalAlignment");
+					FObjectEditorUtils::MigratePropertyValue(this, "HorizontalAlignment", ParentBorder, "HorizontalAlignment");
 				}
 				else if ( PropertyName == "VerticalAlignment" )
 				{
-					FObjectEditorUtils::MigratePropertyValue(this, "VerticalAlignment", Border, "VerticalAlignment");
+					FObjectEditorUtils::MigratePropertyValue(this, "VerticalAlignment", ParentBorder, "VerticalAlignment");
 				}
 			}
 		}
diff --git a/Engine/Source/Runtime/UMG/Private/Components/ComboBoxString.cpp b/Engine/Source/Runtime/UMG/Private/Components/ComboBoxString.cpp
--- a/Engine/Source/Runtime/UMG/Private/Components/ComboBoxString.cpp
+++ b/Engine/Source/Runtime/UMG/Private/Components/ComboBoxString.cpp
@@ -31,9 +31,9 @@ void UComboBoxString::PostLoad()
 	Super::PostLoad();
 
 	// Initialize the set of options from the default set only once.
-	for ( const FString& DefaultOptions : DefaultOptions )
+	for ( const FString& DefaultOption : DefaultOptions )
 	{
-		AddOption(DefaultOptions);
+		AddOption(DefaultOption);
 	}
 }
 
diff --git a/Engine/Source/Runtime/UMG/Private/Components/EditableText.cpp b/Engine/Source/Runtime/UMG/Private/Components/EditableText.cpp
--- a/Engine/Source/Runtime/UMG/Private/Components/EditableText.cpp
+++ b/Engine/Source/Runtime/UMG/Private/Components/EditableText.cpp
@@ -116,14 +116,14 @@ void UEditableText::SetIsReadOnly(bool InbIsReadyOnly)
 	}
 }
 
-void UEditableText::HandleOnTextChanged(const FText& Text)
+void UEditableText::HandleOnTextChanged(const FText& InText)
 {
-	OnTextChanged.Broadcast(Text);
+	OnTextChanged.Broadcast(InText);
 }
 
-void UEditableText::HandleOnTextCommitted(const FText& Text, ETextCommit::Type CommitMethod)
+void UEditableText::HandleOnTextCommitted(const FText& InText, ETextCommit::Type CommitMethod)
 {
-	OnTextCommitted.Broadcast(Text, CommitMethod);
+	OnTextCommitted.Broadcast(InText, CommitMethod);
 }
 
 void UEditableText::PostLoad()
diff --git a/Engine/Source/Runtime/UMG/Private/Components/MultiLineEditableText.cpp b/Engine/Source/Runtime/UMG/Private/Components/MultiLineEditableText.cpp
--- a/Engine/Source/Runtime/UMG/Private/Components/MultiLineEditableText.cpp
+++ b/Engine/Source/Runtime/UMG/Private/Components/MultiLineEditableText.cpp
@@ -83,14 +83,14 @@ void UMultiLineEditableText::SetText(FText InText)
 	}
 }
 
-void UMultiLineEditableText::HandleOnTextChanged(const FText& Text)
+void UMultiLineEditableText::HandleOnTextChanged(const FText& InText)
 {
-	OnTextChanged.Broadcast(Text);
+	OnTextChanged.Broadcast(InText);
 }
 
-void UMultiLineEditableText::HandleOnTextCommitted(const FText& Text, ETextCommit::Type CommitMethod)
+void UMultiLineEditableText::HandleOnTextCommitted(const FText& InText, ETextCommit::Type CommitMethod)
 {
-	OnTextCommitted.Broadcast(Text, CommitMethod);
+	OnTextCommitted.Broadcast(InText, CommitMethod);
 }
 
 #if WITH_EDITOR
diff --git a/Engine/Source/Runtime/UMG/Private/Components/MultiLineEditableTextBox.cpp b/Engine/Source/Runtime/UMG/Private/Components/MultiLineEditableTextBox.cpp
--- a/Engine/Source/Runtime/UMG/Private/Components/MultiLineEditableTextBox.cpp
+++ b/Engine/Source/Runtime/UMG/Private/Components/MultiLineEditableTextBox.cpp
@@ -97,14 +97,14 @@ void UMultiLineEditableTextBox::SetError(FText InError)
 	}
 }
 
-void UMultiLineEditableTextBox::HandleOnTextChanged(const FText& Text)
+void UMultiLineEditableTextBox::HandleOnTextChanged(const FText& InText)
 {
-	OnTextChanged.Broadcast(Text);
+	OnTextChanged.Broadcast(InText);
 }
 
-void UMultiLineEditableTextBox::HandleOnTextCommitted(const FText& Text, ETextCommit::Type CommitMethod)
+void UMultiLineEditableTextBox::HandleOnTextCommitted(const FText& InText, ETextCommit::Type CommitMethod)
 {
-	OnTextCommitted.Broadcast(Text, CommitMethod);
+	OnTextCommitted.Broadcast(InText, CommitMethod);
 }
 
 void UMultiLineEditableTextBox::PostLoad()
diff --git a/Engine/Source/Runtime/UMG/Private/Components/WidgetComponent.cpp b/Engine/Source/Runtime/UMG/Private/Components/WidgetComponent.cpp
--- a/Engine/Source/Runtime/UMG/Private/Components/WidgetComponent.cpp
+++ b/Engine/Source/Runtime/UMG/Private/Components/WidgetComponent.cpp
@@ -392,10 +392,10 @@ void UWidgetComponent::OnRegister()
 						GameViewportWidget->SetCustomHitTestPath(CustomHitTestPath);
 					}
 
-					TSharedPtr<FWidget3DHitTester> WidgetHitTester = StaticCastSharedPtr<FWidget3DHitTester>(CustomHitTestPath);
-					if ( WidgetHitTester->GetWorld() == GetWorld() )
+					TSharedPtr<FWidget3DHitTester> Widget3DHitTester = StaticCastSharedPtr<FWidget3DHitTester>(CustomHitTestPath);
+					if ( Widget3DHitTester->GetWorld() == GetWorld() )
 					{
-						WidgetHitTester->RegisterWidgetComponent(this);
+						Widget3DHitTester->RegisterWidgetComponent(this);
 					}
 				}
 			}
@@ -840,9 +840,9 @@ TArray<FWidgetAndPointer> UWidgetComponent::GetHitWidgetPath( const FHitResult&
 
 	TArray<FWidgetAndPointer> ArrangedWidgets = HitTestGrid->GetBubblePath( LocalHitLocation, 0.0f, bIgnoreEnabledStatus );
 
-	for( FWidgetAndPointer& Widget : ArrangedWidgets )
+	for( FWidgetAndPointer& ArrangedWidget : ArrangedWidgets )
 	{
-		Widget.PointerPosition = VirtualMouseCoordinate;
+		ArrangedWidget.PointerPosition = VirtualMouseCoordinate;
 	}
 
 	return ArrangedWidgets;
