Making sure there's no outstanding load request when PIE begins.

Pre-emptive fix for UE-9775

diff --git a/Engine/Source/Editor/UnrealEd/Private/PlayLevel.cpp b/Engine/Source/Editor/UnrealEd/Private/PlayLevel.cpp
--- a/Engine/Source/Editor/UnrealEd/Private/PlayLevel.cpp
+++ b/Engine/Source/Editor/UnrealEd/Private/PlayLevel.cpp
@@ -1895,6 +1895,9 @@ void UEditorEngine::PlayInEditor( UWorld* InWorld, bool bInSimulateInEditor )
 		GLevelEditorModeTools().DeactivateMode( FBuiltinEditorModes::EM_InterpEdit );
 	}
 
+	// Make sure there's no outstanding load requests
+	FlushAsyncLoading();
+
 	FBlueprintEditorUtils::FindAndSetDebuggableBlueprintInstances();
 
 	FEditorDelegates::BeginPIE.Broadcast(bInSimulateInEditor);
