FBX: UTF8 support for name of meshes, bones, materiasl, and morph targets. FBX officially support UTF8 but FBX importer uses ANSI char. (change ANSI_TO_TCHAR to UTF8_TO_TCHAR and TCHAR_TO_ANSI to TCHAR_TO_UTF8).

diff --git a/Engine/Source/Editor/UnrealEd/Private/Fbx/FbxLightImport.cpp b/Engine/Source/Editor/UnrealEd/Private/Fbx/FbxLightImport.cpp
--- a/Engine/Source/Editor/UnrealEd/Private/Fbx/FbxLightImport.cpp
+++ b/Engine/Source/Editor/UnrealEd/Private/Fbx/FbxLightImport.cpp
@@ -19,7 +19,7 @@ using namespace UnFbx;
 ALight* FFbxImporter::CreateLight(FbxLight* InLight, UWorld* InWorld)
 {
 	ALight* UnrealLight = NULL;
-	FString ActorName = ANSI_TO_TCHAR(MakeName(InLight->GetName()));
+	FString ActorName = UTF8_TO_TCHAR(MakeName(InLight->GetName()));
 
 	// Create the light actor.
 	UClass* LightClass = NULL;
@@ -132,7 +132,7 @@ bool FFbxImporter::FillLightComponent(FbxLight* Light, ULightComponent* UnrealLi
 ACameraActor* FFbxImporter::CreateCamera(FbxCamera* InCamera, UWorld* InWorld)
 {
 	ACameraActor* UnrealCamera = NULL;
-	FString ActorName = ANSI_TO_TCHAR(MakeName(InCamera->GetName()));
+	FString ActorName = UTF8_TO_TCHAR(MakeName(InCamera->GetName()));
 	FActorSpawnParameters SpawnInfo;
 	SpawnInfo.Name = *ActorName;
 	UnrealCamera = InWorld->SpawnActor<ACameraActor>( SpawnInfo );
diff --git a/Engine/Source/Editor/UnrealEd/Private/Fbx/FbxMainExport.cpp b/Engine/Source/Editor/UnrealEd/Private/Fbx/FbxMainExport.cpp
--- a/Engine/Source/Editor/UnrealEd/Private/Fbx/FbxMainExport.cpp
+++ b/Engine/Source/Editor/UnrealEd/Private/Fbx/FbxMainExport.cpp
@@ -126,10 +126,10 @@ void FFbxExporter::CreateDocument()
 	SceneInfo->mSubject = "Export FBX meshes from Unreal";
 	SceneInfo->Original_ApplicationVendor.Set( "Epic Games" );
 	SceneInfo->Original_ApplicationName.Set( "Unreal Engine" );
-	SceneInfo->Original_ApplicationVersion.Set( TCHAR_TO_ANSI(*GEngineVersion.ToString()) );
+	SceneInfo->Original_ApplicationVersion.Set( TCHAR_TO_UTF8(*GEngineVersion.ToString()) );
 	SceneInfo->LastSaved_ApplicationVendor.Set( "Epic Games" );
 	SceneInfo->LastSaved_ApplicationName.Set( "Unreal Engine" );
-	SceneInfo->LastSaved_ApplicationVersion.Set( TCHAR_TO_ANSI(*GEngineVersion.ToString()) );
+	SceneInfo->LastSaved_ApplicationVersion.Set( TCHAR_TO_UTF8(*GEngineVersion.ToString()) );
 
 	Scene->SetSceneInfo(SceneInfo);
 	
@@ -338,7 +338,7 @@ void FFbxExporter::ExportLight( ALight* Actor, AMatineeActor* InMatineeActor )
 	FString FbxNodeName = GetActorNodeName(Actor, InMatineeActor);
 
 	// Export the basic light information
-	FbxLight* Light = FbxLight::Create(Scene, TCHAR_TO_ANSI(*FbxNodeName));
+	FbxLight* Light = FbxLight::Create(Scene, TCHAR_TO_UTF8(*FbxNodeName));
 	Light->Intensity.Set(BaseLight->Intensity);
 	Light->Color.Set(Converter.ConvertToFbxColor(BaseLight->LightColor));
 	
@@ -400,7 +400,7 @@ void FFbxExporter::ExportCamera( ACameraActor* Actor, AMatineeActor* InMatineeAc
 	FString FbxNodeName = GetActorNodeName(Actor, NULL);
 
 	// Create a properly-named FBX camera structure and instantiate it in the FBX scene graph
-	FbxCamera* Camera = FbxCamera::Create(Scene, TCHAR_TO_ANSI(*FbxNodeName));
+	FbxCamera* Camera = FbxCamera::Create(Scene, TCHAR_TO_UTF8(*FbxNodeName));
 
 	// Export the view area information
 	Camera->ProjectionType.Set(FbxCamera::ePerspective);
@@ -441,7 +441,7 @@ void FFbxExporter::ExportBrush(ABrush* Actor, UModel* InModel, bool bConvertToSt
  		Scene->GetRootNode()->AddChild(FbxActor);
  
  		// Export the mesh information
- 		ExportModel(Model, FbxActor, TCHAR_TO_ANSI(*Actor->GetName()));
+ 		ExportModel(Model, FbxActor, TCHAR_TO_UTF8(*Actor->GetName()));
  	}
  	else
  	{
@@ -804,7 +804,7 @@ void FFbxExporter::ExportStaticMesh( UStaticMesh* StaticMesh, const TArray<UMate
 	FString MeshName;
 	StaticMesh->GetName(MeshName);
 	FStaticMeshLODResources& RenderMesh = StaticMesh->GetLODForExport(0);
-	FbxNode* MeshNode = FbxNode::Create(Scene, TCHAR_TO_ANSI(*MeshName));
+	FbxNode* MeshNode = FbxNode::Create(Scene, TCHAR_TO_UTF8(*MeshName));
 	Scene->GetRootNode()->AddChild(MeshNode);
 	ExportStaticMeshToFbx(StaticMesh, RenderMesh, *MeshName, MeshNode, -1, NULL, MaterialOrder );
 }
@@ -816,7 +816,7 @@ void FFbxExporter::ExportStaticMeshLightMap( UStaticMesh* StaticMesh, int32 LODI
 	FString MeshName;
 	StaticMesh->GetName(MeshName);
 	FStaticMeshLODResources& RenderMesh = StaticMesh->GetLODForExport(LODIndex);
-	FbxNode* MeshNode = FbxNode::Create(Scene, TCHAR_TO_ANSI(*MeshName));
+	FbxNode* MeshNode = FbxNode::Create(Scene, TCHAR_TO_UTF8(*MeshName));
 	Scene->GetRootNode()->AddChild(MeshNode);
 	ExportStaticMeshToFbx(StaticMesh, RenderMesh, *MeshName, MeshNode, UVChannel);
 }
@@ -828,7 +828,7 @@ void FFbxExporter::ExportSkeletalMesh( USkeletalMesh* SkeletalMesh )
 	FString MeshName;
 	SkeletalMesh->GetName(MeshName);
 
-	FbxNode* MeshNode = FbxNode::Create(Scene, TCHAR_TO_ANSI(*MeshName));
+	FbxNode* MeshNode = FbxNode::Create(Scene, TCHAR_TO_UTF8(*MeshName));
 	Scene->GetRootNode()->AddChild(MeshNode);
 
 	ExportSkeletalMeshToFbx(SkeletalMesh, NULL, *MeshName, MeshNode);
@@ -950,13 +950,13 @@ FbxSurfaceMaterial* FFbxExporter::ExportMaterial(UMaterial* Material)
 	// Set the shading model
 	if (Material->GetShadingModel() == MSM_DefaultLit)
 	{
-		FbxMaterial = FbxSurfacePhong::Create(Scene, TCHAR_TO_ANSI(*Material->GetName()));
+		FbxMaterial = FbxSurfacePhong::Create(Scene, TCHAR_TO_UTF8(*Material->GetName()));
 		//((FbxSurfacePhong*)FbxMaterial)->Specular.Set(SetMaterialComponent(Material->Specular));
 		//((FbxSurfacePhong*)FbxMaterial)->Shininess.Set(Material->SpecularPower.Constant);
 	}
 	else // if (Material->ShadingModel == MSM_Unlit)
 	{
-		FbxMaterial = FbxSurfaceLambert::Create(Scene, TCHAR_TO_ANSI(*Material->GetName()));
+		FbxMaterial = FbxSurfaceLambert::Create(Scene, TCHAR_TO_UTF8(*Material->GetName()));
 	}
 	
 	((FbxSurfaceLambert*)FbxMaterial)->Emissive.Set(SetMaterialComponent(Material->EmissiveColor));
@@ -1066,7 +1066,7 @@ FbxNode* FFbxExporter::ExportActor(AActor* Actor, AMatineeActor* InMatineeActor,
 			FbxNodeNameToIndexMap.Add( FbxNodeName, 1 );	
 		}
 
-		ActorNode = FbxNode::Create(Scene, TCHAR_TO_ANSI(*FbxNodeName));
+		ActorNode = FbxNode::Create(Scene, TCHAR_TO_UTF8(*FbxNodeName));
 		
 
 		AActor* ParentActor = Actor->GetAttachParentActor();
@@ -1111,10 +1111,10 @@ FbxNode* FFbxExporter::ExportActor(AActor* Actor, AMatineeActor* InMatineeActor,
 
 			if (FbxPivotNodeName == FbxNodeName)
 			{
-				FbxPivotNodeName += ANSI_TO_TCHAR("_pivot");
+				FbxPivotNodeName += UTF8_TO_TCHAR("_pivot");
 			}
 
-			FbxNode* PivotNode = FbxNode::Create(Scene, TCHAR_TO_ANSI(*FbxPivotNodeName));
+			FbxNode* PivotNode = FbxNode::Create(Scene, TCHAR_TO_UTF8(*FbxPivotNodeName));
 			PivotNode->LclRotation.Set(FbxVector4(90, 0, -90));
 
 			if (Actor->IsA(ACameraActor::StaticClass()))
@@ -1170,7 +1170,7 @@ FbxNode* FFbxExporter::ExportActor(AActor* Actor, AMatineeActor* InMatineeActor,
 
 					// This actor has multiple components
 					// create a child node under the actor for each component
-					FbxNode* CompNode = FbxNode::Create(Scene, TCHAR_TO_ANSI(*Component->GetName()));
+					FbxNode* CompNode = FbxNode::Create(Scene, TCHAR_TO_UTF8(*Component->GetName()));
 
 					if( SceneComp != Actor->GetRootComponent() )
 					{
@@ -1859,7 +1859,7 @@ FbxNode* FFbxExporter::ExportStaticMeshToFbx(UStaticMesh* StaticMesh, FStaticMes
 		UniqueVerts = VertRemap;
 	}
 
-	FbxMesh* Mesh = FbxMesh::Create(Scene, TCHAR_TO_ANSI(MeshName));
+	FbxMesh* Mesh = FbxMesh::Create(Scene, TCHAR_TO_UTF8(MeshName));
 
 	// Create and fill in the vertex position data source.
 	// The position vertices are duplicated, for some reason, retrieve only the first half vertices.
@@ -1953,7 +1953,7 @@ FbxNode* FFbxExporter::ExportStaticMeshToFbx(UStaticMesh* StaticMesh, FStaticMes
 			FCString::Sprintf(UVChannelName, TEXT(""));
 		}
 
-		FbxLayerElementUV* UVDiffuseLayer = FbxLayerElementUV::Create(Mesh, TCHAR_TO_ANSI(UVChannelName));
+		FbxLayerElementUV* UVDiffuseLayer = FbxLayerElementUV::Create(Mesh, TCHAR_TO_UTF8(UVChannelName));
 
 		// Note: when eINDEX_TO_DIRECT is used, IndexArray must be 3xTriangle count, DirectArray can be smaller
 		UVDiffuseLayer->SetMappingMode(FbxLayerElement::eByPolygonVertex);
@@ -2142,7 +2142,7 @@ FbxNode* FFbxExporter::ExportSplineMeshToFbx(USplineMeshComponent* SplineMeshCom
 		UniqueVerts = VertRemap;
 	}
 
-	FbxMesh* Mesh = FbxMesh::Create(Scene, TCHAR_TO_ANSI(MeshName));
+	FbxMesh* Mesh = FbxMesh::Create(Scene, TCHAR_TO_UTF8(MeshName));
 
 	// Create and fill in the vertex position data source.
 	// The position vertices are duplicated, for some reason, retrieve only the first half vertices.
@@ -2242,7 +2242,7 @@ FbxNode* FFbxExporter::ExportSplineMeshToFbx(USplineMeshComponent* SplineMeshCom
 			FCString::Sprintf(UVChannelName, TEXT(""));
 		}
 
-		FbxLayerElementUV* UVDiffuseLayer = FbxLayerElementUV::Create(Mesh, TCHAR_TO_ANSI(UVChannelName));
+		FbxLayerElementUV* UVDiffuseLayer = FbxLayerElementUV::Create(Mesh, TCHAR_TO_UTF8(UVChannelName));
 
 		// Note: when eINDEX_TO_DIRECT is used, IndexArray must be 3xTriangle count, DirectArray can be smaller
 		UVDiffuseLayer->SetMappingMode(FbxLayerElement::eByPolygonVertex);
@@ -2400,7 +2400,7 @@ FbxNode* FFbxExporter::ExportLandscapeToFbx(ALandscapeProxy* Landscape, const TC
 		Component->GetComponentExtent(MinX, MinY, MaxX, MaxY);
 	}
 
-	FbxMesh* Mesh = FbxMesh::Create(Scene, TCHAR_TO_ANSI(MeshName));
+	FbxMesh* Mesh = FbxMesh::Create(Scene, TCHAR_TO_UTF8(MeshName));
 
 	// Create and fill in the vertex position data source.
 	const int32 ComponentSizeQuads = ((Landscape->ComponentSizeQuads + 1) >> Landscape->ExportLOD) - 1;
diff --git a/Engine/Source/Editor/UnrealEd/Private/Fbx/FbxMainImport.cpp b/Engine/Source/Editor/UnrealEd/Private/Fbx/FbxMainImport.cpp
--- a/Engine/Source/Editor/UnrealEd/Private/Fbx/FbxMainImport.cpp
+++ b/Engine/Source/Editor/UnrealEd/Private/Fbx/FbxMainImport.cpp
@@ -619,7 +619,7 @@ bool FFbxImporter::OpenFile(FString Filename, bool bParseStatistics, bool bForSc
 	if( !bImportStatus )  // Problem with the file to be imported
 	{
 		UE_LOG(LogFbx, Error,TEXT("Call to FbxImporter::Initialize() failed."));
-		UE_LOG(LogFbx, Warning, TEXT("Error returned: %s"), ANSI_TO_TCHAR(Importer->GetStatus().GetErrorString()));
+		UE_LOG(LogFbx, Warning, TEXT("Error returned: %s"), UTF8_TO_TCHAR(Importer->GetStatus().GetErrorString()));
 
 		if (Importer->GetStatus().GetCode() == FbxStatus::eInvalidFileVersion )
 		{
@@ -710,7 +710,7 @@ bool FFbxImporter::ImportFile(FString Filename)
 	}
 	else
 	{
-		ErrorMessage = ANSI_TO_TCHAR(Importer->GetStatus().GetErrorString());
+		ErrorMessage = UTF8_TO_TCHAR(Importer->GetStatus().GetErrorString());
 		AddTokenizedErrorMessage(FTokenizedMessage::Create(EMessageSeverity::Warning, FText::Format(LOCTEXT("FbxSkeletaLMeshimport_FileLoadingFailed", "FBX Scene Loading Failed : '{0}'"), FText::FromString(ErrorMessage))), FFbxErrors::Generic_LoadingSceneFailed);
 		CleanUp();
 		Result = false;
@@ -780,13 +780,13 @@ bool FFbxImporter::ImportFromFile(const FString& Filename, const FString& Type)
 					{
 						TArray<FAnalyticsEventAttribute> Attribs;
 
-						FString OriginalVendor(ANSI_TO_TCHAR(DocInfo->Original_ApplicationVendor.Get().Buffer()));
-						FString OriginalAppName(ANSI_TO_TCHAR(DocInfo->Original_ApplicationName.Get().Buffer()));
-						FString OriginalAppVersion(ANSI_TO_TCHAR(DocInfo->Original_ApplicationVersion.Get().Buffer()));
+						FString OriginalVendor(UTF8_TO_TCHAR(DocInfo->Original_ApplicationVendor.Get().Buffer()));
+						FString OriginalAppName(UTF8_TO_TCHAR(DocInfo->Original_ApplicationName.Get().Buffer()));
+						FString OriginalAppVersion(UTF8_TO_TCHAR(DocInfo->Original_ApplicationVersion.Get().Buffer()));
 
-						FString LastSavedVendor(ANSI_TO_TCHAR(DocInfo->LastSaved_ApplicationVendor.Get().Buffer()));
-						FString LastSavedAppName(ANSI_TO_TCHAR(DocInfo->LastSaved_ApplicationName.Get().Buffer()));
-						FString LastSavedAppVersion(ANSI_TO_TCHAR(DocInfo->LastSaved_ApplicationVersion.Get().Buffer()));
+						FString LastSavedVendor(UTF8_TO_TCHAR(DocInfo->LastSaved_ApplicationVendor.Get().Buffer()));
+						FString LastSavedAppName(UTF8_TO_TCHAR(DocInfo->LastSaved_ApplicationName.Get().Buffer()));
+						FString LastSavedAppVersion(UTF8_TO_TCHAR(DocInfo->LastSaved_ApplicationVersion.Get().Buffer()));
 
 						Attribs.Add(FAnalyticsEventAttribute(TEXT("Original Application Vendor"), OriginalVendor));
 						Attribs.Add(FAnalyticsEventAttribute(TEXT("Original Application Name"), OriginalAppName));
@@ -904,11 +904,11 @@ FName FFbxImporter::MakeNameForMesh(FString InName, FbxObject* FbxObject)
 
 		if ( InName == FString("None"))
 		{
-			OutputName = FName( *FString::Printf(TEXT("%s"), ANSI_TO_TCHAR(NewName )) );
+			OutputName = FName( *FString::Printf(TEXT("%s"), UTF8_TO_TCHAR(NewName )) );
 		}
 		else
 		{
-			OutputName = FName( *FString::Printf(TEXT("%s_%s"), *InName,ANSI_TO_TCHAR(NewName)) );
+			OutputName = FName( *FString::Printf(TEXT("%s_%s"), *InName,UTF8_TO_TCHAR(NewName)) );
 		}
 	}
 	
@@ -1349,7 +1349,7 @@ void FFbxImporter::RecursiveFindFbxSkelMesh(FbxNode* Node, TArray< TArray<FbxNod
 				AddTokenizedErrorMessage(
 					FTokenizedMessage::Create(
 						EMessageSeverity::Warning, 
-						FText::Format( LOCTEXT("FBX_NoWeightsOnDeformer", "Ignoring mesh {0} because it but no weights."), FText::FromString( ANSI_TO_TCHAR(SkelMeshNode->GetName()) ) )
+						FText::Format( LOCTEXT("FBX_NoWeightsOnDeformer", "Ignoring mesh {0} because it but no weights."), FText::FromString( UTF8_TO_TCHAR(SkelMeshNode->GetName()) ) )
 					), 
 					FFbxErrors::SkeletalMesh_NoWeightsOnDeformer
 				);
@@ -1522,7 +1522,7 @@ FbxNode* FFbxImporter::FindFBXMeshesByBone(const FName& RootBoneName, bool bExpa
 	FbxNode* SkeletonRoot = NULL;
 
 	// find the FBX skeleton node according to name
-	SkeletonRoot = Scene->FindNodeByName(TCHAR_TO_ANSI(*BoneNameString));
+	SkeletonRoot = Scene->FindNodeByName(TCHAR_TO_UTF8(*BoneNameString));
 
 	// SinceFBX bone names are changed on import, it's possible that the 
 	// bone name in the engine doesn't match that of the one in the FBX file and
@@ -1697,7 +1697,7 @@ FbxNode* FFbxImporter::RetrieveObjectFromName(const TCHAR* ObjectName, FbxNode*
 		{
 			FbxNode* Node = Root->GetChild(ChildIndex);
 			FbxMesh* FbxMesh = Node->GetMesh();
-			if (FbxMesh && 0 == FCString::Strcmp(ObjectName,ANSI_TO_TCHAR(Node->GetName())))
+			if (FbxMesh && 0 == FCString::Strcmp(ObjectName,UTF8_TO_TCHAR(Node->GetName())))
 			{
 				Result = Node;
 			}
diff --git a/Engine/Source/Editor/UnrealEd/Private/Fbx/FbxMaterialImport.cpp b/Engine/Source/Editor/UnrealEd/Private/Fbx/FbxMaterialImport.cpp
--- a/Engine/Source/Editor/UnrealEd/Private/Fbx/FbxMaterialImport.cpp
+++ b/Engine/Source/Editor/UnrealEd/Private/Fbx/FbxMaterialImport.cpp
@@ -55,7 +55,7 @@ UTexture* UnFbx::FFbxImporter::ImportTexture( FbxFileTexture* FbxTexture, bool b
 	
 	// create an unreal texture asset
 	UTexture* UnrealTexture = NULL;
-	FString Filename1 = ANSI_TO_TCHAR(FbxTexture->GetFileName());
+	FString Filename1 = UTF8_TO_TCHAR(FbxTexture->GetFileName());
 	FString Extension = FPaths::GetExtension(Filename1).ToLower();
 	// name the texture with file name
 	FString TextureName = FPaths::GetBaseFilename(Filename1);
@@ -97,12 +97,12 @@ UTexture* UnFbx::FFbxImporter::ImportTexture( FbxFileTexture* FbxTexture, bool b
 	if ( ! FFileHelper::LoadFileToArray( DataBinary, *Filename ))
 	{
 		// try fbx file base path + relative path
-		FString Filename2 = FileBasePath / ANSI_TO_TCHAR(FbxTexture->GetRelativeFileName());
+		FString Filename2 = FileBasePath / UTF8_TO_TCHAR(FbxTexture->GetRelativeFileName());
 		Filename = Filename2;
 		if ( ! FFileHelper::LoadFileToArray( DataBinary, *Filename ))
 		{
 			// try fbx file base path + texture file name (no path)
-			FString Filename3 = ANSI_TO_TCHAR(FbxTexture->GetRelativeFileName());
+			FString Filename3 = UTF8_TO_TCHAR(FbxTexture->GetRelativeFileName());
 			FString FileOnly = FPaths::GetCleanFilename(Filename3);
 			Filename3 = FileBasePath / FileOnly;
 			Filename = Filename3;
@@ -240,7 +240,7 @@ bool UnFbx::FFbxImporter::CreateAndLinkExpressionForMaterialProperty(
 		int32 LayeredTextureCount = FbxProperty.GetSrcObjectCount<FbxLayeredTexture>();
 		if (LayeredTextureCount>0)
 		{
-			UE_LOG(LogFbxMaterialImport, Warning,TEXT("Layered Textures are not supported (material %s)"),ANSI_TO_TCHAR(FbxMaterial.GetName()));
+			UE_LOG(LogFbxMaterialImport, Warning,TEXT("Layered Textures are not supported (material %s)"),UTF8_TO_TCHAR(FbxMaterial.GetName()));
 		}
 		else
 		{
@@ -270,7 +270,7 @@ bool UnFbx::FFbxImporter::CreateAndLinkExpressionForMaterialProperty(
 
 						// add/find UVSet and set it to the texture
 						FbxString UVSetName = FbxTexture->UVSet.Get();
-						FString LocalUVSetName = ANSI_TO_TCHAR(UVSetName.Buffer());
+						FString LocalUVSetName = UTF8_TO_TCHAR(UVSetName.Buffer());
 						int32 SetIndex = UVSet.Find(LocalUVSetName);
 						if( (SetIndex != 0 && SetIndex != INDEX_NONE) || ScaleU != 1.0f || ScaleV != 1.0f )
 						{
@@ -365,7 +365,7 @@ void UnFbx::FFbxImporter::FixupMaterial( FbxSurfaceMaterial& FbxMaterial, UMater
 
 FString UnFbx::FFbxImporter::GetMaterialFullName(FbxSurfaceMaterial& FbxMaterial)
 {
-	FString MaterialFullName = ANSI_TO_TCHAR(MakeName(FbxMaterial.GetName()));
+	FString MaterialFullName = UTF8_TO_TCHAR(MakeName(FbxMaterial.GetName()));
 
 	if (MaterialFullName.Len() > 6)
 	{
diff --git a/Engine/Source/Editor/UnrealEd/Private/Fbx/FbxMatineeImport.cpp b/Engine/Source/Editor/UnrealEd/Private/Fbx/FbxMatineeImport.cpp
--- a/Engine/Source/Editor/UnrealEd/Private/Fbx/FbxMatineeImport.cpp
+++ b/Engine/Source/Editor/UnrealEd/Private/Fbx/FbxMatineeImport.cpp
@@ -100,7 +100,7 @@ bool FFbxImporter::HasUnknownCameras( AMatineeActor* InMatineeActor ) const
 	for ( int32 NodeIndex = 0; NodeIndex < NodeCount; ++NodeIndex )
 	{
 		FbxNode* Node = RootNode->GetChild(NodeIndex);
-		if ( _HasUnknownCameras( InMatineeActor, Node, ANSI_TO_TCHAR(Node->GetName()) ) )
+		if ( _HasUnknownCameras( InMatineeActor, Node, UTF8_TO_TCHAR(Node->GetName()) ) )
 		{
 			return true;
 		}
@@ -110,7 +110,7 @@ bool FFbxImporter::HasUnknownCameras( AMatineeActor* InMatineeActor ) const
 		for( int32 ChildIndex = 0; ChildIndex < ChildNodeCount; ++ChildIndex )
 		{
 			FbxNode* ChildNode = Node->GetChild(ChildIndex);
-			if( _HasUnknownCameras( InMatineeActor, ChildNode, ANSI_TO_TCHAR(ChildNode->GetName() ) ) )
+			if( _HasUnknownCameras( InMatineeActor, ChildNode, UTF8_TO_TCHAR(ChildNode->GetName() ) ) )
 			{
 				return true;
 			}
@@ -238,7 +238,7 @@ bool FFbxImporter::ImportMatineeSequence(AMatineeActor* InMatineeActor)
 		// Attempt to name-match the scene node with one of the actors.
 		if( Actor == NULL )
 		{
-			Actor = FindObject<AActor>( ANY_PACKAGE, ANSI_TO_TCHAR(Node->GetName()) );
+			Actor = FindObject<AActor>( ANY_PACKAGE, UTF8_TO_TCHAR(Node->GetName()) );
 		}
 
 		FbxCamera* CameraNode = NULL;
@@ -248,7 +248,7 @@ bool FFbxImporter::ImportMatineeSequence(AMatineeActor* InMatineeActor)
 			if ( bCreateUnknownCameras && CameraNode != NULL )
 			{
 				Actor = GEditor->AddActor( InMatineeActor->GetWorld()->GetCurrentLevel(), ACameraActor::StaticClass(), FTransform::Identity );
-				Actor->SetActorLabel( ANSI_TO_TCHAR(CameraNode->GetName()) );
+				Actor->SetActorLabel( UTF8_TO_TCHAR(CameraNode->GetName()) );
 			}
 			else
 			{
diff --git a/Engine/Source/Editor/UnrealEd/Private/Fbx/FbxSkeletalMeshExport.cpp b/Engine/Source/Editor/UnrealEd/Private/Fbx/FbxSkeletalMeshExport.cpp
--- a/Engine/Source/Editor/UnrealEd/Private/Fbx/FbxSkeletalMeshExport.cpp
+++ b/Engine/Source/Editor/UnrealEd/Private/Fbx/FbxSkeletalMeshExport.cpp
@@ -138,7 +138,7 @@ FbxNode* FFbxExporter::CreateMesh(const USkeletalMesh* SkelMesh, const TCHAR* Me
 	SourceModel.GetVertices(Vertices);
 	if (Vertices.Num() != VertexCount) return NULL;
 
-	FbxMesh* Mesh = FbxMesh::Create(Scene, TCHAR_TO_ANSI(MeshName));
+	FbxMesh* Mesh = FbxMesh::Create(Scene, TCHAR_TO_UTF8(MeshName));
 
 	// Create and fill in the vertex position data source.
 	Mesh->InitControlPoints(VertexCount);
@@ -198,7 +198,7 @@ FbxNode* FFbxExporter::CreateMesh(const USkeletalMesh* SkelMesh, const TCHAR* Me
 			FCString::Sprintf(UVChannelName, TEXT("DiffuseUV"));
 		}
 
-		FbxLayerElementUV* UVDiffuseLayer = FbxLayerElementUV::Create(Mesh, TCHAR_TO_ANSI(UVChannelName));
+		FbxLayerElementUV* UVDiffuseLayer = FbxLayerElementUV::Create(Mesh, TCHAR_TO_UTF8(UVChannelName));
 		UVDiffuseLayer->SetMappingMode(FbxLayerElement::eByControlPoint);
 		UVDiffuseLayer->SetReferenceMode(FbxLayerElement::eDirect);
 
@@ -257,7 +257,7 @@ FbxNode* FFbxExporter::CreateMesh(const USkeletalMesh* SkelMesh, const TCHAR* Me
 		VertexColorArray.Add( FbxColor(VertColor.R, VertColor.G, VertColor.B, VertColor.A ));
 	}
 
-	FbxNode* MeshNode = FbxNode::Create(Scene, TCHAR_TO_ANSI(MeshName));
+	FbxNode* MeshNode = FbxNode::Create(Scene, TCHAR_TO_UTF8(MeshName));
 	MeshNode->SetNodeAttribute(Mesh);
 
 
@@ -283,7 +283,7 @@ FbxNode* FFbxExporter::CreateMesh(const USkeletalMesh* SkelMesh, const TCHAR* Me
 			TCHAR NewMaterialName[MAX_SPRINTF]=TEXT("");
 			FCString::Sprintf( NewMaterialName, TEXT("Fbx Default Material %i"), MaterialIndex );
 
-			FbxMaterial = FbxSurfaceLambert::Create(Scene, TCHAR_TO_ANSI(NewMaterialName));
+			FbxMaterial = FbxSurfaceLambert::Create(Scene, TCHAR_TO_UTF8(NewMaterialName));
 			((FbxSurfaceLambert*)FbxMaterial)->Diffuse.Set(FbxDouble3(0.72, 0.72, 0.72));
 		}
 
diff --git a/Engine/Source/Editor/UnrealEd/Private/Fbx/FbxSkeletalMeshImport.cpp b/Engine/Source/Editor/UnrealEd/Private/Fbx/FbxSkeletalMeshImport.cpp
--- a/Engine/Source/Editor/UnrealEd/Private/Fbx/FbxSkeletalMeshImport.cpp
+++ b/Engine/Source/Editor/UnrealEd/Private/Fbx/FbxSkeletalMeshImport.cpp
@@ -860,7 +860,7 @@ bool UnFbx::FFbxImporter::ImportBone(TArray<FbxNode*>& NodeArray, FSkeletalMeshI
 
 			if(FCStringAnsi::Strcmp(Link->GetName(), AltLink->GetName()) == 0)
 			{
-				FString RawBoneName = ANSI_TO_TCHAR(Link->GetName());
+				FString RawBoneName = UTF8_TO_TCHAR(Link->GetName());
 				AddTokenizedErrorMessage(FTokenizedMessage::Create(EMessageSeverity::Error, FText::Format(LOCTEXT("Error_DuplicateBoneName", "Error, Could not import {0}.\nDuplicate bone name found ('{1}'). Each bone must have a unique name."),
 					FText::FromString(NodeArray[0]->GetName()), FText::FromString(RawBoneName))), FFbxErrors::SkeletalMesh_DuplicateBones);
 				return false;
@@ -946,7 +946,7 @@ bool UnFbx::FFbxImporter::ImportBone(TArray<FbxNode*>& NodeArray, FSkeletalMeshI
 				if(!bUseTime0AsRefPose && !bDisableMissingBindPoseWarning)
 				{
 					bAnyLinksNotInBindPose = true;
-					LinksWithoutBindPoses +=  ANSI_TO_TCHAR(Link->GetName());
+					LinksWithoutBindPoses +=  UTF8_TO_TCHAR(Link->GetName());
 					LinksWithoutBindPoses +=  TEXT("  \n");
 				}
 
@@ -1000,7 +1000,7 @@ bool UnFbx::FFbxImporter::ImportBone(TArray<FbxNode*>& NodeArray, FSkeletalMeshI
 		FString BoneName;
 
 		const char* LinkName = Link->GetName();
-		BoneName = ANSI_TO_TCHAR( MakeName( LinkName ) );
+		BoneName = UTF8_TO_TCHAR( MakeName( LinkName ) );
 		Bone.Name = BoneName;
 
 		VJointPos& JointMatrix = Bone.BonePos;
@@ -1505,7 +1505,7 @@ USkeletalMesh* UnFbx::FFbxImporter::ReimportSkeletalMesh(USkeletalMesh* Mesh, UF
 	}
 
 	char MeshName[1024];
-	FCStringAnsi::Strcpy(MeshName, TCHAR_TO_ANSI(*Mesh->GetName()));
+	FCStringAnsi::Strcpy(MeshName, TCHAR_TO_UTF8(*Mesh->GetName()));
 	TArray<FbxNode*>* FbxNodes = NULL;
 	USkeletalMesh* NewMesh = NULL;
 
@@ -1719,7 +1719,7 @@ bool UnFbx::FFbxImporter::FillSkelMeshImporterFromFbx( FSkeletalMeshImportData&
 					if (ElementUV)
 					{
 						const char* UVSetName = ElementUV->GetName();
-						FString LocalUVSetName = ANSI_TO_TCHAR(UVSetName);
+						FString LocalUVSetName = UTF8_TO_TCHAR(UVSetName);
 
 						UVSets.AddUnique(LocalUVSetName);
 					}
@@ -1785,7 +1785,7 @@ bool UnFbx::FFbxImporter::FillSkelMeshImporterFromFbx( FSkeletalMeshImportData&
 
 			if (ImportOptions->bImportMaterials && Materials.IsValidIndex(MaterialIndex) )
 			{
-				ImportData.Materials[ExistingMatIndex].MaterialImportName = ANSI_TO_TCHAR(MakeName(FbxMaterial->GetName()));
+				ImportData.Materials[ExistingMatIndex].MaterialImportName = UTF8_TO_TCHAR(MakeName(FbxMaterial->GetName()));
 				ImportData.Materials[ExistingMatIndex].Material = Materials[MaterialIndex];
 			}
 		}
@@ -1837,7 +1837,7 @@ bool UnFbx::FFbxImporter::FillSkelMeshImporterFromFbx( FSkeletalMeshImportData&
 
 	if (!Mesh->IsTriangleMesh())
 	{
-		UE_LOG(LogFbx, Log, TEXT("Triangulating skeletal mesh %s"), ANSI_TO_TCHAR(Node->GetName()));
+		UE_LOG(LogFbx, Log, TEXT("Triangulating skeletal mesh %s"), UTF8_TO_TCHAR(Node->GetName()));
 		
 		const bool bReplace = true;
 		FbxNodeAttribute* ConvertedNode = GeometryConverter->Triangulate(Mesh, bReplace);
@@ -1892,7 +1892,7 @@ bool UnFbx::FFbxImporter::FillSkelMeshImporterFromFbx( FSkeletalMeshImportData&
 					if (ElementUV)
 					{
 						const char* UVSetName = ElementUV->GetName();
-						FString LocalUVSetName = ANSI_TO_TCHAR(UVSetName);
+						FString LocalUVSetName = UTF8_TO_TCHAR(UVSetName);
 						if (LocalUVSetName == UVSets[UVIndex])
 						{
 							LayerElementUV[UVIndex] = const_cast<FbxLayerElementUV*>(ElementUV);
@@ -2704,7 +2704,7 @@ void UnFbx::FFbxImporter::ImportMorphTargetsInternal( TArray<FbxNode*>& SkelMesh
 				FbxBlendShape* BlendShape = (FbxBlendShape*)Geometry->GetDeformer(BlendShapeIndex, FbxDeformer::eBlendShape);
 				const int32 BlendShapeChannelCount = BlendShape->GetBlendShapeChannelCount();
 
-				FString BlendShapeName = ANSI_TO_TCHAR(MakeName(BlendShape->GetName()));
+				FString BlendShapeName = UTF8_TO_TCHAR(MakeName(BlendShape->GetName()));
 
 				// see below where this is used for explanation...
 				const bool bMightBeBadMAXFile = (BlendShapeName == FString("Morpher"));
@@ -2717,7 +2717,7 @@ void UnFbx::FFbxImporter::ImportMorphTargetsInternal( TArray<FbxNode*>& SkelMesh
 						//Find which shape should we use according to the weight.
 						const int32 CurrentChannelShapeCount = Channel->GetTargetShapeCount();
 						
-						FString ChannelName = ANSI_TO_TCHAR(MakeName(Channel->GetName()));
+						FString ChannelName = UTF8_TO_TCHAR(MakeName(Channel->GetName()));
 
 						// Maya adds the name of the blendshape and an underscore to the front of the channel name, so remove it
 						if(ChannelName.StartsWith(BlendShapeName))
@@ -2732,13 +2732,13 @@ void UnFbx::FFbxImporter::ImportMorphTargetsInternal( TArray<FbxNode*>& SkelMesh
 							FString ShapeName;
 							if( CurrentChannelShapeCount > 1 )
 							{
-								ShapeName = ANSI_TO_TCHAR(MakeName(Shape->GetName() ) );
+								ShapeName = UTF8_TO_TCHAR(MakeName(Shape->GetName() ) );
 							}
 							else
 							{
 								if (bMightBeBadMAXFile)
 								{
-									ShapeName = ANSI_TO_TCHAR(MakeName(Shape->GetName()));
+									ShapeName = UTF8_TO_TCHAR(MakeName(Shape->GetName()));
 								}
 								else
 								{
diff --git a/Engine/Source/Editor/UnrealEd/Private/Fbx/FbxStaticMeshImport.cpp b/Engine/Source/Editor/UnrealEd/Private/Fbx/FbxStaticMeshImport.cpp
--- a/Engine/Source/Editor/UnrealEd/Private/Fbx/FbxStaticMeshImport.cpp
+++ b/Engine/Source/Editor/UnrealEd/Private/Fbx/FbxStaticMeshImport.cpp
@@ -134,7 +134,7 @@ bool UnFbx::FFbxImporter::BuildStaticMeshFromGeometry(FbxMesh* Mesh, UStaticMesh
 					if (ElementUV)
 					{
 						const char* UVSetName = ElementUV->GetName();
-						FString LocalUVSetName = ANSI_TO_TCHAR(UVSetName);
+						FString LocalUVSetName = UTF8_TO_TCHAR(UVSetName);
 
 						UVSets.AddUnique(LocalUVSetName);
 					}
@@ -245,7 +245,7 @@ bool UnFbx::FFbxImporter::BuildStaticMeshFromGeometry(FbxMesh* Mesh, UStaticMesh
 
 	if (!Mesh->IsTriangleMesh())
 	{
-		UE_LOG(LogFbx, Warning, TEXT("Triangulating static mesh %s"), ANSI_TO_TCHAR(Node->GetName()));
+		UE_LOG(LogFbx, Warning, TEXT("Triangulating static mesh %s"), UTF8_TO_TCHAR(Node->GetName()));
 
 		const bool bReplace = true;
 		FbxNodeAttribute* ConvertedNode = GeometryConverter->Triangulate(Mesh, bReplace);
@@ -302,7 +302,7 @@ bool UnFbx::FFbxImporter::BuildStaticMeshFromGeometry(FbxMesh* Mesh, UStaticMesh
 					if (ElementUV)
 					{
 						const char* UVSetName = ElementUV->GetName();
-						FString LocalUVSetName = ANSI_TO_TCHAR(UVSetName);
+						FString LocalUVSetName = UTF8_TO_TCHAR(UVSetName);
 						if (LocalUVSetName == UVSets[UVIndex])
 						{
 							LayerElementUV[UVIndex] = ElementUV;
@@ -708,7 +708,7 @@ bool UnFbx::FFbxImporter::BuildStaticMeshFromGeometry(FbxMesh* Mesh, UStaticMesh
 UStaticMesh* UnFbx::FFbxImporter::ReimportStaticMesh(UStaticMesh* Mesh, UFbxStaticMeshImportData* TemplateImportData)
 {
 	char MeshName[1024];
-	FCStringAnsi::Strcpy(MeshName,1024,TCHAR_TO_ANSI(*Mesh->GetName()));
+	FCStringAnsi::Strcpy(MeshName,1024,TCHAR_TO_UTF8(*Mesh->GetName()));
 	TArray<FbxNode*> FbxMeshArray;
 	FbxNode* Node = NULL;
 	UStaticMesh* NewMesh = NULL;
@@ -1313,7 +1313,7 @@ static void FindMeshSockets( FbxNode* StartNode, TArray<FbxSocketNode>& OutFbxSo
 		if( Attribute != NULL && Attribute->GetAttributeType() == FbxNodeAttribute::eNull )
 		{
 			// Is this prefixed correctly? If so it is a socket
-			FString SocketName = ANSI_TO_TCHAR( StartNode->GetName() );
+			FString SocketName = UTF8_TO_TCHAR( StartNode->GetName() );
 			if( SocketName.StartsWith( SocketPrefix ) )
 			{
 				// Remove the prefix from the name
@@ -1492,7 +1492,7 @@ bool UnFbx::FFbxImporter::ImportCollisionModels(UStaticMesh* StaticMesh, const F
 
 		if (!FbxMesh->IsTriangleMesh())
 		{
-			FString NodeName = ANSI_TO_TCHAR(MakeName(Node->GetName()));
+			FString NodeName = UTF8_TO_TCHAR(MakeName(Node->GetName()));
 			UE_LOG(LogFbx, Warning, TEXT("Triangulating mesh %s for collision model"), *NodeName);
 
 			const bool bReplace = true;
diff --git a/Engine/Source/Editor/UnrealEd/Private/Fbx/FbxUtilsImport.cpp b/Engine/Source/Editor/UnrealEd/Private/Fbx/FbxUtilsImport.cpp
--- a/Engine/Source/Editor/UnrealEd/Private/Fbx/FbxUtilsImport.cpp
+++ b/Engine/Source/Editor/UnrealEd/Private/Fbx/FbxUtilsImport.cpp
@@ -218,7 +218,7 @@ FbxString FFbxDataConverter::ConvertToFbxString(FName Name)
 	FString UnrealString;
 	Name.ToString(UnrealString);
 
-	OutString = TCHAR_TO_ANSI(*UnrealString);
+	OutString = TCHAR_TO_UTF8(*UnrealString);
 
 	return OutString;
 }
@@ -227,7 +227,7 @@ FbxString FFbxDataConverter::ConvertToFbxString(const FString& String)
 {
 	FbxString OutString;
 
-	OutString = TCHAR_TO_ANSI(*String);
+	OutString = TCHAR_TO_UTF8(*String);
 
 	return OutString;
 }
