Changed check() to ensure() which was causing a crash at runtime when auto reimport cache files could not be written out

This addresses UE-11319 - Crash on project open (Assertion Failed).

diff --git a/Engine/Source/Editor/UnrealEd/Private/AutoReimport/FileCache.cpp b/Engine/Source/Editor/UnrealEd/Private/AutoReimport/FileCache.cpp
--- a/Engine/Source/Editor/UnrealEd/Private/AutoReimport/FileCache.cpp
+++ b/Engine/Source/Editor/UnrealEd/Private/AutoReimport/FileCache.cpp
@@ -184,17 +184,24 @@ void FFileCache::WriteCache()
 {
 	if (bSavedCacheDirty)
 	{
+		const FString ParentFolder = FPaths::GetPath(Config.CacheFile);
+		if (!IFileManager::Get().DirectoryExists(*ParentFolder))
+		{
+			IFileManager::Get().MakeDirectory(*ParentFolder, true);	
+		}
+		
 		FArchive* Ar = IFileManager::Get().CreateFileWriter(*Config.CacheFile);
-		check(Ar);
-
-		*Ar << CachedDirectoryState;
+		if (ensureMsgf(Ar, TEXT("Unable to write file-cache for '%s' to '%s'."), *Config.Directory, *Config.CacheFile))
+		{
+			*Ar << CachedDirectoryState;
 
-		Ar->Close();
-		delete Ar;
+			Ar->Close();
+			delete Ar;
 
-		CachedDirectoryState.Files.Shrink();
+			CachedDirectoryState.Files.Shrink();
 
-		bSavedCacheDirty = false;
+			bSavedCacheDirty = false;
+		}
 	}
 }
 
