Fixed redundant loops in gathering text from assets.

diff --git a/Engine/Source/Editor/UnrealEd/Private/Commandlets/GatherTextFromAssetsCommandlet.cpp b/Engine/Source/Editor/UnrealEd/Private/Commandlets/GatherTextFromAssetsCommandlet.cpp
--- a/Engine/Source/Editor/UnrealEd/Private/Commandlets/GatherTextFromAssetsCommandlet.cpp
+++ b/Engine/Source/Editor/UnrealEd/Private/Commandlets/GatherTextFromAssetsCommandlet.cpp
@@ -386,36 +386,29 @@ namespace
 				// Property is a text property.
 				if (TextProperty)
 				{
-					// Text property may be a native array of fixed size.
-					for(int32 i = 0; i < Property->ArrayDim; ++i)
-					{
-						FText* Text = reinterpret_cast<FText*>(ElementValueAddress);
+					FText* Text = reinterpret_cast<FText*>(ElementValueAddress);
 
-						if( FTextInspector::GetFlags(*Text) & ETextFlag::ConvertedProperty )
-						{
-							ObjectPackage->MarkPackageDirty();
-						}
+					if( FTextInspector::GetFlags(*Text) & ETextFlag::ConvertedProperty )
+					{
+						ObjectPackage->MarkPackageDirty();
+					}
 
-						bool Fixed = false;
-						Commandlet->ProcessTextProperty(TextProperty, Text, Object, /*OUT*/Fixed );
-						if( Fixed )
-						{
-							ObjectPackage->MarkPackageDirty();
-						}
+					bool Fixed = false;
+					Commandlet->ProcessTextProperty(TextProperty, Text, Object, /*OUT*/Fixed );
+					if( Fixed )
+					{
+						ObjectPackage->MarkPackageDirty();
 					}
 				}
 				// Property is a DYNAMIC array property.
 				else if (ArrayProperty)
 				{
-					for(int32 i = 0; i < Property->ArrayDim; ++i)
+					// Iterate over all objects of the array.
+					FScriptArrayHelper ScriptArrayHelper(ArrayProperty, ElementValueAddress);
+					const int32 ElementCount = ScriptArrayHelper.Num();
+					for(int32 j = 0; j < ElementCount; ++j)
 					{
-						// Iterate over all objects of the array.
-						FScriptArrayHelper ScriptArrayHelper(ArrayProperty, ElementValueAddress);
-						const int32 ElementCount = ScriptArrayHelper.Num();
-						for(int32 j = 0; j < ElementCount; ++j)
-						{
-							ProcessProperty( ArrayProperty->Inner, ScriptArrayHelper.GetRawPtr(j) );
-						}
+						ProcessProperty( ArrayProperty->Inner, ScriptArrayHelper.GetRawPtr(j) );
 					}
 				}
 				// Property is a struct property.
@@ -481,15 +474,13 @@ namespace
 
 void UGatherTextFromAssetsCommandlet::ProcessPackages( const TArray< UPackage* >& PackagesToProcess )
 {
-	for( int32 i = 0; i < PackagesToProcess.Num(); ++i )
+	for(UPackage* const Package : PackagesToProcess)
 	{
-		UPackage* Package = PackagesToProcess[i];
 		TArray<UObject*> Objects;
 		GetObjectsWithOuter(Package, Objects);
 
-		for( int32 j = 0; j < Objects.Num(); ++j )
+		for(UObject* const Object : Objects)
 		{
-			UObject* Object = Objects[j];
 			if ( Object->IsA( UBlueprint::StaticClass() ) )
 			{
 				UBlueprint* Blueprint = Cast<UBlueprint>( Object );
@@ -781,16 +772,16 @@ int32 UGatherTextFromAssetsCommandlet::Main(const FString& Params)
 	AssetRegistryModule.Get().SearchAllAssets( true );
 	FARFilter Filter;
 
-	for(int32 i = 0; i < ExcludeClasses.Num(); i++)
+	for(const auto& ExcludeClass : ExcludeClasses)
 	{
-		UClass* FilterClass = FindObject<UClass>(ANY_PACKAGE, *ExcludeClasses[i]);
+		UClass* FilterClass = FindObject<UClass>(ANY_PACKAGE, *ExcludeClass);
 		if(FilterClass)
 		{
 			Filter.ClassNames.Add( FilterClass->GetFName() );
 		}
 		else
 		{
-			UE_LOG(LogGatherTextFromAssetsCommandlet, Warning, TEXT("Invalid exclude class %s"), *ExcludeClasses[i]);
+			UE_LOG(LogGatherTextFromAssetsCommandlet, Warning, TEXT("Invalid exclude class %s"), *ExcludeClass);
 		}
 	}
 
