Added warning when creating dynamic objects at cook time.

diff --git a/Engine/Source/Editor/UnrealEd/Private/CookOnTheFlyServer.cpp b/Engine/Source/Editor/UnrealEd/Private/CookOnTheFlyServer.cpp
--- a/Engine/Source/Editor/UnrealEd/Private/CookOnTheFlyServer.cpp
+++ b/Engine/Source/Editor/UnrealEd/Private/CookOnTheFlyServer.cpp
@@ -37,11 +37,14 @@
 
 DEFINE_LOG_CATEGORY_STATIC(LogCook, Log, All);
 
+
 #define DEBUG_COOKONTHEFLY 0
 #define OUTPUT_TIMING 0
 
 #if OUTPUT_TIMING
 
+
+
 struct FTimerInfo
 {
 public:
@@ -838,6 +841,8 @@ bool UCookOnTheFlyServer::IsCookOnTheFlyMode() const
 	return CurrentCookMode == ECookMode::CookOnTheFly || CurrentCookMode == ECookMode::CookOnTheFlyFromTheEditor; 
 }
 
+COREUOBJECT_API extern bool GOutputCookingWarnings;
+
 uint32 UCookOnTheFlyServer::TickCookOnTheSide( const float TimeSlice, uint32 &CookedPackageCount )
 {
 	struct FCookerTimer
@@ -993,7 +998,6 @@ uint32 UCookOnTheFlyServer::TickCookOnTheSide( const float TimeSlice, uint32 &Co
 		if ( bShouldCook ) // if we should cook the package then cook it otherwise add it to the list of already cooked packages below
 		{
 			SCOPE_TIMER(AllOfLoadPackage);
-
 			UPackage *Package = NULL;
 			{
 				FString PackageName;
@@ -1006,7 +1010,7 @@ uint32 UCookOnTheFlyServer::TickCookOnTheSide( const float TimeSlice, uint32 &Co
 #if DEBUG_COOKONTHEFLY
 			UE_LOG( LogCook, Display, TEXT("Processing request %s"), *BuildFilename);
 #endif
-
+			GOutputCookingWarnings = true;
 			//  if the package is already loaded then try to avoid reloading it :)
 			if ( ( Package == NULL ) || ( Package->IsFullyLoaded() == false ) )
 			{
@@ -1102,6 +1106,7 @@ uint32 UCookOnTheFlyServer::TickCookOnTheSide( const float TimeSlice, uint32 &Co
 					PackagesToSave.AddUnique( Package );
 				}
 			}
+			GOutputCookingWarnings = false;
 		}
 
 		
@@ -1385,8 +1390,16 @@ uint32 UCookOnTheFlyServer::TickCookOnTheSide( const float TimeSlice, uint32 &Co
 					BeginPackageCacheForCookedPlatformData(ObjectsInPackage);
 				}
 
+				bool bShouldSaveAsync = true;
+				FString Temp;
+				if ( FParse::Value( FCommandLine::Get(), TEXT("-diffagainstcookdirectory="), Temp ) || FParse::Value(FCommandLine::Get(), TEXT("-breakonfile="), Temp))
+				{
+					// async save doesn't work with this flags
+					bShouldSaveAsync = false;
+				}
+
 				SCOPE_TIMER(SaveCookedPackage);
-				if( SaveCookedPackage(Package, SAVE_KeepGUID | SAVE_Async | (IsCookFlagSet(ECookInitializationFlags::Unversioned) ? SAVE_Unversioned : 0), bWasUpToDate, AllTargetPlatformNames ) )
+				if( SaveCookedPackage(Package, SAVE_KeepGUID | (bShouldSaveAsync ? SAVE_Async : SAVE_None) | (IsCookFlagSet(ECookInitializationFlags::Unversioned) ? SAVE_Unversioned : 0), bWasUpToDate, AllTargetPlatformNames ) )
 				{
 					// Update flags used to determine garbage collection.
 					if (Package->ContainsMap())
@@ -2394,13 +2407,15 @@ void UCookOnTheFlyServer::AddFileToCook( TArray<FString>& InOutFilesToCook, cons
 { 
 	if (!FPackageName::IsScriptPackage(InFilename))
 	{
-		/*if ( !InOutFilesToCook.Contains(InFilename) )
+#if 1 // randomize cook file order, don't check in enabled...
+		if ( !InOutFilesToCook.Contains(InFilename) )
 		{
 			int Index = FMath::RandRange(0,InOutFilesToCook.Num()-1);
 			InOutFilesToCook.Insert(InFilename, Index );
-		}*/
-
+		}
+#else
 		InOutFilesToCook.AddUnique(InFilename);
+#endif
 	}
 }
 
diff --git a/Engine/Source/Runtime/CoreUObject/Private/UObject/UObjectGlobals.cpp b/Engine/Source/Runtime/CoreUObject/Private/UObject/UObjectGlobals.cpp
--- a/Engine/Source/Runtime/CoreUObject/Private/UObject/UObjectGlobals.cpp
+++ b/Engine/Source/Runtime/CoreUObject/Private/UObject/UObjectGlobals.cpp
@@ -1718,6 +1718,9 @@ void NotifyConstructedDuringAsyncLoading(UObject* Object, bool bSubObject);
 //@todo UE4 this is clunky
 static FRestoreForUObjectOverwrite* ObjectRestoreAfterInitProps = NULL;  
 
+COREUOBJECT_API bool GOutputCookingWarnings = false;
+
+
 UObject* StaticAllocateObject
 (
 	UClass*			InClass,
@@ -1763,6 +1766,12 @@ UObject* StaticAllocateObject
 	if(InName == NAME_None)
 	{
 		InName = MakeUniqueObjectName( InOuter, InClass );
+#if WITH_EDITOR
+		if ( GOutputCookingWarnings && GetTransientPackage() != InOuter->GetOutermost() )
+		{
+			UE_LOG(LogUObjectGlobals, Warning, TEXT("Shouldn't create dynamic objects whilest cooking, if they are saved, they will be saved with unique names, this will damage patch sizes on release.  Package %s, Object %s."), *InOuter->GetOutermost()->GetName(), *InName.ToString() );
+		}
+#endif
 	}
 	else
 	{
