[MERGE] Fixing GUBP documentation warnings.

diff --git a/Engine/Source/Editor/UnrealEd/Classes/Builders/ConeBuilder.h b/Engine/Source/Editor/UnrealEd/Classes/Builders/ConeBuilder.h
--- a/Engine/Source/Editor/UnrealEd/Classes/Builders/ConeBuilder.h
+++ b/Engine/Source/Editor/UnrealEd/Classes/Builders/ConeBuilder.h
@@ -10,7 +10,11 @@
 UCLASS(MinimalAPI, autoexpandcategories=BrushSettings, EditInlineNew, meta=(DisplayName="Cone"))
 class UConeBuilder : public UEditorBrushBuilder
 {
-	GENERATED_UCLASS_BODY()
+public:
+	GENERATED_BODY()
+
+public:
+	UConeBuilder(const FObjectInitializer& ObjectInitializer = FObjectInitializer::Get());
 
 	/** Distance from base to tip of cone */
 	UPROPERTY(EditAnywhere, Category=BrushSettings, meta=(ClampMin = "0.000001"))
diff --git a/Engine/Source/Editor/UnrealEd/Classes/Builders/CubeBuilder.h b/Engine/Source/Editor/UnrealEd/Classes/Builders/CubeBuilder.h
--- a/Engine/Source/Editor/UnrealEd/Classes/Builders/CubeBuilder.h
+++ b/Engine/Source/Editor/UnrealEd/Classes/Builders/CubeBuilder.h
@@ -10,7 +10,11 @@
 UCLASS(MinimalAPI, autoexpandcategories=BrushSettings, EditInlineNew, meta=(DisplayName="Box"))
 class UCubeBuilder : public UEditorBrushBuilder
 {
-	GENERATED_UCLASS_BODY()
+public:
+	GENERATED_BODY()
+
+public:
+	UCubeBuilder(const FObjectInitializer& ObjectInitializer = FObjectInitializer::Get());
 
 	/** The size of the cube in the X dimension */
 	UPROPERTY(EditAnywhere, Category=BrushSettings, meta=(ClampMin = "0.000001"))
diff --git a/Engine/Source/Editor/UnrealEd/Classes/Builders/CurvedStairBuilder.h b/Engine/Source/Editor/UnrealEd/Classes/Builders/CurvedStairBuilder.h
--- a/Engine/Source/Editor/UnrealEd/Classes/Builders/CurvedStairBuilder.h
+++ b/Engine/Source/Editor/UnrealEd/Classes/Builders/CurvedStairBuilder.h
@@ -10,7 +10,11 @@
 UCLASS(MinimalAPI, autoexpandcategories=BrushSettings, EditInlineNew,meta=(DisplayName="Curved Stair"))
 class UCurvedStairBuilder : public UEditorBrushBuilder
 {
-	GENERATED_UCLASS_BODY()
+public:
+	GENERATED_BODY()
+
+public:
+	UCurvedStairBuilder(const FObjectInitializer& ObjectInitializer = FObjectInitializer::Get());
 
 	/** The radius of the inner curve of the stair */
 	UPROPERTY(EditAnywhere, Category=BrushSettings, meta=(ClampMin = "1"))
diff --git a/Engine/Source/Editor/UnrealEd/Classes/Builders/CylinderBuilder.h b/Engine/Source/Editor/UnrealEd/Classes/Builders/CylinderBuilder.h
--- a/Engine/Source/Editor/UnrealEd/Classes/Builders/CylinderBuilder.h
+++ b/Engine/Source/Editor/UnrealEd/Classes/Builders/CylinderBuilder.h
@@ -15,7 +15,11 @@
 UCLASS(MinimalAPI, autoexpandcategories=BrushSettings, EditInlineNew, meta=(DisplayName="Cylinder"))
 class UCylinderBuilder : public UEditorBrushBuilder
 {
-	GENERATED_UCLASS_BODY()
+public:
+	GENERATED_BODY()
+
+public:
+	UCylinderBuilder(const FObjectInitializer& ObjectInitializer = FObjectInitializer::Get());
 
 	/** Distance from base to tip of cylinder */
 	UPROPERTY(EditAnywhere, Category=BrushSettings, meta=(ClampMin = "0.000001"))
diff --git a/Engine/Source/Editor/UnrealEd/Classes/Builders/EditorBrushBuilder.h b/Engine/Source/Editor/UnrealEd/Classes/Builders/EditorBrushBuilder.h
--- a/Engine/Source/Editor/UnrealEd/Classes/Builders/EditorBrushBuilder.h
+++ b/Engine/Source/Editor/UnrealEd/Classes/Builders/EditorBrushBuilder.h
@@ -7,9 +7,11 @@
 UCLASS(abstract, MinimalAPI)
 class UEditorBrushBuilder : public UBrushBuilder
 {
-	GENERATED_UCLASS_BODY()
-
 public:
+	GENERATED_BODY()
+public:
+
+	UEditorBrushBuilder(const FObjectInitializer& ObjectInitializer = FObjectInitializer::Get());
 
 	/** UBrushBuilder interface */
 	virtual void BeginBrush( bool InMergeCoplanars, FName InLayer ) override;
diff --git a/Engine/Source/Editor/UnrealEd/Classes/Builders/LinearStairBuilder.h b/Engine/Source/Editor/UnrealEd/Classes/Builders/LinearStairBuilder.h
--- a/Engine/Source/Editor/UnrealEd/Classes/Builders/LinearStairBuilder.h
+++ b/Engine/Source/Editor/UnrealEd/Classes/Builders/LinearStairBuilder.h
@@ -10,7 +10,11 @@
 UCLASS(MinimalAPI, autoexpandcategories=BrushSettings, EditInlineNew, meta=(DisplayName="Linear Stair"))
 class ULinearStairBuilder : public UEditorBrushBuilder
 {
-	GENERATED_UCLASS_BODY()
+public:
+	GENERATED_BODY()
+
+public:
+	ULinearStairBuilder(const FObjectInitializer& ObjectInitializer = FObjectInitializer::Get());
 
 	/** The length of each step */
 	UPROPERTY(EditAnywhere, Category=BrushSettings, meta=(ClampMin = "1"))
diff --git a/Engine/Source/Editor/UnrealEd/Classes/Builders/SheetBuilder.h b/Engine/Source/Editor/UnrealEd/Classes/Builders/SheetBuilder.h
--- a/Engine/Source/Editor/UnrealEd/Classes/Builders/SheetBuilder.h
+++ b/Engine/Source/Editor/UnrealEd/Classes/Builders/SheetBuilder.h
@@ -19,7 +19,11 @@ enum ESheetAxis
 UCLASS(MinimalAPI, autoexpandcategories=BrushSettings, EditInlineNew, NotPlaceable, meta=(DisplayName="Plane"))
 class USheetBuilder : public UEditorBrushBuilder
 {
-	GENERATED_UCLASS_BODY()
+public:
+	GENERATED_BODY()
+
+public:
+	USheetBuilder(const FObjectInitializer& ObjectInitializer = FObjectInitializer::Get());
 
 	UPROPERTY(EditAnywhere, Category=BrushSettings, meta=(ClampMin = "1"))
 	int32 X;
diff --git a/Engine/Source/Editor/UnrealEd/Classes/Builders/SpiralStairBuilder.h b/Engine/Source/Editor/UnrealEd/Classes/Builders/SpiralStairBuilder.h
--- a/Engine/Source/Editor/UnrealEd/Classes/Builders/SpiralStairBuilder.h
+++ b/Engine/Source/Editor/UnrealEd/Classes/Builders/SpiralStairBuilder.h
@@ -10,7 +10,11 @@
 UCLASS(MinimalAPI, autoexpandcategories=BrushSettings, EditInlineNew, meta=(DisplayName="Spiral Stair"))
 class USpiralStairBuilder : public UEditorBrushBuilder
 {
-	GENERATED_UCLASS_BODY()
+public:
+	GENERATED_BODY()
+
+public:
+	USpiralStairBuilder(const FObjectInitializer& ObjectInitializer = FObjectInitializer::Get());
 
 	/** The radius of the inner curve of the stair */
 	UPROPERTY(EditAnywhere, Category=BrushSettings, meta=(ClampMin = "1"))
diff --git a/Engine/Source/Editor/UnrealEd/Classes/Builders/TetrahedronBuilder.h b/Engine/Source/Editor/UnrealEd/Classes/Builders/TetrahedronBuilder.h
--- a/Engine/Source/Editor/UnrealEd/Classes/Builders/TetrahedronBuilder.h
+++ b/Engine/Source/Editor/UnrealEd/Classes/Builders/TetrahedronBuilder.h
@@ -10,7 +10,11 @@
 UCLASS(MinimalAPI, autoexpandcategories=BrushSettings, EditInlineNew, meta=(DisplayName="Sphere"))
 class UTetrahedronBuilder : public UEditorBrushBuilder
 {
-	GENERATED_UCLASS_BODY()
+public:
+	GENERATED_BODY()
+
+public:
+	UTetrahedronBuilder(const FObjectInitializer& ObjectInitializer = FObjectInitializer::Get());
 
 	/** The radius of this sphere */
 	UPROPERTY(EditAnywhere, Category=BrushSettings, meta=(ClampMin = "0.000001"))
diff --git a/Engine/Source/Editor/UnrealEd/Classes/Builders/VolumetricBuilder.h b/Engine/Source/Editor/UnrealEd/Classes/Builders/VolumetricBuilder.h
--- a/Engine/Source/Editor/UnrealEd/Classes/Builders/VolumetricBuilder.h
+++ b/Engine/Source/Editor/UnrealEd/Classes/Builders/VolumetricBuilder.h
@@ -10,8 +10,12 @@
 UCLASS(MinimalAPI, autoexpandcategories=BrushSettings, EditInlineNew, NotPlaceable, meta=(DisplayName="Volumetric"))
 class UVolumetricBuilder : public UEditorBrushBuilder
 {
-	GENERATED_UCLASS_BODY()
+public:
+	GENERATED_BODY()
 
+public:
+	UVolumetricBuilder(const FObjectInitializer& ObjectInitializer = FObjectInitializer::Get());
+	
 	UPROPERTY(EditAnywhere, Category=BrushSettings, meta=(ClampMin = "0.000001"))
 	float Z;
 
diff --git a/Engine/Source/Editor/UnrealEd/Classes/CookOnTheSide/CookOnTheFlyServer.h b/Engine/Source/Editor/UnrealEd/Classes/CookOnTheSide/CookOnTheFlyServer.h
--- a/Engine/Source/Editor/UnrealEd/Classes/CookOnTheSide/CookOnTheFlyServer.h
+++ b/Engine/Source/Editor/UnrealEd/Classes/CookOnTheSide/CookOnTheFlyServer.h
@@ -51,7 +51,10 @@ namespace ECookMode
 UCLASS()
 class UNREALED_API UCookOnTheFlyServer : public UObject, public FTickableEditorObject
 {
-	GENERATED_UCLASS_BODY()
+	GENERATED_BODY()
+
+	UCookOnTheFlyServer(const FObjectInitializer& ObjectInitializer = FObjectInitializer::Get());
+
 private:
 
 	/** Array which has been made thread safe :) */
diff --git a/Engine/Source/Editor/UnrealEd/Classes/Editor/EditorEngine.h b/Engine/Source/Editor/UnrealEd/Classes/Editor/EditorEngine.h
--- a/Engine/Source/Editor/UnrealEd/Classes/Editor/EditorEngine.h
+++ b/Engine/Source/Editor/UnrealEd/Classes/Editor/EditorEngine.h
@@ -281,7 +281,11 @@ struct FActorOrComponent
 UCLASS(config=Engine, transient)
 class UNREALED_API UEditorEngine : public UEngine
 {
-	GENERATED_UCLASS_BODY()
+public:
+	GENERATED_BODY()
+
+public:
+	UEditorEngine(const FObjectInitializer& ObjectInitializer = FObjectInitializer::Get());
 
 	// Objects.
 	UPROPERTY()
diff --git a/Engine/Source/Editor/UnrealEd/Classes/Editor/UnrealEdEngine.h b/Engine/Source/Editor/UnrealEd/Classes/Editor/UnrealEdEngine.h
--- a/Engine/Source/Editor/UnrealEd/Classes/Editor/UnrealEdEngine.h
+++ b/Engine/Source/Editor/UnrealEd/Classes/Editor/UnrealEdEngine.h
@@ -88,7 +88,9 @@ class FPerformanceMonitor;
 UCLASS(config=Engine, transient)
 class UNREALED_API UUnrealEdEngine : public UEditorEngine, public FNotifyHook
 {
-	GENERATED_UCLASS_BODY()
+public:
+	GENERATED_BODY()
+public:
 
 	/** Global instance of the editor options class. */
 	UPROPERTY()
diff --git a/Engine/Source/Editor/UnrealEd/Classes/Exporters/LevelExporterFBX.h b/Engine/Source/Editor/UnrealEd/Classes/Exporters/LevelExporterFBX.h
--- a/Engine/Source/Editor/UnrealEd/Classes/Exporters/LevelExporterFBX.h
+++ b/Engine/Source/Editor/UnrealEd/Classes/Exporters/LevelExporterFBX.h
@@ -10,8 +10,11 @@
 UCLASS()
 class ULevelExporterFBX : public UExporter
 {
-	GENERATED_UCLASS_BODY()
+public:
+	GENERATED_BODY()
 
+public:
+	ULevelExporterFBX(const FObjectInitializer& ObjectInitializer = FObjectInitializer());
 
 	// Begin UExporter Interface
 	virtual bool ExportBinary( UObject* Object, const TCHAR* Type, FArchive& Ar, FFeedbackContext* Warn, int32 FileIndex = 0, uint32 PortFlags=0 ) override;
diff --git a/Engine/Source/Editor/UnrealEd/Classes/Exporters/LevelExporterLOD.h b/Engine/Source/Editor/UnrealEd/Classes/Exporters/LevelExporterLOD.h
--- a/Engine/Source/Editor/UnrealEd/Classes/Exporters/LevelExporterLOD.h
+++ b/Engine/Source/Editor/UnrealEd/Classes/Exporters/LevelExporterLOD.h
@@ -10,8 +10,11 @@
 UCLASS()
 class ULevelExporterLOD : public UExporter
 {
-	GENERATED_UCLASS_BODY()
+public:
+	GENERATED_BODY()
 
+public:
+	ULevelExporterLOD(const FObjectInitializer& ObjectInitializer = FObjectInitializer());
 
 	// Begin UExporter Interface
 	virtual bool ExportText( const FExportObjectInnerContext* Context, UObject* Object, const TCHAR* Type, FOutputDevice& Ar, FFeedbackContext* Warn, uint32 PortFlags=0 ) override;
diff --git a/Engine/Source/Editor/UnrealEd/Classes/Exporters/LevelExporterOBJ.h b/Engine/Source/Editor/UnrealEd/Classes/Exporters/LevelExporterOBJ.h
--- a/Engine/Source/Editor/UnrealEd/Classes/Exporters/LevelExporterOBJ.h
+++ b/Engine/Source/Editor/UnrealEd/Classes/Exporters/LevelExporterOBJ.h
@@ -10,8 +10,11 @@
 UCLASS()
 class ULevelExporterOBJ : public UExporter
 {
-	GENERATED_UCLASS_BODY()
+public:
+	GENERATED_BODY()
 
+public:
+	ULevelExporterOBJ(const FObjectInitializer& ObjectInitializer = FObjectInitializer());
 
 	// Begin UExporter Interface
 	virtual bool ExportText( const FExportObjectInnerContext* Context, UObject* Object, const TCHAR* Type, FOutputDevice& Ar, FFeedbackContext* Warn, uint32 PortFlags=0 ) override;
diff --git a/Engine/Source/Editor/UnrealEd/Classes/Exporters/LevelExporterSTL.h b/Engine/Source/Editor/UnrealEd/Classes/Exporters/LevelExporterSTL.h
--- a/Engine/Source/Editor/UnrealEd/Classes/Exporters/LevelExporterSTL.h
+++ b/Engine/Source/Editor/UnrealEd/Classes/Exporters/LevelExporterSTL.h
@@ -10,8 +10,11 @@
 UCLASS()
 class ULevelExporterSTL : public UExporter
 {
-	GENERATED_UCLASS_BODY()
+public:
+	GENERATED_BODY()
 
+public:
+	ULevelExporterSTL(const FObjectInitializer& ObjectInitializer = FObjectInitializer());
 
 	// Begin UExporter Interface
 	virtual bool ExportText( const FExportObjectInnerContext* Context, UObject* Object, const TCHAR* Type, FOutputDevice& Ar, FFeedbackContext* Warn, uint32 PortFlags=0 ) override;
diff --git a/Engine/Source/Editor/UnrealEd/Classes/Exporters/LevelExporterT3D.h b/Engine/Source/Editor/UnrealEd/Classes/Exporters/LevelExporterT3D.h
--- a/Engine/Source/Editor/UnrealEd/Classes/Exporters/LevelExporterT3D.h
+++ b/Engine/Source/Editor/UnrealEd/Classes/Exporters/LevelExporterT3D.h
@@ -10,8 +10,11 @@
 UCLASS()
 class ULevelExporterT3D : public UExporter
 {
-	GENERATED_UCLASS_BODY()
+public:
+	GENERATED_BODY()
 
+public:
+	ULevelExporterT3D(const FObjectInitializer& ObjectInitializer = FObjectInitializer());
 
 	// Begin UExporter Interface
 	virtual bool ExportText( const FExportObjectInnerContext* Context, UObject* Object, const TCHAR* Type, FOutputDevice& Ar, FFeedbackContext* Warn, uint32 PortFlags=0 ) override;
diff --git a/Engine/Source/Editor/UnrealEd/Classes/Exporters/ModelExporterT3D.h b/Engine/Source/Editor/UnrealEd/Classes/Exporters/ModelExporterT3D.h
--- a/Engine/Source/Editor/UnrealEd/Classes/Exporters/ModelExporterT3D.h
+++ b/Engine/Source/Editor/UnrealEd/Classes/Exporters/ModelExporterT3D.h
@@ -10,8 +10,11 @@
 UCLASS()
 class UModelExporterT3D : public UExporter
 {
-	GENERATED_UCLASS_BODY()
+public:
+	GENERATED_BODY()
 
+public:
+	UModelExporterT3D(const FObjectInitializer& ObjectInitializer = FObjectInitializer());
 
 	// Begin UExporter Interface
 	virtual bool ExportText( const FExportObjectInnerContext* Context, UObject* Object, const TCHAR* Type, FOutputDevice& Ar, FFeedbackContext* Warn, uint32 PortFlags=0 ) override;
diff --git a/Engine/Source/Editor/UnrealEd/Classes/Exporters/ObjectExporterT3D.h b/Engine/Source/Editor/UnrealEd/Classes/Exporters/ObjectExporterT3D.h
--- a/Engine/Source/Editor/UnrealEd/Classes/Exporters/ObjectExporterT3D.h
+++ b/Engine/Source/Editor/UnrealEd/Classes/Exporters/ObjectExporterT3D.h
@@ -10,7 +10,11 @@
 UCLASS()
 class UObjectExporterT3D : public UExporter
 {
-	GENERATED_UCLASS_BODY()
+public:
+	GENERATED_BODY()
+
+public:
+	UObjectExporterT3D(const FObjectInitializer& ObjectInitializer = FObjectInitializer());
 
 
 	// Begin UExporter Interface
diff --git a/Engine/Source/Editor/UnrealEd/Classes/Exporters/PolysExporterOBJ.h b/Engine/Source/Editor/UnrealEd/Classes/Exporters/PolysExporterOBJ.h
--- a/Engine/Source/Editor/UnrealEd/Classes/Exporters/PolysExporterOBJ.h
+++ b/Engine/Source/Editor/UnrealEd/Classes/Exporters/PolysExporterOBJ.h
@@ -10,8 +10,11 @@
 UCLASS()
 class UPolysExporterOBJ : public UExporter
 {
-	GENERATED_UCLASS_BODY()
+public:
+	GENERATED_BODY()
 
+public:
+	UPolysExporterOBJ(const FObjectInitializer& ObjectInitializer = FObjectInitializer());
 
 	// Begin UExporter Interface
 	virtual bool ExportText( const FExportObjectInnerContext* Context, UObject* Object, const TCHAR* Type, FOutputDevice& Ar, FFeedbackContext* Warn, uint32 PortFlags=0 ) override;
diff --git a/Engine/Source/Editor/UnrealEd/Classes/Exporters/PolysExporterT3D.h b/Engine/Source/Editor/UnrealEd/Classes/Exporters/PolysExporterT3D.h
--- a/Engine/Source/Editor/UnrealEd/Classes/Exporters/PolysExporterT3D.h
+++ b/Engine/Source/Editor/UnrealEd/Classes/Exporters/PolysExporterT3D.h
@@ -10,8 +10,11 @@
 UCLASS()
 class UPolysExporterT3D : public UExporter
 {
-	GENERATED_UCLASS_BODY()
+public:
+	GENERATED_BODY()
 
+public:
+	UPolysExporterT3D(const FObjectInitializer& ObjectInitializer = FObjectInitializer());
 
 	// Begin UExporter Interface
 	virtual bool ExportText( const FExportObjectInnerContext* Context, UObject* Object, const TCHAR* Type, FOutputDevice& Ar, FFeedbackContext* Warn, uint32 PortFlags=0 ) override;
diff --git a/Engine/Source/Editor/UnrealEd/Classes/Exporters/SequenceExporterT3D.h b/Engine/Source/Editor/UnrealEd/Classes/Exporters/SequenceExporterT3D.h
--- a/Engine/Source/Editor/UnrealEd/Classes/Exporters/SequenceExporterT3D.h
+++ b/Engine/Source/Editor/UnrealEd/Classes/Exporters/SequenceExporterT3D.h
@@ -10,8 +10,11 @@
 UCLASS()
 class USequenceExporterT3D : public UExporter
 {
-	GENERATED_UCLASS_BODY()
+public:
+	GENERATED_BODY()
 
+public:
+	USequenceExporterT3D(const FObjectInitializer& ObjectInitializer = FObjectInitializer());
 
 	// Begin UExporter Interface
 	virtual bool ExportText( const FExportObjectInnerContext* Context, UObject* Object, const TCHAR* Type, FOutputDevice& Ar, FFeedbackContext* Warn, uint32 PortFlags=0 ) override;
diff --git a/Engine/Source/Editor/UnrealEd/Classes/Exporters/SoundExporterOGG.h b/Engine/Source/Editor/UnrealEd/Classes/Exporters/SoundExporterOGG.h
--- a/Engine/Source/Editor/UnrealEd/Classes/Exporters/SoundExporterOGG.h
+++ b/Engine/Source/Editor/UnrealEd/Classes/Exporters/SoundExporterOGG.h
@@ -6,7 +6,12 @@
 UCLASS()
 class USoundExporterOGG : public UExporter
 {
-    GENERATED_UCLASS_BODY()
+public:
+	GENERATED_BODY()
+
+public:
+	USoundExporterOGG(const FObjectInitializer& ObjectInitializer = FObjectInitializer());
+
 	// Begin UExporter Interface
 	virtual bool ExportBinary( UObject* Object, const TCHAR* Type, FArchive& Ar, FFeedbackContext* Warn, int32 FileIndex = 0, uint32 PortFlags=0 ) override;
 	virtual bool SupportsObject(UObject* Object) const override;
diff --git a/Engine/Source/Editor/UnrealEd/Classes/Exporters/SoundExporterWAV.h b/Engine/Source/Editor/UnrealEd/Classes/Exporters/SoundExporterWAV.h
--- a/Engine/Source/Editor/UnrealEd/Classes/Exporters/SoundExporterWAV.h
+++ b/Engine/Source/Editor/UnrealEd/Classes/Exporters/SoundExporterWAV.h
@@ -6,7 +6,12 @@
 UCLASS()
 class USoundExporterWAV : public UExporter
 {
-    GENERATED_UCLASS_BODY()
+public:
+	GENERATED_BODY()
+
+public:
+	USoundExporterWAV(const FObjectInitializer& ObjectInitializer = FObjectInitializer());
+
 	// Begin UExporter Interface
 	virtual bool ExportBinary( UObject* Object, const TCHAR* Type, FArchive& Ar, FFeedbackContext* Warn, int32 FileIndex = 0, uint32 PortFlags=0 ) override;
 	virtual bool SupportsObject(UObject* Object) const override;
diff --git a/Engine/Source/Editor/UnrealEd/Classes/Exporters/SoundSurroundExporterWAV.h b/Engine/Source/Editor/UnrealEd/Classes/Exporters/SoundSurroundExporterWAV.h
--- a/Engine/Source/Editor/UnrealEd/Classes/Exporters/SoundSurroundExporterWAV.h
+++ b/Engine/Source/Editor/UnrealEd/Classes/Exporters/SoundSurroundExporterWAV.h
@@ -10,8 +10,11 @@
 UCLASS()
 class USoundSurroundExporterWAV : public UExporter
 {
-	GENERATED_UCLASS_BODY()
+public:
+	GENERATED_BODY()
 
+public:
+	USoundSurroundExporterWAV(const FObjectInitializer& ObjectInitializer = FObjectInitializer());
 
 	// Begin UExporter Interface
 	virtual int32 GetFileCount( void ) const override;
diff --git a/Engine/Source/Editor/UnrealEd/Classes/Exporters/TextBufferExporterTXT.h b/Engine/Source/Editor/UnrealEd/Classes/Exporters/TextBufferExporterTXT.h
--- a/Engine/Source/Editor/UnrealEd/Classes/Exporters/TextBufferExporterTXT.h
+++ b/Engine/Source/Editor/UnrealEd/Classes/Exporters/TextBufferExporterTXT.h
@@ -10,8 +10,11 @@
 UCLASS()
 class UTextBufferExporterTXT : public UExporter
 {
-	GENERATED_UCLASS_BODY()
+public:
+	GENERATED_BODY()
 
+public:
+	UTextBufferExporterTXT(const FObjectInitializer& ObjectInitializer = FObjectInitializer());
 
 	// Begin UExporter Interface
 	virtual bool ExportText( const FExportObjectInnerContext* Context, UObject* Object, const TCHAR* Type, FOutputDevice& Ar, FFeedbackContext* Warn, uint32 PortFlags=0 ) override;
diff --git a/Engine/Source/Editor/UnrealEd/Private/EditorActor.cpp b/Engine/Source/Editor/UnrealEd/Private/EditorActor.cpp
--- a/Engine/Source/Editor/UnrealEd/Private/EditorActor.cpp
+++ b/Engine/Source/Editor/UnrealEd/Private/EditorActor.cpp
@@ -96,12 +96,6 @@ public:
 	}
 };
 
-UUnrealEdEngine::UUnrealEdEngine(const FObjectInitializer& ObjectInitializer)
-	: Super(ObjectInitializer)
-{
-}
-
-
 void UUnrealEdEngine::edactCopySelected( UWorld* InWorld, FString* DestinationData )
 {
 	if (GetSelectedComponentCount() > 0)
diff --git a/Engine/Source/Editor/UnrealEd/Private/EditorActorFolders.cpp b/Engine/Source/Editor/UnrealEd/Private/EditorActorFolders.cpp
--- a/Engine/Source/Editor/UnrealEd/Private/EditorActorFolders.cpp
+++ b/Engine/Source/Editor/UnrealEd/Private/EditorActorFolders.cpp
@@ -8,10 +8,6 @@
 
 #define LOCTEXT_NAMESPACE "FActorFolders"
 
-UEditorActorFolders::UEditorActorFolders(const FObjectInitializer& ObjectInitializer)
-	: Super(ObjectInitializer)
-{ }
-
 void UEditorActorFolders::Serialize(FArchive& Ar)
 {
 	Ar << Folders;
diff --git a/Engine/Source/Editor/UnrealEd/Public/EdGraphNode_Comment.h b/Engine/Source/Editor/UnrealEd/Public/EdGraphNode_Comment.h
--- a/Engine/Source/Editor/UnrealEd/Public/EdGraphNode_Comment.h
+++ b/Engine/Source/Editor/UnrealEd/Public/EdGraphNode_Comment.h
@@ -23,7 +23,11 @@ namespace ECommentBoxMode
 UCLASS(MinimalAPI)
 class UEdGraphNode_Comment : public UEdGraphNode
 {
-	GENERATED_UCLASS_BODY()
+public:
+	GENERATED_BODY()
+
+public:
+	UEdGraphNode_Comment(const FObjectInitializer& ObjectInitializer = FObjectInitializer::Get());
 
 	/** Color to style comment with */
 	UPROPERTY(EditAnywhere, Category=Comment)
diff --git a/Engine/Source/Editor/UnrealEd/Public/EditorActorFolders.h b/Engine/Source/Editor/UnrealEd/Public/EditorActorFolders.h
--- a/Engine/Source/Editor/UnrealEd/Public/EditorActorFolders.h
+++ b/Engine/Source/Editor/UnrealEd/Public/EditorActorFolders.h
@@ -29,7 +29,8 @@ struct FActorFolderProps
 UCLASS()
 class UEditorActorFolders : public UObject
 {
-	GENERATED_UCLASS_BODY()
+public:
+	GENERATED_BODY()
 public:
 	virtual void Serialize(FArchive& Ar) override;
 
diff --git a/Engine/Source/Runtime/InputCore/Classes/InputCoreTypes.h b/Engine/Source/Runtime/InputCore/Classes/InputCoreTypes.h
--- a/Engine/Source/Runtime/InputCore/Classes/InputCoreTypes.h
+++ b/Engine/Source/Runtime/InputCore/Classes/InputCoreTypes.h
@@ -459,7 +459,6 @@ private:
 UCLASS(abstract)
 class UInputCoreTypes : public UObject
 {
-	GENERATED_UCLASS_BODY()
-
+	GENERATED_BODY()
 
 };
\ No newline at end of file
diff --git a/Engine/Source/Runtime/InputCore/Private/InputCoreTypes.cpp b/Engine/Source/Runtime/InputCore/Private/InputCoreTypes.cpp
--- a/Engine/Source/Runtime/InputCore/Private/InputCoreTypes.cpp
+++ b/Engine/Source/Runtime/InputCore/Private/InputCoreTypes.cpp
@@ -4,11 +4,6 @@
 #include "PropertyTag.h"
 #include "InputCoreTypes.h"
 
-UInputCoreTypes::UInputCoreTypes(const FObjectInitializer& ObjectInitializer)
-: Super(ObjectInitializer)
-{
-}
-
 #define LOCTEXT_NAMESPACE "InputKeys"
 
 const FKey EKeys::MouseX("MouseX");
diff --git a/Engine/Source/Runtime/Slate/Private/Framework/Styling/ButtonWidgetStyle.cpp b/Engine/Source/Runtime/Slate/Private/Framework/Styling/ButtonWidgetStyle.cpp
--- a/Engine/Source/Runtime/Slate/Private/Framework/Styling/ButtonWidgetStyle.cpp
+++ b/Engine/Source/Runtime/Slate/Private/Framework/Styling/ButtonWidgetStyle.cpp
@@ -3,8 +3,3 @@
 #include "SlatePrivatePCH.h"
 
 
-UButtonWidgetStyle::UButtonWidgetStyle( const FObjectInitializer& ObjectInitializer )
-	: Super(ObjectInitializer)
-{
-	
-}
\ No newline at end of file
diff --git a/Engine/Source/Runtime/Slate/Private/Framework/Styling/CheckBoxWidgetStyle.cpp b/Engine/Source/Runtime/Slate/Private/Framework/Styling/CheckBoxWidgetStyle.cpp
--- a/Engine/Source/Runtime/Slate/Private/Framework/Styling/CheckBoxWidgetStyle.cpp
+++ b/Engine/Source/Runtime/Slate/Private/Framework/Styling/CheckBoxWidgetStyle.cpp
@@ -3,8 +3,3 @@
 #include "SlatePrivatePCH.h"
 
 
-UCheckBoxWidgetStyle::UCheckBoxWidgetStyle( const FObjectInitializer& ObjectInitializer )
-	: Super(ObjectInitializer)
-{
-	
-}
\ No newline at end of file
diff --git a/Engine/Source/Runtime/Slate/Private/Framework/Styling/ComboBoxWidgetStyle.cpp b/Engine/Source/Runtime/Slate/Private/Framework/Styling/ComboBoxWidgetStyle.cpp
--- a/Engine/Source/Runtime/Slate/Private/Framework/Styling/ComboBoxWidgetStyle.cpp
+++ b/Engine/Source/Runtime/Slate/Private/Framework/Styling/ComboBoxWidgetStyle.cpp
@@ -3,8 +3,3 @@
 #include "SlatePrivatePCH.h"
 
 
-UComboBoxWidgetStyle::UComboBoxWidgetStyle( const FObjectInitializer& ObjectInitializer )
-	: Super(ObjectInitializer)
-{
-	
-}
diff --git a/Engine/Source/Runtime/Slate/Private/Framework/Styling/ComboButtonWidgetStyle.cpp b/Engine/Source/Runtime/Slate/Private/Framework/Styling/ComboButtonWidgetStyle.cpp
--- a/Engine/Source/Runtime/Slate/Private/Framework/Styling/ComboButtonWidgetStyle.cpp
+++ b/Engine/Source/Runtime/Slate/Private/Framework/Styling/ComboButtonWidgetStyle.cpp
@@ -3,8 +3,3 @@
 #include "SlatePrivatePCH.h"
 
 
-UComboButtonWidgetStyle::UComboButtonWidgetStyle( const FObjectInitializer& ObjectInitializer )
-	: Super(ObjectInitializer)
-{
-	
-}
diff --git a/Engine/Source/Runtime/Slate/Private/Framework/Styling/EditableTextBoxWidgetStyle.cpp b/Engine/Source/Runtime/Slate/Private/Framework/Styling/EditableTextBoxWidgetStyle.cpp
--- a/Engine/Source/Runtime/Slate/Private/Framework/Styling/EditableTextBoxWidgetStyle.cpp
+++ b/Engine/Source/Runtime/Slate/Private/Framework/Styling/EditableTextBoxWidgetStyle.cpp
@@ -2,9 +2,3 @@
 
 #include "SlatePrivatePCH.h"
 
-
-UEditableTextBoxWidgetStyle::UEditableTextBoxWidgetStyle( const FObjectInitializer& ObjectInitializer )
-	: Super(ObjectInitializer)
-{
-	
-}
\ No newline at end of file
diff --git a/Engine/Source/Runtime/Slate/Private/Framework/Styling/EditableTextWidgetStyle.cpp b/Engine/Source/Runtime/Slate/Private/Framework/Styling/EditableTextWidgetStyle.cpp
--- a/Engine/Source/Runtime/Slate/Private/Framework/Styling/EditableTextWidgetStyle.cpp
+++ b/Engine/Source/Runtime/Slate/Private/Framework/Styling/EditableTextWidgetStyle.cpp
@@ -3,8 +3,3 @@
 #include "SlatePrivatePCH.h"
 
 
-UEditableTextWidgetStyle::UEditableTextWidgetStyle( const FObjectInitializer& ObjectInitializer )
-	: Super(ObjectInitializer)
-{
-	
-}
\ No newline at end of file
diff --git a/Engine/Source/Runtime/Slate/Private/Framework/Styling/ProgressWidgetStyle.cpp b/Engine/Source/Runtime/Slate/Private/Framework/Styling/ProgressWidgetStyle.cpp
--- a/Engine/Source/Runtime/Slate/Private/Framework/Styling/ProgressWidgetStyle.cpp
+++ b/Engine/Source/Runtime/Slate/Private/Framework/Styling/ProgressWidgetStyle.cpp
@@ -3,8 +3,3 @@
 #include "SlatePrivatePCH.h"
 
 
-UProgressWidgetStyle::UProgressWidgetStyle( const FObjectInitializer& ObjectInitializer )
-	: Super(ObjectInitializer)
-{
-	
-}
\ No newline at end of file
diff --git a/Engine/Source/Runtime/Slate/Private/Framework/Styling/ScrollBarWidgetStyle.cpp b/Engine/Source/Runtime/Slate/Private/Framework/Styling/ScrollBarWidgetStyle.cpp
--- a/Engine/Source/Runtime/Slate/Private/Framework/Styling/ScrollBarWidgetStyle.cpp
+++ b/Engine/Source/Runtime/Slate/Private/Framework/Styling/ScrollBarWidgetStyle.cpp
@@ -2,9 +2,3 @@
 
 #include "SlatePrivatePCH.h"
 
-
-UScrollBarWidgetStyle::UScrollBarWidgetStyle( const FObjectInitializer& ObjectInitializer )
-	: Super(ObjectInitializer)
-{
-	
-}
\ No newline at end of file
diff --git a/Engine/Source/Runtime/Slate/Private/Framework/Styling/ScrollBoxWidgetStyle.cpp b/Engine/Source/Runtime/Slate/Private/Framework/Styling/ScrollBoxWidgetStyle.cpp
--- a/Engine/Source/Runtime/Slate/Private/Framework/Styling/ScrollBoxWidgetStyle.cpp
+++ b/Engine/Source/Runtime/Slate/Private/Framework/Styling/ScrollBoxWidgetStyle.cpp
@@ -3,8 +3,3 @@
 #include "SlatePrivatePCH.h"
 
 
-UScrollBoxWidgetStyle::UScrollBoxWidgetStyle( const FObjectInitializer& ObjectInitializer )
-	: Super(ObjectInitializer)
-{
-	
-}
\ No newline at end of file
diff --git a/Engine/Source/Runtime/Slate/Private/Framework/Styling/SpinBoxWidgetStyle.cpp b/Engine/Source/Runtime/Slate/Private/Framework/Styling/SpinBoxWidgetStyle.cpp
--- a/Engine/Source/Runtime/Slate/Private/Framework/Styling/SpinBoxWidgetStyle.cpp
+++ b/Engine/Source/Runtime/Slate/Private/Framework/Styling/SpinBoxWidgetStyle.cpp
@@ -3,8 +3,3 @@
 #include "SlatePrivatePCH.h"
 
 
-USpinBoxWidgetStyle::USpinBoxWidgetStyle(const FObjectInitializer& ObjectInitializer)
-	: Super(ObjectInitializer)
-{
-
-}
\ No newline at end of file
diff --git a/Engine/Source/Runtime/Slate/Private/Framework/Styling/TextBlockWidgetStyle.cpp b/Engine/Source/Runtime/Slate/Private/Framework/Styling/TextBlockWidgetStyle.cpp
--- a/Engine/Source/Runtime/Slate/Private/Framework/Styling/TextBlockWidgetStyle.cpp
+++ b/Engine/Source/Runtime/Slate/Private/Framework/Styling/TextBlockWidgetStyle.cpp
@@ -3,8 +3,3 @@
 #include "SlatePrivatePCH.h"
 
 
-UTextBlockWidgetStyle::UTextBlockWidgetStyle( const FObjectInitializer& ObjectInitializer )
-	: Super(ObjectInitializer)
-{
-	
-}
\ No newline at end of file
diff --git a/Engine/Source/Runtime/Slate/Public/Framework/Styling/ButtonWidgetStyle.h b/Engine/Source/Runtime/Slate/Public/Framework/Styling/ButtonWidgetStyle.h
--- a/Engine/Source/Runtime/Slate/Public/Framework/Styling/ButtonWidgetStyle.h
+++ b/Engine/Source/Runtime/Slate/Public/Framework/Styling/ButtonWidgetStyle.h
@@ -10,7 +10,8 @@
 UCLASS(BlueprintType, hidecategories=Object, MinimalAPI)
 class UButtonWidgetStyle : public USlateWidgetStyleContainerBase
 {
-	GENERATED_UCLASS_BODY()
+public:
+	GENERATED_BODY()
 
 public:
 	/** The actual data describing the button's appearance. */
diff --git a/Engine/Source/Runtime/Slate/Public/Framework/Styling/CheckBoxWidgetStyle.h b/Engine/Source/Runtime/Slate/Public/Framework/Styling/CheckBoxWidgetStyle.h
--- a/Engine/Source/Runtime/Slate/Public/Framework/Styling/CheckBoxWidgetStyle.h
+++ b/Engine/Source/Runtime/Slate/Public/Framework/Styling/CheckBoxWidgetStyle.h
@@ -9,7 +9,7 @@
 UCLASS(hidecategories=Object, MinimalAPI)
 class UCheckBoxWidgetStyle : public USlateWidgetStyleContainerBase
 {
-	GENERATED_UCLASS_BODY()
+	GENERATED_BODY()
 
 public:
 	/** The actual data describing the button's appearance. */
diff --git a/Engine/Source/Runtime/Slate/Public/Framework/Styling/ComboBoxWidgetStyle.h b/Engine/Source/Runtime/Slate/Public/Framework/Styling/ComboBoxWidgetStyle.h
--- a/Engine/Source/Runtime/Slate/Public/Framework/Styling/ComboBoxWidgetStyle.h
+++ b/Engine/Source/Runtime/Slate/Public/Framework/Styling/ComboBoxWidgetStyle.h
@@ -10,7 +10,7 @@
 UCLASS(hidecategories=Object, MinimalAPI)
 class UComboBoxWidgetStyle : public USlateWidgetStyleContainerBase
 {
-	GENERATED_UCLASS_BODY()
+	GENERATED_BODY()
 
 public:
 	/** The actual data describing the combo box's appearance. */
diff --git a/Engine/Source/Runtime/Slate/Public/Framework/Styling/ComboButtonWidgetStyle.h b/Engine/Source/Runtime/Slate/Public/Framework/Styling/ComboButtonWidgetStyle.h
--- a/Engine/Source/Runtime/Slate/Public/Framework/Styling/ComboButtonWidgetStyle.h
+++ b/Engine/Source/Runtime/Slate/Public/Framework/Styling/ComboButtonWidgetStyle.h
@@ -10,7 +10,7 @@
 UCLASS(hidecategories=Object, MinimalAPI)
 class UComboButtonWidgetStyle : public USlateWidgetStyleContainerBase
 {
-	GENERATED_UCLASS_BODY()
+	GENERATED_BODY()
 
 public:
 	/** The actual data describing the combo button's appearance. */
diff --git a/Engine/Source/Runtime/Slate/Public/Framework/Styling/EditableTextBoxWidgetStyle.h b/Engine/Source/Runtime/Slate/Public/Framework/Styling/EditableTextBoxWidgetStyle.h
--- a/Engine/Source/Runtime/Slate/Public/Framework/Styling/EditableTextBoxWidgetStyle.h
+++ b/Engine/Source/Runtime/Slate/Public/Framework/Styling/EditableTextBoxWidgetStyle.h
@@ -10,7 +10,7 @@
 UCLASS(hidecategories=Object, MinimalAPI)
 class UEditableTextBoxWidgetStyle : public USlateWidgetStyleContainerBase
 {
-	GENERATED_UCLASS_BODY()
+	GENERATED_BODY()
 
 public:
 	/** The actual data describing the button's appearance. */
diff --git a/Engine/Source/Runtime/Slate/Public/Framework/Styling/EditableTextWidgetStyle.h b/Engine/Source/Runtime/Slate/Public/Framework/Styling/EditableTextWidgetStyle.h
--- a/Engine/Source/Runtime/Slate/Public/Framework/Styling/EditableTextWidgetStyle.h
+++ b/Engine/Source/Runtime/Slate/Public/Framework/Styling/EditableTextWidgetStyle.h
@@ -10,7 +10,7 @@
 UCLASS(hidecategories=Object, MinimalAPI)
 class UEditableTextWidgetStyle : public USlateWidgetStyleContainerBase
 {
-	GENERATED_UCLASS_BODY()
+	GENERATED_BODY()
 
 public:
 	/** The actual data describing the button's appearance. */
diff --git a/Engine/Source/Runtime/Slate/Public/Framework/Styling/ProgressWidgetStyle.h b/Engine/Source/Runtime/Slate/Public/Framework/Styling/ProgressWidgetStyle.h
--- a/Engine/Source/Runtime/Slate/Public/Framework/Styling/ProgressWidgetStyle.h
+++ b/Engine/Source/Runtime/Slate/Public/Framework/Styling/ProgressWidgetStyle.h
@@ -10,7 +10,7 @@
 UCLASS(BlueprintType, hidecategories=Object, MinimalAPI)
 class UProgressWidgetStyle : public USlateWidgetStyleContainerBase
 {
-	GENERATED_UCLASS_BODY()
+	GENERATED_BODY()
 
 public:
 	/** The actual data describing the button's appearance. */
diff --git a/Engine/Source/Runtime/Slate/Public/Framework/Styling/ScrollBarWidgetStyle.h b/Engine/Source/Runtime/Slate/Public/Framework/Styling/ScrollBarWidgetStyle.h
--- a/Engine/Source/Runtime/Slate/Public/Framework/Styling/ScrollBarWidgetStyle.h
+++ b/Engine/Source/Runtime/Slate/Public/Framework/Styling/ScrollBarWidgetStyle.h
@@ -10,7 +10,7 @@
 UCLASS(hidecategories=Object, MinimalAPI)
 class UScrollBarWidgetStyle : public USlateWidgetStyleContainerBase
 {
-	GENERATED_UCLASS_BODY()
+	GENERATED_BODY()
 
 public:
 	/** The actual data describing the scrollbox's appearance. */
diff --git a/Engine/Source/Runtime/Slate/Public/Framework/Styling/ScrollBoxWidgetStyle.h b/Engine/Source/Runtime/Slate/Public/Framework/Styling/ScrollBoxWidgetStyle.h
--- a/Engine/Source/Runtime/Slate/Public/Framework/Styling/ScrollBoxWidgetStyle.h
+++ b/Engine/Source/Runtime/Slate/Public/Framework/Styling/ScrollBoxWidgetStyle.h
@@ -10,7 +10,7 @@
 UCLASS(hidecategories=Object, MinimalAPI)
 class UScrollBoxWidgetStyle : public USlateWidgetStyleContainerBase
 {
-	GENERATED_UCLASS_BODY()
+	GENERATED_BODY()
 
 public:
 	/** The actual data describing the scrollbox's appearance. */
diff --git a/Engine/Source/Runtime/Slate/Public/Framework/Styling/SpinBoxWidgetStyle.h b/Engine/Source/Runtime/Slate/Public/Framework/Styling/SpinBoxWidgetStyle.h
--- a/Engine/Source/Runtime/Slate/Public/Framework/Styling/SpinBoxWidgetStyle.h
+++ b/Engine/Source/Runtime/Slate/Public/Framework/Styling/SpinBoxWidgetStyle.h
@@ -10,7 +10,7 @@
 UCLASS(hidecategories = Object, MinimalAPI)
 class USpinBoxWidgetStyle : public USlateWidgetStyleContainerBase
 {
-	GENERATED_UCLASS_BODY()
+	GENERATED_BODY()
 
 public:
 	/** The actual data describing the button's appearance. */
diff --git a/Engine/Source/Runtime/Slate/Public/Framework/Styling/TextBlockWidgetStyle.h b/Engine/Source/Runtime/Slate/Public/Framework/Styling/TextBlockWidgetStyle.h
--- a/Engine/Source/Runtime/Slate/Public/Framework/Styling/TextBlockWidgetStyle.h
+++ b/Engine/Source/Runtime/Slate/Public/Framework/Styling/TextBlockWidgetStyle.h
@@ -10,7 +10,7 @@
 UCLASS(hidecategories=Object, MinimalAPI)
 class UTextBlockWidgetStyle : public USlateWidgetStyleContainerBase
 {
-	GENERATED_UCLASS_BODY()
+	GENERATED_BODY()
 
 public:
 	/** The actual data describing the button's appearance. */
diff --git a/Engine/Source/Runtime/SlateCore/Private/Styling/SlateTypes.cpp b/Engine/Source/Runtime/SlateCore/Private/Styling/SlateTypes.cpp
--- a/Engine/Source/Runtime/SlateCore/Private/Styling/SlateTypes.cpp
+++ b/Engine/Source/Runtime/SlateCore/Private/Styling/SlateTypes.cpp
@@ -653,7 +653,3 @@ const FWindowStyle& FWindowStyle::GetDefault()
 }
 
 
-USlateTypes::USlateTypes( const FObjectInitializer& ObjectInitializer )
-	: Super(ObjectInitializer)
-{
-}
diff --git a/Engine/Source/Runtime/SlateCore/Private/Styling/SlateWidgetStyleAsset.cpp b/Engine/Source/Runtime/SlateCore/Private/Styling/SlateWidgetStyleAsset.cpp
--- a/Engine/Source/Runtime/SlateCore/Private/Styling/SlateWidgetStyleAsset.cpp
+++ b/Engine/Source/Runtime/SlateCore/Private/Styling/SlateWidgetStyleAsset.cpp
@@ -3,8 +3,3 @@
 #include "SlateCorePrivatePCH.h"
 
 
-USlateWidgetStyleAsset::USlateWidgetStyleAsset( const FObjectInitializer& ObjectInitializer )
-	: Super(ObjectInitializer)
-{
-	
-}
diff --git a/Engine/Source/Runtime/SlateCore/Private/Styling/SlateWidgetStyleContainerBase.cpp b/Engine/Source/Runtime/SlateCore/Private/Styling/SlateWidgetStyleContainerBase.cpp
--- a/Engine/Source/Runtime/SlateCore/Private/Styling/SlateWidgetStyleContainerBase.cpp
+++ b/Engine/Source/Runtime/SlateCore/Private/Styling/SlateWidgetStyleContainerBase.cpp
@@ -7,11 +7,6 @@
 DEFINE_LOG_CATEGORY(LogSlateStyle);
 
 
-USlateWidgetStyleContainerBase::USlateWidgetStyleContainerBase( const FObjectInitializer& ObjectInitializer )
-	: Super(ObjectInitializer)
-{ }
-
-
 const struct FSlateWidgetStyle* const USlateWidgetStyleContainerBase::GetStyle() const
 {
 	return nullptr;
diff --git a/Engine/Source/Runtime/SlateCore/Public/Styling/SlateTypes.h b/Engine/Source/Runtime/SlateCore/Public/Styling/SlateTypes.h
--- a/Engine/Source/Runtime/SlateCore/Public/Styling/SlateTypes.h
+++ b/Engine/Source/Runtime/SlateCore/Public/Styling/SlateTypes.h
@@ -1471,6 +1471,6 @@ UCLASS()
 class USlateTypes : public UObject
 {
 public:
-	GENERATED_UCLASS_BODY()
+	GENERATED_BODY()
 
 };
diff --git a/Engine/Source/Runtime/SlateCore/Public/Styling/SlateWidgetStyleAsset.h b/Engine/Source/Runtime/SlateCore/Public/Styling/SlateWidgetStyleAsset.h
--- a/Engine/Source/Runtime/SlateCore/Public/Styling/SlateWidgetStyleAsset.h
+++ b/Engine/Source/Runtime/SlateCore/Public/Styling/SlateWidgetStyleAsset.h
@@ -10,7 +10,7 @@
 UCLASS(hidecategories=Object)
 class SLATECORE_API USlateWidgetStyleAsset : public UObject
 {
-	GENERATED_UCLASS_BODY()
+	GENERATED_BODY()
 		  
 public:  
 	/**  */
diff --git a/Engine/Source/Runtime/SlateCore/Public/Styling/SlateWidgetStyleContainerBase.h b/Engine/Source/Runtime/SlateCore/Public/Styling/SlateWidgetStyleContainerBase.h
--- a/Engine/Source/Runtime/SlateCore/Public/Styling/SlateWidgetStyleContainerBase.h
+++ b/Engine/Source/Runtime/SlateCore/Public/Styling/SlateWidgetStyleContainerBase.h
@@ -13,7 +13,7 @@ SLATECORE_API DECLARE_LOG_CATEGORY_EXTERN(LogSlateStyle, Log, All);
 UCLASS(hidecategories=Object)
 class SLATECORE_API USlateWidgetStyleContainerBase : public UObject, public ISlateWidgetStyleContainerInterface
 {
-	GENERATED_UCLASS_BODY()
+	GENERATED_BODY()
 
 public:
 
