Disabled randomized cook order code.

diff --git a/Engine/Source/Editor/UnrealEd/Private/CookOnTheFlyServer.cpp b/Engine/Source/Editor/UnrealEd/Private/CookOnTheFlyServer.cpp
--- a/Engine/Source/Editor/UnrealEd/Private/CookOnTheFlyServer.cpp
+++ b/Engine/Source/Editor/UnrealEd/Private/CookOnTheFlyServer.cpp
@@ -2407,7 +2407,7 @@ void UCookOnTheFlyServer::AddFileToCook( TArray<FString>& InOutFilesToCook, cons
 { 
 	if (!FPackageName::IsScriptPackage(InFilename))
 	{
-#if 1 // randomize cook file order, don't check in enabled...
+#if 0 // randomize cook file order, don't check in enabled...
 		if ( !InOutFilesToCook.Contains(InFilename) )
 		{
 			int Index = FMath::RandRange(0,InOutFilesToCook.Num()-1);
