Fixed displaying of aspect ratios in Common Screen Resolution tooltips in Play settings

diff --git a/Engine/Source/Editor/UnrealEd/Private/Settings/LevelEditorPlaySettingsCustomization.h b/Engine/Source/Editor/UnrealEd/Private/Settings/LevelEditorPlaySettingsCustomization.h
--- a/Engine/Source/Editor/UnrealEd/Private/Settings/LevelEditorPlaySettingsCustomization.h
+++ b/Engine/Source/Editor/UnrealEd/Private/Settings/LevelEditorPlaySettingsCustomization.h
@@ -233,7 +233,7 @@ protected:
 				Args.Add(TEXT("Height"), FText::AsNumber(Iter->Height, NULL, I18N.GetInvariantCulture()));
 				Args.Add(TEXT("AspectRatio"), FText::FromString(Iter->AspectRatio));
 
-				MenuBuilder.AddMenuEntry(FText::FromString(Iter->Description), FText::Format(LOCTEXT("CommonResolutionFormat", "{Width} x {Height} (AspectRatio)"), Args), FSlateIcon(), Action);
+				MenuBuilder.AddMenuEntry(FText::FromString(Iter->Description), FText::Format(LOCTEXT("CommonResolutionFormat", "{Width} x {Height} ({AspectRatio})"), Args), FSlateIcon(), Action);
 			}
 		}
 		MenuBuilder.EndSection();
