When deleting Blueprint assets, child Blueprints will not be re-parented and recompiled when they are also slated for deletion.

diff --git a/Engine/Source/Editor/UnrealEd/Private/ObjectTools.cpp b/Engine/Source/Editor/UnrealEd/Private/ObjectTools.cpp
--- a/Engine/Source/Editor/UnrealEd/Private/ObjectTools.cpp
+++ b/Engine/Source/Editor/UnrealEd/Private/ObjectTools.cpp
@@ -1946,14 +1946,18 @@ namespace ObjectTools
 									UBlueprint* ChildBlueprint = Cast<UBlueprint>(ChildClass->ClassGeneratedBy);
 									if(ChildBlueprint != nullptr)
 									{
-										ChildBlueprint->Modify();
-										ChildBlueprint->ParentClass = BlueprintObject->ParentClass;
-
-										// Recompile the child blueprint to fix up the generated class
-										FKismetEditorUtilities::CompileBlueprint(ChildBlueprint, false, true);
-
-										// Defer garbage collection until after we're done processing the list of objects
-										bNeedsGarbageCollection = true;
+										// Do not reparent and recompile a Blueprint that is going to be deleted.
+										if (ObjectsToDelete.Find(ChildBlueprint) == INDEX_NONE)
+										{
+											ChildBlueprint->Modify();
+											ChildBlueprint->ParentClass = BlueprintObject->ParentClass;
+
+											// Recompile the child blueprint to fix up the generated class
+											FKismetEditorUtilities::CompileBlueprint(ChildBlueprint, false, true);
+
+											// Defer garbage collection until after we're done processing the list of objects
+											bNeedsGarbageCollection = true;
+										}
 									}
 								}
 							}
