Fixed text export of UModel objects, so that Brush objects are forward declared before being defined.

UModel objects were not being forward declared, meaning that the reference in the BrushComponent was ambiguous (and also referring to the wrong object).  By luck, this was being fixed up at the end of ULevelFactory::FactoryCreateText for this particular case, but other exported UModel* properties would likely fail.  This change correctly forward declares any UModel objects in the export text and allows them to be redefined during import if they are uninitialized.

diff --git a/Engine/Source/Editor/UnrealEd/Private/EditorExporters.cpp b/Engine/Source/Editor/UnrealEd/Private/EditorExporters.cpp
--- a/Engine/Source/Editor/UnrealEd/Private/EditorExporters.cpp
+++ b/Engine/Source/Editor/UnrealEd/Private/EditorExporters.cpp
@@ -317,8 +317,9 @@ UModelExporterT3D::UModelExporterT3D(const FObjectInitializer& ObjectInitializer
 	bText = true;
 	PreferredFormatIndex = 0;
 	FormatExtension.Add(TEXT("T3D"));
+	FormatExtension.Add(TEXT("COPY"));
+	FormatDescription.Add(TEXT("Unreal model text"));
 	FormatDescription.Add(TEXT("Unreal model text"));
-
 }
 
 bool UModelExporterT3D::ExportText( const FExportObjectInnerContext* Context, UObject* Object, const TCHAR* Type, FOutputDevice& Ar, FFeedbackContext* Warn, uint32 PortFlags )
@@ -326,8 +327,10 @@ bool UModelExporterT3D::ExportText( const FExportObjectInnerContext* Context, UO
 	UModel* Model = CastChecked<UModel>( Object );
 
 	Ar.Logf( TEXT("%sBegin Brush Name=%s\r\n"), FCString::Spc(TextIndent), *Model->GetName() );
-		UExporter::ExportToOutputDevice( Context, Model->Polys, NULL, Ar, Type, TextIndent+3, PortFlags );
-//		ExportObjectInner( Context, Model, Ar, PortFlags | PPF_ExportsNotFullyQualified );
+	if (!(PortFlags & PPF_SeparateDeclare))
+	{
+		UExporter::ExportToOutputDevice(Context, Model->Polys, NULL, Ar, Type, TextIndent + 3, PortFlags);
+	}
 	Ar.Logf( TEXT("%sEnd Brush\r\n"), FCString::Spc(TextIndent) );
 
 	return 1;
diff --git a/Engine/Source/Editor/UnrealEd/Private/EditorObject.cpp b/Engine/Source/Editor/UnrealEd/Private/EditorObject.cpp
--- a/Engine/Source/Editor/UnrealEd/Private/EditorObject.cpp
+++ b/Engine/Source/Editor/UnrealEd/Private/EditorObject.cpp
@@ -221,14 +221,16 @@ static const TCHAR* ImportProperties(
 			TCHAR BrushName[NAME_SIZE];
 			if( FParse::Value( Str, TEXT("Name="), BrushName, NAME_SIZE ) )
 			{
-				// If a brush with this name already exists in the
-				// level, rename the existing one.  This is necessary
-				// because we can't rename the brush we're importing without
-				// losing our ability to associate it with the actor properties
-				// that reference it.
+				// If an initialized brush with this name already exists in the level, rename the existing one.
+				// It is deemed to be initialized if it has a non-zero poly count.
+				// If it is uninitialized, the existing object will have been created by a forward reference in the import text,
+				// and it will now be redefined.  This relies on the behavior that NewObject<> will return an existing pointer
+				// if an object with the same name and outer is passed.
 				UModel* ExistingBrush = FindObject<UModel>( SubobjectRoot, BrushName );
-				if( ExistingBrush )
+				if (ExistingBrush && ExistingBrush->Polys && ExistingBrush->Polys->Element.Num() > 0)
+				{
 					ExistingBrush->Rename();
+				}
 
 				// Create model.
 				UModelFactory* ModelFactory = NewObject<UModelFactory>();
diff --git a/Engine/Source/Runtime/Engine/Private/UnrealExporter.cpp b/Engine/Source/Runtime/Engine/Private/UnrealExporter.cpp
--- a/Engine/Source/Runtime/Engine/Private/UnrealExporter.cpp
+++ b/Engine/Source/Runtime/Engine/Private/UnrealExporter.cpp
@@ -604,7 +604,7 @@ void UExporter::ExportObjectInner(const FExportObjectInnerContext* Context, UObj
 	{
 		for (UObject* Obj : ObjectInners)
 		{
-			if (!Obj->HasAnyFlags(RF_TextExportTransient) && Obj->GetClass() != UModel::StaticClass())
+			if (!Obj->HasAnyFlags(RF_TextExportTransient))
 			{
 				// export the object
 				UExporter::ExportToOutputDevice( Context, Obj, NULL, Ar, (PortFlags & PPF_Copy) ? TEXT("Copy") : TEXT("T3D"), TextIndent, PortFlags | PPF_SeparateDeclare, false, ExportRootScope );
