Fixed issue where a map containing sublevels which have just been renamed is incorrectly opened.

diff --git a/Engine/Source/Editor/UnrealEd/Private/EditorServer.cpp b/Engine/Source/Editor/UnrealEd/Private/EditorServer.cpp
--- a/Engine/Source/Editor/UnrealEd/Private/EditorServer.cpp
+++ b/Engine/Source/Editor/UnrealEd/Private/EditorServer.cpp
@@ -2183,6 +2183,19 @@ bool UEditorEngine::Map_Load(const TCHAR* Str, FOutputDevice& Ar)
 					}
 					EditorDestroyWorld( Context, LocalizedLoadingMap, NewWorld );
 
+					// Unload all other map packages currently loaded, before opening a new map.
+					// The world is only initialized correctly as part of the level loading process, so ensure that every map package needs loading.
+					TArray<UPackage*> WorldPackages;
+					for (TObjectIterator<UWorld> It; It; ++It)
+					{
+						UPackage* Package = Cast<UPackage>(It->GetOuter());
+						if (Package && Package != GetTransientPackage())
+						{
+							WorldPackages.AddUnique(Package);
+						}
+					}
+					PackageTools::UnloadPackages(WorldPackages);
+
 					// Refresh ExistingPackage and Existing World now that GC has occurred.
 					ExistingPackage = FindPackage(NULL, *LongTempFname);
 					if (ExistingPackage)
