Fixed issue where BSP flickered off and on when moving a brush in a level.

- Fixed issue in FEditorViewportClient::StopTracking() where the HitProxy was being invalidated with immediate effect.  When this is done, the viewport must be redrawn directly afterwards.

diff --git a/Engine/Source/Editor/UnrealEd/Private/EditorViewportClient.cpp b/Engine/Source/Editor/UnrealEd/Private/EditorViewportClient.cpp
--- a/Engine/Source/Editor/UnrealEd/Private/EditorViewportClient.cpp
+++ b/Engine/Source/Editor/UnrealEd/Private/EditorViewportClient.cpp
@@ -2062,18 +2062,45 @@ void FEditorViewportClient::StopTracking()
 		MouseDeltaTracker->EndTracking( this );
 
 		Widget->SetCurrentAxis( EAxisList::None );
-		Invalidate( true, true );
+
+		// Force an immediate redraw of the viewport and hit proxy.
+		// The results are required straight away, so it is not sufficient to defer the redraw until the next tick.
+		if (Viewport)
+		{
+			Viewport->InvalidateHitProxy();
+			Viewport->Draw();
+
+			// If there are child viewports, force a redraw on those too
+			FSceneViewStateInterface* ParentView = ViewState.GetReference();
+			if (ParentView->IsViewParent())
+			{
+				for (FEditorViewportClient* ViewportClient : GEditor->AllViewportClients)
+				{
+					if (ViewportClient != nullptr)
+					{
+						FSceneViewStateInterface* ViewportParentView = ViewportClient->ViewState.GetReference();
+
+						if (ViewportParentView != nullptr &&
+							ViewportParentView->HasViewParent() &&
+							ViewportParentView->GetViewParent() == ParentView &&
+							!ViewportParentView->IsViewParent())
+						{
+							ViewportClient->Viewport->InvalidateHitProxy();
+							ViewportClient->Viewport->Draw();
+						}
+					}
+				}
+			}
+		}
 
 		SetRequiredCursorOverride( false );
 
 		bWidgetAxisControlledByDrag = false;
 
-		// Update the hovered hit proxy here.  If the user didnt move the mouse
-		// they still need to be able to pick up the gizmo without moving the mouse again
-		Viewport->InvalidateHitProxy();
-		HHitProxy* HitProxy = Viewport->GetHitProxy(CachedMouseX,CachedMouseY);
-
-		CheckHoveredHitProxy(HitProxy);
+ 		// Update the hovered hit proxy here.  If the user didnt move the mouse
+ 		// they still need to be able to pick up the gizmo without moving the mouse again
+ 		HHitProxy* HitProxy = Viewport->GetHitProxy(CachedMouseX,CachedMouseY);
+  		CheckHoveredHitProxy(HitProxy);
 
 		bIsTracking = false;	
 	}
