If a realtime viewport is visible, Slate will correctly stay awake after ending a PIE session

diff --git a/Engine/Source/Editor/UnrealEd/Private/EditorViewportClient.cpp b/Engine/Source/Editor/UnrealEd/Private/EditorViewportClient.cpp
--- a/Engine/Source/Editor/UnrealEd/Private/EditorViewportClient.cpp
+++ b/Engine/Source/Editor/UnrealEd/Private/EditorViewportClient.cpp
@@ -365,6 +365,26 @@ FEditorViewportClient::~FEditorViewportClient()
 	}
 }
 
+void FEditorViewportClient::RestoreRealtime(const bool bAllowDisable)
+{
+	if (bAllowDisable)
+	{
+		bIsRealtime = bStoredRealtime;
+		bShowStats = bStoredShowStats;
+	}
+	else
+	{
+		bIsRealtime |= bStoredRealtime;
+		bShowStats |= bStoredShowStats;
+	}
+
+	if (EditorViewportWidget.IsValid() && ( bIsRealtime || bShowStats ))
+	{
+		// Invalidate the viewport widget to re-register its active timer
+		EditorViewportWidget.Pin()->Invalidate();
+	}
+}
+
 void FEditorViewportClient::RedrawRequested(FViewport* InViewport)
 {
 	bNeedsRedraw = true;
diff --git a/Engine/Source/Editor/UnrealEd/Public/EditorViewportClient.h b/Engine/Source/Editor/UnrealEd/Public/EditorViewportClient.h
--- a/Engine/Source/Editor/UnrealEd/Public/EditorViewportClient.h
+++ b/Engine/Source/Editor/UnrealEd/Public/EditorViewportClient.h
@@ -290,19 +290,7 @@ public:
 	 * Restores realtime setting to stored value. This will only enable realtime and 
 	 * never disable it (unless bAllowDisable is true)
 	 */
-	void RestoreRealtime( const bool bAllowDisable = false )
-	{
-		if( bAllowDisable )
-		{
-			bIsRealtime = bStoredRealtime;
-			bShowStats = bStoredShowStats;
-		}
-		else
-		{
-			bIsRealtime |= bStoredRealtime;
-			bShowStats |= bStoredShowStats;
-		}
-	}
+	void RestoreRealtime(const bool bAllowDisable = false);
 
 
 	// this set ups camera for both orbit and non orbit control
