Fix crash drag dropping objects into the world when a slow task dialog appears due to a long load.

diff --git a/Engine/Source/Editor/UnrealEd/Private/DragAndDrop/AssetDragDropOp.cpp b/Engine/Source/Editor/UnrealEd/Private/DragAndDrop/AssetDragDropOp.cpp
--- a/Engine/Source/Editor/UnrealEd/Private/DragAndDrop/AssetDragDropOp.cpp
+++ b/Engine/Source/Editor/UnrealEd/Private/DragAndDrop/AssetDragDropOp.cpp
@@ -143,6 +143,13 @@ void FAssetDragDropOp::Init()
 {
 	if ( AssetData.Num() > 0 && ThumbnailSize > 0 )
 	{
+		// Load all assets first so that there is no loading going on while attempting to drag
+		// Can cause unsafe frame reentry 
+		for( FAssetData& Data : AssetData )
+		{
+			Data.GetAsset();
+		}
+
 		// Create a thumbnail pool to hold the single thumbnail rendered
 		ThumbnailPool = MakeShareable( new FAssetThumbnailPool(1, /*InAreRealTileThumbnailsAllowed=*/false) );
 
