Reintroducing the complex open asset automations test.  This is intended to help out the content team.

diff --git a/Engine/Source/Editor/UnrealEd/Private/Tests/EditorOpenAssetAutomationTests.cpp b/Engine/Source/Editor/UnrealEd/Private/Tests/EditorOpenAssetAutomationTests.cpp
--- a/Engine/Source/Editor/UnrealEd/Private/Tests/EditorOpenAssetAutomationTests.cpp
+++ b/Engine/Source/Editor/UnrealEd/Private/Tests/EditorOpenAssetAutomationTests.cpp
@@ -240,583 +240,661 @@ void FOpenAllAssetTests(OpenAssetParameters& AssetList)
 
 //////////////////////////////////////////////////////////////////////////
 /**
-* This test opens each AIM OFFSET BLENDSPACE into its sub-editor.
+* This test opens each AimOffsetBlendSpace into its sub-editor.
 */
-IMPLEMENT_SIMPLE_AUTOMATION_TEST(FOpenAimOffsetBlendspaceAssetEditorsSLOW, "Tools.Skeletal Mesh.Open All Aim Offset Blendspace Assets (SLOW)", (EAutomationTestFlags::ATF_Editor | EAutomationTestFlags::ATF_RequiresUser));
+IMPLEMENT_COMPLEX_AUTOMATION_TEST(FOpenAimOffsetBlendSpaceAssetEditors, "Tools.Skeletal Mesh.Animation.BlendSpaces.Open Aim Offset BlendSpace Assets", (EAutomationTestFlags::ATF_Editor | EAutomationTestFlags::ATF_RequiresUser));
 
-bool FOpenAimOffsetBlendspaceAssetEditorsSLOW::RunTest(const FString& Parameters)
+void FOpenAimOffsetBlendSpaceAssetEditors::GetTests(TArray<FString>& OutBeautifiedNames, TArray<FString>& OutTestCommands) const
 {
-	OpenAssetParameters AssetList;
-	AssetList.TimeSpent = FPlatformTime::Seconds();
-
-	FEditorAutomationTestUtilities::CollectGameContentByClass(UAimOffsetBlendSpace::StaticClass(), true, AssetList.AssetPathList);
-	FOpenAllAssetTests(AssetList);
+	//This grabs each AimOffsetBlendSpace asset in the Game/Content directory
+	FEditorAutomationTestUtilities::CollectGameContentTestsByClass(UAimOffsetBlendSpace::StaticClass(), true, OutBeautifiedNames, OutTestCommands);
+}
 
-	return true;
+bool FOpenAimOffsetBlendSpaceAssetEditors::RunTest(const FString& Parameters)
+{
+	bool bDidTheTestPass = FOpenActualAssetEditors(Parameters);
+	CollectGarbage(GARBAGE_COLLECTION_KEEPFLAGS);
+	return bDidTheTestPass;
 }
 
+//////////////////////////////////////////////////////////////////////////
 /**
-* This test opens each AIM OFFSET BLENDSPACE 1D into its sub-editor.
+* This test opens each AimOffsetBlendSpace1D into its sub-editor.
 */
-IMPLEMENT_SIMPLE_AUTOMATION_TEST(FOpenAimOffSetBlendspace1DAssetEditorsSLOW, "Tools.Skeletal Mesh.Open All Aim Offset Blendspace 1D Assets (SLOW)", (EAutomationTestFlags::ATF_Editor | EAutomationTestFlags::ATF_RequiresUser));
+IMPLEMENT_COMPLEX_AUTOMATION_TEST(FOpenAimOffsetBlendSpace1DAssetEditors, "Tools.Skeletal Mesh.Animation.BlendSpaces.Open Aim Offset BlendSpace1D Assets", (EAutomationTestFlags::ATF_Editor | EAutomationTestFlags::ATF_RequiresUser));
 
-bool FOpenAimOffSetBlendspace1DAssetEditorsSLOW::RunTest(const FString& Parameters)
+void FOpenAimOffsetBlendSpace1DAssetEditors::GetTests(TArray<FString>& OutBeautifiedNames, TArray<FString>& OutTestCommands) const
 {
-	OpenAssetParameters AssetList;
-	AssetList.TimeSpent = FPlatformTime::Seconds();
-
-	FEditorAutomationTestUtilities::CollectGameContentByClass(UAimOffsetBlendSpace1D::StaticClass(), true, AssetList.AssetPathList);
-
-	FOpenAllAssetTests(AssetList);
+	//This grabs each AimOffsetBlendSpace1D asset in the Game/Content directory
+	FEditorAutomationTestUtilities::CollectGameContentTestsByClass(UAimOffsetBlendSpace1D::StaticClass(), true, OutBeautifiedNames, OutTestCommands);
+}
 
-	return true;
+bool FOpenAimOffsetBlendSpace1DAssetEditors::RunTest(const FString& Parameters)
+{
+	bool bDidTheTestPass = FOpenActualAssetEditors(Parameters);
+	CollectGarbage(GARBAGE_COLLECTION_KEEPFLAGS);
+	return bDidTheTestPass;
 }
 
+//////////////////////////////////////////////////////////////////////////
 /**
-* This test opens each ANIMMONTAGE into its sub-editor.
+* This test opens each AnimMontage into its sub-editor.
 */
-IMPLEMENT_SIMPLE_AUTOMATION_TEST(FOpenAnimMontageAssetEditorsSLOW, "Tools.Skeletal Mesh.Open All Anim Montage Assets (SLOW)", (EAutomationTestFlags::ATF_Editor | EAutomationTestFlags::ATF_RequiresUser));
+IMPLEMENT_COMPLEX_AUTOMATION_TEST(FOpenAnimMontageAssetEditors, "Tools.Skeletal Mesh.Animation.Open AnimMontage Assets", (EAutomationTestFlags::ATF_Editor | EAutomationTestFlags::ATF_RequiresUser));
 
-bool FOpenAnimMontageAssetEditorsSLOW::RunTest(const FString& Parameters)
+void FOpenAnimMontageAssetEditors::GetTests(TArray<FString>& OutBeautifiedNames, TArray<FString>& OutTestCommands) const
 {
-	OpenAssetParameters AssetList;
-	AssetList.TimeSpent = FPlatformTime::Seconds();
-
-	FEditorAutomationTestUtilities::CollectGameContentByClass(UAnimMontage::StaticClass(), true, AssetList.AssetPathList);
-
-	FOpenAllAssetTests(AssetList);
+	//This grabs each AnimMontage asset in the Game/Content directory
+	FEditorAutomationTestUtilities::CollectGameContentTestsByClass(UAnimMontage::StaticClass(), true, OutBeautifiedNames, OutTestCommands);
+}
 
-	return true;
+bool FOpenAnimMontageAssetEditors::RunTest(const FString& Parameters)
+{
+	bool bDidTheTestPass = FOpenActualAssetEditors(Parameters);
+	CollectGarbage(GARBAGE_COLLECTION_KEEPFLAGS);
+	return bDidTheTestPass;
 }
 
+//////////////////////////////////////////////////////////////////////////
 /**
-* This test opens each ANIMSEQUENCE into its sub-editor.
+* This test opens each AnimSequence into its sub-editor.
 */
-IMPLEMENT_SIMPLE_AUTOMATION_TEST(FOpenAnimSequenceAssetEditorsSLOW, "Tools.Skeletal Mesh.Open All Anim Sequence Assets (SLOW)", (EAutomationTestFlags::ATF_Editor | EAutomationTestFlags::ATF_RequiresUser));
+IMPLEMENT_COMPLEX_AUTOMATION_TEST(FOpenAnimSequenceAssetEditors, "Tools.Skeletal Mesh.Animation.Open AnimSequence Assets", (EAutomationTestFlags::ATF_Editor | EAutomationTestFlags::ATF_RequiresUser));
 
-bool FOpenAnimSequenceAssetEditorsSLOW::RunTest(const FString& Parameters)
+void FOpenAnimSequenceAssetEditors::GetTests(TArray<FString>& OutBeautifiedNames, TArray<FString>& OutTestCommands) const
 {
-	OpenAssetParameters AssetList;
-	AssetList.TimeSpent = FPlatformTime::Seconds();
-
-	FEditorAutomationTestUtilities::CollectGameContentByClass(UAnimSequence::StaticClass(), true, AssetList.AssetPathList);
-
-	FOpenAllAssetTests(AssetList);
+	//This grabs each AnimSequence asset in the Game/Content directory
+	FEditorAutomationTestUtilities::CollectGameContentTestsByClass(UAnimSequence::StaticClass(), true, OutBeautifiedNames, OutTestCommands);
+}
 
-	return true;
+bool FOpenAnimSequenceAssetEditors::RunTest(const FString& Parameters)
+{
+	bool bDidTheTestPass = FOpenActualAssetEditors(Parameters);
+	CollectGarbage(GARBAGE_COLLECTION_KEEPFLAGS);
+	return bDidTheTestPass;
 }
 
+//////////////////////////////////////////////////////////////////////////
 /**
-* This test opens each BEHAVIORTREE into its sub-editor.
+* This test opens each BEHAVIOR TREE into its sub-editor.
 */
-IMPLEMENT_SIMPLE_AUTOMATION_TEST(FOpenBehaviorTreeAssetEditorsSLOW, "Tools.AI.Open All Behaviour Tree Assets (SLOW)", (EAutomationTestFlags::ATF_Editor | EAutomationTestFlags::ATF_RequiresUser));
+IMPLEMENT_COMPLEX_AUTOMATION_TEST(FOpenBehaviorTreeAssetEditors, "Tools.AI.Open Behavior Tree Assets", (EAutomationTestFlags::ATF_Editor | EAutomationTestFlags::ATF_RequiresUser));
 
-bool FOpenBehaviorTreeAssetEditorsSLOW::RunTest(const FString& Parameters)
+void FOpenBehaviorTreeAssetEditors::GetTests(TArray<FString>& OutBeautifiedNames, TArray<FString>& OutTestCommands) const
 {
-	OpenAssetParameters AssetList;
-	AssetList.TimeSpent = FPlatformTime::Seconds();
-
-	FEditorAutomationTestUtilities::CollectGameContentByClass(UBehaviorTree::StaticClass(), true, AssetList.AssetPathList);
-
-	FOpenAllAssetTests(AssetList);
+	//This grabs each Behavior Tree asset in the Game/Content directory
+	FEditorAutomationTestUtilities::CollectGameContentTestsByClass(UBehaviorTree::StaticClass(), true, OutBeautifiedNames, OutTestCommands);
+}
 
-	return true;
+bool FOpenBehaviorTreeAssetEditors::RunTest(const FString& Parameters)
+{
+	bool bDidTheTestPass = FOpenActualAssetEditors(Parameters);
+	CollectGarbage(GARBAGE_COLLECTION_KEEPFLAGS);
+	return bDidTheTestPass;
 }
 
+//////////////////////////////////////////////////////////////////////////
 /**
-* This test opens each BLENDSPACE into its sub-editor.
+* This test opens each BlendSpace into its sub-editor.
 */
-IMPLEMENT_SIMPLE_AUTOMATION_TEST(FOpenBlendSpaceAssetEditorsSLOW, "Tools.Skeletal Mesh.Open All BlendSpace Assets (SLOW)", (EAutomationTestFlags::ATF_Editor | EAutomationTestFlags::ATF_RequiresUser));
+IMPLEMENT_COMPLEX_AUTOMATION_TEST(FOpenBlendSpaceAssetEditors, "Tools.Skeletal Mesh.Animation.BlendSpaces.Open BlendSpace Assets", (EAutomationTestFlags::ATF_Editor | EAutomationTestFlags::ATF_RequiresUser));
 
-bool FOpenBlendSpaceAssetEditorsSLOW::RunTest(const FString& Parameters)
+void FOpenBlendSpaceAssetEditors::GetTests(TArray<FString>& OutBeautifiedNames, TArray<FString>& OutTestCommands) const
 {
-	OpenAssetParameters AssetList;
-	AssetList.TimeSpent = FPlatformTime::Seconds();
-
-	FEditorAutomationTestUtilities::CollectGameContentByClass(UBlendSpace::StaticClass(), true, AssetList.AssetPathList);
-
-	FOpenAllAssetTests(AssetList);
+	//This grabs each BlendSpace asset in the Game/Content directory
+	FEditorAutomationTestUtilities::CollectGameContentTestsByClass(UBlendSpace::StaticClass(), false, OutBeautifiedNames, OutTestCommands);
+}
 
-	return true;
+bool FOpenBlendSpaceAssetEditors::RunTest(const FString& Parameters)
+{
+	bool bDidTheTestPass = FOpenActualAssetEditors(Parameters);
+	CollectGarbage(GARBAGE_COLLECTION_KEEPFLAGS);
+	return bDidTheTestPass;
 }
 
+//////////////////////////////////////////////////////////////////////////
 /**
-* This test opens each BLENDSPACE 1D into its sub-editor.
+* This test opens each BlendSpace1D into its sub-editor.
 */
-IMPLEMENT_SIMPLE_AUTOMATION_TEST(FOpenBlendSpace1DAssetEditorsSLOW, "Tools.Skeletal Mesh.Open All BlendSpace 1D Assets (SLOW)", (EAutomationTestFlags::ATF_Editor | EAutomationTestFlags::ATF_RequiresUser));
+IMPLEMENT_COMPLEX_AUTOMATION_TEST(FOpenBlendSpace1DAssetEditors, "Tools.Skeletal Mesh.Animation.BlendSpaces.Open BlendSpace1D Assets", (EAutomationTestFlags::ATF_Editor | EAutomationTestFlags::ATF_RequiresUser));
 
-bool FOpenBlendSpace1DAssetEditorsSLOW::RunTest(const FString& Parameters)
+void FOpenBlendSpace1DAssetEditors::GetTests(TArray<FString>& OutBeautifiedNames, TArray<FString>& OutTestCommands) const
 {
-	OpenAssetParameters AssetList;
-	AssetList.TimeSpent = FPlatformTime::Seconds();
-
-	FEditorAutomationTestUtilities::CollectGameContentByClass(UBlendSpace::StaticClass(), true, AssetList.AssetPathList);
-
-	FOpenAllAssetTests(AssetList);
+	//This grabs each BlendSpace1D asset in the Game/Content directory
+	FEditorAutomationTestUtilities::CollectGameContentTestsByClass(UBlendSpace1D::StaticClass(), false, OutBeautifiedNames, OutTestCommands);
+}
 
-	return true;
+bool FOpenBlendSpace1DAssetEditors::RunTest(const FString& Parameters)
+{
+	bool bDidTheTestPass = FOpenActualAssetEditors(Parameters);
+	CollectGarbage(GARBAGE_COLLECTION_KEEPFLAGS);
+	return bDidTheTestPass;
 }
 
+
+//////////////////////////////////////////////////////////////////////////
 /**
 * This test opens each BLUEPRINT into its sub-editor.
 */
-IMPLEMENT_SIMPLE_AUTOMATION_TEST(FOpenBlueprintAssetEditorsSLOW, "Tools.Blueprint.Open All Blueprint Assets (SLOW)", (EAutomationTestFlags::ATF_Editor | EAutomationTestFlags::ATF_RequiresUser));
+IMPLEMENT_COMPLEX_AUTOMATION_TEST(FOpenBlueprintAssetEditors, "Tools.Blueprint.Open Blueprint Assets", (EAutomationTestFlags::ATF_Editor | EAutomationTestFlags::ATF_RequiresUser));
 
-bool FOpenBlueprintAssetEditorsSLOW::RunTest(const FString& Parameters)
+void FOpenBlueprintAssetEditors::GetTests(TArray<FString>& OutBeautifiedNames, TArray<FString>& OutTestCommands) const
 {
-	OpenAssetParameters AssetList;
-	AssetList.TimeSpent = FPlatformTime::Seconds();
-
-	FEditorAutomationTestUtilities::CollectGameContentByClass(UBlueprint::StaticClass(), true, AssetList.AssetPathList);
-
-	FOpenAllAssetTests(AssetList);
+	//This grabs each BlendSpace asset in the Game/Content directory
+	FEditorAutomationTestUtilities::CollectGameContentTestsByClass(UBlueprint::StaticClass(), true, OutBeautifiedNames, OutTestCommands);
+}
 
-	return true;
+bool FOpenBlueprintAssetEditors::RunTest(const FString& Parameters)
+{
+	bool bDidTheTestPass = FOpenActualAssetEditors(Parameters);
+	CollectGarbage(GARBAGE_COLLECTION_KEEPFLAGS);
+	return bDidTheTestPass;
 }
 
+//////////////////////////////////////////////////////////////////////////
 /**
-* This test opens each DESTRUCTIBLE MESH into its sub-editor.
+* This test opens each UserDefinedEnum into its sub-editor.
 */
-IMPLEMENT_SIMPLE_AUTOMATION_TEST(FOpenDestructibleMeshAssetEditorsSLOW, "Tools.Skeletal Mesh.Open All Destructible Mesh Assets (SLOW)", (EAutomationTestFlags::ATF_Editor | EAutomationTestFlags::ATF_RequiresUser));
+IMPLEMENT_COMPLEX_AUTOMATION_TEST(FOpenUserDefinedEnumAssetEditors, "Tools.Blueprint.Open User Defined Enum Assets", (EAutomationTestFlags::ATF_Editor | EAutomationTestFlags::ATF_RequiresUser));
 
-bool FOpenDestructibleMeshAssetEditorsSLOW::RunTest(const FString& Parameters)
+void FOpenUserDefinedEnumAssetEditors::GetTests(TArray<FString>& OutBeautifiedNames, TArray<FString>& OutTestCommands) const
 {
-	OpenAssetParameters AssetList;
-	AssetList.TimeSpent = FPlatformTime::Seconds();
-
-	FEditorAutomationTestUtilities::CollectGameContentByClass(UDestructibleMesh::StaticClass(), false, AssetList.AssetPathList);
-
-	FOpenAllAssetTests(AssetList);
+	//This grabs each UserDefinedEnum asset in the Game/Content directory
+	FEditorAutomationTestUtilities::CollectGameContentTestsByClass(UUserDefinedEnum::StaticClass(), true, OutBeautifiedNames, OutTestCommands);
+}
 
-	return true;
+bool FOpenUserDefinedEnumAssetEditors::RunTest(const FString& Parameters)
+{
+	bool bDidTheTestPass = FOpenActualAssetEditors(Parameters);
+	CollectGarbage(GARBAGE_COLLECTION_KEEPFLAGS);
+	return bDidTheTestPass;
 }
 
+//////////////////////////////////////////////////////////////////////////
 /**
-* This test opens each DIALOGUE VOICE into its sub-editor.
+* This test opens each DESTRUCTIBLE MESH into its sub-editor.
 */
-IMPLEMENT_SIMPLE_AUTOMATION_TEST(FOpenDialogueVoiceAssetEditorsSLOW, "Tools.Sounds.Open All Dialogue Voice Assets (SLOW)", (EAutomationTestFlags::ATF_Editor | EAutomationTestFlags::ATF_RequiresUser));
+IMPLEMENT_COMPLEX_AUTOMATION_TEST(FOpenDestructibleMeshAssetEditors, "Tools.Skeletal Mesh.Open Destructible Mesh Assets", (EAutomationTestFlags::ATF_Editor | EAutomationTestFlags::ATF_RequiresUser));
 
-bool FOpenDialogueVoiceAssetEditorsSLOW::RunTest(const FString& Parameters)
+void FOpenDestructibleMeshAssetEditors::GetTests(TArray<FString>& OutBeautifiedNames, TArray<FString>& OutTestCommands) const
 {
-	OpenAssetParameters AssetList;
-	AssetList.TimeSpent = FPlatformTime::Seconds();
-
-	FEditorAutomationTestUtilities::CollectGameContentByClass(UDialogueVoice::StaticClass(), true, AssetList.AssetPathList);
-
-	FOpenAllAssetTests(AssetList);
+	//This grabs each DestructibleMesh asset in the Game/Content directory
+	FEditorAutomationTestUtilities::CollectGameContentTestsByClass(UDestructibleMesh::StaticClass(), false, OutBeautifiedNames, OutTestCommands);
+}
 
-	return true;
+bool FOpenDestructibleMeshAssetEditors::RunTest(const FString& Parameters)
+{
+	bool bDidTheTestPass = FOpenActualAssetEditors(Parameters);
+	CollectGarbage(GARBAGE_COLLECTION_KEEPFLAGS);
+	return bDidTheTestPass;
 }
 
+//////////////////////////////////////////////////////////////////////////
 /**
-* This test opens each DIALOGUE WAVE into its sub-editor.
+* This test opens each DialogueVoice into its sub-editor.
 */
-IMPLEMENT_SIMPLE_AUTOMATION_TEST(FOpenDialogueWaveAssetEditorsSLOW, "Tools.Sounds.Open All Dialogue Wave Assets (SLOW)", (EAutomationTestFlags::ATF_Editor | EAutomationTestFlags::ATF_RequiresUser));
+IMPLEMENT_COMPLEX_AUTOMATION_TEST(FOpenDialogueVoiceAssetEditors, "Tools.Sounds.Open Dialogue Voice Assets", (EAutomationTestFlags::ATF_Editor | EAutomationTestFlags::ATF_RequiresUser));
 
-bool FOpenDialogueWaveAssetEditorsSLOW::RunTest(const FString& Parameters)
+void FOpenDialogueVoiceAssetEditors::GetTests(TArray<FString>& OutBeautifiedNames, TArray<FString>& OutTestCommands) const
 {
-	OpenAssetParameters AssetList;
-	AssetList.TimeSpent = FPlatformTime::Seconds();
-
-	FEditorAutomationTestUtilities::CollectGameContentByClass(UDialogueWave::StaticClass(), true, AssetList.AssetPathList);
-
-	FOpenAllAssetTests(AssetList);
+	//This grabs each DialogueVoice asset in the Game/Content directory
+	FEditorAutomationTestUtilities::CollectGameContentTestsByClass(UDialogueVoice::StaticClass(), true, OutBeautifiedNames, OutTestCommands);
+}
 
-	return true;
+bool FOpenDialogueVoiceAssetEditors::RunTest(const FString& Parameters)
+{
+	bool bDidTheTestPass = FOpenActualAssetEditors(Parameters);
+	CollectGarbage(GARBAGE_COLLECTION_KEEPFLAGS);
+	return bDidTheTestPass;
 }
 
+//////////////////////////////////////////////////////////////////////////
 /**
-* This test opens each MATERIAL into its sub-editor.
+* This test opens each DialogueWave into its sub-editor.
 */
-IMPLEMENT_SIMPLE_AUTOMATION_TEST(FOpenMaterialAssetEditorsSLOW, "Tools.Materials.Open All Material Assets (SLOW)", (EAutomationTestFlags::ATF_Editor | EAutomationTestFlags::ATF_RequiresUser));
+IMPLEMENT_COMPLEX_AUTOMATION_TEST(FOpenDialogueWaveAssetEditors, "Tools.Sounds.Open Dialogue Wave Assets", (EAutomationTestFlags::ATF_Editor | EAutomationTestFlags::ATF_RequiresUser));
 
-bool FOpenMaterialAssetEditorsSLOW::RunTest(const FString& Parameters)
+void FOpenDialogueWaveAssetEditors::GetTests(TArray<FString>& OutBeautifiedNames, TArray<FString>& OutTestCommands) const
 {
-	OpenAssetParameters AssetList;
-	AssetList.TimeSpent = FPlatformTime::Seconds();
-
-	FEditorAutomationTestUtilities::CollectGameContentByClass(UMaterial::StaticClass(), true, AssetList.AssetPathList);
-
-	FOpenAllAssetTests(AssetList);
+	//This grabs each DialogueWave asset in the Game/Content directory
+	FEditorAutomationTestUtilities::CollectGameContentTestsByClass(UDialogueWave::StaticClass(), true, OutBeautifiedNames, OutTestCommands);
+}
 
-	return true;
+bool FOpenDialogueWaveAssetEditors::RunTest(const FString& Parameters)
+{
+	bool bDidTheTestPass = FOpenActualAssetEditors(Parameters);
+	CollectGarbage(GARBAGE_COLLECTION_KEEPFLAGS);
+	return bDidTheTestPass;
 }
 
+//////////////////////////////////////////////////////////////////////////
 /**
-* This test opens each MATERIAL FUNCTION into its sub-editor.
+* This test opens each Material into its sub-editor.
 */
-IMPLEMENT_SIMPLE_AUTOMATION_TEST(FOpenMaterialFunctionAssetEditorsSLOW, "Tools.Materials.Open All Material Function Assets (SLOW)", (EAutomationTestFlags::ATF_Editor | EAutomationTestFlags::ATF_RequiresUser));
+IMPLEMENT_COMPLEX_AUTOMATION_TEST(FOpenMaterialAssetEditors, "Tools.Materials.Open Material Assets", (EAutomationTestFlags::ATF_Editor | EAutomationTestFlags::ATF_RequiresUser));
 
-bool FOpenMaterialFunctionAssetEditorsSLOW::RunTest(const FString& Parameters)
+void FOpenMaterialAssetEditors::GetTests(TArray<FString>& OutBeautifiedNames, TArray<FString>& OutTestCommands) const
 {
-	OpenAssetParameters AssetList;
-	AssetList.TimeSpent = FPlatformTime::Seconds();
-
-	FEditorAutomationTestUtilities::CollectGameContentByClass(UMaterialFunction::StaticClass(), true, AssetList.AssetPathList);
-
-	FOpenAllAssetTests(AssetList);
-
-	return true;
+	//This grabs each Material asset in the Game/Content directory
+	FEditorAutomationTestUtilities::CollectGameContentTestsByClass(UMaterial::StaticClass(), true, OutBeautifiedNames, OutTestCommands);
 }
 
-/**
-* This test opens each MATERIAL INSTANCE CONSTANT into its sub-editor.
-*/
-IMPLEMENT_SIMPLE_AUTOMATION_TEST(FOpenMICAssetEditorsSLOW, "Tools.Materials.Open All Material Instance Constant Assets (SLOW)", (EAutomationTestFlags::ATF_Editor | EAutomationTestFlags::ATF_RequiresUser));
-
-bool FOpenMICAssetEditorsSLOW::RunTest(const FString& Parameters)
+bool FOpenMaterialAssetEditors::RunTest(const FString& Parameters)
 {
-	OpenAssetParameters AssetList;
-	AssetList.TimeSpent = FPlatformTime::Seconds();
+	bool bDidTheTestPass = FOpenActualAssetEditors(Parameters);
+	CollectGarbage(GARBAGE_COLLECTION_KEEPFLAGS);
+	return bDidTheTestPass;
+}
 
-	FEditorAutomationTestUtilities::CollectGameContentByClass(UMaterialInstanceConstant::StaticClass(), true, AssetList.AssetPathList);
+////////////////////////////////////////////////////////////////////////////
+///**
+//* This test opens each MaterialFunction into its sub-editor.
+//*/
+IMPLEMENT_COMPLEX_AUTOMATION_TEST(FOpenMaterialFunctionAssetEditors, "Tools.Materials.Open Material Function Assets", (EAutomationTestFlags::ATF_Editor | EAutomationTestFlags::ATF_RequiresUser));
 
-	FOpenAllAssetTests(AssetList);
+void FOpenMaterialFunctionAssetEditors::GetTests(TArray<FString>& OutBeautifiedNames, TArray<FString>& OutTestCommands) const
+{
+	//This grabs each MaterialFunction asset in the Game/Content directory
+	FEditorAutomationTestUtilities::CollectGameContentTestsByClass(UMaterialFunction::StaticClass(), true, OutBeautifiedNames, OutTestCommands);
+}
 
-	return true;
+bool FOpenMaterialFunctionAssetEditors::RunTest(const FString& Parameters)
+{
+	bool bDidTheTestPass = FOpenActualAssetEditors(Parameters);
+	CollectGarbage(GARBAGE_COLLECTION_KEEPFLAGS);
+	return bDidTheTestPass;
 }
 
+//////////////////////////////////////////////////////////////////////////
 /**
-* This test opens each MATERIAL PARAMETER COLLECTION into its sub-editor.
+* This test opens each MaterialInstanceConstant into its sub-editor.
 */
-IMPLEMENT_SIMPLE_AUTOMATION_TEST(FOpenMaterialParameterCollectionAssetEditorsSLOW, "Tools.Materials.Open All Material Parameter Collection Assets (SLOW)", (EAutomationTestFlags::ATF_Editor | EAutomationTestFlags::ATF_RequiresUser));
+IMPLEMENT_COMPLEX_AUTOMATION_TEST(FOpenMaterialInstanceConstantAssetEditors, "Tools.Materials.Open Material Instance Constant Assets", (EAutomationTestFlags::ATF_Editor | EAutomationTestFlags::ATF_RequiresUser));
 
-bool FOpenMaterialParameterCollectionAssetEditorsSLOW::RunTest(const FString& Parameters)
+void FOpenMaterialInstanceConstantAssetEditors::GetTests(TArray<FString>& OutBeautifiedNames, TArray<FString>& OutTestCommands) const
 {
-	OpenAssetParameters AssetList;
-	AssetList.TimeSpent = FPlatformTime::Seconds();
-
-	FEditorAutomationTestUtilities::CollectGameContentByClass(UMaterialParameterCollection::StaticClass(), true, AssetList.AssetPathList);
-
-	FOpenAllAssetTests(AssetList);
+	//This grabs each MaterialInstanceConstant asset in the Game/Content directory
+	FEditorAutomationTestUtilities::CollectGameContentTestsByClass(UMaterialInstanceConstant::StaticClass(), true, OutBeautifiedNames, OutTestCommands);
+}
 
-	return true;
+bool FOpenMaterialInstanceConstantAssetEditors::RunTest(const FString& Parameters)
+{
+	bool bDidTheTestPass = FOpenActualAssetEditors(Parameters);
+	CollectGarbage(GARBAGE_COLLECTION_KEEPFLAGS);
+	return bDidTheTestPass;
 }
 
+//////////////////////////////////////////////////////////////////////////
 /**
-* This test opens each PARTICLE SYSTEM into its sub-editor.
+* This test opens each MaterialParameterCollection into its sub-editor.
 */
-IMPLEMENT_SIMPLE_AUTOMATION_TEST(FOpenParticleSystemAssetEditorsSLOW, "Tools.Particle System.Open All Particle System Assets (SLOW)", (EAutomationTestFlags::ATF_Editor | EAutomationTestFlags::ATF_RequiresUser));
+IMPLEMENT_COMPLEX_AUTOMATION_TEST(FOpenMaterialParameterCollectionAssetEditors, "Tools.Materials.Open Material Parameter Collection Assets", (EAutomationTestFlags::ATF_Editor | EAutomationTestFlags::ATF_RequiresUser));
 
-bool FOpenParticleSystemAssetEditorsSLOW::RunTest(const FString& Parameters)
+void FOpenMaterialParameterCollectionAssetEditors::GetTests(TArray<FString>& OutBeautifiedNames, TArray<FString>& OutTestCommands) const
 {
-	OpenAssetParameters AssetList;
-	AssetList.TimeSpent = FPlatformTime::Seconds();
-
-	FEditorAutomationTestUtilities::CollectGameContentByClass(UParticleSystem::StaticClass(), true, AssetList.AssetPathList);
-
-	FOpenAllAssetTests(AssetList);
+	//This grabs each MaterialParameterCollection asset in the Game/Content directory
+	FEditorAutomationTestUtilities::CollectGameContentTestsByClass(UMaterialParameterCollection::StaticClass(), true, OutBeautifiedNames, OutTestCommands);
+}
 
-	return true;
+bool FOpenMaterialParameterCollectionAssetEditors::RunTest(const FString& Parameters)
+{
+	bool bDidTheTestPass = FOpenActualAssetEditors(Parameters);
+	CollectGarbage(GARBAGE_COLLECTION_KEEPFLAGS);
+	return bDidTheTestPass;
 }
 
+//////////////////////////////////////////////////////////////////////////
 /**
-* This test opens each PHYSICAL MATERIAL into its sub-editor.
+* This test opens each PARTICLE SYSTEM into its sub-editor.
 */
-IMPLEMENT_SIMPLE_AUTOMATION_TEST(FOpenPhysicalMaterialAssetEditorsSLOW, "Tools.Physics.Open All Physical Material Assets (SLOW)", (EAutomationTestFlags::ATF_Editor | EAutomationTestFlags::ATF_RequiresUser));
+IMPLEMENT_COMPLEX_AUTOMATION_TEST(FOpenParticleSystemAssetEditors, "Tools.Particle System.Open Particle System Assets", (EAutomationTestFlags::ATF_Editor | EAutomationTestFlags::ATF_RequiresUser));
 
-bool FOpenPhysicalMaterialAssetEditorsSLOW::RunTest(const FString& Parameters)
+void FOpenParticleSystemAssetEditors::GetTests(TArray<FString>& OutBeautifiedNames, TArray<FString>& OutTestCommands) const
 {
-	OpenAssetParameters AssetList;
-	AssetList.TimeSpent = FPlatformTime::Seconds();
-
-	FEditorAutomationTestUtilities::CollectGameContentByClass(UPhysicalMaterial::StaticClass(), true, AssetList.AssetPathList);
-
-	FOpenAllAssetTests(AssetList);
+	//This grabs each ParticleSystem asset in the Game/Content directory
+	FEditorAutomationTestUtilities::CollectGameContentTestsByClass(UParticleSystem::StaticClass(), true, OutBeautifiedNames, OutTestCommands);
+}
 
-	return true;
+bool FOpenParticleSystemAssetEditors::RunTest(const FString& Parameters)
+{
+	bool bDidTheTestPass = FOpenActualAssetEditors(Parameters);
+	CollectGarbage(GARBAGE_COLLECTION_KEEPFLAGS);
+	return bDidTheTestPass;
 }
 
+//////////////////////////////////////////////////////////////////////////
 /**
-* This test opens each PHYSICS ASSET into its sub-editor.
+* This test opens each PhysicalMaterial into its sub-editor.
 */
-IMPLEMENT_SIMPLE_AUTOMATION_TEST(FOpenPhysicAssetEditorsSLOW, "Tools.Physics.Open All Physic Assets (SLOW)", (EAutomationTestFlags::ATF_Editor | EAutomationTestFlags::ATF_RequiresUser));
+IMPLEMENT_COMPLEX_AUTOMATION_TEST(FOpenPhysicalMaterialAssetEditors, "Tools.Physics.Open Physical Material Assets", (EAutomationTestFlags::ATF_Editor | EAutomationTestFlags::ATF_RequiresUser));
 
-bool FOpenPhysicAssetEditorsSLOW::RunTest(const FString& Parameters)
+void FOpenPhysicalMaterialAssetEditors::GetTests(TArray<FString>& OutBeautifiedNames, TArray<FString>& OutTestCommands) const
 {
-	OpenAssetParameters AssetList;
-	AssetList.TimeSpent = FPlatformTime::Seconds();
-
-	FEditorAutomationTestUtilities::CollectGameContentByClass(UPhysicsAsset::StaticClass(), true, AssetList.AssetPathList);
-
-	FOpenAllAssetTests(AssetList);
+	//This grabs each PhysicalMaterial asset in the Game/Content directory
+	FEditorAutomationTestUtilities::CollectGameContentTestsByClass(UPhysicalMaterial::StaticClass(), true, OutBeautifiedNames, OutTestCommands);
+}
 
-	return true;
+bool FOpenPhysicalMaterialAssetEditors::RunTest(const FString& Parameters)
+{
+	bool bDidTheTestPass = FOpenActualAssetEditors(Parameters);
+	CollectGarbage(GARBAGE_COLLECTION_KEEPFLAGS);
+	return bDidTheTestPass;
 }
 
+
+//////////////////////////////////////////////////////////////////////////
 /**
-* This test opens each REVERB EFFECT into its sub-editor.
+* This test opens each PhysicsAsset into its sub-editor.
 */
-IMPLEMENT_SIMPLE_AUTOMATION_TEST(FOpenReverbEffectAssetEditorsSLOW, "Tools.Sounds.Open All Reverb Effect Assets (SLOW)", (EAutomationTestFlags::ATF_Editor | EAutomationTestFlags::ATF_RequiresUser));
+IMPLEMENT_COMPLEX_AUTOMATION_TEST(FOpenPhysicsAssetAssetEditors, "Tools.Physics.Open Physics Assets", (EAutomationTestFlags::ATF_Editor | EAutomationTestFlags::ATF_RequiresUser));
 
-bool FOpenReverbEffectAssetEditorsSLOW::RunTest(const FString& Parameters)
+void FOpenPhysicsAssetAssetEditors::GetTests(TArray<FString>& OutBeautifiedNames, TArray<FString>& OutTestCommands) const
 {
-	OpenAssetParameters AssetList;
-	AssetList.TimeSpent = FPlatformTime::Seconds();
-
-	FEditorAutomationTestUtilities::CollectGameContentByClass(UReverbEffect::StaticClass(), true, AssetList.AssetPathList);
-
-	FOpenAllAssetTests(AssetList);
+	//This grabs each PhysicsAsset asset in the Game/Content directory
+	FEditorAutomationTestUtilities::CollectGameContentTestsByClass(UPhysicsAsset::StaticClass(), true, OutBeautifiedNames, OutTestCommands);
+}
 
-	return true;
+bool FOpenPhysicsAssetAssetEditors::RunTest(const FString& Parameters)
+{
+	bool bDidTheTestPass = FOpenActualAssetEditors(Parameters);
+	CollectGarbage(GARBAGE_COLLECTION_KEEPFLAGS);
+	return bDidTheTestPass;
 }
 
+
+//////////////////////////////////////////////////////////////////////////
 /**
-* This test opens each SLATE WIDGET STYLE ASSET into its sub-editor.
+* This test opens each ReverbEffect into its sub-editor.
 */
-IMPLEMENT_SIMPLE_AUTOMATION_TEST(FOpenSlateWidgetStyleAssetEditorsSLOW, "Tools.Slate.Open All Slate Widget Style Assets (SLOW)", (EAutomationTestFlags::ATF_Editor | EAutomationTestFlags::ATF_RequiresUser));
+IMPLEMENT_COMPLEX_AUTOMATION_TEST(FOpenReverbEffectAssetEditors, "Tools.Sounds.Open Reverb Effect Assets", (EAutomationTestFlags::ATF_Editor | EAutomationTestFlags::ATF_RequiresUser));
 
-bool FOpenSlateWidgetStyleAssetEditorsSLOW::RunTest(const FString& Parameters)
+void FOpenReverbEffectAssetEditors::GetTests(TArray<FString>& OutBeautifiedNames, TArray<FString>& OutTestCommands) const
 {
-	OpenAssetParameters AssetList;
-	AssetList.TimeSpent = FPlatformTime::Seconds();
+	//This grabs each ReverbEffect asset in the Game/Content directory
+	FEditorAutomationTestUtilities::CollectGameContentTestsByClass(UReverbEffect::StaticClass(), true, OutBeautifiedNames, OutTestCommands);
+}
 
-	FEditorAutomationTestUtilities::CollectGameContentByClass(USlateWidgetStyleAsset::StaticClass(), true, AssetList.AssetPathList);
+bool FOpenReverbEffectAssetEditors::RunTest(const FString& Parameters)
+{
+	bool bDidTheTestPass = FOpenActualAssetEditors(Parameters);
+	CollectGarbage(GARBAGE_COLLECTION_KEEPFLAGS);
+	return bDidTheTestPass;
+}
 
-	FOpenAllAssetTests(AssetList);
 
-	return true;
-}
 
+//////////////////////////////////////////////////////////////////////////
 /**
-* This test opens each SLATE BRUSH ASSET into its sub-editor.
+* This test opens each SlateWidgetStyleAsset into its sub-editor.
 */
-IMPLEMENT_SIMPLE_AUTOMATION_TEST(FOpenSlateBrushAssetEditorsSLOW, "Tools.Slate.Open All Slate Brush Assets (SLOW)", (EAutomationTestFlags::ATF_Editor | EAutomationTestFlags::ATF_RequiresUser));
+IMPLEMENT_COMPLEX_AUTOMATION_TEST(FOpenSlateWidgetStyleAssetAssetEditors, "Tools.Slate.Open Slate Widget Style Assets", (EAutomationTestFlags::ATF_Editor | EAutomationTestFlags::ATF_RequiresUser));
 
-bool FOpenSlateBrushAssetEditorsSLOW::RunTest(const FString& Parameters)
+void FOpenSlateWidgetStyleAssetAssetEditors::GetTests(TArray<FString>& OutBeautifiedNames, TArray<FString>& OutTestCommands) const
 {
-	OpenAssetParameters AssetList;
-	AssetList.TimeSpent = FPlatformTime::Seconds();
-
-	FEditorAutomationTestUtilities::CollectGameContentByClass(USlateBrushAsset::StaticClass(), true, AssetList.AssetPathList);
-
-	FOpenAllAssetTests(AssetList);
+	//This grabs each SlateWidgetStyleAsset asset in the Game/Content directory
+	FEditorAutomationTestUtilities::CollectGameContentTestsByClass(USlateWidgetStyleAsset::StaticClass(), true, OutBeautifiedNames, OutTestCommands);
+}
 
-	return true;
+bool FOpenSlateWidgetStyleAssetAssetEditors::RunTest(const FString& Parameters)
+{
+	bool bDidTheTestPass = FOpenActualAssetEditors(Parameters);
+	CollectGarbage(GARBAGE_COLLECTION_KEEPFLAGS);
+	return bDidTheTestPass;
 }
 
+//////////////////////////////////////////////////////////////////////////
 /**
-* This test opens each SOUNDATTENUATION into its sub-editor.
+* This test opens each Slate Brush Asset into its sub-editor.
 */
-IMPLEMENT_SIMPLE_AUTOMATION_TEST(FOpenSoundAttenuationAssetEditorsSLOW, "Tools.Sounds.Open All Sound Attenuation Assets (SLOW)", (EAutomationTestFlags::ATF_Editor | EAutomationTestFlags::ATF_RequiresUser));
+IMPLEMENT_COMPLEX_AUTOMATION_TEST(FOpenSlateBrushAssetAssetEditors, "Tools.Slate.Open Slate Brush Assets", (EAutomationTestFlags::ATF_Editor | EAutomationTestFlags::ATF_RequiresUser));
 
-bool FOpenSoundAttenuationAssetEditorsSLOW::RunTest(const FString& Parameters)
+void FOpenSlateBrushAssetAssetEditors::GetTests(TArray<FString>& OutBeautifiedNames, TArray<FString>& OutTestCommands) const
 {
-	OpenAssetParameters AssetList;
-	AssetList.TimeSpent = FPlatformTime::Seconds();
-
-	FEditorAutomationTestUtilities::CollectGameContentByClass(USoundAttenuation::StaticClass(), true, AssetList.AssetPathList);
-
-	FOpenAllAssetTests(AssetList);
+	//This grabs each SlateBrushAsset asset in the Game/Content directory
+	FEditorAutomationTestUtilities::CollectGameContentTestsByClass(USlateBrushAsset::StaticClass(), true, OutBeautifiedNames, OutTestCommands);
+}
 
-	return true;
+bool FOpenSlateBrushAssetAssetEditors::RunTest(const FString& Parameters)
+{
+	bool bDidTheTestPass = FOpenActualAssetEditors(Parameters);
+	CollectGarbage(GARBAGE_COLLECTION_KEEPFLAGS);
+	return bDidTheTestPass;
 }
 
+//////////////////////////////////////////////////////////////////////////
 /**
-* This test opens each STATIC MESH into its sub-editor.
+* This test opens each SoundAttenuation into its sub-editor.
 */
-IMPLEMENT_SIMPLE_AUTOMATION_TEST(FOpenStaticMeshAssetEditorsSLOW, "Tools.Static Mesh.Open All Static Mesh Assets (SLOW)", (EAutomationTestFlags::ATF_Editor | EAutomationTestFlags::ATF_RequiresUser));
+IMPLEMENT_COMPLEX_AUTOMATION_TEST(FOpenSoundAttenuationAssetEditors, "Tools.Sounds.Open Sound Attenuation Assets", (EAutomationTestFlags::ATF_Editor | EAutomationTestFlags::ATF_RequiresUser));
 
-bool FOpenStaticMeshAssetEditorsSLOW::RunTest(const FString& Parameters)
+void FOpenSoundAttenuationAssetEditors::GetTests(TArray<FString>& OutBeautifiedNames, TArray<FString>& OutTestCommands) const
 {
-	OpenAssetParameters AssetList;
-	AssetList.TimeSpent = FPlatformTime::Seconds();
-
-	FEditorAutomationTestUtilities::CollectGameContentByClass(UStaticMesh::StaticClass(), true, AssetList.AssetPathList);
-
-	FOpenAllAssetTests(AssetList);
+	//This grabs each SoundAttenuation asset in the Game/Content directory
+	FEditorAutomationTestUtilities::CollectGameContentTestsByClass(USoundAttenuation::StaticClass(), true, OutBeautifiedNames, OutTestCommands);
+}
 
-	return true;
+bool FOpenSoundAttenuationAssetEditors::RunTest(const FString& Parameters)
+{
+	bool bDidTheTestPass = FOpenActualAssetEditors(Parameters);
+	CollectGarbage(GARBAGE_COLLECTION_KEEPFLAGS);
+	return bDidTheTestPass;
 }
 
+//////////////////////////////////////////////////////////////////////////
 /**
-* This test opens each SKELETAL MESH into its sub-editor.
+* This test opens each Static Mesh into their sub-editor.
 */
-IMPLEMENT_SIMPLE_AUTOMATION_TEST(FOpenSkeletalMeshAssetEditorsSLOW, "Tools.Skeletal Mesh.Open All Skeletal Mesh Assets (SLOW)", (EAutomationTestFlags::ATF_Editor | EAutomationTestFlags::ATF_RequiresUser));
+IMPLEMENT_COMPLEX_AUTOMATION_TEST(FOpenStaticMeshAssetEditors, "Tools.Static Mesh.Open Static Mesh Assets", (EAutomationTestFlags::ATF_Editor | EAutomationTestFlags::ATF_RequiresUser));
 
-bool FOpenSkeletalMeshAssetEditorsSLOW::RunTest(const FString& Parameters)
+void FOpenStaticMeshAssetEditors::GetTests(TArray<FString>& OutBeautifiedNames, TArray<FString>& OutTestCommands) const
 {
-	OpenAssetParameters AssetList;
-	AssetList.TimeSpent = FPlatformTime::Seconds();
-
-	FEditorAutomationTestUtilities::CollectGameContentByClass(USkeletalMesh::StaticClass(), false, AssetList.AssetPathList);
-
-	FOpenAllAssetTests(AssetList);
+	//This grabs each Static Mesh in the Game/Content
+	FEditorAutomationTestUtilities::CollectGameContentTestsByClass(UStaticMesh::StaticClass(), true, OutBeautifiedNames, OutTestCommands);
+}
 
-	return true;
+bool FOpenStaticMeshAssetEditors::RunTest(const FString& Parameters)
+{
+	bool bDidTheTestPass = FOpenActualAssetEditors(Parameters);
+	CollectGarbage(GARBAGE_COLLECTION_KEEPFLAGS);
+	return bDidTheTestPass;
 }
 
+//////////////////////////////////////////////////////////////////////////
 /**
-* This test opens each SKELETON into its sub-editor.
+* This test opens each Skeletal Mesh into their sub-editor.
 */
-IMPLEMENT_SIMPLE_AUTOMATION_TEST(FOpenSkeletonAssetEditorsSLOW, "Tools.Skeletal Mesh.Open All Skeleton Assets (SLOW)", (EAutomationTestFlags::ATF_Editor | EAutomationTestFlags::ATF_RequiresUser));
+IMPLEMENT_COMPLEX_AUTOMATION_TEST(FOpenSkeletalMeshAssetEditors, "Tools.Skeletal Mesh.Open Skeletal Mesh Assets", (EAutomationTestFlags::ATF_Editor | EAutomationTestFlags::ATF_RequiresUser));
 
-bool FOpenSkeletonAssetEditorsSLOW::RunTest(const FString& Parameters)
+void FOpenSkeletalMeshAssetEditors::GetTests(TArray<FString>& OutBeautifiedNames, TArray<FString>& OutTestCommands) const
 {
-	OpenAssetParameters AssetList;
-	AssetList.TimeSpent = FPlatformTime::Seconds();
-
-	FEditorAutomationTestUtilities::CollectGameContentByClass(USkeleton::StaticClass(), true, AssetList.AssetPathList);
-
-	FOpenAllAssetTests(AssetList);
+	//This grabs each SKeletal Mesh in the Game/Content
+	FEditorAutomationTestUtilities::CollectGameContentTestsByClass(USkeletalMesh::StaticClass(), false, OutBeautifiedNames, OutTestCommands);
+}
 
-	return true;
+bool FOpenSkeletalMeshAssetEditors::RunTest(const FString& Parameters)
+{
+	bool bDidTheTestPass = FOpenActualAssetEditors(Parameters);
+	CollectGarbage(GARBAGE_COLLECTION_KEEPFLAGS);
+	return bDidTheTestPass;
 }
 
+//////////////////////////////////////////////////////////////////////////
 /**
-* This test opens each SOUND CLASS into its sub-editor.
+* This test opens each Skeleton into its sub-editor.
 */
-IMPLEMENT_SIMPLE_AUTOMATION_TEST(FOpenSoundClassAssetEditorsSLOW, "Tools.Sounds.Open All Sound Class Assets (SLOW)", (EAutomationTestFlags::ATF_Editor | EAutomationTestFlags::ATF_RequiresUser));
+IMPLEMENT_COMPLEX_AUTOMATION_TEST(FOpenSkeletonAssetEditors, "Tools.Skeletal Mesh.Open Skeleton Assets", (EAutomationTestFlags::ATF_Editor | EAutomationTestFlags::ATF_RequiresUser));
 
-bool FOpenSoundClassAssetEditorsSLOW::RunTest(const FString& Parameters)
+void FOpenSkeletonAssetEditors::GetTests(TArray<FString>& OutBeautifiedNames, TArray<FString>& OutTestCommands) const
 {
-	OpenAssetParameters AssetList;
-	AssetList.TimeSpent = FPlatformTime::Seconds();
-
-	FEditorAutomationTestUtilities::CollectGameContentByClass(USoundClass::StaticClass(), true, AssetList.AssetPathList);
-
-	FOpenAllAssetTests(AssetList);
+	//This grabs each Skeleton asset in the Game/Content directory
+	FEditorAutomationTestUtilities::CollectGameContentTestsByClass(USkeleton::StaticClass(), true, OutBeautifiedNames, OutTestCommands);
+}
 
-	return true;
+bool FOpenSkeletonAssetEditors::RunTest(const FString& Parameters)
+{
+	bool bDidTheTestPass = FOpenActualAssetEditors(Parameters);
+	CollectGarbage(GARBAGE_COLLECTION_KEEPFLAGS);
+	return bDidTheTestPass;
 }
 
+//////////////////////////////////////////////////////////////////////////
 /**
-* This test opens each SOUND CUE into its sub-editor.
+* This test opens each SoundClass into its sub-editor.
 */
-IMPLEMENT_SIMPLE_AUTOMATION_TEST(FOpenSoundCueAssetEditorsSLOW, "Tools.Sounds.Open All Sound Cue Assets (SLOW)", (EAutomationTestFlags::ATF_Editor | EAutomationTestFlags::ATF_RequiresUser));
+IMPLEMENT_COMPLEX_AUTOMATION_TEST(FOpenSoundClassAssetEditors, "Tools.Sounds.Open SoundClass Assets", (EAutomationTestFlags::ATF_Editor | EAutomationTestFlags::ATF_RequiresUser));
 
-bool FOpenSoundCueAssetEditorsSLOW::RunTest(const FString& Parameters)
+void FOpenSoundClassAssetEditors::GetTests(TArray<FString>& OutBeautifiedNames, TArray<FString>& OutTestCommands) const
 {
-	OpenAssetParameters AssetList;
-	AssetList.TimeSpent = FPlatformTime::Seconds();
-
-	FEditorAutomationTestUtilities::CollectGameContentByClass(USoundCue::StaticClass(), true, AssetList.AssetPathList);
-
-	FOpenAllAssetTests(AssetList);
+	//This grabs each SoundClass asset in the Game/Content directory
+	FEditorAutomationTestUtilities::CollectGameContentTestsByClass(USoundClass::StaticClass(), true, OutBeautifiedNames, OutTestCommands);
+}
 
-	return true;
+bool FOpenSoundClassAssetEditors::RunTest(const FString& Parameters)
+{
+	bool bDidTheTestPass = FOpenActualAssetEditors(Parameters);
+	CollectGarbage(GARBAGE_COLLECTION_KEEPFLAGS);
+	return bDidTheTestPass;
 }
 
+
+//////////////////////////////////////////////////////////////////////////
 /**
-* This test opens each SOUND MIX into its sub-editor.
+* This test opens each SoundCue into its sub-editor.
 */
-IMPLEMENT_SIMPLE_AUTOMATION_TEST(FOpenSoundMixAssetEditorsSLOW, "Tools.Sounds.Open All Sound Mix Assets (SLOW)", (EAutomationTestFlags::ATF_Editor | EAutomationTestFlags::ATF_RequiresUser));
+IMPLEMENT_COMPLEX_AUTOMATION_TEST(FOpenSoundCueAssetEditors, "Tools.Sounds.Open Sound Cue Assets", (EAutomationTestFlags::ATF_Editor | EAutomationTestFlags::ATF_RequiresUser));
 
-bool FOpenSoundMixAssetEditorsSLOW::RunTest(const FString& Parameters)
+void FOpenSoundCueAssetEditors::GetTests(TArray<FString>& OutBeautifiedNames, TArray<FString>& OutTestCommands) const
 {
-	OpenAssetParameters AssetList;
-	AssetList.TimeSpent = FPlatformTime::Seconds();
-
-	FEditorAutomationTestUtilities::CollectGameContentByClass(USoundMix::StaticClass(), true, AssetList.AssetPathList);
-
-	FOpenAllAssetTests(AssetList);
+	//This grabs each SoundCue asset in the Game/Content directory
+	FEditorAutomationTestUtilities::CollectGameContentTestsByClass(USoundCue::StaticClass(), true, OutBeautifiedNames, OutTestCommands);
+}
 
-	return true;
+bool FOpenSoundCueAssetEditors::RunTest(const FString& Parameters)
+{
+	bool bDidTheTestPass = FOpenActualAssetEditors(Parameters);
+	CollectGarbage(GARBAGE_COLLECTION_KEEPFLAGS);
+	return bDidTheTestPass;
 }
 
+
+//////////////////////////////////////////////////////////////////////////
 /**
-* This test opens each SOUND WAVE into its sub-editor.
+* This test opens each SoundMix into its sub-editor.
 */
-IMPLEMENT_SIMPLE_AUTOMATION_TEST(FOpenSoundWaveAssetEditorsSLOW, "Tools.Sounds.Open All Sound Wave Assets (SLOW)", (EAutomationTestFlags::ATF_Editor | EAutomationTestFlags::ATF_RequiresUser));
+IMPLEMENT_COMPLEX_AUTOMATION_TEST(FOpenSoundMixAssetEditors, "Tools.Sounds.Open Sound Mix Assets", (EAutomationTestFlags::ATF_Editor | EAutomationTestFlags::ATF_RequiresUser));
 
-bool FOpenSoundWaveAssetEditorsSLOW::RunTest(const FString& Parameters)
+void FOpenSoundMixAssetEditors::GetTests(TArray<FString>& OutBeautifiedNames, TArray<FString>& OutTestCommands) const
 {
-	OpenAssetParameters AssetList;
-	AssetList.TimeSpent = FPlatformTime::Seconds();
-
-	FEditorAutomationTestUtilities::CollectGameContentByClass(USoundWave::StaticClass(), true, AssetList.AssetPathList);
-
-	FOpenAllAssetTests(AssetList);
+	//This grabs each SoundMix asset in the Game/Content directory
+	FEditorAutomationTestUtilities::CollectGameContentTestsByClass(USoundMix::StaticClass(), true, OutBeautifiedNames, OutTestCommands);
+}
 
-	return true;
+bool FOpenSoundMixAssetEditors::RunTest(const FString& Parameters)
+{
+	bool bDidTheTestPass = FOpenActualAssetEditors(Parameters);
+	CollectGarbage(GARBAGE_COLLECTION_KEEPFLAGS);
+	return bDidTheTestPass;
 }
 
+
+//////////////////////////////////////////////////////////////////////////
 /**
-* This test opens each SUBSURFACEPROFILE into its sub-editor.
+* This test opens each SoundWave into its sub-editor.
 */
-IMPLEMENT_SIMPLE_AUTOMATION_TEST(FOpenSubsurfaceProfileAssetEditorsSLOW, "Tools.Materials.Open All Subsurface Profile Assets (SLOW)", (EAutomationTestFlags::ATF_Editor | EAutomationTestFlags::ATF_RequiresUser));
+IMPLEMENT_COMPLEX_AUTOMATION_TEST(FOpenSoundWaveAssetEditors, "Tools.Sounds.Open Sound Wave Assets", (EAutomationTestFlags::ATF_Editor | EAutomationTestFlags::ATF_RequiresUser));
 
-bool FOpenSubsurfaceProfileAssetEditorsSLOW::RunTest(const FString& Parameters)
+void FOpenSoundWaveAssetEditors::GetTests(TArray<FString>& OutBeautifiedNames, TArray<FString>& OutTestCommands) const
 {
-	OpenAssetParameters AssetList;
-	AssetList.TimeSpent = FPlatformTime::Seconds();
-
-	FEditorAutomationTestUtilities::CollectGameContentByClass(USubsurfaceProfile::StaticClass(), true, AssetList.AssetPathList);
-
-	FOpenAllAssetTests(AssetList);
+	//This grabs each SoundWave asset in the Game/Content directory
+	FEditorAutomationTestUtilities::CollectGameContentTestsByClass(USoundWave::StaticClass(), true, OutBeautifiedNames, OutTestCommands);
+}
 
-	return true;
+bool FOpenSoundWaveAssetEditors::RunTest(const FString& Parameters)
+{
+	bool bDidTheTestPass = FOpenActualAssetEditors(Parameters);
+	CollectGarbage(GARBAGE_COLLECTION_KEEPFLAGS);
+	return bDidTheTestPass;
 }
 
+//////////////////////////////////////////////////////////////////////////
 /**
-* This test opens each TEXTURE into its sub-editor.
+* This test opens each SubsurfaceProfile into its sub-editor.
 */
-IMPLEMENT_SIMPLE_AUTOMATION_TEST(FOpenTextureAssetEditorsSLOW, "Tools.Textures.Open All Texture Assets (SLOW)", (EAutomationTestFlags::ATF_Editor | EAutomationTestFlags::ATF_RequiresUser));
+IMPLEMENT_COMPLEX_AUTOMATION_TEST(FOpenSubsurfaceProfileAssetEditors, "Tools.Materials.Open Subsurface Profile Assets", (EAutomationTestFlags::ATF_Editor | EAutomationTestFlags::ATF_RequiresUser));
 
-bool FOpenTextureAssetEditorsSLOW::RunTest(const FString& Parameters)
+void FOpenSubsurfaceProfileAssetEditors::GetTests(TArray<FString>& OutBeautifiedNames, TArray<FString>& OutTestCommands) const
 {
-	OpenAssetParameters AssetList;
-	AssetList.TimeSpent = FPlatformTime::Seconds();
+	//This grabs each SubsurfaceProfile asset in the Game/Content directory
+	FEditorAutomationTestUtilities::CollectGameContentTestsByClass(USubsurfaceProfile::StaticClass(), true, OutBeautifiedNames, OutTestCommands);
+}
 
-	FEditorAutomationTestUtilities::CollectGameContentByClass(UTexture::StaticClass(), true, AssetList.AssetPathList);
+bool FOpenSubsurfaceProfileAssetEditors::RunTest(const FString& Parameters)
+{
+	bool bDidTheTestPass = FOpenActualAssetEditors(Parameters);
+	CollectGarbage(GARBAGE_COLLECTION_KEEPFLAGS);
+	return bDidTheTestPass;
+}
 
-	FOpenAllAssetTests(AssetList);
 
-	return true;
-}
 
+//////////////////////////////////////////////////////////////////////////
 /**
-* This test opens each USER DEFINED ENUM into its sub-editor.
+* This test opens each TEXTURE asset into its sub-editor.
 */
-IMPLEMENT_SIMPLE_AUTOMATION_TEST(FOpenUserDefinedEnumAssetEditorsSLOW, "Tools.Blueprint.Open All User Defined Enum Assets (SLOW)", (EAutomationTestFlags::ATF_Editor | EAutomationTestFlags::ATF_RequiresUser));
+IMPLEMENT_COMPLEX_AUTOMATION_TEST(FOpenTextureAssetEditors, "Tools.Textures.Open Texture Assets", (EAutomationTestFlags::ATF_Editor | EAutomationTestFlags::ATF_RequiresUser));
 
-bool FOpenUserDefinedEnumAssetEditorsSLOW::RunTest(const FString& Parameters)
+void FOpenTextureAssetEditors::GetTests(TArray<FString>& OutBeautifiedNames, TArray<FString>& OutTestCommands) const
 {
-	OpenAssetParameters AssetList;
-	AssetList.TimeSpent = FPlatformTime::Seconds();
-
-	FEditorAutomationTestUtilities::CollectGameContentByClass(UUserDefinedEnum::StaticClass(), true, AssetList.AssetPathList);
-
-	FOpenAllAssetTests(AssetList);
+	//This grabs each texture file in the Game/Content
+	FEditorAutomationTestUtilities::CollectGameContentTestsByClass(UTexture::StaticClass(), true, OutBeautifiedNames, OutTestCommands);
+}
 
-	return true;
+bool FOpenTextureAssetEditors::RunTest(const FString& Parameters)
+{
+	bool bDidTheTestPass = FOpenActualAssetEditors(Parameters);
+	CollectGarbage(GARBAGE_COLLECTION_KEEPFLAGS);
+	return bDidTheTestPass;
 }
 
+
+//////////////////////////////////////////////////////////////////////////
 /**
-* This test opens each FONT into its sub-editor.
+* This test opens each Font into its sub-editor.
 */
-IMPLEMENT_SIMPLE_AUTOMATION_TEST(FOpenFontAssetEditorsSLOW, "Tools.Textures.Open All Font Assets (SLOW)", (EAutomationTestFlags::ATF_Editor | EAutomationTestFlags::ATF_RequiresUser));
+IMPLEMENT_COMPLEX_AUTOMATION_TEST(FOpenFontAssetEditors, "Tools.Textures.Open Font Assets", (EAutomationTestFlags::ATF_Editor | EAutomationTestFlags::ATF_RequiresUser));
 
-bool FOpenFontAssetEditorsSLOW::RunTest(const FString& Parameters)
+void FOpenFontAssetEditors::GetTests(TArray<FString>& OutBeautifiedNames, TArray<FString>& OutTestCommands) const
 {
-	OpenAssetParameters AssetList;
-	AssetList.TimeSpent = FPlatformTime::Seconds();
-
-	FEditorAutomationTestUtilities::CollectGameContentByClass(UFont::StaticClass(), true, AssetList.AssetPathList);
-
-	FOpenAllAssetTests(AssetList);
+	//This grabs each Font asset in the Game/Content directory
+	FEditorAutomationTestUtilities::CollectGameContentTestsByClass(UFont::StaticClass(), true, OutBeautifiedNames, OutTestCommands);
+}
 
-	return true;
+bool FOpenFontAssetEditors::RunTest(const FString& Parameters)
+{
+	bool bDidTheTestPass = FOpenActualAssetEditors(Parameters);
+	CollectGarbage(GARBAGE_COLLECTION_KEEPFLAGS);
+	return bDidTheTestPass;
 }
 
+
 ////////////////////////////////////////////////////////////////////////////
 ///**
 //* This test opens each misc asset into its sub-editor.
 //*/
-//IMPLEMENT_COMPLEX_AUTOMATION_TEST(FOpenMiscAssetEditors, "Tools.Misc.Open Misc Assets", (EAutomationTestFlags::ATF_Editor | EAutomationTestFlags::ATF_RequiresUser));
-//
-//void FOpenMiscAssetEditors::GetTests(TArray<FString>& OutBeautifiedNames, TArray<FString>& OutTestCommands) const
-//{
-//	//This grabs each TextureRenderTarget2D asset in the Game/Content directory
-//	FEditorAutomationTestUtilities::CollectMiscGameContentTestsByClass(OutBeautifiedNames, OutTestCommands);
-//}
-//
-//bool FOpenMiscAssetEditors::RunTest(const FString& Parameters)
-//{
-//	bool bDidTheTestPass = FOpenActualAssetEditors(Parameters);
-//	CollectGarbage(GARBAGE_COLLECTION_KEEPFLAGS);
-//	return bDidTheTestPass;
-//}
-//
-//
-//
+IMPLEMENT_COMPLEX_AUTOMATION_TEST(FOpenMiscAssetEditors, "Tools.Misc.Open Misc Assets", (EAutomationTestFlags::ATF_Editor | EAutomationTestFlags::ATF_RequiresUser));
+
+void FOpenMiscAssetEditors::GetTests(TArray<FString>& OutBeautifiedNames, TArray<FString>& OutTestCommands) const
+{
+	//This grabs each TextureRenderTarget2D asset in the Game/Content directory
+	FEditorAutomationTestUtilities::CollectMiscGameContentTestsByClass(OutBeautifiedNames, OutTestCommands);
+}
+
+bool FOpenMiscAssetEditors::RunTest(const FString& Parameters)
+{
+	bool bDidTheTestPass = FOpenActualAssetEditors(Parameters);
+	CollectGarbage(GARBAGE_COLLECTION_KEEPFLAGS);
+	return bDidTheTestPass;
+}
+
+
+
