Static Mesh Editor - Improving the Socket Manager's look and feel and enabling it by default.

diff --git a/Engine/Source/Editor/StaticMeshEditor/Private/StaticMeshEditor.cpp b/Engine/Source/Editor/StaticMeshEditor/Private/StaticMeshEditor.cpp
--- a/Engine/Source/Editor/StaticMeshEditor/Private/StaticMeshEditor.cpp
+++ b/Engine/Source/Editor/StaticMeshEditor/Private/StaticMeshEditor.cpp
@@ -136,7 +136,7 @@ void FStaticMeshEditor::InitStaticMeshEditor( const EToolkitMode::Type Mode, con
 
 	BuildSubTools();
 
-	const TSharedRef<FTabManager::FLayout> StandaloneDefaultLayout = FTabManager::NewLayout( "Standalone_StaticMeshEditor_Layout_v3" )
+	const TSharedRef<FTabManager::FLayout> StandaloneDefaultLayout = FTabManager::NewLayout( "Standalone_StaticMeshEditor_Layout_v4" )
 	->AddArea
 	(
 		FTabManager::NewPrimaryArea() ->SetOrientation(Orient_Vertical)
@@ -149,27 +149,30 @@ void FStaticMeshEditor::InitStaticMeshEditor( const EToolkitMode::Type Mode, con
 		)
 		->Split
 		(
-			FTabManager::NewSplitter() ->SetOrientation(Orient_Horizontal) ->SetSizeCoefficient(0.9f)
+			FTabManager::NewSplitter() ->SetOrientation(Orient_Horizontal)
+			->SetSizeCoefficient(0.9f)
 			->Split
 			(
 				FTabManager::NewStack()
 				->SetSizeCoefficient(0.6f)
-				->AddTab(ViewportTabId, ETabState::OpenedTab) ->SetHideTabWell( true )
+				->AddTab(ViewportTabId, ETabState::OpenedTab)
+				->SetHideTabWell( true )
 			)
 			->Split
 			(
-				FTabManager::NewSplitter() ->SetOrientation(Orient_Vertical) ->SetSizeCoefficient(0.2f)
+				FTabManager::NewSplitter() ->SetOrientation(Orient_Vertical)
+				->SetSizeCoefficient(0.2f)
 				->Split
 				(
 					FTabManager::NewStack()
-					->SetSizeCoefficient(0.5f)
+					->SetSizeCoefficient(0.7f)
 					->AddTab(PropertiesTabId, ETabState::OpenedTab)
 				)
 				->Split
 				(
 					FTabManager::NewStack()
-					->SetSizeCoefficient(0.5f)
-					->AddTab(SocketManagerTabId, ETabState::ClosedTab)
+					->SetSizeCoefficient(0.3f)
+					->AddTab(SocketManagerTabId, ETabState::OpenedTab)
 					->AddTab(CollisionTabId, ETabState::ClosedTab)
 				)
 			)
diff --git a/Engine/Source/Editor/UnrealEd/Private/SSocketManager.cpp b/Engine/Source/Editor/UnrealEd/Private/SSocketManager.cpp
--- a/Engine/Source/Editor/UnrealEd/Private/SSocketManager.cpp
+++ b/Engine/Source/Editor/UnrealEd/Private/SSocketManager.cpp
@@ -196,160 +196,104 @@ void SSocketManager::Construct(const FArguments& InArgs)
 	WorldSpaceRotation = FVector::ZeroVector;
 
 	this->ChildSlot
+	[
+		SNew(SVerticalBox)
+
+		+ SVerticalBox::Slot()
+		.FillHeight(1.0f)
 		[
-			SNew(SHorizontalBox)
-			+SHorizontalBox::Slot()
-			.MaxWidth(220.0f)
-			.Padding(4.0f, 0.0f, 0.0f, 0.0f)
+			SNew(SSplitter)
+			.Orientation(Orient_Horizontal)
+
+			+ SSplitter::Slot()
+			.Value(.3f)
 			[
-				SNew(SVerticalBox)
-				+SVerticalBox::Slot()
-				.AutoHeight()
-				.Padding(4.0f, 0.0f, 0.0f, 4.0f)
+				SNew(SBorder)
+				.BorderImage(FEditorStyle::GetBrush("ToolPanel.GroupBorder"))
 				[
-					SNew(STextBlock)
-						.Text( this, &SSocketManager::GetSocketHeaderText )
-				]
+					SNew(SVerticalBox)
 
-				+SVerticalBox::Slot()
+					+ SVerticalBox::Slot()
+					.AutoHeight()
+					.Padding(0, 0, 0, 4)
+					[
+						SNew(SButton)
+						.ButtonStyle(FEditorStyle::Get(), "FlatButton.Success")
+						.ForegroundColor(FLinearColor::White)
+						.Text(LOCTEXT("CreateSocket", "Create Socket"))
+						.OnClicked(this, &SSocketManager::CreateSocket_Execute)
+						.HAlign(HAlign_Center)
+					]
+
+					+ SVerticalBox::Slot()
 					.FillHeight(1.0f)
 					[
-						SNew(SBorder)
-						[
-							SAssignNew(SocketListView, SListView<TSharedPtr< SocketListItem > >)
+						SAssignNew(SocketListView, SListView<TSharedPtr< SocketListItem > >)
 
-							.SelectionMode(ESelectionMode::Single)
+						.SelectionMode(ESelectionMode::Single)
 
-							.ListItemsSource( &SocketList )
+						.ListItemsSource(&SocketList)
 
-							// Generates the actual widget for a tree item
-							.OnGenerateRow( this, &SSocketManager::MakeWidgetFromOption ) 
+						// Generates the actual widget for a tree item
+						.OnGenerateRow(this, &SSocketManager::MakeWidgetFromOption)
 
-							// Find out when the user selects something in the tree
-							.OnSelectionChanged( this, &SSocketManager::SocketSelectionChanged_Execute )
+						// Find out when the user selects something in the tree
+						.OnSelectionChanged(this, &SSocketManager::SocketSelectionChanged_Execute)
 
-							// Allow for some spacing between items with a larger item height.
-							.ItemHeight(20.0f)
+						// Allow for some spacing between items with a larger item height.
+						.ItemHeight(20.0f)
 
-							.OnContextMenuOpening( this, &SSocketManager::OnContextMenuOpening )
-							.OnItemScrolledIntoView( this, &SSocketManager::OnItemScrolledIntoView )
+						.OnContextMenuOpening(this, &SSocketManager::OnContextMenuOpening)
+						.OnItemScrolledIntoView(this, &SSocketManager::OnItemScrolledIntoView)
 
-							.HeaderRow
-							(
-								SNew(SHeaderRow)
-								.Visibility(EVisibility::Collapsed)
-								+SHeaderRow::Column(TEXT("Socket"))
-							)
-						]
+						.HeaderRow
+						(
+						SNew(SHeaderRow)
+						.Visibility(EVisibility::Collapsed)
+						+ SHeaderRow::Column(TEXT("Socket"))
+						)
 					]
 
-				+SVerticalBox::Slot()
+					+ SVerticalBox::Slot()
 					.AutoHeight()
-					.Padding(0.0f, 8.0f, 0.0f, 0.0f)
 					[
-						SNew(SHorizontalBox)
-						+SHorizontalBox::Slot()
-						.HAlign(HAlign_Right)
-						.FillWidth(1.0f)
-						.Padding(5.0f, 0.0f, 4.0f, 0.0f)
-						[
-							SNew(SButton)
-							.Text(LOCTEXT("CreateSocket", "Create Socket"))
-							.OnClicked(this, &SSocketManager::CreateSocket_Execute)
-						]
-						+SHorizontalBox::Slot()
-							.HAlign(HAlign_Left)
-							.FillWidth(1.0f)
-							.Padding(4.0f, 0.0f, 5.0f, 0.0f)
-							[
-								SNew(SButton)
-								.Text(LOCTEXT("DeleteSocket", "Delete Socket"))
-								.OnClicked(this, &SSocketManager::DeleteSelectedSocket_Execute)
-							]
+						SNew(SSeparator)
 					]
 
-				+SVerticalBox::Slot()
+					+ SVerticalBox::Slot()
 					.AutoHeight()
-					.HAlign(HAlign_Center)
-					.Padding(0.0f, 12.0f, 0.0f, 4.0f)
 					[
 						SNew(STextBlock)
-						.Text(LOCTEXT("WorldSpaceRotation", "World Space Rotation"))
+						.Text(this, &SSocketManager::GetSocketHeaderText)
 					]
+				]
+			]
 
-				+SVerticalBox::Slot()
-					.AutoHeight()
-					.Padding(10.0f, 4.0f, 10.0f, 4.0f)
-					[
-						SNew(SHorizontalBox)
-						+SHorizontalBox::Slot()
-						.FillWidth(0.25f)
-						[
-							SNew(STextBlock)
-							.Text(LOCTEXT("Pitch", "Pitch"))
-						]
-						+SHorizontalBox::Slot()
-							.FillWidth(1.0f)
-							[
-								SAssignNew(PitchRotation, SSpinBox<float>)
-								.MaxValue(360.f)
-								.MinValue(0.f)
-								.OnValueChanged(this, &SSocketManager::PitchRotation_ValueChanged)
-								.Value(this, &SSocketManager::GetWorldSpacePitchValue)
-							]
-					]
+			+ SSplitter::Slot()
+			.Value(.7f)
+			[
+				SNew(SOverlay)
 
-				+SVerticalBox::Slot()
-					.AutoHeight()
-					.Padding(10.0f, 4.0f, 10.0f, 4.0f)
+				+ SOverlay::Slot()
+				[
+					SNew(SBorder)
+					.BorderImage(FEditorStyle::GetBrush("ToolPanel.GroupBorder"))
+					.HAlign(HAlign_Center)
+					.VAlign(VAlign_Center)
+					.Visibility(this, &SSocketManager::GetSelectSocketMessageVisibility)
 					[
-						SNew(SHorizontalBox)
-						+SHorizontalBox::Slot()
-						.FillWidth(0.25f)
-						[
-							SNew(STextBlock)
-							.Text(LOCTEXT("Yaw", "Yaw"))
-						]
-						+SHorizontalBox::Slot()
-							.FillWidth(1.0f)
-							[
-								SAssignNew(YawRotation, SSpinBox<float>)
-								.MaxValue(360.f)
-								.MinValue(0.f)
-								.OnValueChanged(this, &SSocketManager::YawRotation_ValueChanged)
-								.Value(this, &SSocketManager::GetWorldSpaceYawValue)
-							]
+						SNew(STextBlock)
+						.Text(LOCTEXT("NoSocketSelected", "Select a Socket"))
 					]
+				]
 
-				+SVerticalBox::Slot()
-					.AutoHeight()
-					.Padding(10.0f, 4.0f, 10.0f, 4.0f)
-					[
-						SNew(SHorizontalBox)
-						+SHorizontalBox::Slot()
-						.FillWidth(0.25f)
-						[
-							SNew(STextBlock)
-								.Text(LOCTEXT("Roll", "Roll"))
-						]
-						+SHorizontalBox::Slot()
-							.FillWidth(1.0f)
-							[
-								SAssignNew(RollRotation, SSpinBox<float>)
-								.MaxValue(360.f)
-								.MinValue(0.f)
-								.OnValueChanged(this, &SSocketManager::RollRotation_ValueChanged)
-								.Value(this, &SSocketManager::GetWorldSpaceRollValue)
-							]
-					]
-			]
-			+SHorizontalBox::Slot()
-				.FillWidth(1.0f)
-				.Padding(4.0f, 0.0f, 0.0f, 0.0f)
+				+ SOverlay::Slot()
 				[
 					SocketDetailsView.ToSharedRef()
 				]
-		];
+			]
+		]
+	];
 
 	RefreshSocketList();
 
@@ -374,7 +318,12 @@ UStaticMeshSocket* SSocketManager::GetSelectedSocket() const
 		return SocketListView->GetSelectedItems()[0]->Socket;
 	}
 
-	return NULL;
+	return nullptr;
+}
+
+EVisibility SSocketManager::GetSelectSocketMessageVisibility() const
+{
+	return SocketListView->GetSelectedItems().Num() > 0 ? EVisibility::Hidden : EVisibility::Visible;
 }
 
 void SSocketManager::SetSelectedSocket(UStaticMeshSocket* InSelectedSocket)
@@ -624,12 +573,6 @@ FReply SSocketManager::CreateSocket_Execute()
 	return FReply::Handled();
 }
 
-FReply SSocketManager::DeleteSelectedSocket_Execute()
-{
-	DeleteSelectedSocket();
-	return FReply::Handled();
-}
-
 FText SSocketManager::GetSocketHeaderText() const
 {
 	UStaticMesh* CurrentStaticMesh = nullptr;
@@ -638,7 +581,7 @@ FText SSocketManager::GetSocketHeaderText() const
 	{
 		CurrentStaticMesh = StaticMeshEditorPinned->GetStaticMesh();
 	}
-	return FText::Format(LOCTEXT("SocketHeader_TotalFmt", "Sockets ({0} Total)"), FText::AsNumber((CurrentStaticMesh != nullptr) ? CurrentStaticMesh->Sockets.Num() : 0));
+	return FText::Format(LOCTEXT("SocketHeader_TotalFmt", "{0} sockets"), FText::AsNumber((CurrentStaticMesh != nullptr) ? CurrentStaticMesh->Sockets.Num() : 0));
 }
 
 void SSocketManager::SocketName_TextChanged(const FText& InText)
@@ -646,66 +589,6 @@ void SSocketManager::SocketName_TextChanged(const FText& InText)
 	CheckForDuplicateSocket(InText.ToString());
 }
 
-void SSocketManager::PitchRotation_ValueChanged(float InValue)
-{
-	if( InValue != WorldSpaceRotation.X )
-	{
-		WorldSpaceRotation.X = InValue;
-
-		RotateSocket_WorldSpace();
-	}
-}
-
-void SSocketManager::YawRotation_ValueChanged(float InValue)
-{
-	if( InValue != WorldSpaceRotation.Y )
-	{
-		WorldSpaceRotation.Y = InValue;
-
-		RotateSocket_WorldSpace();
-	}
-}
-
-void SSocketManager::RollRotation_ValueChanged(float InValue)
-{
-	if( InValue != WorldSpaceRotation.Z )
-	{
-		WorldSpaceRotation.Z = InValue;
-
-		RotateSocket_WorldSpace();
-	}
-}
-
-float SSocketManager::GetWorldSpacePitchValue() const
-{
-	return WorldSpaceRotation.X;
-}
-
-float SSocketManager::GetWorldSpaceYawValue() const
-{
-	return WorldSpaceRotation.Y;
-}
-
-float SSocketManager::GetWorldSpaceRollValue() const
-{
-	return WorldSpaceRotation.Z;
-}
-
-void SSocketManager::RotateSocket_WorldSpace()
-{
-	TSharedPtr<IStaticMeshEditor> StaticMeshEditorPinned = StaticMeshEditorPtr.Pin();
-	if (StaticMeshEditorPinned.IsValid() && SocketListView->GetSelectedItems().Num())
-	{
-		UProperty* ChangedProperty = FindField<UProperty>( UStaticMeshSocket::StaticClass(), "RelativeRotation" );
-		UStaticMeshSocket* SelectedSocket = SocketListView->GetSelectedItems()[0]->Socket;
-		SelectedSocket->PreEditChange(ChangedProperty);
-		SelectedSocket->RelativeRotation = FRotator( WorldSpaceRotation.X, WorldSpaceRotation.Y, WorldSpaceRotation.Z );
-
-		FPropertyChangedEvent PropertyChangedEvent( ChangedProperty );
-		SelectedSocket->PostEditChangeProperty(PropertyChangedEvent);
-	}
-}
-
 TSharedPtr<SWidget> SSocketManager::OnContextMenuOpening()
 {
 	const bool bShouldCloseWindowAfterMenuSelection = true;
diff --git a/Engine/Source/Editor/UnrealEd/Private/SSocketManager.h b/Engine/Source/Editor/UnrealEd/Private/SSocketManager.h
--- a/Engine/Source/Editor/UnrealEd/Private/SSocketManager.h
+++ b/Engine/Source/Editor/UnrealEd/Private/SSocketManager.h
@@ -45,6 +45,9 @@ private:
 	/** Refreshes the socket list. */
 	void RefreshSocketList();
 
+	/** Gets the visibility of the select a socket message */
+	EVisibility GetSelectSocketMessageVisibility() const;
+
 	/** 
 	 *	Updates the details to the selected socket.
 	 *
@@ -58,35 +61,11 @@ private:
 	/** Callback for the Create Socket button. */
 	FReply CreateSocket_Execute();
 
-	/** Callback for the Delete Socket button. */
-	FReply DeleteSelectedSocket_Execute();
-
 	FText GetSocketHeaderText() const;
 
 	/** Callback for when the socket name textbox is changed, verifies the name is not a duplicate. */
 	void SocketName_TextChanged(const FText& InText);
 
-	/** Callback for the world space rotation value for Pitch being changed. */
-	void PitchRotation_ValueChanged(float InValue);
-
-	/** Callback for the world space rotation value for Yaw being changed. */
-	void YawRotation_ValueChanged(float InValue);
-
-	/** Callback for the world space rotation value for Roll being changed. */
-	void RollRotation_ValueChanged(float InValue);
-
-	/** Callback to get the world space rotation pitch value. */
-	float GetWorldSpacePitchValue() const;
-
-	/** Callback to get the world space rotation yaw value. */
-	float GetWorldSpaceYawValue() const;
-
-	/** Callback to get the world space rotation roll value. */
-	float GetWorldSpaceRollValue() const;
-
-	/** Handles rotating the socket in world space. */
-	void RotateSocket_WorldSpace();
-
 	/** Callback to retrieve the context menu for the list view */
 	TSharedPtr<SWidget> OnContextMenuOpening();
 
