Github pull request 786 Add sub-level editor delegate

diff --git a/Engine/Source/Editor/UnrealEd/Private/Editor.cpp b/Engine/Source/Editor/UnrealEd/Private/Editor.cpp
--- a/Engine/Source/Editor/UnrealEd/Private/Editor.cpp
+++ b/Engine/Source/Editor/UnrealEd/Private/Editor.cpp
@@ -149,6 +149,7 @@ FSimpleMulticastDelegate								FEditorDelegates::OnShutdownPostPackagesSaved;
 FEditorDelegates::FOnAssetsPreDelete					FEditorDelegates::OnAssetsPreDelete;
 FEditorDelegates::FOnAssetsDeleted						FEditorDelegates::OnAssetsDeleted;
 FSimpleMulticastDelegate								FEditorDelegates::OnActionAxisMappingsChanged;
+FEditorDelegates::FOnAddLevelToWorld					FEditorDelegates::OnAddLevelToWorld;
 
 /*-----------------------------------------------------------------------------
 	Globals.
diff --git a/Engine/Source/Editor/UnrealEd/Private/EditorLevelUtils.cpp b/Engine/Source/Editor/UnrealEd/Private/EditorLevelUtils.cpp
--- a/Engine/Source/Editor/UnrealEd/Private/EditorLevelUtils.cpp
+++ b/Engine/Source/Editor/UnrealEd/Private/EditorLevelUtils.cpp
@@ -243,6 +243,11 @@ namespace EditorLevelUtils
 			}
 		}
 
+		if(NewLevel) // if the level was successfully added
+		{
+			FEditorDelegates::OnAddLevelToWorld.Broadcast(NewLevel);
+		}
+
 		return NewLevel;
 	}
 
diff --git a/Engine/Source/Editor/UnrealEd/Public/Editor.h b/Engine/Source/Editor/UnrealEd/Public/Editor.h
--- a/Engine/Source/Editor/UnrealEd/Public/Editor.h
+++ b/Engine/Source/Editor/UnrealEd/Public/Editor.h
@@ -94,6 +94,8 @@ struct UNREALED_API FEditorDelegates
 	DECLARE_MULTICAST_DELEGATE_OneParam(FOnAssetsPreDelete, const TArray<UObject*>&);
 	/** delegate type for when one or more assets have been deleted */
 	DECLARE_MULTICAST_DELEGATE_OneParam(FOnAssetsDeleted, const TArray<UClass*>& /*DeletedAssetClasses*/);
+	/** delegate type for when a new level is added to the world */
+	DECLARE_MULTICAST_DELEGATE_OneParam(FOnAddLevelToWorld, ULevel*);
 
 	/** Called when the CurrentLevel is switched to a new level.  Note that this event won't be fired for temporary
 		changes to the current level, such as when copying/pasting actors. */
@@ -189,6 +191,8 @@ struct UNREALED_API FEditorDelegates
 	static FOnAssetsDeleted OnAssetsDeleted;
 	/** Called when Action or Axis mappings have been changed */
 	static FSimpleMulticastDelegate OnActionAxisMappingsChanged;
+	/** Called from FEditorUtils::AddLevelToWorld after the level is added successfully to the world. */
+	static FOnAddLevelToWorld OnAddLevelToWorld;
 };
 
 /**
