UE-9764 Setting a bsp box to tessellated can result in a crash when interacting with the viewport

diff --git a/Engine/Source/Editor/UnrealEd/Private/EditorServer.cpp b/Engine/Source/Editor/UnrealEd/Private/EditorServer.cpp
--- a/Engine/Source/Editor/UnrealEd/Private/EditorServer.cpp
+++ b/Engine/Source/Editor/UnrealEd/Private/EditorServer.cpp
@@ -1547,6 +1547,8 @@ void UEditorEngine::RebuildLevel(ULevel& Level)
 	{
 		IFA->MapRebuild();
 	}
+
+	GLevelEditorModeTools().MapChangeNotify();
 }
 
 void UEditorEngine::RebuildModelFromBrushes(UModel* Model, bool bSelectedBrushesOnly)
