Level editor viewport widgets are now hidden when an actor has no scene components.

UE-10129 - Improve UI for Actors without Components

diff --git a/Engine/Source/Editor/UnrealEd/Private/EdMode.cpp b/Engine/Source/Editor/UnrealEd/Private/EdMode.cpp
--- a/Engine/Source/Editor/UnrealEd/Private/EdMode.cpp
+++ b/Engine/Source/Editor/UnrealEd/Private/EdMode.cpp
@@ -460,7 +460,7 @@ bool FEdMode::GetCustomDrawingCoordinateSystem(FMatrix& InMatrix, void* InData)
 bool FEdMode::ShouldDrawWidget() const
 {
 	bool bDrawWidget = false;
-	if (GEditor->GetSelectedComponentCount() > 0)
+	if (GEditor->GetSelectedComponentCount() > 0) // when components are selected, only show the widget when one or more are scene components
 	{
 		for (FSelectedEditableComponentIterator It(GEditor->GetSelectedEditableComponentIterator()); It; ++It)
 		{
@@ -471,9 +471,19 @@ bool FEdMode::ShouldDrawWidget() const
 			}
 		}
 	}
-	else
+	else // when actors are selected, only show the widget when all selected actors have scene components
 	{
-		bDrawWidget = (GEditor->GetSelectedActors()->GetTop<AActor>() != NULL);
+		bDrawWidget = GEditor->GetSelectedActorCount() > 0;
+
+		for (FSelectionIterator It(GEditor->GetSelectedActorIterator()); It; ++It)
+		{
+			AActor* Actor = Cast<AActor>(*It);
+			if (Actor == nullptr || Actor->FindComponentByClass<USceneComponent>() == nullptr)
+			{
+				bDrawWidget = false;
+				break;
+			}
+		}
 	}
 
 	return bDrawWidget;
