Making it so that the static mesh uv test doesn't work on engine content.

[REDACTED_HANDLE] make a seperate test for the static mesh uv test.

diff --git a/Engine/Source/Editor/UnrealEd/Private/Tests/EditorAutomationTests.cpp b/Engine/Source/Editor/UnrealEd/Private/Tests/EditorAutomationTests.cpp
--- a/Engine/Source/Editor/UnrealEd/Private/Tests/EditorAutomationTests.cpp
+++ b/Engine/Source/Editor/UnrealEd/Private/Tests/EditorAutomationTests.cpp
@@ -1853,7 +1853,7 @@ namespace StaticMeshUVTest
 			LoadedPackageCount(0),
 			CurrentTest(BaseTest)
 		{
-			FPackageName::FindPackagesInDirectory(ContentPackages, *FPaths::EngineContentDir());
+			//FPackageName::FindPackagesInDirectory(ContentPackages, *FPaths::EngineContentDir());
 			FPackageName::FindPackagesInDirectory(ContentPackages, *FPaths::GameContentDir());
 		}
 
