Make SEditorViewport toolbar handling more flexible
- MakeViewportToolbar (top toolbar only) is now optional to implement in subclasses and may be deprecated in the future
- Introduced a new method PopulateViewportOverlays that allows arbitrary toolbar creation (bottom, sides, etc...)

diff --git a/Engine/Source/Editor/UnrealEd/Private/SEditorViewport.cpp b/Engine/Source/Editor/UnrealEd/Private/SEditorViewport.cpp
--- a/Engine/Source/Editor/UnrealEd/Private/SEditorViewport.cpp
+++ b/Engine/Source/Editor/UnrealEd/Private/SEditorViewport.cpp
@@ -82,6 +82,8 @@ void SEditorViewport::Construct( const FArguments& InArgs )
 				ViewportToolbar.ToSharedRef()
 			];
 	}
+
+	PopulateViewportOverlays(ViewportOverlay.ToSharedRef());
 }
 
 FReply SEditorViewport::OnKeyDown( const FGeometry& MyGeometry, const FKeyEvent& InKeyEvent )
diff --git a/Engine/Source/Editor/UnrealEd/Public/SEditorViewport.h b/Engine/Source/Editor/UnrealEd/Public/SEditorViewport.h
--- a/Engine/Source/Editor/UnrealEd/Public/SEditorViewport.h
+++ b/Engine/Source/Editor/UnrealEd/Public/SEditorViewport.h
@@ -53,7 +53,13 @@ public:
 
 protected:
 	virtual TSharedRef<FEditorViewportClient> MakeEditorViewportClient() = 0;
-	virtual TSharedPtr<SWidget> MakeViewportToolbar() = 0;
+
+	// Implement this to add a viewport toolbar to the inside top of the viewport
+	virtual TSharedPtr<SWidget> MakeViewportToolbar() { return TSharedPtr<SWidget>(nullptr); }
+
+	// Implement this to add an arbitrary set of toolbars or other overlays to the inside of the viewport
+	virtual void PopulateViewportOverlays(TSharedRef<class SOverlay> Overlay) { }
+
 	virtual void BindCommands();
 	virtual const FSlateBrush* OnGetViewportBorderBrush() const { return NULL; }
 	virtual FSlateColor OnGetViewportBorderColorAndOpacity() const { return FLinearColor::Black; }
