Adjusted the LoadAllMapsInEditor test to use load non-engine content.

diff --git a/Engine/Source/Editor/UnrealEd/Private/Tests/EditorAutomationTests.cpp b/Engine/Source/Editor/UnrealEd/Private/Tests/EditorAutomationTests.cpp
--- a/Engine/Source/Editor/UnrealEd/Private/Tests/EditorAutomationTests.cpp
+++ b/Engine/Source/Editor/UnrealEd/Private/Tests/EditorAutomationTests.cpp
@@ -738,8 +738,11 @@ void FLoadAllMapsInEditorTest::GetTests(TArray<FString>& OutBeautifiedNames, TAr
 		{
 			if (FAutomationTestFramework::GetInstance().ShouldTestContent(Filename))
 			{
-				OutBeautifiedNames.Add(FPaths::GetBaseFilename(Filename));
-				OutTestCommands.Add(Filename);
+				if (!Filename.Contains(TEXT("/Engine/")))
+				{
+					OutBeautifiedNames.Add(FPaths::GetBaseFilename(Filename));
+					OutTestCommands.Add(Filename);
+				}
 			}
 		}
 	}
