Merging

//depot/UE4-GDC2015/Engine/Source/Editor/UnrealEd/Private/StaticMeshEdit.cpp

to //depot/UE4-GDC2015/Engine/Source/Editor/UnrealEd/Private/StaticMeshEdit.cpp

diff --git a/Engine/Source/Editor/UnrealEd/Private/StaticMeshEdit.cpp b/Engine/Source/Editor/UnrealEd/Private/StaticMeshEdit.cpp
--- a/Engine/Source/Editor/UnrealEd/Private/StaticMeshEdit.cpp
+++ b/Engine/Source/Editor/UnrealEd/Private/StaticMeshEdit.cpp
@@ -814,14 +814,21 @@ static void TransformPolys(UPolys* Polys,const FMatrix& Matrix)
 
 }
 
-struct ExistingStaticMeshData
+// LOD data to copy over
+struct ExistingLODMeshData
 {
+	FMeshBuildSettings			ExistingBuildSettings;
+	FMeshReductionSettings		ExistingReductionSettings;
+	FRawMesh					ExistingRawMesh;
 	TArray<UMaterialInterface*>	ExistingMaterials;
+	float						ExistingScreenSize;
+};
+
+struct ExistingStaticMeshData
+{
 	FMeshSectionInfoMap			ExistingSectionInfoMap;
-	TArray<FMeshBuildSettings>	ExistingBuildSettings;
-	TArray<FMeshReductionSettings> ExistingReductionSettings;
-	TArray<float>				ExistingScreenSize;
-	
+	TArray<ExistingLODMeshData>	ExistingLODData;
+
 	TArray<UStaticMeshSocket*>	ExistingSockets;
 
 	bool						ExistingUseMaximumStreamingTexelRatio;
@@ -852,14 +859,39 @@ ExistingStaticMeshData* SaveExistingStaticMeshData(UStaticMesh* ExistingMesh)
 	{
 		ExistingMeshDataPtr = new ExistingStaticMeshData();
 
-		ExistingMeshDataPtr->ExistingMaterials = ExistingMesh->Materials;
-		ExistingMeshDataPtr->ExistingSectionInfoMap = ExistingMesh->SectionInfoMap;
+		FMeshSectionInfoMap OldSectionInfoMap = ExistingMesh->SectionInfoMap;
+
+		ExistingMeshDataPtr->ExistingLODData.AddZeroed(ExistingMesh->SourceModels.Num());
 
+		// refresh material and section info map here
+		// we have to make sure it only contains valid item
+		// we go through section info and only add it back if used, otherwise we don't want to use
+		ExistingMesh->SectionInfoMap.Clear();
+		int32 TotalMaterialIndex=0;
 		for(int32 i=0; i<ExistingMesh->SourceModels.Num(); i++)
 		{
-			ExistingMeshDataPtr->ExistingBuildSettings.Add(ExistingMesh->SourceModels[i].BuildSettings);
-			ExistingMeshDataPtr->ExistingReductionSettings.Add(ExistingMesh->SourceModels[i].ReductionSettings);
-			ExistingMeshDataPtr->ExistingScreenSize.Add(ExistingMesh->SourceModels[i].ScreenSize);
+			FStaticMeshLODResources& LOD = ExistingMesh->RenderData->LODResources[i];
+			int32 NumSections = LOD.Sections.Num();
+			for(int32 SectionIndex = 0; SectionIndex < NumSections; ++SectionIndex)
+			{
+				FMeshSectionInfo Info = OldSectionInfoMap.Get(i, SectionIndex);
+				if(ExistingMesh->Materials.IsValidIndex(Info.MaterialIndex))
+				{
+					ExistingMeshDataPtr->ExistingLODData[i].ExistingMaterials.Add(ExistingMesh->Materials[Info.MaterialIndex]);
+
+					// @Todo @fixme
+					// have to refresh material index since it might be pointing at wrong one
+					// this will break IF the base material number grows or shoterns and index will be off
+					// I think we have to save material index per section, so that we don't have to worry about global index
+					Info.MaterialIndex = TotalMaterialIndex++;
+					ExistingMeshDataPtr->ExistingSectionInfoMap.Set(i, SectionIndex, Info);
+				}
+			}
+
+			ExistingMeshDataPtr->ExistingLODData[i].ExistingBuildSettings = ExistingMesh->SourceModels[i].BuildSettings;
+			ExistingMeshDataPtr->ExistingLODData[i].ExistingReductionSettings = ExistingMesh->SourceModels[i].ReductionSettings;
+			ExistingMeshDataPtr->ExistingLODData[i].ExistingScreenSize = ExistingMesh->SourceModels[i].ScreenSize;
+			ExistingMesh->SourceModels[i].RawMeshBulkData->LoadRawMesh(ExistingMeshDataPtr->ExistingLODData[i].ExistingRawMesh);
 		}
 
 		ExistingMeshDataPtr->ExistingSockets = ExistingMesh->Sockets;
@@ -886,27 +918,31 @@ void RestoreExistingMeshData(struct ExistingStaticMeshData* ExistingMeshDataPtr,
 {
 	if ( ExistingMeshDataPtr && NewMesh )
 	{
-		int32 NumCommonMaterials = FMath::Min(NewMesh->Materials.Num(), ExistingMeshDataPtr->ExistingMaterials.Num());
+		int32 NumCommonMaterials = FMath::Min(NewMesh->Materials.Num(), ExistingMeshDataPtr->ExistingLODData[0].ExistingMaterials.Num());
 		for (int32 MaterialIndex = 0; MaterialIndex < NumCommonMaterials; ++MaterialIndex)
 		{
-			NewMesh->Materials[MaterialIndex] = ExistingMeshDataPtr->ExistingMaterials[MaterialIndex];
+			NewMesh->Materials[MaterialIndex] = ExistingMeshDataPtr->ExistingLODData[0].ExistingMaterials[MaterialIndex];
 		}
 
 		NewMesh->SectionInfoMap.CopyFrom(ExistingMeshDataPtr->ExistingSectionInfoMap);
 
-		int32 NumCommonLODs = FMath::Min<int32>(ExistingMeshDataPtr->ExistingBuildSettings.Num(), NewMesh->SourceModels.Num());
+		int32 NumCommonLODs = FMath::Min<int32>(ExistingMeshDataPtr->ExistingLODData.Num(), NewMesh->SourceModels.Num());
 		for(int32 i=0; i<NumCommonLODs; i++)
 		{
-			NewMesh->SourceModels[i].BuildSettings = ExistingMeshDataPtr->ExistingBuildSettings[i];
-			NewMesh->SourceModels[i].ReductionSettings = ExistingMeshDataPtr->ExistingReductionSettings[i];
-			NewMesh->SourceModels[i].ScreenSize = ExistingMeshDataPtr->ExistingScreenSize[i];
+			NewMesh->SourceModels[i].BuildSettings = ExistingMeshDataPtr->ExistingLODData[i].ExistingBuildSettings;
+			NewMesh->SourceModels[i].ReductionSettings = ExistingMeshDataPtr->ExistingLODData[i].ExistingReductionSettings;
+			NewMesh->SourceModels[i].ScreenSize = ExistingMeshDataPtr->ExistingLODData[i].ExistingScreenSize;
 		}
-		for(int32 i=NumCommonLODs; i < ExistingMeshDataPtr->ExistingBuildSettings.Num(); ++i)
+
+		for(int32 i=NumCommonLODs; i < ExistingMeshDataPtr->ExistingLODData.Num(); ++i)
 		{
+			NewMesh->Materials.Append(ExistingMeshDataPtr->ExistingLODData[i].ExistingMaterials);
+
 			FStaticMeshSourceModel* SrcModel = new(NewMesh->SourceModels) FStaticMeshSourceModel();
-			SrcModel->BuildSettings = ExistingMeshDataPtr->ExistingBuildSettings[i];
-			SrcModel->ReductionSettings = ExistingMeshDataPtr->ExistingReductionSettings[i];
-			SrcModel->ScreenSize = ExistingMeshDataPtr->ExistingScreenSize[i];
+			SrcModel->RawMeshBulkData->SaveRawMesh(ExistingMeshDataPtr->ExistingLODData[i].ExistingRawMesh);
+			SrcModel->BuildSettings = ExistingMeshDataPtr->ExistingLODData[i].ExistingBuildSettings;
+			SrcModel->ReductionSettings = ExistingMeshDataPtr->ExistingLODData[i].ExistingReductionSettings;
+			SrcModel->ScreenSize = ExistingMeshDataPtr->ExistingLODData[i].ExistingScreenSize;
 		}
 
 		// Assign sockets from old version of this StaticMesh.
