Merging

//depot/UE4-GDC2015/Engine/Source/Editor/UnrealEd/Private/StaticMeshEdit.cpp

to //depot/UE4-GDC2015/Engine/Source/Editor/UnrealEd/Private/StaticMeshEdit.cpp

diff --git a/Engine/Source/Editor/UnrealEd/Private/StaticMeshEdit.cpp b/Engine/Source/Editor/UnrealEd/Private/StaticMeshEdit.cpp
--- a/Engine/Source/Editor/UnrealEd/Private/StaticMeshEdit.cpp
+++ b/Engine/Source/Editor/UnrealEd/Private/StaticMeshEdit.cpp
@@ -877,14 +877,19 @@ ExistingStaticMeshData* SaveExistingStaticMeshData(UStaticMesh* ExistingMesh)
 				FMeshSectionInfo Info = OldSectionInfoMap.Get(i, SectionIndex);
 				if(ExistingMesh->Materials.IsValidIndex(Info.MaterialIndex))
 				{
-					ExistingMeshDataPtr->ExistingLODData[i].ExistingMaterials.Add(ExistingMesh->Materials[Info.MaterialIndex]);
-
-					// @Todo @fixme
-					// have to refresh material index since it might be pointing at wrong one
-					// this will break IF the base material number grows or shoterns and index will be off
-					// I think we have to save material index per section, so that we don't have to worry about global index
-					Info.MaterialIndex = TotalMaterialIndex++;
-					ExistingMeshDataPtr->ExistingSectionInfoMap.Set(i, SectionIndex, Info);
+					// we only save per LOD separeate IF the material index isn't added yet. 
+					// if it's already added, we don't have to add another one. 
+					if (Info.MaterialIndex >= TotalMaterialIndex)
+					{
+						ExistingMeshDataPtr->ExistingLODData[i].ExistingMaterials.Add(ExistingMesh->Materials[Info.MaterialIndex]);
+
+						// @Todo @fixme
+						// have to refresh material index since it might be pointing at wrong one
+						// this will break IF the base material number grows or shoterns and index will be off
+						// I think we have to save material index per section, so that we don't have to worry about global index
+						Info.MaterialIndex = TotalMaterialIndex++;
+						ExistingMeshDataPtr->ExistingSectionInfoMap.Set(i, SectionIndex, Info);
+					}
 				}
 			}
 
@@ -936,10 +941,17 @@ void RestoreExistingMeshData(struct ExistingStaticMeshData* ExistingMeshDataPtr,
 
 		for(int32 i=NumCommonLODs; i < ExistingMeshDataPtr->ExistingLODData.Num(); ++i)
 		{
-			NewMesh->Materials.Append(ExistingMeshDataPtr->ExistingLODData[i].ExistingMaterials);
+			if (ExistingMeshDataPtr->ExistingLODData[i].ExistingMaterials.Num() > 0)
+			{
+				NewMesh->Materials.Append(ExistingMeshDataPtr->ExistingLODData[i].ExistingMaterials);
+			}
 
 			FStaticMeshSourceModel* SrcModel = new(NewMesh->SourceModels) FStaticMeshSourceModel();
-			SrcModel->RawMeshBulkData->SaveRawMesh(ExistingMeshDataPtr->ExistingLODData[i].ExistingRawMesh);
+
+			if (ExistingMeshDataPtr->ExistingLODData[i].ExistingRawMesh.IsValidOrFixable())
+			{
+				SrcModel->RawMeshBulkData->SaveRawMesh(ExistingMeshDataPtr->ExistingLODData[i].ExistingRawMesh);
+			}
 			SrcModel->BuildSettings = ExistingMeshDataPtr->ExistingLODData[i].ExistingBuildSettings;
 			SrcModel->ReductionSettings = ExistingMeshDataPtr->ExistingLODData[i].ExistingReductionSettings;
 			SrcModel->ScreenSize = ExistingMeshDataPtr->ExistingLODData[i].ExistingScreenSize;
