Editor: Allow editor mode toolkits a chance at processing command bindings

diff --git a/Engine/Source/Editor/UnrealEd/Private/EdMode.cpp b/Engine/Source/Editor/UnrealEd/Private/EdMode.cpp
--- a/Engine/Source/Editor/UnrealEd/Private/EdMode.cpp
+++ b/Engine/Source/Editor/UnrealEd/Private/EdMode.cpp
@@ -260,24 +260,34 @@ bool FEdMode::CapturedMouseMove( FEditorViewportClient* InViewportClient, FViewp
 
 bool FEdMode::InputKey(FEditorViewportClient* ViewportClient, FViewport* Viewport, FKey Key, EInputEvent Event)
 {
-	if( GetCurrentTool() && GetCurrentTool()->InputKey( ViewportClient, Viewport, Key, Event ) )
+	// First try the currently selected tool
+	if ((GetCurrentTool() != nullptr) && GetCurrentTool()->InputKey(ViewportClient, Viewport, Key, Event))
 	{
 		return true;
 	}
 	else
 	{
-		// Pass input up to selected actors if not in a tool mode
+		// Next pass input to the mode toolkit
+		if (Toolkit.IsValid() && ((Event == IE_Pressed) || (Event == IE_Repeat)))
+		{
+			if (Toolkit->GetToolkitCommands()->ProcessCommandBindings(Key, FSlateApplication::Get().GetModifierKeys(), (Event == IE_Repeat)))
+			{
+				return true;
+			}
+		}
+
+		// Finally, pass input up to selected actors if not in a tool mode
 		TArray<AActor*> SelectedActors;
 		Owner->GetSelectedActors()->GetSelectedObjects<AActor>(SelectedActors);
 
-		for( TArray<AActor*>::TIterator it(SelectedActors); it; ++it )
+		for (TArray<AActor*>::TIterator It(SelectedActors); It; ++It)
 		{
 			// Tell the object we've had a key press
-			(*it)->EditorKeyPressed(Key, Event);
+			(*It)->EditorKeyPressed(Key, Event);
 		}
 	}
 
-	return 0;
+	return false;
 }
 
 bool FEdMode::InputAxis(FEditorViewportClient* InViewportClient, FViewport* Viewport, int32 ControllerId, FKey Key, float Delta, float DeltaTime)
