Removed experiemental IWCE setting since it isnt experimental

diff --git a/Engine/Source/Editor/LevelEditor/Private/SActorDetails.cpp b/Engine/Source/Editor/LevelEditor/Private/SActorDetails.cpp
--- a/Engine/Source/Editor/LevelEditor/Private/SActorDetails.cpp
+++ b/Engine/Source/Editor/LevelEditor/Private/SActorDetails.cpp
@@ -190,7 +190,7 @@ void SActorDetails::SetObjects(const TArray<UObject*>& InObjects, bool bForceRef
 
 		bool bShowingComponents = false;
 
-		if(InObjects.Num() == 1 && FKismetEditorUtilities::CanCreateBlueprintOfClass(InObjects[0]->GetClass()) && GetDefault<UEditorExperimentalSettings>()->bInWorldBPEditing)
+		if(InObjects.Num() == 1 && FKismetEditorUtilities::CanCreateBlueprintOfClass(InObjects[0]->GetClass()))
 		{
 			auto Actor = GetSelectedActorInEditor();
 			if(Actor)
diff --git a/Engine/Source/Editor/PlacementMode/Private/SPlacementModeTools.cpp b/Engine/Source/Editor/PlacementMode/Private/SPlacementModeTools.cpp
--- a/Engine/Source/Editor/PlacementMode/Private/SPlacementModeTools.cpp
+++ b/Engine/Source/Editor/PlacementMode/Private/SPlacementModeTools.cpp
@@ -691,13 +691,27 @@ TSharedRef< SWidget > SPlacementModeTools::BuildBasicWidget()
 	TSharedRef<SVerticalBox> VerticalBox = SNew( SVerticalBox )
 	+SVerticalBox::Slot()
 	.AutoHeight()
+	[
+		BuildDraggableAssetWidget(UActorFactoryEmptyActor::StaticClass())
+	]
+	+SVerticalBox::Slot()
+	.AutoHeight()
+	[
+		BuildDraggableAssetWidget(UActorFactoryCharacter::StaticClass())
+	]
+	+SVerticalBox::Slot()
+	.AutoHeight()
+	[
+		BuildDraggableAssetWidget(UActorFactoryPawn::StaticClass())
+	]
+	.AutoHeight()
 	[
 		BuildDraggableAssetWidget(UActorFactoryPointLight::StaticClass())
 	]
 	+ SVerticalBox::Slot()
 	.AutoHeight()
 	[
-		BuildDraggableAssetWidget( UActorFactoryPlayerStart::StaticClass() )
+		BuildDraggableAssetWidget(UActorFactoryPlayerStart::StaticClass())
 	]
 	+ SVerticalBox::Slot()
 	.AutoHeight()
@@ -734,27 +748,6 @@ TSharedRef< SWidget > SPlacementModeTools::BuildBasicWidget()
 		BuildDraggableAssetWidget( UActorFactoryTriggerSphere::StaticClass() )
 	];
 
-	if( GetDefault<UEditorExperimentalSettings>()->bInWorldBPEditing )
-	{
-		VerticalBox->InsertSlot(0)
-		.AutoHeight()
-		[
-			BuildDraggableAssetWidget(UActorFactoryEmptyActor::StaticClass())
-		];
-
-		VerticalBox->InsertSlot(1)
-		.AutoHeight()
-		[
-			BuildDraggableAssetWidget(UActorFactoryCharacter::StaticClass())
-		];
-
-		VerticalBox->InsertSlot(2)
-		.AutoHeight()
-		[
-			BuildDraggableAssetWidget(UActorFactoryPawn::StaticClass())
-		];
-	}
-
 	return VerticalBox;
 }
 
diff --git a/Engine/Source/Editor/UnrealEd/Classes/Settings/EditorExperimentalSettings.h b/Engine/Source/Editor/UnrealEd/Classes/Settings/EditorExperimentalSettings.h
--- a/Engine/Source/Editor/UnrealEd/Classes/Settings/EditorExperimentalSettings.h
+++ b/Engine/Source/Editor/UnrealEd/Classes/Settings/EditorExperimentalSettings.h
@@ -33,10 +33,6 @@ public:
 	UPROPERTY(EditAnywhere, config, Category=Tools, meta=(DisplayName="Editor Utility Blueprints (Blutility)"))
 	bool bEnableEditorUtilityBlueprints;
 
-	/** Enable In World BP Editing (WIP). */
-	UPROPERTY(EditAnywhere, config, Category = Tools, meta = (DisplayName = "In World Blueprint Editing"))
-	bool bInWorldBPEditing;
-
 	/** Enable Single Layout BP Editor. */
 	UPROPERTY(EditAnywhere, config, Category = Tools, meta = ( DisplayName = "Single Layout Blueprint Editor" ))
 	bool bUnifiedBlueprintEditor;
diff --git a/Engine/Source/Editor/UnrealEd/Private/Factories/ActorFactory.cpp b/Engine/Source/Editor/UnrealEd/Private/Factories/ActorFactory.cpp
--- a/Engine/Source/Editor/UnrealEd/Private/Factories/ActorFactory.cpp
+++ b/Engine/Source/Editor/UnrealEd/Private/Factories/ActorFactory.cpp
@@ -1190,13 +1190,12 @@ UActorFactoryEmptyActor::UActorFactoryEmptyActor(const FObjectInitializer& Objec
 
 bool UActorFactoryEmptyActor::CanCreateActorFrom( const FAssetData& AssetData, FText& OutErrorMsg )
 {
-	return GetDefault<UEditorExperimentalSettings>()->bInWorldBPEditing ? AssetData.ObjectPath == FName(*AActor::StaticClass()->GetPathName()) : false;
+	return AssetData.ObjectPath == FName(*AActor::StaticClass()->GetPathName());
 }
 
 AActor* UActorFactoryEmptyActor::SpawnActor( UObject* Asset, ULevel* InLevel, const FVector& Location, const FRotator& Rotation, EObjectFlags ObjectFlags, const FName& Name )
 {
 	AActor* NewActor = nullptr;
-	if(GetDefault<UEditorExperimentalSettings>()->bInWorldBPEditing) 
 	{
 		// Spawn a temporary actor for dragging around
 		NewActor = Super::SpawnActor(Asset, InLevel, Location, Rotation, ObjectFlags, Name);
@@ -1229,16 +1228,12 @@ UActorFactoryCharacter::UActorFactoryCharacter(const FObjectInitializer& ObjectI
 
 bool UActorFactoryCharacter::CanCreateActorFrom(const FAssetData& AssetData, FText& OutErrorMsg)
 {
-	return GetDefault<UEditorExperimentalSettings>()->bInWorldBPEditing ? AssetData.ObjectPath == FName(*ACharacter::StaticClass()->GetPathName()) : false;
+	return AssetData.ObjectPath == FName(*ACharacter::StaticClass()->GetPathName());
 }
 
 AActor* UActorFactoryCharacter::SpawnActor(UObject* Asset, ULevel* InLevel, const FVector& Location, const FRotator& Rotation, EObjectFlags ObjectFlags, const FName& Name)
 {
-	AActor* NewActor = nullptr;
-	if(GetDefault<UEditorExperimentalSettings>()->bInWorldBPEditing)
-	{
-			NewActor = Super::SpawnActor( Asset, InLevel, Location, Rotation, ObjectFlags, Name );
-		}
+	AActor* NewActor = Super::SpawnActor( Asset, InLevel, Location, Rotation, ObjectFlags, Name );
 
 	return NewActor;
 }
@@ -1256,17 +1251,12 @@ UActorFactoryPawn::UActorFactoryPawn(const FObjectInitializer& ObjectInitializer
 
 bool UActorFactoryPawn::CanCreateActorFrom(const FAssetData& AssetData, FText& OutErrorMsg)
 {
-	return GetDefault<UEditorExperimentalSettings>()->bInWorldBPEditing ? AssetData.ObjectPath == FName(*APawn::StaticClass()->GetPathName()) : false;
+	return AssetData.ObjectPath == FName(*APawn::StaticClass()->GetPathName());
 }
 
 AActor* UActorFactoryPawn::SpawnActor(UObject* Asset, ULevel* InLevel, const FVector& Location, const FRotator& Rotation, EObjectFlags ObjectFlags, const FName& Name)
 {
-	AActor* NewActor = nullptr;
-	if(GetDefault<UEditorExperimentalSettings>()->bInWorldBPEditing)
-	{
-		NewActor = Super::SpawnActor(Asset, InLevel, Location, Rotation, ObjectFlags, Name);
-	}
-
+	AActor* NewActor = Super::SpawnActor(Asset, InLevel, Location, Rotation, ObjectFlags, Name);
 	return NewActor;
 }
 
diff --git a/Engine/Source/Editor/UnrealEd/Private/LevelEditorViewport.cpp b/Engine/Source/Editor/UnrealEd/Private/LevelEditorViewport.cpp
--- a/Engine/Source/Editor/UnrealEd/Private/LevelEditorViewport.cpp
+++ b/Engine/Source/Editor/UnrealEd/Private/LevelEditorViewport.cpp
@@ -1845,7 +1845,7 @@ void FLevelEditorViewportClient::ProcessClick(FSceneView& View, HHitProxy* HitPr
 
 			const bool bSelectComponent = bActorAlreadySelectedExclusively && bActorIsBlueprintable && (bComponentAlreadySelected != bWasDoubleClick);
 
-			if (bSelectComponent && GetDefault<UEditorExperimentalSettings>()->bInWorldBPEditing)
+			if (bSelectComponent)
 			{
 				ClickHandlers::ClickComponent(this, ActorHitProxy, Click);
 			}
diff --git a/Engine/Source/Editor/UnrealEd/Private/Settings/SettingsClasses.cpp b/Engine/Source/Editor/UnrealEd/Private/Settings/SettingsClasses.cpp
--- a/Engine/Source/Editor/UnrealEd/Private/Settings/SettingsClasses.cpp
+++ b/Engine/Source/Editor/UnrealEd/Private/Settings/SettingsClasses.cpp
@@ -53,7 +53,6 @@ UDestructableMeshEditorSettings::UDestructableMeshEditorSettings( const FObjectI
 
 UEditorExperimentalSettings::UEditorExperimentalSettings( const FObjectInitializer& ObjectInitializer )
 	: Super(ObjectInitializer)
-	, bInWorldBPEditing(true)
 	, bUnifiedBlueprintEditor(true)
 	, bBlueprintableComponents(true)
 
