Fix mapcheck messages added during object loading not being shown in the map check window.

diff --git a/Engine/Source/Editor/UnrealEd/Classes/Editor/EditorEngine.h b/Engine/Source/Editor/UnrealEd/Classes/Editor/EditorEngine.h
--- a/Engine/Source/Editor/UnrealEd/Classes/Editor/EditorEngine.h
+++ b/Engine/Source/Editor/UnrealEd/Classes/Editor/EditorEngine.h
@@ -2418,7 +2418,7 @@ private:
 	/**
 	 * Checks map for common errors.
 	 */
-	bool Map_Check( UWorld* InWorld, const TCHAR* Str, FOutputDevice& Ar, bool bCheckDeprecatedOnly, EMapCheckNotification::Type Notification = EMapCheckNotification::DisplayResults );
+	bool Map_Check(UWorld* InWorld, const TCHAR* Str, FOutputDevice& Ar, bool bCheckDeprecatedOnly, EMapCheckNotification::Type Notification = EMapCheckNotification::DisplayResults, bool bClearLog = true);
 	bool Map_Scale( UWorld* InWorld, const TCHAR* Str, FOutputDevice& Ar );
 	bool Map_Setbrush( UWorld* InWorld, const TCHAR* Str, FOutputDevice& Ar );
 
diff --git a/Engine/Source/Editor/UnrealEd/Private/EditorServer.cpp b/Engine/Source/Editor/UnrealEd/Private/EditorServer.cpp
--- a/Engine/Source/Editor/UnrealEd/Private/EditorServer.cpp
+++ b/Engine/Source/Editor/UnrealEd/Private/EditorServer.cpp
@@ -3438,7 +3438,7 @@ void UEditorEngine::SelectByPropertyColoration(UWorld* InWorld)
 }
 
 
-bool UEditorEngine::Map_Check( UWorld* InWorld, const TCHAR* Str, FOutputDevice& Ar, bool bCheckDeprecatedOnly, EMapCheckNotification::Type Notification/*= EMapCheckNotification::DisplayResults*/)
+bool UEditorEngine::Map_Check( UWorld* InWorld, const TCHAR* Str, FOutputDevice& Ar, bool bCheckDeprecatedOnly, EMapCheckNotification::Type Notification/*= EMapCheckNotification::DisplayResults*/, bool bClearLog/*= true*/)
 {
 #define LOCTEXT_NAMESPACE "EditorEngine"
 	const FText CheckMapLocText = NSLOCTEXT("UnrealEd", "CheckingMap", "Checking map");
@@ -3446,6 +3446,8 @@ bool UEditorEngine::Map_Check( UWorld* InWorld, const TCHAR* Str, FOutputDevice&
 	const double StartTime = FPlatformTime::Seconds();
 
 	FMessageLog MapCheckLog("MapCheck");
+
+	if (bClearLog)
 	{
 		FFormatNamedArguments Arguments;
 		Arguments.Add(TEXT("Name"), FText::FromString(FPackageName::GetShortName(InWorld->GetOutermost())));
@@ -5703,6 +5705,7 @@ bool UEditorEngine::HandleMapCommand( const TCHAR* Str, FOutputDevice& Ar, UWorl
 	else if (FParse::Command (&Str,TEXT("CHECK")))
 	{
 		EMapCheckNotification::Type Notification = EMapCheckNotification::DisplayResults;
+		bool bClearLog = true;
 		if(FParse::Command(&Str,TEXT("DONTDISPLAYDIALOG")))
 		{
 			Notification = EMapCheckNotification::DontDisplayResults;
@@ -5711,11 +5714,16 @@ bool UEditorEngine::HandleMapCommand( const TCHAR* Str, FOutputDevice& Ar, UWorl
 		{
 			Notification = EMapCheckNotification::NotifyOfResults;
 		}
-		return Map_Check( InWorld, Str, Ar, false, Notification );
+		if (FParse::Command(&Str, TEXT("NOCLEARLOG")))
+		{
+			bClearLog = false;
+		}
+		return Map_Check(InWorld, Str, Ar, false, Notification, bClearLog);
 	}
 	else if (FParse::Command (&Str,TEXT("CHECKDEP")))
 	{
 		EMapCheckNotification::Type Notification = EMapCheckNotification::DisplayResults;
+		bool bClearLog = true;
 		if(FParse::Command(&Str,TEXT("DONTDISPLAYDIALOG")))
 		{
 			Notification = EMapCheckNotification::DontDisplayResults;
@@ -5724,7 +5732,11 @@ bool UEditorEngine::HandleMapCommand( const TCHAR* Str, FOutputDevice& Ar, UWorl
 		{
 			Notification = EMapCheckNotification::NotifyOfResults;
 		}
-		return Map_Check( InWorld, Str, Ar, true, Notification );
+		if (FParse::Command(&Str, TEXT("NOCLEARLOG")))
+		{
+			bClearLog = false;
+		}
+		return Map_Check(InWorld, Str, Ar, true, Notification, bClearLog);
 	}
 	else if (FParse::Command (&Str,TEXT("SCALE")))
 	{
diff --git a/Engine/Source/Editor/UnrealEd/Private/FileHelpers.cpp b/Engine/Source/Editor/UnrealEd/Private/FileHelpers.cpp
--- a/Engine/Source/Editor/UnrealEd/Private/FileHelpers.cpp
+++ b/Engine/Source/Editor/UnrealEd/Private/FileHelpers.cpp
@@ -2019,7 +2019,7 @@ void FEditorFileUtils::LoadMap(const FString& InFilename, bool LoadAsTemplate, b
 	if( !GIsDemoMode )
 	{
 		// Check for deprecated actor classes.
-		GEditor->Exec( World, TEXT("MAP CHECKDEP") );
+		GEditor->Exec(World, TEXT("MAP CHECKDEP NOCLEARLOG"));
 		FMessageLog("MapCheck").Open( EMessageSeverity::Warning );
 	}
 
