Opening a class via the Content Browser now opens the source file as well as the header file

diff --git a/Engine/Source/Developer/AssetTools/Private/AssetTypeActions/AssetTypeActions_Class.cpp b/Engine/Source/Developer/AssetTools/Private/AssetTypeActions/AssetTypeActions_Class.cpp
--- a/Engine/Source/Developer/AssetTools/Private/AssetTypeActions/AssetTypeActions_Class.cpp
+++ b/Engine/Source/Developer/AssetTools/Private/AssetTypeActions/AssetTypeActions_Class.cpp
@@ -141,12 +141,23 @@ void FAssetTypeActions_Class::OpenAssetEditor( const TArray<UObject*>& InObjects
 		UClass* const Class = Cast<UClass>(Object);
 		if(Class)
 		{
+			TArray<FString> FilesToOpen;
+
 			FString ClassHeaderPath;
 			if(FSourceCodeNavigation::FindClassHeaderPath(Class, ClassHeaderPath))
 			{
 				const FString AbsoluteHeaderPath = IFileManager::Get().ConvertToAbsolutePathForExternalAppForRead(*ClassHeaderPath);
-				FSourceCodeNavigation::OpenSourceFile(AbsoluteHeaderPath);
+				FilesToOpen.Add(AbsoluteHeaderPath);
 			}
+
+			FString ClassSourcePath;
+			if(FSourceCodeNavigation::FindClassSourcePath(Class, ClassSourcePath))
+			{
+				const FString AbsoluteSourcePath = IFileManager::Get().ConvertToAbsolutePathForExternalAppForRead(*ClassSourcePath);
+				FilesToOpen.Add(AbsoluteSourcePath);
+			}
+
+			FSourceCodeNavigation::OpenSourceFiles(FilesToOpen);
 		}
 	}
 }
diff --git a/Engine/Source/Editor/UnrealEd/Private/SourceCodeNavigation.cpp b/Engine/Source/Editor/UnrealEd/Private/SourceCodeNavigation.cpp
--- a/Engine/Source/Editor/UnrealEd/Private/SourceCodeNavigation.cpp
+++ b/Engine/Source/Editor/UnrealEd/Private/SourceCodeNavigation.cpp
@@ -1656,6 +1656,73 @@ bool FSourceCodeNavigation::FindClassHeaderPath( const UField *Field, FString &O
 	return false;
 }
 
+bool FSourceCodeNavigation::FindClassSourcePath( const UField *Field, FString &OutClassSourcePath )
+{
+	// Get the class package, and skip past the "/Script/" portion to get the module name
+	UPackage *ModulePackage = Field->GetTypedOuter<UPackage>();
+	FString ModulePackageName = ModulePackage->GetName();
+
+	int32 ModuleNameIdx;
+	if(ModulePackageName.FindLastChar(TEXT('/'), ModuleNameIdx))
+	{
+		// Find the base path for the module
+		FString ModuleBasePath;
+		if(FSourceCodeNavigation::FindModulePath(*ModulePackageName + ModuleNameIdx + 1, ModuleBasePath))
+		{
+			// Get the metadata for the class path relative to the module base
+			// Given this we can try and find the corresponding .cpp file
+			FString ModuleRelativePath = ModulePackage->GetMetaData()->GetValue(Field, TEXT("ModuleRelativePath"));
+			if(ModuleRelativePath.Len() > 0)
+			{
+				const FString PotentialCppLeafname = FPaths::GetBaseFilename(ModuleRelativePath) + TEXT(".cpp");
+				FString PotentialCppFilename = ModuleBasePath / FPaths::GetPath(ModuleRelativePath) / PotentialCppLeafname;
+
+				// Is the .cpp file in the same folder as the header file?
+				if(FPaths::FileExists(PotentialCppFilename))
+				{
+					OutClassSourcePath = PotentialCppFilename;
+					return true;
+				}
+
+				const FString PublicPath = ModuleBasePath / "Public" / "";		// Ensure trailing /
+				const FString PrivatePath = ModuleBasePath / "Private" / "";	// Ensure trailing /
+				const FString ClassesPath = ModuleBasePath / "Classes" / "";	// Ensure trailing /
+
+				// If the path starts with Public or Classes, try swapping those out with Private
+				if(PotentialCppFilename.StartsWith(PublicPath))
+				{
+					PotentialCppFilename.ReplaceInline(*PublicPath, *PrivatePath);
+				}
+				else if(PotentialCppFilename.StartsWith(ClassesPath))
+				{
+					PotentialCppFilename.ReplaceInline(*ClassesPath, *PrivatePath);
+				}
+				else
+				{
+					PotentialCppFilename.Empty();
+				}
+				if(!PotentialCppFilename.IsEmpty() && FPaths::FileExists(PotentialCppFilename))
+				{
+					OutClassSourcePath = PotentialCppFilename;
+					return true;
+				}
+
+				// Still no luck, try and search for the file on the filesystem
+				TArray<FString> Filenames;
+				IFileManager::Get().FindFilesRecursive(Filenames, *ModuleBasePath, *PotentialCppLeafname, true, false, false);
+				
+				if(Filenames.Num() > 0)
+				{
+					// Assume it's the first match (we should really only find a single file with a given name within a project anyway)
+					OutClassSourcePath = Filenames[0];
+					return true;
+				}
+			}
+		}
+	}
+	return false;
+}
+
 void FSourceCodeNavigationImpl::OnSuggestedIDEInstallerDownloadComplete(FHttpRequestPtr Request, FHttpResponsePtr Response, bool bWasSuccessful, FOnIDEInstallerDownloadComplete OnDownloadComplete)
 {
 	if (bWasSuccessful)
diff --git a/Engine/Source/Editor/UnrealEd/Public/SourceCodeNavigation.h b/Engine/Source/Editor/UnrealEd/Public/SourceCodeNavigation.h
--- a/Engine/Source/Editor/UnrealEd/Public/SourceCodeNavigation.h
+++ b/Engine/Source/Editor/UnrealEd/Public/SourceCodeNavigation.h
@@ -169,9 +169,12 @@ public:
 	/** Finds the base directory for a given module name. Does not rely on symbols; finds matching .build.cs files. */
 	UNREALED_API static bool FindModulePath( const FString& ModuleName, FString &OutModulePath );
 
-	/** Finds the path to a given class. Does not rely on symbols; finds matching .build.cs files. */
+	/** Finds the path to a given class header. Does not rely on symbols; finds matching .build.cs files. */
 	UNREALED_API static bool FindClassHeaderPath( const UField *Field, FString &OutClassHeaderPath );
 
+	/** Finds the path to a given class source. Does not rely on symbols; finds matching .build.cs files. */
+	UNREALED_API static bool FindClassSourcePath( const UField *Field, FString &OutClassSourcePath );
+
 	/** Opens a single source file */
 	UNREALED_API static bool OpenSourceFile( const FString& AbsoluteSourcePath, int32 LineNumber = 0, int32 ColumnNumber = 0 );
 
