UE-11460 A level will not save if a Blueprint that uses a Blueprint Component with an Event Dispatcher is placed in it

FBlueprintCompileReinstancer::ReplaceInstancesOfClass copies objects flags.

diff --git a/Engine/Source/Editor/UnrealEd/Private/Kismet2/KismetReinstanceUtilities.cpp b/Engine/Source/Editor/UnrealEd/Private/Kismet2/KismetReinstanceUtilities.cpp
--- a/Engine/Source/Editor/UnrealEd/Private/Kismet2/KismetReinstanceUtilities.cpp
+++ b/Engine/Source/Editor/UnrealEd/Private/Kismet2/KismetReinstanceUtilities.cpp
@@ -871,13 +871,19 @@ void FBlueprintCompileReinstancer::ReplaceInstancesOfClass(UClass* OldClass, UCl
 			}
 			else
 			{
+				auto OldFlags = OldObject->GetFlags();
+
 				FName OldName(OldObject->GetFName());
 				OldObject->Rename(NULL, OldObject->GetOuter(), REN_DoNotDirty | REN_DontCreateRedirectors);
 				NewUObject = NewObject<UObject>(OldObject->GetOuter(), NewClass, OldName);
 				check(NewUObject != nullptr);
 
+				auto FlagMask = RF_Public | RF_ArchetypeObject | RF_Transactional | RF_Transient | RF_TextExportTransient | RF_InheritableComponentTemplate; //TODO: what about RF_RootSet and RF_Standalone ?
+				NewObject->SetFlags(OldFlags & FlagMask);
+
 				UEditorEngine::CopyPropertiesForUnrelatedObjects(OldObject, NewUObject);
 
+
 				if (UAnimInstance* AnimTree = Cast<UAnimInstance>(NewUObject))
 				{
 					// Initialising the anim instance isn't enough to correctly set up the skeletal mesh again in a
