Improved output logging from localization manifest and archive generation commandlets - now gives the proper path of the file being written to.

diff --git a/Engine/Source/Editor/UnrealEd/Classes/Commandlets/GenerateGatherArchiveCommandlet.h b/Engine/Source/Editor/UnrealEd/Classes/Commandlets/GenerateGatherArchiveCommandlet.h
--- a/Engine/Source/Editor/UnrealEd/Classes/Commandlets/GenerateGatherArchiveCommandlet.h
+++ b/Engine/Source/Editor/UnrealEd/Classes/Commandlets/GenerateGatherArchiveCommandlet.h
@@ -21,7 +21,7 @@ public:
 	virtual int32 Main( const FString& Params ) override;
 	// End UCommandlet Interface
 	
-	bool WriteArchiveToFile( TSharedRef< FJsonObject> ArchiveJSONObject, const FString& OutputDirectoryPath, const TCHAR* Culture, const FString& FileName );
+	bool WriteArchiveToFile( TSharedRef< FJsonObject> ArchiveJSONObject, const FString& OutputFilePath );
 	void BuildArchiveFromManifest( TSharedRef< const FInternationalizationManifest > InManifest, TSharedRef< FInternationalizationArchive > Archive, const FString& SourceCulture, const FString& TargetCulture  );
 	void AppendArchiveData( TSharedRef< const FInternationalizationArchive > InArchiveToAppend, TSharedRef< FInternationalizationArchive > ArchiveCombined );
 	void ConditionTranslation( FLocItem& LocItem );
diff --git a/Engine/Source/Editor/UnrealEd/Classes/Commandlets/GenerateGatherManifestCommandlet.h b/Engine/Source/Editor/UnrealEd/Classes/Commandlets/GenerateGatherManifestCommandlet.h
--- a/Engine/Source/Editor/UnrealEd/Classes/Commandlets/GenerateGatherManifestCommandlet.h
+++ b/Engine/Source/Editor/UnrealEd/Classes/Commandlets/GenerateGatherManifestCommandlet.h
@@ -16,7 +16,7 @@ public:
 	// Begin UCommandlet Interface
 	virtual int32 Main( const FString& Params ) override;
 
-	bool WriteManifest( const TSharedPtr<FInternationalizationManifest>& InManifest, const FString& OutputFilePath );
+	bool WriteManifestToFile( const TSharedPtr<FInternationalizationManifest>& InManifest, const FString& OutputFilePath );
 	// End UCommandlet Interface
 #endif
 };
diff --git a/Engine/Source/Editor/UnrealEd/Private/Commandlets/GenerateGatherArchiveCommandlet.cpp b/Engine/Source/Editor/UnrealEd/Private/Commandlets/GenerateGatherArchiveCommandlet.cpp
--- a/Engine/Source/Editor/UnrealEd/Private/Commandlets/GenerateGatherArchiveCommandlet.cpp
+++ b/Engine/Source/Editor/UnrealEd/Private/Commandlets/GenerateGatherArchiveCommandlet.cpp
@@ -251,9 +251,10 @@ int32 UGenerateGatherArchiveCommandlet::Main( const FString& Params )
 		TSharedRef< FJsonObject > OutputArchiveJsonObj = MakeShareable( new FJsonObject );
 		ArchiveSerializer.SerializeArchive( OutputInternationalizationArchive, OutputArchiveJsonObj );
 
-		if( !WriteArchiveToFile( OutputArchiveJsonObj, DestinationPath, *CulturesToGenerate[Culture], *ArchiveName ) )
+		const FString ArchivePath = FPaths::ConvertRelativePathToFull(DestinationPath) / CulturesToGenerate[Culture] / ArchiveName;
+		if( !WriteArchiveToFile( OutputArchiveJsonObj, ArchivePath ) )
 		{
-			UE_LOG( LogGenerateArchiveCommandlet, Error,TEXT("Failed to write archive to %s."), *DestinationPath );				
+			UE_LOG( LogGenerateArchiveCommandlet, Error,TEXT("Failed to write archive to %s."), *ArchivePath );				
 			return -1;
 		}
 	}
@@ -261,17 +262,15 @@ int32 UGenerateGatherArchiveCommandlet::Main( const FString& Params )
 	return 0;
 }
 
-bool UGenerateGatherArchiveCommandlet::WriteArchiveToFile(TSharedRef< FJsonObject > ArchiveJSONObject, const FString& OutputDirectoryPath, const TCHAR* Culture, const FString& FileName)
+bool UGenerateGatherArchiveCommandlet::WriteArchiveToFile(TSharedRef< FJsonObject > ArchiveJSONObject, const FString& OutputFilePath)
 {
-	FString FullOutPath = FPaths::ConvertRelativePathToFull(OutputDirectoryPath) / Culture / FileName;
+	UE_LOG(LogGenerateArchiveCommandlet, Log, TEXT("Writing archive to %s."), *OutputFilePath);
 
-	if( !WriteJSONToTextFile(ArchiveJSONObject, FullOutPath, SourceControlInfo ) )
+	if( !WriteJSONToTextFile(ArchiveJSONObject, OutputFilePath, SourceControlInfo ) )
 	{
 		return false;
 	}
 
-	UE_LOG(LogGenerateArchiveCommandlet, Log, TEXT("Writing archive to %s."), *OutputDirectoryPath);
-
 	return true;
 }
 
diff --git a/Engine/Source/Editor/UnrealEd/Private/Commandlets/GenerateGatherManifestCommandlet.cpp b/Engine/Source/Editor/UnrealEd/Private/Commandlets/GenerateGatherManifestCommandlet.cpp
--- a/Engine/Source/Editor/UnrealEd/Private/Commandlets/GenerateGatherManifestCommandlet.cpp
+++ b/Engine/Source/Editor/UnrealEd/Private/Commandlets/GenerateGatherManifestCommandlet.cpp
@@ -83,17 +83,19 @@ int32 UGenerateGatherManifestCommandlet::Main( const FString& Params )
 		ManifestInfo->ApplyManifestDependencies();
 	}
 	
-
-	if( !WriteManifest( ManifestInfo->GetManifest(), DestinationPath / ManifestName ) )
+	const FString ManifestPath = DestinationPath / ManifestName;
+	if( !WriteManifestToFile( ManifestInfo->GetManifest(), ManifestPath ) )
 	{
-		UE_LOG( LogGenerateManifestCommandlet, Error,TEXT("Failed to write manifest to %s."), *DestinationPath );				
+		UE_LOG( LogGenerateManifestCommandlet, Error,TEXT("Failed to write manifest to %s."), *ManifestPath );				
 		return -1;
 	}
 	return 0;
 }
 
-bool UGenerateGatherManifestCommandlet::WriteManifest( const TSharedPtr<FInternationalizationManifest>& InManifest, const FString& OutputFilePath )
+bool UGenerateGatherManifestCommandlet::WriteManifestToFile( const TSharedPtr<FInternationalizationManifest>& InManifest, const FString& OutputFilePath )
 {
+	UE_LOG(LogGenerateManifestCommandlet, Log, TEXT("Writing archive to %s."), *OutputFilePath);
+
 	// We can not continue if the provided manifest is not valid
 	if( !InManifest.IsValid() )
 	{
