Bp Editor crash, with intermediate graphs shown.

diff --git a/Engine/Source/Editor/UnrealEd/Private/Kismet2/BlueprintEditorUtils.cpp b/Engine/Source/Editor/UnrealEd/Private/Kismet2/BlueprintEditorUtils.cpp
--- a/Engine/Source/Editor/UnrealEd/Private/Kismet2/BlueprintEditorUtils.cpp
+++ b/Engine/Source/Editor/UnrealEd/Private/Kismet2/BlueprintEditorUtils.cpp
@@ -2650,7 +2650,8 @@ bool FBlueprintEditorUtils::DoesSupportLocalVariables(UEdGraph const* InGraph)
 	if(InGraph)
 	{
 		UBlueprint* Blueprint = FBlueprintEditorUtils::FindBlueprintForGraph(InGraph);
-		return Blueprint->BlueprintType != BPTYPE_Interface
+		return Blueprint
+			&& Blueprint->BlueprintType != BPTYPE_Interface
 			&& InGraph->GetSchema()->GetGraphType(InGraph) == EGraphType::GT_Function
 			&& !InGraph->IsA(UAnimationTransitionGraph::StaticClass());
 	}
