Mesh proxy supports Metallic, Roughness, Specular material properties
Landscape export LOD uses MaxLOD from landscape actor by default
Sub-levels LOD generation will be still accessible when sub-level has landscape actor but user has no Simplygon available, in this case all geometry except landscape will be stripped

diff --git a/Engine/Source/Developer/MeshUtilities/Private/MeshUtilities.cpp b/Engine/Source/Developer/MeshUtilities/Private/MeshUtilities.cpp
--- a/Engine/Source/Developer/MeshUtilities/Private/MeshUtilities.cpp
+++ b/Engine/Source/Developer/MeshUtilities/Private/MeshUtilities.cpp
@@ -3668,28 +3668,19 @@ void FMeshUtilities::CreateProxyMesh(
 	const FString AssetBaseName = FPackageName::GetShortName(ProxyBasePackageName);
 	const FString AssetBasePath = InOuter ? TEXT("") : FPackageName::GetLongPackagePath(ProxyBasePackageName) + TEXT("/");
 	
-	TArray<ALandscapeProxy*>		LandscapesToMerge;
-	TArray<UStaticMeshComponent*>	ComponentsToMerge;
+	TArray<UStaticMeshComponent*> ComponentsToMerge;
 	
-	//Collect components of the corresponding actor
+	// Collect components to merge
 	for (AActor* Actor : SourceActors)
 	{
-		ALandscapeProxy* LandscapeActor = Cast<ALandscapeProxy>(Actor);
-		if (LandscapeActor)
-		{
-			LandscapesToMerge.Add(LandscapeActor);
-		}
-		else
-		{
-			TInlineComponentArray<UStaticMeshComponent*> Components;
-			Actor->GetComponents<UStaticMeshComponent>(Components);
-			// TODO: support instanced static meshes
-			Components.RemoveAll([](UStaticMeshComponent* Val){ return Val->IsA(UInstancedStaticMeshComponent::StaticClass()); });
-			// TODO: support non-opaque materials
-			Components.RemoveAll(&NonOpaqueMaterialPredicate);
-			//
-			ComponentsToMerge.Append(Components);
-		}
+		TInlineComponentArray<UStaticMeshComponent*> Components;
+		Actor->GetComponents<UStaticMeshComponent>(Components);
+		// TODO: support instanced static meshes
+		Components.RemoveAll([](UStaticMeshComponent* Val){ return Val->IsA(UInstancedStaticMeshComponent::StaticClass()); });
+		// TODO: support non-opaque materials
+		Components.RemoveAll(&NonOpaqueMaterialPredicate);
+		//
+		ComponentsToMerge.Append(Components);
 	}
 	
 	// Convert collected static mesh components and landscapes into raw meshes and flatten materials
@@ -3698,8 +3689,8 @@ void FMeshUtilities::CreateProxyMesh(
 	TMap<int32, TArray<int32>>						MaterialMap;
 	FBox											ProxyBounds(0);
 	
-	RawMeshes.Empty(ComponentsToMerge.Num() + LandscapesToMerge.Num());
-	UniqueMaterials.Empty(ComponentsToMerge.Num() + LandscapesToMerge.Num());
+	RawMeshes.Empty(ComponentsToMerge.Num());
+	UniqueMaterials.Empty(ComponentsToMerge.Num());
 	
 	// Convert static mesh components
 	TArray<UMaterialInterface*> StaticMeshMaterials;
@@ -3720,58 +3711,29 @@ void FMeshUtilities::CreateProxyMesh(
 		}
 	}
 
-	// Convert materials into flatten materials
-	for (UMaterialInterface* Material : StaticMeshMaterials)
+	if (RawMeshes.Num() == 0)
 	{
-		UniqueMaterials.Add(MaterialExportUtils::FFlattenMaterial());
-		MaterialExportUtils::ExportMaterial(Material, UniqueMaterials.Last());
+		return;
 	}
-	
-	// Convert landscapes
-	for (ALandscapeProxy* Landscape : LandscapesToMerge)
+
+	// Convert materials into flatten materials
 	{
-		TArray<int32> RawMeshMaterialMap;
-		int32 RawMeshId = RawMeshes.Add(FRawMesh());
+		MaterialExportUtils::FFlattenMaterial FlattenMaterial;
+		FIntPoint TargetTextureSize = FIntPoint(InProxySettings.TextureWidth, InProxySettings.TextureHeight);
 
-		if (Landscape->ExportToRawMesh(INDEX_NONE, RawMeshes[RawMeshId]))
-		{
-			// Landscape has one unique material
-			int32 MatIdx = UniqueMaterials.Add(MaterialExportUtils::FFlattenMaterial());
-			RawMeshMaterialMap.Add(MatIdx);
-			MaterialMap.Add(RawMeshId, RawMeshMaterialMap);
-			// This is texture resolution for a landscape, probably need to be calculated using landscape size
-			UniqueMaterials.Last().DiffuseSize = FIntPoint(1024, 1024);
-			// FIXME: Landscape material exporter currently renders world space normal map, so it can't be merged with other meshes normal maps
-			UniqueMaterials.Last().NormalSize = FIntPoint::ZeroValue;
-
-			// Use only landscape primitives for texture flattening
-			TSet<FPrimitiveComponentId> PrimitivesToHide;
-			for (TObjectIterator<UPrimitiveComponent> It; It; ++It)
-			{
-				UPrimitiveComponent* PrimitiveComp = *It;
-				const bool bTargetPrim = PrimitiveComp->GetOuter() == Landscape;
-				if (!bTargetPrim && PrimitiveComp->IsRegistered() && PrimitiveComp->SceneProxy)
-				{
-					PrimitivesToHide.Add(PrimitiveComp->SceneProxy->GetPrimitiveComponentId());
-				}
-			}
-			
-			MaterialExportUtils::ExportMaterial(Landscape, PrimitivesToHide, UniqueMaterials.Last());
-			
-			//Store the bounds for each component
-			ProxyBounds+= Landscape->GetComponentsBoundingBox(true);
-		}
-		else
+		FlattenMaterial.DiffuseSize = TargetTextureSize;
+		FlattenMaterial.NormalSize = InProxySettings.bExportNormalMap ?	TargetTextureSize : FIntPoint::ZeroValue;
+		FlattenMaterial.MetallicSize = InProxySettings.bExportMetallicMap ? TargetTextureSize : FIntPoint::ZeroValue;
+		FlattenMaterial.RoughnessSize = InProxySettings.bExportRoughnessMap ? TargetTextureSize : FIntPoint::ZeroValue;
+		FlattenMaterial.SpecularSize = InProxySettings.bExportSpecularMap ? TargetTextureSize : FIntPoint::ZeroValue;
+
+		for (UMaterialInterface* Material : StaticMeshMaterials)
 		{
-			RawMeshes.RemoveAt(RawMeshId);
+			UniqueMaterials.Add(FlattenMaterial);
+			MaterialExportUtils::ExportMaterial(Material, UniqueMaterials.Last());
 		}
 	}
-
-	if (RawMeshes.Num() == 0)
-	{
-		return;
-	}
-	
+		
 	//For each raw mesh, re-map the material indices according to the MaterialMap
 	for (int32 RawMeshIndex = 0; RawMeshIndex < RawMeshes.Num(); ++RawMeshIndex)
 	{
diff --git a/Engine/Source/Developer/SimplygonMeshReduction/Private/SimplygonMeshReduction.cpp b/Engine/Source/Developer/SimplygonMeshReduction/Private/SimplygonMeshReduction.cpp
--- a/Engine/Source/Developer/SimplygonMeshReduction/Private/SimplygonMeshReduction.cpp
+++ b/Engine/Source/Developer/SimplygonMeshReduction/Private/SimplygonMeshReduction.cpp
@@ -423,10 +423,16 @@ public:
 	{
 		bool bCastMaterials;
 		bool bCastNormals;
+		bool bCastMetallic;
+		bool bCastRoughness;
+		bool bCastSpecular;
 
 		FMaterialCastingProperties()
-			: bCastMaterials(false),
-			  bCastNormals(false)			
+			: bCastMaterials(false)
+			, bCastNormals(false)
+			, bCastMetallic(false)
+			, bCastRoughness(false)
+			, bCastSpecular(false)
 		{
 		}
 	};
@@ -543,7 +549,7 @@ public:
 			}
 
 			
-			//Cast normal texture data
+			// Normal texture data
 			if(CastProperties.bCastNormals)
 			{
 				//Create Image data where the normal data is stored
@@ -565,13 +571,75 @@ public:
 				OutputMaterialLOD->SetLayeredTextureLevel(SimplygonSDK::SG_MATERIAL_CHANNEL_NORMALS, 0, 0);
 			}
 
+			// Metallic texture data
+			if(CastProperties.bCastMetallic)
+			{
+				// Create Image data where the metallic data is stored
+				SimplygonSDK::spImageData OutputMetallicData = SDK->CreateImageData();
+				// Cast the data using a color caster
+				SimplygonSDK::spColorCaster cast = SDK->CreateColorCaster();
+				cast->SetColorType( SimplygonSDK::SG_MATERIAL_CHANNEL_METALLIC );
+				cast->SetSourceMaterials( OriginalMaterials );
+				cast->SetMappingImage( MappingImage );
+				cast->SetOutputChannels( 3 );
+				cast->SetOutputChannelBitDepth( 8 );
+				cast->SetsRGB(false);
+				cast->SetDilation( 10 );
+				cast->SetOutputImage(OutputMetallicData);
+				cast->CastMaterials(); // Fetch!
+
+				OutputMaterialLOD->SetLayeredTextureImage(SimplygonSDK::SG_MATERIAL_CHANNEL_METALLIC, 0, OutputMetallicData);
+				OutputMaterialLOD->SetLayeredTextureLevel(SimplygonSDK::SG_MATERIAL_CHANNEL_METALLIC, 0, 0);
+			}
+
+			// Roughness texture data
+			if(CastProperties.bCastRoughness)
+			{
+				// Create Image data where the metallic data is stored
+				SimplygonSDK::spImageData OutputRoughnessData = SDK->CreateImageData();
+				// Cast the data using a color caster
+				SimplygonSDK::spColorCaster cast = SDK->CreateColorCaster();
+				cast->SetColorType( SimplygonSDK::SG_MATERIAL_CHANNEL_ROUGHNESS );
+				cast->SetSourceMaterials( OriginalMaterials );
+				cast->SetMappingImage( MappingImage ); 
+				cast->SetOutputChannels( 3 );
+				cast->SetOutputChannelBitDepth( 8 );
+				cast->SetsRGB(false);
+				cast->SetDilation( 10 ); 
+				cast->SetOutputImage(OutputRoughnessData);
+				cast->CastMaterials(); // Fetch!
+
+				OutputMaterialLOD->SetLayeredTextureImage(SimplygonSDK::SG_MATERIAL_CHANNEL_ROUGHNESS, 0, OutputRoughnessData);
+				OutputMaterialLOD->SetLayeredTextureLevel(SimplygonSDK::SG_MATERIAL_CHANNEL_ROUGHNESS, 0, 0);
+			}
+
+			// Specular texture data
+			if(CastProperties.bCastSpecular)
+			{
+				// Create Image data where the metallic data is stored
+				SimplygonSDK::spImageData OutputSpecularData = SDK->CreateImageData();
+				// Cast the data using a color caster
+				SimplygonSDK::spColorCaster cast = SDK->CreateColorCaster();
+				cast->SetColorType( SimplygonSDK::SG_MATERIAL_CHANNEL_SPECULAR );
+				cast->SetSourceMaterials( OriginalMaterials );
+				cast->SetMappingImage( MappingImage ); 
+				cast->SetOutputChannels( 3 );
+				cast->SetOutputChannelBitDepth( 8 );
+				cast->SetsRGB(false);
+				cast->SetDilation( 10 ); 
+				cast->SetOutputImage(OutputSpecularData);
+				cast->CastMaterials(); // Fetch!
+
+				OutputMaterialLOD->SetLayeredTextureImage(SimplygonSDK::SG_MATERIAL_CHANNEL_SPECULAR, 0, OutputSpecularData);
+				OutputMaterialLOD->SetLayeredTextureLevel(SimplygonSDK::SG_MATERIAL_CHANNEL_SPECULAR, 0, 0);
+			}
+						
 			//Create a new material table for the new materials
 			SimplygonSDK::spMaterialTable OutputTable = SDK->CreateMaterialTable();
 			OutputTable->AddMaterial(OutputMaterialLOD);
 
 			//Convert the simplygon material to unreal materials
 			CreateFlattenMaterialFromSGMaterial(OutputTable, OutMaterial);
-
 		}
 
 		//Collect the proxy mesh
@@ -1451,7 +1519,7 @@ private:
 	/**
 	 * ProxyLOD Related Methods
 	 */
-	void SetMaterialChannel(
+	void SetMaterialChannelData(
 		const TArray<FColor>& InSamples, 
 		FIntPoint InTextureSize,
 		SimplygonSDK::spMaterial& InSGMaterial, 
@@ -1483,9 +1551,10 @@ private:
 			InSGMaterial->SetLayeredTextureImage(SGMaterialChannelName, 0, ImageData);
 			InSGMaterial->SetLayeredTextureLevel(SGMaterialChannelName, 0, 0);
 		}
-		else
+		else if (NumTexels == 1)
 		{
 			// handle uniform value
+			InSGMaterial->SetColorRGB(SGMaterialChannelName, InSamples[0].R, InSamples[0].G, InSamples[0].B);
 		}
 	}
 
@@ -1518,78 +1587,100 @@ private:
 
 			if(!SGMaterial->HasUserChannel(SimplygonSDK::SG_MATERIAL_CHANNEL_ROUGHNESS))
 				SGMaterial->AddUserChannel(SimplygonSDK::SG_MATERIAL_CHANNEL_ROUGHNESS);
+
+			if(!SGMaterial->HasUserChannel(SimplygonSDK::SG_MATERIAL_CHANNEL_SPECULAR))
+				SGMaterial->AddUserChannel(SimplygonSDK::SG_MATERIAL_CHANNEL_SPECULAR);
 					
 			SGMaterial->SetName(TCHAR_TO_ANSI(*FString::Printf(TEXT("Material%d"), MaterialIndex)));
 			
-			// Does current material have BaseColor?
+			// BaseColor
 			if (FlattenMaterial.DiffuseSamples.Num())
 			{
-				SetMaterialChannel(FlattenMaterial.DiffuseSamples, FlattenMaterial.DiffuseSize, SGMaterial, SimplygonSDK::SG_MATERIAL_CHANNEL_DIFFUSE);
+				SetMaterialChannelData(FlattenMaterial.DiffuseSamples, FlattenMaterial.DiffuseSize, SGMaterial, SimplygonSDK::SG_MATERIAL_CHANNEL_DIFFUSE);
 			}
 
-			//Does current material have a normalmap?
+			// Normal
 			if (FlattenMaterial.NormalSamples.Num())
 			{
 				OutCastProperties.bCastNormals = true;
-				SetMaterialChannel(FlattenMaterial.NormalSamples, FlattenMaterial.NormalSize, SGMaterial, SimplygonSDK::SG_MATERIAL_CHANNEL_NORMALS);
+				SetMaterialChannelData(FlattenMaterial.NormalSamples, FlattenMaterial.NormalSize, SGMaterial, SimplygonSDK::SG_MATERIAL_CHANNEL_NORMALS);
+			}
+
+			// Metallic
+			if (FlattenMaterial.MetallicSamples.Num())
+			{
+				OutCastProperties.bCastMetallic = true;
+				SetMaterialChannelData(FlattenMaterial.MetallicSamples, FlattenMaterial.MetallicSize, SGMaterial, SimplygonSDK::SG_MATERIAL_CHANNEL_METALLIC);
+			}
+
+			// Roughness
+			if (FlattenMaterial.RoughnessSamples.Num())
+			{
+				OutCastProperties.bCastRoughness = true;
+				SetMaterialChannelData(FlattenMaterial.RoughnessSamples, FlattenMaterial.RoughnessSize, SGMaterial, SimplygonSDK::SG_MATERIAL_CHANNEL_ROUGHNESS);
 			}
 
+			// Specular
+			if (FlattenMaterial.SpecularSamples.Num())
+			{
+				OutCastProperties.bCastSpecular = true;
+				SetMaterialChannelData(FlattenMaterial.SpecularSamples, FlattenMaterial.SpecularSize, SGMaterial, SimplygonSDK::SG_MATERIAL_CHANNEL_SPECULAR);
+			}
+			
  			OutSGMaterialTable->AddMaterial(SGMaterial);
 		}
 
 		return true;
 	}
 
-	void CreateFlattenMaterialFromSGMaterial(
-		SimplygonSDK::spMaterialTable& InSGMaterialTable, 
-		MaterialExportUtils::FFlattenMaterial& OutMaterial)
+	void GetMaterialChannelData(const SimplygonSDK::spMaterial& InSGMaterial, const char* SGMaterialChannelName, TArray<FColor>& OutSamples, FIntPoint& OutTextureSize)
 	{
-		//Get the diffuse data from SGMaterial
-		SimplygonSDK::spImageData SGDiffuseData = InSGMaterialTable->GetMaterial(0)->GetLayeredTextureImage(SimplygonSDK::SG_MATERIAL_CHANNEL_DIFFUSE, 0);
+		SimplygonSDK::spImageData SGChannelData = InSGMaterial->GetLayeredTextureImage(SGMaterialChannelName, 0);
+		if (SGChannelData)
 		{
-			SimplygonSDK::spUnsignedCharArray DiffuseImageColors = SimplygonSDK::SafeCast<SimplygonSDK::IUnsignedCharArray>( SGDiffuseData->GetColors() );
+			SimplygonSDK::spUnsignedCharArray ImageColors = SimplygonSDK::SafeCast<SimplygonSDK::IUnsignedCharArray>(SGChannelData->GetColors());
 			
-			OutMaterial.DiffuseSize.X = SGDiffuseData->GetXSize();
-			OutMaterial.DiffuseSize.Y = SGDiffuseData->GetYSize();
-			int32 TexelsCount = OutMaterial.DiffuseSize.X*OutMaterial.DiffuseSize.Y;
-		
-			//Fill diffuse texture with color data collected from simplygon
-			OutMaterial.DiffuseSamples.SetNumUninitialized(TexelsCount);
-			for (int32 TexelIndex = 0; TexelIndex < TexelsCount; ++TexelIndex)
-			{
-				uint8 DiffuseData[4];
-				DiffuseImageColors->GetTuple(TexelIndex, (unsigned char*)&DiffuseData);
+			OutTextureSize.X = SGChannelData->GetXSize();
+			OutTextureSize.Y = SGChannelData->GetYSize();
 			
-				OutMaterial.DiffuseSamples[TexelIndex].R = DiffuseData[0];
-				OutMaterial.DiffuseSamples[TexelIndex].G = DiffuseData[1];
-				OutMaterial.DiffuseSamples[TexelIndex].B = DiffuseData[2];
-				OutMaterial.DiffuseSamples[TexelIndex].A = DiffuseData[3];
-			}
-		}
+			int32 TexelsCount = OutTextureSize.X*OutTextureSize.Y;
+			OutSamples.Empty(TexelsCount);
+			OutSamples.AddUninitialized(TexelsCount);
 
-		//Get the normal data from SGMaterial
-		SimplygonSDK::spImageData SGNormalData = InSGMaterialTable->GetMaterial(0)->GetLayeredTextureImage(SimplygonSDK::SG_MATERIAL_CHANNEL_NORMALS, 0);
-		if (SGNormalData)
-		{
-			SimplygonSDK::spUnsignedCharArray NormalImageColors = SimplygonSDK::SafeCast<SimplygonSDK::IUnsignedCharArray>( SGNormalData->GetColors() );
-			OutMaterial.NormalSize.X = SGNormalData->GetXSize();
-			OutMaterial.NormalSize.Y = SGNormalData->GetYSize();
-			int32 TexelsCount = OutMaterial.NormalSize.X*OutMaterial.NormalSize.Y;
-			
-			//Fill normal texture with color data collected from simplygon
-			OutMaterial.NormalSamples.SetNumUninitialized(TexelsCount);
 			for (int32 TexelIndex = 0; TexelIndex < TexelsCount; ++TexelIndex)
 			{
-				uint8 NormalData[3];
-				NormalImageColors->GetTuple(TexelIndex, (unsigned char*)&NormalData);
+				uint8 ColorData[4];
+				ImageColors->GetTuple(TexelIndex, (unsigned char*)&ColorData);
 			
-				OutMaterial.NormalSamples[TexelIndex].R = NormalData[0];
-				OutMaterial.NormalSamples[TexelIndex].G = NormalData[1];
-				OutMaterial.NormalSamples[TexelIndex].B = NormalData[2];
-				OutMaterial.NormalSamples[TexelIndex].A = FColor::White.A;
+				OutSamples[TexelIndex].R = ColorData[0];
+				OutSamples[TexelIndex].G = ColorData[1];
+				OutSamples[TexelIndex].B = ColorData[2];
+				OutSamples[TexelIndex].A = (SGMaterialChannelName == SimplygonSDK::SG_MATERIAL_CHANNEL_DIFFUSE ? ColorData[3] : FColor::White.A);
 			}
 		}
 	}
+
+	void CreateFlattenMaterialFromSGMaterial(
+		SimplygonSDK::spMaterialTable& InSGMaterialTable, 
+		MaterialExportUtils::FFlattenMaterial& OutMaterial)
+	{
+		SimplygonSDK::spMaterial SGMaterial = InSGMaterialTable->GetMaterial(0);
+		
+		// Diffuse
+		GetMaterialChannelData(SGMaterial, SimplygonSDK::SG_MATERIAL_CHANNEL_DIFFUSE, OutMaterial.DiffuseSamples, OutMaterial.DiffuseSize);
+
+		// Normal
+		GetMaterialChannelData(SGMaterial, SimplygonSDK::SG_MATERIAL_CHANNEL_NORMALS, OutMaterial.NormalSamples, OutMaterial.NormalSize);
+	
+		// Metallic
+		GetMaterialChannelData(SGMaterial, SimplygonSDK::SG_MATERIAL_CHANNEL_METALLIC, OutMaterial.MetallicSamples, OutMaterial.MetallicSize);
+
+		// Roughness
+		GetMaterialChannelData(SGMaterial, SimplygonSDK::SG_MATERIAL_CHANNEL_ROUGHNESS, OutMaterial.RoughnessSamples, OutMaterial.RoughnessSize);
+
+		// Specular
+		GetMaterialChannelData(SGMaterial, SimplygonSDK::SG_MATERIAL_CHANNEL_SPECULAR, OutMaterial.SpecularSamples, OutMaterial.SpecularSize);
+	}
 };
 
 TScopedPointer<FSimplygonMeshReduction> GSimplygonMeshReduction;
diff --git a/Engine/Source/Editor/UnrealEd/Private/Dialogs/SMeshProxyDialog.cpp b/Engine/Source/Editor/UnrealEd/Private/Dialogs/SMeshProxyDialog.cpp
--- a/Engine/Source/Editor/UnrealEd/Private/Dialogs/SMeshProxyDialog.cpp
+++ b/Engine/Source/Editor/UnrealEd/Private/Dialogs/SMeshProxyDialog.cpp
@@ -70,6 +70,16 @@ protected:
 	/** TextureResolution accessors */
 	void SetTextureResolution(TSharedPtr<FString> NewSelection, ESelectInfo::Type SelectInfo);
 
+	/** Export material properties acessors **/
+	ECheckBoxState GetExportNormalMap() const;
+	void SetExportNormalMap(ECheckBoxState NewValue);
+	ECheckBoxState GetExportMetallicMap() const;
+	void SetExportMetallicMap(ECheckBoxState NewValue);
+	ECheckBoxState GetExportRoughnessMap() const;
+	void SetExportRoughnessMap(ECheckBoxState NewValue);
+	ECheckBoxState GetExportSpecularMap() const;
+	void SetExportSpecularMap(ECheckBoxState NewValue);
+	
 	/** Delegates for Merge button state */
 	bool GetMergeButtonEnabledState() const;
 	
@@ -179,6 +189,7 @@ void SMeshProxyDialog::Construct(const FArguments& InArgs)
 	CreateLayout();
 }
 
+BEGIN_SLATE_FUNCTION_BUILD_OPTIMIZATION
 void SMeshProxyDialog::CreateLayout()
 {
 	this->ChildSlot
@@ -238,7 +249,30 @@ void SMeshProxyDialog::CreateLayout()
 					[
 						SNew(STextBlock).Text(LOCTEXT("TextureResolutionLabel", "Texture Resolution"))
 					]
-
+					+SVerticalBox::Slot()
+					.FillHeight(1.0f)
+					.VAlign(VAlign_Center)
+					[
+						SNew(STextBlock).Text(LOCTEXT("ExportNormalMapLabel", "Export Normal Map"))
+					]
+					+SVerticalBox::Slot()
+					.FillHeight(1.0f)
+					.VAlign(VAlign_Center)
+					[
+						SNew(STextBlock).Text(LOCTEXT("ExportMetallicMapLabel", "Export Metallic Map"))
+					]
+					+SVerticalBox::Slot()
+					.FillHeight(1.0f)
+					.VAlign(VAlign_Center)
+					[
+						SNew(STextBlock).Text(LOCTEXT("ExportRoughnessMapLabel", "Export Roughness Map"))
+					]
+					+SVerticalBox::Slot()
+					.FillHeight(1.0f)
+					.VAlign(VAlign_Center)
+					[
+						SNew(STextBlock).Text(LOCTEXT("ExportSpecularMapLabel", "Export Specular Map"))
+					]
 					+SVerticalBox::Slot()
 					.FillHeight(1.0f)
 					.VAlign(VAlign_Center)
@@ -379,6 +413,50 @@ void SMeshProxyDialog::CreateLayout()
 						.OnSelectionChanged(this, &SMeshProxyDialog::SetTextureResolution)
 					]
 
+					// Export Normal Map
+					+SVerticalBox::Slot()
+					.FillHeight(1.0f)
+					.VAlign(VAlign_Center)
+					[
+						SNew(SCheckBox)
+						.Type(ESlateCheckBoxType::CheckBox)
+						.IsChecked(this, &SMeshProxyDialog::GetExportNormalMap)
+						.OnCheckStateChanged(this, &SMeshProxyDialog::SetExportNormalMap)
+					]
+
+					// Export Metallic Map
+					+SVerticalBox::Slot()
+					.FillHeight(1.0f)
+					.VAlign(VAlign_Center)
+					[
+						SNew(SCheckBox)
+						.Type(ESlateCheckBoxType::CheckBox)
+						.IsChecked(this, &SMeshProxyDialog::GetExportMetallicMap)
+						.OnCheckStateChanged(this, &SMeshProxyDialog::SetExportMetallicMap)
+					]
+					
+					// Export Roughness Map
+					+SVerticalBox::Slot()
+					.FillHeight(1.0f)
+					.VAlign(VAlign_Center)
+					[
+						SNew(SCheckBox)
+						.Type(ESlateCheckBoxType::CheckBox)
+						.IsChecked(this, &SMeshProxyDialog::GetExportRoughnessMap)
+						.OnCheckStateChanged(this, &SMeshProxyDialog::SetExportRoughnessMap)
+					]
+
+					// Export Specular Map
+					+SVerticalBox::Slot()
+					.FillHeight(1.0f)
+					.VAlign(VAlign_Center)
+					[
+						SNew(SCheckBox)
+						.Type(ESlateCheckBoxType::CheckBox)
+						.IsChecked(this, &SMeshProxyDialog::GetExportSpecularMap)
+						.OnCheckStateChanged(this, &SMeshProxyDialog::SetExportSpecularMap)
+					]
+					
 					// Package Name
 					+SVerticalBox::Slot()
 					.FillHeight(1.0f)
@@ -441,6 +519,7 @@ void SMeshProxyDialog::CreateLayout()
 		]
 	];
 }
+END_SLATE_FUNCTION_BUILD_OPTIMIZATION
 
 void SMeshProxyDialog::Reset(const bool bRefresh)
 {
@@ -564,6 +643,45 @@ void SMeshProxyDialog::SetTextureResolution(TSharedPtr<FString> NewSelection, ES
 	ProxySettings.TextureHeight = Resolution;
 }
 
+ECheckBoxState SMeshProxyDialog::GetExportNormalMap() const
+{
+	return ProxySettings.bExportNormalMap ? ECheckBoxState::Checked :  ECheckBoxState::Unchecked;
+}
+
+void SMeshProxyDialog::SetExportNormalMap(ECheckBoxState NewValue)
+{
+	ProxySettings.bExportNormalMap = (NewValue == ECheckBoxState::Checked);
+}
+
+ECheckBoxState SMeshProxyDialog::GetExportMetallicMap() const
+{
+	return ProxySettings.bExportMetallicMap ? ECheckBoxState::Checked :  ECheckBoxState::Unchecked;
+}
+
+void SMeshProxyDialog::SetExportMetallicMap(ECheckBoxState NewValue)
+{
+	ProxySettings.bExportMetallicMap = (NewValue == ECheckBoxState::Checked);
+}
+
+ECheckBoxState SMeshProxyDialog::GetExportRoughnessMap() const
+{
+	return ProxySettings.bExportRoughnessMap ? ECheckBoxState::Checked :  ECheckBoxState::Unchecked;
+}
+
+void SMeshProxyDialog::SetExportRoughnessMap(ECheckBoxState NewValue)
+{
+	ProxySettings.bExportRoughnessMap = (NewValue == ECheckBoxState::Checked);
+}
+
+ECheckBoxState SMeshProxyDialog::GetExportSpecularMap() const
+{
+	return ProxySettings.bExportSpecularMap ? ECheckBoxState::Checked :  ECheckBoxState::Unchecked;
+}
+
+void SMeshProxyDialog::SetExportSpecularMap(ECheckBoxState NewValue)
+{
+	ProxySettings.bExportSpecularMap = (NewValue == ECheckBoxState::Checked);
+}
 
 bool SMeshProxyDialog::GetMergeButtonEnabledState() const
 {
@@ -609,19 +727,6 @@ FReply SMeshProxyDialog::OnMergeClicked()
 		GEditor->EndTransaction();
 		GWarn->EndSlowTask();
 
-		//Create a new world actor and place it at the same location as the original assets
-		if (Actors[0] != nullptr)
-		{
-			UStaticMesh* ProxyMesh = nullptr;
-			if (AssetsToSync.FindItemByClass(&ProxyMesh))
-			{
-				UWorld* World = Actors[0]->GetWorld();
-				AStaticMeshActor* WorldActor = Cast<AStaticMeshActor>(World->SpawnActor(AStaticMeshActor::StaticClass(), &ProxyLocation));
-				WorldActor->GetStaticMeshComponent()->StaticMesh = ProxyMesh;
-				WorldActor->SetActorLabel(ProxyMesh->GetName());
-			}
-		}
-
 		//Update the asset registry that a new static mash and material has been created
 		if (AssetsToSync.Num())
 		{
diff --git a/Engine/Source/Editor/UnrealEd/Private/EditorExporters.cpp b/Engine/Source/Editor/UnrealEd/Private/EditorExporters.cpp
--- a/Engine/Source/Editor/UnrealEd/Private/EditorExporters.cpp
+++ b/Engine/Source/Editor/UnrealEd/Private/EditorExporters.cpp
@@ -742,9 +742,11 @@ public:
 		// Override with a special usage so we won't re-use the shader map used by the material for rendering
 		switch (InPropertyToCompile)
 		{
-		case MP_DiffuseColor: ResourceId.Usage = EMaterialShaderMapUsage::MaterialExportDiffuseColor; break;
-		case MP_SpecularColor: ResourceId.Usage = EMaterialShaderMapUsage::MaterialExportSpecularColor; break;
+		case MP_BaseColor: ResourceId.Usage = EMaterialShaderMapUsage::MaterialExportBaseColor; break;
+		case MP_Specular: ResourceId.Usage = EMaterialShaderMapUsage::MaterialExportSpecular; break;
 		case MP_Normal: ResourceId.Usage = EMaterialShaderMapUsage::MaterialExportNormal; break;
+		case MP_Metallic: ResourceId.Usage = EMaterialShaderMapUsage::MaterialExportMetallic; break;
+		case MP_Roughness: ResourceId.Usage = EMaterialShaderMapUsage::MaterialExportRoughness; break;
 		};
 		
 		CacheShaders(ResourceId, GMaxRHIShaderPlatform, true);
@@ -818,17 +820,7 @@ public:
 	/** helper for CompilePropertyAndSetMaterialProperty() */
 	int32 CompilePropertyAndSetMaterialPropertyWithoutCast(EMaterialProperty Property, FMaterialCompiler* Compiler) const
 	{
-		// MAKE SURE THIS MATCHES THE CHART IN WillFillData
-		// 						  RETURNED VALUES (F16 'textures')
-		// 	BLEND MODE  | DIFFUSE     | SPECULAR     | EMISSIVE    | NORMAL    | TRANSMISSIVE              |
-		// 	------------+-------------+--------------+-------------+-----------+---------------------------|
-		// 	Opaque      | Diffuse     | Spec,SpecPwr | Emissive    | Normal    | 0 (EMPTY)                 |
-		// 	Masked      | Diffuse     | Spec,SpecPwr | Emissive    | Normal    | Opacity Mask              |
-		// 	Translucent | 0 (EMPTY)   | 0 (EMPTY)    | Emissive    | 0 (EMPTY) | (Emsv | Diffuse)*Opacity  |
-		// 	Additive    | 0 (EMPTY)   | 0 (EMPTY)    | Emissive    | 0 (EMPTY) | (Emsv | Diffuse)*Opacity  |
-		// 	Modulative  | 0 (EMPTY)   | 0 (EMPTY)    | Emissive    | 0 (EMPTY) | Emsv | Diffuse            |
-		// 	------------+-------------+--------------+-------------+-----------+---------------------------|
-		if( Property == MP_EmissiveColor )
+		if (Property == MP_EmissiveColor)
 		{
 			UMaterial* ProxyMaterial = MaterialInterface->GetMaterial();
 			EBlendMode BlendMode = MaterialInterface->GetBlendMode();
@@ -839,21 +831,20 @@ public:
 			case MP_EmissiveColor:
 				// Emissive is ALWAYS returned...
 				return Compiler->ForceCast(MaterialInterface->CompileProperty(Compiler,MP_EmissiveColor),MCT_Float3,true,true);
-			case MP_DiffuseColor:
+			case MP_BaseColor:
 				// Only return for Opaque and Masked...
 				if (BlendMode == BLEND_Opaque || BlendMode == BLEND_Masked)
 				{
-					EMaterialProperty DiffuseInput = MP_BaseColor;
-					return Compiler->ForceCast(MaterialInterface->CompileProperty(Compiler, DiffuseInput),MCT_Float3,true,true);
+					return Compiler->ForceCast(MaterialInterface->CompileProperty(Compiler, MP_BaseColor),MCT_Float3,true,true);
 				}
 				break;
-			case MP_SpecularColor: 
+			case MP_Specular: 
+			case MP_Roughness:
+			case MP_Metallic:
 				// Only return for Opaque and Masked...
 				if (BlendMode == BLEND_Opaque || BlendMode == BLEND_Masked)
 				{
-					return Compiler->AppendVector(
-						Compiler->ForceCast(MaterialInterface->CompileProperty(Compiler, MP_SpecularColor),MCT_Float3,true,true), 
-						Compiler->ForceCast(MaterialInterface->CompileProperty(Compiler,MP_Roughness),MCT_Float1));
+					return Compiler->ForceCast(MaterialInterface->CompileProperty(Compiler, PropertyToCompile),MCT_Float1,true,true);
 				}
 				break;
 			case MP_Normal:
@@ -867,41 +858,6 @@ public:
 						MCT_Float3, true, true );
 				}
 				break;
-			
-			case MP_Opacity:
-				if (BlendMode == BLEND_Masked)
-				{
-					return MaterialInterface->CompileProperty(Compiler, MP_OpacityMask);
-				}
-				else if (IsTranslucentBlendMode((EBlendMode)BlendMode) && ProxyMaterial->GetCastShadowAsMasked())
-				{
-					return MaterialInterface->CompileProperty(Compiler, MP_Opacity);
-				}
-				else if (BlendMode == BLEND_Modulate)
-				{
-					if (ShadingModel == MSM_Unlit)
-					{
-						return Compiler->ForceCast(MaterialInterface->CompileProperty(Compiler, MP_EmissiveColor),MCT_Float3,true,true);
-					}
-					else
-					{
-						return Compiler->ForceCast(MaterialInterface->CompileProperty(Compiler, MP_DiffuseColor),MCT_Float3,true,true);
-					}
-				}
-				else if ((BlendMode == BLEND_Translucent) || (BlendMode == BLEND_Additive))
-				{
-					int32 ColoredOpacity = -1;
-					if (ShadingModel == MSM_Unlit)
-					{
-						ColoredOpacity = Compiler->ForceCast(MaterialInterface->CompileProperty(Compiler, MP_EmissiveColor),MCT_Float3,true,true);
-					}
-					else
-					{
-						ColoredOpacity = Compiler->ForceCast(MaterialInterface->CompileProperty(Compiler, MP_DiffuseColor),MCT_Float3,true,true);
-					}
-					return Compiler->Lerp(Compiler->Constant3(1.0f, 1.0f, 1.0f), ColoredOpacity, Compiler->ForceCast( MaterialInterface->CompileProperty(Compiler,MP_Opacity), MCT_Float1));
-				}
-				break;
 			default:
 				return Compiler->Constant(1.0f);
 			}
@@ -1000,20 +956,20 @@ public:
 		case MP_EmissiveColor:
 			bConnected = InMaterial->EmissiveColor.Expression != NULL;
 			break;
-		case MP_DiffuseColor:
+		case MP_BaseColor:
 			bConnected = InMaterial->BaseColor.Expression != NULL;
 			break;
-		case MP_SpecularColor:
+		case MP_Specular:
 			bConnected = InMaterial->Specular.Expression != NULL;
 			break;
 		case MP_Normal:
 			bConnected = InMaterial->Normal.Expression != NULL;
 			break;
-		case MP_Opacity:
-			bConnected = InMaterial->Opacity.Expression != NULL;
+		case MP_Metallic:
+			bConnected = InMaterial->Metallic.Expression != NULL;
 			break;
-		case MP_OpacityMask:
-			bConnected = InMaterial->OpacityMask.Expression != NULL;
+		case MP_Roughness:
+			bConnected = InMaterial->Roughness.Expression != NULL;
 			break;
 		default:
 			break;
@@ -1053,14 +1009,16 @@ public:
 			// Emissive is ALWAYS returned...
 			bExpressionIsNULL = !IsMaterialInputConnected(Material, PropertyToCompile);
 			break;
-		case MP_DiffuseColor:
+		case MP_BaseColor:
 			// Only return for Opaque and Masked...
 			if (BlendMode == BLEND_Opaque || BlendMode == BLEND_Masked)
 			{
 				bExpressionIsNULL = !IsMaterialInputConnected(Material, PropertyToCompile);
 			}
 			break;
-		case MP_SpecularColor: 
+		case MP_Specular: 
+		case MP_Metallic:
+		case MP_Roughness:
 			// Only return for Opaque and Masked...
 			if (BlendMode == BLEND_Opaque || BlendMode == BLEND_Masked)
 			{
@@ -1076,40 +1034,6 @@ public:
 				OutUniformValue.B = 255;	// Default normal is (0,0,1)
 			}
 			break;
-		case MP_Opacity:
-			if (BlendMode == BLEND_Masked)
-			{
-				bExpressionIsNULL = !IsMaterialInputConnected(Material, MP_OpacityMask);
-				OutUniformValue.R = 255;
-				OutUniformValue.G = 255;
-				OutUniformValue.B = 255;
-				OutUniformValue.A = 255;
-			}
-			else
-			if ((BlendMode == BLEND_Modulate) ||
-				(BlendMode == BLEND_Translucent) || 
-				(BlendMode == BLEND_Additive))
-			{
-				bool bColorInputIsNULL = false;
-				if (ShadingModel == MSM_Unlit)
-				{
-					bColorInputIsNULL = !IsMaterialInputConnected(Material, MP_EmissiveColor);
-				}
-				else
-				{
-					bColorInputIsNULL = !IsMaterialInputConnected(Material, MP_DiffuseColor);
-				}
-				if (BlendMode == BLEND_Translucent
-					|| BlendMode == BLEND_Additive)
-				{
-					bExpressionIsNULL = bColorInputIsNULL && !IsMaterialInputConnected(Material, PropertyToCompile);
-				}
-				else
-				{
-					bExpressionIsNULL = bColorInputIsNULL;
-				}
-			}
-			break;
 		}
 
 		return bExpressionIsNULL;
@@ -1177,17 +1101,6 @@ public:
 
 	static bool WillFillData(EBlendMode InBlendMode, EMaterialProperty InMaterialProperty)
 	{
-		// MAKE SURE THIS MATCHES THE CHART IN CompileProperty
-		// 						  RETURNED VALUES (F16 'textures')
-		// 	BLEND MODE  | DIFFUSE     | SPECULAR     | EMISSIVE    | NORMAL    | TRANSMISSIVE              |
-		// 	------------+-------------+--------------+-------------+-----------+---------------------------|
-		// 	Opaque      | Diffuse     | Spec,SpecPwr | Emissive    | Normal    | 0 (EMPTY)                 |
-		// 	Masked      | Diffuse     | Spec,SpecPwr | Emissive    | Normal    | Opacity Mask              |
-		// 	Translucent | 0 (EMPTY)   | 0 (EMPTY)    | Emissive    | 0 (EMPTY) | (Emsv | Diffuse)*Opacity  |
-		// 	Additive    | 0 (EMPTY)   | 0 (EMPTY)    | Emissive    | 0 (EMPTY) | (Emsv | Diffuse)*Opacity  |
-		// 	Modulative  | 0 (EMPTY)   | 0 (EMPTY)    | Emissive    | 0 (EMPTY) | Emsv | Diffuse            |
-		// 	------------+-------------+--------------+-------------+-----------+---------------------------|
-
 		if (InMaterialProperty == MP_EmissiveColor)
 		{
 			return true;
@@ -1199,44 +1112,11 @@ public:
 			{
 				switch (InMaterialProperty)
 				{
-				case MP_DiffuseColor:	return true;
-				case MP_SpecularColor:	return true;
-				case MP_Normal:			return true;
-				case MP_Opacity:		return false;
-				}
-			}
-			break;
-		case BLEND_Masked:
-			{
-				switch (InMaterialProperty)
-				{
-				case MP_DiffuseColor:	return true;
-				case MP_SpecularColor:	return true;
+				case MP_BaseColor:		return true;
+				case MP_Specular:		return true;
 				case MP_Normal:			return true;
-				case MP_Opacity:		return true;
-				}
-			}
-			break;
-		case BLEND_Translucent:
-		case BLEND_Additive:
-			{
-				switch (InMaterialProperty)
-				{
-				case MP_DiffuseColor:	return false;
-				case MP_SpecularColor:	return false;
-				case MP_Normal:			return false;
-				case MP_Opacity:		return true;
-				}
-			}
-			break;
-		case BLEND_Modulate:
-			{
-				switch (InMaterialProperty)
-				{
-				case MP_DiffuseColor:	return false;
-				case MP_SpecularColor:	return false;
-				case MP_Normal:			return false;
-				case MP_Opacity:		return true;
+				case MP_Metallic:		return true;
+				case MP_Roughness:		return true;
 				}
 			}
 			break;
@@ -1692,13 +1572,13 @@ void ExportOBJs(FOutputDevice& FileAr, FStringOutputDevice* MemAr, FFeedbackCont
 
 			{
 				FString BMPFilename = FPaths::GetPath(MaterialLibFilename) / MaterialName + TEXT("_D.bmp");
-				ExportMaterialPropertyTexture(BMPFilename, Material, MP_DiffuseColor, RenderTarget, Canvas);
+				ExportMaterialPropertyTexture(BMPFilename, Material, MP_BaseColor, RenderTarget, Canvas);
 				MaterialLib->Logf(TEXT("\tmap_Kd %s\r\n"), *FPaths::GetCleanFilename(BMPFilename));
 			}
 
 			{
 				FString BMPFilename = FPaths::GetPath(MaterialLibFilename) / MaterialName + TEXT("_S.bmp");
-				ExportMaterialPropertyTexture(BMPFilename, Material, MP_SpecularColor, RenderTarget, Canvas);
+				ExportMaterialPropertyTexture(BMPFilename, Material, MP_Specular, RenderTarget, Canvas);
 				MaterialLib->Logf(TEXT("\tmap_Ks %s\r\n"), *FPaths::GetCleanFilename(BMPFilename));
 			}
 
@@ -2643,7 +2523,7 @@ namespace MaterialExportUtils
 			
 			//
 			OutFlattenMaterial.DiffuseSamples.Empty(OutFlattenMaterial.DiffuseSize.X * OutFlattenMaterial.DiffuseSize.Y);
-			bool bResult = ExportMaterialProperty(InMaterial, MP_DiffuseColor, RenderTargetDiffuse, OutFlattenMaterial.DiffuseSamples);
+			bool bResult = ExportMaterialProperty(InMaterial, MP_BaseColor, RenderTargetDiffuse, OutFlattenMaterial.DiffuseSamples);
 
 			// Uniform value
 			if (OutFlattenMaterial.DiffuseSamples.Num() == 1)
@@ -2693,6 +2573,102 @@ namespace MaterialExportUtils
 			}
 		}
 
+		// Render metallic property
+		if (OutFlattenMaterial.MetallicSize.X > 0 && 
+			OutFlattenMaterial.MetallicSize.Y > 0)
+		{
+			// Create temporary render target
+			auto RenderTargetMetallic = NewObject<UTextureRenderTarget2D>();
+			check(RenderTargetMetallic);
+			RenderTargetMetallic->AddToRoot();
+			RenderTargetMetallic->ClearColor = FLinearColor(0.0f, 0.0f, 0.0f, 0.0f);
+			RenderTargetMetallic->InitCustomFormat(
+				OutFlattenMaterial.MetallicSize.X,
+				OutFlattenMaterial.MetallicSize.Y, PF_FloatRGB, true);
+
+			//
+			OutFlattenMaterial.MetallicSamples.Empty(OutFlattenMaterial.MetallicSize.X * OutFlattenMaterial.MetallicSize.Y);
+			bool bResult = ExportMaterialProperty(InMaterial, MP_Metallic, RenderTargetMetallic, OutFlattenMaterial.MetallicSamples);
+
+			// Uniform value
+			if (OutFlattenMaterial.MetallicSamples.Num() == 1)
+			{
+				OutFlattenMaterial.MetallicSize = FIntPoint(1,1);
+			}
+		
+			RenderTargetMetallic->RemoveFromRoot();
+			RenderTargetMetallic = nullptr;
+
+			if (!bResult)
+			{
+				return false;
+			}
+		}
+
+		// Render roughness property
+		if (OutFlattenMaterial.RoughnessSize.X > 0 && 
+			OutFlattenMaterial.RoughnessSize.Y > 0)
+		{
+			// Create temporary render target
+			auto RenderTargetRoughness = NewObject<UTextureRenderTarget2D>();
+			check(RenderTargetRoughness);
+			RenderTargetRoughness->AddToRoot();
+			RenderTargetRoughness->ClearColor = FLinearColor(0.0f, 0.0f, 0.0f, 0.0f);
+			RenderTargetRoughness->InitCustomFormat(
+				OutFlattenMaterial.RoughnessSize.X,
+				OutFlattenMaterial.RoughnessSize.Y, PF_FloatRGB, true);
+
+			//
+			OutFlattenMaterial.RoughnessSamples.Empty(OutFlattenMaterial.RoughnessSize.X * OutFlattenMaterial.RoughnessSize.Y);
+			bool bResult = ExportMaterialProperty(InMaterial, MP_Roughness, RenderTargetRoughness, OutFlattenMaterial.RoughnessSamples);
+
+			// Uniform value
+			if (OutFlattenMaterial.RoughnessSamples.Num() == 1)
+			{
+				OutFlattenMaterial.RoughnessSize = FIntPoint(1,1);
+			}
+		
+			RenderTargetRoughness->RemoveFromRoot();
+			RenderTargetRoughness = nullptr;
+
+			if (!bResult)
+			{
+				return false;
+			}
+		}
+
+		// Render specular property
+		if (OutFlattenMaterial.SpecularSize.X > 0 && 
+			OutFlattenMaterial.SpecularSize.Y > 0)
+		{
+			// Create temporary render target
+			auto RenderTargetSpecular = NewObject<UTextureRenderTarget2D>();
+			check(RenderTargetSpecular);
+			RenderTargetSpecular->AddToRoot();
+			RenderTargetSpecular->ClearColor = FLinearColor(0.0f, 0.0f, 0.0f, 0.0f);
+			RenderTargetSpecular->InitCustomFormat(
+				OutFlattenMaterial.SpecularSize.X,
+				OutFlattenMaterial.SpecularSize.Y, PF_FloatRGB, true);
+
+			//
+			OutFlattenMaterial.SpecularSamples.Empty(OutFlattenMaterial.SpecularSize.X * OutFlattenMaterial.SpecularSize.Y);
+			bool bResult = ExportMaterialProperty(InMaterial, MP_Specular, RenderTargetSpecular, OutFlattenMaterial.SpecularSamples);
+
+			// Uniform value
+			if (OutFlattenMaterial.SpecularSamples.Num() == 1)
+			{
+				OutFlattenMaterial.SpecularSize = FIntPoint(1,1);
+			}
+		
+			RenderTargetSpecular->RemoveFromRoot();
+			RenderTargetSpecular = nullptr;
+
+			if (!bResult)
+			{
+				return false;
+			}
+		}
+		
 		OutFlattenMaterial.MaterialId = InMaterial->GetLightingGuid();
 		return true;
 	}
@@ -2945,12 +2921,12 @@ namespace MaterialExportUtils
 		{
 			const FString AssetName = TEXT("T_") + AssetBaseName + TEXT("_S");
 			const bool bSRGB = false;
-			UTexture2D* Texture = CreateTexture(Outer, AssetBasePath + AssetName, InFlattenMaterial.SpecularSize, InFlattenMaterial.SpecularSamples, TC_Default, TEXTUREGROUP_World, Flags, bSRGB);
+			UTexture2D* Texture = CreateTexture(Outer, AssetBasePath + AssetName, InFlattenMaterial.SpecularSize, InFlattenMaterial.SpecularSamples, TC_Grayscale, TEXTUREGROUP_World, Flags, bSRGB);
 			OutGeneratedAssets.Add(Texture);
 			
 			auto SpecularExpression = NewObject<UMaterialExpressionTextureSample>(Material);
 			SpecularExpression->Texture = Texture;
-			SpecularExpression->SamplerType = EMaterialSamplerType::SAMPLERTYPE_LinearColor;
+			SpecularExpression->SamplerType = EMaterialSamplerType::SAMPLERTYPE_LinearGrayscale;
 			SpecularExpression->MaterialExpressionEditorX = -400;
 			SpecularExpression->MaterialExpressionEditorY = MaterialNodeY;
 			Material->Expressions.Add(SpecularExpression);
diff --git a/Engine/Source/Editor/UnrealEd/Private/EditorGroup.cpp b/Engine/Source/Editor/UnrealEd/Private/EditorGroup.cpp
--- a/Engine/Source/Editor/UnrealEd/Private/EditorGroup.cpp
+++ b/Engine/Source/Editor/UnrealEd/Private/EditorGroup.cpp
@@ -227,7 +227,7 @@ void UUnrealEdEngine::edactMergeActors()
 			// Create window
 			TSharedRef<SWindow> MeshProxyDialog = SNew(SWindow);
 			MeshProxyDialog->SetTitle(NSLOCTEXT("SMeshProxyDialog", "MeshProxyDialogTitle", "Mesh Proxy"));
-			MeshProxyDialog->Resize(FVector2D(400.0f, 280.0f));
+			MeshProxyDialog->Resize(FVector2D(400.0f, 350.0f));
 
 			MeshProxyControls->AssignWindowContent(MeshProxyDialog);
 			MeshProxyControls->SetParentWindow(MeshProxyDialog);
diff --git a/Engine/Source/Editor/WorldBrowser/Private/Tiles/WorldTileCollectionModel.cpp b/Engine/Source/Editor/WorldBrowser/Private/Tiles/WorldTileCollectionModel.cpp
--- a/Engine/Source/Editor/WorldBrowser/Private/Tiles/WorldTileCollectionModel.cpp
+++ b/Engine/Source/Editor/WorldBrowser/Private/Tiles/WorldTileCollectionModel.cpp
@@ -1124,16 +1124,16 @@ void FWorldTileCollectionModel::PostWorldOriginOffset(UWorld* InWorld, FIntVecto
 	}
 }
 
-bool FWorldTileCollectionModel::HasLandscapeLevel(const FWorldTileModelList& InLevels) const
+bool FWorldTileCollectionModel::AreAnySelectedLevelsHaveLandscape() const
 {
-	for (auto It = InLevels.CreateConstIterator(); It; ++It)
+	for (auto LevelModel : SelectedLevelsList)
 	{
-		if ((*It)->IsLandscapeBased())
+		if (LevelModel->IsLoaded() && StaticCastSharedPtr<FWorldTileModel>(LevelModel)->IsLandscapeBased())
 		{
 			return true;
 		}
 	}
-	
+
 	return false;
 }
 
@@ -1829,21 +1829,14 @@ void FWorldTileCollectionModel::OnNewCurrentLevel()
 	Focus(CurrentLevelModel->GetLevelBounds(), FWorldTileCollectionModel::OriginAtCenter);
 }
 
-bool FWorldTileCollectionModel::HasGenerateLODLevelSupport() const
+bool FWorldTileCollectionModel::HasMeshProxySupport() const
 {
 	return bMeshProxyAvailable;
 }
 
 bool FWorldTileCollectionModel::GenerateLODLevels(FLevelModelList InLevelList, int32 TargetLODIndex)
 {
-	if (!HasGenerateLODLevelSupport())
-	{
-		return false;
-	}
-
 	IMeshUtilities& MeshUtilities = FModuleManager::Get().LoadModuleChecked<IMeshUtilities>("MeshUtilities");
-	IMeshMerging* MeshMerging = MeshUtilities.GetMeshMergingInterface();
-	check(MeshMerging);
 
 	// Select tiles that can be processed
 	TArray<TSharedPtr<FWorldTileModel>> TilesToProcess;
@@ -1872,6 +1865,32 @@ bool FWorldTileCollectionModel::GenerateLODLevels(FLevelModelList InLevelList, i
 	// Generate LOD maps for each tile
 	for (TSharedPtr<FWorldTileModel> TileModel : TilesToProcess)
 	{
+		TArray<AActor*>				Actors;
+		TArray<ALandscapeProxy*>	LandscapeActors;
+		// Separate landscape actors from all others
+		for (AActor* Actor : TileModel->GetLevelObject()->Actors)
+		{
+			if (Actor)
+			{
+				ALandscapeProxy* LandscapeProxy = Cast<ALandscapeProxy>(Actor);
+				if (LandscapeProxy)
+				{
+					LandscapeActors.Add(LandscapeProxy);
+				}
+				else 
+				{
+					Actors.Add(Actor);
+				}
+			}
+		}
+
+		// Check if we can simplify this level
+		IMeshMerging* MeshMerging = MeshUtilities.GetMeshMergingInterface();
+		if (MeshMerging == nullptr && LandscapeActors.Num() == 0)
+		{
+			continue;
+		}
+							
 		// We have to make original level visible, to correctly export it
 		const bool bVisibleLevel = TileModel->IsVisible();
 		if (!bVisibleLevel)
@@ -1918,42 +1937,31 @@ bool FWorldTileCollectionModel::GenerateLODLevels(FLevelModelList InLevelList, i
 		};
 
 		TArray<FAssetToSpawnInfo>	AssetsToSpawn;
-		TArray<AActor*>				Actors;
-		TArray<ALandscapeProxy*>	LandscapeActors;
 		TArray<UObject*>			GeneratedAssets;
 		
 		// Where generated assets will be stored
 		UPackage* AssetsOuter = SimplificationDetails.bCreatePackagePerAsset ? nullptr : LODPackage;
 		// In case we don't have outer generated assets should have same path as LOD level
 		const FString AssetsPath = AssetsOuter ? TEXT("") : FPackageName::GetLongPackagePath(LODLevelPackageName) + TEXT("/");
-		
-		// Separate landscape actors from all others
-		for (AActor* Actor : TileModel->GetLevelObject()->Actors)
-		{
-			if (Actor)
-			{
-				ALandscapeProxy* LandscapeProxy = Cast<ALandscapeProxy>(Actor);
-				if (LandscapeProxy)
-				{
-					LandscapeActors.Add(LandscapeProxy);
-				}
-				else 
-				{
-					Actors.Add(Actor);
-				}
-			}
-		}
 	
 		// Generate Proxy LOD mesh for all actors excluding landscapes
-		if (Actors.Num())
+		if (Actors.Num() && MeshMerging != nullptr)
 		{
 			GWarn->StatusUpdate(0, 10, LOCTEXT("GeneratingProxyMesh", "Generating Proxy Mesh"));
 
 			FMeshProxySettings ProxySettings;
 			ProxySettings.ScreenSize = ProxySettings.ScreenSize*(SimplificationDetails.DetailsPercentage/100.f);
+			ProxySettings.TextureWidth = 1024; // TODO: Expose texture size
+			ProxySettings.TextureHeight = 1024;
+			ProxySettings.bExportNormalMap = SimplificationDetails.bGenerateMeshNormalMap;
+			ProxySettings.bExportMetallicMap = SimplificationDetails.bGenerateMeshMetallicMap;
+			ProxySettings.bExportRoughnessMap = SimplificationDetails.bGenerateMeshRoughnessMap;
+			ProxySettings.bExportSpecularMap = SimplificationDetails.bGenerateMeshSpecularMap;
+
 			TArray<UObject*> OutAssets;
 			FVector OutProxyLocation;
-			FString ProxyPackageName = TEXT("PROXY_") + FPackageName::GetShortName(TileModel->TileDetails->PackageName);
+			FString ProxyPackageName = FString::Printf(TEXT("PROXY_%s_LOD%d"), *FPackageName::GetShortName(TileModel->TileDetails->PackageName), TargetLODIndex + 1);
+			
 			// Generate proxy mesh and proxy material assets
 			MeshUtilities.CreateProxyMesh(Actors, ProxySettings, AssetsOuter, AssetsPath + ProxyPackageName, OutAssets, OutProxyLocation);
 		
@@ -1980,8 +1988,14 @@ bool FWorldTileCollectionModel::GenerateLODLevels(FLevelModelList InLevelList, i
 			FRawMesh LandscapeRawMesh;
 			FFlattenMaterial LandscapeFlattenMaterial;
 			FVector LandscapeWorldLocation = Landscape->GetActorLocation();
+			
+			int32 LandscapeLOD = SimplificationDetails.LandscapeExportLOD;
+			if (!SimplificationDetails.bOverrideLandscapeExportLOD)
+			{
+				LandscapeLOD = Landscape->MaxLODLevel >= 0 ? Landscape->MaxLODLevel : FMath::CeilLogTwo(Landscape->SubsectionSizeQuads + 1) - 1;
+			}
 		
-			Landscape->ExportToRawMesh(SimplificationDetails.LandscapeExportLOD, LandscapeRawMesh);
+			Landscape->ExportToRawMesh(LandscapeLOD, LandscapeRawMesh);
 		
 			for (FVector& VertexPos : LandscapeRawMesh.VertexPositions)
 			{
@@ -2027,7 +2041,7 @@ bool FWorldTileCollectionModel::GenerateLODLevels(FLevelModelList InLevelList, i
 			{
 				Landscape->FlushGrassComponents(); // wipe this and let it fix itself later
 			}
-			FString LandscapeBaseAssetName = Landscape->GetName();
+			FString LandscapeBaseAssetName = FString::Printf(TEXT("%s_LOD%d"), *Landscape->GetName(), TargetLODIndex + 1);
 			// Construct landscape material
 			UMaterial* StaticLandscapeMaterial = MaterialExportUtils::CreateMaterial(
 				LandscapeFlattenMaterial, AssetsOuter, *(AssetsPath + LandscapeBaseAssetName), RF_Public|RF_Standalone, GeneratedAssets);
diff --git a/Engine/Source/Editor/WorldBrowser/Private/Tiles/WorldTileCollectionModel.h b/Engine/Source/Editor/WorldBrowser/Private/Tiles/WorldTileCollectionModel.h
--- a/Engine/Source/Editor/WorldBrowser/Private/Tiles/WorldTileCollectionModel.h
+++ b/Engine/Source/Editor/WorldBrowser/Private/Tiles/WorldTileCollectionModel.h
@@ -78,8 +78,8 @@ public:
 	/** Show a levels in the editor and place them to actual world position */
 	void UnshelveLevels(const FWorldTileModelList& InLevels);
 
-	/** Whether specified list of levels has at least one landscape level */
-	bool HasLandscapeLevel(const FWorldTileModelList& InLevels) const;
+	/** Whether any of the currently selected levels have landscape actor */
+	bool AreAnySelectedLevelsHaveLandscape() const;
 	
 	/** Creates a new empty level 
 	 *	@return LevelModel of a new empty level
@@ -234,9 +234,9 @@ private:
 	void ReimportTiledLandscape_Executed(FName TargetLayer);
 
 public:
-	/** Whether Editor has support for generating LOD levels */	
-	bool HasGenerateLODLevelSupport() const;
-	
+	/** Whether Editor has support for generating static mesh proxies */	
+	bool HasMeshProxySupport() const;
+
 	/** 
 	 * Generates simplified versions of a specified levels. Levels has to be loaded.
 	 * Currently all static meshes found inside one level will be merged into one proxy mesh using Simplygon ProxyLOD
diff --git a/Engine/Source/Editor/WorldBrowser/Private/Tiles/WorldTileDetailsCustomization.cpp b/Engine/Source/Editor/WorldBrowser/Private/Tiles/WorldTileDetailsCustomization.cpp
--- a/Engine/Source/Editor/WorldBrowser/Private/Tiles/WorldTileDetailsCustomization.cpp
+++ b/Engine/Source/Editor/WorldBrowser/Private/Tiles/WorldTileDetailsCustomization.cpp
@@ -254,6 +254,7 @@ void FTileLODEntryDetailsCustomization::CustomizeHeader(TSharedRef<class IProper
 				.Text(LOCTEXT("Generate", "Generate"))
 				.OnClicked(this, &FTileLODEntryDetailsCustomization::OnGenerateTile)
 				.IsEnabled(this, &FTileLODEntryDetailsCustomization::IsGenerateTileEnabled)
+				.ToolTipText(LOCTEXT("GenerateLODToolTip", "Creates simplified sub-level by merging geometry into static mesh proxy (requires Simplygon) and exporting landscapes into static meshes"))
 			]
 		];
 }
@@ -305,7 +306,7 @@ bool FTileLODEntryDetailsCustomization::IsGenerateTileEnabled() const
 	TSharedPtr<FWorldTileCollectionModel> PinnedWorldModel = WorldModel.Pin();
 	if (PinnedWorldModel.IsValid())
 	{
-		return PinnedWorldModel->HasGenerateLODLevelSupport() && PinnedWorldModel->AreAnySelectedLevelsLoaded();
+		return PinnedWorldModel->AreAnySelectedLevelsLoaded() && (PinnedWorldModel->HasMeshProxySupport() || PinnedWorldModel->AreAnySelectedLevelsHaveLandscape());
 	}
 	
 	return false;
diff --git a/Engine/Source/Runtime/Engine/Classes/Engine/EngineTypes.h b/Engine/Source/Runtime/Engine/Classes/Engine/EngineTypes.h
--- a/Engine/Source/Runtime/Engine/Classes/Engine/EngineTypes.h
+++ b/Engine/Source/Runtime/Engine/Classes/Engine/EngineTypes.h
@@ -2097,6 +2097,22 @@ struct FMeshProxySettings
 	UPROPERTY(EditAnywhere, Category=ProxySettings)
 	int32 TextureHeight;
 
+	/** Whether to create normal map for the proxy mesh*/
+	UPROPERTY(EditAnywhere, Category=ProxySettings)
+	bool bExportNormalMap;
+
+	/** Whether to create metallic map for the proxy mesh*/
+	UPROPERTY(EditAnywhere, Category=ProxySettings)
+	bool bExportMetallicMap;
+
+	/** Whether to create roughness map for the proxy mesh*/
+	UPROPERTY(EditAnywhere, Category=ProxySettings)
+	bool bExportRoughnessMap;
+
+	/** Whether to create specular map for the proxy mesh*/
+	UPROPERTY(EditAnywhere, Category=ProxySettings)
+	bool bExportSpecularMap;
+
 	/** Should Simplygon recalculate normals for the proxy mesh? */
 	UPROPERTY(EditAnywhere, Category=ProxySettings)
 	bool bRecalculateNormals;
@@ -2125,6 +2141,10 @@ struct FMeshProxySettings
 		: ScreenSize(300)
 		, TextureWidth(512)
 		, TextureHeight(512)
+		, bExportNormalMap(true)
+		, bExportMetallicMap(false)
+		, bExportRoughnessMap(false)
+		, bExportSpecularMap(false)
 		, bRecalculateNormals(true)
 		, HardAngleThreshold(80.0f)
 		, MergeDistance(4)
@@ -2139,7 +2159,11 @@ struct FMeshProxySettings
 	{
 		return ScreenSize == Other.ScreenSize
 			&& TextureWidth == Other.TextureWidth
-			&& TextureHeight == Other.TextureHeight 
+			&& TextureHeight == Other.TextureHeight
+			&& bExportNormalMap == Other.bExportNormalMap
+			&& bExportMetallicMap == Other.bExportMetallicMap
+			&& bExportRoughnessMap == Other.bExportRoughnessMap
+			&& bExportSpecularMap == Other.bExportSpecularMap
 			&& bRecalculateNormals == Other.bRecalculateNormals
 			&& HardAngleThreshold == Other.HardAngleThreshold
 			&& MergeDistance == Other.MergeDistance;
diff --git a/Engine/Source/Runtime/Engine/Classes/Engine/Level.h b/Engine/Source/Runtime/Engine/Classes/Engine/Level.h
--- a/Engine/Source/Runtime/Engine/Classes/Engine/Level.h
+++ b/Engine/Source/Runtime/Engine/Classes/Engine/Level.h
@@ -235,32 +235,51 @@ struct ENGINE_API FLevelSimplificationDetails
 {
 	GENERATED_USTRUCT_BODY()
 
-	// Percentage of details relative to main tile details
-	UPROPERTY(Category=ReductionSettings, EditAnywhere, meta=(ClampMin = "0", ClampMax = "100", UIMin = "0", UIMax = "100"))	
-	float DetailsPercentage;
-
 	// Whether to create separate packages for each generated asset. All in map package otherwise
-	UPROPERTY(Category=Landscape, EditAnywhere)
+	UPROPERTY(Category=General, EditAnywhere)
 	bool bCreatePackagePerAsset;
+
+	// Percentage of details for static mesh proxy
+	UPROPERTY(Category=StaticMesh, EditAnywhere, meta=(DisplayName="Static Mesh Details Percentage", ClampMin = "0", ClampMax = "100", UIMin = "0", UIMax = "100"))	
+	float DetailsPercentage;
+
+	// Whether to generate normal map for static mesh proxy
+	UPROPERTY(Category=StaticMesh, EditAnywhere, meta=(DisplayName="Static Mesh Normal Map"))
+	bool bGenerateMeshNormalMap;
+	
+	// Whether to generate metallic map for static mesh proxy
+	UPROPERTY(Category=StaticMesh, EditAnywhere, meta=(DisplayName="Static Mesh Metallic Map"))
+	bool bGenerateMeshMetallicMap;
+
+	// Whether to generate roughness map for static mesh proxy
+	UPROPERTY(Category=StaticMesh, EditAnywhere, meta=(DisplayName="Static Mesh Roughness Map"))
+	bool bGenerateMeshRoughnessMap;
 	
-	// Landscape LOD to use for static mesh generation
-	UPROPERTY(Category=Landscape, EditAnywhere, meta=(ClampMin = "0", ClampMax = "7", UIMin = "0", UIMax = "7"))
+	// Whether to generate specular map for static mesh proxy
+	UPROPERTY(Category=StaticMesh, EditAnywhere, meta=(DisplayName="Static Mesh Specular Map"))
+	bool bGenerateMeshSpecularMap;
+
+	UPROPERTY()
+	bool bOverrideLandscapeExportLOD;
+
+	// Landscape LOD to use for static mesh generation, when not specified 'Max LODLevel' from landscape actor will be used
+	UPROPERTY(Category=Landscape, EditAnywhere, meta=(ClampMin = "0", ClampMax = "7", UIMin = "0", UIMax = "7", editcondition = "bOverrideLandscapeExportLOD"))
 	int32 LandscapeExportLOD;
 
 	// Whether to generate normal map for landscape static mesh
-	UPROPERTY(Category=Landscape, EditAnywhere)
+	UPROPERTY(Category=Landscape, EditAnywhere, meta=(DisplayName="Landscape Normal Map"))
 	bool bGenerateLandscapeNormalMap;
 
 	// Whether to generate metallic map for landscape static mesh
-	UPROPERTY(Category=Landscape, EditAnywhere)
+	UPROPERTY(Category=Landscape, EditAnywhere, meta=(DisplayName="Landscape Metallic Map"))
 	bool bGenerateLandscapeMetallicMap;
 
 	// Whether to generate roughness map for landscape static mesh
-	UPROPERTY(Category=Landscape, EditAnywhere)
+	UPROPERTY(Category=Landscape, EditAnywhere, meta=(DisplayName="Landscape Roughness Map"))
 	bool bGenerateLandscapeRoughnessMap;
 	
 	// Whether to generate specular map for landscape static mesh
-	UPROPERTY(Category=Landscape, EditAnywhere)
+	UPROPERTY(Category=Landscape, EditAnywhere, meta=(DisplayName="Landscape Specular Map"))
 	bool bGenerateLandscapeSpecularMap;
 	
 	// Whether to bake foliage into landscape static mesh texture
diff --git a/Engine/Source/Runtime/Engine/Private/Level.cpp b/Engine/Source/Runtime/Engine/Private/Level.cpp
--- a/Engine/Source/Runtime/Engine/Private/Level.cpp
+++ b/Engine/Source/Runtime/Engine/Private/Level.cpp
@@ -162,8 +162,13 @@ FArchive& operator<<( FArchive& Ar, FPrecomputedVolumeDistanceField& D )
 }
 
 FLevelSimplificationDetails::FLevelSimplificationDetails()
- : DetailsPercentage(70.f)
- , bCreatePackagePerAsset(true)
+ : bCreatePackagePerAsset(true)
+ , DetailsPercentage(70.f)
+ , bGenerateMeshNormalMap(true)
+ , bGenerateMeshMetallicMap(false)
+ , bGenerateMeshRoughnessMap(false)
+ , bGenerateMeshSpecularMap(false)
+ , bOverrideLandscapeExportLOD(false)
  , LandscapeExportLOD(7)
  , bGenerateLandscapeNormalMap(true)
  , bGenerateLandscapeMetallicMap(false)
@@ -177,8 +182,13 @@ FLevelSimplificationDetails::FLevelSimplificationDetails()
 bool FLevelSimplificationDetails::operator == (const FLevelSimplificationDetails& Other) const
 {
 	return
-		DetailsPercentage == Other.DetailsPercentage &&
 		bCreatePackagePerAsset == Other.bCreatePackagePerAsset &&
+		DetailsPercentage == Other.DetailsPercentage &&
+		bGenerateMeshNormalMap == Other.bGenerateMeshNormalMap &&
+		bGenerateMeshMetallicMap == Other.bGenerateMeshMetallicMap &&
+		bGenerateMeshRoughnessMap == Other.bGenerateMeshRoughnessMap &&
+		bGenerateMeshSpecularMap == Other.bGenerateMeshSpecularMap &&
+		bOverrideLandscapeExportLOD == Other.bOverrideLandscapeExportLOD &&
 		LandscapeExportLOD == Other.LandscapeExportLOD &&
 		bGenerateLandscapeNormalMap == Other.bGenerateLandscapeNormalMap &&
 		bGenerateLandscapeMetallicMap == Other.bGenerateLandscapeMetallicMap &&
diff --git a/Engine/Source/Runtime/Engine/Public/MaterialShared.h b/Engine/Source/Runtime/Engine/Public/MaterialShared.h
--- a/Engine/Source/Runtime/Engine/Public/MaterialShared.h
+++ b/Engine/Source/Runtime/Engine/Public/MaterialShared.h
@@ -329,9 +329,11 @@ namespace EMaterialShaderMapUsage
 		LightmassExportDiffuse,
 		LightmassExportOpacity,
 		LightmassExportNormal,
-		MaterialExportDiffuseColor,
-		MaterialExportSpecularColor,
-		MaterialExportNormal
+		MaterialExportBaseColor,
+		MaterialExportSpecular,
+		MaterialExportNormal,
+		MaterialExportMetallic,
+		MaterialExportRoughness
 	};
 }
 
