Improved viewport gizmo control mechanism.
#ttp 345815 - NUX: MVP: Quality pass on viewport gizmos (esp rotation)
#ttp 345479 - NUX:Level Viewport: Frustratingly glitchy input when trying to drag gizmos after right clicking (moves camera instead of actor!)

diff --git a/Engine/Config/BaseEngine.ini b/Engine/Config/BaseEngine.ini
--- a/Engine/Config/BaseEngine.ini
+++ b/Engine/Config/BaseEngine.ini
@@ -1347,7 +1347,6 @@ bSubtitlesEnabled=True
 GameCommandLine=-log
 FOVAngle=90.000000
 GodMode=True
-InvertwidgetZAxis=True
 UseAxisIndicator=True
 MatineeCurveDetail=0.1
 HeightMapExportClassName="TerrainHeightMapExporterTextT3D"
diff --git a/Engine/Source/Editor/UnrealEd/Classes/Editor/EditorEngine.h b/Engine/Source/Editor/UnrealEd/Classes/Editor/EditorEngine.h
--- a/Engine/Source/Editor/UnrealEd/Classes/Editor/EditorEngine.h
+++ b/Engine/Source/Editor/UnrealEd/Classes/Editor/EditorEngine.h
@@ -371,9 +371,6 @@ public:
 	UPROPERTY(EditAnywhere, config, Category=Advanced)
 	uint32 GodMode:1;
 
-	UPROPERTY(EditAnywhere, config, Category=Advanced)
-	uint32 InvertwidgetZAxis:1;
-
 	UPROPERTY(EditAnywhere, config, Category=Advanced)
 	FString GameCommandLine;
 
diff --git a/Engine/Source/Editor/UnrealEd/Private/EditorViewportClient.cpp b/Engine/Source/Editor/UnrealEd/Private/EditorViewportClient.cpp
--- a/Engine/Source/Editor/UnrealEd/Private/EditorViewportClient.cpp
+++ b/Engine/Source/Editor/UnrealEd/Private/EditorViewportClient.cpp
@@ -1614,19 +1614,9 @@ void FEditorViewportClient::UpdateMouseDelta()
 		return;
 	}
 
-	FVector DragDelta = FVector::ZeroVector;
+	FVector DragDelta = MouseDeltaTracker->GetDelta();
 
-	bool bNeedMovementSnapping = true;
-	if( Widget->GetCurrentAxis() != EAxisList::None && bNeedMovementSnapping )
-	{
-		DragDelta = MouseDeltaTracker->GetDeltaSnapped();
-	}
-	else
-	{
-		DragDelta = MouseDeltaTracker->GetDelta();
-	}
-
-	GEditor->MouseMovement += FVector( FMath::Abs(DragDelta.X), FMath::Abs(DragDelta.Y), FMath::Abs(DragDelta.Z) );
+	GEditor->MouseMovement += DragDelta.GetAbs();
 
 	if( Viewport )
 	{
@@ -1636,102 +1626,90 @@ void FEditorViewportClient::UpdateMouseDelta()
 			const bool MiddleMouseButtonDown = Viewport->KeyState(EKeys::MiddleMouseButton);
 			const bool RightMouseButtonDown = Viewport->KeyState(EKeys::RightMouseButton);
 			const bool bIsUsingTrackpad = FSlateApplication::Get().IsUsingTrackpad();
-
 			const bool bIsNonOrbitMiddleMouse = MiddleMouseButtonDown && !IsAltPressed();
 
-			for( int32 x = 0 ; x < 3 ; ++x )
+			// Convert the movement delta into drag/rotation deltas
+			FVector Drag;
+			FRotator Rot;
+			FVector Scale;
+			EAxisList::Type CurrentAxis = Widget->GetCurrentAxis();
+			if ( IsOrtho() && ( LeftMouseButtonDown || bIsUsingTrackpad ) && RightMouseButtonDown )
+			{
+				bWidgetAxisControlledByDrag = false;
+				Widget->SetCurrentAxis( EAxisList::None );
+				MouseDeltaTracker->ConvertMovementDeltaToDragRot(this, DragDelta, Drag, Rot, Scale);
+				Widget->SetCurrentAxis( CurrentAxis );
+				CurrentAxis = EAxisList::None;
+			}
+			else
 			{
-				FVector WkDelta( 0.f, 0.f, 0.f );
-				WkDelta[x] = DragDelta[x];
-				if( !WkDelta.IsZero() )
+				//if Absolute Translation, and not just moving the camera around
+				if (IsUsingAbsoluteTranslation())
 				{
-					// Convert the movement delta into drag/rotation deltas
-					FVector Drag;
-					FRotator Rot;
-					FVector Scale;
-					EAxisList::Type CurrentAxis = Widget->GetCurrentAxis();
-					if ( IsOrtho() && ( LeftMouseButtonDown || bIsUsingTrackpad ) && RightMouseButtonDown )
-					{
-						bWidgetAxisControlledByDrag = false;
-						Widget->SetCurrentAxis( EAxisList::None );
-						MouseDeltaTracker->ConvertMovementDeltaToDragRot( this, WkDelta, Drag, Rot, Scale );
-						Widget->SetCurrentAxis( CurrentAxis );
-						CurrentAxis = EAxisList::None;
-					}
-					else
-					{
-						//if Absolute Translation, and not just moving the camera around
-						if (IsUsingAbsoluteTranslation())
-						{
-							// Compute a view.
-							FSceneViewFamilyContext ViewFamily(FSceneViewFamily::ConstructionValues(
-								Viewport,
-								GetScene(),
-								EngineShowFlags)
-								.SetRealtimeUpdate( IsRealtime() ) );
-							FSceneView* View = CalcSceneView( &ViewFamily );
-
-							MouseDeltaTracker->AbsoluteTranslationConvertMouseToDragRot(View, this, Drag, Rot, Scale);
-							//only need to go through this loop once
-							x = 3;
-						} 
-						else
-						{
-							MouseDeltaTracker->ConvertMovementDeltaToDragRot( this, WkDelta, Drag, Rot, Scale );
-						}
-					}
+					// Compute a view.
+					FSceneViewFamilyContext ViewFamily(FSceneViewFamily::ConstructionValues(
+						Viewport,
+						GetScene(),
+						EngineShowFlags)
+						.SetRealtimeUpdate( IsRealtime() ) );
+					FSceneView* View = CalcSceneView( &ViewFamily );
+
+					MouseDeltaTracker->AbsoluteTranslationConvertMouseToDragRot(View, this, Drag, Rot, Scale);
+				} 
+				else
+				{
+					MouseDeltaTracker->ConvertMovementDeltaToDragRot(this, DragDelta, Drag, Rot, Scale);
+				}
+			}
 
-					const bool bInputHandledByGizmos = InputWidgetDelta( Viewport, CurrentAxis, Drag, Rot, Scale );
+			const bool bInputHandledByGizmos = InputWidgetDelta( Viewport, CurrentAxis, Drag, Rot, Scale );
 					
-					if( !Rot.IsZero() )
-					{
-						Widget->UpdateDeltaRotation();
-					}
+			if( !Rot.IsZero() )
+			{
+				Widget->UpdateDeltaRotation();
+			}
 
-					if( !bInputHandledByGizmos )
+			if( !bInputHandledByGizmos )
+			{
+				if ( ShouldOrbitCamera() )
+				{
+					FVector TempDrag;
+					FRotator TempRot;
+					InputAxisForOrbit( Viewport, DragDelta, TempDrag, TempRot );
+				}
+				else
+				{
+					// Disable orbit camera
+					const bool bEnable=false;
+					ToggleOrbitCamera(bEnable);
+
+					if ( ShouldPanOrDollyCamera() )
 					{
-						if ( ShouldOrbitCamera() )
+						if( !IsOrtho())
+						{
+							const float CameraSpeed = GetCameraSpeed();
+							Drag *= CameraSpeed;
+						}
+						MoveViewportCamera( Drag, Rot );
+
+						if ( IsPerspective() && LeftMouseButtonDown && !MiddleMouseButtonDown && !RightMouseButtonDown )
 						{
-							FVector TempDrag;
-							FRotator TempRot;
-							InputAxisForOrbit( Viewport, DragDelta, TempDrag, TempRot );
+							FEditorViewportStats::Using(FEditorViewportStats::CAT_PERSPECTIVE_MOUSE_DOLLY);
 						}
 						else
 						{
-							// Disable orbit camera
-							const bool bEnable=false;
-							ToggleOrbitCamera(bEnable);
-
-							if ( ShouldPanOrDollyCamera() )
+							if ( !Drag.IsZero() )
 							{
-								if( !IsOrtho())
-								{
-									const float CameraSpeed = GetCameraSpeed();
-									Drag *= CameraSpeed;
-								}
-								MoveViewportCamera( Drag, Rot );
-
-								if ( IsPerspective() && LeftMouseButtonDown && !MiddleMouseButtonDown && !RightMouseButtonDown )
-								{
-									FEditorViewportStats::Using(FEditorViewportStats::CAT_PERSPECTIVE_MOUSE_DOLLY);
-								}
-								else
-								{
-									if ( !Drag.IsZero() )
-									{
-										FEditorViewportStats::Using(IsPerspective() ? FEditorViewportStats::CAT_PERSPECTIVE_MOUSE_PAN : FEditorViewportStats::CAT_ORTHOGRAPHIC_MOUSE_PAN);
-									}
-								}
+								FEditorViewportStats::Using(IsPerspective() ? FEditorViewportStats::CAT_PERSPECTIVE_MOUSE_PAN : FEditorViewportStats::CAT_ORTHOGRAPHIC_MOUSE_PAN);
 							}
 						}
 					}
-			
-
-					// Clean up
-					MouseDeltaTracker->ReduceBy( WkDelta );
 				}
 			}
 
+			// Clean up
+			MouseDeltaTracker->ReduceBy( DragDelta );
+
 			Invalidate( false, false );
 		}
 	}
diff --git a/Engine/Source/Editor/UnrealEd/Private/MouseDeltaTracker.cpp b/Engine/Source/Editor/UnrealEd/Private/MouseDeltaTracker.cpp
--- a/Engine/Source/Editor/UnrealEd/Private/MouseDeltaTracker.cpp
+++ b/Engine/Source/Editor/UnrealEd/Private/MouseDeltaTracker.cpp
@@ -179,11 +179,18 @@ void FMouseDeltaTracker::DetermineCurrentAxis(FEditorViewportClient* InViewportC
  */
 void FMouseDeltaTracker::StartTracking(FEditorViewportClient* InViewportClient, const int32 InX, const int32 InY, const FInputEventState& InInputState, bool bNudge, bool bResetDragToolState)
 {
-	// If the builder brush was selected and is still selected we don't want to deselect it due to the possibility that you can get two mouse down
-	// events on the same viewport before you get a mouse up event.
-	bool bNoActorSelected = !GEditor->GetSelectedActors()->GetTop<AActor>();
-	ABrush* Brush = Cast< ABrush >( GEditor->GetSelectedActors()->GetTop<AActor>() );
+	// Initialize widget axis (in case it hasn't been set by the hovered hit proxy)
 
+	if (InViewportClient->Widget)
+	{
+		check(InViewportClient->Viewport);
+		HHitProxy* HitProxy = InViewportClient->Viewport->GetHitProxy(InX, InY);
+		if (HitProxy && HitProxy->IsA(HWidgetAxis::StaticGetType()))
+		{
+			EAxisList::Type ProxyAxis = ((HWidgetAxis*)HitProxy)->Axis;
+			InViewportClient->SetCurrentWidgetAxis(ProxyAxis);
+		}
+	}
 
 	DetermineCurrentAxis( InViewportClient );
 
@@ -207,12 +214,12 @@ void FMouseDeltaTracker::StartTracking(FEditorViewportClient* InViewportClient,
 
 	if (InViewportClient->Widget)
 	{
+		InViewportClient->Widget->SetDragStartPosition(FVector2D(InX, InY));
 		InViewportClient->Widget->SetDragging(bIsDragging);
 		if (InViewportClient->GetWidgetMode() == FWidget::WM_Rotate)
 		{
 			InViewportClient->Invalidate();
 		}
-
 	}
 
 	// Clear bool that tracks whether AddDelta has been called
@@ -570,30 +577,30 @@ const FVector FMouseDeltaTracker::GetScreenDelta() const
 /**
  * Converts the delta movement to drag/rotation/scale based on the viewport type or widget axis
  */
-void FMouseDeltaTracker::ConvertMovementDeltaToDragRot(FEditorViewportClient* InViewportClient, const FVector& InDragDelta, FVector& InDrag, FRotator& InRotation, FVector& InScale) const
+void FMouseDeltaTracker::ConvertMovementDeltaToDragRot(FEditorViewportClient* InViewportClient, FVector& InOutDragDelta, FVector& OutDrag, FRotator& OutRotation, FVector& OutScale) const
 {
-	InDrag = FVector::ZeroVector;
-	InRotation = FRotator::ZeroRotator;
-	InScale = FVector::ZeroVector;
+	OutDrag = FVector::ZeroVector;
+	OutRotation = FRotator::ZeroRotator;
+	OutScale = FVector::ZeroVector;
 
 	if( InViewportClient->GetCurrentWidgetAxis() != EAxisList::None )
 	{
-		InViewportClient->Widget->ConvertMouseMovementToAxisMovement( InViewportClient, InViewportClient->GetWidgetLocation(), InDragDelta, InDrag, InRotation, InScale );
+		InViewportClient->Widget->ConvertMouseMovementToAxisMovement( InViewportClient, InViewportClient->GetWidgetLocation(), InOutDragDelta, OutDrag, OutRotation, OutScale );
 	}
 	else
 	{
-		InViewportClient->ConvertMovementToDragRot( InDragDelta, InDrag, InRotation );
+		InViewportClient->ConvertMovementToDragRot( InOutDragDelta, OutDrag, OutRotation );
 	}
 }
 
 /**
  * Absolute Translation conversion from mouse position on the screen to widget axis movement/rotation.
  */
-void FMouseDeltaTracker::AbsoluteTranslationConvertMouseToDragRot(FSceneView* InView, FEditorViewportClient* InViewportClient,FVector& InDrag, FRotator& InRotation, FVector& InScale ) const
+void FMouseDeltaTracker::AbsoluteTranslationConvertMouseToDragRot(FSceneView* InView, FEditorViewportClient* InViewportClient,FVector& OutDrag, FRotator& OutRotation, FVector& OutScale ) const
 {
-	InDrag = FVector::ZeroVector;
-	InRotation = FRotator::ZeroRotator;
-	InScale = FVector::ZeroVector;
+	OutDrag = FVector::ZeroVector;
+	OutRotation = FRotator::ZeroRotator;
+	OutScale = FVector::ZeroVector;
 
 	check ( InViewportClient->GetCurrentWidgetAxis() != EAxisList::None );
 
@@ -601,7 +608,7 @@ void FMouseDeltaTracker::AbsoluteTranslationConvertMouseToDragRot(FSceneView* In
 	check(InViewportClient->Viewport);
 	FVector2D MousePosition(InViewportClient->Viewport->GetMouseX(), InViewportClient->Viewport->GetMouseY());
 
-	InViewportClient->Widget->AbsoluteTranslationConvertMouseMovementToAxisMovement(InView, InViewportClient, InViewportClient->GetWidgetLocation(), MousePosition, InDrag, InRotation, InScale );
+	InViewportClient->Widget->AbsoluteTranslationConvertMouseMovementToAxisMovement(InView, InViewportClient, InViewportClient->GetWidgetLocation(), MousePosition, OutDrag, OutRotation, OutScale );
 }
 
 /**
diff --git a/Engine/Source/Editor/UnrealEd/Private/UnrealWidget.cpp b/Engine/Source/Editor/UnrealEd/Private/UnrealWidget.cpp
--- a/Engine/Source/Editor/UnrealEd/Private/UnrealWidget.cpp
+++ b/Engine/Source/Editor/UnrealEd/Private/UnrealWidget.cpp
@@ -65,6 +65,7 @@ FWidget::FWidget()
 	bDragging = false;
 	bSnapEnabled = false;
 	bDefaultVisibility = true;
+	bIsOrthoDrawingFullRing = false;
 }
 
 extern ENGINE_API void StringSize(UFont* Font,int32& XL,int32& YL,const TCHAR* Text, FCanvas* Canvas);
@@ -193,16 +194,16 @@ void FWidget::Render( const FSceneView* View,FPrimitiveDrawInterface* PDI, FEdit
 /**
  * Draws an arrow head line for a specific axis.
  */
-void FWidget::Render_Axis( const FSceneView* View, FPrimitiveDrawInterface* PDI, EAxisList::Type InAxis, FMatrix& InMatrix, UMaterialInterface* InMaterial, const FLinearColor& InColor, FVector2D& OutAxisEnd, const FVector& InScale, bool bDrawWidget, bool bCubeHead )
+void FWidget::Render_Axis( const FSceneView* View, FPrimitiveDrawInterface* PDI, EAxisList::Type InAxis, FMatrix& InMatrix, UMaterialInterface* InMaterial, const FLinearColor& InColor, FVector2D& OutAxisDir, const FVector& InScale, bool bDrawWidget, bool bCubeHead )
 {
 	FMatrix AxisRotation = FMatrix::Identity;
 	if( InAxis == EAxisList::Y )
 	{
-		AxisRotation = FRotationMatrix( FRotator(0,90.f,0) );
+		AxisRotation = FRotationMatrix::MakeFromXZ(FVector(0, 1, 0), FVector(0, 0, 1));
 	}
 	else if( InAxis == EAxisList::Z )
 	{
-		AxisRotation = FRotationMatrix( FRotator(90.f,0,0) );
+		AxisRotation = FRotationMatrix::MakeFromXY(FVector(0, 0, 1), FVector(0, 1, 0));
 	}
 
 	FMatrix ArrowToWorld = AxisRotation * InMatrix;
@@ -266,10 +267,13 @@ void FWidget::Render_Axis( const FSceneView* View, FPrimitiveDrawInterface* PDI,
 		PDI->SetHitProxy( NULL );
 	}
 
-	if(!View->ScreenToPixel(View->WorldToScreen(ArrowToWorld.TransformPosition(FVector(64,0,0))),OutAxisEnd))
+	FVector2D AxisEnd;
+	if(!View->ScreenToPixel(View->WorldToScreen(ArrowToWorld.TransformPosition(FVector(64,0,0))),AxisEnd))
 	{
-		OutAxisEnd.X = OutAxisEnd.Y = 0;
+		AxisEnd.X = AxisEnd.Y = 0;
 	}
+
+	OutAxisDir = (AxisEnd - Origin).GetSafeNormal();
 }
 
 void FWidget::Render_Cube( FPrimitiveDrawInterface* PDI, const FMatrix& InMatrix, const UMaterialInterface* InMaterial, const FVector& InScale )
@@ -410,19 +414,19 @@ void FWidget::Render_Translate( const FSceneView* View, FPrimitiveDrawInterface*
 	if( DrawAxis&EAxisList::X && (bIsPerspective || bIsLocalSpace || !bIsOrthoYZ) )
 	{
 		UMaterialInstanceDynamic* XMaterial = ( CurrentAxis&EAxisList::X ? CurrentAxisMaterial : AxisMaterialX );
-		Render_Axis( View, PDI, EAxisList::X, WidgetMatrix, XMaterial, XColor, XAxisEnd, Scale, bDrawWidget );
+		Render_Axis( View, PDI, EAxisList::X, WidgetMatrix, XMaterial, XColor, XAxisDir, Scale, bDrawWidget );
 	}
 
 	if( DrawAxis&EAxisList::Y && (bIsPerspective || bIsLocalSpace || !bIsOrthoXZ) )
 	{
 		UMaterialInstanceDynamic* YMaterial = ( CurrentAxis&EAxisList::Y ? CurrentAxisMaterial : AxisMaterialY );
-		Render_Axis( View, PDI, EAxisList::Y, WidgetMatrix, YMaterial, YColor, YAxisEnd, Scale, bDrawWidget );
+		Render_Axis( View, PDI, EAxisList::Y, WidgetMatrix, YMaterial, YColor, YAxisDir, Scale, bDrawWidget );
 	}
 
 	if( DrawAxis&EAxisList::Z && (bIsPerspective || bIsLocalSpace || !bIsOrthoXY) )
 	{
 		UMaterialInstanceDynamic* ZMaterial = ( CurrentAxis&EAxisList::Z ? CurrentAxisMaterial : AxisMaterialZ );
-		Render_Axis( View, PDI, EAxisList::Z, WidgetMatrix, ZMaterial, ZColor, ZAxisEnd, Scale, bDrawWidget );
+		Render_Axis( View, PDI, EAxisList::Z, WidgetMatrix, ZMaterial, ZColor, ZAxisDir, Scale, bDrawWidget );
 	}
 
 	// Draw the grabbers
@@ -506,43 +510,36 @@ void FWidget::Render_Rotate( const FSceneView* View,FPrimitiveDrawInterface* PDI
 	float Scale = View->WorldToScreen( InLocation ).W * ( 4.0f / View->ViewRect.Width() / View->ViewMatrices.ProjMatrix.M[0][0] );
 
 	//get the axes 
-	FVector XAxis = CustomCoordSystem.TransformVector(FVector(-1, 0, 0));
-	FVector YAxis = CustomCoordSystem.TransformVector(FVector(0, -1, 0));
+	FVector XAxis = CustomCoordSystem.TransformVector(FVector(1, 0, 0));
+	FVector YAxis = CustomCoordSystem.TransformVector(FVector(0, 1, 0));
 	FVector ZAxis = CustomCoordSystem.TransformVector(FVector(0, 0, 1));
 
 	EAxisList::Type DrawAxis = GetAxisToDraw( ViewportClient->GetWidgetMode() );
 
-	FMatrix XMatrix = FRotationMatrix( FRotator(0,90.f,0) ) * FTranslationMatrix( InLocation );
-
 	FVector DirectionToWidget = View->IsPerspectiveProjection() ? (InLocation - View->ViewMatrices.ViewOrigin) : -View->GetViewDirection();
 	DirectionToWidget.Normalize();
 
 	// Draw a circle for each axis
 	if (bDrawWidget || bDragging)
 	{
+		bIsOrthoDrawingFullRing = false;
+
 		//no draw the arc segments
 		if( DrawAxis&EAxisList::X )
 		{
-			DrawRotationArc(View, PDI, EAxisList::X, InLocation, YAxis, ZAxis, DirectionToWidget, AxisColorX, Scale);
+			DrawRotationArc(View, PDI, EAxisList::X, InLocation, ZAxis, YAxis, DirectionToWidget, AxisColorX, Scale, XAxisDir);
 		}
 
 		if( DrawAxis&EAxisList::Y )
 		{
-			DrawRotationArc(View, PDI, EAxisList::Y, InLocation, ZAxis, XAxis, DirectionToWidget, AxisColorY, Scale);
+			DrawRotationArc(View, PDI, EAxisList::Y, InLocation, XAxis, ZAxis, DirectionToWidget, AxisColorY, Scale, YAxisDir);
 		}
 
 		if( DrawAxis&EAxisList::Z )
 		{
-			DrawRotationArc(View, PDI, EAxisList::Z, InLocation, XAxis, YAxis, DirectionToWidget, AxisColorZ, Scale);
+			DrawRotationArc(View, PDI, EAxisList::Z, InLocation, XAxis, YAxis, DirectionToWidget, AxisColorZ, Scale, ZAxisDir);
 		}
 	}
-
-	// Update Axis by projecting the axis vector to screenspace.
-	View->ScreenToPixel(View->WorldToScreen( XMatrix.TransformPosition( FVector(96,0,0) ) ), XAxisEnd);
-	// Update Axis by projecting the axis vector to screenspace.
-	View->ScreenToPixel(View->WorldToScreen( XMatrix.TransformPosition( FVector(0,96,0) ) ), YAxisEnd);
-	// Update Axis by projecting the axis vector to screenspace.
-	View->ScreenToPixel(View->WorldToScreen( XMatrix.TransformPosition( FVector(0,0,96) ) ), ZAxisEnd);
 }
 
 /**
@@ -597,17 +594,17 @@ void FWidget::Render_Scale( const FSceneView* View,FPrimitiveDrawInterface* PDI,
 	// Draw the axis lines with cube heads	
 	if( !bIsOrthoYZ && DrawAxis&EAxisList::X )
 	{
-		Render_Axis( View, PDI, EAxisList::X, WidgetMatrix, XMaterial, XColor, XAxisEnd, Scale, bDrawWidget, true );
+		Render_Axis( View, PDI, EAxisList::X, WidgetMatrix, XMaterial, XColor, XAxisDir, Scale, bDrawWidget, true );
 	}
 
 	if( !bIsOrthoXZ && DrawAxis&EAxisList::Y )
 	{
-		Render_Axis( View, PDI, EAxisList::Y, WidgetMatrix, YMaterial, YColor, YAxisEnd, Scale, bDrawWidget, true );
+		Render_Axis( View, PDI, EAxisList::Y, WidgetMatrix, YMaterial, YColor, YAxisDir, Scale, bDrawWidget, true );
 	}
 
 	if( !bIsOrthoXY &&  DrawAxis&EAxisList::Z )
 	{
-		Render_Axis( View, PDI, EAxisList::Z, WidgetMatrix, ZMaterial, ZColor, ZAxisEnd, Scale, bDrawWidget, true );
+		Render_Axis( View, PDI, EAxisList::Z, WidgetMatrix, ZMaterial, ZColor, ZAxisDir, Scale, bDrawWidget, true );
 	}
 
 	// Draw grabber handles and center cube
@@ -718,17 +715,17 @@ void FWidget::Render_TranslateRotateZ( const FSceneView* View, FPrimitiveDrawInt
 		// Draw the axis lines with arrow heads
 		if( DrawAxis&EAxisList::X && (bIsPerspective || bIsLocalSpace || View->ViewMatrices.ViewMatrix.M[0][2] != -1.f) )
 		{
-			Render_Axis( View, PDI, EAxisList::X, AxisMatrix, XMaterial, XColor, XAxisEnd, Scale, bDrawWidget );
+			Render_Axis( View, PDI, EAxisList::X, AxisMatrix, XMaterial, XColor, XAxisDir, Scale, bDrawWidget );
 		}
 
 		if( DrawAxis&EAxisList::Y && (bIsPerspective || bIsLocalSpace || View->ViewMatrices.ViewMatrix.M[1][2] != -1.f) )
 		{
-			Render_Axis( View, PDI, EAxisList::Y, AxisMatrix, YMaterial, YColor, YAxisEnd, Scale, bDrawWidget );
+			Render_Axis( View, PDI, EAxisList::Y, AxisMatrix, YMaterial, YColor, YAxisDir, Scale, bDrawWidget );
 		}
 
 		if( DrawAxis&EAxisList::Z && (bIsPerspective || bIsLocalSpace || View->ViewMatrices.ViewMatrix.M[0][1] != 1.f) )
 		{
-			Render_Axis( View, PDI, EAxisList::Z, AxisMatrix, ZMaterial, ZColor, ZAxisEnd, Scale, bDrawWidget );
+			Render_Axis( View, PDI, EAxisList::Z, AxisMatrix, ZMaterial, ZColor, ZAxisDir, Scale, bDrawWidget );
 		}
 
 		const bool bDisabled = IsWidgetDisabled();
@@ -774,181 +771,131 @@ void FWidget::Render_TranslateRotateZ( const FSceneView* View, FPrimitiveDrawInt
 /**
  * Converts mouse movement on the screen to widget axis movement/rotation.
  */
-void FWidget::ConvertMouseMovementToAxisMovement( FEditorViewportClient* InViewportClient, const FVector& InLocation, const FVector& InDiff, FVector& InDrag, FRotator& InRotation, FVector& InScale )
+void FWidget::ConvertMouseMovementToAxisMovement( FEditorViewportClient* InViewportClient, const FVector& InLocation, FVector& InOutDelta, FVector& OutDrag, FRotator& OutRotation, FVector& OutScale )
 {
-	FSceneViewFamilyContext ViewFamily( FSceneViewFamily::ConstructionValues( InViewportClient->Viewport, InViewportClient->GetScene(), InViewportClient->EngineShowFlags ));
-	FSceneView* View = InViewportClient->CalcSceneView(&ViewFamily);
-
-	FPlane Wk;
-	FVector2D AxisEnd;
-	FVector Diff = InDiff;
-
-	InDrag = FVector::ZeroVector;
-	InRotation = FRotator::ZeroRotator;
-	InScale = FVector::ZeroVector; 
+	OutDrag = FVector::ZeroVector;
+	OutRotation = FRotator::ZeroRotator;
+	OutScale = FVector::ZeroVector; 
 
-	// Get the end of the axis (in screen space) based on which axis is being pulled
+	const int32 WidgetMode = InViewportClient->GetWidgetMode();
 
-	switch( CurrentAxis )
-	{
-		case EAxisList::X:	AxisEnd = XAxisEnd;		break;
-		case EAxisList::Y:	AxisEnd = YAxisEnd;		break;
-		case EAxisList::Z:	AxisEnd = ZAxisEnd;		break;
-		case EAxisList::XY:	AxisEnd = Diff.X != 0 ? XAxisEnd : YAxisEnd;		break;
-		case EAxisList::XZ:	AxisEnd = Diff.X != 0 ? XAxisEnd : ZAxisEnd;		break;
-		case EAxisList::YZ:	AxisEnd = Diff.X != 0 ? YAxisEnd : ZAxisEnd;		break;
-		case EAxisList::XYZ:	AxisEnd = Diff.X != 0 ? YAxisEnd : ZAxisEnd;		break;
-		default:
-			break;
-	}
+	// Get input delta as 2D vector, adjusted for inverted screen space Y axis
+	const FVector2D DragDir = FVector2D(InOutDelta.X, -InOutDelta.Y);
 
-	// Screen space Y axis is inverted
+	// Get offset of the drag start position from the widget origin
+	const FVector2D DirectionToMousePos = FVector2D(DragStartPos - Origin).GetSafeNormal();
 
-	Diff.Y *= -1;
+	// Get direction of tangent of circle formed from the widget origin and the drag start position
+	const FVector2D TangentDir = FVector2D(-DirectionToMousePos.Y, DirectionToMousePos.X);
 
-	// Get the directions of the axis (on the screen) and the mouse drag direction (in screen space also).
-	if(!View->ScreenToPixel(View->WorldToScreen(InLocation),Origin))
+	switch (WidgetMode)
 	{
-		Origin.X = Origin.Y = 0;
-	}
+		case WM_Translate:
+		{
+			// Get drag delta in widget axis space
+			OutDrag = FVector(
+				(CurrentAxis & EAxisList::X) ? FVector2D::DotProduct(XAxisDir, DragDir) : 0.0f,
+				(CurrentAxis & EAxisList::Y) ? FVector2D::DotProduct(YAxisDir, DragDir) : 0.0f,
+				(CurrentAxis & EAxisList::Z) ? FVector2D::DotProduct(ZAxisDir, DragDir) : 0.0f
+				);
 
+			// Snap to grid in widget axis space
+			const FVector GridSize = FVector(GEditor->GetGridSize());
+			FSnappingUtils::SnapPointToGrid(OutDrag, GridSize);
 
-	FVector2D AxisDir = AxisEnd - Origin;
-	AxisDir.Normalize();
+			// Convert to effective screen space delta, and replace input delta, adjusted for inverted screen space Y axis
+			const FVector2D EffectiveDelta = OutDrag.X * XAxisDir + OutDrag.Y * YAxisDir + OutDrag.Z * ZAxisDir;
+			InOutDelta = FVector(EffectiveDelta.X, -EffectiveDelta.Y, 0.0f);
 
-	FVector2D DragDir( Diff.X, Diff.Y );
-	DragDir.Normalize();
+			// Transform drag delta into world space
+			OutDrag = CustomCoordSystem.TransformPosition(OutDrag);
+		}
+		break;
 
+		case WM_Rotate:
+			{
+				FRotator Rotation;
+				FVector2D EffectiveDelta;
 
-	// Use the most dominant axis the mouse is being dragged along -
-	// unless we are modifying a single axis in an Ortho viewport.
-	uint32 idx = GetDominantAxisIndex( Diff, InViewportClient );
+				if (CurrentAxis == EAxisList::X)
+				{
+					// Get screen direction representing positive rotation
+					const FVector2D AxisDir = bIsOrthoDrawingFullRing ? TangentDir : XAxisDir;
 
-	float Val = Diff[idx];
+					// Get rotation in widget local space
+					Rotation = FRotator(0, 0, FVector2D::DotProduct(AxisDir, DragDir));
+					FSnappingUtils::SnapRotatorToGrid(Rotation);
 
-	FMatrix InputCoordSystem = InViewportClient->GetWidgetCoordSystem();
+					// Record delta rotation (used by the widget to render the accumulated delta)
+					CurrentDeltaRotation = Rotation.Roll;
 
-	const int32 WidgetMode = InViewportClient->GetWidgetMode();
+					// Use to calculate the new input delta
+					EffectiveDelta = AxisDir * Rotation.Roll;
+				}
+				else if (CurrentAxis == EAxisList::Y)
+				{
+					// TODO: Determine why -TangentDir is necessary here, and fix whatever is causing it
+					const FVector2D AxisDir = bIsOrthoDrawingFullRing ? -TangentDir : YAxisDir;
 
+					Rotation = FRotator(FVector2D::DotProduct(AxisDir, DragDir), 0, 0);
+					FSnappingUtils::SnapRotatorToGrid(Rotation);
 
-	if( WidgetMode == FWidget::WM_Rotate && InViewportClient->IsPerspective() )
-	{
-		FVector DirectionToWidget = InLocation - View->ViewMatrices.ViewOrigin;
+					CurrentDeltaRotation = Rotation.Pitch;
+					EffectiveDelta = AxisDir * Rotation.Pitch;
+				}
+				else if (CurrentAxis == EAxisList::Z)
+				{
+					const FVector2D AxisDir = bIsOrthoDrawingFullRing ? TangentDir : ZAxisDir;
 
-		FVector XAxis = InputCoordSystem.TransformVector(FVector(1, 0, 0));
-		FVector YAxis = InputCoordSystem.TransformVector(FVector(0, 1, 0));
-		FVector ZAxis = InputCoordSystem.TransformVector(FVector(0, 0, 1));
+					Rotation = FRotator(0, FVector2D::DotProduct(AxisDir, DragDir), 0);
+					FSnappingUtils::SnapRotatorToGrid(Rotation);
 
-		float XDot = XAxis|DirectionToWidget;
-		float YDot = YAxis|DirectionToWidget;
-		float ZDot = ZAxis|DirectionToWidget;
+					CurrentDeltaRotation = Rotation.Yaw;
+					EffectiveDelta = AxisDir * Rotation.Yaw;
+				}
 
-		switch(CurrentAxis)
-		{
-		case EAxisList::X:
-			if( FMath::IsNegativeFloat(XDot ) )
-			{
-				Val *= -1;
-			}
-			break;
-		case EAxisList::Y:
-			if( FMath::IsNegativeFloat(YDot) )
-			{
-				Val *= -1;
-			}
-			break;
-		case EAxisList::Z:
-			if( FMath::IsNegativeFloat(ZDot) )
-			{
-				Val *= -1;
+				// Adjust the input delta according to how much rotation was actually applied
+				InOutDelta = FVector(EffectiveDelta.X, -EffectiveDelta.Y, 0.0f);
+
+				// Need to get the delta rotation in the current coordinate space of the widget
+				OutRotation = (CustomCoordSystem.Inverse() * FRotationMatrix(Rotation) * CustomCoordSystem).Rotator();
 			}
 			break;
-		}
-	}
-	else
-	{
-
-		// If the axis dir is negative, it is pointing in the negative screen direction.  In this situation, the mouse
-		// drag must be inverted so that you are still dragging in the right logical direction.
-		//
-		// For example, if the X axis is pointing left and you drag left, this will ensure that the widget moves left.
-		//Only valid for single axis movement.  For planar movement, this widget gets caught up at the origin and oscillates
-		
-		if(( AxisDir[idx] < 0 ) && ((CurrentAxis == EAxisList::X) || (CurrentAxis == EAxisList::Y) || (CurrentAxis == EAxisList::Z)))
-		{
-			Val *= -1;
-		}
-	}
-
-
-	// Honor INI option to invert Z axis movement on the widget
-
-	if( idx == 1 && (CurrentAxis&EAxisList::Z) && GEditor->InvertwidgetZAxis && ((WidgetMode==WM_Translate) || (WidgetMode==WM_Rotate) || (WidgetMode==WM_TranslateRotateZ)) &&
-		// Don't apply this if the origin and the AxisEnd are the same
-		AxisDir.IsNearlyZero() == false)
-	{
-		Val *= -1;
-	}
 
-
-
-	switch( WidgetMode )
-	{
-		case WM_Translate:
-			switch( CurrentAxis )
+		case WM_Scale:
 			{
-				case EAxisList::X:	InDrag = FVector( Val, 0, 0 );	break;
-				case EAxisList::Y:	InDrag = FVector( 0, Val, 0 );	break;
-				case EAxisList::Z:	InDrag = FVector( 0, 0, -Val );	break;
-				case EAxisList::XY:	InDrag = ( InDiff.X != 0 ? FVector( Val, 0, 0 ) : FVector( 0, Val, 0 ) );	break;
-				case EAxisList::XZ:	InDrag = ( InDiff.X != 0 ? FVector( Val, 0, 0 ) : FVector( 0, 0, Val ) );	break;
-				case EAxisList::YZ:	InDrag = ( InDiff.X != 0 ? FVector( 0, Val, 0 ) : FVector( 0, 0, Val ) );	break;
-			}
+				FVector2D AxisDir = FVector2D::ZeroVector;
 
-			InDrag = InputCoordSystem.TransformPosition( InDrag );
-			break;
+				if (CurrentAxis & EAxisList::X)
+				{
+					AxisDir += XAxisDir;
+				}
 
-		case WM_Rotate:
-			{
-				FVector WorldManDir;
-				FVector Axis;
-				switch( CurrentAxis )
+				if (CurrentAxis & EAxisList::Y)
 				{
-					case EAxisList::X:	
-						Axis = FVector( -1, 0, 0 );
-						break;
-					case EAxisList::Y:	
-						Axis = FVector( 0, -1, 0 );	
-						break;
-					case EAxisList::Z:	
-						Axis = FVector( 0, 0, 1 );	
-						break;
-					default:
-						// Prevent this from crashing when axis incorrect and
-						// make sure Axis is set to something sensible.
-						ensureMsgf( 0, TEXT("Axis not correctly set while rotating! Axis value was %d"), (uint32)CurrentAxis );
-						Axis = FVector( -1, 0, 0 );
-						break;
+					AxisDir += YAxisDir;
 				}
-			
-				Axis = InputCoordSystem.TransformVector( Axis );
-
-				const float RotationSpeed = GetRotationSpeed();
-				const FQuat DeltaQ( Axis, Val*RotationSpeed );
-				CurrentDeltaRotation = Val;
-				
-				InRotation = FRotator(DeltaQ);
-			}
-			break;
 
-		case WM_Scale:
-			{
-				FVector Axis;
-				Axis.X = (CurrentAxis & EAxisList::X) == 0 ? 0 : 1;
-				Axis.Y = (CurrentAxis & EAxisList::Y) == 0 ? 0 : 1;
-				Axis.Z = (CurrentAxis & EAxisList::Z) == 0 ? 0 : 1;
+				if (CurrentAxis & EAxisList::Z)
+				{
+					AxisDir += ZAxisDir;
+				}
+
+				const float ScaleDelta = FVector2D::DotProduct(AxisDir.GetSafeNormal(), DragDir);
 
-				InScale = Axis * Val;
+				OutScale = FVector(
+					(CurrentAxis & EAxisList::X) ? ScaleDelta : 0.0f,
+					(CurrentAxis & EAxisList::Y) ? ScaleDelta : 0.0f,
+					(CurrentAxis & EAxisList::Z) ? ScaleDelta : 0.0f
+					);
+
+				// Snap to grid in widget axis space
+				const FVector GridSize = FVector(GEditor->GetGridSize());
+				FSnappingUtils::SnapScale(OutScale, GridSize);
+
+				// Convert to effective screen space delta, and replace input delta, adjusted for inverted screen space Y axis
+				const FVector2D EffectiveDelta = OutScale.X * XAxisDir + OutScale.Y * YAxisDir + OutScale.Z * ZAxisDir;
+				InOutDelta = FVector(EffectiveDelta.X, -EffectiveDelta.Y, 0.0f);
 			}
 			break;
 
@@ -956,26 +903,35 @@ void FWidget::ConvertMouseMovementToAxisMovement( FEditorViewportClient* InViewp
 			{
 				if( CurrentAxis == EAxisList::ZRotation )
 				{
-					FVector Axis = FVector( 0, 0, 1 );
-					Axis = InputCoordSystem.TransformVector( Axis );
+					const FVector2D AxisDir = bIsOrthoDrawingFullRing ? TangentDir : ZAxisDir;
+					FRotator Rotation = FRotator(0, FVector2D::DotProduct(AxisDir, DragDir), 0);
+					FSnappingUtils::SnapRotatorToGrid(Rotation);
+					CurrentDeltaRotation = Rotation.Yaw;
 
-					const float RotationSpeed = GetRotationSpeed();
-					const FQuat DeltaQ( Axis, Val*RotationSpeed );
-					CurrentDeltaRotation = Val;
+					const FVector2D EffectiveDelta = AxisDir * Rotation.Yaw;
+					InOutDelta = FVector(EffectiveDelta.X, -EffectiveDelta.Y, 0.0f);
 
-					InRotation = FRotator(DeltaQ);
+					OutRotation = (CustomCoordSystem.Inverse() * FRotationMatrix(Rotation) * CustomCoordSystem).Rotator();
 				}
 				else
 				{
-					switch( CurrentAxis )
-					{
-						case EAxisList::X:	InDrag = FVector( Val, 0, 0 );	break;
-						case EAxisList::Y:	InDrag = FVector( 0, Val, 0 );	break;
-						case EAxisList::Z:	InDrag = FVector( 0, 0, -Val );	break;
-						case EAxisList::XY:	InDrag = ( InDiff.X != 0 ? FVector( Val, 0, 0 ) : FVector( 0, Val, 0 ) );	break;
-					}
-
-					InDrag = InputCoordSystem.TransformPosition( InDrag );
+					// Get drag delta in widget axis space
+					OutDrag = FVector(
+						(CurrentAxis & EAxisList::X) ? FVector2D::DotProduct(XAxisDir, DragDir) : 0.0f,
+						(CurrentAxis & EAxisList::Y) ? FVector2D::DotProduct(YAxisDir, DragDir) : 0.0f,
+						(CurrentAxis & EAxisList::Z) ? FVector2D::DotProduct(ZAxisDir, DragDir) : 0.0f
+						);
+
+					// Snap to grid in widget axis space
+					const FVector GridSize = FVector(GEditor->GetGridSize());
+					FSnappingUtils::SnapPointToGrid(OutDrag, GridSize);
+
+					// Convert to effective screen space delta, and replace input delta, adjusted for inverted screen space Y axis
+					const FVector2D EffectiveDelta = OutDrag.X * XAxisDir + OutDrag.Y * YAxisDir + OutDrag.Z * ZAxisDir;
+					InOutDelta = FVector(EffectiveDelta.X, -EffectiveDelta.Y, 0.0f);
+
+					// Transform drag delta into world space
+					OutDrag = CustomCoordSystem.TransformPosition(OutDrag);
 				}
 			}
 			break;
@@ -1035,11 +991,6 @@ void GetPlaneNormalAndMask(const FVector& InAxis, FVector& OutPlaneNormal, FVect
  */
 void FWidget::AbsoluteTranslationConvertMouseMovementToAxisMovement(FSceneView* InView, FEditorViewportClient* InViewportClient, const FVector& InLocation, const FVector2D& InMousePosition, FVector& OutDrag, FRotator& OutRotation, FVector& OutScale )
 {
-	//reset all output variables
-	//OutDrag = FVector::ZeroVector;
-	//OutScale = FVector::ZeroVector;
-	//OutRotation = FRotator::ZeroRotator;
-
 	// Compute a world space ray from the screen space mouse coordinates
 	FViewportCursorLocation MouseViewportRay( InView, InViewportClient, InMousePosition.X, InMousePosition.Y );
 
@@ -1266,17 +1217,11 @@ FVector FWidget::GetAbsoluteTranslationInitialOffset(const FVector& InNewPositio
 
 
 /**
- * Returns true if we're in Local Space editing mode or editing BSP (which uses the World axes anyway
+ * Returns true if we're in Local Space editing mode
  */
 bool FWidget::IsRotationLocalSpace() const
 {
-	bool bIsLocalSpace = ( CustomCoordSystemSpace == COORD_Local );
-	//for bsp and things that don't have a "true" local space, they will always use world.  So do NOT invert.
-	if (bIsLocalSpace && CustomCoordSystem.Equals(FMatrix::Identity))
-	{
-		bIsLocalSpace = false;
-	}
-	return bIsLocalSpace;
+	return ( CustomCoordSystemSpace == COORD_Local );
 }
 
 void FWidget::UpdateDeltaRotation()
@@ -1293,8 +1238,7 @@ void FWidget::UpdateDeltaRotation()
  */
 float FWidget::GetDeltaRotation() const
 {
-	bool bIsLocalSpace = IsRotationLocalSpace();
-	return (bIsLocalSpace ? -1 : 1)*TotalDeltaRotation;
+	return TotalDeltaRotation;
 }
 
 
@@ -1316,12 +1260,13 @@ uint8 SmallOuterAlpha = 0x0f;
  * @param InColor - The color associated with the axis of rotation
  * @param InScale - Multiplier to maintain a constant screen size for rendering the widget
  */
-void FWidget::DrawRotationArc(const FSceneView* View, FPrimitiveDrawInterface* PDI, EAxisList::Type InAxis, const FVector& InLocation, const FVector& Axis0, const FVector& Axis1, const FVector& InDirectionToWidget, const FColor& InColor, const float InScale)
+void FWidget::DrawRotationArc(const FSceneView* View, FPrimitiveDrawInterface* PDI, EAxisList::Type InAxis, const FVector& InLocation, const FVector& Axis0, const FVector& Axis1, const FVector& InDirectionToWidget, const FColor& InColor, const float InScale, FVector2D& OutAxisDir)
 {
 	bool bIsPerspective = ( View->ViewMatrices.ProjMatrix.M[3][3] < 1.0f );
 	bool bIsOrtho = !bIsPerspective;
+
 	//if we're in an ortho viewport and the ring is perpendicular to the camera (both Axis0 & Axis1 are perpendicular)
-	bool bIsOrthoDrawingFullRing = bIsOrtho && (FMath::Abs(Axis0|InDirectionToWidget) < KINDA_SMALL_NUMBER) && (FMath::Abs(Axis1|InDirectionToWidget) < KINDA_SMALL_NUMBER);
+	bIsOrthoDrawingFullRing |= bIsOrtho && (FMath::Abs(Axis0|InDirectionToWidget) < KINDA_SMALL_NUMBER) && (FMath::Abs(Axis1|InDirectionToWidget) < KINDA_SMALL_NUMBER);
 
 	FColor ArcColor = InColor;
 	ArcColor.A = LargeOuterAlpha;
@@ -1331,11 +1276,12 @@ void FWidget::DrawRotationArc(const FSceneView* View, FPrimitiveDrawInterface* P
 		if ((CurrentAxis&InAxis) || (bIsOrthoDrawingFullRing))
 		{
 			float DeltaRotation = GetDeltaRotation();
-			float AbsRotation = FRotator::ClampAxis(fabs(DeltaRotation));
+			float AdjustedDeltaRotation = IsRotationLocalSpace() ? -DeltaRotation : DeltaRotation;
+			float AbsRotation = FRotator::ClampAxis(FMath::Abs(DeltaRotation));
 			float AngleOfChangeRadians (AbsRotation * PI / 180.f);
 
 			//always draw clockwise, so if we're negative we need to flip the angle
-			float StartAngle = DeltaRotation < 0.0f ? -AngleOfChangeRadians : 0.0f;
+			float StartAngle = AdjustedDeltaRotation < 0.0f ? -AngleOfChangeRadians : 0.0f;
 			float FilledAngle = AngleOfChangeRadians;
 
 			//the axis of rotation
@@ -1352,7 +1298,7 @@ void FWidget::DrawRotationArc(const FSceneView* View, FPrimitiveDrawInterface* P
 			DrawStartStopMarker(PDI, InLocation, Axis0, Axis1, 0, ArcColor, InScale);
 			//Filled Arrow
 			ArcColor.A = LargeOuterAlpha;
-			DrawStartStopMarker(PDI, InLocation, Axis0, Axis1, DeltaRotation, ArcColor, InScale);
+			DrawStartStopMarker(PDI, InLocation, Axis0, Axis1, AdjustedDeltaRotation, ArcColor, InScale);
 
 			ArcColor.A = 255;
 
@@ -1379,22 +1325,22 @@ void FWidget::DrawRotationArc(const FSceneView* View, FPrimitiveDrawInterface* P
 			//Rotate Colors to match Axis 0
 			Swap(AxisColor.R, AxisColor.G);
 			Swap(AxisColor.B, AxisColor.R);
-			AxisColor.A = (DeltaRotation == 0) ? MAX_uint8 : LargeOuterAlpha;
+			AxisColor.A = (AdjustedDeltaRotation == 0) ? MAX_uint8 : LargeOuterAlpha;
 			DrawSnapMarker(PDI, SnapLocation,  Axis0,  Axis1, AxisColor, InScale, .25f);
-			AxisColor.A = (DeltaRotation == 180.f) ? MAX_uint8 : LargeOuterAlpha;
+			AxisColor.A = (AdjustedDeltaRotation == 180.f) ? MAX_uint8 : LargeOuterAlpha;
 			DrawSnapMarker(PDI, SnapLocation, -Axis0, -Axis1, AxisColor, InScale, .25f);
 
 			//Rotate Colors to match Axis 1
 			Swap(AxisColor.R, AxisColor.G);
 			Swap(AxisColor.B, AxisColor.R);
-			AxisColor.A = (DeltaRotation == 90.f) ? MAX_uint8 : LargeOuterAlpha;
+			AxisColor.A = (AdjustedDeltaRotation == 90.f) ? MAX_uint8 : LargeOuterAlpha;
 			DrawSnapMarker(PDI, SnapLocation,  Axis1, -Axis0, AxisColor, InScale, .25f);
-			AxisColor.A = (DeltaRotation == 270.f) ? MAX_uint8 : LargeOuterAlpha;
+			AxisColor.A = (AdjustedDeltaRotation == 270.f) ? MAX_uint8 : LargeOuterAlpha;
 			DrawSnapMarker(PDI, SnapLocation, -Axis1,  Axis0, AxisColor, InScale, .25f);
 
 			if (bDragging)
 			{
-				float OffsetAngle = IsRotationLocalSpace() ? 0 : DeltaRotation;
+				float OffsetAngle = IsRotationLocalSpace() ? 0 : AdjustedDeltaRotation;
 
 				CacheRotationHUDText(View, PDI, InLocation, Axis0.RotateAngleAxis(OffsetAngle, ZAxis), Axis1.RotateAngleAxis(OffsetAngle, ZAxis), DeltaRotation, InScale);
 			}
@@ -1403,10 +1349,28 @@ void FWidget::DrawRotationArc(const FSceneView* View, FPrimitiveDrawInterface* P
 	else
 	{
 		//Reverse the axes based on camera view
-		FVector RenderAxis0 = ((Axis0|InDirectionToWidget) <= 0.0f) ? Axis0 : -Axis0;
-		FVector RenderAxis1 = ((Axis1|InDirectionToWidget) <= 0.0f) ? Axis1 : -Axis1;
+		bool bMirrorAxis0 = ((Axis0 | InDirectionToWidget) <= 0.0f);
+		bool bMirrorAxis1 = ((Axis1 | InDirectionToWidget) <= 0.0f);
+
+		FVector RenderAxis0 = bMirrorAxis0 ? Axis0 : -Axis0;
+		FVector RenderAxis1 = bMirrorAxis1 ? Axis1 : -Axis1;
+		float Direction = (bMirrorAxis0 ^ bMirrorAxis1) ? -1.0f : 1.0f;
 
 		DrawPartialRotationArc(View, PDI, InAxis, InLocation, RenderAxis0, RenderAxis1, 0, PI/2, ArcColor, InScale, InDirectionToWidget);
+
+		FVector2D Axis0ScreenLocation;
+		if (!View->ScreenToPixel(View->WorldToScreen(InLocation + RenderAxis0 * 64.0f), Axis0ScreenLocation))
+		{
+			Axis0ScreenLocation.X = Axis0ScreenLocation.Y = 0;
+		}
+
+		FVector2D Axis1ScreenLocation;
+		if (!View->ScreenToPixel(View->WorldToScreen(InLocation + RenderAxis1 * 64.0f), Axis1ScreenLocation))
+		{
+			Axis1ScreenLocation.X = Axis1ScreenLocation.Y = 0;
+		}
+
+		OutAxisDir = ((Axis1ScreenLocation - Axis0ScreenLocation) * Direction).GetSafeNormal();
 	}
 }
 
diff --git a/Engine/Source/Editor/UnrealEd/Public/MouseDeltaTracker.h b/Engine/Source/Editor/UnrealEd/Public/MouseDeltaTracker.h
--- a/Engine/Source/Editor/UnrealEd/Public/MouseDeltaTracker.h
+++ b/Engine/Source/Editor/UnrealEd/Public/MouseDeltaTracker.h
@@ -78,11 +78,11 @@ public:
 	/**
 	 * Converts the delta movement to drag/rotation/scale based on the viewport type or widget axis.
 	 */
-	void UNREALED_API ConvertMovementDeltaToDragRot(FEditorViewportClient* InViewportClient, const FVector& InDragDelta, FVector& InDrag, FRotator& InRotation, FVector& InScale) const;
+	void UNREALED_API ConvertMovementDeltaToDragRot(FEditorViewportClient* InViewportClient, FVector& InDragDelta, FVector& OutDrag, FRotator& OutRotation, FVector& OutScale) const;
 	/**
 	 * Absolute Translation conversion from mouse position on the screen to widget axis movement/rotation.
 	 */
-	void AbsoluteTranslationConvertMouseToDragRot(FSceneView* InView, FEditorViewportClient* InViewportClient, FVector& InDrag, FRotator& InRotation, FVector& InScale ) const;
+	void AbsoluteTranslationConvertMouseToDragRot(FSceneView* InView, FEditorViewportClient* InViewportClient, FVector& OutDrag, FRotator& OutRotation, FVector& OutScale ) const;
 
 	/**
 	 * Subtracts the specified value from End and EndSnapped.
diff --git a/Engine/Source/Editor/UnrealEd/Public/UnrealWidget.h b/Engine/Source/Editor/UnrealEd/Public/UnrealWidget.h
--- a/Engine/Source/Editor/UnrealEd/Public/UnrealWidget.h
+++ b/Engine/Source/Editor/UnrealEd/Public/UnrealWidget.h
@@ -44,7 +44,7 @@ public:
 	 * Draws an arrow head line for a specific axis.
 	 * @param	bCubeHead		[opt] If true, render a cube at the axis tips.  If false (the default), render a cone.
 	 */
-	void Render_Axis(const FSceneView* View, FPrimitiveDrawInterface* PDI, EAxisList::Type InAxis, FMatrix& InMatrix, UMaterialInterface* InMaterial, const FLinearColor& InColor, FVector2D& OutAxisEnd, const FVector& InScale, bool bDrawWidget, bool bCubeHead=false);
+	void Render_Axis(const FSceneView* View, FPrimitiveDrawInterface* PDI, EAxisList::Type InAxis, FMatrix& InMatrix, UMaterialInterface* InMaterial, const FLinearColor& InColor, FVector2D& OutAxisDir, const FVector& InScale, bool bDrawWidget, bool bCubeHead=false);
 
 	/**
 	 * Draws a cube
@@ -74,7 +74,7 @@ public:
 	/**
 	 * Converts mouse movement on the screen to widget axis movement/rotation.
 	 */
-	void ConvertMouseMovementToAxisMovement( FEditorViewportClient* InViewportClient, const FVector& InLocation, const FVector& InDiff, FVector& InDrag, FRotator& InRotation, FVector& InScale );
+	void ConvertMouseMovementToAxisMovement( FEditorViewportClient* InViewportClient, const FVector& InLocation, FVector& InDiff, FVector& OutDrag, FRotator& OutRotation, FVector& OutScale );
 
 	/**
 	 * Absolute Translation conversion from mouse movement on the screen to widget axis movement/rotation.
@@ -120,11 +120,22 @@ public:
 		return CurrentAxis;
 	}
 
+	/** 
+	 * @return	The widget origin in viewport space.
+	 */
 	FVector2D GetOrigin() const
 	{
 		return Origin;
 	}
 
+	/**
+	 * @return	The mouse drag start position in viewport space.
+	 */
+	void SetDragStartPosition(const FVector2D& Position)
+	{
+		DragStartPos = Position;
+	}
+
 	/**
 	 * Returns whether we are actively dragging
 	 */
@@ -260,8 +271,9 @@ private:
 	 * @param InDirectionToWidget - Direction from camera to the widget
 	 * @param InColor - The color associated with the axis of rotation
 	 * @param InScale - Multiplier to maintain a constant screen size for rendering the widget
+	 * @param OutAxisDir - Viewport-space direction of rotation arc chord is placed here
 	 */
-	void DrawRotationArc(const FSceneView* View, FPrimitiveDrawInterface* PDI, EAxisList::Type InAxis, const FVector& InLocation, const FVector& Axis0, const FVector& Axis1, const FVector& InDirectionToWidget, const FColor& InColor, const float InScale);
+	void DrawRotationArc(const FSceneView* View, FPrimitiveDrawInterface* PDI, EAxisList::Type InAxis, const FVector& InLocation, const FVector& Axis0, const FVector& Axis1, const FVector& InDirectionToWidget, const FColor& InColor, const float InScale, FVector2D& OutAxisEnd);
 
 	/**
 	 * If actively dragging, draws a ring representing the potential rotation of the selected objects, snap ticks, and "delta" markers
@@ -340,8 +352,12 @@ private:
 	/** The axis currently being moused over */
 	EAxisList::Type CurrentAxis;
 
-	/** Locations of the various points on the widget */
-	FVector2D Origin, XAxisEnd, YAxisEnd, ZAxisEnd;
+	/** Viewport space origin location of the widget */
+	FVector2D Origin;
+	/** Viewport space direction vectors of the axes on the widget */
+	FVector2D XAxisDir, YAxisDir, ZAxisDir;
+	/** Drag start position in viewport space */
+	FVector2D DragStartPos;
 
 	enum
 	{
@@ -389,6 +405,8 @@ private:
 	bool bSnapEnabled;
 	/** Default visibility for the widget if an Editor Mode Tool doesn't override it */
 	bool bDefaultVisibility;
+	/** Whether we are drawing the full ring in rotation mode (ortho viewports only) */
+	bool bIsOrthoDrawingFullRing;
 
 	/** Total delta rotation applied since the widget was dragged */
 	float TotalDeltaRotation;
