Fix FSlateUpdateableTexture::UpdateTextureThreadSafe not working for offline slate renderers.  Also added a comment clarifying the safety of the UpdateTexture functions in FSlateUpdateableTexture

diff --git a/Engine/Source/Developer/StandaloneRenderer/Private/OpenGL/SlateOpenGLTextures.h b/Engine/Source/Developer/StandaloneRenderer/Private/OpenGL/SlateOpenGLTextures.h
--- a/Engine/Source/Developer/StandaloneRenderer/Private/OpenGL/SlateOpenGLTextures.h
+++ b/Engine/Source/Developer/StandaloneRenderer/Private/OpenGL/SlateOpenGLTextures.h
@@ -32,7 +32,7 @@ public:
 	virtual FSlateShaderResource* GetSlateResource() override {return this;}
 	virtual void ResizeTexture( uint32 Width, uint32 Height ) override;
 	virtual void UpdateTexture(const TArray<uint8>& Bytes) override;
-	virtual void UpdateTextureThreadSafe(const TArray<uint8>& Bytes) {}
+	virtual void UpdateTextureThreadSafe(const TArray<uint8>& Bytes) override { UpdateTexture(Bytes); }
 private:
 	static GLuint NullTexture;
 
diff --git a/Engine/Source/Developer/StandaloneRenderer/Private/Windows/D3D/SlateD3DTextures.h b/Engine/Source/Developer/StandaloneRenderer/Private/Windows/D3D/SlateD3DTextures.h
--- a/Engine/Source/Developer/StandaloneRenderer/Private/Windows/D3D/SlateD3DTextures.h
+++ b/Engine/Source/Developer/StandaloneRenderer/Private/Windows/D3D/SlateD3DTextures.h
@@ -33,7 +33,7 @@ public:
 	virtual FSlateShaderResource* GetSlateResource() override {return this;}
 	virtual void ResizeTexture( uint32 Width, uint32 Height ) override;
 	virtual void UpdateTexture(const TArray<uint8>& Bytes) override;
-	virtual void UpdateTextureThreadSafe(const TArray<uint8>& Bytes) {}
+	virtual void UpdateTextureThreadSafe(const TArray<uint8>& Bytes) override { UpdateTexture(Bytes); }
 private:
 	/** Actual texture.  In D3D the SRV is used by the shader so our parent class holds that */
 	TRefCountPtr<ID3D11Texture2D> D3DTexture;
diff --git a/Engine/Source/Editor/UnrealEd/Private/Editor/ActorPositioning.cpp b/Engine/Source/Editor/UnrealEd/Private/Editor/ActorPositioning.cpp
--- a/Engine/Source/Editor/UnrealEd/Private/Editor/ActorPositioning.cpp
+++ b/Engine/Source/Editor/UnrealEd/Private/Editor/ActorPositioning.cpp
@@ -207,15 +207,17 @@ FTransform FActorPositioning::GetSurfaceAlignedTransform(const FPositioningData&
 
 FTransform FActorPositioning::GetSnappedSurfaceAlignedTransform(const FSnappedPositioningData& Data)
 {
-	FVector SnappedLocation = Data.SurfaceLocation;
-	FSnappingUtils::SnapPointToGrid(SnappedLocation, FVector(0.f));
+	FSnappedPositioningData SnappedData(Data);
+
+	//FVector SnappedLocation = Data.SurfaceLocation;
+	FSnappingUtils::SnapPointToGrid(SnappedData.SurfaceLocation, FVector(0.f));
 
 	// Secondly, attempt vertex snapping
 	FVector AlignToNormal;
-	if (!Data.LevelViewportClient || !FSnappingUtils::SnapLocationToNearestVertex( SnappedLocation, Data.LevelViewportClient->GetDropPreviewLocation(), Data.LevelViewportClient, AlignToNormal, Data.bDrawSnapHelpers ))
+	if (!Data.LevelViewportClient || !FSnappingUtils::SnapLocationToNearestVertex( SnappedData.SurfaceLocation, SnappedData.LevelViewportClient->GetDropPreviewLocation(), SnappedData.LevelViewportClient, AlignToNormal, SnappedData.bDrawSnapHelpers ))
 	{
-		AlignToNormal = Data.SurfaceNormal;
+		AlignToNormal = SnappedData.SurfaceNormal;
 	}
 
-	return GetSurfaceAlignedTransform(Data);
+	return GetSurfaceAlignedTransform(SnappedData);
 }
\ No newline at end of file
diff --git a/Engine/Source/Editor/UnrealEd/Private/LevelEditorViewport.cpp b/Engine/Source/Editor/UnrealEd/Private/LevelEditorViewport.cpp
--- a/Engine/Source/Editor/UnrealEd/Private/LevelEditorViewport.cpp
+++ b/Engine/Source/Editor/UnrealEd/Private/LevelEditorViewport.cpp
@@ -1070,9 +1070,16 @@ bool FLevelEditorViewportClient::UpdateDropPreviewActors(int32 MouseX, int32 Mou
 			.UseStartTransform(PreDragActorTransforms.FindRef(Actor))
 			.UsePlacementExtent(Actor->GetPlacementExtent());
 
-		const FTransform ActorTransform = FActorPositioning::GetSnappedSurfaceAlignedTransform(PositioningData);
+		FActorOrComponent ActorOrComponent( Actor );
 
-		Actor->SetActorTransform(ActorTransform);
+		if( DroppedOnActor )
+		{
+			GEditor->SnapObjectTo( FActorOrComponent( DroppedOnActor ), true, true, true, false, ActorOrComponent );
+		}
+
+		//const FTransform ActorTransform = FActorPositioning::GetSurfaceAlignedTransform(PositioningData);
+
+		//Actor->SetActorTransform(ActorTransform);
 		Actor->SetIsTemporarilyHiddenInEditor(false);
 		Actor->PostEditMove(false);
 	}
diff --git a/Engine/Source/Runtime/SlateCore/Public/Textures/SlateUpdatableTexture.h b/Engine/Source/Runtime/SlateCore/Public/Textures/SlateUpdatableTexture.h
--- a/Engine/Source/Runtime/SlateCore/Public/Textures/SlateUpdatableTexture.h
+++ b/Engine/Source/Runtime/SlateCore/Public/Textures/SlateUpdatableTexture.h
@@ -37,7 +37,8 @@ public:
 	virtual void ResizeTexture( uint32 Width, uint32 Height ) = 0;
 
 	/**
-	 * Updates the texture contents via a byte array
+	 * Updates the texture contents via a byte array. 
+	 * Note: This method is not thread safe so make sure you do not use the Bytes data on another after it is passed in
 	 *
 	 * @param Bytes Array of texture data
 	 */
