Made the navigation notification not require clicking OK to dismiss

diff --git a/Engine/Source/Editor/UnrealEd/Private/NavigationBuildingNotification.cpp b/Engine/Source/Editor/UnrealEd/Private/NavigationBuildingNotification.cpp
--- a/Engine/Source/Editor/UnrealEd/Private/NavigationBuildingNotification.cpp
+++ b/Engine/Source/Editor/UnrealEd/Private/NavigationBuildingNotification.cpp
@@ -74,14 +74,10 @@ void FNavigationBuildingNotificationImpl::BuildFinished()
 	if (EEngine != NULL && (FEditorBuildUtils::IsBuildingNavigationFromUserRequest()))
 	{
 		FNotificationInfo Info( NSLOCTEXT("NavigationBuild", "NavigationBuildDoneMessage", "Navigation building completed.") );
-		Info.bFireAndForget = false;
+		Info.bFireAndForget = true;
 		Info.bUseThrobber = false;
-		Info.FadeOutDuration = 0.0f;
-		Info.ExpireDuration = 0.0f;
-		Info.ButtonDetails.Add(FNotificationButtonInfo(
-			NSLOCTEXT("NavigationBuild", "NavigationBuildOk","Ok"),
-			FText(),
-			FSimpleDelegate::CreateRaw(this, &FNavigationBuildingNotificationImpl::ClearCompleteNotification)));
+		Info.FadeOutDuration = 3.0f;
+		Info.ExpireDuration = 3.0f;
 
 		NavigationBuiltCompleteNotification = FSlateNotificationManager::Get().AddNotification(Info);
 		if (NavigationBuiltCompleteNotification.IsValid())
