Fixed issues with BSP surface texture transform when the parent brush is scaled or rotated.

- Fixed a bug in FPoly::Scale where the texture was being scaled in the wrong direction.
- Corrected UEditorEngine::polyUpdateMaster so that the "master poly" Base (UV offset) and UV directions take into account any transform on the parent Brush.

diff --git a/Engine/Source/Editor/UnrealEd/Private/EditorBsp.cpp b/Engine/Source/Editor/UnrealEd/Private/EditorBsp.cpp
--- a/Engine/Source/Editor/UnrealEd/Private/EditorBsp.cpp
+++ b/Engine/Source/Editor/UnrealEd/Private/EditorBsp.cpp
@@ -1153,6 +1153,12 @@ int UEditorEngine::bspBrushCSG
 	if( ReallyBig ) GWarn->StatusUpdate(0,0, NSLOCTEXT("UnrealEd", "Transforming", "Transforming"));
 
 	UMaterialInterface* SelectedMaterialInstance = GetSelectedObjects()->GetTop<UMaterialInterface>();
+
+	const FVector PrePivot = Actor->GetPrePivot();
+	const FVector Scale = Actor->GetActorScale();
+	const FRotator Rotation = Actor->GetActorRotation();
+	const FVector Location = Actor->GetActorLocation();
+
 	for( i=0; i<Brush->Polys->Element.Num(); i++ )
 	{
 		FPoly& CurrentPoly = Brush->Polys->Element[i];
@@ -1183,9 +1189,9 @@ int UEditorEngine::bspBrushCSG
 			DestEdPoly.iLink = i;
 
 		// Transform it.
-		DestEdPoly.Scale( Actor->GetPrePivot(), Actor->GetActorScale() );
-		DestEdPoly.Rotate( Actor->GetPrePivot(), Actor->GetActorRotation() );
-		DestEdPoly.Transform( Actor->GetPrePivot(), Actor->GetActorLocation() );
+		DestEdPoly.Scale( PrePivot, Scale );
+		DestEdPoly.Rotate( PrePivot, Rotation );
+		DestEdPoly.Transform( PrePivot, Location );
 
 		// Add poly to the temp model.
 		new(TempModel->Polys->Element)FPoly( DestEdPoly );
diff --git a/Engine/Source/Editor/UnrealEd/Private/EditorCsg.cpp b/Engine/Source/Editor/UnrealEd/Private/EditorCsg.cpp
--- a/Engine/Source/Editor/UnrealEd/Private/EditorCsg.cpp
+++ b/Engine/Source/Editor/UnrealEd/Private/EditorCsg.cpp
@@ -353,24 +353,30 @@ void UEditorEngine::polyUpdateMaster
 )
 {
 	FBspSurf &Surf = Model->Surfs[iSurf];
-	if( !Surf.Actor )
+	ABrush* Actor = Surf.Actor;
+	if( !Actor )
 		return;
 
-	for( int32 iEdPoly = Surf.iBrushPoly; iEdPoly < Surf.Actor->Brush->Polys->Element.Num(); iEdPoly++ )
+	const FVector ActorLocation = Actor->GetActorLocation();
+	const FVector ActorPrePivot = Actor->GetPrePivot();
+	const FVector ActorScale = Actor->GetActorScale();
+	const FRotator ActorRotation = -Actor->GetActorRotation();
+
+	for( int32 iEdPoly = Surf.iBrushPoly; iEdPoly < Actor->Brush->Polys->Element.Num(); iEdPoly++ )
 	{
-		FPoly& MasterEdPoly = Surf.Actor->Brush->Polys->Element[iEdPoly];
+		FPoly& MasterEdPoly = Actor->Brush->Polys->Element[iEdPoly];
 		if( iEdPoly==Surf.iBrushPoly || MasterEdPoly.iLink==Surf.iBrushPoly )
 		{
-			Surf.Actor->Brush->Polys->Element.ModifyItem( iEdPoly );
+			Actor->Brush->Polys->Element.ModifyItem( iEdPoly );
 
 			MasterEdPoly.Material  = Surf.Material;
 			MasterEdPoly.PolyFlags = Surf.PolyFlags & ~(PF_NoEdit);
 
 			if( UpdateTexCoords )
 			{
-				MasterEdPoly.Base = (Model->Points[Surf.pBase] - Surf.Actor->GetActorLocation()) + Surf.Actor->GetPrePivot();
-				MasterEdPoly.TextureU = Model->Vectors[Surf.vTextureU];
-				MasterEdPoly.TextureV = Model->Vectors[Surf.vTextureV];
+				MasterEdPoly.Base = ActorRotation.RotateVector(Model->Points[Surf.pBase] - ActorLocation) / ActorScale + ActorPrePivot;
+				MasterEdPoly.TextureU = ActorRotation.RotateVector(Model->Vectors[Surf.vTextureU]) * ActorScale;
+				MasterEdPoly.TextureV = ActorRotation.RotateVector(Model->Vectors[Surf.vTextureV]) * ActorScale;
 			}
 		}
 	}
diff --git a/Engine/Source/Runtime/Engine/Private/Polygon.cpp b/Engine/Source/Runtime/Engine/Private/Polygon.cpp
--- a/Engine/Source/Runtime/Engine/Private/Polygon.cpp
+++ b/Engine/Source/Runtime/Engine/Private/Polygon.cpp
@@ -501,8 +501,8 @@ void FPoly::Scale
 	Base = PreSubtract + (Base - PreSubtract) * Scale;
 
 	// Scale the texture vectors.
-	TextureU *= Scale;
-	TextureV *= Scale;
+	TextureU /= Scale;
+	TextureV /= Scale;
 
 	// Recalculate the normal. Non-uniform scale or mirroring requires this.
 	CalcNormal(true);
