Back out mistakenly checked in files

diff --git a/Engine/Source/Editor/UnrealEd/Private/Editor/ActorPositioning.cpp b/Engine/Source/Editor/UnrealEd/Private/Editor/ActorPositioning.cpp
--- a/Engine/Source/Editor/UnrealEd/Private/Editor/ActorPositioning.cpp
+++ b/Engine/Source/Editor/UnrealEd/Private/Editor/ActorPositioning.cpp
@@ -207,17 +207,15 @@ FTransform FActorPositioning::GetSurfaceAlignedTransform(const FPositioningData&
 
 FTransform FActorPositioning::GetSnappedSurfaceAlignedTransform(const FSnappedPositioningData& Data)
 {
-	FSnappedPositioningData SnappedData(Data);
-
-	//FVector SnappedLocation = Data.SurfaceLocation;
-	FSnappingUtils::SnapPointToGrid(SnappedData.SurfaceLocation, FVector(0.f));
+	FVector SnappedLocation = Data.SurfaceLocation;
+	FSnappingUtils::SnapPointToGrid(SnappedLocation, FVector(0.f));
 
 	// Secondly, attempt vertex snapping
 	FVector AlignToNormal;
-	if (!Data.LevelViewportClient || !FSnappingUtils::SnapLocationToNearestVertex( SnappedData.SurfaceLocation, SnappedData.LevelViewportClient->GetDropPreviewLocation(), SnappedData.LevelViewportClient, AlignToNormal, SnappedData.bDrawSnapHelpers ))
+	if (!Data.LevelViewportClient || !FSnappingUtils::SnapLocationToNearestVertex( SnappedLocation, Data.LevelViewportClient->GetDropPreviewLocation(), Data.LevelViewportClient, AlignToNormal, Data.bDrawSnapHelpers ))
 	{
-		AlignToNormal = SnappedData.SurfaceNormal;
+		AlignToNormal = Data.SurfaceNormal;
 	}
 
-	return GetSurfaceAlignedTransform(SnappedData);
+	return GetSurfaceAlignedTransform(Data);
 }
\ No newline at end of file
diff --git a/Engine/Source/Editor/UnrealEd/Private/LevelEditorViewport.cpp b/Engine/Source/Editor/UnrealEd/Private/LevelEditorViewport.cpp
--- a/Engine/Source/Editor/UnrealEd/Private/LevelEditorViewport.cpp
+++ b/Engine/Source/Editor/UnrealEd/Private/LevelEditorViewport.cpp
@@ -1070,16 +1070,9 @@ bool FLevelEditorViewportClient::UpdateDropPreviewActors(int32 MouseX, int32 Mou
 			.UseStartTransform(PreDragActorTransforms.FindRef(Actor))
 			.UsePlacementExtent(Actor->GetPlacementExtent());
 
-		FActorOrComponent ActorOrComponent( Actor );
+		const FTransform ActorTransform = FActorPositioning::GetSnappedSurfaceAlignedTransform(PositioningData);
 
-		if( DroppedOnActor )
-		{
-			GEditor->SnapObjectTo( FActorOrComponent( DroppedOnActor ), true, true, true, false, ActorOrComponent );
-		}
-
-		//const FTransform ActorTransform = FActorPositioning::GetSurfaceAlignedTransform(PositioningData);
-
-		//Actor->SetActorTransform(ActorTransform);
+		Actor->SetActorTransform(ActorTransform);
 		Actor->SetIsTemporarilyHiddenInEditor(false);
 		Actor->PostEditMove(false);
 	}
