Launch on was being very conservitive when selecting not to use cook in editor.
Also switched over to using correct platform ini file for checking mobilehdr setting.

diff --git a/Engine/Source/Editor/UnrealEd/Private/UnrealEdEngine.cpp b/Engine/Source/Editor/UnrealEd/Private/UnrealEdEngine.cpp
--- a/Engine/Source/Editor/UnrealEd/Private/UnrealEdEngine.cpp
+++ b/Engine/Source/Editor/UnrealEd/Private/UnrealEdEngine.cpp
@@ -148,12 +148,18 @@ bool CanCookForPlatformInThisProcess( const FString& PlatformName )
 	// hack remove this hack when we properly support changing the mobileHDR setting 
 	// check if our mobile hdr setting in memory is different from the one which is saved in the config file
 	
+	
+	FConfigFile PlatformEngineIni;
+	GConfig->LoadLocalIniFile(PlatformEngineIni, TEXT("Engine"), true, *PlatformName );
+
+	FString IniValueString;
 	bool ConfigSetting = false;
-	if ( !GConfig->GetBool( TEXT("/Script/Engine.RendererSettings"), TEXT("r.MobileHDR"), ConfigSetting, GEngineIni) )
+	if ( PlatformEngineIni.GetString( TEXT("/Script/Engine.RendererSettings"), TEXT("r.MobileHDR"), IniValueString ) == false )
 	{
-		// if we can't get the config setting then don't risk it
-		return false;
+		// must always match the RSetting setting because we don't have a config setting
+		return true; 
 	}
+	ConfigSetting = IniValueString.ToBool();
 
 	// this was stolen from void IsMobileHDR()
 	static auto* MobileHDRCvar = IConsoleManager::Get().FindTConsoleVariableDataInt(TEXT("r.MobileHDR"));
@@ -161,6 +167,7 @@ bool CanCookForPlatformInThisProcess( const FString& PlatformName )
 
 	if ( CurrentRSetting != ConfigSetting )
 	{
+		UE_LOG(LogUnrealEdEngine, Warning, TEXT("Unable to use cook in editor because r.MobileHDR from Engine ini doesn't match console value r.MobileHDR"));
 		return false;
 	}
 	////////////////////////////////////////
