GitHub pull request 918

Update FbxMainExport.cpp, export selection as .fbx crash fixed

diff --git a/Engine/Source/Editor/UnrealEd/Private/Fbx/FbxMainExport.cpp b/Engine/Source/Editor/UnrealEd/Private/Fbx/FbxMainExport.cpp
--- a/Engine/Source/Editor/UnrealEd/Private/Fbx/FbxMainExport.cpp
+++ b/Engine/Source/Editor/UnrealEd/Private/Fbx/FbxMainExport.cpp
@@ -690,7 +690,7 @@ void FFbxExporter::ExportBSP( UModel* Model, bool bSelectedOnly )
 
 		ABrush* BrushActor = Surf.Actor;
 
-		if( (Surf.PolyFlags & PF_Selected) || !bSelectedOnly || (BrushActor && BrushActor->IsSelected() ) )
+		if( (Surf.PolyFlags & PF_Selected) || !bSelectedOnly || (BrushActor && BrushActor->IsSelected() ) && Node.NumVertices >= 3 )
 		{
 			FPoly Poly;
 			GEditor->polyFindMaster( Model, Node.iSurf, Poly );
