Editor: Add some additional methods to IAssetEditorInstance to support reasoning about their owning tab manager and last activation time

diff --git a/Engine/Source/Editor/UnrealEd/Private/MiniCurveEditor.cpp b/Engine/Source/Editor/UnrealEd/Private/MiniCurveEditor.cpp
--- a/Engine/Source/Editor/UnrealEd/Private/MiniCurveEditor.cpp
+++ b/Engine/Source/Editor/UnrealEd/Private/MiniCurveEditor.cpp
@@ -76,3 +76,14 @@ bool SMiniCurveEditor::CloseWindow()
 	return true;
 }
 
+TSharedPtr<class FTabManager> SMiniCurveEditor::GetAssociatedTabManager()
+{
+	//@TODO: This editor should probably derive from FAssetEditorToolkit instead!
+	return TSharedPtr<class FTabManager>();
+}
+
+double SMiniCurveEditor::GetLastActivationTime()
+{
+	//@TODO: This editor should probably derive from FAssetEditorToolkit instead!
+	return 0.0;
+}
diff --git a/Engine/Source/Editor/UnrealEd/Private/Toolkits/AssetEditorToolkit.cpp b/Engine/Source/Editor/UnrealEd/Private/Toolkits/AssetEditorToolkit.cpp
--- a/Engine/Source/Editor/UnrealEd/Private/Toolkits/AssetEditorToolkit.cpp
+++ b/Engine/Source/Editor/UnrealEd/Private/Toolkits/AssetEditorToolkit.cpp
@@ -327,6 +327,27 @@ void FAssetEditorToolkit::InvokeTab(const FTabId& TabId)
 	GetTabManager()->InvokeTab(TabId);
 }
 
+TSharedPtr<class FTabManager> FAssetEditorToolkit::GetAssociatedTabManager()
+{
+	return TabManager;
+}
+
+double FAssetEditorToolkit::GetLastActivationTime()
+{
+	double MostRecentTime = 0.0;
+
+	if (TabManager.IsValid())
+	{
+		TSharedPtr<SDockTab> OwnerTab = TabManager->GetOwnerTab();
+		if (OwnerTab.IsValid())
+		{
+			MostRecentTime = OwnerTab->GetLastActivationTime();
+		}
+	}
+
+	return MostRecentTime;
+}
+
 TSharedPtr< IToolkitHost > FAssetEditorToolkit::GetPreviousWorldCentricToolkitHost()
 {
 	return PreviousWorldCentricToolkitHost.Pin();
diff --git a/Engine/Source/Editor/UnrealEd/Public/MiniCurveEditor.h b/Engine/Source/Editor/UnrealEd/Public/MiniCurveEditor.h
--- a/Engine/Source/Editor/UnrealEd/Public/MiniCurveEditor.h
+++ b/Engine/Source/Editor/UnrealEd/Public/MiniCurveEditor.h
@@ -1,8 +1,7 @@
 // Copyright 1998-2015 Epic Games, Inc. All Rights Reserved.
 
+#pragma once
 
-#ifndef __MiniCurveEditor_h__
-#define __MiniCurveEditor_h__
 #include "Toolkits/AssetEditorManager.h"
 
 class UNREALED_API SMiniCurveEditor :  public SCompoundWidget,public IAssetEditorInstance
@@ -27,6 +26,8 @@ public:
 	virtual bool CloseWindow() override;
 	virtual bool IsPrimaryEditor() const override { return true; };
 	virtual void InvokeTab(const struct FTabId& TabId) override {}
+	virtual TSharedPtr<class FTabManager> GetAssociatedTabManager() override;
+	virtual double GetLastActivationTime() override;
 
 private:
 
@@ -46,5 +47,3 @@ private:
 protected:
 	TWeakPtr<SWindow> WidgetWindow;
 };
-
-#endif // MiniCurveEditor
\ No newline at end of file
diff --git a/Engine/Source/Editor/UnrealEd/Public/Toolkits/AssetEditorManager.h b/Engine/Source/Editor/UnrealEd/Public/Toolkits/AssetEditorManager.h
--- a/Engine/Source/Editor/UnrealEd/Public/Toolkits/AssetEditorManager.h
+++ b/Engine/Source/Editor/UnrealEd/Public/Toolkits/AssetEditorManager.h
@@ -1,8 +1,6 @@
 // Copyright 1998-2015 Epic Games, Inc. All Rights Reserved.
 
-
-#ifndef __AssetEditorManager_h__
-#define __AssetEditorManager_h__
+#pragma once
 
 #include "IToolkit.h"	// For EToolkitMode
 #include "Ticker.h" 	// For automation to request assets to load
@@ -21,6 +19,8 @@ public:
 	virtual bool CloseWindow() = 0;
 	virtual bool IsPrimaryEditor() const = 0;
 	virtual void InvokeTab(const struct FTabId& TabId) = 0;
+	virtual TSharedPtr<class FTabManager> GetAssociatedTabManager() = 0;
+	virtual double GetLastActivationTime() = 0;
 };
 
 
@@ -29,8 +29,7 @@ public:
  *
  * @todo toolkit: Merge this functionality into FToolkitManager
  */
-class UNREALED_API FAssetEditorManager
-	: public FGCObject
+class UNREALED_API FAssetEditorManager : public FGCObject
 {
 public:
 
@@ -44,14 +43,14 @@ public:
 	* Tries to open an editor for the specified asset.  Returns true if the asset is open in an editor.
 	* If the file is already open in an editor, it will not create another editor window but instead bring it to front
 	*/
-	bool OpenEditorForAsset( UObject* Asset, const EToolkitMode::Type ToolkitMode = EToolkitMode::Standalone, TSharedPtr< class IToolkitHost > OpenedFromLevelEditor = TSharedPtr< IToolkitHost >() );
+	bool OpenEditorForAsset(UObject* Asset, const EToolkitMode::Type ToolkitMode = EToolkitMode::Standalone, TSharedPtr< class IToolkitHost > OpenedFromLevelEditor = TSharedPtr< IToolkitHost >());
 
 	/** 
 	* Tries to open an editor for all of the specified assets. 
 	* If any of the assets are already open, it will not create a new editor for them.
 	* If all assets are of the same type, the supporting AssetTypeAction (if it exists) is responsible for the details of how to handle opening multiple assets at once.
 	*/
-	bool OpenEditorForAssets( const TArray< UObject* >& Assets, const EToolkitMode::Type ToolkitMode = EToolkitMode::Standalone, TSharedPtr< class IToolkitHost > OpenedFromLevelEditor = TSharedPtr< IToolkitHost >() );
+	bool OpenEditorForAssets(const TArray< UObject* >& Assets, const EToolkitMode::Type ToolkitMode = EToolkitMode::Standalone, TSharedPtr< class IToolkitHost > OpenedFromLevelEditor = TSharedPtr< IToolkitHost >());
 
 	/** Opens editors for the supplied assets (via OpenEditorForAsset) */
 	void OpenEditorsForAssets(const TArray<FString>& AssetsToOpen);
@@ -84,7 +83,7 @@ public:
 	void NotifyEditorClosed(IAssetEditorInstance* Instance);
 
 	// FGCObject interface
-	virtual void AddReferencedObjects( FReferenceCollector& Collector ) override;
+	virtual void AddReferencedObjects(FReferenceCollector& Collector) override;
 
 	/** Close all open asset editors */
 	bool CloseAllAssetEditors();
@@ -109,7 +108,7 @@ private:
 private:
 
 	/** Handles FAssetEditorRequestOpenAsset messages. */
-	void HandleRequestOpenAssetMessage( const FAssetEditorRequestOpenAsset& Message, const IMessageContextRef& Context );
+	void HandleRequestOpenAssetMessage(const FAssetEditorRequestOpenAsset& Message, const IMessageContextRef& Context);
 
 	/** Opens an asset by path */
 	void OpenEditorForAsset(const FString& AssetPathName);
@@ -187,6 +186,3 @@ private:
 	/** A pointer to the notification used by RestorePreviouslyOpenAssets */
 	TWeakPtr<SNotificationItem> RestorePreviouslyOpenAssetsNotificationPtr;
 };
-
-
-#endif
diff --git a/Engine/Source/Editor/UnrealEd/Public/Toolkits/AssetEditorToolkit.h b/Engine/Source/Editor/UnrealEd/Public/Toolkits/AssetEditorToolkit.h
--- a/Engine/Source/Editor/UnrealEd/Public/Toolkits/AssetEditorToolkit.h
+++ b/Engine/Source/Editor/UnrealEd/Public/Toolkits/AssetEditorToolkit.h
@@ -68,6 +68,8 @@ public:
 	virtual bool CloseWindow() override;
 	virtual bool IsPrimaryEditor() const override { return true; };
 	virtual void InvokeTab(const struct FTabId& TabId) override;
+	virtual TSharedPtr<class FTabManager> GetAssociatedTabManager() override;
+	virtual double GetLastActivationTime() override;
 
 	/**
 	 * Fills in the supplied menu with commands for working with this asset file
