Editor: Add a Control+Tab menu to quickly switch between different open assets or activate their associated tool tabs
- Added SGlobalTabSwitchingDialog (summoned by either Control+Tab or Command+Tab by default)
- Split SGlobalOpenAssetDialog out into separate files
- Fixed an issue where Ctrl+P would not always be able to create a window (when there is no known global active window to determine preferred screen location from; now defaults to the monitor with the main frame if necessary)

#epicfriday [UE-1302] [UE-960]

diff --git a/Engine/Source/Editor/EditorStyle/Private/SlateEditorStyle.cpp b/Engine/Source/Editor/EditorStyle/Private/SlateEditorStyle.cpp
--- a/Engine/Source/Editor/EditorStyle/Private/SlateEditorStyle.cpp
+++ b/Engine/Source/Editor/EditorStyle/Private/SlateEditorStyle.cpp
@@ -723,6 +723,8 @@ void FSlateEditorStyle::FStyle::SetupGeneralStyles()
 			.SetNormal(FSlateNoResource() )
 			.SetPressed(FSlateNoResource() )
 			.SetHovered(BORDER_BRUSH( "Old/HyperlinkUnderline", FMargin(0,0,0,3/16.0f) ) );
+		Set("HoverOnlyHyperlinkButton", HoverOnlyHyperlinkButton);
+
 		FHyperlinkStyle HoverOnlyHyperlink = FHyperlinkStyle()
 			.SetUnderlineStyle(HoverOnlyHyperlinkButton)
 			.SetTextStyle(NormalText)
@@ -1885,6 +1887,21 @@ void FSlateEditorStyle::FStyle::SetupGeneralStyles()
 		/* ... and add new style */
 		Set( "ToolBar.CheckBox", ToolBarCheckBoxStyle );
 
+
+		/* Read-only checkbox that appears next to a menu item */
+		/* Set images for various SCheckBox states associated with read-only toolbar check box items... */
+		const FCheckBoxStyle BasicToolBarCheckStyle = FCheckBoxStyle()
+			.SetUncheckedImage(IMAGE_BRUSH("Icons/Empty_14x", Icon14x14))
+			.SetUncheckedHoveredImage(IMAGE_BRUSH("Icons/Empty_14x", Icon14x14))
+			.SetUncheckedPressedImage(IMAGE_BRUSH("Common/SmallCheckBox_Hovered", Icon14x14))
+			.SetCheckedImage(IMAGE_BRUSH("Common/SmallCheck", Icon14x14))
+			.SetCheckedHoveredImage(IMAGE_BRUSH("Common/SmallCheck", Icon14x14))
+			.SetCheckedPressedImage(IMAGE_BRUSH("Common/SmallCheck", Icon14x14))
+			.SetUndeterminedImage(IMAGE_BRUSH("Icons/Empty_14x", Icon14x14))
+			.SetUndeterminedHoveredImage(FSlateNoResource())
+			.SetUndeterminedPressedImage(FSlateNoResource());
+		Set( "ToolBar.Check", BasicToolBarCheckStyle );
+
 		// This radio button is actually just a check box with different images
 		/* Create style for "ToolBar.RadioButton" widget ... */
 		const FCheckBoxStyle ToolbarRadioButtonCheckBoxStyle = FCheckBoxStyle()
@@ -1925,6 +1942,32 @@ void FSlateEditorStyle::FStyle::SetupGeneralStyles()
 		Set( "ToolBar.Button.Checked_Pressed", new BOX_BRUSH( "Common/RoundedSelection_16x",  4.0f/16.0f, SelectionColor ) );
 	}
 
+	// Ctrl+Tab menu
+	{
+		Set("ControlTabMenu.Background", new BOX_BRUSH("Old/Menu_Background", FMargin(8.0f / 64.0f)));
+
+		Set("ControlTabMenu.HeadingStyle",
+			FTextBlockStyle(NormalText)
+			.SetFont(TTF_CORE_FONT("Fonts/Roboto-Bold", 14))
+			.SetColorAndOpacity(FLinearColor::White)
+			);
+
+		Set("ControlTabMenu.AssetTypeStyle",
+			FTextBlockStyle(NormalText)
+			.SetColorAndOpacity(FLinearColor::White)
+			);
+
+		Set("ControlTabMenu.AssetPathStyle",
+			FTextBlockStyle(NormalText)
+			.SetColorAndOpacity(FLinearColor::White)
+			);
+
+		Set("ControlTabMenu.AssetNameStyle",
+			FTextBlockStyle(NormalText)
+			.SetFont(TTF_CORE_FONT("Fonts/Roboto-Regular", 14))
+			.SetColorAndOpacity(FLinearColor::White)
+			);
+	}
 
 	// MenuBar
 	{
diff --git a/Engine/Source/Editor/UnrealEd/Private/Toolkits/GlobalEditorCommonCommands.cpp b/Engine/Source/Editor/UnrealEd/Private/Toolkits/GlobalEditorCommonCommands.cpp
--- a/Engine/Source/Editor/UnrealEd/Private/Toolkits/GlobalEditorCommonCommands.cpp
+++ b/Engine/Source/Editor/UnrealEd/Private/Toolkits/GlobalEditorCommonCommands.cpp
@@ -2,93 +2,18 @@
 
 #include "UnrealEd.h"
 #include "GlobalEditorCommonCommands.h"
-
-//////////////////////////////////////////////////////////////////////////
-// SGlobalOpenAssetDialog
-
-#include "AssetData.h"
-#include "ContentBrowserModule.h"
-#include "Toolkits/AssetEditorManager.h"
 #include "OutputLogModule.h"
+#include "LevelEditor.h"
+#include "SDockTab.h"
 
-class SGlobalOpenAssetDialog : public SCompoundWidget
-{
-public:
-	SLATE_BEGIN_ARGS(SGlobalOpenAssetDialog){}
-	SLATE_END_ARGS()
-
-public:
-	void Construct(const FArguments& InArgs, FVector2D InSize)
-	{
-		FContentBrowserModule& ContentBrowserModule = FModuleManager::Get().LoadModuleChecked<FContentBrowserModule>(TEXT("ContentBrowser"));
-
-		FAssetPickerConfig AssetPickerConfig;
-		AssetPickerConfig.OnAssetDoubleClicked = FOnAssetSelected::CreateSP(this, &SGlobalOpenAssetDialog::OnAssetSelectedFromPicker);
-		AssetPickerConfig.OnAssetEnterPressed = FOnAssetEnterPressed::CreateSP(this, &SGlobalOpenAssetDialog::OnPressedEnterOnAssetsInPicker);
-		AssetPickerConfig.InitialAssetViewType = EAssetViewType::List;
-		AssetPickerConfig.bAllowNullSelection = false;
-		AssetPickerConfig.bShowBottomToolbar = true;
-		AssetPickerConfig.bAutohideSearchBar = false;
-		AssetPickerConfig.bCanShowClasses = false;
-		AssetPickerConfig.SaveSettingsName = TEXT("GlobalAssetPicker");
-
-		ChildSlot
-		[
-			SNew(SBox)
-			.WidthOverride(InSize.X)
-			.HeightOverride(InSize.Y)
-			[
-				SNew(SVerticalBox)
-				+SVerticalBox::Slot()
-				.FillHeight(1.0f)
-				[
-					ContentBrowserModule.Get().CreateAssetPicker(AssetPickerConfig)
-				]
-			]
-		];
-	}
-
-	// SWidget interface
-	virtual FReply OnPreviewKeyDown(const FGeometry& MyGeometry, const FKeyEvent& InKeyEvent) override;
-	// End of SWidget interface
+#include "SGlobalOpenAssetDialog.h"
+#include "SGlobalTabSwitchingDialog.h"
 
-protected:
-	void OnAssetSelectedFromPicker(const FAssetData& AssetData)
-	{
-		if (UObject* ObjectToEdit = AssetData.GetAsset())
-		{
-			FAssetEditorManager::Get().OpenEditorForAsset(ObjectToEdit);
-		}
-	}
-
-	void OnPressedEnterOnAssetsInPicker(const TArray<FAssetData>& SelectedAssets)
-	{
-		for (auto AssetIt = SelectedAssets.CreateConstIterator(); AssetIt; ++AssetIt)
-		{
-			if (UObject* ObjectToEdit = AssetIt->GetAsset())
-			{
-				FAssetEditorManager::Get().OpenEditorForAsset(ObjectToEdit);
-			}
-		}
-	}
-};
-
-FReply SGlobalOpenAssetDialog::OnPreviewKeyDown(const FGeometry& MyGeometry, const FKeyEvent& InKeyEvent)
-{
-	if (InKeyEvent.GetKey() == EKeys::Escape)
-	{
-		FSlateApplication::Get().DismissAllMenus();
-		return FReply::Handled();
-	}
-
-	return FReply::Unhandled();
-}
+#define LOCTEXT_NAMESPACE "GlobalEditorCommonCommands"
 
 //////////////////////////////////////////////////////////////////////////
 // FGlobalEditorCommonCommands
 
-#define LOCTEXT_NAMESPACE ""
-
 FGlobalEditorCommonCommands::FGlobalEditorCommonCommands()
 	: TCommands<FGlobalEditorCommonCommands>(TEXT("SystemWideCommands"), NSLOCTEXT("Contexts", "SystemWideCommands", "System-wide"), NAME_None, FEditorStyle::GetStyleSetName())
 {
@@ -96,7 +21,9 @@ FGlobalEditorCommonCommands::FGlobalEditorCommonCommands()
 
 void FGlobalEditorCommonCommands::RegisterCommands()
 {
-	//UI_COMMAND( SummonControlTabNavigation, "Tab Navigation", "Summons a list of open assets and tabs", EUserInterfaceActionType::Button, FInputGesture(EModifierKey::Control, EKeys::Tab) );
+	UI_COMMAND( SummonControlTabNavigation, "Tab Navigation", "Summons a list of open assets and tabs", EUserInterfaceActionType::Button, FInputGesture(EModifierKey::Control, EKeys::Tab) );
+	UI_COMMAND( SummonControlTabNavigationAlternate, "Tab Navigation", "Summons a list of open assets and tabs", EUserInterfaceActionType::Button, FInputGesture(EModifierKey::Command, EKeys::Tab) );
+
 	UI_COMMAND( SummonOpenAssetDialog, "Open Asset...", "Summons an asset picker", EUserInterfaceActionType::Button, FInputGesture(EModifierKey::Control, EKeys::P) );
 	UI_COMMAND( SummonOpenAssetDialogAlternate, "Open Asset...", "Summons an asset picker", EUserInterfaceActionType::Button, FInputGesture(EModifierKey::Alt | EModifierKey::Shift, EKeys::O));
 	UI_COMMAND( FindInContentBrowser, "Find in Content Browser", "Summons the Content Browser and navigates to the selected asset", EUserInterfaceActionType::Button, FInputGesture(EModifierKey::Control, EKeys::B));
@@ -111,9 +38,13 @@ void FGlobalEditorCommonCommands::MapActions(TSharedRef<FUICommandList>& Toolkit
 {
 	Register();
 
-// 	ToolkitCommands->MapAction(
-// 		Get().SummonControlTabNavigation,
-// 		FExecuteAction::CreateStatic( &FGlobalEditorCommonCommands::OnPressedCtrlTab ) );
+ 	ToolkitCommands->MapAction(
+ 		Get().SummonControlTabNavigation,
+ 		FExecuteAction::CreateStatic( &FGlobalEditorCommonCommands::OnPressedCtrlTab, Get().SummonControlTabNavigation ) );
+
+ 	ToolkitCommands->MapAction(
+ 		Get().SummonControlTabNavigationAlternate,
+ 		FExecuteAction::CreateStatic( &FGlobalEditorCommonCommands::OnPressedCtrlTab, Get().SummonControlTabNavigationAlternate ) );
 
 	ToolkitCommands->MapAction(
 		Get().SummonOpenAssetDialog,
@@ -128,41 +59,61 @@ void FGlobalEditorCommonCommands::MapActions(TSharedRef<FUICommandList>& Toolkit
 		FExecuteAction::CreateStatic( &FGlobalEditorCommonCommands::OnSummonedConsoleCommandBox ) );
 }
 
-void FGlobalEditorCommonCommands::OnPressedCtrlTab()
+void FGlobalEditorCommonCommands::OnPressedCtrlTab(TSharedPtr<FUICommandInfo> TriggeringCommand)
 {
+	if (!SGlobalTabSwitchingDialog::IsAlreadyOpen())
+	{
+		const FVector2D TabListSize(700.0f, 486.0f);
+
+		// Create the contents of the popup
+		TSharedRef<SWidget> ActualWidget = SNew(SGlobalTabSwitchingDialog, TabListSize, *TriggeringCommand->GetActiveGesture());
+
+		TSharedPtr<SWindow> NewWindow = OpenPopup(ActualWidget, TabListSize);
+	}
 }
 
 void FGlobalEditorCommonCommands::OnSummonedAssetPicker()
 {
 	const FVector2D AssetPickerSize(600.0f, 586.0f);
 
-	FMenuBuilder MenuBuilder(false, NULL);
-
 	// Create the contents of the popup
 	TSharedRef<SWidget> ActualWidget = SNew(SGlobalOpenAssetDialog, AssetPickerSize);
 
 	// Wrap the picker widget in a multibox-style menu body
+	FMenuBuilder MenuBuilder(/*BShouldCloseAfterSelection=*/ false, /*CommandList=*/ nullptr);
 	MenuBuilder.BeginSection("AssetPickerOpenAsset", NSLOCTEXT("GlobalAssetPicker", "WindowTitle", "Open Asset"));
-	{
-		MenuBuilder.AddWidget(ActualWidget, FText::GetEmpty(), true);
-	}
+	MenuBuilder.AddWidget(ActualWidget, FText::GetEmpty(), /*bNoIndent=*/ true);
 	MenuBuilder.EndSection();
 
-	TSharedRef<SWidget> WindowContents = MenuBuilder.MakeWidget();
+	OpenPopup(MenuBuilder.MakeWidget(), AssetPickerSize);
+}
 
+TSharedPtr<SWindow> FGlobalEditorCommonCommands::OpenPopup(TSharedRef<SWidget> WindowContents, const FVector2D& PopupDesiredSize)
+{
 	// Determine where the pop-up should open
 	TSharedPtr<SWindow> ParentWindow = FSlateApplication::Get().GetActiveTopLevelWindow();
 	FVector2D WindowPosition = FSlateApplication::Get().GetCursorPos();
+	if (!ParentWindow.IsValid())
+	{
+		TSharedPtr<SDockTab> LevelEditorTab = FModuleManager::Get().GetModuleChecked<FLevelEditorModule>("LevelEditor").GetLevelEditorTab();
+		ParentWindow = LevelEditorTab->GetParentWindow();
+		check(ParentWindow.IsValid());
+	}
+
 	if (ParentWindow.IsValid())
 	{
 		FSlateRect ParentMonitorRect = ParentWindow->GetFullScreenInfo();
 		const FVector2D MonitorCenter((ParentMonitorRect.Right + ParentMonitorRect.Left) * 0.5f, (ParentMonitorRect.Top + ParentMonitorRect.Bottom) * 0.5f);
-		WindowPosition = MonitorCenter - AssetPickerSize * 0.5f;
+		WindowPosition = MonitorCenter - PopupDesiredSize * 0.5f;
 
 		// Open the pop-up
 		FPopupTransitionEffect TransitionEffect(FPopupTransitionEffect::None);
-		TSharedRef<SWindow> PopupWindow = FSlateApplication::Get().PushMenu(ParentWindow.ToSharedRef(), WindowContents, WindowPosition, TransitionEffect);
+		TSharedRef<SWindow> PopupWindow = FSlateApplication::Get().PushMenu(ParentWindow.ToSharedRef(), WindowContents, WindowPosition, TransitionEffect, /*bFocusImmediately=*/ true);
+		
+		return PopupWindow;
 	}
+
+	return TSharedPtr<SWindow>();
 }
 
 static void CloseDebugConsole()
@@ -187,4 +138,6 @@ void FGlobalEditorCommonCommands::OnSummonedConsoleCommandBox()
 	}
 }
 
+//////////////////////////////////////////////////////////////////////////
+
 #undef LOCTEXT_NAMESPACE
\ No newline at end of file
diff --git a/Engine/Source/Editor/UnrealEd/Private/Toolkits/SGlobalOpenAssetDialog.cpp b/Engine/Source/Editor/UnrealEd/Private/Toolkits/SGlobalOpenAssetDialog.cpp
new file mode 100644
--- /dev/null
+++ b/Engine/Source/Editor/UnrealEd/Private/Toolkits/SGlobalOpenAssetDialog.cpp
@@ -0,0 +1,78 @@
+// Copyright 1998-2015 Epic Games, Inc. All Rights Reserved.
+
+#include "UnrealEd.h"
+#include "SGlobalOpenAssetDialog.h"
+
+#include "AssetData.h"
+#include "ContentBrowserModule.h"
+#include "Toolkits/AssetEditorManager.h"
+
+#define LOCTEXT_NAMESPACE "SGlobalOpenAssetDialog"
+
+//////////////////////////////////////////////////////////////////////////
+// SGlobalOpenAssetDialog
+
+void SGlobalOpenAssetDialog::Construct(const FArguments& InArgs, FVector2D InSize)
+{
+	FContentBrowserModule& ContentBrowserModule = FModuleManager::Get().LoadModuleChecked<FContentBrowserModule>(TEXT("ContentBrowser"));
+
+	FAssetPickerConfig AssetPickerConfig;
+	AssetPickerConfig.OnAssetDoubleClicked = FOnAssetSelected::CreateSP(this, &SGlobalOpenAssetDialog::OnAssetSelectedFromPicker);
+	AssetPickerConfig.OnAssetEnterPressed = FOnAssetEnterPressed::CreateSP(this, &SGlobalOpenAssetDialog::OnPressedEnterOnAssetsInPicker);
+	AssetPickerConfig.InitialAssetViewType = EAssetViewType::List;
+	AssetPickerConfig.bAllowNullSelection = false;
+	AssetPickerConfig.bShowBottomToolbar = true;
+	AssetPickerConfig.bAutohideSearchBar = false;
+	AssetPickerConfig.bCanShowClasses = false;
+	AssetPickerConfig.SaveSettingsName = TEXT("GlobalAssetPicker");
+
+	ChildSlot
+	[
+		SNew(SBox)
+		.WidthOverride(InSize.X)
+		.HeightOverride(InSize.Y)
+		[
+			SNew(SVerticalBox)
+			+ SVerticalBox::Slot()
+			.FillHeight(1.0f)
+			[
+				ContentBrowserModule.Get().CreateAssetPicker(AssetPickerConfig)
+			]
+		]
+	];
+}
+
+
+FReply SGlobalOpenAssetDialog::OnPreviewKeyDown(const FGeometry& MyGeometry, const FKeyEvent& InKeyEvent)
+{
+	if (InKeyEvent.GetKey() == EKeys::Escape)
+	{
+		FSlateApplication::Get().DismissAllMenus();
+		return FReply::Handled();
+	}
+
+	return FReply::Unhandled();
+}
+
+void SGlobalOpenAssetDialog::OnAssetSelectedFromPicker(const FAssetData& AssetData)
+{
+	if (UObject* ObjectToEdit = AssetData.GetAsset())
+	{
+		FAssetEditorManager::Get().OpenEditorForAsset(ObjectToEdit);
+	}
+}
+
+void SGlobalOpenAssetDialog::OnPressedEnterOnAssetsInPicker(const TArray<FAssetData>& SelectedAssets)
+{
+	for (auto AssetIt = SelectedAssets.CreateConstIterator(); AssetIt; ++AssetIt)
+	{
+		if (UObject* ObjectToEdit = AssetIt->GetAsset())
+		{
+			FAssetEditorManager::Get().OpenEditorForAsset(ObjectToEdit);
+		}
+	}
+}
+
+//////////////////////////////////////////////////////////////////////////
+
+#undef LOCTEXT_NAMESPACE
\ No newline at end of file
diff --git a/Engine/Source/Editor/UnrealEd/Private/Toolkits/SGlobalOpenAssetDialog.h b/Engine/Source/Editor/UnrealEd/Private/Toolkits/SGlobalOpenAssetDialog.h
new file mode 100644
--- /dev/null
+++ b/Engine/Source/Editor/UnrealEd/Private/Toolkits/SGlobalOpenAssetDialog.h
@@ -0,0 +1,24 @@
+// Copyright 1998-2015 Epic Games, Inc. All Rights Reserved.
+
+#pragma once
+
+//////////////////////////////////////////////////////////////////////////
+// SGlobalOpenAssetDialog
+
+class SGlobalOpenAssetDialog : public SCompoundWidget
+{
+public:
+	SLATE_BEGIN_ARGS(SGlobalOpenAssetDialog){}
+	SLATE_END_ARGS()
+
+public:
+	void Construct(const FArguments& InArgs, FVector2D InSize);
+
+	// SWidget interface
+	virtual FReply OnPreviewKeyDown(const FGeometry& MyGeometry, const FKeyEvent& InKeyEvent) override;
+	// End of SWidget interface
+
+protected:
+	void OnAssetSelectedFromPicker(const class FAssetData& AssetData);
+	void OnPressedEnterOnAssetsInPicker(const TArray<class FAssetData>& SelectedAssets);
+};
diff --git a/Engine/Source/Editor/UnrealEd/Private/Toolkits/SGlobalTabSwitchingDialog.cpp b/Engine/Source/Editor/UnrealEd/Private/Toolkits/SGlobalTabSwitchingDialog.cpp
new file mode 100644
--- /dev/null
+++ b/Engine/Source/Editor/UnrealEd/Private/Toolkits/SGlobalTabSwitchingDialog.cpp
@@ -0,0 +1,520 @@
+// Copyright 1998-2015 Epic Games, Inc. All Rights Reserved.
+
+#include "UnrealEd.h"
+#include "SGlobalTabSwitchingDialog.h"
+#include "SDockTab.h"
+
+#include "AssetEditorManager.h"
+#include "AssetThumbnail.h"
+#include "LevelEditor.h"
+#include "Editor/WorkspaceMenuStructure/Public/WorkspaceMenuStructureModule.h"
+
+#define LOCTEXT_NAMESPACE "SGlobalTabSwitchingDialog"
+
+//////////////////////////////////////////////////////////////////////////
+// FTabSwitchingListItemBase
+
+class FTabSwitchingListItemBase
+{
+public:
+	FTabSwitchingListItemBase()
+		: LastAccessTime(0.0)
+	{
+	}
+
+	virtual ~FTabSwitchingListItemBase() {}
+
+	virtual TSharedRef<SWidget> CreateWidget(TSharedPtr<class FAssetThumbnailPool> AssetThumbnailPool)
+	{
+		return SNullWidget::NullWidget;
+	}
+
+	virtual FText GetTypeString() const
+	{
+		return FText::GetEmpty();
+	}
+
+	virtual FText GetPathString() const
+	{
+		return FText::GetEmpty();
+	}
+
+	virtual void ActivateTab() { }
+
+	virtual void ShowInContentBrowser()
+	{
+	}
+
+	virtual TSharedPtr<FTabManager> GetAssociatedTabManager()
+	{
+		return TSharedPtr<FTabManager>();
+	}
+public:
+	double LastAccessTime;
+};
+
+//////////////////////////////////////////////////////////////////////////
+// FTabSwitchingListItem_Asset
+
+class FTabSwitchingListItem_Asset : public FTabSwitchingListItemBase
+{
+public:
+	FTabSwitchingListItem_Asset(UObject* InAsset)
+		: MyAsset(InAsset)
+	{
+		if (IAssetEditorInstance* EditorInstance = FAssetEditorManager::Get().FindEditorForAsset(MyAsset, /*bFocusIfOpen=*/ false))
+		{
+			LastAccessTime = EditorInstance->GetLastActivationTime();
+		}
+	}
+
+	virtual TSharedRef<SWidget> CreateWidget(TSharedPtr<class FAssetThumbnailPool> AssetThumbnailPool) override
+	{
+		// Create a label for the asset name
+		const bool bDirtyState = MyAsset->GetOutermost()->IsDirty();
+		FFormatNamedArguments Args;
+		Args.Add(TEXT("AssetName"), FText::AsCultureInvariant(MyAsset->GetName()));
+		Args.Add(TEXT("DirtyState"), bDirtyState ? LOCTEXT("AssetModified", " [Modified]") : LOCTEXT("AssetNotModified", ""));
+		FText AssetText = FText::Format(LOCTEXT("AssetEntryLabel", "{AssetName}{DirtyState}"), Args);
+
+		// Create a thumbnail to represent the asset type
+		const int32 ThumbnailSize = 48;
+
+		FAssetData AssetData(MyAsset);
+		Thumbnail = MakeShareable(new FAssetThumbnail(MyAsset, ThumbnailSize, ThumbnailSize, AssetThumbnailPool));
+
+		FAssetThumbnailConfig ThumbnailConfig;
+
+		return SNew(SHorizontalBox)
+			+SHorizontalBox::Slot()
+			.AutoWidth()
+			.VAlign(VAlign_Center)
+			.Padding(4.0f, 4.0f)
+			[
+				SNew(SBox)
+				.WidthOverride(ThumbnailSize)
+				.HeightOverride(ThumbnailSize)
+				[
+					Thumbnail->MakeThumbnailWidget(ThumbnailConfig)
+				]
+			]
+			+SHorizontalBox::Slot()
+			.FillWidth(1.0f)
+			.VAlign(VAlign_Center)
+			.Padding(8.0f, 0.0f, 8.0f, 0.0f)
+			[
+				SNew(STextBlock)
+				.TextStyle(FEditorStyle::Get(), "ControlTabMenu.AssetNameStyle")
+				.Text(AssetText)
+			];
+	}
+
+	virtual void ShowInContentBrowser()
+	{
+		TArray<UObject*> ObjectsToSync;
+		ObjectsToSync.Add(MyAsset);
+		GEditor->SyncBrowserToObjects(ObjectsToSync);
+	}
+
+	virtual FText GetTypeString() const override
+	{
+		return MyAsset->GetClass()->GetDisplayNameText();
+	}
+
+	virtual FText GetPathString() const override
+	{
+		return FText::AsCultureInvariant(MyAsset->GetOutermost()->GetName());
+	}
+
+	virtual void ActivateTab() override
+	{
+		FAssetEditorManager::Get().FindEditorForAsset(MyAsset, /*bFocusIfOpen=*/ true);
+	}
+
+	virtual TSharedPtr<FTabManager> GetAssociatedTabManager() override
+	{
+		IAssetEditorInstance* Instance = FAssetEditorManager::Get().FindEditorForAsset(MyAsset, /*bFocusIfOpen=*/ false);
+		return Instance->GetAssociatedTabManager();
+	}
+
+protected:
+	UObject* MyAsset;
+	TSharedPtr<class FAssetThumbnail> Thumbnail;
+};
+
+//////////////////////////////////////////////////////////////////////////
+// FTabSwitchingListItem_World
+
+class FTabSwitchingListItem_World : public FTabSwitchingListItem_Asset
+{
+public:
+	static TSharedPtr<FTabSwitchingListItem_World> MakeWorldItem()
+	{
+		UWorld* MyWorld = nullptr;
+		for (const FWorldContext& Context : GEngine->GetWorldContexts())
+		{
+			if (Context.WorldType == EWorldType::PIE)
+			{
+				MyWorld = Context.World();
+				break;
+			}
+			else if (Context.WorldType == EWorldType::Editor)
+			{
+				MyWorld = Context.World();
+			}
+		}
+
+		check(MyWorld != nullptr);
+		return MakeShareable(new FTabSwitchingListItem_World(MyWorld));
+	}
+
+	virtual void ActivateTab() override
+	{
+		const TSharedPtr<FTabManager> LevelEditorTabManager = GetAssociatedTabManager();
+		FGlobalTabmanager::Get()->DrawAttention(LevelEditorTabManager->GetOwnerTab().ToSharedRef());
+	}
+
+	virtual TSharedPtr<FTabManager> GetAssociatedTabManager() override
+	{
+		return FModuleManager::Get().GetModuleChecked<FLevelEditorModule>("LevelEditor").GetLevelEditorTabManager();
+	}
+
+protected:
+	FTabSwitchingListItem_World(UWorld* InWorld)
+		: FTabSwitchingListItem_Asset(InWorld)
+	{
+		const TSharedPtr<SDockTab> LevelEditorTabPtr = FModuleManager::Get().GetModuleChecked<FLevelEditorModule>("LevelEditor").GetLevelEditorTab();
+		LastAccessTime = LevelEditorTabPtr->GetLastActivationTime();
+	}
+};
+
+//////////////////////////////////////////////////////////////////////////
+// FTabSwitchingListItem_Asset
+
+bool SGlobalTabSwitchingDialog::bIsAlreadyOpen = false;
+
+SGlobalTabSwitchingDialog::FTabListItemPtr SGlobalTabSwitchingDialog::GetMainTabListSelectedItem() const
+{
+	TArray<FTabListItemPtr> SelectedItems = MainTabsListWidget->GetSelectedItems();
+	if (SelectedItems.Num() > 0)
+	{
+		return SelectedItems[0];
+	}
+	else
+	{
+		return FTabListItemPtr();
+	}
+}
+
+FReply SGlobalTabSwitchingDialog::OnBrowseToSelectedAsset()
+{
+	FTabListItemPtr SelectedItem = GetMainTabListSelectedItem();
+	if (SelectedItem.IsValid())
+	{
+		SelectedItem->ShowInContentBrowser();
+
+		FSlateApplication::Get().DismissAllMenus();
+	}
+	return FReply::Handled();
+}
+
+void SGlobalTabSwitchingDialog::OnMainTabListSelectionChanged(FTabListItemPtr InItem, ESelectInfo::Type SelectInfo)
+{
+	TSharedRef<SWidget> NewTopContents = SNullWidget::NullWidget;
+	TSharedRef<SWidget> NewBottomContents = SNullWidget::NullWidget;
+	TSharedRef<SWidget> NewToolTabsContent = SNullWidget::NullWidget;
+
+	TArray<FTabListItemPtr> SelectedItems = MainTabsListWidget->GetSelectedItems();
+
+	if (SelectedItems.Num() > 0)
+	{
+		FTabListItemPtr SelectedItem = SelectedItems[0];
+
+		NewTopContents = SelectedItem->CreateWidget(AssetThumbnailPool);
+
+		NewBottomContents =
+			SNew(SHorizontalBox)
+			+SHorizontalBox::Slot()
+			.AutoWidth()
+			[
+				SNew(SButton)
+				.ButtonStyle(FEditorStyle::Get(), "HoverOnlyHyperlinkButton")
+				.VAlign(VAlign_Center)
+				.HAlign(HAlign_Center)
+				.OnClicked(this, &SGlobalTabSwitchingDialog::OnBrowseToSelectedAsset)
+				.ToolTipText(LOCTEXT("BrowseButtonToolTipText", "Browse to Asset in Content Browser"))
+				[
+					SNew(SHorizontalBox)
+					+SHorizontalBox::Slot()
+					.AutoWidth()
+					.VAlign(VAlign_Center)
+					[
+						SNew(SImage)
+						.Image(FEditorStyle::GetBrush("PropertyWindow.Button_Browse"))
+					]
+					+SHorizontalBox::Slot()
+					.AutoWidth()
+					.VAlign(VAlign_Center)
+					.Padding(FMargin(4.0f, 0.0f, 0.0f, 0.0f))
+					[
+						SNew(STextBlock)
+						.TextStyle(FEditorStyle::Get(), "ControlTabMenu.AssetPathStyle")
+						.Text(SelectedItem->GetPathString())
+					]
+				]
+			]
+			+SHorizontalBox::Slot()
+			.FillWidth(1.0f)
+			[
+				SNullWidget::NullWidget
+			]
+			+SHorizontalBox::Slot()
+			.AutoWidth()
+			.HAlign(HAlign_Right)
+			[
+				SNew(STextBlock)
+				.TextStyle(FEditorStyle::Get(), "ControlTabMenu.AssetTypeStyle")
+				.Text(SelectedItem->GetTypeString())
+			];
+
+		// Create the list of the tool tabs in the current active tab
+		TSharedPtr<FTabManager> SelectedAssetTabManager = SelectedItem->GetAssociatedTabManager();
+		if (SelectedAssetTabManager.IsValid())
+		{
+			FMenuBuilder ToolTabMenuBuilder(/*bShouldCloseAfterSelection=*/ true, /*CommandList=*/ nullptr);
+			ToolTabMenuBuilder.GetMultiBox()->SetStyle(&FEditorStyle::Get(), "ToolBar");
+			
+			// Local editor tabs
+			SelectedAssetTabManager->PopulateLocalTabSpawnerMenu(ToolTabMenuBuilder);
+			
+			// General tabs
+			const IWorkspaceMenuStructure& MenuStructure = WorkspaceMenu::GetMenuStructure();
+			SelectedAssetTabManager->PopulateTabSpawnerMenu(ToolTabMenuBuilder, MenuStructure.GetStructureRoot());
+
+			// Turn the builder into a widget
+			NewToolTabsContent = ToolTabMenuBuilder.MakeWidget();
+		}
+	}
+
+	NewTabItemToActivateDisplayBox->SetContent(NewTopContents);
+	NewTabItemToActivatePathBox->SetContent(NewBottomContents);
+	ToolTabsListBox->SetContent(NewToolTabsContent);
+}
+
+void SGlobalTabSwitchingDialog::OnMainTabListItemClicked(FTabListItemPtr InItem)
+{
+	DismissDialog();
+}
+
+void SGlobalTabSwitchingDialog::CycleSelection(bool bForwards)
+{
+	// This is done here each time in case someone clicks off of the selected item (and to prime the pump at startup),
+	// otherwise the code below wouldn't cycle back into an item if nothing was selected
+	if ((MainTabsListWidget->GetNumItemsSelected() == 0) && (MainTabsListDataSource.Num() > 0))
+	{
+		MainTabsListWidget->SetSelection(MainTabsListDataSource[0]);
+	}
+
+	// Move to the next/previous item
+	FTabListItemPtr OldSelectedItem = GetMainTabListSelectedItem();
+	if (OldSelectedItem.IsValid())
+	{
+		int32 OldSelectionIndex;
+		if (MainTabsListDataSource.Find(OldSelectedItem, /*out*/ OldSelectionIndex))
+		{
+			const int32 NewSelectionIndex = (OldSelectionIndex + MainTabsListDataSource.Num() + (bForwards ? 1 : -1)) % MainTabsListDataSource.Num();
+			if (NewSelectionIndex != OldSelectionIndex)
+			{
+				FTabListItemPtr NewSelectedItem = MainTabsListDataSource[NewSelectionIndex];
+				MainTabsListWidget->SetSelection(NewSelectedItem);
+				MainTabsListWidget->RequestScrollIntoView(NewSelectedItem);
+			}
+		}
+	}
+}
+
+void SGlobalTabSwitchingDialog::DismissDialog()
+{
+	FTabListItemPtr SelectedItem = GetMainTabListSelectedItem();
+	if (SelectedItem.IsValid())
+	{
+		SelectedItem->ActivateTab();
+	}
+
+	FSlateApplication::Get().DismissAllMenus();
+}
+
+void SGlobalTabSwitchingDialog::Construct(const FArguments& InArgs, FVector2D InSize, FInputGesture InTriggerGesture)
+{
+	check(!bIsAlreadyOpen);
+	bIsAlreadyOpen = true;
+
+	TriggerGesture = InTriggerGesture;
+
+	AssetThumbnailPool = MakeShareable(new FAssetThumbnailPool(128));
+
+	// Populate the list with open asset editors
+	TArray<UObject*> OpenAssetList = FAssetEditorManager::Get().GetAllEditedAssets();
+	for (UObject* OpenAsset : OpenAssetList)
+	{
+		MainTabsListDataSource.Add(MakeShareable(new FTabSwitchingListItem_Asset(OpenAsset)));
+	}
+
+	MainTabsListDataSource.Add(FTabSwitchingListItem_World::MakeWorldItem());
+
+	// Sort the list by access time
+	MainTabsListDataSource.Sort([](const FTabListItemPtr& A, const FTabListItemPtr& B) { return (A->LastAccessTime > B->LastAccessTime); });
+
+	// Create the widgets
+	NewTabItemToActivateDisplayBox =
+		SNew(SBox)
+		.Padding(FMargin(0.0f, 0.0f, 10.0f, 0.0f))
+		.HeightOverride(70.0f)
+		.VAlign(VAlign_Top);
+
+	NewTabItemToActivatePathBox =
+		SNew(SBox)
+		.Padding(FMargin(0.0f, 10.0f, 10.0f, 10.0f))
+		.HeightOverride(40.0f)
+		.VAlign(VAlign_Center);
+
+	ToolTabsListBox =
+		SNew(SBox)
+		.Padding(FMargin(0.0f, 0.0f, 15.0f, 0.0f));
+
+	MainTabsListWidget = SNew(STabListWidget)
+		.ItemHeight(64)
+		.ListItemsSource(&MainTabsListDataSource)
+		.OnGenerateRow(this, &SGlobalTabSwitchingDialog::OnGenerateTabSwitchListItemWidget)
+		.OnSelectionChanged(this, &SGlobalTabSwitchingDialog::OnMainTabListSelectionChanged)
+		.OnMouseButtonClick(this, &SGlobalTabSwitchingDialog::OnMainTabListItemClicked)
+		.SelectionMode(ESelectionMode::Single);
+
+	TSharedRef<SWidget> ToolTabList = SNew(SVerticalBox)
+		+SVerticalBox::Slot()
+		.AutoHeight()
+		[
+			SNew(STextBlock)
+			.TextStyle(FEditorStyle::Get(), "ControlTabMenu.HeadingStyle")
+			.Text(LOCTEXT("ChangeToolsHeading", "Tool Windows"))
+		]
+		+SVerticalBox::Slot()
+		.FillHeight(1.0f)
+		[
+			ToolTabsListBox.ToSharedRef()
+		];
+
+	TSharedRef<SWidget> DocumentTabList = SNew(SVerticalBox)
+		+SVerticalBox::Slot()
+		.AutoHeight()
+		[
+			SNew(STextBlock)
+			.TextStyle(FEditorStyle::Get(), "ControlTabMenu.HeadingStyle")
+			.Text(LOCTEXT("OpenAssetsHeading", "Active Files"))
+		]
+		+SVerticalBox::Slot()
+		.FillHeight(1.0f)
+		[
+			SNew(SScrollBorder, MainTabsListWidget.ToSharedRef())
+			[
+				MainTabsListWidget.ToSharedRef()
+			]
+		];
+
+	ChildSlot
+	[
+		SNew(SBorder)
+		.BorderImage(FEditorStyle::Get().GetBrush("ControlTabMenu.Background"))
+		.ForegroundColor(FCoreStyle::Get().GetSlateColor("DefaultForeground"))
+		[
+			SNew(SBox)
+			.WidthOverride(InSize.X)
+			.HeightOverride(InSize.Y)
+			.Padding(FMargin(12.0f, 12.0f, 12.0f, 0.0f))
+			[
+				SNew(SVerticalBox)
+				+SVerticalBox::Slot()
+				.AutoHeight()
+				[
+					NewTabItemToActivateDisplayBox.ToSharedRef()
+				]
+				+SVerticalBox::Slot()
+				.FillHeight(1.0f)
+				[
+					SNew(SHorizontalBox)
+					+SHorizontalBox::Slot()
+					.FillWidth(0.4f)
+					.Padding(FMargin(0.0f, 0.0f, 0.0f, 0.0f))
+					[
+						ToolTabList
+					]
+
+					+SHorizontalBox::Slot()
+					.FillWidth(1.0f)
+					[
+						DocumentTabList
+					]
+
+				]
+				+SVerticalBox::Slot()
+				.AutoHeight()
+				[
+					NewTabItemToActivatePathBox.ToSharedRef()
+				]
+			]
+		]
+	];
+
+	// Pick the second most recent or least recent file based on whether Shift was held down when we were summoned
+	if (MainTabsListDataSource.Num() > 0)
+	{
+		CycleSelection(/*bForwards=*/ !FSlateApplication::Get().GetModifierKeys().IsShiftDown());
+	}
+}
+
+FReply SGlobalTabSwitchingDialog::OnKeyUp(const FGeometry& MyGeometry, const FKeyEvent& InKeyEvent)
+{
+	// Check to see if the trigger modifier key was released, which should close the dialog
+	const bool bCloseViaControl = TriggerGesture.NeedsControl() && ((InKeyEvent.GetKey() == EKeys::LeftControl) || (InKeyEvent.GetKey() == EKeys::RightControl));
+	const bool bCloseViaCommand = TriggerGesture.NeedsCommand() && ((InKeyEvent.GetKey() == EKeys::LeftCommand) || (InKeyEvent.GetKey() == EKeys::RightCommand));
+	const bool bCloseViaAlt = TriggerGesture.NeedsAlt() && ((InKeyEvent.GetKey() == EKeys::LeftAlt) || (InKeyEvent.GetKey() == EKeys::RightAlt));
+	
+	if (bCloseViaControl || bCloseViaCommand || bCloseViaAlt)
+	{
+		DismissDialog();
+		return FReply::Handled();
+	}
+
+	return FReply::Unhandled();
+}
+
+FReply SGlobalTabSwitchingDialog::OnKeyDown(const FGeometry& MyGeometry, const FKeyEvent& InKeyEvent)
+{
+	if (TriggerGesture.Key == InKeyEvent.GetKey())
+	{
+		const bool bCycleForward = !InKeyEvent.IsShiftDown();
+		CycleSelection(bCycleForward);
+	}
+
+	return FReply::Unhandled();
+}
+
+FReply SGlobalTabSwitchingDialog::OnPreviewKeyDown(const FGeometry& MyGeometry, const FKeyEvent& InKeyEvent)
+{
+	if (InKeyEvent.GetKey() == EKeys::Escape)
+	{
+		FSlateApplication::Get().DismissAllMenus();
+		return FReply::Handled();
+	}
+
+	return FReply::Unhandled();
+}
+
+TSharedRef<ITableRow> SGlobalTabSwitchingDialog::OnGenerateTabSwitchListItemWidget(FTabListItemPtr InItem, const TSharedRef<STableViewBase>& OwnerTable)
+{
+	return SNew(STableRow<FTabListItemPtr>, OwnerTable)[InItem->CreateWidget(AssetThumbnailPool)];
+}
+
+//////////////////////////////////////////////////////////////////////////
+
+#undef LOCTEXT_NAMESPACE
\ No newline at end of file
diff --git a/Engine/Source/Editor/UnrealEd/Private/Toolkits/SGlobalTabSwitchingDialog.h b/Engine/Source/Editor/UnrealEd/Private/Toolkits/SGlobalTabSwitchingDialog.h
new file mode 100644
--- /dev/null
+++ b/Engine/Source/Editor/UnrealEd/Private/Toolkits/SGlobalTabSwitchingDialog.h
@@ -0,0 +1,73 @@
+// Copyright 1998-2015 Epic Games, Inc. All Rights Reserved.
+
+#pragma once
+
+//////////////////////////////////////////////////////////////////////////
+// SGlobalTabSwitchingDialog
+
+class SGlobalTabSwitchingDialog : public SCompoundWidget
+{
+public:
+	SLATE_BEGIN_ARGS(SGlobalTabSwitchingDialog){}
+	SLATE_END_ARGS()
+
+public:
+	void Construct(const FArguments& InArgs, FVector2D InSize, FInputGesture InTriggerGesture);
+
+	~SGlobalTabSwitchingDialog()
+	{
+		bIsAlreadyOpen = false;
+	}
+
+	// SWidget interface
+	virtual FReply OnPreviewKeyDown(const FGeometry& MyGeometry, const FKeyEvent& InKeyEvent) override;
+	virtual FReply OnKeyDown(const FGeometry& MyGeometry, const FKeyEvent& InKeyEvent) override;
+	virtual FReply OnKeyUp(const FGeometry& MyGeometry, const FKeyEvent& InKeyEvent) override;
+	virtual bool SupportsKeyboardFocus() const { return true; }
+	// End of SWidget interface
+
+	// Is an instance already open?
+	static bool IsAlreadyOpen()
+	{
+		return bIsAlreadyOpen;
+	}
+
+private:
+	typedef TSharedPtr<class FTabSwitchingListItemBase> FTabListItemPtr;
+	typedef SListView<FTabListItemPtr> STabListWidget;
+
+private:
+	TSharedRef<ITableRow> OnGenerateTabSwitchListItemWidget(FTabListItemPtr InItem, const TSharedRef<STableViewBase>& OwnerTable);
+
+	void CycleSelection(bool bForwards);
+	void OnMainTabListSelectionChanged(FTabListItemPtr InItem, ESelectInfo::Type SelectInfo);
+	void OnMainTabListItemClicked(FTabListItemPtr InItem);
+	void DismissDialog();
+	FReply OnBrowseToSelectedAsset();
+
+	FTabListItemPtr GetMainTabListSelectedItem() const;
+
+private:
+	// The gesture that triggered the dialog (so we can handle the correct Tab/`/etc... key repeat, and dismiss on the correct control/command modifier release
+	FInputGesture TriggerGesture;
+
+	// Pool for maintaining and rendering thumbnails
+	TSharedPtr<class FAssetThumbnailPool> AssetThumbnailPool;
+
+	// The array of 'document' items
+	TArray<FTabListItemPtr> MainTabsListDataSource;
+
+	// The widget representing the list of 'document' items
+	TSharedPtr<STabListWidget> MainTabsListWidget;
+
+	// The container widget for the indication of the asset that will be activated when the dialog closes
+	TSharedPtr<SBox> NewTabItemToActivateDisplayBox;
+
+	// The container widget for the indication of the path to the asset that will be activated when the dialog closes
+	TSharedPtr<SBox> NewTabItemToActivatePathBox;
+
+	// The container widget for the tool tabs list
+	TSharedPtr<SBox> ToolTabsListBox;
+
+	static bool bIsAlreadyOpen;
+};
diff --git a/Engine/Source/Editor/UnrealEd/Public/Toolkits/GlobalEditorCommonCommands.h b/Engine/Source/Editor/UnrealEd/Public/Toolkits/GlobalEditorCommonCommands.h
--- a/Engine/Source/Editor/UnrealEd/Public/Toolkits/GlobalEditorCommonCommands.h
+++ b/Engine/Source/Editor/UnrealEd/Public/Toolkits/GlobalEditorCommonCommands.h
@@ -15,12 +15,15 @@ public:
 	static void MapActions(TSharedRef<FUICommandList>& ToolkitCommands);
 
 protected:
-	static void OnPressedCtrlTab();
+	static void OnPressedCtrlTab(TSharedPtr<FUICommandInfo> TriggeringCommand);
 	static void OnSummonedAssetPicker();
 	static void OnSummonedConsoleCommandBox();
+
+	static TSharedPtr<SWindow> OpenPopup(TSharedRef<SWidget> WindowContents, const FVector2D& PopupDesiredSize);
 public:
 	TSharedPtr<FUICommandInfo> FindInContentBrowser;
 	TSharedPtr<FUICommandInfo> SummonControlTabNavigation;
+	TSharedPtr<FUICommandInfo> SummonControlTabNavigationAlternate;
 	TSharedPtr<FUICommandInfo> SummonOpenAssetDialog;
 	TSharedPtr<FUICommandInfo> SummonOpenAssetDialogAlternate;
 	TSharedPtr<FUICommandInfo> OpenDocumentation;
