Correct rendertarget format for rendering metallic, roughness, specular properties

diff --git a/Engine/Source/Editor/UnrealEd/Private/EditorExporters.cpp b/Engine/Source/Editor/UnrealEd/Private/EditorExporters.cpp
--- a/Engine/Source/Editor/UnrealEd/Private/EditorExporters.cpp
+++ b/Engine/Source/Editor/UnrealEd/Private/EditorExporters.cpp
@@ -2584,7 +2584,7 @@ namespace MaterialExportUtils
 			RenderTargetMetallic->ClearColor = FLinearColor(0.0f, 0.0f, 0.0f, 0.0f);
 			RenderTargetMetallic->InitCustomFormat(
 				OutFlattenMaterial.MetallicSize.X,
-				OutFlattenMaterial.MetallicSize.Y, PF_FloatRGB, true);
+				OutFlattenMaterial.MetallicSize.Y, PF_B8G8R8A8, true);
 
 			//
 			OutFlattenMaterial.MetallicSamples.Empty(OutFlattenMaterial.MetallicSize.X * OutFlattenMaterial.MetallicSize.Y);
@@ -2616,7 +2616,7 @@ namespace MaterialExportUtils
 			RenderTargetRoughness->ClearColor = FLinearColor(0.0f, 0.0f, 0.0f, 0.0f);
 			RenderTargetRoughness->InitCustomFormat(
 				OutFlattenMaterial.RoughnessSize.X,
-				OutFlattenMaterial.RoughnessSize.Y, PF_FloatRGB, true);
+				OutFlattenMaterial.RoughnessSize.Y, PF_B8G8R8A8, true);
 
 			//
 			OutFlattenMaterial.RoughnessSamples.Empty(OutFlattenMaterial.RoughnessSize.X * OutFlattenMaterial.RoughnessSize.Y);
@@ -2648,7 +2648,7 @@ namespace MaterialExportUtils
 			RenderTargetSpecular->ClearColor = FLinearColor(0.0f, 0.0f, 0.0f, 0.0f);
 			RenderTargetSpecular->InitCustomFormat(
 				OutFlattenMaterial.SpecularSize.X,
-				OutFlattenMaterial.SpecularSize.Y, PF_FloatRGB, true);
+				OutFlattenMaterial.SpecularSize.Y, PF_B8G8R8A8, true);
 
 			//
 			OutFlattenMaterial.SpecularSamples.Empty(OutFlattenMaterial.SpecularSize.X * OutFlattenMaterial.SpecularSize.Y);
