Added double-click handling to components in SCS editor that focuses the viewport cameras on the item.

UE-8539 - IWCE: Double-clicking on a component in the SCS editor should focus the camera on it

diff --git a/Engine/Source/Editor/LevelEditor/Private/SActorDetails.cpp b/Engine/Source/Editor/LevelEditor/Private/SActorDetails.cpp
--- a/Engine/Source/Editor/LevelEditor/Private/SActorDetails.cpp
+++ b/Engine/Source/Editor/LevelEditor/Private/SActorDetails.cpp
@@ -112,7 +112,8 @@ void SActorDetails::Construct(const FArguments& InArgs, const FName TabIdentifie
 		.EditorMode(EComponentEditorMode::ActorInstance)
 		.AllowEditing(this, &SActorDetails::GetAllowComponentTreeEditing)
 		.ActorContext(this, &SActorDetails::GetActorContext)
-		.OnSelectionUpdated(this, &SActorDetails::OnSCSEditorTreeViewSelectionChanged);
+		.OnSelectionUpdated(this, &SActorDetails::OnSCSEditorTreeViewSelectionChanged)
+		.OnItemDoubleClicked(this, &SActorDetails::OnSCSEditorTreeViewItemDoubleClicked);
 		
 	ComponentsBox->SetContent(SCSEditor.ToSharedRef());
 
@@ -468,6 +469,22 @@ void SActorDetails::OnSCSEditorTreeViewSelectionChanged(const TArray<FSCSEditorT
 	}
 }
 
+void SActorDetails::OnSCSEditorTreeViewItemDoubleClicked(const TSharedPtr<class FSCSEditorTreeNode> ClickedNode)
+{
+	if (ClickedNode.IsValid())
+	{
+		if (ClickedNode->GetNodeType() == FSCSEditorTreeNode::ComponentNode)
+		{
+			USceneComponent* SceneComponent = Cast<USceneComponent>(ClickedNode->GetComponentTemplate());
+			if (SceneComponent != nullptr)
+			{
+				const bool bActiveViewportOnly = false;
+				GEditor->MoveViewportCamerasToComponent(SceneComponent, bActiveViewportOnly);
+			}
+		}
+	}
+}
+
 void SActorDetails::UpdateComponentTreeFromEditorSelection()
 {
 	if (!DetailsView->IsLocked())
diff --git a/Engine/Source/Editor/LevelEditor/Private/SActorDetails.h b/Engine/Source/Editor/LevelEditor/Private/SActorDetails.h
--- a/Engine/Source/Editor/LevelEditor/Private/SActorDetails.h
+++ b/Engine/Source/Editor/LevelEditor/Private/SActorDetails.h
@@ -34,6 +34,7 @@ private:
 	void OnEditorSelectionChanged(UObject* Object);
 	void OnSCSEditorRootSelected(AActor* Actor);
 	void OnSCSEditorTreeViewSelectionChanged(const TArray<TSharedPtr<class FSCSEditorTreeNode> >& SelectedNodes);
+	void OnSCSEditorTreeViewItemDoubleClicked(const TSharedPtr<class FSCSEditorTreeNode> ClickedNode);
 	void UpdateComponentTreeFromEditorSelection();
 
 	bool IsPropertyReadOnly(const struct FPropertyAndParent& PropertyAndParent) const;
diff --git a/Engine/Source/Editor/UnrealEd/Classes/Editor/EditorEngine.h b/Engine/Source/Editor/UnrealEd/Classes/Editor/EditorEngine.h
--- a/Engine/Source/Editor/UnrealEd/Classes/Editor/EditorEngine.h
+++ b/Engine/Source/Editor/UnrealEd/Classes/Editor/EditorEngine.h
@@ -946,6 +946,13 @@ public:
 	*/
 	void MoveViewportCamerasToActor(const TArray<AActor*> &Actors, const TArray<UPrimitiveComponent*>& Components, bool bActiveViewportOnly);
 
+	/**
+	* Moves all viewport cameras to focus on the provided component.
+	* @param	Component				Target component
+	* @param	bActiveViewportOnly		If true, move/reorient only the active viewport.
+	*/
+	void MoveViewportCamerasToComponent(USceneComponent* Component, bool bActiveViewportOnly);
+
 	/** 
 	 * Snaps an actor in a direction.  Optionally will align with the trace normal.
 	 * @param InActor			Actor to move to the floor.
@@ -2776,6 +2783,13 @@ private:
 	/** Gets the init values for worlds opened via Map_Load in the editor */
 	UWorld::InitializationValues GetEditorWorldInitializationValues() const;
 
+	/**
+	* Moves all viewport cameras to focus on the provided bounding box.
+	* @param	BoundingBox				Target box
+	* @param	bActiveViewportOnly		If true, move/reorient only the active viewport.
+	*/
+	void MoveViewportCamerasToBox(const FBox& BoundingBox, bool bActiveViewportOnly) const;
+
 	/** Mapping of delegate handles for each online Login() call while in flight */
 	TMap<FName, FDelegateHandle> OnLoginPIECompleteDelegateHandlesForPIEInstances;
 
diff --git a/Engine/Source/Editor/UnrealEd/Private/EditorServer.cpp b/Engine/Source/Editor/UnrealEd/Private/EditorServer.cpp
--- a/Engine/Source/Editor/UnrealEd/Private/EditorServer.cpp
+++ b/Engine/Source/Editor/UnrealEd/Private/EditorServer.cpp
@@ -66,6 +66,11 @@ DEFINE_LOG_CATEGORY_STATIC(LogEditorServer, Log, All);
 /** Used for the "tagsounds" and "checksounds" commands only			*/
 static FUObjectAnnotationSparseBool DebugSoundAnnotation;
 
+namespace EditorEngineDefs
+{
+	/** Limit the minimum size of the bounding box when centering cameras on individual components to avoid extreme zooming */
+	static const float MinComponentBoundsForZoom = 50.0f;
+}
 
 namespace 
 {
@@ -4459,42 +4464,7 @@ void UEditorEngine::MoveViewportCamerasToActor(const TArray<AActor*> &Actors, co
 		}
 	}
 
-	// Make sure we had atleast one non-null actor in the array passed in.
-	if( BoundingBox.GetSize() != FVector::ZeroVector )
-	{
-		if ( bActiveViewportOnly )
-		{
-			if ( GCurrentLevelEditingViewportClient )
-			{
-				GCurrentLevelEditingViewportClient->FocusViewportOnBox( BoundingBox );
-				
-				// Update Linked Orthographic viewports.
-				if( GCurrentLevelEditingViewportClient->IsOrtho() && GetDefault<ULevelEditorViewportSettings>()->bUseLinkedOrthographicViewports )
-				{
-					// Search through all viewports
-					for( auto ViewportIt = GEditor->LevelViewportClients.CreateConstIterator(); ViewportIt ; ++ViewportIt )
-					{
-						FLevelEditorViewportClient* LinkedViewportClient = *ViewportIt;
-						// Only update other orthographic viewports
-						if( LinkedViewportClient && LinkedViewportClient != GCurrentLevelEditingViewportClient && LinkedViewportClient->IsOrtho() )
-						{
-							LinkedViewportClient->FocusViewportOnBox( BoundingBox );
-						}
-					}
-				}
-			}
-
-		}
-		else
-		{
-			// Update all viewports.
-			for( int32 i = 0; i < LevelViewportClients.Num(); i++ )
-			{
-				FEditorViewportClient* ViewportClient = LevelViewportClients[i];
-				ViewportClient->FocusViewportOnBox( BoundingBox );
-			}
-		}
-	}
+	MoveViewportCamerasToBox(BoundingBox, bActiveViewportOnly);
 
 	// Warn the user with a suppressable dialog if they attempted to zoom to actors that are in an invisible level
 	if ( InvisLevelActors.Num() > 0 )
@@ -4518,6 +4488,29 @@ void UEditorEngine::MoveViewportCamerasToActor(const TArray<AActor*> &Actors, co
 	FEditorDelegates::OnFocusViewportOnActors.Broadcast(Actors);
 }
 
+void UEditorEngine::MoveViewportCamerasToComponent(USceneComponent* Component, bool bActiveViewportOnly)
+{
+	if (Component != nullptr)
+	{
+		if (FLevelUtils::IsLevelVisible(Component->GetComponentLevel()) && Component->IsRegistered())
+		{
+			FBox Box = Component->Bounds.GetBox();
+			FVector Center;
+			FVector Extents;
+			Box.GetCenterAndExtents(Center, Extents);
+
+			// Apply a minimum size to the extents of the component's box to avoid the camera's zooming too close to small or zero-sized components
+			if (Extents.SizeSquared() < EditorEngineDefs::MinComponentBoundsForZoom * EditorEngineDefs::MinComponentBoundsForZoom)
+			{
+				FVector NewExtents(EditorEngineDefs::MinComponentBoundsForZoom, SMALL_NUMBER, SMALL_NUMBER);
+				Box = FBox(Center - NewExtents, Center + NewExtents);
+			}
+
+			MoveViewportCamerasToBox(Box, bActiveViewportOnly);
+		}
+	}
+}
+
 /** 
  * Snaps an actor in a direction.  Optionally will align with the trace normal.
  * @param InActor			Actor to move to the floor.
@@ -6587,3 +6580,43 @@ void UEditorEngine::AutoMergeStaticMeshes()
 	}
 #endif // #if TODO_STATICMESH
 }
+
+void UEditorEngine::MoveViewportCamerasToBox(const FBox& BoundingBox, bool bActiveViewportOnly) const
+{
+	// Make sure we had at least one non-null actor in the array passed in.
+	if (BoundingBox.GetSize() != FVector::ZeroVector)
+	{
+		if (bActiveViewportOnly)
+		{
+			if (GCurrentLevelEditingViewportClient)
+			{
+				GCurrentLevelEditingViewportClient->FocusViewportOnBox(BoundingBox);
+
+				// Update Linked Orthographic viewports.
+				if (GCurrentLevelEditingViewportClient->IsOrtho() && GetDefault<ULevelEditorViewportSettings>()->bUseLinkedOrthographicViewports)
+				{
+					// Search through all viewports
+					for (auto ViewportIt = LevelViewportClients.CreateConstIterator(); ViewportIt; ++ViewportIt)
+					{
+						FLevelEditorViewportClient* LinkedViewportClient = *ViewportIt;
+						// Only update other orthographic viewports
+						if (LinkedViewportClient && LinkedViewportClient != GCurrentLevelEditingViewportClient && LinkedViewportClient->IsOrtho())
+						{
+							LinkedViewportClient->FocusViewportOnBox(BoundingBox);
+						}
+					}
+				}
+			}
+
+		}
+		else
+		{
+			// Update all viewports.
+			for (auto ViewportIt = LevelViewportClients.CreateConstIterator(); ViewportIt; ++ViewportIt)
+			{
+				FLevelEditorViewportClient* LinkedViewportClient = *ViewportIt;
+				LinkedViewportClient->FocusViewportOnBox(BoundingBox);
+			}
+		}
+	}
+}
