PR #758: fix importing apb (Contributed by [REDACTED_PERSON])

https://github.com/EpicGames/UnrealEngine/pull/758

diff --git a/Engine/Source/Editor/UnrealEd/Private/Factories/ApexDestructibleAssetImport.cpp b/Engine/Source/Editor/UnrealEd/Private/Factories/ApexDestructibleAssetImport.cpp
--- a/Engine/Source/Editor/UnrealEd/Private/Factories/ApexDestructibleAssetImport.cpp
+++ b/Engine/Source/Editor/UnrealEd/Private/Factories/ApexDestructibleAssetImport.cpp
@@ -728,7 +728,11 @@ bool SetApexDestructibleAsset(UDestructibleMesh& DestructibleMesh, NxDestructibl
 {
 	DestructibleMesh.PreEditChange(NULL);
 
-	ExistingDestMeshData * ExistDestMeshDataPtr = SaveExistingDestMeshData(&DestructibleMesh);
+	ExistingDestMeshData * ExistDestMeshDataPtr = nullptr;
+	if(Options & EDestructibleImportOptions::PreserveSettings)
+	{
+		ExistDestMeshDataPtr = SaveExistingDestMeshData(&DestructibleMesh);
+	}
 	
 	// The asset is going away, which will destroy any actors created from it.  We must destroy the physics state of any destructible mesh components before we release the asset.
 	for(TObjectIterator<UDestructibleComponent> It; It; ++It)
