Added 3 new orthographic viewport modes.

diff --git a/Engine/Source/Editor/EditorStyle/Private/SlateEditorStyle.cpp b/Engine/Source/Editor/EditorStyle/Private/SlateEditorStyle.cpp
--- a/Engine/Source/Editor/EditorStyle/Private/SlateEditorStyle.cpp
+++ b/Engine/Source/Editor/EditorStyle/Private/SlateEditorStyle.cpp
@@ -4156,8 +4156,11 @@ void FSlateEditorStyle::FStyle::SetupLevelEditorStyle()
 		Set( "EditorViewport.CollisionVisibility", new IMAGE_BRUSH( "Icons/icon_ViewMode_CollisionVisibility_16px", Icon16x16 ) );
 		Set( "EditorViewport.Perspective", new IMAGE_BRUSH( "Icons/icon_ViewMode_ViewPerspective_16px", Icon16x16 ) );
 		Set( "EditorViewport.Top", new IMAGE_BRUSH( "Icons/icon_ViewMode_ViewTop_16px", Icon16x16 ) );
-		Set( "EditorViewport.Side", new IMAGE_BRUSH( "Icons/icon_ViewMode_ViewSide_16px", Icon16x16 ) );
+		Set( "EditorViewport.Left", new IMAGE_BRUSH( "Icons/icon_ViewMode_ViewLeft_16px", Icon16x16 ) );
 		Set( "EditorViewport.Front", new IMAGE_BRUSH( "Icons/icon_ViewMode_ViewFront_16px", Icon16x16 ) );
+		Set( "EditorViewport.Bottom", new IMAGE_BRUSH("Icons/icon_ViewMode_ViewBottom_16px", Icon16x16 ) );
+		Set( "EditorViewport.Right", new IMAGE_BRUSH("Icons/icon_ViewMode_ViewRight_16px", Icon16x16 ) );
+		Set( "EditorViewport.Back", new IMAGE_BRUSH("Icons/icon_ViewMode_ViewBack_16px", Icon16x16 ) );
 #endif
 
 #if WITH_EDITOR || IS_PROGRAM
diff --git a/Engine/Source/Editor/Kismet/Private/SSCSEditorViewport.cpp b/Engine/Source/Editor/Kismet/Private/SSCSEditorViewport.cpp
--- a/Engine/Source/Editor/Kismet/Private/SSCSEditorViewport.cpp
+++ b/Engine/Source/Editor/Kismet/Private/SSCSEditorViewport.cpp
@@ -125,13 +125,25 @@ public:
 				break;
 
 			case LVT_OrthoYZ:
-				Label = NSLOCTEXT("BlueprintEditor", "CameraMenuTitle_Side", "Side");
+				Label = NSLOCTEXT("BlueprintEditor", "CameraMenuTitle_Left", "Left");
 				break;
 
 			case LVT_OrthoXZ:
 				Label = NSLOCTEXT("BlueprintEditor", "CameraMenuTitle_Front", "Front");
 				break;
 
+			case LVT_OrthoNegativeXY:
+				Label = NSLOCTEXT("BlueprintEditor", "CameraMenuTitle_Bottom", "Bottom");
+				break;
+
+			case LVT_OrthoNegativeYZ:
+				Label = NSLOCTEXT("BlueprintEditor", "CameraMenuTitle_Right", "Right");
+				break;
+
+			case LVT_OrthoNegativeXZ:
+				Label = NSLOCTEXT("BlueprintEditor", "CameraMenuTitle_Back", "Back");
+				break;
+
 			case LVT_OrthoFreelook:
 				Label = NSLOCTEXT("BlueprintEditor", "CameraMenuTitle_OrthoFreelook", "Ortho");
 				break;
@@ -145,8 +157,11 @@ public:
 	{
 		static FName PerspectiveIconName("EditorViewport.Perspective");
 		static FName TopIconName("EditorViewport.Top");
-		static FName SideIconName("EditorViewport.Side");
+		static FName LeftIconName("EditorViewport.Left");
 		static FName FrontIconName("EditorViewport.Front");
+		static FName BottomIconName("EditorViewport.Bottom");
+		static FName RightIconName("EditorViewport.Right");
+		static FName BackIconName("EditorViewport.Back");
 
 		FName Icon = NAME_None;
 
@@ -163,12 +178,24 @@ public:
 				break;
 
 			case LVT_OrthoYZ:
-				Icon = SideIconName;
+				Icon = LeftIconName;
 				break;
 
 			case LVT_OrthoXZ:
 				Icon = FrontIconName;
 				break;
+
+			case LVT_OrthoNegativeXY:
+				Icon = BottomIconName;
+				break;
+
+			case LVT_OrthoNegativeYZ:
+				Icon = RightIconName;
+				break;
+
+			case LVT_OrthoNegativeXZ:
+				Icon = BackIconName;
+				break;
 			}
 		}
 
@@ -186,8 +213,11 @@ public:
 
 		CameraMenuBuilder.BeginSection("LevelViewportCameraType_Ortho", NSLOCTEXT("BlueprintEditor", "CameraTypeHeader_Ortho", "Orthographic"));
 			CameraMenuBuilder.AddMenuEntry(FEditorViewportCommands::Get().Top);
-			CameraMenuBuilder.AddMenuEntry(FEditorViewportCommands::Get().Side);
+			CameraMenuBuilder.AddMenuEntry(FEditorViewportCommands::Get().Bottom);
+			CameraMenuBuilder.AddMenuEntry(FEditorViewportCommands::Get().Left);
+			CameraMenuBuilder.AddMenuEntry(FEditorViewportCommands::Get().Right);
 			CameraMenuBuilder.AddMenuEntry(FEditorViewportCommands::Get().Front);
+			CameraMenuBuilder.AddMenuEntry(FEditorViewportCommands::Get().Back);
 		CameraMenuBuilder.EndSection();
 
 		return CameraMenuBuilder.MakeWidget();
diff --git a/Engine/Source/Editor/LandscapeEditor/Private/LandscapeEdMode.cpp b/Engine/Source/Editor/LandscapeEditor/Private/LandscapeEdMode.cpp
--- a/Engine/Source/Editor/LandscapeEditor/Private/LandscapeEdMode.cpp
+++ b/Engine/Source/Editor/LandscapeEditor/Private/LandscapeEdMode.cpp
@@ -1929,7 +1929,7 @@ void FEdModeLandscape::Render(const FSceneView* View, FViewport* Viewport, FPrim
 		}
 		else //if (NewLandscapePreviewMode == ENewLandscapePreviewMode::NewLandscape)
 		{
-			if (ViewportType == LVT_Perspective || ViewportType == LVT_OrthoXY)
+			if (ViewportType == LVT_Perspective || ViewportType == LVT_OrthoXY || ViewportType == LVT_OrthoNegativeXY)
 			{
 				for (int32 x = 0; x <= ComponentCountX * QuadsPerComponent; x++)
 				{
@@ -1975,7 +1975,7 @@ void FEdModeLandscape::Render(const FSceneView* View, FViewport* Viewport, FPrim
 				PDI->DrawLine(Transform.TransformPosition(FVector(ComponentCountX * QuadsPerComponent, 0, 0)), Transform.TransformPosition(FVector(ComponentCountX * QuadsPerComponent, ComponentCountY * ComponentSize, 0)), EdgeColour, SDPG_World);
 			}
 
-			if (ViewportType == LVT_Perspective || ViewportType == LVT_OrthoXY)
+			if (ViewportType == LVT_Perspective || ViewportType == LVT_OrthoXY || ViewportType == LVT_OrthoNegativeXY)
 			{
 				for (int32 y = 0; y <= ComponentCountY * QuadsPerComponent; y++)
 				{
diff --git a/Engine/Source/Editor/LevelEditor/Private/LevelViewportLayout2x2.cpp b/Engine/Source/Editor/LevelEditor/Private/LevelViewportLayout2x2.cpp
--- a/Engine/Source/Editor/LevelEditor/Private/LevelViewportLayout2x2.cpp
+++ b/Engine/Source/Editor/LevelEditor/Private/LevelViewportLayout2x2.cpp
@@ -51,7 +51,7 @@ TSharedRef<SWidget> FLevelViewportLayout2x2::MakeViewportLayout(const FString& L
 	.TopLeft()
 	[
 		SAssignNew( ViewportWidgetTL, SLevelViewport )
-		.ViewportType( LVT_OrthoYZ )		// Side
+		.ViewportType( LVT_OrthoYZ )		// Left
 		.ParentLayout( AsShared() )
 		.ParentLevelEditor( ParentLevelEditor )
 		.ConfigKey( TopLeftKey )
diff --git a/Engine/Source/Editor/LevelEditor/Private/SLevelViewportToolBar.cpp b/Engine/Source/Editor/LevelEditor/Private/SLevelViewportToolBar.cpp
--- a/Engine/Source/Editor/LevelEditor/Private/SLevelViewportToolBar.cpp
+++ b/Engine/Source/Editor/LevelEditor/Private/SLevelViewportToolBar.cpp
@@ -246,12 +246,24 @@ FText SLevelViewportToolBar::GetCameraMenuLabel() const
 				break;
 
 			case LVT_OrthoYZ:
-				Label = LOCTEXT("CameraMenuTitle_Side", "Side");
+				Label = LOCTEXT("CameraMenuTitle_Left", "Left");
 				break;
 
 			case LVT_OrthoXZ:
 				Label = LOCTEXT("CameraMenuTitle_Front", "Front");
 				break;
+
+			case LVT_OrthoNegativeXY:
+				Label = LOCTEXT("CameraMenuTitle_Bottom", "Bottom");
+				break;
+
+			case LVT_OrthoNegativeYZ:
+				Label = LOCTEXT("CameraMenuTitle_Right", "Right");
+				break;
+
+			case LVT_OrthoNegativeXZ:
+				Label = LOCTEXT("CameraMenuTitle_Back", "Back");
+				break;
 		}
 	}
 
@@ -303,8 +315,11 @@ const FSlateBrush* SLevelViewportToolBar::GetCameraMenuLabelIcon() const
 	{
 		static FName PerspectiveIcon("EditorViewport.Perspective");
 		static FName TopIcon("EditorViewport.Top");
-		static FName SideIcon("EditorViewport.Side");
+		static FName LeftIcon("EditorViewport.Left");
 		static FName FrontIcon("EditorViewport.Front");
+		static FName BottomIcon("EditorViewport.Bottom");
+		static FName RightIcon("EditorViewport.Right");
+		static FName BackIcon("EditorViewport.Back");
 
 		switch( PinnedViewport->GetLevelViewportClient().ViewportType )
 		{
@@ -317,12 +332,24 @@ const FSlateBrush* SLevelViewportToolBar::GetCameraMenuLabelIcon() const
 				break;
 
 			case LVT_OrthoYZ:
-				Icon = SideIcon;
+				Icon = LeftIcon;
 				break;
 
 			case LVT_OrthoXZ:
 				Icon = FrontIcon;
 				break;
+
+			case LVT_OrthoNegativeXY:
+				Icon = BottomIcon;
+				break;
+
+			case LVT_OrthoNegativeYZ:
+				Icon = RightIcon;
+				break;
+
+			case LVT_OrthoNegativeXZ:
+				Icon = BackIcon;
+				break;
 		}
 	}
 
@@ -652,8 +679,11 @@ TSharedRef<SWidget> SLevelViewportToolBar::GenerateCameraMenu() const
 
 	CameraMenuBuilder.BeginSection("LevelViewportCameraType_Ortho", LOCTEXT("CameraTypeHeader_Ortho", "Orthographic") );
 		CameraMenuBuilder.AddMenuEntry(FEditorViewportCommands::Get().Top);
-		CameraMenuBuilder.AddMenuEntry(FEditorViewportCommands::Get().Side);
+		CameraMenuBuilder.AddMenuEntry(FEditorViewportCommands::Get().Bottom);
+		CameraMenuBuilder.AddMenuEntry(FEditorViewportCommands::Get().Left);
+		CameraMenuBuilder.AddMenuEntry(FEditorViewportCommands::Get().Right);
 		CameraMenuBuilder.AddMenuEntry(FEditorViewportCommands::Get().Front);
+		CameraMenuBuilder.AddMenuEntry(FEditorViewportCommands::Get().Back);
 	CameraMenuBuilder.EndSection();
 
 	return CameraMenuBuilder.MakeWidget();
diff --git a/Engine/Source/Editor/Persona/Private/SAnimViewportToolBar.cpp b/Engine/Source/Editor/Persona/Private/SAnimViewportToolBar.cpp
--- a/Engine/Source/Editor/Persona/Private/SAnimViewportToolBar.cpp
+++ b/Engine/Source/Editor/Persona/Private/SAnimViewportToolBar.cpp
@@ -727,8 +727,11 @@ TSharedRef<SWidget> SAnimViewportToolBar::GenerateViewportTypeMenu() const
 
 	CameraMenuBuilder.BeginSection("LevelViewportCameraType_Ortho", LOCTEXT("CameraTypeHeader_Ortho", "Orthographic"));
 	CameraMenuBuilder.AddMenuEntry(FEditorViewportCommands::Get().Top);
-	CameraMenuBuilder.AddMenuEntry(FEditorViewportCommands::Get().Side);
+	CameraMenuBuilder.AddMenuEntry(FEditorViewportCommands::Get().Bottom);
+	CameraMenuBuilder.AddMenuEntry(FEditorViewportCommands::Get().Left);
+	CameraMenuBuilder.AddMenuEntry(FEditorViewportCommands::Get().Right);
 	CameraMenuBuilder.AddMenuEntry(FEditorViewportCommands::Get().Front);
+	CameraMenuBuilder.AddMenuEntry(FEditorViewportCommands::Get().Back);
 	CameraMenuBuilder.EndSection();
 
 	return CameraMenuBuilder.MakeWidget();
@@ -821,12 +824,24 @@ FText SAnimViewportToolBar::GetCameraMenuLabel() const
 			break;
 
 		case LVT_OrthoYZ:
-			Label = LOCTEXT("CameraMenuTitle_Side", "Side");
+			Label = LOCTEXT("CameraMenuTitle_Left", "Left");
 			break;
 
 		case LVT_OrthoXZ:
 			Label = LOCTEXT("CameraMenuTitle_Front", "Front");
 			break;
+
+		case LVT_OrthoNegativeXY:
+			Label = LOCTEXT("CameraMenuTitle_Bottom", "Bottom");
+			break;
+
+		case LVT_OrthoNegativeYZ:
+			Label = LOCTEXT("CameraMenuTitle_Right", "Right");
+			break;
+
+		case LVT_OrthoNegativeXZ:
+			Label = LOCTEXT("CameraMenuTitle_Back", "Back");
+			break;
 		}
 	}
 
@@ -850,12 +865,24 @@ const FSlateBrush* SAnimViewportToolBar::GetCameraMenuLabelIcon() const
 			break;
 
 		case LVT_OrthoYZ:
-			Icon = FName("EditorViewport.Side");
+			Icon = FName("EditorViewport.Left");
 			break;
 
 		case LVT_OrthoXZ:
 			Icon = FName("EditorViewport.Front");
 			break;
+
+		case LVT_OrthoNegativeXY:
+			Icon = FName("EditorViewport.Bottom");
+			break;
+
+		case LVT_OrthoNegativeYZ:
+			Icon = FName("EditorViewport.Right");
+			break;
+
+		case LVT_OrthoNegativeXZ:
+			Icon = FName("EditorViewport.Back");
+			break;
 		}
 	}
 
diff --git a/Engine/Source/Editor/PhAT/Private/PhAT.cpp b/Engine/Source/Editor/PhAT/Private/PhAT.cpp
--- a/Engine/Source/Editor/PhAT/Private/PhAT.cpp
+++ b/Engine/Source/Editor/PhAT/Private/PhAT.cpp
@@ -1191,7 +1191,7 @@ void FPhAT::BindCommands()
 		);
 
 	ToolkitCommands->MapAction(
-		Commands.SideView,
+		Commands.LeftView,
 		FExecuteAction::CreateSP(this, &FPhAT::OnViewType, ELevelViewportType::LVT_OrthoYZ),
 		FCanExecuteAction()
 		);
@@ -1202,6 +1202,24 @@ void FPhAT::BindCommands()
 		FCanExecuteAction()
 		);
 
+	ToolkitCommands->MapAction(
+		Commands.BottomView,
+		FExecuteAction::CreateSP(this, &FPhAT::OnViewType, ELevelViewportType::LVT_OrthoNegativeXY),
+		FCanExecuteAction()
+		);
+
+	ToolkitCommands->MapAction(
+		Commands.RightView,
+		FExecuteAction::CreateSP(this, &FPhAT::OnViewType, ELevelViewportType::LVT_OrthoNegativeYZ),
+		FCanExecuteAction()
+		);
+
+	ToolkitCommands->MapAction(
+		Commands.BackView,
+		FExecuteAction::CreateSP(this, &FPhAT::OnViewType, ELevelViewportType::LVT_OrthoNegativeXZ),
+		FCanExecuteAction()
+		);
+
 	ToolkitCommands->MapAction(
 		Commands.MakeBodyKinematic,
 		FExecuteAction::CreateSP(this, &FPhAT::OnSetBodyPhysicsType, EPhysicsType::PhysType_Kinematic ),
diff --git a/Engine/Source/Editor/PhAT/Private/PhATActions.cpp b/Engine/Source/Editor/PhAT/Private/PhATActions.cpp
--- a/Engine/Source/Editor/PhAT/Private/PhATActions.cpp
+++ b/Engine/Source/Editor/PhAT/Private/PhATActions.cpp
@@ -69,9 +69,12 @@ void FPhATCommands::RegisterCommands()
 	UI_COMMAND(HierarchyFilterAll, "All Bones", "Show Entire Hierarchy", EUserInterfaceActionType::Button, FInputGesture());
 	UI_COMMAND(HierarchyFilterBodies, "Bones With Bodies", "Filter Bones With Bodies", EUserInterfaceActionType::Button, FInputGesture());
 	UI_COMMAND(PerspectiveView, "Perspective", "Perspective", EUserInterfaceActionType::Button, FInputGesture());
-	UI_COMMAND(SideView, "Side", "Orthographic view from side", EUserInterfaceActionType::Button, FInputGesture());
+	UI_COMMAND(LeftView, "Left", "Orthographic view from left", EUserInterfaceActionType::Button, FInputGesture());
 	UI_COMMAND(TopView, "Top", "Orthographic view from top", EUserInterfaceActionType::Button, FInputGesture());
 	UI_COMMAND(FrontView, "Front", "Orthographic view from front", EUserInterfaceActionType::Button, FInputGesture());
+	UI_COMMAND(RightView, "Right", "Orthographic view from right", EUserInterfaceActionType::Button, FInputGesture());
+	UI_COMMAND(BottomView, "Bottom", "Orthographic view from bottom", EUserInterfaceActionType::Button, FInputGesture());
+	UI_COMMAND(BackView, "Back", "Orthographic view from back", EUserInterfaceActionType::Button, FInputGesture());
 	UI_COMMAND(Mirror, "Mirror", "Finds the body on the other side and duplicates constraint and body", EUserInterfaceActionType::Button, FInputGesture(EKeys::M));
 
 
diff --git a/Engine/Source/Editor/PhAT/Private/PhATActions.h b/Engine/Source/Editor/PhAT/Private/PhATActions.h
--- a/Engine/Source/Editor/PhAT/Private/PhATActions.h
+++ b/Engine/Source/Editor/PhAT/Private/PhATActions.h
@@ -81,8 +81,11 @@ public:
 	TSharedPtr<FUICommandInfo> HierarchyFilterBodies;
 	TSharedPtr<FUICommandInfo> PerspectiveView;
 	TSharedPtr<FUICommandInfo> TopView;
-	TSharedPtr<FUICommandInfo> SideView;
+	TSharedPtr<FUICommandInfo> LeftView;
 	TSharedPtr<FUICommandInfo> FrontView;
+	TSharedPtr<FUICommandInfo> BottomView;
+	TSharedPtr<FUICommandInfo> RightView;
+	TSharedPtr<FUICommandInfo> BackView;
 	TSharedPtr<FUICommandInfo> Mirror;
 
 	/** Hotkey only commands */
diff --git a/Engine/Source/Editor/PhAT/Private/SPhATPreviewToolbar.cpp b/Engine/Source/Editor/PhAT/Private/SPhATPreviewToolbar.cpp
--- a/Engine/Source/Editor/PhAT/Private/SPhATPreviewToolbar.cpp
+++ b/Engine/Source/Editor/PhAT/Private/SPhATPreviewToolbar.cpp
@@ -80,8 +80,11 @@ TSharedRef<SWidget> SPhATPreviewViewportToolBar::GeneratePerspectiveMenu() const
 	{
 		PerspectiveMenuBuilder.AddMenuEntry(Actions.PerspectiveView);
 		PerspectiveMenuBuilder.AddMenuEntry(Actions.TopView);
-		PerspectiveMenuBuilder.AddMenuEntry(Actions.SideView);
+		PerspectiveMenuBuilder.AddMenuEntry(Actions.BottomView);
+		PerspectiveMenuBuilder.AddMenuEntry(Actions.LeftView);
+		PerspectiveMenuBuilder.AddMenuEntry(Actions.RightView);
 		PerspectiveMenuBuilder.AddMenuEntry(Actions.FrontView);
+		PerspectiveMenuBuilder.AddMenuEntry(Actions.BackView);
 	}
 
 	return PerspectiveMenuBuilder.MakeWidget();
diff --git a/Engine/Source/Editor/UnrealEd/Classes/Editor/UnrealEdTypes.h b/Engine/Source/Editor/UnrealEd/Classes/Editor/UnrealEdTypes.h
--- a/Engine/Source/Editor/UnrealEd/Classes/Editor/UnrealEdTypes.h
+++ b/Engine/Source/Editor/UnrealEd/Classes/Editor/UnrealEdTypes.h
@@ -115,11 +115,14 @@ struct FLightmassParameterizedMaterialSettings
 UENUM()
 enum ELevelViewportType
 {
-	LVT_OrthoXY = 0,	// Top
-	LVT_OrthoXZ = 1,	// Front
-	LVT_OrthoYZ = 2,	// Side
+	LVT_OrthoXY = 0,			// Top
+	LVT_OrthoXZ = 1,			// Front
+	LVT_OrthoYZ = 2,			// Left
 	LVT_Perspective = 3,
 	LVT_OrthoFreelook = 4,
+	LVT_OrthoNegativeXY = 5,	// Bottom
+	LVT_OrthoNegativeXZ = 6,	// Back
+	LVT_OrthoNegativeYZ = 7,	// Right
 	LVT_MAX,
 
 	LVT_None = 255,
diff --git a/Engine/Source/Editor/UnrealEd/Private/DragTool_BoxSelect.cpp b/Engine/Source/Editor/UnrealEd/Private/DragTool_BoxSelect.cpp
--- a/Engine/Source/Editor/UnrealEd/Private/DragTool_BoxSelect.cpp
+++ b/Engine/Source/Editor/UnrealEd/Private/DragTool_BoxSelect.cpp
@@ -255,14 +255,17 @@ void FDragTool_ActorBoxSelect::CalculateBox( FBox& OutBox )
 	switch(LevelViewportClient->ViewportType)
 	{
 	case LVT_OrthoXY:
+	case LVT_OrthoNegativeXY:
 		OutBox.Min.Z = -WORLD_MAX;
 		OutBox.Max.Z = WORLD_MAX;
 		break;
 	case LVT_OrthoXZ:
+	case LVT_OrthoNegativeXZ:
 		OutBox.Min.Y = -WORLD_MAX;
 		OutBox.Max.Y = WORLD_MAX;
 		break;
 	case LVT_OrthoYZ:
+	case LVT_OrthoNegativeYZ:
 		OutBox.Min.X = -WORLD_MAX;
 		OutBox.Max.X = WORLD_MAX;
 		break;
diff --git a/Engine/Source/Editor/UnrealEd/Private/Editor/ActorPositioning.cpp b/Engine/Source/Editor/UnrealEd/Private/Editor/ActorPositioning.cpp
--- a/Engine/Source/Editor/UnrealEd/Private/Editor/ActorPositioning.cpp
+++ b/Engine/Source/Editor/UnrealEd/Private/Editor/ActorPositioning.cpp
@@ -33,7 +33,8 @@ FActorPositionTraceResult FActorPositioning::TraceWorldForPosition(const FViewpo
 
 	// Start with a ray that encapsulates the entire world
 	FVector RayStart = Cursor.GetOrigin();
-	if (ViewportType == LVT_OrthoXY || ViewportType == LVT_OrthoXZ || ViewportType == LVT_OrthoYZ)
+	if (ViewportType == LVT_OrthoXY || ViewportType == LVT_OrthoXZ || ViewportType == LVT_OrthoYZ ||
+		ViewportType == LVT_OrthoNegativeXY || ViewportType == LVT_OrthoNegativeXZ || ViewportType == LVT_OrthoNegativeYZ)
 	{
 		RayStart -= Cursor.GetDirection() * HALF_WORLD_MAX/2;
 	}
diff --git a/Engine/Source/Editor/UnrealEd/Private/EditorModeInterpolation.cpp b/Engine/Source/Editor/UnrealEd/Private/EditorModeInterpolation.cpp
--- a/Engine/Source/Editor/UnrealEd/Private/EditorModeInterpolation.cpp
+++ b/Engine/Source/Editor/UnrealEd/Private/EditorModeInterpolation.cpp
@@ -227,12 +227,15 @@ bool FModeTool_InterpEdit::InputAxis(FEditorViewportClient* InViewportClient, FV
 		switch(ViewportType)
 		{
 		case LVT_OrthoXY:
+		case LVT_OrthoNegativeXY:
 			InViewportClient->SetCurrentWidgetAxis(EAxisList::XY);
 			break;
 		case LVT_OrthoXZ:
+		case LVT_OrthoNegativeXZ:
 			InViewportClient->SetCurrentWidgetAxis(EAxisList::XZ);
 			break;
 		case LVT_OrthoYZ:
+		case LVT_OrthoNegativeYZ:
 			InViewportClient->SetCurrentWidgetAxis(EAxisList::YZ);
 			break;
 		}
diff --git a/Engine/Source/Editor/UnrealEd/Private/EditorViewportClient.cpp b/Engine/Source/Editor/UnrealEd/Private/EditorViewportClient.cpp
--- a/Engine/Source/Editor/UnrealEd/Private/EditorViewportClient.cpp
+++ b/Engine/Source/Editor/UnrealEd/Private/EditorViewportClient.cpp
@@ -698,6 +698,30 @@ FSceneView* FEditorViewportClient::CalcSceneView(FSceneViewFamily* ViewFamily)
 					FPlane(0, 1, 0, 0),
 					FPlane(0, 0, ViewLocation.X, 1));
 			}
+			else if (EffectiveViewportType == LVT_OrthoNegativeXY)
+			{
+				ViewInitOptions.ViewRotationMatrix = FMatrix(
+					FPlane(-1, 0, 0, 0),
+					FPlane(0, -1, 0, 0),
+					FPlane(0, 0, 1, 0),
+					FPlane(0, 0, -ViewLocation.Z, 1));
+			}
+			else if (EffectiveViewportType == LVT_OrthoNegativeXZ)
+			{
+				ViewInitOptions.ViewRotationMatrix = FMatrix(
+					FPlane(-1, 0, 0, 0),
+					FPlane(0, 0, 1, 0),
+					FPlane(0, 1, 0, 0),
+					FPlane(0, 0, -ViewLocation.Y, 1));
+			}
+			else if (EffectiveViewportType == LVT_OrthoNegativeYZ)
+			{
+				ViewInitOptions.ViewRotationMatrix = FMatrix(
+					FPlane(0, 0, -1, 0),
+					FPlane(-1, 0, 0, 0),
+					FPlane(0, 1, 0, 0),
+					FPlane(0, 0, ViewLocation.X, 1));
+			}
 			else
 			{
 				// Unknown viewport type
@@ -2537,6 +2561,15 @@ void FEditorViewportClient::OnOrthoZoom( const struct FInputEventState& InputSta
 		case LVT_OrthoYZ:
 			OldOffsetFromCenter.Set(0.0f, DeltaFromCenterX, DeltaFromCenterY);
 			break;
+		case LVT_OrthoNegativeXY:
+			OldOffsetFromCenter.Set(-DeltaFromCenterX, -DeltaFromCenterY, 0.0f);
+			break;
+		case LVT_OrthoNegativeXZ:
+			OldOffsetFromCenter.Set(-DeltaFromCenterX, 0.0f, DeltaFromCenterY);
+			break;
+		case LVT_OrthoNegativeYZ:
+			OldOffsetFromCenter.Set(0.0f, -DeltaFromCenterX, DeltaFromCenterY);
+			break;
 		case LVT_OrthoFreelook:
 			//@TODO: CAMERA: How to handle this
 			break;
@@ -2930,6 +2963,27 @@ void FEditorViewportClient::Draw(FViewport* InViewport, FCanvas* Canvas)
 				DrawScaleUnits(Viewport, Canvas, *View);
 				break;
 			}
+		case LVT_OrthoNegativeXY:
+			{
+				const FRotator XYRot(90.0f, 90.0f, 0.0f);
+				DrawAxes(Viewport, Canvas, &XYRot, EAxisList::XY);
+				DrawScaleUnits(Viewport, Canvas, *View);
+				break;
+			}
+		case LVT_OrthoNegativeXZ:
+			{
+				const FRotator XZRot(0.0f, 90.0f, 0.0f);
+				DrawAxes(Viewport, Canvas, &XZRot, EAxisList::XZ);
+				DrawScaleUnits(Viewport, Canvas, *View);
+				break;
+			}
+		case LVT_OrthoNegativeYZ:
+			{
+				const FRotator YZRot(0.0f, 180.0f, 0.0f);
+				DrawAxes(Viewport, Canvas, &YZRot, EAxisList::YZ);
+				DrawScaleUnits(Viewport, Canvas, *View);
+				break;
+			}
 		default:
 			{
 				DrawAxes(Viewport, Canvas);
@@ -3093,6 +3147,9 @@ FVector FEditorViewportClient::TranslateDelta( FKey InKey, float InDelta, bool I
 	case LVT_OrthoXY:
 	case LVT_OrthoXZ:
 	case LVT_OrthoYZ:
+	case LVT_OrthoNegativeXY:
+	case LVT_OrthoNegativeXZ:
+	case LVT_OrthoNegativeYZ:
 		{
 			LastMouseX += X;
 			LastMouseY -= Y;
@@ -3140,10 +3197,19 @@ FVector FEditorViewportClient::TranslateDelta( FKey InKey, float InDelta, bool I
 						vec.Y *= -1.0f;
 						break;
 					case LVT_OrthoXZ:
-						vec = FVector( X * UnitsPerPixel, 0.f, Y * UnitsPerPixel );
+						vec = FVector(X * UnitsPerPixel, 0.f, Y * UnitsPerPixel);
 						break;
 					case LVT_OrthoYZ:
-						vec = FVector( 0.f, X * UnitsPerPixel, Y * UnitsPerPixel );
+						vec = FVector(0.f, X * UnitsPerPixel, Y * UnitsPerPixel);
+						break;
+ 					case LVT_OrthoNegativeXY:
+ 						vec = FVector(-X * UnitsPerPixel, -Y * UnitsPerPixel, 0.0f);
+ 						break;
+					case LVT_OrthoNegativeXZ:
+						vec = FVector(-X * UnitsPerPixel, 0.f, Y * UnitsPerPixel);
+						break;
+					case LVT_OrthoNegativeYZ:
+						vec = FVector(0.f, -X * UnitsPerPixel, Y * UnitsPerPixel);
 						break;
 					}
 				}
@@ -3265,6 +3331,9 @@ bool FEditorViewportClient::InputGesture(FViewport* InViewport, EGestureEvent::T
 	case LVT_OrthoXY:
 	case LVT_OrthoXZ:
 	case LVT_OrthoYZ:
+	case LVT_OrthoNegativeXY:
+	case LVT_OrthoNegativeXZ:
+	case LVT_OrthoNegativeYZ:
 		{
 			if (GestureType == EGestureEvent::Scroll && !LeftMouseButtonDown && !RightMouseButtonDown)
 			{
@@ -3284,6 +3353,15 @@ bool FEditorViewportClient::InputGesture(FViewport* InViewport, EGestureEvent::T
 					case LVT_OrthoYZ:
 						CurrentGestureDragDelta += FVector(0, -AdjustedGestureDelta.X, AdjustedGestureDelta.Y);
 						break;
+					case LVT_OrthoNegativeXY:
+						CurrentGestureDragDelta += FVector(AdjustedGestureDelta.X, -AdjustedGestureDelta.Y, 0);
+						break;
+					case LVT_OrthoNegativeXZ:
+						CurrentGestureDragDelta += FVector(AdjustedGestureDelta.X, 0, AdjustedGestureDelta.Y);
+						break;
+					case LVT_OrthoNegativeYZ:
+						CurrentGestureDragDelta += FVector(0, AdjustedGestureDelta.X, AdjustedGestureDelta.Y);
+						break;
 				}
 
 				FEditorViewportStats::Used(FEditorViewportStats::CAT_ORTHOGRAPHIC_GESTURE_SCROLL);
@@ -3371,6 +3449,9 @@ void FEditorViewportClient::ConvertMovementToDragRot(const FVector& InDelta,
 	case LVT_OrthoXY:
 	case LVT_OrthoXZ:
 	case LVT_OrthoYZ:
+	case LVT_OrthoNegativeXY:
+	case LVT_OrthoNegativeXZ:
+	case LVT_OrthoNegativeYZ:
 		{
 			if( ( LeftMouseButtonDown || bIsUsingTrackpad ) && RightMouseButtonDown )
 			{
@@ -3452,6 +3533,9 @@ void FEditorViewportClient::ConvertMovementToOrbitDragRot(const FVector& InDelta
 	case LVT_OrthoXY:
 	case LVT_OrthoXZ:
 	case LVT_OrthoYZ:
+	case LVT_OrthoNegativeXY:
+	case LVT_OrthoNegativeXZ:
+	case LVT_OrthoNegativeYZ:
 		{
 			if( ( LeftMouseButtonDown || bIsUsingTrackpad ) && RightMouseButtonDown )
 			{
@@ -3775,6 +3859,9 @@ void FEditorViewportClient::MoveViewportCamera(const FVector& InDrag, const FRot
 	case LVT_OrthoXY:
 	case LVT_OrthoXZ:
 	case LVT_OrthoYZ:
+	case LVT_OrthoNegativeXY:
+	case LVT_OrthoNegativeXZ:
+	case LVT_OrthoNegativeYZ:
 		{
 			const bool LeftMouseButtonDown = Viewport->KeyState(EKeys::LeftMouseButton);
 			const bool RightMouseButtonDown = Viewport->KeyState(EKeys::RightMouseButton);
@@ -3859,12 +3946,15 @@ void FEditorViewportClient::CheckHoveredHitProxy( HHitProxy* HoveredHitProxy )
 				switch( GetViewportType() )
 				{
 				case LVT_OrthoXY:
+				case LVT_OrthoNegativeXY:
 					NewAxis = EAxisList::Z;
 					break;
 				case LVT_OrthoXZ:
+				case LVT_OrthoNegativeXZ:
 					NewAxis = EAxisList::Y;
 					break;
 				case LVT_OrthoYZ:
+				case LVT_OrthoNegativeYZ:
 					NewAxis = EAxisList::X;
 					break;
 				default:
diff --git a/Engine/Source/Editor/UnrealEd/Private/EditorViewportCommands.cpp b/Engine/Source/Editor/UnrealEd/Private/EditorViewportCommands.cpp
--- a/Engine/Source/Editor/UnrealEd/Private/EditorViewportCommands.cpp
+++ b/Engine/Source/Editor/UnrealEd/Private/EditorViewportCommands.cpp
@@ -10,8 +10,11 @@ void FEditorViewportCommands::RegisterCommands()
 {
 	UI_COMMAND( Perspective, "Perspective", "Switches the viewport to perspective view", EUserInterfaceActionType::RadioButton, FInputGesture( EModifierKey::Alt, EKeys::G ) );
 	UI_COMMAND( Front, "Front", "Switches the viewport to front view", EUserInterfaceActionType::RadioButton, FInputGesture( EModifierKey::Alt, EKeys::H ) );
+	UI_COMMAND( Back, "Back", "Switches the viewport to back view", EUserInterfaceActionType::RadioButton, FInputGesture( EModifierKey::Alt | EModifierKey::Shift, EKeys::H ) );
 	UI_COMMAND( Top, "Top", "Switches the viewport to top view", EUserInterfaceActionType::RadioButton, FInputGesture( EModifierKey::Alt, EKeys::J ) );
-	UI_COMMAND( Side, "Side", "Switches the viewport to side view", EUserInterfaceActionType::RadioButton, FInputGesture( EModifierKey::Alt, EKeys::K ) );
+	UI_COMMAND( Bottom, "Bottom", "Switches the viewport to bottom view", EUserInterfaceActionType::RadioButton, FInputGesture( EModifierKey::Alt | EModifierKey::Shift, EKeys::J ) );
+	UI_COMMAND( Left, "Left", "Switches the viewport to left view", EUserInterfaceActionType::RadioButton, FInputGesture( EModifierKey::Alt, EKeys::K ) );
+	UI_COMMAND( Right, "Right", "Switches the viewport to right view", EUserInterfaceActionType::RadioButton, FInputGesture( EModifierKey::Alt | EModifierKey::Shift, EKeys::K ) );
 
 	UI_COMMAND( WireframeMode, "Brush Wireframe View Mode", "Renders the scene in brush wireframe", EUserInterfaceActionType::RadioButton, FInputGesture( EModifierKey::Alt, EKeys::Two ) );
 	UI_COMMAND( UnlitMode, "Unlit View Mode", "Renders the scene with no lights", EUserInterfaceActionType::RadioButton, FInputGesture( EModifierKey::Alt, EKeys::Three ) );
diff --git a/Engine/Source/Editor/UnrealEd/Private/LevelEditorViewport.cpp b/Engine/Source/Editor/UnrealEd/Private/LevelEditorViewport.cpp
--- a/Engine/Source/Editor/UnrealEd/Private/LevelEditorViewport.cpp
+++ b/Engine/Source/Editor/UnrealEd/Private/LevelEditorViewport.cpp
@@ -139,17 +139,17 @@ static FVector4 AttemptToSnapLocationToOriginPlane( const FViewportCursorLocatio
 			Location = NewLocFloor;
 		}
 	}
-	else if ( ViewportType == LVT_OrthoXY )
+	else if ( ViewportType == LVT_OrthoXY || ViewportType == LVT_OrthoNegativeXY )
 	{
 		// In ortho place the brush at the origin of the hidden axis
 		Location.Z = 0;
 	}
-	else if ( ViewportType == LVT_OrthoXZ )
+	else if ( ViewportType == LVT_OrthoXZ || ViewportType == LVT_OrthoNegativeXZ )
 	{
 		// In ortho place the brush at the origin of the hidden axis
 		Location.Y = 0;
 	}
-	else if ( ViewportType == LVT_OrthoYZ )
+	else if ( ViewportType == LVT_OrthoYZ || ViewportType == LVT_OrthoNegativeYZ )
 	{
 		// In ortho place the brush at the origin of the hidden axis
 		Location.X = 0;
@@ -2747,10 +2747,13 @@ EAxisList::Type FLevelEditorViewportClient::GetHorizAxis() const
 	switch( GetViewportType() )
 	{
 	case LVT_OrthoXY:
+	case LVT_OrthoNegativeXY:
 		return EAxisList::X;
 	case LVT_OrthoXZ:
+	case LVT_OrthoNegativeXZ:
 		return EAxisList::X;
 	case LVT_OrthoYZ:
+	case LVT_OrthoNegativeYZ:
 		return EAxisList::Y;
 	}
 
@@ -2766,10 +2769,13 @@ EAxisList::Type FLevelEditorViewportClient::GetVertAxis() const
 	switch( GetViewportType() )
 	{
 	case LVT_OrthoXY:
+	case LVT_OrthoNegativeXY:
 		return EAxisList::Y;
 	case LVT_OrthoXZ:
+	case LVT_OrthoNegativeXZ:
 		return EAxisList::Z;
 	case LVT_OrthoYZ:
+	case LVT_OrthoNegativeYZ:
 		return EAxisList::Z;
 	}
 
diff --git a/Engine/Source/Editor/UnrealEd/Private/LevelViewportClickHandlers.cpp b/Engine/Source/Editor/UnrealEd/Private/LevelViewportClickHandlers.cpp
--- a/Engine/Source/Editor/UnrealEd/Private/LevelViewportClickHandlers.cpp
+++ b/Engine/Source/Editor/UnrealEd/Private/LevelViewportClickHandlers.cpp
@@ -519,10 +519,21 @@ namespace ClickHandlers
 		bool bResult = false;
 		switch( ViewportType )
 		{
-		case LVT_OrthoXY:	bResult = FMath::Abs(Vec0.X - Vec1.X) < Tolerance && FMath::Abs(Vec0.Y - Vec1.Y) < Tolerance;	break;
-		case LVT_OrthoXZ:	bResult = FMath::Abs(Vec0.X - Vec1.X) < Tolerance && FMath::Abs(Vec0.Z - Vec1.Z) < Tolerance;	break;
-		case LVT_OrthoYZ:	bResult = FMath::Abs(Vec0.Y - Vec1.Y) < Tolerance && FMath::Abs(Vec0.Z - Vec1.Z) < Tolerance;	break;
-		default:			check( 0 );		break;
+		case LVT_OrthoXY:
+		case LVT_OrthoNegativeXY:
+			bResult = FMath::Abs(Vec0.X - Vec1.X) < Tolerance && FMath::Abs(Vec0.Y - Vec1.Y) < Tolerance;
+			break;
+		case LVT_OrthoXZ:
+		case LVT_OrthoNegativeXZ:
+			bResult = FMath::Abs(Vec0.X - Vec1.X) < Tolerance && FMath::Abs(Vec0.Z - Vec1.Z) < Tolerance;
+			break;
+		case LVT_OrthoYZ:
+		case LVT_OrthoNegativeYZ:
+			bResult = FMath::Abs(Vec0.Y - Vec1.Y) < Tolerance && FMath::Abs(Vec0.Z - Vec1.Z) < Tolerance;
+			break;
+		default:
+			check( 0 );
+			break;
 		}
 		return bResult;
 	}
diff --git a/Engine/Source/Editor/UnrealEd/Private/MouseDeltaTracker.cpp b/Engine/Source/Editor/UnrealEd/Private/MouseDeltaTracker.cpp
--- a/Engine/Source/Editor/UnrealEd/Private/MouseDeltaTracker.cpp
+++ b/Engine/Source/Editor/UnrealEd/Private/MouseDeltaTracker.cpp
@@ -80,12 +80,15 @@ void FMouseDeltaTracker::DetermineCurrentAxis(FEditorViewportClient* InViewportC
 								}
 								break;
 							case LVT_OrthoXY:
+							case LVT_OrthoNegativeXY:
 								InViewportClient->SetCurrentWidgetAxis( EAxisList::XY );
 								break;
 							case LVT_OrthoXZ:
+							case LVT_OrthoNegativeXZ:
 								InViewportClient->SetCurrentWidgetAxis( EAxisList::XZ );
 								break;
 							case LVT_OrthoYZ:
+							case LVT_OrthoNegativeYZ:
 								InViewportClient->SetCurrentWidgetAxis( EAxisList::YZ );
 								break;
 							default:
@@ -117,13 +120,16 @@ void FMouseDeltaTracker::DetermineCurrentAxis(FEditorViewportClient* InViewportC
 							}
 							break;
 						case LVT_OrthoXY:
-							InViewportClient->SetCurrentWidgetAxis( EAxisList::XY );
+						case LVT_OrthoNegativeXY:
+							InViewportClient->SetCurrentWidgetAxis(EAxisList::XY);
 							break;
 						case LVT_OrthoXZ:
-							InViewportClient->SetCurrentWidgetAxis( EAxisList::XZ );
+						case LVT_OrthoNegativeXZ:
+							InViewportClient->SetCurrentWidgetAxis(EAxisList::XZ);
 							break;
 						case LVT_OrthoYZ:
-							InViewportClient->SetCurrentWidgetAxis( EAxisList::YZ );
+						case LVT_OrthoNegativeYZ:
+							InViewportClient->SetCurrentWidgetAxis(EAxisList::YZ);
 							break;
 						default:
 							break;
@@ -148,13 +154,16 @@ void FMouseDeltaTracker::DetermineCurrentAxis(FEditorViewportClient* InViewportC
 							}
 							break;
 						case LVT_OrthoXY:
-							InViewportClient->SetCurrentWidgetAxis( EAxisList::Z );
+						case LVT_OrthoNegativeXY:
+							InViewportClient->SetCurrentWidgetAxis(EAxisList::Z);
 							break;
 						case LVT_OrthoXZ:
-							InViewportClient->SetCurrentWidgetAxis( EAxisList::Y );
+						case LVT_OrthoNegativeXZ:
+							InViewportClient->SetCurrentWidgetAxis(EAxisList::Y);
 							break;
 						case LVT_OrthoYZ:
-							InViewportClient->SetCurrentWidgetAxis( EAxisList::X );
+						case LVT_OrthoNegativeYZ:
+							InViewportClient->SetCurrentWidgetAxis(EAxisList::X);
 							break;
 						default:
 							break;
diff --git a/Engine/Source/Editor/UnrealEd/Private/SCommonEditorViewportToolbarBase.cpp b/Engine/Source/Editor/UnrealEd/Private/SCommonEditorViewportToolbarBase.cpp
--- a/Engine/Source/Editor/UnrealEd/Private/SCommonEditorViewportToolbarBase.cpp
+++ b/Engine/Source/Editor/UnrealEd/Private/SCommonEditorViewportToolbarBase.cpp
@@ -119,13 +119,25 @@ FText SCommonEditorViewportToolbarBase::GetCameraMenuLabel() const
 			break;
 
 		case LVT_OrthoYZ:
-			Label = LOCTEXT("CameraMenuTitle_Side", "Side");
+			Label = LOCTEXT("CameraMenuTitle_Left", "Left");
 			break;
 
 		case LVT_OrthoXZ:
 			Label = LOCTEXT("CameraMenuTitle_Front", "Front");
 			break;
 
+		case LVT_OrthoNegativeXY:
+			Label = LOCTEXT("CameraMenuTitle_Bottom", "Bottom");
+			break;
+
+		case LVT_OrthoNegativeYZ:
+			Label = LOCTEXT("CameraMenuTitle_Right", "Right");
+			break;
+
+		case LVT_OrthoNegativeXZ:
+			Label = LOCTEXT("CameraMenuTitle_Back", "Back");
+			break;
+
 		case LVT_OrthoFreelook:
 			Label = LOCTEXT("CameraMenuTitle_OrthoFreelook", "Ortho");
 			break;
@@ -149,12 +161,24 @@ const FSlateBrush* SCommonEditorViewportToolbarBase::GetCameraMenuLabelIcon() co
 			break;
 
 		case LVT_OrthoYZ:
-			Icon = FName( "EditorViewport.Side" );
+			Icon = FName( "EditorViewport.Left" );
 			break;
 
 		case LVT_OrthoXZ:
 			Icon = FName( "EditorViewport.Front" );
 			break;
+
+		case LVT_OrthoNegativeXY:
+			Icon = FName("EditorViewport.Bottom");
+			break;
+
+		case LVT_OrthoNegativeYZ:
+			Icon = FName("EditorViewport.Right");
+			break;
+
+		case LVT_OrthoNegativeXZ:
+			Icon = FName("EditorViewport.Back");
+			break;
 	}
 
 	return FEditorStyle::GetBrush( Icon );
@@ -201,9 +225,12 @@ TSharedRef<SWidget> SCommonEditorViewportToolbarBase::GenerateCameraMenu() const
 
 	CameraMenuBuilder.BeginSection("LevelViewportCameraType_Ortho", LOCTEXT("CameraTypeHeader_Ortho", "Othographic") );
 		CameraMenuBuilder.AddMenuEntry(FEditorViewportCommands::Get().Top);
-		CameraMenuBuilder.AddMenuEntry(FEditorViewportCommands::Get().Side);
+		CameraMenuBuilder.AddMenuEntry(FEditorViewportCommands::Get().Bottom);
+		CameraMenuBuilder.AddMenuEntry(FEditorViewportCommands::Get().Left);
+		CameraMenuBuilder.AddMenuEntry(FEditorViewportCommands::Get().Right);
 		CameraMenuBuilder.AddMenuEntry(FEditorViewportCommands::Get().Front);
-	CameraMenuBuilder.EndSection();
+		CameraMenuBuilder.AddMenuEntry(FEditorViewportCommands::Get().Back);
+		CameraMenuBuilder.EndSection();
 
 	return CameraMenuBuilder.MakeWidget();
 }
diff --git a/Engine/Source/Editor/UnrealEd/Private/SEditorViewport.cpp b/Engine/Source/Editor/UnrealEd/Private/SEditorViewport.cpp
--- a/Engine/Source/Editor/UnrealEd/Private/SEditorViewport.cpp
+++ b/Engine/Source/Editor/UnrealEd/Private/SEditorViewport.cpp
@@ -174,7 +174,7 @@ void SEditorViewport::BindCommands()
 		FIsActionChecked::CreateSP(ClientRef, &FEditorViewportClient::IsActiveViewportType, LVT_OrthoXZ));
 
 	CommandListRef.MapAction( 
-		Commands.Side,
+		Commands.Left,
 		FExecuteAction::CreateSP( ClientRef, &FEditorViewportClient::SetViewportType, LVT_OrthoYZ ),
 		FCanExecuteAction(),
 		FIsActionChecked::CreateSP(ClientRef, &FEditorViewportClient::IsActiveViewportType, LVT_OrthoYZ));
@@ -185,6 +185,24 @@ void SEditorViewport::BindCommands()
 		FCanExecuteAction(),
 		FIsActionChecked::CreateSP(ClientRef, &FEditorViewportClient::IsActiveViewportType, LVT_OrthoXY));
 
+	CommandListRef.MapAction(
+		Commands.Back,
+		FExecuteAction::CreateSP(ClientRef, &FEditorViewportClient::SetViewportType, LVT_OrthoNegativeXZ),
+		FCanExecuteAction(),
+		FIsActionChecked::CreateSP(ClientRef, &FEditorViewportClient::IsActiveViewportType, LVT_OrthoNegativeXZ));
+
+	CommandListRef.MapAction(
+		Commands.Right,
+		FExecuteAction::CreateSP(ClientRef, &FEditorViewportClient::SetViewportType, LVT_OrthoNegativeYZ),
+		FCanExecuteAction(),
+		FIsActionChecked::CreateSP(ClientRef, &FEditorViewportClient::IsActiveViewportType, LVT_OrthoNegativeYZ));
+
+	CommandListRef.MapAction(
+		Commands.Bottom,
+		FExecuteAction::CreateSP(ClientRef, &FEditorViewportClient::SetViewportType, LVT_OrthoNegativeXY),
+		FCanExecuteAction(),
+		FIsActionChecked::CreateSP(ClientRef, &FEditorViewportClient::IsActiveViewportType, LVT_OrthoNegativeXY));
+
 	CommandListRef.MapAction(
 		Commands.ScreenCapture,
 		FExecuteAction::CreateSP( this, &SEditorViewport::OnScreenCapture ),
diff --git a/Engine/Source/Editor/UnrealEd/Public/EditorViewportCommands.h b/Engine/Source/Editor/UnrealEd/Public/EditorViewportCommands.h
--- a/Engine/Source/Editor/UnrealEd/Public/EditorViewportCommands.h
+++ b/Engine/Source/Editor/UnrealEd/Public/EditorViewportCommands.h
@@ -26,12 +26,21 @@ public:
 	/** Changes the viewport to top view */
 	TSharedPtr< FUICommandInfo > Top;
 
-	/** Changes the viewport to side view */
-	TSharedPtr< FUICommandInfo > Side;
+	/** Changes the viewport to bottom view */
+	TSharedPtr< FUICommandInfo > Bottom;
+
+	/** Changes the viewport to left view */
+	TSharedPtr< FUICommandInfo > Left;
+
+	/** Changes the viewport to right view */
+	TSharedPtr< FUICommandInfo > Right;
 
 	/** Changes the viewport to front view */
 	TSharedPtr< FUICommandInfo > Front;
 
+	/** Changes the viewport to back view */
+	TSharedPtr< FUICommandInfo > Back;
+
 	/** Changes the viewport to wireframe */
 	TSharedPtr< FUICommandInfo > WireframeMode;
 
