Additional files which got lost as part of CL 2480276.

diff --git a/Engine/Source/Editor/UnrealEd/Private/EditorComponents.cpp b/Engine/Source/Editor/UnrealEd/Private/EditorComponents.cpp
--- a/Engine/Source/Editor/UnrealEd/Private/EditorComponents.cpp
+++ b/Engine/Source/Editor/UnrealEd/Private/EditorComponents.cpp
@@ -389,6 +389,9 @@ void FEditorCommonDrawHelper::DrawOldGrid(const FSceneView* View,FPrimitiveDrawI
 
 	bool bIsPerspective = ( View->ViewMatrices.ProjMatrix.M[3][3] < 1.0f );
 
+	// Don't attempt to draw the grid lines from the maximum half world extent as it may be clipped due to floating point truncation errors
+	const float HalfWorldExtent = HALF_WORLD_MAX - 1000.0f;
+
 	// Draw 3D perspective grid
 	if( bIsPerspective)
 	{
@@ -442,12 +445,15 @@ void FEditorCommonDrawHelper::DrawOldGrid(const FSceneView* View,FPrimitiveDrawI
 		FLinearColor AxisColors[3];
 		GetAxisColors(AxisColors, false);
 
-		if( bIsOrthoXY )
+		if (bIsOrthoXY)
 		{
-			FVector StartY( 0.0f, +HALF_WORLD_MAX1, -HALF_WORLD_MAX );
-			FVector EndY( 0.0f, -HALF_WORLD_MAX1, -HALF_WORLD_MAX );
-			FVector StartX( +HALF_WORLD_MAX1, 0.0f, -HALF_WORLD_MAX );
-			FVector EndX( -HALF_WORLD_MAX1, 0.0f, -HALF_WORLD_MAX );
+			const bool bNegative = View->ViewMatrices.ViewMatrix.M[2][2] > 0.0f;
+
+			FVector StartY(0.0f, +HalfWorldExtent, bNegative ? HalfWorldExtent : -HalfWorldExtent);
+			FVector EndY(0.0f, -HalfWorldExtent, bNegative ? HalfWorldExtent : -HalfWorldExtent);
+			FVector StartX(+HalfWorldExtent, 0.0f, bNegative ? HalfWorldExtent : -HalfWorldExtent);
+			FVector EndX(-HalfWorldExtent, 0.0f, bNegative ? HalfWorldExtent : -HalfWorldExtent);
+
 			DrawGridSection( GEditor->GetGridSize(), &StartY, &EndY, &StartY.X, &EndY.X, 0, View, PDI);
 			DrawGridSection( GEditor->GetGridSize(), &StartX, &EndX, &StartX.Y, &EndX.Y, 1, View, PDI);
 			DrawOriginAxisLine( &StartY, &EndY, &StartY.X, &EndY.X, View, PDI, AxisColors[1] );
@@ -455,10 +461,13 @@ void FEditorCommonDrawHelper::DrawOldGrid(const FSceneView* View,FPrimitiveDrawI
 		}
 		else if( bIsOrthoXZ )
 		{
-			FVector StartZ( 0.0f, -HALF_WORLD_MAX, +HALF_WORLD_MAX1 );
-			FVector EndZ( 0.0f, -HALF_WORLD_MAX, -HALF_WORLD_MAX1 );
-			FVector StartX( +HALF_WORLD_MAX1, -HALF_WORLD_MAX, 0.0f );
-			FVector EndX( -HALF_WORLD_MAX1, -HALF_WORLD_MAX, 0.0f );
+			const bool bNegative = View->ViewMatrices.ViewMatrix.M[1][2] > 0.0f;
+
+			FVector StartZ(0.0f, bNegative ? HalfWorldExtent : -HalfWorldExtent, +HalfWorldExtent);
+			FVector EndZ(0.0f, bNegative ? HalfWorldExtent : -HalfWorldExtent, -HalfWorldExtent);
+			FVector StartX(+HalfWorldExtent, bNegative ? HalfWorldExtent : -HalfWorldExtent, 0.0f);
+			FVector EndX(-HalfWorldExtent, bNegative ? HalfWorldExtent : -HalfWorldExtent, 0.0f);
+
 			DrawGridSection( GEditor->GetGridSize(), &StartZ, &EndZ, &StartZ.X, &EndZ.X, 0, View, PDI);
 			DrawGridSection( GEditor->GetGridSize(), &StartX, &EndX, &StartX.Z, &EndX.Z, 2, View, PDI);
 			DrawOriginAxisLine( &StartZ, &EndZ, &StartZ.X, &EndZ.X, View, PDI, AxisColors[2] );
@@ -466,10 +475,13 @@ void FEditorCommonDrawHelper::DrawOldGrid(const FSceneView* View,FPrimitiveDrawI
 		}
 		else if( bIsOrthoYZ )
 		{
-			FVector StartZ( +HALF_WORLD_MAX, 0.0f, +HALF_WORLD_MAX1 );
-			FVector EndZ( +HALF_WORLD_MAX, 0.0f, -HALF_WORLD_MAX1 );
-			FVector StartY( +HALF_WORLD_MAX, +HALF_WORLD_MAX1, 0.0f );
-			FVector EndY( +HALF_WORLD_MAX, -HALF_WORLD_MAX1, 0.0f );
+			const bool bNegative = View->ViewMatrices.ViewMatrix.M[0][2] < 0.0f;
+
+			FVector StartZ(bNegative ? -HalfWorldExtent : HalfWorldExtent, 0.0f, +HalfWorldExtent);
+			FVector EndZ(bNegative ? -HalfWorldExtent : HalfWorldExtent, 0.0f, -HalfWorldExtent);
+			FVector StartY(bNegative ? -HalfWorldExtent : HalfWorldExtent, +HalfWorldExtent, 0.0f);
+			FVector EndY(bNegative ? -HalfWorldExtent : HalfWorldExtent, -HalfWorldExtent, 0.0f);
+
 			DrawGridSection( GEditor->GetGridSize(), &StartZ, &EndZ, &StartZ.Y, &EndZ.Y, 1, View, PDI);
 			DrawGridSection( GEditor->GetGridSize(), &StartY, &EndY, &StartY.Z, &EndY.Z, 2, View, PDI);
 			DrawOriginAxisLine( &StartZ, &EndZ, &StartZ.Y, &EndZ.Y, View, PDI, AxisColors[2] );
@@ -480,15 +492,15 @@ void FEditorCommonDrawHelper::DrawOldGrid(const FSceneView* View,FPrimitiveDrawI
 		{
 			float KillZ = GWorld->GetWorldSettings()->KillZ;
 
-			PDI->DrawLine(FVector(-HALF_WORLD_MAX, 0, KillZ), FVector(HALF_WORLD_MAX, 0, KillZ), FColor::Red, SDPG_Foreground);
-			PDI->DrawLine(FVector(0, -HALF_WORLD_MAX, KillZ), FVector(0, HALF_WORLD_MAX, KillZ), FColor::Red, SDPG_Foreground);
+			PDI->DrawLine(FVector(-HalfWorldExtent, 0, KillZ), FVector(HalfWorldExtent, 0, KillZ), FColor::Red, SDPG_Foreground);
+			PDI->DrawLine(FVector(0, -HalfWorldExtent, KillZ), FVector(0, HalfWorldExtent, KillZ), FColor::Red, SDPG_Foreground);
 		}
 	}
 
 	// Draw orthogonal worldframe.
 	if(bDrawWorldBox)
 	{
-		DrawWireBox(PDI, FBox( FVector(-HALF_WORLD_MAX1,-HALF_WORLD_MAX1,-HALF_WORLD_MAX1),FVector(HALF_WORLD_MAX1,HALF_WORLD_MAX1,HALF_WORLD_MAX1) ), GEngine->C_WorldBox, eDPG );
+		DrawWireBox(PDI, FBox(FVector(-HalfWorldExtent, -HalfWorldExtent, -HalfWorldExtent), FVector(HalfWorldExtent, HalfWorldExtent, HalfWorldExtent)), GEngine->C_WorldBox, eDPG);
 	}
 }
 
diff --git a/Engine/Source/Editor/UnrealEd/Private/UnrealWidget.cpp b/Engine/Source/Editor/UnrealEd/Private/UnrealWidget.cpp
--- a/Engine/Source/Editor/UnrealEd/Private/UnrealWidget.cpp
+++ b/Engine/Source/Editor/UnrealEd/Private/UnrealWidget.cpp
@@ -441,7 +441,7 @@ void FWidget::Render_Translate( const FSceneView* View, FPrimitiveDrawInterface*
 		// but in the ortho viewports it will prevent scaling in the direction of the camera and thus intersecting the near plane.
 		FVector FlattenScale = FVector(Scale.Component(0) == 1.0f ? 1.0f / UniformScale : 1.0f, Scale.Component(1) == 1.0f ? 1.0f / UniformScale : 1.0f, Scale.Component(2) == 1.0f ? 1.0f / UniformScale : 1.0f);
 
-		if( bIsPerspective || bIsLocalSpace || View->ViewMatrices.ViewMatrix.M[2][1] == 0.f )
+		if (bIsPerspective || bIsLocalSpace || bIsOrthoXY)
 		{
 			if( (DrawAxis&EAxisList::XY) == EAxisList::XY )							// Top
 			{
@@ -456,7 +456,7 @@ void FWidget::Render_Translate( const FSceneView* View, FPrimitiveDrawInterface*
 			}
 		}
 
-		if( bIsPerspective || bIsLocalSpace || View->ViewMatrices.ViewMatrix.M[1][2] == -1.f )		// Front
+		if (bIsPerspective || bIsLocalSpace || bIsOrthoXZ)		// Front
 		{
 			if( (DrawAxis&EAxisList::XZ) == EAxisList::XZ ) 
 			{
@@ -471,7 +471,7 @@ void FWidget::Render_Translate( const FSceneView* View, FPrimitiveDrawInterface*
 			}
 		}
 
-		if( bIsPerspective || bIsLocalSpace || View->ViewMatrices.ViewMatrix.M[1][0] == 1.f )		// Side
+		if( bIsPerspective || bIsLocalSpace || bIsOrthoYZ )		// Side
 		{
 			if( (DrawAxis&EAxisList::YZ) == EAxisList::YZ ) 
 			{
