Extra fix to compiler error output in output log, so that error lines are displayed in the log without prepended text so the error throwback in VS works OK.

diff --git a/Engine/Source/Editor/UnrealEd/Private/Kismet2/CompilerResultsLog.cpp b/Engine/Source/Editor/UnrealEd/Private/Kismet2/CompilerResultsLog.cpp
--- a/Engine/Source/Editor/UnrealEd/Private/Kismet2/CompilerResultsLog.cpp
+++ b/Engine/Source/Editor/UnrealEd/Private/Kismet2/CompilerResultsLog.cpp
@@ -397,6 +397,11 @@ TArray< TSharedRef<FTokenizedMessage> > FCompilerResultsLog::ParseCompilerLogDum
 			Message->AddToken( FTextToken::Create( FText::FromString( Link ) )->OnMessageTokenActivated(FOnMessageTokenActivated::CreateStatic(&FCompilerResultsLog::OnGotoError) ) );
 			Message->AddToken( FTextToken::Create( FText::FromString( RightStr ) ) );
 			Messages.Add(Message);
+
+			if (Severity == EMessageSeverity::Error)
+			{
+				FPlatformMisc::LowLevelOutputDebugStringf(TEXT("%s"), *Line);
+			}
 		}
 		else
 		{
@@ -404,6 +409,7 @@ TArray< TSharedRef<FTokenizedMessage> > FCompilerResultsLog::ParseCompilerLogDum
 			if (Line.Contains(TEXT("error LNK"), ESearchCase::CaseSensitive))
 			{
 				Severity = EMessageSeverity::Error;
+				FPlatformMisc::LowLevelOutputDebugStringf(TEXT("%s"), *Line);
 			}
 
 			TSharedRef<FTokenizedMessage> Message = FTokenizedMessage::Create( Severity );
