Fixed issue with -iterate not detecting changes to root global shader map key or root material key.

diff --git a/Engine/Source/Editor/UnrealEd/Private/CookOnTheFlyServer.cpp b/Engine/Source/Editor/UnrealEd/Private/CookOnTheFlyServer.cpp
--- a/Engine/Source/Editor/UnrealEd/Private/CookOnTheFlyServer.cpp
+++ b/Engine/Source/Editor/UnrealEd/Private/CookOnTheFlyServer.cpp
@@ -11,6 +11,7 @@
 #include "DerivedDataCacheInterface.h"
 #include "ISourceControlModule.h"
 #include "GlobalShader.h"
+#include "MaterialShader.h"
 #include "TargetPlatform.h"
 #include "IConsoleManager.h"
 #include "Developer/PackageDependencyInfo/Public/PackageDependencyInfo.h"
@@ -2101,6 +2102,11 @@ bool UCookOnTheFlyServer::GetCurrentIniVersionStrings( const ITargetPlatform* Ta
 	GetVersionFormatNumbersForIniVersionStrings( IniVersionStrings, TEXT("TextureFormat"), TPM->GetTextureFormats() );
 	GetVersionFormatNumbersForIniVersionStrings( IniVersionStrings, TEXT("ShaderFormat"), TPM->GetShaderFormats() );
 
+	FString MaterialShaderMapDDCVersion = FString::Printf(TEXT("MaterialShaderMapDDCVersion:%s"), *GetMaterialShaderMapDDCKey());
+	IniVersionStrings.Emplace( MoveTemp(MaterialShaderMapDDCVersion) );
+	FString GlobalDDCVersion = FString::Printf(TEXT("GlobalDDCVersion:%s"), *GetGlobalShaderMapDDCKey());
+	IniVersionStrings.Emplace( MoveTemp(GlobalDDCVersion) );
+
 	return true;
 }
 
diff --git a/Engine/Source/Runtime/Engine/Private/GlobalShader.cpp b/Engine/Source/Runtime/Engine/Private/GlobalShader.cpp
--- a/Engine/Source/Runtime/Engine/Private/GlobalShader.cpp
+++ b/Engine/Source/Runtime/Engine/Private/GlobalShader.cpp
@@ -307,6 +307,16 @@ FString SaveGlobalShaderFile(EShaderPlatform Platform, FString SavePath)
 	return FullPath;
 }
 
+FString GetGlobalShaderMapDDCKey()
+{
+	return FString(GLOBALSHADERMAP_DERIVEDDATA_VER);
+}
+
+FString GetMaterialShaderMapDDCKey()
+{
+	return FString(MATERIALSHADERMAP_DERIVEDDATA_VER);
+}
+
 /** Creates a string key for the derived data cache entry for the global shader map. */
 FString GetGlobalShaderMapKeyString(const FGlobalShaderMapId& ShaderMapId, EShaderPlatform Platform)
 {
diff --git a/Engine/Source/Runtime/Engine/Public/GlobalShader.h b/Engine/Source/Runtime/Engine/Public/GlobalShader.h
--- a/Engine/Source/Runtime/Engine/Public/GlobalShader.h
+++ b/Engine/Source/Runtime/Engine/Public/GlobalShader.h
@@ -219,6 +219,11 @@ extern ENGINE_API void FinishRecompileGlobalShaders();
 /** Called by the shader compiler to process completed global shader jobs. */
 extern void ProcessCompiledGlobalShaders(const TArray<FShaderCompileJob*>& CompilationResults);
 
+extern ENGINE_API FString GetGlobalShaderMapDDCKey();
+
+extern ENGINE_API FString GetMaterialShaderMapDDCKey();
+
+
 /** 
  * Recompiles global shaders
  *
