Made it so cooker warnings about duplicate object creation only happen once in the log.

diff --git a/Engine/Source/Editor/UnrealEd/Private/CookOnTheFlyServer.cpp b/Engine/Source/Editor/UnrealEd/Private/CookOnTheFlyServer.cpp
--- a/Engine/Source/Editor/UnrealEd/Private/CookOnTheFlyServer.cpp
+++ b/Engine/Source/Editor/UnrealEd/Private/CookOnTheFlyServer.cpp
@@ -1011,7 +1011,13 @@ uint32 UCookOnTheFlyServer::TickCookOnTheSide( const float TimeSlice, uint32 &Co
 #if DEBUG_COOKONTHEFLY
 			UE_LOG( LogCook, Display, TEXT("Processing request %s"), *BuildFilename);
 #endif
-			GOutputCookingWarnings = true;
+			static TSet<FString> CookWarningsList;
+			if ( CookWarningsList.Contains(BuildFilename) == false )
+			{
+				CookWarningsList.Add( BuildFilename );
+				GOutputCookingWarnings = true;
+			}
+			
 			//  if the package is already loaded then try to avoid reloading it :)
 			if ( ( Package == NULL ) || ( Package->IsFullyLoaded() == false ) )
 			{
