[UE-11011] Fix for stale property reference in bytecode, we were not tracking blueprint dependencies properly when there were circular references. bCachedDependenciesUpToDate would be true after we loaded a blueprint, but because it circularly referenced another blueprint we could load that other blueprint after loading the first one and setting bCachedDependenciesUpToDate to true.

diff --git a/Engine/Source/Editor/UnrealEd/Private/Kismet2/BlueprintEditorUtils.cpp b/Engine/Source/Editor/UnrealEd/Private/Kismet2/BlueprintEditorUtils.cpp
--- a/Engine/Source/Editor/UnrealEd/Private/Kismet2/BlueprintEditorUtils.cpp
+++ b/Engine/Source/Editor/UnrealEd/Private/Kismet2/BlueprintEditorUtils.cpp
@@ -1126,6 +1126,7 @@ UClass* FBlueprintEditorUtils::RegenerateBlueprintClass(UBlueprint* Blueprint, U
 
 	if( ShouldRegenerateBlueprint(Blueprint) && !Blueprint->bHasBeenRegenerated )
 	{
+		Blueprint->bCachedDependenciesUpToDate = false;
 		Blueprint->bIsRegeneratingOnLoad = true;
 
 		// Cache off the linker index, if needed
