Enabled UHT to digest GENERATED_BODY instead of GENERATED_UCLASS_BODY, GENERATED_USTRUCT_BODY, GENERATED_UINTERFACE_BODY or GENERATED_IINTERFACE_BODY, changed every occurence to the new syntax and fix

diff --git a/Engine/Source/Editor/UnrealEd/Classes/Editor/EditorUserSettings.h b/Engine/Source/Editor/UnrealEd/Classes/Editor/EditorUserSettings.h
--- a/Engine/Source/Editor/UnrealEd/Classes/Editor/EditorUserSettings.h
+++ b/Engine/Source/Editor/UnrealEd/Classes/Editor/EditorUserSettings.h
@@ -10,7 +10,9 @@ UCLASS(minimalapi, autoexpandcategories=(ViewportControls, ViewportLookAndFeel,
 	   config=EditorUserSettings)
 class UEditorUserSettings : public UObject
 {
-	GENERATED_UCLASS_BODY()
+	GENERATED_BODY()
+public:
+	UEditorUserSettings(const FObjectInitializer& ObjectInitializer = FObjectInitializer::Get());
 
 	// =====================================================================
 	// The following options are exposed in the Preferences Editor 
diff --git a/Engine/Source/Editor/UnrealEd/Classes/Editor/GroupActor.h b/Engine/Source/Editor/UnrealEd/Classes/Editor/GroupActor.h
--- a/Engine/Source/Editor/UnrealEd/Classes/Editor/GroupActor.h
+++ b/Engine/Source/Editor/UnrealEd/Classes/Editor/GroupActor.h
@@ -11,7 +11,9 @@
 UCLASS(hidedropdown,MinimalAPI, notplaceable, NotBlueprintable)
 class AGroupActor : public AActor
 {
-	GENERATED_UCLASS_BODY()
+	GENERATED_BODY()
+public:
+	UNREALED_API AGroupActor(const FObjectInitializer& ObjectInitializer = FObjectInitializer::Get());
 
 #if WITH_EDITORONLY_DATA
 	UPROPERTY()
diff --git a/Engine/Source/Editor/UnrealEd/Classes/Editor/PropertyEditorTestObject.h b/Engine/Source/Editor/UnrealEd/Classes/Editor/PropertyEditorTestObject.h
--- a/Engine/Source/Editor/UnrealEd/Classes/Editor/PropertyEditorTestObject.h
+++ b/Engine/Source/Editor/UnrealEd/Classes/Editor/PropertyEditorTestObject.h
@@ -38,7 +38,7 @@ enum ArrayLabelEnum
 USTRUCT()
 struct FPropertyEditorTestSubStruct
 {
-	GENERATED_USTRUCT_BODY()
+	GENERATED_BODY()
 
 	FPropertyEditorTestSubStruct()
 		: FirstProperty( 7897789 )
@@ -65,7 +65,7 @@ struct FPropertyEditorTestSubStruct
 USTRUCT()
 struct FPropertyEditorTestBasicStruct
 {
-	GENERATED_USTRUCT_BODY()
+	GENERATED_BODY()
 
 	FPropertyEditorTestBasicStruct()
 		: IntPropertyInsideAStruct( 0 )
@@ -90,7 +90,9 @@ struct FPropertyEditorTestBasicStruct
 UCLASS(transient)
 class UPropertyEditorTestObject : public UObject
 {
-    GENERATED_UCLASS_BODY()
+    GENERATED_BODY()
+public:
+    UPropertyEditorTestObject(const FObjectInitializer& ObjectInitializer = FObjectInitializer::Get());
 
 	// Integer
 	UPROPERTY(EditAnywhere, Category=BasicProperties)
@@ -261,4 +263,4 @@ class UPropertyEditorTestObject : public UObject
 
 	UPROPERTY(EditAnywhere, Category=AssetPropertyTests)
 	AActor* OnlyActorsAllowed;
-};
\ No newline at end of file
+};
diff --git a/Engine/Source/Editor/UnrealEd/Classes/Editor/TemplateMapMetadata.h b/Engine/Source/Editor/UnrealEd/Classes/Editor/TemplateMapMetadata.h
--- a/Engine/Source/Editor/UnrealEd/Classes/Editor/TemplateMapMetadata.h
+++ b/Engine/Source/Editor/UnrealEd/Classes/Editor/TemplateMapMetadata.h
@@ -7,6 +7,8 @@
 UCLASS(deprecated)
 class UDEPRECATED_TemplateMapMetadata : public UObject
 {
-	GENERATED_UCLASS_BODY()
+	GENERATED_BODY()
+public:
+	UDEPRECATED_TemplateMapMetadata(const FObjectInitializer& ObjectInitializer = FObjectInitializer::Get());
 };
 
diff --git a/Engine/Source/Editor/UnrealEd/Classes/Editor/Transactor.h b/Engine/Source/Editor/UnrealEd/Classes/Editor/Transactor.h
--- a/Engine/Source/Editor/UnrealEd/Classes/Editor/Transactor.h
+++ b/Engine/Source/Editor/UnrealEd/Classes/Editor/Transactor.h
@@ -410,7 +410,9 @@ public:
 UCLASS(transient)
 class UTransactor : public UObject
 {
-    GENERATED_UCLASS_BODY()
+    GENERATED_BODY()
+public:
+    UTransactor(const FObjectInitializer& ObjectInitializer = FObjectInitializer::Get());
 	/**
 	 * Begins a new undo transaction.  An undo transaction is defined as all actions
 	 * which take place when the user selects "undo" a single time.
diff --git a/Engine/Source/Editor/UnrealEd/Classes/Editor/UnrealEdTypes.h b/Engine/Source/Editor/UnrealEd/Classes/Editor/UnrealEdTypes.h
--- a/Engine/Source/Editor/UnrealEd/Classes/Editor/UnrealEdTypes.h
+++ b/Engine/Source/Editor/UnrealEd/Classes/Editor/UnrealEdTypes.h
@@ -14,7 +14,7 @@
 USTRUCT()
 struct FLightmassParameterValue
 {
-	GENERATED_USTRUCT_BODY()
+	GENERATED_BODY()
 
 	/** If true, override the given parameter with the given settings */
 	UPROPERTY(EditAnywhere, BlueprintReadWrite, Category=LightmassParameterValue)
@@ -32,7 +32,7 @@ struct FLightmassParameterValue
 USTRUCT()
 struct FLightmassBooleanParameterValue : public FLightmassParameterValue
 {
-	GENERATED_USTRUCT_BODY()
+	GENERATED_BODY()
 
 	/** The boolean value to override the parent value with */
 	UPROPERTY(EditAnywhere, BlueprintReadWrite, Category=LightmassBooleanParameterValue)
@@ -50,7 +50,7 @@ struct FLightmassBooleanParameterValue : public FLightmassParameterValue
 USTRUCT()
 struct FLightmassScalarParameterValue : public FLightmassParameterValue
 {
-	GENERATED_USTRUCT_BODY()
+	GENERATED_BODY()
 
 	/** The scalar value to override the parent value with */
 	UPROPERTY(EditAnywhere, BlueprintReadWrite, Category=LightmassScalarParameterValue)
@@ -71,7 +71,7 @@ struct FLightmassScalarParameterValue : public FLightmassParameterValue
 USTRUCT()
 struct FLightmassParameterizedMaterialSettings
 {
-	GENERATED_USTRUCT_BODY()
+	GENERATED_BODY()
 
 	/** If true, forces translucency to cast static shadows as if the material were masked. */
 	UPROPERTY(EditAnywhere, BlueprintReadWrite, Category=LightmassParameterizedMaterialSettings)
@@ -132,7 +132,9 @@ enum ELevelViewportType
 UCLASS(abstract, config=UnrealEd)
 class UUnrealEdTypes : public UObject
 {
-	GENERATED_UCLASS_BODY()
+	GENERATED_BODY()
+public:
+	UUnrealEdTypes(const FObjectInitializer& ObjectInitializer = FObjectInitializer::Get());
 
 };
 
diff --git a/Engine/Source/Editor/UnrealEd/Classes/Exporters/AnimSequenceExporterFBX.h b/Engine/Source/Editor/UnrealEd/Classes/Exporters/AnimSequenceExporterFBX.h
--- a/Engine/Source/Editor/UnrealEd/Classes/Exporters/AnimSequenceExporterFBX.h
+++ b/Engine/Source/Editor/UnrealEd/Classes/Exporters/AnimSequenceExporterFBX.h
@@ -11,7 +11,9 @@
 UCLASS()
 class UAnimSequenceExporterFBX : public UExporter
 {
-	GENERATED_UCLASS_BODY()
+	GENERATED_BODY()
+public:
+	UAnimSequenceExporterFBX(const FObjectInitializer& ObjectInitializer = FObjectInitializer::Get());
 
 
 	// Begin UExporter Interface
diff --git a/Engine/Source/Editor/UnrealEd/Classes/Exporters/ExportTextContainer.h b/Engine/Source/Editor/UnrealEd/Classes/Exporters/ExportTextContainer.h
--- a/Engine/Source/Editor/UnrealEd/Classes/Exporters/ExportTextContainer.h
+++ b/Engine/Source/Editor/UnrealEd/Classes/Exporters/ExportTextContainer.h
@@ -7,7 +7,9 @@
 UCLASS(MinimalAPI)
 class UExportTextContainer : public UObject
 {
-	GENERATED_UCLASS_BODY()
+	GENERATED_BODY()
+public:
+	UNREALED_API UExportTextContainer(const FObjectInitializer& ObjectInitializer = FObjectInitializer::Get());
 
 	/** ExportText representation of one or more objects */
 	UPROPERTY()
diff --git a/Engine/Source/Editor/UnrealEd/Classes/Exporters/LevelExporterFBX.h b/Engine/Source/Editor/UnrealEd/Classes/Exporters/LevelExporterFBX.h
--- a/Engine/Source/Editor/UnrealEd/Classes/Exporters/LevelExporterFBX.h
+++ b/Engine/Source/Editor/UnrealEd/Classes/Exporters/LevelExporterFBX.h
@@ -10,11 +10,10 @@
 UCLASS()
 class ULevelExporterFBX : public UExporter
 {
-public:
 	GENERATED_BODY()
-
 public:
-	ULevelExporterFBX(const FObjectInitializer& ObjectInitializer = FObjectInitializer());
+	ULevelExporterFBX(const FObjectInitializer& ObjectInitializer = FObjectInitializer::Get());
+
 
 	// Begin UExporter Interface
 	virtual bool ExportBinary( UObject* Object, const TCHAR* Type, FArchive& Ar, FFeedbackContext* Warn, int32 FileIndex = 0, uint32 PortFlags=0 ) override;
diff --git a/Engine/Source/Editor/UnrealEd/Classes/Exporters/LevelExporterLOD.h b/Engine/Source/Editor/UnrealEd/Classes/Exporters/LevelExporterLOD.h
--- a/Engine/Source/Editor/UnrealEd/Classes/Exporters/LevelExporterLOD.h
+++ b/Engine/Source/Editor/UnrealEd/Classes/Exporters/LevelExporterLOD.h
@@ -10,11 +10,10 @@
 UCLASS()
 class ULevelExporterLOD : public UExporter
 {
-public:
 	GENERATED_BODY()
-
 public:
-	ULevelExporterLOD(const FObjectInitializer& ObjectInitializer = FObjectInitializer());
+	ULevelExporterLOD(const FObjectInitializer& ObjectInitializer = FObjectInitializer::Get());
+
 
 	// Begin UExporter Interface
 	virtual bool ExportText( const FExportObjectInnerContext* Context, UObject* Object, const TCHAR* Type, FOutputDevice& Ar, FFeedbackContext* Warn, uint32 PortFlags=0 ) override;
diff --git a/Engine/Source/Editor/UnrealEd/Classes/Exporters/LevelExporterOBJ.h b/Engine/Source/Editor/UnrealEd/Classes/Exporters/LevelExporterOBJ.h
--- a/Engine/Source/Editor/UnrealEd/Classes/Exporters/LevelExporterOBJ.h
+++ b/Engine/Source/Editor/UnrealEd/Classes/Exporters/LevelExporterOBJ.h
@@ -10,11 +10,10 @@
 UCLASS()
 class ULevelExporterOBJ : public UExporter
 {
-public:
 	GENERATED_BODY()
-
 public:
-	ULevelExporterOBJ(const FObjectInitializer& ObjectInitializer = FObjectInitializer());
+	ULevelExporterOBJ(const FObjectInitializer& ObjectInitializer = FObjectInitializer::Get());
+
 
 	// Begin UExporter Interface
 	virtual bool ExportText( const FExportObjectInnerContext* Context, UObject* Object, const TCHAR* Type, FOutputDevice& Ar, FFeedbackContext* Warn, uint32 PortFlags=0 ) override;
diff --git a/Engine/Source/Editor/UnrealEd/Classes/Exporters/LevelExporterSTL.h b/Engine/Source/Editor/UnrealEd/Classes/Exporters/LevelExporterSTL.h
--- a/Engine/Source/Editor/UnrealEd/Classes/Exporters/LevelExporterSTL.h
+++ b/Engine/Source/Editor/UnrealEd/Classes/Exporters/LevelExporterSTL.h
@@ -10,11 +10,10 @@
 UCLASS()
 class ULevelExporterSTL : public UExporter
 {
-public:
 	GENERATED_BODY()
-
 public:
-	ULevelExporterSTL(const FObjectInitializer& ObjectInitializer = FObjectInitializer());
+	ULevelExporterSTL(const FObjectInitializer& ObjectInitializer = FObjectInitializer::Get());
+
 
 	// Begin UExporter Interface
 	virtual bool ExportText( const FExportObjectInnerContext* Context, UObject* Object, const TCHAR* Type, FOutputDevice& Ar, FFeedbackContext* Warn, uint32 PortFlags=0 ) override;
diff --git a/Engine/Source/Editor/UnrealEd/Classes/Exporters/LevelExporterT3D.h b/Engine/Source/Editor/UnrealEd/Classes/Exporters/LevelExporterT3D.h
--- a/Engine/Source/Editor/UnrealEd/Classes/Exporters/LevelExporterT3D.h
+++ b/Engine/Source/Editor/UnrealEd/Classes/Exporters/LevelExporterT3D.h
@@ -10,11 +10,10 @@
 UCLASS()
 class ULevelExporterT3D : public UExporter
 {
-public:
 	GENERATED_BODY()
-
 public:
-	ULevelExporterT3D(const FObjectInitializer& ObjectInitializer = FObjectInitializer());
+	ULevelExporterT3D(const FObjectInitializer& ObjectInitializer = FObjectInitializer::Get());
+
 
 	// Begin UExporter Interface
 	virtual bool ExportText( const FExportObjectInnerContext* Context, UObject* Object, const TCHAR* Type, FOutputDevice& Ar, FFeedbackContext* Warn, uint32 PortFlags=0 ) override;
diff --git a/Engine/Source/Editor/UnrealEd/Classes/Exporters/ModelExporterT3D.h b/Engine/Source/Editor/UnrealEd/Classes/Exporters/ModelExporterT3D.h
--- a/Engine/Source/Editor/UnrealEd/Classes/Exporters/ModelExporterT3D.h
+++ b/Engine/Source/Editor/UnrealEd/Classes/Exporters/ModelExporterT3D.h
@@ -10,11 +10,10 @@
 UCLASS()
 class UModelExporterT3D : public UExporter
 {
-public:
 	GENERATED_BODY()
-
 public:
-	UModelExporterT3D(const FObjectInitializer& ObjectInitializer = FObjectInitializer());
+	UModelExporterT3D(const FObjectInitializer& ObjectInitializer = FObjectInitializer::Get());
+
 
 	// Begin UExporter Interface
 	virtual bool ExportText( const FExportObjectInnerContext* Context, UObject* Object, const TCHAR* Type, FOutputDevice& Ar, FFeedbackContext* Warn, uint32 PortFlags=0 ) override;
diff --git a/Engine/Source/Editor/UnrealEd/Classes/Exporters/ObjectExporterT3D.h b/Engine/Source/Editor/UnrealEd/Classes/Exporters/ObjectExporterT3D.h
--- a/Engine/Source/Editor/UnrealEd/Classes/Exporters/ObjectExporterT3D.h
+++ b/Engine/Source/Editor/UnrealEd/Classes/Exporters/ObjectExporterT3D.h
@@ -10,11 +10,9 @@
 UCLASS()
 class UObjectExporterT3D : public UExporter
 {
-public:
 	GENERATED_BODY()
-
 public:
-	UObjectExporterT3D(const FObjectInitializer& ObjectInitializer = FObjectInitializer());
+	UObjectExporterT3D(const FObjectInitializer& ObjectInitializer = FObjectInitializer::Get());
 
 
 	// Begin UExporter Interface
diff --git a/Engine/Source/Editor/UnrealEd/Classes/Exporters/PolysExporterOBJ.h b/Engine/Source/Editor/UnrealEd/Classes/Exporters/PolysExporterOBJ.h
--- a/Engine/Source/Editor/UnrealEd/Classes/Exporters/PolysExporterOBJ.h
+++ b/Engine/Source/Editor/UnrealEd/Classes/Exporters/PolysExporterOBJ.h
@@ -10,11 +10,10 @@
 UCLASS()
 class UPolysExporterOBJ : public UExporter
 {
-public:
 	GENERATED_BODY()
-
 public:
-	UPolysExporterOBJ(const FObjectInitializer& ObjectInitializer = FObjectInitializer());
+	UPolysExporterOBJ(const FObjectInitializer& ObjectInitializer = FObjectInitializer::Get());
+
 
 	// Begin UExporter Interface
 	virtual bool ExportText( const FExportObjectInnerContext* Context, UObject* Object, const TCHAR* Type, FOutputDevice& Ar, FFeedbackContext* Warn, uint32 PortFlags=0 ) override;
diff --git a/Engine/Source/Editor/UnrealEd/Classes/Exporters/PolysExporterT3D.h b/Engine/Source/Editor/UnrealEd/Classes/Exporters/PolysExporterT3D.h
--- a/Engine/Source/Editor/UnrealEd/Classes/Exporters/PolysExporterT3D.h
+++ b/Engine/Source/Editor/UnrealEd/Classes/Exporters/PolysExporterT3D.h
@@ -10,11 +10,10 @@
 UCLASS()
 class UPolysExporterT3D : public UExporter
 {
-public:
 	GENERATED_BODY()
-
 public:
-	UPolysExporterT3D(const FObjectInitializer& ObjectInitializer = FObjectInitializer());
+	UPolysExporterT3D(const FObjectInitializer& ObjectInitializer = FObjectInitializer::Get());
+
 
 	// Begin UExporter Interface
 	virtual bool ExportText( const FExportObjectInnerContext* Context, UObject* Object, const TCHAR* Type, FOutputDevice& Ar, FFeedbackContext* Warn, uint32 PortFlags=0 ) override;
diff --git a/Engine/Source/Editor/UnrealEd/Classes/Exporters/SequenceExporterT3D.h b/Engine/Source/Editor/UnrealEd/Classes/Exporters/SequenceExporterT3D.h
--- a/Engine/Source/Editor/UnrealEd/Classes/Exporters/SequenceExporterT3D.h
+++ b/Engine/Source/Editor/UnrealEd/Classes/Exporters/SequenceExporterT3D.h
@@ -10,11 +10,10 @@
 UCLASS()
 class USequenceExporterT3D : public UExporter
 {
-public:
 	GENERATED_BODY()
-
 public:
-	USequenceExporterT3D(const FObjectInitializer& ObjectInitializer = FObjectInitializer());
+	USequenceExporterT3D(const FObjectInitializer& ObjectInitializer = FObjectInitializer::Get());
+
 
 	// Begin UExporter Interface
 	virtual bool ExportText( const FExportObjectInnerContext* Context, UObject* Object, const TCHAR* Type, FOutputDevice& Ar, FFeedbackContext* Warn, uint32 PortFlags=0 ) override;
diff --git a/Engine/Source/Editor/UnrealEd/Classes/Exporters/SkeletalMeshExporterFBX.h b/Engine/Source/Editor/UnrealEd/Classes/Exporters/SkeletalMeshExporterFBX.h
--- a/Engine/Source/Editor/UnrealEd/Classes/Exporters/SkeletalMeshExporterFBX.h
+++ b/Engine/Source/Editor/UnrealEd/Classes/Exporters/SkeletalMeshExporterFBX.h
@@ -10,7 +10,9 @@
 UCLASS()
 class USkeletalMeshExporterFBX : public UExporter
 {
-	GENERATED_UCLASS_BODY()
+	GENERATED_BODY()
+public:
+	USkeletalMeshExporterFBX(const FObjectInitializer& ObjectInitializer = FObjectInitializer::Get());
 
 
 	// Begin UExporter Interface
diff --git a/Engine/Source/Editor/UnrealEd/Classes/Exporters/SoundExporterOGG.h b/Engine/Source/Editor/UnrealEd/Classes/Exporters/SoundExporterOGG.h
--- a/Engine/Source/Editor/UnrealEd/Classes/Exporters/SoundExporterOGG.h
+++ b/Engine/Source/Editor/UnrealEd/Classes/Exporters/SoundExporterOGG.h
@@ -6,11 +6,9 @@
 UCLASS()
 class USoundExporterOGG : public UExporter
 {
+    GENERATED_BODY()
 public:
-	GENERATED_BODY()
-
-public:
-	USoundExporterOGG(const FObjectInitializer& ObjectInitializer = FObjectInitializer());
+    USoundExporterOGG(const FObjectInitializer& ObjectInitializer = FObjectInitializer::Get());
 
 	// Begin UExporter Interface
 	virtual bool ExportBinary( UObject* Object, const TCHAR* Type, FArchive& Ar, FFeedbackContext* Warn, int32 FileIndex = 0, uint32 PortFlags=0 ) override;
diff --git a/Engine/Source/Editor/UnrealEd/Classes/Exporters/SoundExporterWAV.h b/Engine/Source/Editor/UnrealEd/Classes/Exporters/SoundExporterWAV.h
--- a/Engine/Source/Editor/UnrealEd/Classes/Exporters/SoundExporterWAV.h
+++ b/Engine/Source/Editor/UnrealEd/Classes/Exporters/SoundExporterWAV.h
@@ -6,11 +6,9 @@
 UCLASS()
 class USoundExporterWAV : public UExporter
 {
+    GENERATED_BODY()
 public:
-	GENERATED_BODY()
-
-public:
-	USoundExporterWAV(const FObjectInitializer& ObjectInitializer = FObjectInitializer());
+    USoundExporterWAV(const FObjectInitializer& ObjectInitializer = FObjectInitializer::Get());
 
 	// Begin UExporter Interface
 	virtual bool ExportBinary( UObject* Object, const TCHAR* Type, FArchive& Ar, FFeedbackContext* Warn, int32 FileIndex = 0, uint32 PortFlags=0 ) override;
diff --git a/Engine/Source/Editor/UnrealEd/Classes/Exporters/SoundSurroundExporterWAV.h b/Engine/Source/Editor/UnrealEd/Classes/Exporters/SoundSurroundExporterWAV.h
--- a/Engine/Source/Editor/UnrealEd/Classes/Exporters/SoundSurroundExporterWAV.h
+++ b/Engine/Source/Editor/UnrealEd/Classes/Exporters/SoundSurroundExporterWAV.h
@@ -10,11 +10,10 @@
 UCLASS()
 class USoundSurroundExporterWAV : public UExporter
 {
-public:
 	GENERATED_BODY()
-
 public:
-	USoundSurroundExporterWAV(const FObjectInitializer& ObjectInitializer = FObjectInitializer());
+	USoundSurroundExporterWAV(const FObjectInitializer& ObjectInitializer = FObjectInitializer::Get());
+
 
 	// Begin UExporter Interface
 	virtual int32 GetFileCount( void ) const override;
diff --git a/Engine/Source/Editor/UnrealEd/Classes/Exporters/StaticMeshExporterFBX.h b/Engine/Source/Editor/UnrealEd/Classes/Exporters/StaticMeshExporterFBX.h
--- a/Engine/Source/Editor/UnrealEd/Classes/Exporters/StaticMeshExporterFBX.h
+++ b/Engine/Source/Editor/UnrealEd/Classes/Exporters/StaticMeshExporterFBX.h
@@ -10,7 +10,9 @@
 UCLASS()
 class UStaticMeshExporterFBX : public UExporter
 {
-	GENERATED_UCLASS_BODY()
+	GENERATED_BODY()
+public:
+	UStaticMeshExporterFBX(const FObjectInitializer& ObjectInitializer = FObjectInitializer::Get());
 
 
 	// Begin UExporter Interface
diff --git a/Engine/Source/Editor/UnrealEd/Classes/Exporters/StaticMeshExporterOBJ.h b/Engine/Source/Editor/UnrealEd/Classes/Exporters/StaticMeshExporterOBJ.h
--- a/Engine/Source/Editor/UnrealEd/Classes/Exporters/StaticMeshExporterOBJ.h
+++ b/Engine/Source/Editor/UnrealEd/Classes/Exporters/StaticMeshExporterOBJ.h
@@ -10,7 +10,9 @@
 UCLASS()
 class UStaticMeshExporterOBJ : public UExporter
 {
-	GENERATED_UCLASS_BODY()
+	GENERATED_BODY()
+public:
+	UStaticMeshExporterOBJ(const FObjectInitializer& ObjectInitializer = FObjectInitializer::Get());
 
 
 	// Begin UExporter Interface
diff --git a/Engine/Source/Editor/UnrealEd/Classes/Exporters/TextBufferExporterTXT.h b/Engine/Source/Editor/UnrealEd/Classes/Exporters/TextBufferExporterTXT.h
--- a/Engine/Source/Editor/UnrealEd/Classes/Exporters/TextBufferExporterTXT.h
+++ b/Engine/Source/Editor/UnrealEd/Classes/Exporters/TextBufferExporterTXT.h
@@ -10,11 +10,10 @@
 UCLASS()
 class UTextBufferExporterTXT : public UExporter
 {
-public:
 	GENERATED_BODY()
-
 public:
-	UTextBufferExporterTXT(const FObjectInitializer& ObjectInitializer = FObjectInitializer());
+	UTextBufferExporterTXT(const FObjectInitializer& ObjectInitializer = FObjectInitializer::Get());
+
 
 	// Begin UExporter Interface
 	virtual bool ExportText( const FExportObjectInnerContext* Context, UObject* Object, const TCHAR* Type, FOutputDevice& Ar, FFeedbackContext* Warn, uint32 PortFlags=0 ) override;
diff --git a/Engine/Source/Editor/UnrealEd/Classes/Exporters/TextureCubeExporterHDR.h b/Engine/Source/Editor/UnrealEd/Classes/Exporters/TextureCubeExporterHDR.h
--- a/Engine/Source/Editor/UnrealEd/Classes/Exporters/TextureCubeExporterHDR.h
+++ b/Engine/Source/Editor/UnrealEd/Classes/Exporters/TextureCubeExporterHDR.h
@@ -10,7 +10,9 @@
 UCLASS()
 class UNREALED_API UTextureCubeExporterHDR : public UExporter
 {
-	GENERATED_UCLASS_BODY()
+	GENERATED_BODY()
+public:
+	UTextureCubeExporterHDR(const FObjectInitializer& ObjectInitializer = FObjectInitializer::Get());
 
 	// Begin UExporter Interface
 	virtual bool ExportBinary(UObject* Object, const TCHAR* Type, FArchive& Ar, FFeedbackContext* Warn, int32 FileIndex = 0, uint32 PortFlags = 0) override;
diff --git a/Engine/Source/Editor/UnrealEd/Classes/Exporters/TextureExporterBMP.h b/Engine/Source/Editor/UnrealEd/Classes/Exporters/TextureExporterBMP.h
--- a/Engine/Source/Editor/UnrealEd/Classes/Exporters/TextureExporterBMP.h
+++ b/Engine/Source/Editor/UnrealEd/Classes/Exporters/TextureExporterBMP.h
@@ -10,7 +10,9 @@
 UCLASS()
 class UTextureExporterBMP : public UExporter
 {
-	GENERATED_UCLASS_BODY()
+	GENERATED_BODY()
+public:
+	UTextureExporterBMP(const FObjectInitializer& ObjectInitializer = FObjectInitializer::Get());
 
 
 	bool ExportBinary( UObject* Object, const TCHAR* Type, FArchive& Ar, FFeedbackContext* Warn, int32 FileIndex = 0, uint32 PortFlags=0 );
diff --git a/Engine/Source/Editor/UnrealEd/Classes/Exporters/TextureExporterHDR.h b/Engine/Source/Editor/UnrealEd/Classes/Exporters/TextureExporterHDR.h
--- a/Engine/Source/Editor/UnrealEd/Classes/Exporters/TextureExporterHDR.h
+++ b/Engine/Source/Editor/UnrealEd/Classes/Exporters/TextureExporterHDR.h
@@ -10,7 +10,9 @@
 UCLASS()
 class UNREALED_API UTextureExporterHDR : public UExporter
 {
-	GENERATED_UCLASS_BODY()
+	GENERATED_BODY()
+public:
+	UTextureExporterHDR(const FObjectInitializer& ObjectInitializer = FObjectInitializer::Get());
 
 	// Begin UExporter Interface
 	virtual bool ExportBinary(UObject* Object, const TCHAR* Type, FArchive& Ar, FFeedbackContext* Warn, int32 FileIndex = 0, uint32 PortFlags = 0) override;
diff --git a/Engine/Source/Editor/UnrealEd/Classes/Exporters/TextureExporterPCX.h b/Engine/Source/Editor/UnrealEd/Classes/Exporters/TextureExporterPCX.h
--- a/Engine/Source/Editor/UnrealEd/Classes/Exporters/TextureExporterPCX.h
+++ b/Engine/Source/Editor/UnrealEd/Classes/Exporters/TextureExporterPCX.h
@@ -10,7 +10,9 @@
 UCLASS()
 class UTextureExporterPCX : public UExporter
 {
-	GENERATED_UCLASS_BODY()
+	GENERATED_BODY()
+public:
+	UTextureExporterPCX(const FObjectInitializer& ObjectInitializer = FObjectInitializer::Get());
 
 
 	// Begin UExporter Interface
diff --git a/Engine/Source/Editor/UnrealEd/Classes/Exporters/TextureExporterTGA.h b/Engine/Source/Editor/UnrealEd/Classes/Exporters/TextureExporterTGA.h
--- a/Engine/Source/Editor/UnrealEd/Classes/Exporters/TextureExporterTGA.h
+++ b/Engine/Source/Editor/UnrealEd/Classes/Exporters/TextureExporterTGA.h
@@ -10,7 +10,9 @@
 UCLASS()
 class UNREALED_API UTextureExporterTGA : public UExporter
 {
-	GENERATED_UCLASS_BODY()
+	GENERATED_BODY()
+public:
+	UTextureExporterTGA(const FObjectInitializer& ObjectInitializer = FObjectInitializer::Get());
 
 
 	// Begin UExporter Interface
diff --git a/Engine/Source/Editor/UnrealEd/Classes/Factories/AimOffsetBlendSpaceFactory1D.h b/Engine/Source/Editor/UnrealEd/Classes/Factories/AimOffsetBlendSpaceFactory1D.h
--- a/Engine/Source/Editor/UnrealEd/Classes/Factories/AimOffsetBlendSpaceFactory1D.h
+++ b/Engine/Source/Editor/UnrealEd/Classes/Factories/AimOffsetBlendSpaceFactory1D.h
@@ -10,7 +10,9 @@
 UCLASS(hidecategories=Object, MinimalAPI)
 class UAimOffsetBlendSpaceFactory1D : public UBlendSpaceFactory1D
 {
-	GENERATED_UCLASS_BODY()
+	GENERATED_BODY()
+public:
+	UNREALED_API UAimOffsetBlendSpaceFactory1D(const FObjectInitializer& ObjectInitializer = FObjectInitializer::Get());
 
 	// Begin UFactory Interface
 	virtual UObject* FactoryCreateNew(UClass* Class,UObject* InParent,FName Name,EObjectFlags Flags,UObject* Context,FFeedbackContext* Warn) override;
diff --git a/Engine/Source/Editor/UnrealEd/Classes/Factories/AimOffsetBlendSpaceFactoryNew.h b/Engine/Source/Editor/UnrealEd/Classes/Factories/AimOffsetBlendSpaceFactoryNew.h
--- a/Engine/Source/Editor/UnrealEd/Classes/Factories/AimOffsetBlendSpaceFactoryNew.h
+++ b/Engine/Source/Editor/UnrealEd/Classes/Factories/AimOffsetBlendSpaceFactoryNew.h
@@ -10,7 +10,9 @@
 UCLASS(hidecategories=Object, MinimalAPI)
 class UAimOffsetBlendSpaceFactoryNew : public UBlendSpaceFactoryNew
 {
-	GENERATED_UCLASS_BODY()
+	GENERATED_BODY()
+public:
+	UNREALED_API UAimOffsetBlendSpaceFactoryNew(const FObjectInitializer& ObjectInitializer = FObjectInitializer::Get());
 
 	// Begin UFactory Interface
 	virtual UObject* FactoryCreateNew(UClass* Class,UObject* InParent,FName Name,EObjectFlags Flags,UObject* Context,FFeedbackContext* Warn) override;
diff --git a/Engine/Source/Editor/UnrealEd/Classes/Factories/AnimBlueprintFactory.h b/Engine/Source/Editor/UnrealEd/Classes/Factories/AnimBlueprintFactory.h
--- a/Engine/Source/Editor/UnrealEd/Classes/Factories/AnimBlueprintFactory.h
+++ b/Engine/Source/Editor/UnrealEd/Classes/Factories/AnimBlueprintFactory.h
@@ -11,7 +11,9 @@
 UCLASS(HideCategories=Object,MinimalAPI)
 class UAnimBlueprintFactory : public UFactory
 {
-	GENERATED_UCLASS_BODY()
+	GENERATED_BODY()
+public:
+	UNREALED_API UAnimBlueprintFactory(const FObjectInitializer& ObjectInitializer = FObjectInitializer::Get());
 
 	// The type of blueprint that will be created
 	UPROPERTY(EditAnywhere, Category=AnimBlueprintFactory)
diff --git a/Engine/Source/Editor/UnrealEd/Classes/Factories/AnimCompositeFactory.h b/Engine/Source/Editor/UnrealEd/Classes/Factories/AnimCompositeFactory.h
--- a/Engine/Source/Editor/UnrealEd/Classes/Factories/AnimCompositeFactory.h
+++ b/Engine/Source/Editor/UnrealEd/Classes/Factories/AnimCompositeFactory.h
@@ -10,7 +10,9 @@
 UCLASS(HideCategories=Object,MinimalAPI)
 class UAnimCompositeFactory : public UFactory
 {
-	GENERATED_UCLASS_BODY()
+	GENERATED_BODY()
+public:
+	UNREALED_API UAnimCompositeFactory(const FObjectInitializer& ObjectInitializer = FObjectInitializer::Get());
 
 	UPROPERTY()
 	class USkeleton* TargetSkeleton;
diff --git a/Engine/Source/Editor/UnrealEd/Classes/Factories/AnimMontageFactory.h b/Engine/Source/Editor/UnrealEd/Classes/Factories/AnimMontageFactory.h
--- a/Engine/Source/Editor/UnrealEd/Classes/Factories/AnimMontageFactory.h
+++ b/Engine/Source/Editor/UnrealEd/Classes/Factories/AnimMontageFactory.h
@@ -10,7 +10,9 @@
 UCLASS(HideCategories=Object,MinimalAPI)
 class UAnimMontageFactory : public UFactory
 {
-	GENERATED_UCLASS_BODY()
+	GENERATED_BODY()
+public:
+	UNREALED_API UAnimMontageFactory(const FObjectInitializer& ObjectInitializer = FObjectInitializer::Get());
 
 	UPROPERTY()
 	class USkeleton* TargetSkeleton;
diff --git a/Engine/Source/Editor/UnrealEd/Classes/Factories/AnimSequenceFactory.h b/Engine/Source/Editor/UnrealEd/Classes/Factories/AnimSequenceFactory.h
--- a/Engine/Source/Editor/UnrealEd/Classes/Factories/AnimSequenceFactory.h
+++ b/Engine/Source/Editor/UnrealEd/Classes/Factories/AnimSequenceFactory.h
@@ -10,7 +10,9 @@
 UCLASS(HideCategories=Object,MinimalAPI)
 class UAnimSequenceFactory : public UFactory
 {
-	GENERATED_UCLASS_BODY()
+	GENERATED_BODY()
+public:
+	UNREALED_API UAnimSequenceFactory(const FObjectInitializer& ObjectInitializer = FObjectInitializer::Get());
 
 	UPROPERTY()
 	class USkeleton* TargetSkeleton;
diff --git a/Engine/Source/Editor/UnrealEd/Classes/Factories/BlendSpaceFactory1D.h b/Engine/Source/Editor/UnrealEd/Classes/Factories/BlendSpaceFactory1D.h
--- a/Engine/Source/Editor/UnrealEd/Classes/Factories/BlendSpaceFactory1D.h
+++ b/Engine/Source/Editor/UnrealEd/Classes/Factories/BlendSpaceFactory1D.h
@@ -10,7 +10,9 @@
 UCLASS(hidecategories=Object, MinimalAPI)
 class UBlendSpaceFactory1D : public UFactory
 {
-	GENERATED_UCLASS_BODY()
+	GENERATED_BODY()
+public:
+	UNREALED_API UBlendSpaceFactory1D(const FObjectInitializer& ObjectInitializer = FObjectInitializer::Get());
 
 	UPROPERTY()
 	class USkeleton*	TargetSkeleton;
diff --git a/Engine/Source/Editor/UnrealEd/Classes/Factories/BlendSpaceFactoryNew.h b/Engine/Source/Editor/UnrealEd/Classes/Factories/BlendSpaceFactoryNew.h
--- a/Engine/Source/Editor/UnrealEd/Classes/Factories/BlendSpaceFactoryNew.h
+++ b/Engine/Source/Editor/UnrealEd/Classes/Factories/BlendSpaceFactoryNew.h
@@ -10,7 +10,9 @@
 UCLASS(hidecategories=Object, MinimalAPI)
 class UBlendSpaceFactoryNew : public UFactory
 {
-	GENERATED_UCLASS_BODY()
+	GENERATED_BODY()
+public:
+	UNREALED_API UBlendSpaceFactoryNew(const FObjectInitializer& ObjectInitializer = FObjectInitializer::Get());
 
 	UPROPERTY()
 	class USkeleton*	TargetSkeleton;
diff --git a/Engine/Source/Editor/UnrealEd/Classes/Factories/BlueprintFactory.h b/Engine/Source/Editor/UnrealEd/Classes/Factories/BlueprintFactory.h
--- a/Engine/Source/Editor/UnrealEd/Classes/Factories/BlueprintFactory.h
+++ b/Engine/Source/Editor/UnrealEd/Classes/Factories/BlueprintFactory.h
@@ -10,7 +10,9 @@
 UCLASS(hidecategories=Object, collapsecategories)
 class UNREALED_API UBlueprintFactory : public UFactory
 {
-	GENERATED_UCLASS_BODY()
+	GENERATED_BODY()
+public:
+	UBlueprintFactory(const FObjectInitializer& ObjectInitializer = FObjectInitializer::Get());
 
 	// The parent class of the created blueprint
 	UPROPERTY(EditAnywhere, Category=BlueprintFactory, meta=(AllowAbstract = "", BlueprintBaseOnly = ""))
diff --git a/Engine/Source/Editor/UnrealEd/Classes/Factories/BlueprintFunctionLibraryFactory.h b/Engine/Source/Editor/UnrealEd/Classes/Factories/BlueprintFunctionLibraryFactory.h
--- a/Engine/Source/Editor/UnrealEd/Classes/Factories/BlueprintFunctionLibraryFactory.h
+++ b/Engine/Source/Editor/UnrealEd/Classes/Factories/BlueprintFunctionLibraryFactory.h
@@ -10,7 +10,9 @@
 UCLASS(hidecategories=Object, collapsecategories)
 class UBlueprintFunctionLibraryFactory : public UBlueprintFactory
 {
-	GENERATED_UCLASS_BODY()
+	GENERATED_BODY()
+public:
+	UBlueprintFunctionLibraryFactory(const FObjectInitializer& ObjectInitializer = FObjectInitializer::Get());
 
 	// UFactory interface
 	virtual FText GetDisplayName() const override;
diff --git a/Engine/Source/Editor/UnrealEd/Classes/Factories/BlueprintInterfaceFactory.h b/Engine/Source/Editor/UnrealEd/Classes/Factories/BlueprintInterfaceFactory.h
--- a/Engine/Source/Editor/UnrealEd/Classes/Factories/BlueprintInterfaceFactory.h
+++ b/Engine/Source/Editor/UnrealEd/Classes/Factories/BlueprintInterfaceFactory.h
@@ -10,7 +10,9 @@
 UCLASS(hidecategories=Object, collapsecategories)
 class UBlueprintInterfaceFactory : public UFactory
 {
-	GENERATED_UCLASS_BODY()
+	GENERATED_BODY()
+public:
+	UBlueprintInterfaceFactory(const FObjectInitializer& ObjectInitializer = FObjectInitializer::Get());
 
 	// UFactory interface
 	virtual FText GetDisplayName() const override;
diff --git a/Engine/Source/Editor/UnrealEd/Classes/Factories/BlueprintMacroFactory.h b/Engine/Source/Editor/UnrealEd/Classes/Factories/BlueprintMacroFactory.h
--- a/Engine/Source/Editor/UnrealEd/Classes/Factories/BlueprintMacroFactory.h
+++ b/Engine/Source/Editor/UnrealEd/Classes/Factories/BlueprintMacroFactory.h
@@ -10,7 +10,9 @@
 UCLASS(hidecategories=Object, collapsecategories)
 class UBlueprintMacroFactory : public UBlueprintFactory
 {
-	GENERATED_UCLASS_BODY()
+	GENERATED_BODY()
+public:
+	UBlueprintMacroFactory(const FObjectInitializer& ObjectInitializer = FObjectInitializer::Get());
 
 	// UFactory interface
 	virtual FText GetDisplayName() const override;
diff --git a/Engine/Source/Editor/UnrealEd/Classes/Factories/CSVImportFactory.h b/Engine/Source/Editor/UnrealEd/Classes/Factories/CSVImportFactory.h
--- a/Engine/Source/Editor/UnrealEd/Classes/Factories/CSVImportFactory.h
+++ b/Engine/Source/Editor/UnrealEd/Classes/Factories/CSVImportFactory.h
@@ -12,7 +12,9 @@ DECLARE_LOG_CATEGORY_EXTERN(LogCSVImportFactory, Log, All);
 UCLASS(hidecategories=Object)
 class UNREALED_API UCSVImportFactory : public UFactory
 {
-	GENERATED_UCLASS_BODY()
+	GENERATED_BODY()
+public:
+	UCSVImportFactory(const FObjectInitializer& ObjectInitializer = FObjectInitializer::Get());
 
 	// Begin UFactory Interface
 	virtual FText GetDisplayName() const override;
diff --git a/Engine/Source/Editor/UnrealEd/Classes/Factories/CameraAnimFactory.h b/Engine/Source/Editor/UnrealEd/Classes/Factories/CameraAnimFactory.h
--- a/Engine/Source/Editor/UnrealEd/Classes/Factories/CameraAnimFactory.h
+++ b/Engine/Source/Editor/UnrealEd/Classes/Factories/CameraAnimFactory.h
@@ -7,7 +7,9 @@
 UCLASS(hidecategories=Object, MinimalAPI)
 class UCameraAnimFactory : public UFactory
 {
-	GENERATED_UCLASS_BODY()
+	GENERATED_BODY()
+public:
+	UNREALED_API UCameraAnimFactory(const FObjectInitializer& ObjectInitializer = FObjectInitializer::Get());
 
 	// Begin UFactory Interface
 	virtual FText GetDisplayName() const override;
diff --git a/Engine/Source/Editor/UnrealEd/Classes/Factories/CurveFactory.h b/Engine/Source/Editor/UnrealEd/Classes/Factories/CurveFactory.h
--- a/Engine/Source/Editor/UnrealEd/Classes/Factories/CurveFactory.h
+++ b/Engine/Source/Editor/UnrealEd/Classes/Factories/CurveFactory.h
@@ -21,7 +21,9 @@ enum ECurveType
 UCLASS(hidecategories=Object)
 class UCurveFactory : public UFactory
 {
-	GENERATED_UCLASS_BODY()
+	GENERATED_BODY()
+public:
+	UCurveFactory(const FObjectInitializer& ObjectInitializer = FObjectInitializer::Get());
 
 	/** The type of blueprint that will be created */
 	UPROPERTY(EditAnywhere, Category=CurveFactory)
diff --git a/Engine/Source/Editor/UnrealEd/Classes/Factories/CurveImportFactory.h b/Engine/Source/Editor/UnrealEd/Classes/Factories/CurveImportFactory.h
--- a/Engine/Source/Editor/UnrealEd/Classes/Factories/CurveImportFactory.h
+++ b/Engine/Source/Editor/UnrealEd/Classes/Factories/CurveImportFactory.h
@@ -10,7 +10,9 @@
 UCLASS(hidecategories=Object)
 class UCurveImportFactory : public UFactory
 {
-	GENERATED_UCLASS_BODY()
+	GENERATED_BODY()
+public:
+	UCurveImportFactory(const FObjectInitializer& ObjectInitializer = FObjectInitializer::Get());
 
 	// Begin UFactory Interface
 	virtual UObject* FactoryCreateText( UClass* InClass, UObject* InParent, FName InName, EObjectFlags Flags, UObject* Context, const TCHAR* Type, const TCHAR*& Buffer, const TCHAR* BufferEnd, FFeedbackContext* Warn ) override;
diff --git a/Engine/Source/Editor/UnrealEd/Classes/Factories/DataAssetFactory.h b/Engine/Source/Editor/UnrealEd/Classes/Factories/DataAssetFactory.h
--- a/Engine/Source/Editor/UnrealEd/Classes/Factories/DataAssetFactory.h
+++ b/Engine/Source/Editor/UnrealEd/Classes/Factories/DataAssetFactory.h
@@ -9,7 +9,9 @@ class UDataAsset;
 UCLASS(hidecategories=Object, MinimalAPI)
 class UDataAssetFactory : public UFactory
 {
-	GENERATED_UCLASS_BODY()
+	GENERATED_BODY()
+public:
+	UNREALED_API UDataAssetFactory(const FObjectInitializer& ObjectInitializer = FObjectInitializer::Get());
 
 	UPROPERTY(EditAnywhere, Category=DataAsset)
 	TSubclassOf<UDataAsset> DataAssetClass;
diff --git a/Engine/Source/Editor/UnrealEd/Classes/Factories/DataTableFactory.h b/Engine/Source/Editor/UnrealEd/Classes/Factories/DataTableFactory.h
--- a/Engine/Source/Editor/UnrealEd/Classes/Factories/DataTableFactory.h
+++ b/Engine/Source/Editor/UnrealEd/Classes/Factories/DataTableFactory.h
@@ -7,7 +7,9 @@
 UCLASS(hidecategories=Object, MinimalAPI)
 class UDataTableFactory : public UFactory
 {
-	GENERATED_UCLASS_BODY()
+	GENERATED_BODY()
+public:
+	UNREALED_API UDataTableFactory(const FObjectInitializer& ObjectInitializer = FObjectInitializer::Get());
 
 	UPROPERTY()
 	class UScriptStruct* Struct;
diff --git a/Engine/Source/Editor/UnrealEd/Classes/Factories/DestructibleMeshFactory.h b/Engine/Source/Editor/UnrealEd/Classes/Factories/DestructibleMeshFactory.h
--- a/Engine/Source/Editor/UnrealEd/Classes/Factories/DestructibleMeshFactory.h
+++ b/Engine/Source/Editor/UnrealEd/Classes/Factories/DestructibleMeshFactory.h
@@ -8,7 +8,9 @@
 UCLASS(hideCategories=Object, MinimalAPI)
 class UDestructibleMeshFactory : public UFactory
 {
-    GENERATED_UCLASS_BODY()
+    GENERATED_BODY()
+public:
+    UNREALED_API UDestructibleMeshFactory(const FObjectInitializer& ObjectInitializer = FObjectInitializer::Get());
 	// Begin UFactory Interface
 	virtual FText GetDisplayName() const override;
 #if WITH_APEX
diff --git a/Engine/Source/Editor/UnrealEd/Classes/Factories/DialogueVoiceFactory.h b/Engine/Source/Editor/UnrealEd/Classes/Factories/DialogueVoiceFactory.h
--- a/Engine/Source/Editor/UnrealEd/Classes/Factories/DialogueVoiceFactory.h
+++ b/Engine/Source/Editor/UnrealEd/Classes/Factories/DialogueVoiceFactory.h
@@ -10,7 +10,9 @@
 UCLASS(hidecategories=Object)
 class UDialogueVoiceFactory : public UFactory
 {
-	GENERATED_UCLASS_BODY()
+	GENERATED_BODY()
+public:
+	UDialogueVoiceFactory(const FObjectInitializer& ObjectInitializer = FObjectInitializer::Get());
 
 	// Begin UFactory Interface
 	virtual UObject* FactoryCreateNew(UClass* Class,UObject* InParent,FName Name,EObjectFlags Flags,UObject* Context,FFeedbackContext* Warn) override;
diff --git a/Engine/Source/Editor/UnrealEd/Classes/Factories/DialogueWaveFactory.h b/Engine/Source/Editor/UnrealEd/Classes/Factories/DialogueWaveFactory.h
--- a/Engine/Source/Editor/UnrealEd/Classes/Factories/DialogueWaveFactory.h
+++ b/Engine/Source/Editor/UnrealEd/Classes/Factories/DialogueWaveFactory.h
@@ -10,7 +10,9 @@
 UCLASS(hidecategories=Object)
 class UDialogueWaveFactory : public UFactory
 {
-	GENERATED_UCLASS_BODY()
+	GENERATED_BODY()
+public:
+	UDialogueWaveFactory(const FObjectInitializer& ObjectInitializer = FObjectInitializer::Get());
 
 	// Begin UFactory Interface
 	virtual UObject* FactoryCreateNew(UClass* Class,UObject* InParent,FName Name,EObjectFlags Flags,UObject* Context,FFeedbackContext* Warn) override;
diff --git a/Engine/Source/Editor/UnrealEd/Classes/Factories/EnumFactory.h b/Engine/Source/Editor/UnrealEd/Classes/Factories/EnumFactory.h
--- a/Engine/Source/Editor/UnrealEd/Classes/Factories/EnumFactory.h
+++ b/Engine/Source/Editor/UnrealEd/Classes/Factories/EnumFactory.h
@@ -7,7 +7,9 @@
 UCLASS(hidecategories=Object, MinimalAPI)
 class UEnumFactory : public UFactory
 {
-	GENERATED_UCLASS_BODY()
+	GENERATED_BODY()
+public:
+	UNREALED_API UEnumFactory(const FObjectInitializer& ObjectInitializer = FObjectInitializer::Get());
 
 	// Begin UFactory Interface
 	virtual UObject* FactoryCreateNew(UClass* Class,UObject* InParent,FName Name,EObjectFlags Flags,UObject* Context,FFeedbackContext* Warn) override;
diff --git a/Engine/Source/Editor/UnrealEd/Classes/Factories/Factory.h b/Engine/Source/Editor/UnrealEd/Classes/Factories/Factory.h
--- a/Engine/Source/Editor/UnrealEd/Classes/Factories/Factory.h
+++ b/Engine/Source/Editor/UnrealEd/Classes/Factories/Factory.h
@@ -12,7 +12,9 @@
 UCLASS(abstract)
 class UNREALED_API UFactory : public UObject
 {
-	GENERATED_UCLASS_BODY()
+	GENERATED_BODY()
+public:
+	UFactory(const FObjectInitializer& ObjectInitializer = FObjectInitializer::Get());
 
 	/** The class manufactured by this factory. */
 	UPROPERTY()
diff --git a/Engine/Source/Editor/UnrealEd/Classes/Factories/FbxAnimSequenceImportData.h b/Engine/Source/Editor/UnrealEd/Classes/Factories/FbxAnimSequenceImportData.h
--- a/Engine/Source/Editor/UnrealEd/Classes/Factories/FbxAnimSequenceImportData.h
+++ b/Engine/Source/Editor/UnrealEd/Classes/Factories/FbxAnimSequenceImportData.h
@@ -26,7 +26,9 @@ enum EFBXAnimationLengthImportType
 UCLASS()
 class UNREALED_API UFbxAnimSequenceImportData : public UFbxAssetImportData
 {
-	GENERATED_UCLASS_BODY()
+	GENERATED_BODY()
+public:
+	UFbxAnimSequenceImportData(const FObjectInitializer& ObjectInitializer = FObjectInitializer::Get());
 
 	/** Type of asset to import from the FBX file */
 	UPROPERTY(EditAnywhere, Category=ImportSettings, meta=(DisplayName = "Animation Length"))
@@ -61,4 +63,4 @@ class UNREALED_API UFbxAnimSequenceImportData : public UFbxAssetImportData
 
 	virtual bool CanEditChange(const UProperty* InProperty) const;
 
-};
\ No newline at end of file
+};
diff --git a/Engine/Source/Editor/UnrealEd/Classes/Factories/FbxAssetImportData.h b/Engine/Source/Editor/UnrealEd/Classes/Factories/FbxAssetImportData.h
--- a/Engine/Source/Editor/UnrealEd/Classes/Factories/FbxAssetImportData.h
+++ b/Engine/Source/Editor/UnrealEd/Classes/Factories/FbxAssetImportData.h
@@ -11,7 +11,9 @@
 UCLASS(config=EditorUserSettings, HideCategories=Object, abstract)
 class UNREALED_API UFbxAssetImportData : public UAssetImportData
 {
-	GENERATED_UCLASS_BODY()
+	GENERATED_BODY()
+public:
+	UFbxAssetImportData(const FObjectInitializer& ObjectInitializer = FObjectInitializer::Get());
 
 	UPROPERTY(EditAnywhere, Category=Transform, meta=(ImportType="StaticMesh|SkeletalMesh|Animation", ImportCategory="Transform"))
 	FVector ImportTranslation;
@@ -21,4 +23,4 @@ class UNREALED_API UFbxAssetImportData : public UAssetImportData
 
 	UPROPERTY(EditAnywhere, Category=Transform, meta=(ImportType="StaticMesh|SkeletalMesh|Animation", ImportCategory="Transform"))
 	float ImportUniformScale;
-};
\ No newline at end of file
+};
diff --git a/Engine/Source/Editor/UnrealEd/Classes/Factories/FbxFactory.h b/Engine/Source/Editor/UnrealEd/Classes/Factories/FbxFactory.h
--- a/Engine/Source/Editor/UnrealEd/Classes/Factories/FbxFactory.h
+++ b/Engine/Source/Editor/UnrealEd/Classes/Factories/FbxFactory.h
@@ -6,7 +6,9 @@
 UCLASS(hidecategories=Object)
 class UNREALED_API UFbxFactory : public UFactory
 {
-	GENERATED_UCLASS_BODY()
+	GENERATED_BODY()
+public:
+	UFbxFactory(const FObjectInitializer& ObjectInitializer = FObjectInitializer::Get());
 
 	UPROPERTY()
 	class UFbxImportUI* ImportUI;
diff --git a/Engine/Source/Editor/UnrealEd/Classes/Factories/FbxImportUI.h b/Engine/Source/Editor/UnrealEd/Classes/Factories/FbxImportUI.h
--- a/Engine/Source/Editor/UnrealEd/Classes/Factories/FbxImportUI.h
+++ b/Engine/Source/Editor/UnrealEd/Classes/Factories/FbxImportUI.h
@@ -25,7 +25,9 @@ enum EFBXImportType
 UCLASS(config=EditorUserSettings, AutoExpandCategories=(FTransform), HideCategories=Object, MinimalAPI)
 class UFbxImportUI : public UObject
 {
-	GENERATED_UCLASS_BODY()
+	GENERATED_BODY()
+public:
+	UNREALED_API UFbxImportUI(const FObjectInitializer& ObjectInitializer = FObjectInitializer::Get());
 
 	/** Whether or not the imported file is in OBJ format */
 	UPROPERTY()
diff --git a/Engine/Source/Editor/UnrealEd/Classes/Factories/FbxMeshImportData.h b/Engine/Source/Editor/UnrealEd/Classes/Factories/FbxMeshImportData.h
--- a/Engine/Source/Editor/UnrealEd/Classes/Factories/FbxMeshImportData.h
+++ b/Engine/Source/Editor/UnrealEd/Classes/Factories/FbxMeshImportData.h
@@ -18,7 +18,9 @@ enum EFBXNormalImportMethod
 UCLASS(config=EditorUserSettings, configdonotcheckdefaults, abstract)
 class UFbxMeshImportData : public UFbxAssetImportData
 {
-	GENERATED_UCLASS_BODY()
+	GENERATED_BODY()
+public:
+	UFbxMeshImportData(const FObjectInitializer& ObjectInitializer = FObjectInitializer::Get());
 
 	/** Enables importing of mesh LODs from FBX LOD groups, if present in the FBX file */
 	UPROPERTY(EditAnywhere, AdvancedDisplay, config, Category=ImportSettings, meta=(OBJRestrict="true", ImportType="Mesh", ToolTip="If enabled, creates LOD models for Unreal meshes from LODs in the import file; If not enabled, only the base mesh from the LOD group is imported"))
@@ -29,4 +31,4 @@ class UFbxMeshImportData : public UFbxAssetImportData
 	TEnumAsByte<enum EFBXNormalImportMethod> NormalImportMethod;
 
 	bool CanEditChange( const UProperty* InProperty ) const;
-};
\ No newline at end of file
+};
diff --git a/Engine/Source/Editor/UnrealEd/Classes/Factories/FbxSkeletalMeshImportData.h b/Engine/Source/Editor/UnrealEd/Classes/Factories/FbxSkeletalMeshImportData.h
--- a/Engine/Source/Editor/UnrealEd/Classes/Factories/FbxSkeletalMeshImportData.h
+++ b/Engine/Source/Editor/UnrealEd/Classes/Factories/FbxSkeletalMeshImportData.h
@@ -9,7 +9,9 @@
 UCLASS()
 class UFbxSkeletalMeshImportData : public UFbxMeshImportData
 {
-	GENERATED_UCLASS_BODY()
+	GENERATED_BODY()
+public:
+	UFbxSkeletalMeshImportData(const FObjectInitializer& ObjectInitializer = FObjectInitializer::Get());
 
 	/** Enable this option to update Skeleton (of the mesh)'s reference pose. Mesh's reference pose is always updated.  */
 	UPROPERTY(EditAnywhere, AdvancedDisplay, Category=ImportSettings, meta=(ImportType="SkeletalMesh", ToolTip="If enabled, update the Skeleton (of the mesh being imported)'s reference pose."))
@@ -39,4 +41,4 @@ class UFbxSkeletalMeshImportData : public UFbxMeshImportData
 	static UFbxSkeletalMeshImportData* GetImportDataForSkeletalMesh(USkeletalMesh* SkeletalMesh, UFbxSkeletalMeshImportData* TemplateForCreation);
 
 	bool CanEditChange( const UProperty* InProperty ) const;
-};
\ No newline at end of file
+};
diff --git a/Engine/Source/Editor/UnrealEd/Classes/Factories/FbxStaticMeshImportData.h b/Engine/Source/Editor/UnrealEd/Classes/Factories/FbxStaticMeshImportData.h
--- a/Engine/Source/Editor/UnrealEd/Classes/Factories/FbxStaticMeshImportData.h
+++ b/Engine/Source/Editor/UnrealEd/Classes/Factories/FbxStaticMeshImportData.h
@@ -24,7 +24,9 @@ namespace EVertexColorImportOption
 UCLASS(config=EditorUserSettings, AutoExpandCategories=(Options), MinimalAPI)
 class UFbxStaticMeshImportData : public UFbxMeshImportData
 {
-	GENERATED_UCLASS_BODY()
+	GENERATED_BODY()
+public:
+	UNREALED_API UFbxStaticMeshImportData(const FObjectInitializer& ObjectInitializer = FObjectInitializer::Get());
 
 	/** For static meshes, enabling this option will combine all meshes in the FBX into a single monolithic mesh in Unreal */
 	UPROPERTY(EditAnywhere, AdvancedDisplay, Category=ImportSettings, meta=(ImportType="StaticMesh"))
diff --git a/Engine/Source/Editor/UnrealEd/Classes/Factories/FbxTextureImportData.h b/Engine/Source/Editor/UnrealEd/Classes/Factories/FbxTextureImportData.h
--- a/Engine/Source/Editor/UnrealEd/Classes/Factories/FbxTextureImportData.h
+++ b/Engine/Source/Editor/UnrealEd/Classes/Factories/FbxTextureImportData.h
@@ -9,11 +9,13 @@
 UCLASS(AutoExpandCategories=(Texture))
 class UNREALED_API UFbxTextureImportData : public UFbxAssetImportData
 {
-	GENERATED_UCLASS_BODY()
+	GENERATED_BODY()
+public:
+	UFbxTextureImportData(const FObjectInitializer& ObjectInitializer = FObjectInitializer::Get());
 
 	/** If either importing of textures (or materials) is enabled, this option will cause normal map values to be inverted */
 	UPROPERTY(EditAnywhere, AdvancedDisplay, config, Category=ImportSettings, meta=(OBJRestrict="true"))
 	uint32 bInvertNormalMaps:1;
 
 	bool CanEditChange( const UProperty* InProperty ) const;
-};
\ No newline at end of file
+};
diff --git a/Engine/Source/Editor/UnrealEd/Classes/Factories/FontFactory.h b/Engine/Source/Editor/UnrealEd/Classes/Factories/FontFactory.h
--- a/Engine/Source/Editor/UnrealEd/Classes/Factories/FontFactory.h
+++ b/Engine/Source/Editor/UnrealEd/Classes/Factories/FontFactory.h
@@ -10,7 +10,9 @@
 UCLASS()
 class UNREALED_API UFontFactory : public UFactory
 {
-	GENERATED_UCLASS_BODY()
+	GENERATED_BODY()
+public:
+	UFontFactory(const FObjectInitializer& ObjectInitializer = FObjectInitializer::Get());
 
 	// Begin UFactory Interface
 	virtual UObject* FactoryCreateNew(UClass* InClass, UObject* InParent, FName InName, EObjectFlags Flags, UObject* Context, FFeedbackContext* Warn) override;
diff --git a/Engine/Source/Editor/UnrealEd/Classes/Factories/FontFileImportFactory.h b/Engine/Source/Editor/UnrealEd/Classes/Factories/FontFileImportFactory.h
--- a/Engine/Source/Editor/UnrealEd/Classes/Factories/FontFileImportFactory.h
+++ b/Engine/Source/Editor/UnrealEd/Classes/Factories/FontFileImportFactory.h
@@ -10,7 +10,9 @@
 UCLASS()
 class UNREALED_API UFontFileImportFactory : public UFactory
 {
-	GENERATED_UCLASS_BODY()
+	GENERATED_BODY()
+public:
+	UFontFileImportFactory(const FObjectInitializer& ObjectInitializer = FObjectInitializer::Get());
 
 	// Begin UFactory Interface
 	virtual UObject* FactoryCreateBinary(UClass* Class, UObject* InParent, FName Name, EObjectFlags Flags, UObject* Context, const TCHAR* Type, const uint8*& Buffer, const uint8* BufferEnd, FFeedbackContext* Warn) override;
diff --git a/Engine/Source/Editor/UnrealEd/Classes/Factories/ForceFeedbackEffectFactory.h b/Engine/Source/Editor/UnrealEd/Classes/Factories/ForceFeedbackEffectFactory.h
--- a/Engine/Source/Editor/UnrealEd/Classes/Factories/ForceFeedbackEffectFactory.h
+++ b/Engine/Source/Editor/UnrealEd/Classes/Factories/ForceFeedbackEffectFactory.h
@@ -6,7 +6,9 @@
 UCLASS(hidecategories=Object)
 class UForceFeedbackEffectFactory : public UFactory
 {
-	GENERATED_UCLASS_BODY()
+	GENERATED_BODY()
+public:
+	UForceFeedbackEffectFactory(const FObjectInitializer& ObjectInitializer = FObjectInitializer::Get());
 
 
 	// Begin UFactory Interface
diff --git a/Engine/Source/Editor/UnrealEd/Classes/Factories/InterpDataFactoryNew.h b/Engine/Source/Editor/UnrealEd/Classes/Factories/InterpDataFactoryNew.h
--- a/Engine/Source/Editor/UnrealEd/Classes/Factories/InterpDataFactoryNew.h
+++ b/Engine/Source/Editor/UnrealEd/Classes/Factories/InterpDataFactoryNew.h
@@ -10,7 +10,9 @@
 UCLASS(hidecategories=Object)
 class UInterpDataFactoryNew : public UFactory
 {
-	GENERATED_UCLASS_BODY()
+	GENERATED_BODY()
+public:
+	UInterpDataFactoryNew(const FObjectInitializer& ObjectInitializer = FObjectInitializer::Get());
 
 
 	// Begin UFactory Interface
diff --git a/Engine/Source/Editor/UnrealEd/Classes/Factories/LevelFactory.h b/Engine/Source/Editor/UnrealEd/Classes/Factories/LevelFactory.h
--- a/Engine/Source/Editor/UnrealEd/Classes/Factories/LevelFactory.h
+++ b/Engine/Source/Editor/UnrealEd/Classes/Factories/LevelFactory.h
@@ -10,7 +10,9 @@
 UCLASS(MinimalAPI)
 class ULevelFactory : public UFactory
 {
-	GENERATED_UCLASS_BODY()
+	GENERATED_BODY()
+public:
+	UNREALED_API ULevelFactory(const FObjectInitializer& ObjectInitializer = FObjectInitializer::Get());
 
 
 	// Begin UFactory Interface
diff --git a/Engine/Source/Editor/UnrealEd/Classes/Factories/MaterialFactoryNew.h b/Engine/Source/Editor/UnrealEd/Classes/Factories/MaterialFactoryNew.h
--- a/Engine/Source/Editor/UnrealEd/Classes/Factories/MaterialFactoryNew.h
+++ b/Engine/Source/Editor/UnrealEd/Classes/Factories/MaterialFactoryNew.h
@@ -10,7 +10,9 @@
 UCLASS(hidecategories=Object, collapsecategories, MinimalAPI)
 class UMaterialFactoryNew : public UFactory
 {
-	GENERATED_UCLASS_BODY()
+	GENERATED_BODY()
+public:
+	UNREALED_API UMaterialFactoryNew(const FObjectInitializer& ObjectInitializer = FObjectInitializer::Get());
 
 	/** An initial texture to place in the newly created material */
 	UPROPERTY()
@@ -19,4 +21,4 @@ class UMaterialFactoryNew : public UFactory
 	// Begin UFactory Interface
 	virtual UObject* FactoryCreateNew(UClass* Class,UObject* InParent,FName Name,EObjectFlags Flags,UObject* Context,FFeedbackContext* Warn) override;
 	// Begin UFactory Interface	
-};
\ No newline at end of file
+};
diff --git a/Engine/Source/Editor/UnrealEd/Classes/Factories/MaterialFunctionFactoryNew.h b/Engine/Source/Editor/UnrealEd/Classes/Factories/MaterialFunctionFactoryNew.h
--- a/Engine/Source/Editor/UnrealEd/Classes/Factories/MaterialFunctionFactoryNew.h
+++ b/Engine/Source/Editor/UnrealEd/Classes/Factories/MaterialFunctionFactoryNew.h
@@ -10,7 +10,9 @@
 UCLASS(hidecategories=Object, collapsecategories)
 class UMaterialFunctionFactoryNew : public UFactory
 {
-	GENERATED_UCLASS_BODY()
+	GENERATED_BODY()
+public:
+	UMaterialFunctionFactoryNew(const FObjectInitializer& ObjectInitializer = FObjectInitializer::Get());
 
 
 	// Begin UFactory Interface
diff --git a/Engine/Source/Editor/UnrealEd/Classes/Factories/MaterialInstanceConstantFactoryNew.h b/Engine/Source/Editor/UnrealEd/Classes/Factories/MaterialInstanceConstantFactoryNew.h
--- a/Engine/Source/Editor/UnrealEd/Classes/Factories/MaterialInstanceConstantFactoryNew.h
+++ b/Engine/Source/Editor/UnrealEd/Classes/Factories/MaterialInstanceConstantFactoryNew.h
@@ -10,7 +10,9 @@
 UCLASS(hidecategories=Object, collapsecategories,MinimalAPI)
 class UMaterialInstanceConstantFactoryNew : public UFactory
 {
-	GENERATED_UCLASS_BODY()
+	GENERATED_BODY()
+public:
+	UNREALED_API UMaterialInstanceConstantFactoryNew(const FObjectInitializer& ObjectInitializer = FObjectInitializer::Get());
 
 	UPROPERTY()
 	class UMaterialInterface* InitialParent;
diff --git a/Engine/Source/Editor/UnrealEd/Classes/Factories/MaterialParameterCollectionFactoryNew.h b/Engine/Source/Editor/UnrealEd/Classes/Factories/MaterialParameterCollectionFactoryNew.h
--- a/Engine/Source/Editor/UnrealEd/Classes/Factories/MaterialParameterCollectionFactoryNew.h
+++ b/Engine/Source/Editor/UnrealEd/Classes/Factories/MaterialParameterCollectionFactoryNew.h
@@ -10,7 +10,9 @@
 UCLASS(hidecategories=Object, collapsecategories)
 class UMaterialParameterCollectionFactoryNew : public UFactory
 {
-	GENERATED_UCLASS_BODY()
+	GENERATED_BODY()
+public:
+	UMaterialParameterCollectionFactoryNew(const FObjectInitializer& ObjectInitializer = FObjectInitializer::Get());
 
 	// Begin UFactory Interface
 	virtual UObject* FactoryCreateNew(UClass* Class,UObject* InParent,FName Name,EObjectFlags Flags,UObject* Context,FFeedbackContext* Warn) override;
diff --git a/Engine/Source/Editor/UnrealEd/Classes/Factories/ModelFactory.h b/Engine/Source/Editor/UnrealEd/Classes/Factories/ModelFactory.h
--- a/Engine/Source/Editor/UnrealEd/Classes/Factories/ModelFactory.h
+++ b/Engine/Source/Editor/UnrealEd/Classes/Factories/ModelFactory.h
@@ -10,7 +10,9 @@
 UCLASS()
 class UModelFactory : public UFactory
 {
-	GENERATED_UCLASS_BODY()
+	GENERATED_BODY()
+public:
+	UModelFactory(const FObjectInitializer& ObjectInitializer = FObjectInitializer::Get());
 
 
 	// Begin UFactory Interface
diff --git a/Engine/Source/Editor/UnrealEd/Classes/Factories/NiagaraEffectFactoryNew.h b/Engine/Source/Editor/UnrealEd/Classes/Factories/NiagaraEffectFactoryNew.h
--- a/Engine/Source/Editor/UnrealEd/Classes/Factories/NiagaraEffectFactoryNew.h
+++ b/Engine/Source/Editor/UnrealEd/Classes/Factories/NiagaraEffectFactoryNew.h
@@ -6,7 +6,9 @@
 UCLASS(hidecategories = Object)
 class UNiagaraEffectFactoryNew : public UFactory
 {
-	GENERATED_UCLASS_BODY()
+	GENERATED_BODY()
+public:
+	UNiagaraEffectFactoryNew(const FObjectInitializer& ObjectInitializer = FObjectInitializer::Get());
 
 	// Begin UFactory Interface
 	virtual UObject* FactoryCreateNew(UClass* Class, UObject* InParent, FName Name, EObjectFlags Flags, UObject* Context, FFeedbackContext* Warn) override;
diff --git a/Engine/Source/Editor/UnrealEd/Classes/Factories/NiagaraScriptFactoryNew.h b/Engine/Source/Editor/UnrealEd/Classes/Factories/NiagaraScriptFactoryNew.h
--- a/Engine/Source/Editor/UnrealEd/Classes/Factories/NiagaraScriptFactoryNew.h
+++ b/Engine/Source/Editor/UnrealEd/Classes/Factories/NiagaraScriptFactoryNew.h
@@ -6,7 +6,9 @@
 UCLASS(hidecategories=Object)
 class UNiagaraScriptFactoryNew : public UFactory
 {
-	GENERATED_UCLASS_BODY()
+	GENERATED_BODY()
+public:
+	UNiagaraScriptFactoryNew(const FObjectInitializer& ObjectInitializer = FObjectInitializer::Get());
 
 	// Begin UFactory Interface
 	virtual UObject* FactoryCreateNew(UClass* Class,UObject* InParent,FName Name,EObjectFlags Flags,UObject* Context,FFeedbackContext* Warn) override;
diff --git a/Engine/Source/Editor/UnrealEd/Classes/Factories/ObjectLibraryFactory.h b/Engine/Source/Editor/UnrealEd/Classes/Factories/ObjectLibraryFactory.h
--- a/Engine/Source/Editor/UnrealEd/Classes/Factories/ObjectLibraryFactory.h
+++ b/Engine/Source/Editor/UnrealEd/Classes/Factories/ObjectLibraryFactory.h
@@ -6,7 +6,9 @@
 UCLASS(hidecategories=Object)
 class UObjectLibraryFactory : public UFactory
 {
-	GENERATED_UCLASS_BODY()
+	GENERATED_BODY()
+public:
+	UObjectLibraryFactory(const FObjectInitializer& ObjectInitializer = FObjectInitializer::Get());
 
 	// Begin UFactory Interface
 	virtual UObject* FactoryCreateNew(UClass* Class,UObject* InParent,FName Name,EObjectFlags Flags,UObject* Context,FFeedbackContext* Warn) override;
diff --git a/Engine/Source/Editor/UnrealEd/Classes/Factories/PackFactory.h b/Engine/Source/Editor/UnrealEd/Classes/Factories/PackFactory.h
--- a/Engine/Source/Editor/UnrealEd/Classes/Factories/PackFactory.h
+++ b/Engine/Source/Editor/UnrealEd/Classes/Factories/PackFactory.h
@@ -9,9 +9,11 @@
 UCLASS(hidecategories=Object)
 class UPackFactory : public UFactory
 {
-	GENERATED_UCLASS_BODY()
+	GENERATED_BODY()
+public:
+	UPackFactory(const FObjectInitializer& ObjectInitializer = FObjectInitializer::Get());
 
 	// Begin UFactory Interface
 	virtual UObject* FactoryCreateBinary( UClass* Class, UObject* InParent, FName Name, EObjectFlags Flags, UObject* Context, const TCHAR* Type, const uint8*& Buffer, const uint8* BufferEnd, FFeedbackContext* Warn ) override;
 	// End UFactory Interface
-};
\ No newline at end of file
+};
diff --git a/Engine/Source/Editor/UnrealEd/Classes/Factories/PackageFactory.h b/Engine/Source/Editor/UnrealEd/Classes/Factories/PackageFactory.h
--- a/Engine/Source/Editor/UnrealEd/Classes/Factories/PackageFactory.h
+++ b/Engine/Source/Editor/UnrealEd/Classes/Factories/PackageFactory.h
@@ -10,7 +10,9 @@
 UCLASS()
 class UPackageFactory : public UFactory
 {
-	GENERATED_UCLASS_BODY()
+	GENERATED_BODY()
+public:
+	UPackageFactory(const FObjectInitializer& ObjectInitializer = FObjectInitializer::Get());
 
 
 	// Begin UFactory Interface
diff --git a/Engine/Source/Editor/UnrealEd/Classes/Factories/ParticleSystemFactoryNew.h b/Engine/Source/Editor/UnrealEd/Classes/Factories/ParticleSystemFactoryNew.h
--- a/Engine/Source/Editor/UnrealEd/Classes/Factories/ParticleSystemFactoryNew.h
+++ b/Engine/Source/Editor/UnrealEd/Classes/Factories/ParticleSystemFactoryNew.h
@@ -10,7 +10,9 @@
 UCLASS(hidecategories=Object)
 class UParticleSystemFactoryNew : public UFactory
 {
-	GENERATED_UCLASS_BODY()
+	GENERATED_BODY()
+public:
+	UParticleSystemFactoryNew(const FObjectInitializer& ObjectInitializer = FObjectInitializer::Get());
 
 
 	// Begin UFactory Interface
diff --git a/Engine/Source/Editor/UnrealEd/Classes/Factories/PhysicalMaterialFactoryNew.h b/Engine/Source/Editor/UnrealEd/Classes/Factories/PhysicalMaterialFactoryNew.h
--- a/Engine/Source/Editor/UnrealEd/Classes/Factories/PhysicalMaterialFactoryNew.h
+++ b/Engine/Source/Editor/UnrealEd/Classes/Factories/PhysicalMaterialFactoryNew.h
@@ -10,7 +10,9 @@
 UCLASS(hidecategories=Object)
 class UPhysicalMaterialFactoryNew : public UFactory
 {
-	GENERATED_UCLASS_BODY()
+	GENERATED_BODY()
+public:
+	UPhysicalMaterialFactoryNew(const FObjectInitializer& ObjectInitializer = FObjectInitializer::Get());
 
 
 	// Begin UFactory Interface
diff --git a/Engine/Source/Editor/UnrealEd/Classes/Factories/PolysFactory.h b/Engine/Source/Editor/UnrealEd/Classes/Factories/PolysFactory.h
--- a/Engine/Source/Editor/UnrealEd/Classes/Factories/PolysFactory.h
+++ b/Engine/Source/Editor/UnrealEd/Classes/Factories/PolysFactory.h
@@ -10,7 +10,9 @@
 UCLASS()
 class UPolysFactory : public UFactory
 {
-	GENERATED_UCLASS_BODY()
+	GENERATED_BODY()
+public:
+	UPolysFactory(const FObjectInitializer& ObjectInitializer = FObjectInitializer::Get());
 
 
 	// Begin UFactory Interface
diff --git a/Engine/Source/Editor/UnrealEd/Classes/Factories/ReimportCurveFactory.h b/Engine/Source/Editor/UnrealEd/Classes/Factories/ReimportCurveFactory.h
--- a/Engine/Source/Editor/UnrealEd/Classes/Factories/ReimportCurveFactory.h
+++ b/Engine/Source/Editor/UnrealEd/Classes/Factories/ReimportCurveFactory.h
@@ -6,7 +6,9 @@
 UCLASS()
 class UReimportCurveFactory : public UCSVImportFactory, public FReimportHandler
 {
-	GENERATED_UCLASS_BODY()
+	GENERATED_BODY()
+public:
+	UReimportCurveFactory(const FObjectInitializer& ObjectInitializer = FObjectInitializer::Get());
 
 
 	// Begin FReimportHandler interface
diff --git a/Engine/Source/Editor/UnrealEd/Classes/Factories/ReimportCurveTableFactory.h b/Engine/Source/Editor/UnrealEd/Classes/Factories/ReimportCurveTableFactory.h
--- a/Engine/Source/Editor/UnrealEd/Classes/Factories/ReimportCurveTableFactory.h
+++ b/Engine/Source/Editor/UnrealEd/Classes/Factories/ReimportCurveTableFactory.h
@@ -6,7 +6,9 @@
 UCLASS()
 class UReimportCurveTableFactory : public UCSVImportFactory, public FReimportHandler
 {
-	GENERATED_UCLASS_BODY()
+	GENERATED_BODY()
+public:
+	UReimportCurveTableFactory(const FObjectInitializer& ObjectInitializer = FObjectInitializer::Get());
 
 	// Begin FReimportHandler interface
 	virtual bool CanReimport( UObject* Obj, TArray<FString>& OutFilenames ) override;
diff --git a/Engine/Source/Editor/UnrealEd/Classes/Factories/ReimportDataTableFactory.h b/Engine/Source/Editor/UnrealEd/Classes/Factories/ReimportDataTableFactory.h
--- a/Engine/Source/Editor/UnrealEd/Classes/Factories/ReimportDataTableFactory.h
+++ b/Engine/Source/Editor/UnrealEd/Classes/Factories/ReimportDataTableFactory.h
@@ -6,7 +6,9 @@
 UCLASS()
 class UReimportDataTableFactory : public UCSVImportFactory, public FReimportHandler
 {
-	GENERATED_UCLASS_BODY()
+	GENERATED_BODY()
+public:
+	UReimportDataTableFactory(const FObjectInitializer& ObjectInitializer = FObjectInitializer::Get());
 
 	// Begin FReimportHandler interface
 	virtual bool CanReimport( UObject* Obj, TArray<FString>& OutFilenames ) override;
diff --git a/Engine/Source/Editor/UnrealEd/Classes/Factories/ReimportDestructibleMeshFactory.h b/Engine/Source/Editor/UnrealEd/Classes/Factories/ReimportDestructibleMeshFactory.h
--- a/Engine/Source/Editor/UnrealEd/Classes/Factories/ReimportDestructibleMeshFactory.h
+++ b/Engine/Source/Editor/UnrealEd/Classes/Factories/ReimportDestructibleMeshFactory.h
@@ -8,7 +8,9 @@
 UCLASS(hideCategories=Object)
 class UReimportDestructibleMeshFactory : public UDestructibleMeshFactory, public FReimportHandler
 {
-    GENERATED_UCLASS_BODY()
+    GENERATED_BODY()
+public:
+    UReimportDestructibleMeshFactory(const FObjectInitializer& ObjectInitializer = FObjectInitializer::Get());
 public:
 	virtual FText GetDisplayName() const override;
 #if WITH_APEX
diff --git a/Engine/Source/Editor/UnrealEd/Classes/Factories/ReimportFbxAnimSequenceFactory.h b/Engine/Source/Editor/UnrealEd/Classes/Factories/ReimportFbxAnimSequenceFactory.h
--- a/Engine/Source/Editor/UnrealEd/Classes/Factories/ReimportFbxAnimSequenceFactory.h
+++ b/Engine/Source/Editor/UnrealEd/Classes/Factories/ReimportFbxAnimSequenceFactory.h
@@ -10,7 +10,9 @@
 UCLASS(collapsecategories)
 class UReimportFbxAnimSequenceFactory : public UFbxFactory, public FReimportHandler
 {
-	GENERATED_UCLASS_BODY()
+	GENERATED_BODY()
+public:
+	UReimportFbxAnimSequenceFactory(const FObjectInitializer& ObjectInitializer = FObjectInitializer::Get());
 
 
 	// Begin FReimportHandler interface
diff --git a/Engine/Source/Editor/UnrealEd/Classes/Factories/ReimportFbxSkeletalMeshFactory.h b/Engine/Source/Editor/UnrealEd/Classes/Factories/ReimportFbxSkeletalMeshFactory.h
--- a/Engine/Source/Editor/UnrealEd/Classes/Factories/ReimportFbxSkeletalMeshFactory.h
+++ b/Engine/Source/Editor/UnrealEd/Classes/Factories/ReimportFbxSkeletalMeshFactory.h
@@ -10,7 +10,9 @@
 UCLASS(collapsecategories)
 class UReimportFbxSkeletalMeshFactory : public UFbxFactory, public FReimportHandler
 {
-	GENERATED_UCLASS_BODY()
+	GENERATED_BODY()
+public:
+	UReimportFbxSkeletalMeshFactory(const FObjectInitializer& ObjectInitializer = FObjectInitializer::Get());
 
 
 	// Begin FReimportHandler interface
diff --git a/Engine/Source/Editor/UnrealEd/Classes/Factories/ReimportFbxStaticMeshFactory.h b/Engine/Source/Editor/UnrealEd/Classes/Factories/ReimportFbxStaticMeshFactory.h
--- a/Engine/Source/Editor/UnrealEd/Classes/Factories/ReimportFbxStaticMeshFactory.h
+++ b/Engine/Source/Editor/UnrealEd/Classes/Factories/ReimportFbxStaticMeshFactory.h
@@ -10,7 +10,9 @@
 UCLASS(collapsecategories)
 class UReimportFbxStaticMeshFactory : public UFbxFactory, public FReimportHandler
 {
-	GENERATED_UCLASS_BODY()
+	GENERATED_BODY()
+public:
+	UReimportFbxStaticMeshFactory(const FObjectInitializer& ObjectInitializer = FObjectInitializer::Get());
 
 
 	// Begin FReimportHandler interface
diff --git a/Engine/Source/Editor/UnrealEd/Classes/Factories/ReimportSoundFactory.h b/Engine/Source/Editor/UnrealEd/Classes/Factories/ReimportSoundFactory.h
--- a/Engine/Source/Editor/UnrealEd/Classes/Factories/ReimportSoundFactory.h
+++ b/Engine/Source/Editor/UnrealEd/Classes/Factories/ReimportSoundFactory.h
@@ -10,7 +10,9 @@
 UCLASS(hidecategories=Object, collapsecategories)
 class UReimportSoundFactory : public USoundFactory, public FReimportHandler
 {
-	GENERATED_UCLASS_BODY()
+	GENERATED_BODY()
+public:
+	UReimportSoundFactory(const FObjectInitializer& ObjectInitializer = FObjectInitializer::Get());
 
 
 	// Begin FReimportHandler interface
diff --git a/Engine/Source/Editor/UnrealEd/Classes/Factories/ReimportSoundSurroundFactory.h b/Engine/Source/Editor/UnrealEd/Classes/Factories/ReimportSoundSurroundFactory.h
--- a/Engine/Source/Editor/UnrealEd/Classes/Factories/ReimportSoundSurroundFactory.h
+++ b/Engine/Source/Editor/UnrealEd/Classes/Factories/ReimportSoundSurroundFactory.h
@@ -10,7 +10,9 @@
 UCLASS(hidecategories=Object, collapsecategories)
 class UReimportSoundSurroundFactory : public USoundSurroundFactory, public FReimportHandler
 {
-	GENERATED_UCLASS_BODY()
+	GENERATED_BODY()
+public:
+	UReimportSoundSurroundFactory(const FObjectInitializer& ObjectInitializer = FObjectInitializer::Get());
 
 	UPROPERTY()
 	TArray<FString> ReimportPaths;
diff --git a/Engine/Source/Editor/UnrealEd/Classes/Factories/ReimportTextureFactory.h b/Engine/Source/Editor/UnrealEd/Classes/Factories/ReimportTextureFactory.h
--- a/Engine/Source/Editor/UnrealEd/Classes/Factories/ReimportTextureFactory.h
+++ b/Engine/Source/Editor/UnrealEd/Classes/Factories/ReimportTextureFactory.h
@@ -10,7 +10,9 @@
 UCLASS(hidecategories=(LightMap, DitherMipMaps, LODGroup), collapsecategories)
 class UReimportTextureFactory : public UTextureFactory, public FReimportHandler
 {
-	GENERATED_UCLASS_BODY()
+	GENERATED_BODY()
+public:
+	UReimportTextureFactory(const FObjectInitializer& ObjectInitializer = FObjectInitializer::Get());
 
 	UPROPERTY()
 	class UTexture* pOriginalTex;
diff --git a/Engine/Source/Editor/UnrealEd/Classes/Factories/ReimportVectorFieldStaticFactory.h b/Engine/Source/Editor/UnrealEd/Classes/Factories/ReimportVectorFieldStaticFactory.h
--- a/Engine/Source/Editor/UnrealEd/Classes/Factories/ReimportVectorFieldStaticFactory.h
+++ b/Engine/Source/Editor/UnrealEd/Classes/Factories/ReimportVectorFieldStaticFactory.h
@@ -10,7 +10,9 @@
 UCLASS()
 class UReimportVectorFieldStaticFactory : public UVectorFieldStaticFactory, public FReimportHandler
 {
-	GENERATED_UCLASS_BODY()
+	GENERATED_BODY()
+public:
+	UReimportVectorFieldStaticFactory(const FObjectInitializer& ObjectInitializer = FObjectInitializer::Get());
 
 	// Begin FReimportHandler interface
 	virtual bool CanReimport( UObject* Obj, TArray<FString>& OutFilenames ) override;
diff --git a/Engine/Source/Editor/UnrealEd/Classes/Factories/ReverbEffectFactory.h b/Engine/Source/Editor/UnrealEd/Classes/Factories/ReverbEffectFactory.h
--- a/Engine/Source/Editor/UnrealEd/Classes/Factories/ReverbEffectFactory.h
+++ b/Engine/Source/Editor/UnrealEd/Classes/Factories/ReverbEffectFactory.h
@@ -10,7 +10,9 @@
 UCLASS(hidecategories=Object)
 class UReverbEffectFactory : public UFactory
 {
-	GENERATED_UCLASS_BODY()
+	GENERATED_BODY()
+public:
+	UReverbEffectFactory(const FObjectInitializer& ObjectInitializer = FObjectInitializer::Get());
 
 
 	// Begin UFactory Interface
diff --git a/Engine/Source/Editor/UnrealEd/Classes/Factories/SkeletonFactory.h b/Engine/Source/Editor/UnrealEd/Classes/Factories/SkeletonFactory.h
--- a/Engine/Source/Editor/UnrealEd/Classes/Factories/SkeletonFactory.h
+++ b/Engine/Source/Editor/UnrealEd/Classes/Factories/SkeletonFactory.h
@@ -10,7 +10,9 @@
 UCLASS(HideCategories=Object,MinimalAPI)
 class USkeletonFactory : public UFactory
 {
-	GENERATED_UCLASS_BODY()
+	GENERATED_BODY()
+public:
+	UNREALED_API USkeletonFactory(const FObjectInitializer& ObjectInitializer = FObjectInitializer::Get());
 
 	// The skeletal mesh with which to initialize this skeleton.
 	UPROPERTY()
diff --git a/Engine/Source/Editor/UnrealEd/Classes/Factories/SlateBrushAssetFactory.h b/Engine/Source/Editor/UnrealEd/Classes/Factories/SlateBrushAssetFactory.h
--- a/Engine/Source/Editor/UnrealEd/Classes/Factories/SlateBrushAssetFactory.h
+++ b/Engine/Source/Editor/UnrealEd/Classes/Factories/SlateBrushAssetFactory.h
@@ -8,7 +8,9 @@
 UCLASS(hidecategories=Object, MinimalAPI)
 class USlateBrushAssetFactory : public UFactory
 {
-	GENERATED_UCLASS_BODY()
+	GENERATED_BODY()
+public:
+	UNREALED_API USlateBrushAssetFactory(const FObjectInitializer& ObjectInitializer = FObjectInitializer::Get());
 
 	/** An initial texture to assign to the newly created slate brush */
 	UPROPERTY()
diff --git a/Engine/Source/Editor/UnrealEd/Classes/Factories/SlateWidgetStyleAssetFactory.h b/Engine/Source/Editor/UnrealEd/Classes/Factories/SlateWidgetStyleAssetFactory.h
--- a/Engine/Source/Editor/UnrealEd/Classes/Factories/SlateWidgetStyleAssetFactory.h
+++ b/Engine/Source/Editor/UnrealEd/Classes/Factories/SlateWidgetStyleAssetFactory.h
@@ -8,7 +8,9 @@
 UCLASS(hidecategories=Object)
 class USlateWidgetStyleAssetFactory : public UFactory
 {
-	GENERATED_UCLASS_BODY()
+	GENERATED_BODY()
+public:
+	USlateWidgetStyleAssetFactory(const FObjectInitializer& ObjectInitializer = FObjectInitializer::Get());
 
 	UPROPERTY(EditAnywhere, Category=StyleType)
 	TSubclassOf< USlateWidgetStyleContainerBase > StyleType;
diff --git a/Engine/Source/Editor/UnrealEd/Classes/Factories/SoundAttenuationFactory.h b/Engine/Source/Editor/UnrealEd/Classes/Factories/SoundAttenuationFactory.h
--- a/Engine/Source/Editor/UnrealEd/Classes/Factories/SoundAttenuationFactory.h
+++ b/Engine/Source/Editor/UnrealEd/Classes/Factories/SoundAttenuationFactory.h
@@ -10,7 +10,9 @@
 UCLASS(hidecategories=Object, MinimalAPI)
 class USoundAttenuationFactory : public UFactory
 {
-	GENERATED_UCLASS_BODY()
+	GENERATED_BODY()
+public:
+	UNREALED_API USoundAttenuationFactory(const FObjectInitializer& ObjectInitializer = FObjectInitializer::Get());
 
 	// Begin UFactory Interface
 	virtual UObject* FactoryCreateNew(UClass* Class,UObject* InParent,FName Name,EObjectFlags Flags,UObject* Context,FFeedbackContext* Warn) override;
diff --git a/Engine/Source/Editor/UnrealEd/Classes/Factories/SoundClassFactory.h b/Engine/Source/Editor/UnrealEd/Classes/Factories/SoundClassFactory.h
--- a/Engine/Source/Editor/UnrealEd/Classes/Factories/SoundClassFactory.h
+++ b/Engine/Source/Editor/UnrealEd/Classes/Factories/SoundClassFactory.h
@@ -10,7 +10,9 @@
 UCLASS(hidecategories=Object)
 class USoundClassFactory : public UFactory
 {
-	GENERATED_UCLASS_BODY()
+	GENERATED_BODY()
+public:
+	USoundClassFactory(const FObjectInitializer& ObjectInitializer = FObjectInitializer::Get());
 
 
 	// Begin UFactory Interface
diff --git a/Engine/Source/Editor/UnrealEd/Classes/Factories/SoundCueFactoryNew.h b/Engine/Source/Editor/UnrealEd/Classes/Factories/SoundCueFactoryNew.h
--- a/Engine/Source/Editor/UnrealEd/Classes/Factories/SoundCueFactoryNew.h
+++ b/Engine/Source/Editor/UnrealEd/Classes/Factories/SoundCueFactoryNew.h
@@ -10,7 +10,9 @@
 UCLASS(hidecategories=Object, MinimalAPI)
 class USoundCueFactoryNew : public UFactory
 {
-	GENERATED_UCLASS_BODY()
+	GENERATED_BODY()
+public:
+	UNREALED_API USoundCueFactoryNew(const FObjectInitializer& ObjectInitializer = FObjectInitializer::Get());
 
 
 	// Begin UFactory Interface
diff --git a/Engine/Source/Editor/UnrealEd/Classes/Factories/SoundFactory.h b/Engine/Source/Editor/UnrealEd/Classes/Factories/SoundFactory.h
--- a/Engine/Source/Editor/UnrealEd/Classes/Factories/SoundFactory.h
+++ b/Engine/Source/Editor/UnrealEd/Classes/Factories/SoundFactory.h
@@ -10,7 +10,9 @@
 UCLASS(hidecategories=Object)
 class USoundFactory : public UFactory
 {
-	GENERATED_UCLASS_BODY()
+	GENERATED_BODY()
+public:
+	USoundFactory(const FObjectInitializer& ObjectInitializer = FObjectInitializer::Get());
 
 	/** If enabled, a sound cue will automatically be created for the sound */
 	UPROPERTY(EditAnywhere, Category=SoundFactory, meta=(ToolTip="If enabled, a sound cue will automatically be created for the sound"))
diff --git a/Engine/Source/Editor/UnrealEd/Classes/Factories/SoundMixFactory.h b/Engine/Source/Editor/UnrealEd/Classes/Factories/SoundMixFactory.h
--- a/Engine/Source/Editor/UnrealEd/Classes/Factories/SoundMixFactory.h
+++ b/Engine/Source/Editor/UnrealEd/Classes/Factories/SoundMixFactory.h
@@ -10,7 +10,9 @@
 UCLASS(hidecategories=Object)
 class USoundMixFactory : public UFactory
 {
-	GENERATED_UCLASS_BODY()
+	GENERATED_BODY()
+public:
+	USoundMixFactory(const FObjectInitializer& ObjectInitializer = FObjectInitializer::Get());
 
 
 	// Begin UFactory Interface
diff --git a/Engine/Source/Editor/UnrealEd/Classes/Factories/SoundSurroundFactory.h b/Engine/Source/Editor/UnrealEd/Classes/Factories/SoundSurroundFactory.h
--- a/Engine/Source/Editor/UnrealEd/Classes/Factories/SoundSurroundFactory.h
+++ b/Engine/Source/Editor/UnrealEd/Classes/Factories/SoundSurroundFactory.h
@@ -10,7 +10,9 @@
 UCLASS(hidecategories=Object)
 class USoundSurroundFactory : public UFactory
 {
-	GENERATED_UCLASS_BODY()
+	GENERATED_BODY()
+public:
+	USoundSurroundFactory(const FObjectInitializer& ObjectInitializer = FObjectInitializer::Get());
 
 	UPROPERTY()
 	float CueVolume;
diff --git a/Engine/Source/Editor/UnrealEd/Classes/Factories/StructureFactory.h b/Engine/Source/Editor/UnrealEd/Classes/Factories/StructureFactory.h
--- a/Engine/Source/Editor/UnrealEd/Classes/Factories/StructureFactory.h
+++ b/Engine/Source/Editor/UnrealEd/Classes/Factories/StructureFactory.h
@@ -7,7 +7,9 @@
 UCLASS(hidecategories=Object, MinimalAPI)
 class UStructureFactory : public UFactory
 {
-	GENERATED_UCLASS_BODY()
+	GENERATED_BODY()
+public:
+	UNREALED_API UStructureFactory(const FObjectInitializer& ObjectInitializer = FObjectInitializer::Get());
 
 	// Begin UFactory Interface
 	virtual UObject* FactoryCreateNew(UClass* Class,UObject* InParent,FName Name,EObjectFlags Flags,UObject* Context,FFeedbackContext* Warn) override;
diff --git a/Engine/Source/Editor/UnrealEd/Classes/Factories/SubsurfaceProfileFactory.h b/Engine/Source/Editor/UnrealEd/Classes/Factories/SubsurfaceProfileFactory.h
--- a/Engine/Source/Editor/UnrealEd/Classes/Factories/SubsurfaceProfileFactory.h
+++ b/Engine/Source/Editor/UnrealEd/Classes/Factories/SubsurfaceProfileFactory.h
@@ -6,7 +6,9 @@
 UCLASS(hidecategories=Object)
 class USubsurfaceProfileFactory : public UFactory
 {
-	GENERATED_UCLASS_BODY()
+	GENERATED_BODY()
+public:
+	USubsurfaceProfileFactory(const FObjectInitializer& ObjectInitializer = FObjectInitializer::Get());
 
 
 	// Begin UFactory Interface
diff --git a/Engine/Source/Editor/UnrealEd/Classes/Factories/Texture2dFactoryNew.h b/Engine/Source/Editor/UnrealEd/Classes/Factories/Texture2dFactoryNew.h
--- a/Engine/Source/Editor/UnrealEd/Classes/Factories/Texture2dFactoryNew.h
+++ b/Engine/Source/Editor/UnrealEd/Classes/Factories/Texture2dFactoryNew.h
@@ -10,7 +10,9 @@
 UCLASS(hidecategories=Object, MinimalAPI)
 class UTexture2DFactoryNew : public UFactory
 {
-	GENERATED_UCLASS_BODY()
+	GENERATED_BODY()
+public:
+	UNREALED_API UTexture2DFactoryNew(const FObjectInitializer& ObjectInitializer = FObjectInitializer::Get());
 
 	/** width of new texture */
 	UPROPERTY()
diff --git a/Engine/Source/Editor/UnrealEd/Classes/Factories/TextureCubeThumbnailRenderer.h b/Engine/Source/Editor/UnrealEd/Classes/Factories/TextureCubeThumbnailRenderer.h
--- a/Engine/Source/Editor/UnrealEd/Classes/Factories/TextureCubeThumbnailRenderer.h
+++ b/Engine/Source/Editor/UnrealEd/Classes/Factories/TextureCubeThumbnailRenderer.h
@@ -11,7 +11,9 @@
 UCLASS()
 class UTextureCubeThumbnailRenderer : public UTextureThumbnailRenderer
 {
-	GENERATED_UCLASS_BODY()
+	GENERATED_BODY()
+public:
+	UTextureCubeThumbnailRenderer(const FObjectInitializer& ObjectInitializer = FObjectInitializer::Get());
 
 
 	// Begin UThumbnailRenderer Interface
diff --git a/Engine/Source/Editor/UnrealEd/Classes/Factories/TextureFactory.h b/Engine/Source/Editor/UnrealEd/Classes/Factories/TextureFactory.h
--- a/Engine/Source/Editor/UnrealEd/Classes/Factories/TextureFactory.h
+++ b/Engine/Source/Editor/UnrealEd/Classes/Factories/TextureFactory.h
@@ -5,10 +5,11 @@
 #include "Factories/Factory.h"
 #include "TextureFactory.generated.h"
 
-UCLASS(customconstructor, collapsecategories, hidecategories=Object)
+UCLASS(collapsecategories, hidecategories=Object)
 class UNREALED_API UTextureFactory : public UFactory
 {
-	GENERATED_UCLASS_BODY()
+	GENERATED_BODY()
+public:
 
 	UPROPERTY()
 	uint32 NoCompression:1;
diff --git a/Engine/Source/Editor/UnrealEd/Classes/Factories/TextureRenderTargetCubeFactoryNew.h b/Engine/Source/Editor/UnrealEd/Classes/Factories/TextureRenderTargetCubeFactoryNew.h
--- a/Engine/Source/Editor/UnrealEd/Classes/Factories/TextureRenderTargetCubeFactoryNew.h
+++ b/Engine/Source/Editor/UnrealEd/Classes/Factories/TextureRenderTargetCubeFactoryNew.h
@@ -10,7 +10,9 @@
 UCLASS(hidecategories=(Object, Texture))
 class UTextureRenderTargetCubeFactoryNew : public UFactory
 {
-	GENERATED_UCLASS_BODY()
+	GENERATED_BODY()
+public:
+	UTextureRenderTargetCubeFactoryNew(const FObjectInitializer& ObjectInitializer = FObjectInitializer::Get());
 
 	/** width of new texture */
 	UPROPERTY(meta=(ToolTip="Width of the texture render target"))
diff --git a/Engine/Source/Editor/UnrealEd/Classes/Factories/TextureRenderTargetFactoryNew.h b/Engine/Source/Editor/UnrealEd/Classes/Factories/TextureRenderTargetFactoryNew.h
--- a/Engine/Source/Editor/UnrealEd/Classes/Factories/TextureRenderTargetFactoryNew.h
+++ b/Engine/Source/Editor/UnrealEd/Classes/Factories/TextureRenderTargetFactoryNew.h
@@ -10,7 +10,9 @@
 UCLASS(hidecategories=(Object, Texture))
 class UTextureRenderTargetFactoryNew : public UFactory
 {
-	GENERATED_UCLASS_BODY()
+	GENERATED_BODY()
+public:
+	UTextureRenderTargetFactoryNew(const FObjectInitializer& ObjectInitializer = FObjectInitializer::Get());
 
 	/** Width of the texture render target */
 	UPROPERTY(meta=(ToolTip="Width of the texture render target"))
diff --git a/Engine/Source/Editor/UnrealEd/Classes/Factories/TouchInterfaceFactory.h b/Engine/Source/Editor/UnrealEd/Classes/Factories/TouchInterfaceFactory.h
--- a/Engine/Source/Editor/UnrealEd/Classes/Factories/TouchInterfaceFactory.h
+++ b/Engine/Source/Editor/UnrealEd/Classes/Factories/TouchInterfaceFactory.h
@@ -6,7 +6,9 @@
 UCLASS(hidecategories=Object)
 class UTouchInterfaceFactory : public UFactory
 {
-	GENERATED_UCLASS_BODY()
+	GENERATED_BODY()
+public:
+	UTouchInterfaceFactory(const FObjectInitializer& ObjectInitializer = FObjectInitializer::Get());
 
 
 	// Begin UFactory Interface
diff --git a/Engine/Source/Editor/UnrealEd/Classes/Factories/TrueTypeFontFactory.h b/Engine/Source/Editor/UnrealEd/Classes/Factories/TrueTypeFontFactory.h
--- a/Engine/Source/Editor/UnrealEd/Classes/Factories/TrueTypeFontFactory.h
+++ b/Engine/Source/Editor/UnrealEd/Classes/Factories/TrueTypeFontFactory.h
@@ -10,7 +10,9 @@
 UCLASS(hidecategories=Object, collapsecategories)
 class UNREALED_API UTrueTypeFontFactory : public UTextureFactory, public FReimportHandler
 {
-	GENERATED_UCLASS_BODY()
+	GENERATED_BODY()
+public:
+	UTrueTypeFontFactory(const FObjectInitializer& ObjectInitializer = FObjectInitializer::Get());
 
 	/** Import options for this font */
 	UPROPERTY(EditAnywhere, Instanced, Category = TrueTypeFontFactory, meta = (ToolTip = "Import options for the font"))
diff --git a/Engine/Source/Editor/UnrealEd/Classes/Factories/VectorFieldStaticFactory.h b/Engine/Source/Editor/UnrealEd/Classes/Factories/VectorFieldStaticFactory.h
--- a/Engine/Source/Editor/UnrealEd/Classes/Factories/VectorFieldStaticFactory.h
+++ b/Engine/Source/Editor/UnrealEd/Classes/Factories/VectorFieldStaticFactory.h
@@ -11,7 +11,9 @@
 UCLASS()
 class UVectorFieldStaticFactory : public UFactory
 {
-	GENERATED_UCLASS_BODY()
+	GENERATED_BODY()
+public:
+	UVectorFieldStaticFactory(const FObjectInitializer& ObjectInitializer = FObjectInitializer::Get());
 
 
 	// Begin UFactory Interface
diff --git a/Engine/Source/Editor/UnrealEd/Classes/Factories/WorldFactory.h b/Engine/Source/Editor/UnrealEd/Classes/Factories/WorldFactory.h
--- a/Engine/Source/Editor/UnrealEd/Classes/Factories/WorldFactory.h
+++ b/Engine/Source/Editor/UnrealEd/Classes/Factories/WorldFactory.h
@@ -11,7 +11,9 @@
 UCLASS(MinimalAPI)
 class UWorldFactory : public UFactory
 {
-	GENERATED_UCLASS_BODY()
+	GENERATED_BODY()
+public:
+	UNREALED_API UWorldFactory(const FObjectInitializer& ObjectInitializer = FObjectInitializer::Get());
 
 	TEnumAsByte<EWorldType::Type> WorldType;
 	bool bInformEngineOfWorld;
diff --git a/Engine/Source/Editor/UnrealEd/Classes/MaterialEditor/DEditorFontParameterValue.h b/Engine/Source/Editor/UnrealEd/Classes/MaterialEditor/DEditorFontParameterValue.h
--- a/Engine/Source/Editor/UnrealEd/Classes/MaterialEditor/DEditorFontParameterValue.h
+++ b/Engine/Source/Editor/UnrealEd/Classes/MaterialEditor/DEditorFontParameterValue.h
@@ -10,7 +10,7 @@
 USTRUCT()
 struct UNREALED_API FDFontParameters
 {
-	GENERATED_USTRUCT_BODY()
+	GENERATED_BODY()
 
 	UPROPERTY(EditAnywhere, Category = DFontParameter)
 	class UFont* FontValue;
@@ -22,7 +22,9 @@ struct UNREALED_API FDFontParameters
 UCLASS(hidecategories = Object, collapsecategories)
 class UNREALED_API UDEditorFontParameterValue : public UDEditorParameterValue
 {
-	GENERATED_UCLASS_BODY()
+	GENERATED_BODY()
+public:
+	UDEditorFontParameterValue(const FObjectInitializer& ObjectInitializer = FObjectInitializer::Get());
 
 	UPROPERTY(EditAnywhere, Category = DEditorFontParameterValue)
 	struct FDFontParameters ParameterValue;
diff --git a/Engine/Source/Editor/UnrealEd/Classes/MaterialEditor/DEditorParameterValue.h b/Engine/Source/Editor/UnrealEd/Classes/MaterialEditor/DEditorParameterValue.h
--- a/Engine/Source/Editor/UnrealEd/Classes/MaterialEditor/DEditorParameterValue.h
+++ b/Engine/Source/Editor/UnrealEd/Classes/MaterialEditor/DEditorParameterValue.h
@@ -9,7 +9,9 @@
 UCLASS(hidecategories=Object, collapsecategories, editinlinenew)
 class UNREALED_API UDEditorParameterValue : public UObject
 {
-	GENERATED_UCLASS_BODY()
+	GENERATED_BODY()
+public:
+	UDEditorParameterValue(const FObjectInitializer& ObjectInitializer = FObjectInitializer::Get());
 
 	UPROPERTY(EditAnywhere, Category=DEditorParameterValue)
 	uint32 bOverride:1;
diff --git a/Engine/Source/Editor/UnrealEd/Classes/MaterialEditor/DEditorScalarParameterValue.h b/Engine/Source/Editor/UnrealEd/Classes/MaterialEditor/DEditorScalarParameterValue.h
--- a/Engine/Source/Editor/UnrealEd/Classes/MaterialEditor/DEditorScalarParameterValue.h
+++ b/Engine/Source/Editor/UnrealEd/Classes/MaterialEditor/DEditorScalarParameterValue.h
@@ -9,7 +9,9 @@
 UCLASS(hidecategories=Object, collapsecategories, editinlinenew)
 class UNREALED_API UDEditorScalarParameterValue : public UDEditorParameterValue
 {
-	GENERATED_UCLASS_BODY()
+	GENERATED_BODY()
+public:
+	UDEditorScalarParameterValue(const FObjectInitializer& ObjectInitializer = FObjectInitializer::Get());
 
 	UPROPERTY(EditAnywhere, Category=DEditorScalarParameterValue)
 	float ParameterValue;
diff --git a/Engine/Source/Editor/UnrealEd/Classes/MaterialEditor/DEditorStaticComponentMaskParameterValue.h b/Engine/Source/Editor/UnrealEd/Classes/MaterialEditor/DEditorStaticComponentMaskParameterValue.h
--- a/Engine/Source/Editor/UnrealEd/Classes/MaterialEditor/DEditorStaticComponentMaskParameterValue.h
+++ b/Engine/Source/Editor/UnrealEd/Classes/MaterialEditor/DEditorStaticComponentMaskParameterValue.h
@@ -9,7 +9,7 @@
 USTRUCT()
 struct UNREALED_API FDComponentMaskParameter
 {
-	GENERATED_USTRUCT_BODY()
+	GENERATED_BODY()
 
 	UPROPERTY(EditAnywhere, Category=DComponentMaskParameter)
 	uint32 R:1;
@@ -46,7 +46,9 @@ struct UNREALED_API FDComponentMaskParameter
 UCLASS(hidecategories=Object, collapsecategories)
 class UNREALED_API UDEditorStaticComponentMaskParameterValue : public UDEditorParameterValue
 {
-	GENERATED_UCLASS_BODY()
+	GENERATED_BODY()
+public:
+	UDEditorStaticComponentMaskParameterValue(const FObjectInitializer& ObjectInitializer = FObjectInitializer::Get());
 
 	UPROPERTY(EditAnywhere, Category=DEditorStaticComponentMaskParameterValue)
 	struct FDComponentMaskParameter ParameterValue;
diff --git a/Engine/Source/Editor/UnrealEd/Classes/MaterialEditor/DEditorStaticSwitchParameterValue.h b/Engine/Source/Editor/UnrealEd/Classes/MaterialEditor/DEditorStaticSwitchParameterValue.h
--- a/Engine/Source/Editor/UnrealEd/Classes/MaterialEditor/DEditorStaticSwitchParameterValue.h
+++ b/Engine/Source/Editor/UnrealEd/Classes/MaterialEditor/DEditorStaticSwitchParameterValue.h
@@ -9,7 +9,9 @@
 UCLASS(hidecategories=Object, collapsecategories)
 class UNREALED_API UDEditorStaticSwitchParameterValue : public UDEditorParameterValue
 {
-	GENERATED_UCLASS_BODY()
+	GENERATED_BODY()
+public:
+	UDEditorStaticSwitchParameterValue(const FObjectInitializer& ObjectInitializer = FObjectInitializer::Get());
 
 	UPROPERTY(EditAnywhere, Category=DEditorStaticSwitchParameterValue)
 	uint32 ParameterValue:1;
diff --git a/Engine/Source/Editor/UnrealEd/Classes/MaterialEditor/DEditorTextureParameterValue.h b/Engine/Source/Editor/UnrealEd/Classes/MaterialEditor/DEditorTextureParameterValue.h
--- a/Engine/Source/Editor/UnrealEd/Classes/MaterialEditor/DEditorTextureParameterValue.h
+++ b/Engine/Source/Editor/UnrealEd/Classes/MaterialEditor/DEditorTextureParameterValue.h
@@ -9,7 +9,9 @@
 UCLASS(hidecategories=Object, collapsecategories)
 class UNREALED_API UDEditorTextureParameterValue : public UDEditorParameterValue
 {
-	GENERATED_UCLASS_BODY()
+	GENERATED_BODY()
+public:
+	UDEditorTextureParameterValue(const FObjectInitializer& ObjectInitializer = FObjectInitializer::Get());
 
 	UPROPERTY(EditAnywhere, Category=DEditorTextureParameterValue)
 	class UTexture* ParameterValue;
diff --git a/Engine/Source/Editor/UnrealEd/Classes/MaterialEditor/DEditorVectorParameterValue.h b/Engine/Source/Editor/UnrealEd/Classes/MaterialEditor/DEditorVectorParameterValue.h
--- a/Engine/Source/Editor/UnrealEd/Classes/MaterialEditor/DEditorVectorParameterValue.h
+++ b/Engine/Source/Editor/UnrealEd/Classes/MaterialEditor/DEditorVectorParameterValue.h
@@ -9,7 +9,9 @@
 UCLASS(hidecategories=Object, collapsecategories, editinlinenew)
 class UNREALED_API UDEditorVectorParameterValue : public UDEditorParameterValue
 {
-	GENERATED_UCLASS_BODY()
+	GENERATED_BODY()
+public:
+	UDEditorVectorParameterValue(const FObjectInitializer& ObjectInitializer = FObjectInitializer::Get());
 
 	UPROPERTY(EditAnywhere, Category=DEditorVectorParameterValue)
 	FLinearColor ParameterValue;
diff --git a/Engine/Source/Editor/UnrealEd/Classes/MaterialEditor/MaterialEditorInstanceConstant.h b/Engine/Source/Editor/UnrealEd/Classes/MaterialEditor/MaterialEditorInstanceConstant.h
--- a/Engine/Source/Editor/UnrealEd/Classes/MaterialEditor/MaterialEditorInstanceConstant.h
+++ b/Engine/Source/Editor/UnrealEd/Classes/MaterialEditor/MaterialEditorInstanceConstant.h
@@ -13,7 +13,7 @@
 USTRUCT()
 struct FEditorParameterGroup
 {
-	GENERATED_USTRUCT_BODY()
+	GENERATED_BODY()
 
 	UPROPERTY()
 	FName GroupName;
@@ -26,7 +26,7 @@ struct FEditorParameterGroup
 USTRUCT()
 struct FEditorParameterValue
 {
-	GENERATED_USTRUCT_BODY()
+	GENERATED_BODY()
 
 	UPROPERTY(EditAnywhere, Category=EditorParameterValue)
 	uint32 bOverride:1;
@@ -48,7 +48,7 @@ struct FEditorParameterValue
 USTRUCT()
 struct FEditorVectorParameterValue : public FEditorParameterValue
 {
-	GENERATED_USTRUCT_BODY()
+	GENERATED_BODY()
 
 	UPROPERTY(EditAnywhere, Category=EditorVectorParameterValue)
 	FLinearColor ParameterValue;
@@ -64,7 +64,7 @@ struct FEditorVectorParameterValue : public FEditorParameterValue
 USTRUCT()
 struct FEditorScalarParameterValue : public FEditorParameterValue
 {
-	GENERATED_USTRUCT_BODY()
+	GENERATED_BODY()
 
 	UPROPERTY(EditAnywhere, Category=EditorScalarParameterValue)
 	float ParameterValue;
@@ -80,7 +80,7 @@ struct FEditorScalarParameterValue : public FEditorParameterValue
 USTRUCT()
 struct FEditorTextureParameterValue : public FEditorParameterValue
 {
-	GENERATED_USTRUCT_BODY()
+	GENERATED_BODY()
 
 	UPROPERTY(EditAnywhere, Category=EditorTextureParameterValue)
 	class UTexture* ParameterValue;
@@ -96,7 +96,7 @@ struct FEditorTextureParameterValue : public FEditorParameterValue
 USTRUCT()
 struct FEditorFontParameterValue : public FEditorParameterValue
 {
-	GENERATED_USTRUCT_BODY()
+	GENERATED_BODY()
 
 	UPROPERTY(EditAnywhere, Category=EditorFontParameterValue)
 	class UFont* FontValue;
@@ -116,7 +116,7 @@ struct FEditorFontParameterValue : public FEditorParameterValue
 USTRUCT()
 struct FEditorStaticSwitchParameterValue : public FEditorParameterValue
 {
-	GENERATED_USTRUCT_BODY()
+	GENERATED_BODY()
 
 	UPROPERTY(EditAnywhere, Category=EditorStaticSwitchParameterValue)
 	uint32 ParameterValue:1;
@@ -142,7 +142,7 @@ struct FEditorStaticSwitchParameterValue : public FEditorParameterValue
 USTRUCT()
 struct FComponentMaskParameter
 {
-	GENERATED_USTRUCT_BODY()
+	GENERATED_BODY()
 
 	UPROPERTY(EditAnywhere, Category=ComponentMaskParameter)
 	uint32 R:1;
@@ -180,7 +180,7 @@ struct FComponentMaskParameter
 USTRUCT()
 struct FEditorStaticComponentMaskParameterValue : public FEditorParameterValue
 {
-	GENERATED_USTRUCT_BODY()
+	GENERATED_BODY()
 
 	UPROPERTY(EditAnywhere, Category=EditorStaticComponentMaskParameterValue)
 	struct FComponentMaskParameter ParameterValue;
@@ -202,7 +202,9 @@ struct FEditorStaticComponentMaskParameterValue : public FEditorParameterValue
 UCLASS(hidecategories=Object, collapsecategories)
 class UNREALED_API UMaterialEditorInstanceConstant : public UObject
 {
-	GENERATED_UCLASS_BODY()
+	GENERATED_BODY()
+public:
+	UMaterialEditorInstanceConstant(const FObjectInitializer& ObjectInitializer = FObjectInitializer::Get());
 
 	/** Physical material to use for this graphics material. Used for sounds, effects etc.*/
 	UPROPERTY(EditAnywhere, Category=MaterialEditorInstanceConstant)
diff --git a/Engine/Source/Editor/UnrealEd/Classes/MaterialEditor/MaterialEditorMeshComponent.h b/Engine/Source/Editor/UnrealEd/Classes/MaterialEditor/MaterialEditorMeshComponent.h
--- a/Engine/Source/Editor/UnrealEd/Classes/MaterialEditor/MaterialEditorMeshComponent.h
+++ b/Engine/Source/Editor/UnrealEd/Classes/MaterialEditor/MaterialEditorMeshComponent.h
@@ -8,7 +8,9 @@
 UCLASS()
 class UNREALED_API UMaterialEditorMeshComponent : public UStaticMeshComponent
 {
-	GENERATED_UCLASS_BODY()
+	GENERATED_BODY()
+public:
+	UMaterialEditorMeshComponent(const FObjectInitializer& ObjectInitializer = FObjectInitializer::Get());
 
 	// USceneComponent Interface
 	virtual FBoxSphereBounds CalcBounds(const FTransform& LocalToWorld) const override;
diff --git a/Engine/Source/Editor/UnrealEd/Classes/MaterialEditor/PreviewMaterial.h b/Engine/Source/Editor/UnrealEd/Classes/MaterialEditor/PreviewMaterial.h
--- a/Engine/Source/Editor/UnrealEd/Classes/MaterialEditor/PreviewMaterial.h
+++ b/Engine/Source/Editor/UnrealEd/Classes/MaterialEditor/PreviewMaterial.h
@@ -9,7 +9,9 @@
 UCLASS()
 class UNREALED_API UPreviewMaterial : public UMaterial
 {
-	GENERATED_UCLASS_BODY()
+	GENERATED_BODY()
+public:
+	UPreviewMaterial(const FObjectInitializer& ObjectInitializer = FObjectInitializer::Get());
 
 
 	// Begin UMaterial interface.
diff --git a/Engine/Source/Editor/UnrealEd/Classes/MaterialGraph/MaterialGraph.h b/Engine/Source/Editor/UnrealEd/Classes/MaterialGraph/MaterialGraph.h
--- a/Engine/Source/Editor/UnrealEd/Classes/MaterialGraph/MaterialGraph.h
+++ b/Engine/Source/Editor/UnrealEd/Classes/MaterialGraph/MaterialGraph.h
@@ -69,7 +69,9 @@ struct FMaterialInputInfo
 UCLASS(MinimalAPI)
 class UMaterialGraph : public UEdGraph
 {
-	GENERATED_UCLASS_BODY()
+	GENERATED_BODY()
+public:
+	UNREALED_API UMaterialGraph(const FObjectInitializer& ObjectInitializer = FObjectInitializer::Get());
 
 	/** Material this Graph represents */
 	UPROPERTY()
diff --git a/Engine/Source/Editor/UnrealEd/Classes/MaterialGraph/MaterialGraphNode.h b/Engine/Source/Editor/UnrealEd/Classes/MaterialGraph/MaterialGraphNode.h
--- a/Engine/Source/Editor/UnrealEd/Classes/MaterialGraph/MaterialGraphNode.h
+++ b/Engine/Source/Editor/UnrealEd/Classes/MaterialGraph/MaterialGraphNode.h
@@ -7,7 +7,9 @@
 UCLASS(MinimalAPI)
 class UMaterialGraphNode : public UMaterialGraphNode_Base
 {
-	GENERATED_UCLASS_BODY()
+	GENERATED_BODY()
+public:
+	UNREALED_API UMaterialGraphNode(const FObjectInitializer& ObjectInitializer = FObjectInitializer::Get());
 
 	/** Material Expression this node is representing */
 	UPROPERTY()
diff --git a/Engine/Source/Editor/UnrealEd/Classes/MaterialGraph/MaterialGraphNode_Base.h b/Engine/Source/Editor/UnrealEd/Classes/MaterialGraph/MaterialGraphNode_Base.h
--- a/Engine/Source/Editor/UnrealEd/Classes/MaterialGraph/MaterialGraphNode_Base.h
+++ b/Engine/Source/Editor/UnrealEd/Classes/MaterialGraph/MaterialGraphNode_Base.h
@@ -7,7 +7,9 @@
 UCLASS(MinimalAPI)
 class UMaterialGraphNode_Base : public UEdGraphNode
 {
-	GENERATED_UCLASS_BODY()
+	GENERATED_BODY()
+public:
+	UNREALED_API UMaterialGraphNode_Base(const FObjectInitializer& ObjectInitializer = FObjectInitializer::Get());
 
 	/** Create all of the input pins required */
 	virtual void CreateInputPins() {};
diff --git a/Engine/Source/Editor/UnrealEd/Classes/MaterialGraph/MaterialGraphNode_Comment.h b/Engine/Source/Editor/UnrealEd/Classes/MaterialGraph/MaterialGraphNode_Comment.h
--- a/Engine/Source/Editor/UnrealEd/Classes/MaterialGraph/MaterialGraphNode_Comment.h
+++ b/Engine/Source/Editor/UnrealEd/Classes/MaterialGraph/MaterialGraphNode_Comment.h
@@ -9,7 +9,9 @@
 UCLASS(MinimalAPI)
 class UMaterialGraphNode_Comment : public UEdGraphNode_Comment
 {
-	GENERATED_UCLASS_BODY()
+	GENERATED_BODY()
+public:
+	UNREALED_API UMaterialGraphNode_Comment(const FObjectInitializer& ObjectInitializer = FObjectInitializer::Get());
 
 	/** Material Comment that this node represents */
 	UPROPERTY()
diff --git a/Engine/Source/Editor/UnrealEd/Classes/MaterialGraph/MaterialGraphNode_Root.h b/Engine/Source/Editor/UnrealEd/Classes/MaterialGraph/MaterialGraphNode_Root.h
--- a/Engine/Source/Editor/UnrealEd/Classes/MaterialGraph/MaterialGraphNode_Root.h
+++ b/Engine/Source/Editor/UnrealEd/Classes/MaterialGraph/MaterialGraphNode_Root.h
@@ -7,7 +7,9 @@
 UCLASS(MinimalAPI)
 class UMaterialGraphNode_Root : public UMaterialGraphNode_Base
 {
-	GENERATED_UCLASS_BODY()
+	GENERATED_BODY()
+public:
+	UNREALED_API UMaterialGraphNode_Root(const FObjectInitializer& ObjectInitializer = FObjectInitializer::Get());
 
 	/** Material whose inputs this root node represents */
 	UPROPERTY()
diff --git a/Engine/Source/Editor/UnrealEd/Classes/MaterialGraph/MaterialGraphSchema.h b/Engine/Source/Editor/UnrealEd/Classes/MaterialGraph/MaterialGraphSchema.h
--- a/Engine/Source/Editor/UnrealEd/Classes/MaterialGraph/MaterialGraphSchema.h
+++ b/Engine/Source/Editor/UnrealEd/Classes/MaterialGraph/MaterialGraphSchema.h
@@ -8,7 +8,7 @@
 USTRUCT()
 struct UNREALED_API FMaterialGraphSchemaAction_NewNode : public FEdGraphSchemaAction
 {
-	GENERATED_USTRUCT_BODY();
+	GENERATED_BODY()
 
 	/** Class of expression we want to create */
 	UPROPERTY()
@@ -45,7 +45,7 @@ struct UNREALED_API FMaterialGraphSchemaAction_NewNode : public FEdGraphSchemaAc
 USTRUCT()
 struct UNREALED_API FMaterialGraphSchemaAction_NewFunctionCall : public FEdGraphSchemaAction
 {
-	GENERATED_USTRUCT_BODY();
+	GENERATED_BODY()
 
 	/** Path to function that we want to call */
 	UPROPERTY()
@@ -72,7 +72,7 @@ struct UNREALED_API FMaterialGraphSchemaAction_NewFunctionCall : public FEdGraph
 USTRUCT()
 struct UNREALED_API FMaterialGraphSchemaAction_NewComment : public FEdGraphSchemaAction
 {
-	GENERATED_USTRUCT_BODY();
+	GENERATED_BODY()
 
 	// Simple type info
 	static FName StaticGetTypeId() {static FName Type("FMaterialGraphSchemaAction_NewComment"); return Type;}
@@ -95,7 +95,7 @@ struct UNREALED_API FMaterialGraphSchemaAction_NewComment : public FEdGraphSchem
 USTRUCT()
 struct UNREALED_API FMaterialGraphSchemaAction_Paste : public FEdGraphSchemaAction
 {
-	GENERATED_USTRUCT_BODY();
+	GENERATED_BODY()
 
 	// Simple type info
 	static FName StaticGetTypeId() {static FName Type("FMaterialGraphSchemaAction_Paste"); return Type;}
@@ -117,7 +117,9 @@ struct UNREALED_API FMaterialGraphSchemaAction_Paste : public FEdGraphSchemaActi
 UCLASS(MinimalAPI)
 class UMaterialGraphSchema : public UEdGraphSchema
 {
-	GENERATED_UCLASS_BODY()
+	GENERATED_BODY()
+public:
+	UNREALED_API UMaterialGraphSchema(const FObjectInitializer& ObjectInitializer = FObjectInitializer::Get());
 
 	// Allowable PinType.PinCategory values
 	UPROPERTY()
diff --git a/Engine/Source/Editor/UnrealEd/Classes/Niagara/EdGraphSchema_Niagara.h b/Engine/Source/Editor/UnrealEd/Classes/Niagara/EdGraphSchema_Niagara.h
--- a/Engine/Source/Editor/UnrealEd/Classes/Niagara/EdGraphSchema_Niagara.h
+++ b/Engine/Source/Editor/UnrealEd/Classes/Niagara/EdGraphSchema_Niagara.h
@@ -9,7 +9,7 @@
 USTRUCT()
 struct UNREALED_API FNiagaraSchemaAction_NewNode : public FEdGraphSchemaAction
 {
-	GENERATED_USTRUCT_BODY();
+	GENERATED_BODY()
 
 	/** Template of node we want to create */
 	UPROPERTY()
@@ -45,7 +45,9 @@ struct UNREALED_API FNiagaraSchemaAction_NewNode : public FEdGraphSchemaAction
 UCLASS(MinimalAPI)
 class UEdGraphSchema_Niagara : public UEdGraphSchema
 {
-	GENERATED_UCLASS_BODY()
+	GENERATED_BODY()
+public:
+	UNREALED_API UEdGraphSchema_Niagara(const FObjectInitializer& ObjectInitializer = FObjectInitializer::Get());
 
 	// Allowable PinType.PinCategory values
 	UPROPERTY()
diff --git a/Engine/Source/Editor/UnrealEd/Classes/Niagara/NiagaraGraph.h b/Engine/Source/Editor/UnrealEd/Classes/Niagara/NiagaraGraph.h
--- a/Engine/Source/Editor/UnrealEd/Classes/Niagara/NiagaraGraph.h
+++ b/Engine/Source/Editor/UnrealEd/Classes/Niagara/NiagaraGraph.h
@@ -7,7 +7,9 @@
 UCLASS(MinimalAPI)
 class UNiagaraGraph : public UEdGraph
 {
-	GENERATED_UCLASS_BODY()
+	GENERATED_BODY()
+public:
+	UNREALED_API UNiagaraGraph(const FObjectInitializer& ObjectInitializer = FObjectInitializer::Get());
 	
 	/** Get the source that owns this graph */
 	class UNiagaraScriptSource* GetSource() const;
diff --git a/Engine/Source/Editor/UnrealEd/Classes/Niagara/NiagaraNode.h b/Engine/Source/Editor/UnrealEd/Classes/Niagara/NiagaraNode.h
--- a/Engine/Source/Editor/UnrealEd/Classes/Niagara/NiagaraNode.h
+++ b/Engine/Source/Editor/UnrealEd/Classes/Niagara/NiagaraNode.h
@@ -6,7 +6,9 @@
 UCLASS()
 class UNREALED_API UNiagaraNode : public UEdGraphNode
 {
-	GENERATED_UCLASS_BODY()
+	GENERATED_BODY()
+public:
+	UNiagaraNode(const FObjectInitializer& ObjectInitializer = FObjectInitializer::Get());
 
 public:
 
diff --git a/Engine/Source/Editor/UnrealEd/Classes/Niagara/NiagaraNodeConstant.h b/Engine/Source/Editor/UnrealEd/Classes/Niagara/NiagaraNodeConstant.h
--- a/Engine/Source/Editor/UnrealEd/Classes/Niagara/NiagaraNodeConstant.h
+++ b/Engine/Source/Editor/UnrealEd/Classes/Niagara/NiagaraNodeConstant.h
@@ -8,7 +8,9 @@
 UCLASS(MinimalAPI, Deprecated)
 class UDEPRECATED_NiagaraNodeConstant : public UNiagaraNode
 {
-	GENERATED_UCLASS_BODY()
+	GENERATED_BODY()
+public:
+	UNREALED_API UDEPRECATED_NiagaraNodeConstant(const FObjectInitializer& ObjectInitializer = FObjectInitializer::Get());
 
 public:
 	/** The type of the constant we're creating. */
diff --git a/Engine/Source/Editor/UnrealEd/Classes/Niagara/NiagaraNodeGetAttr.h b/Engine/Source/Editor/UnrealEd/Classes/Niagara/NiagaraNodeGetAttr.h
--- a/Engine/Source/Editor/UnrealEd/Classes/Niagara/NiagaraNodeGetAttr.h
+++ b/Engine/Source/Editor/UnrealEd/Classes/Niagara/NiagaraNodeGetAttr.h
@@ -6,7 +6,9 @@
 UCLASS(MinimalAPI, Deprecated)
 class UDEPRECATED_NiagaraNodeGetAttr : public UNiagaraNode
 {
-	GENERATED_UCLASS_BODY()
+	GENERATED_BODY()
+public:
+	UNREALED_API UDEPRECATED_NiagaraNodeGetAttr(const FObjectInitializer& ObjectInitializer = FObjectInitializer::Get());
 
 public:
 
