Back out changelist 2481333

diff --git a/Engine/Source/Editor/UnrealEd/Classes/ActorFactories/ActorFactoryTargetPoint.h b/Engine/Source/Editor/UnrealEd/Classes/ActorFactories/ActorFactoryTargetPoint.h
--- a/Engine/Source/Editor/UnrealEd/Classes/ActorFactories/ActorFactoryTargetPoint.h
+++ b/Engine/Source/Editor/UnrealEd/Classes/ActorFactories/ActorFactoryTargetPoint.h
@@ -6,8 +6,6 @@
 UCLASS( MinimalAPI, config = Editor, collapsecategories, hidecategories = Object )
 class UActorFactoryTargetPoint : public UActorFactory
 {
-	GENERATED_BODY()
-public:
-	UNREALED_API UActorFactoryTargetPoint(const FObjectInitializer& ObjectInitializer = FObjectInitializer::Get());
+	GENERATED_UCLASS_BODY()
 
 };
diff --git a/Engine/Source/Editor/UnrealEd/Classes/ActorFactories/ActorFactoryTextRender.h b/Engine/Source/Editor/UnrealEd/Classes/ActorFactories/ActorFactoryTextRender.h
--- a/Engine/Source/Editor/UnrealEd/Classes/ActorFactories/ActorFactoryTextRender.h
+++ b/Engine/Source/Editor/UnrealEd/Classes/ActorFactories/ActorFactoryTextRender.h
@@ -7,9 +7,7 @@
 UCLASS(MinimalAPI, config=Editor)
 class UActorFactoryTextRender : public UActorFactory
 {
-	GENERATED_BODY()
-public:
-	UNREALED_API UActorFactoryTextRender(const FObjectInitializer& ObjectInitializer = FObjectInitializer::Get());
+	GENERATED_UCLASS_BODY()
 };
 
 
diff --git a/Engine/Source/Editor/UnrealEd/Classes/ActorFactories/ActorFactoryTriggerBox.h b/Engine/Source/Editor/UnrealEd/Classes/ActorFactories/ActorFactoryTriggerBox.h
--- a/Engine/Source/Editor/UnrealEd/Classes/ActorFactories/ActorFactoryTriggerBox.h
+++ b/Engine/Source/Editor/UnrealEd/Classes/ActorFactories/ActorFactoryTriggerBox.h
@@ -7,9 +7,7 @@
 UCLASS(MinimalAPI, config=Editor, collapsecategories, hidecategories=Object)
 class UActorFactoryTriggerBox : public UActorFactory
 {
-	GENERATED_BODY()
-public:
-	UNREALED_API UActorFactoryTriggerBox(const FObjectInitializer& ObjectInitializer = FObjectInitializer::Get());
+	GENERATED_UCLASS_BODY()
 
 };
 
diff --git a/Engine/Source/Editor/UnrealEd/Classes/ActorFactories/ActorFactoryTriggerCapsule.h b/Engine/Source/Editor/UnrealEd/Classes/ActorFactories/ActorFactoryTriggerCapsule.h
--- a/Engine/Source/Editor/UnrealEd/Classes/ActorFactories/ActorFactoryTriggerCapsule.h
+++ b/Engine/Source/Editor/UnrealEd/Classes/ActorFactories/ActorFactoryTriggerCapsule.h
@@ -7,9 +7,7 @@
 UCLASS(MinimalAPI, config=Editor, collapsecategories, hidecategories=Object)
 class UActorFactoryTriggerCapsule : public UActorFactory
 {
-	GENERATED_BODY()
-public:
-	UNREALED_API UActorFactoryTriggerCapsule(const FObjectInitializer& ObjectInitializer = FObjectInitializer::Get());
+	GENERATED_UCLASS_BODY()
 
 };
 
diff --git a/Engine/Source/Editor/UnrealEd/Classes/ActorFactories/ActorFactoryTriggerSphere.h b/Engine/Source/Editor/UnrealEd/Classes/ActorFactories/ActorFactoryTriggerSphere.h
--- a/Engine/Source/Editor/UnrealEd/Classes/ActorFactories/ActorFactoryTriggerSphere.h
+++ b/Engine/Source/Editor/UnrealEd/Classes/ActorFactories/ActorFactoryTriggerSphere.h
@@ -7,9 +7,7 @@
 UCLASS(MinimalAPI, config=Editor, collapsecategories, hidecategories=Object)
 class UActorFactoryTriggerSphere : public UActorFactory
 {
-	GENERATED_BODY()
-public:
-	UNREALED_API UActorFactoryTriggerSphere(const FObjectInitializer& ObjectInitializer = FObjectInitializer::Get());
+	GENERATED_UCLASS_BODY()
 
 };
 
diff --git a/Engine/Source/Editor/UnrealEd/Classes/ActorFactories/ActorFactoryVectorFieldVolume.h b/Engine/Source/Editor/UnrealEd/Classes/ActorFactories/ActorFactoryVectorFieldVolume.h
--- a/Engine/Source/Editor/UnrealEd/Classes/ActorFactories/ActorFactoryVectorFieldVolume.h
+++ b/Engine/Source/Editor/UnrealEd/Classes/ActorFactories/ActorFactoryVectorFieldVolume.h
@@ -7,9 +7,7 @@
 UCLASS(MinimalAPI, config=Editor, collapsecategories, hidecategories=Object)
 class UActorFactoryVectorFieldVolume : public UActorFactory
 {
-	GENERATED_BODY()
-public:
-	UNREALED_API UActorFactoryVectorFieldVolume(const FObjectInitializer& ObjectInitializer = FObjectInitializer::Get());
+	GENERATED_UCLASS_BODY()
 
 	// Begin UActorFactory Interface
 	virtual void PostSpawnActor( UObject* Asset, AActor* NewActor ) override;
diff --git a/Engine/Source/Editor/UnrealEd/Classes/Animation/DebugSkelMeshComponent.h b/Engine/Source/Editor/UnrealEd/Classes/Animation/DebugSkelMeshComponent.h
--- a/Engine/Source/Editor/UnrealEd/Classes/Animation/DebugSkelMeshComponent.h
+++ b/Engine/Source/Editor/UnrealEd/Classes/Animation/DebugSkelMeshComponent.h
@@ -8,7 +8,7 @@
 USTRUCT()
 struct FSelectedSocketInfo
 {
-	GENERATED_BODY()
+	GENERATED_USTRUCT_BODY()
 
 	/** Default constructor */
 	FSelectedSocketInfo()
@@ -38,9 +38,7 @@ struct FSelectedSocketInfo
 UCLASS(transient, MinimalAPI)
 class UDebugSkelMeshComponent : public USkeletalMeshComponent
 {
-	GENERATED_BODY()
-public:
-	UNREALED_API UDebugSkelMeshComponent(const FObjectInitializer& ObjectInitializer = FObjectInitializer::Get());
+	GENERATED_UCLASS_BODY()
 
 	/** If true, render a wireframe skeleton of the mesh animated with the raw (uncompressed) animation data. */
 	UPROPERTY()
diff --git a/Engine/Source/Editor/UnrealEd/Classes/Animation/EditorAnimBaseObj.h b/Engine/Source/Editor/UnrealEd/Classes/Animation/EditorAnimBaseObj.h
--- a/Engine/Source/Editor/UnrealEd/Classes/Animation/EditorAnimBaseObj.h
+++ b/Engine/Source/Editor/UnrealEd/Classes/Animation/EditorAnimBaseObj.h
@@ -16,9 +16,7 @@ DECLARE_DELEGATE_TwoParams( FOnAnimObjectChange, class UObject*, bool)
 UCLASS(abstract, MinimalAPI)
 class UEditorAnimBaseObj: public UObject
 {
-	GENERATED_BODY()
-public:
-	UNREALED_API UEditorAnimBaseObj(const FObjectInitializer& ObjectInitializer = FObjectInitializer::Get());
+	GENERATED_UCLASS_BODY()
 public:
 	
 
diff --git a/Engine/Source/Editor/UnrealEd/Classes/Animation/EditorAnimCompositeSegment.h b/Engine/Source/Editor/UnrealEd/Classes/Animation/EditorAnimCompositeSegment.h
--- a/Engine/Source/Editor/UnrealEd/Classes/Animation/EditorAnimCompositeSegment.h
+++ b/Engine/Source/Editor/UnrealEd/Classes/Animation/EditorAnimCompositeSegment.h
@@ -16,9 +16,7 @@ DECLARE_DELEGATE_OneParam( FOnAnimSegmentChanged, class UEditorAnimSegment*)
 UCLASS(hidecategories=UObject, MinimalAPI, BlueprintType)
 class UEditorAnimCompositeSegment: public UEditorAnimBaseObj
 {
-	GENERATED_BODY()
-public:
-	UNREALED_API UEditorAnimCompositeSegment(const FObjectInitializer& ObjectInitializer = FObjectInitializer::Get());
+	GENERATED_UCLASS_BODY()
 public:
 
 	/** Default blend in time. */
diff --git a/Engine/Source/Editor/UnrealEd/Classes/Animation/EditorAnimSegment.h b/Engine/Source/Editor/UnrealEd/Classes/Animation/EditorAnimSegment.h
--- a/Engine/Source/Editor/UnrealEd/Classes/Animation/EditorAnimSegment.h
+++ b/Engine/Source/Editor/UnrealEd/Classes/Animation/EditorAnimSegment.h
@@ -16,9 +16,7 @@ DECLARE_DELEGATE_OneParam( FOnAnimSegmentChanged, class UEditorAnimSegment*)
 UCLASS(hidecategories=UObject, MinimalAPI, BlueprintType)
 class UEditorAnimSegment: public UEditorAnimBaseObj
 {
-	GENERATED_BODY()
-public:
-	UNREALED_API UEditorAnimSegment(const FObjectInitializer& ObjectInitializer = FObjectInitializer::Get());
+	GENERATED_UCLASS_BODY()
 public:
 
 	/** Default blend in time. */
diff --git a/Engine/Source/Editor/UnrealEd/Classes/Animation/EditorCompositeSection.h b/Engine/Source/Editor/UnrealEd/Classes/Animation/EditorCompositeSection.h
--- a/Engine/Source/Editor/UnrealEd/Classes/Animation/EditorCompositeSection.h
+++ b/Engine/Source/Editor/UnrealEd/Classes/Animation/EditorCompositeSection.h
@@ -13,9 +13,7 @@
 UCLASS(hidecategories=UObject, MinimalAPI, BlueprintType)
 class UEditorCompositeSection: public UEditorAnimBaseObj
 {
-	GENERATED_BODY()
-public:
-	UNREALED_API UEditorCompositeSection(const FObjectInitializer& ObjectInitializer = FObjectInitializer::Get());
+	GENERATED_UCLASS_BODY()
 public:
 
 	/** Default blend in time. */
diff --git a/Engine/Source/Editor/UnrealEd/Classes/Animation/EditorNotifyObject.h b/Engine/Source/Editor/UnrealEd/Classes/Animation/EditorNotifyObject.h
--- a/Engine/Source/Editor/UnrealEd/Classes/Animation/EditorNotifyObject.h
+++ b/Engine/Source/Editor/UnrealEd/Classes/Animation/EditorNotifyObject.h
@@ -14,9 +14,7 @@ class UEditorNotifyEntry;
 UCLASS(MinimalAPI)
 class UEditorNotifyObject : public UEditorAnimBaseObj
 {
-	GENERATED_BODY()
-public:
-	UNREALED_API UEditorNotifyObject(const FObjectInitializer& ObjectInitializer = FObjectInitializer::Get());
+	GENERATED_UCLASS_BODY()
 
 	/** Set up the editor object
 	 *	@param TrackIdx Index of the track the notify belongs to
@@ -36,4 +34,4 @@ public:
 	/** The notify event to modify */
 	UPROPERTY(EditAnywhere, Category=Event)
 	FAnimNotifyEvent Event;
-};
+};
\ No newline at end of file
diff --git a/Engine/Source/Editor/UnrealEd/Classes/Animation/EditorParentPlayerListObj.h b/Engine/Source/Editor/UnrealEd/Classes/Animation/EditorParentPlayerListObj.h
--- a/Engine/Source/Editor/UnrealEd/Classes/Animation/EditorParentPlayerListObj.h
+++ b/Engine/Source/Editor/UnrealEd/Classes/Animation/EditorParentPlayerListObj.h
@@ -11,9 +11,7 @@ class UAnimBlueprint;
 UCLASS()
 class UNREALED_API UEditorParentPlayerListObj : public UObject
 {
-	GENERATED_BODY()
-public:
-	UEditorParentPlayerListObj(const FObjectInitializer& ObjectInitializer = FObjectInitializer::Get());
+	GENERATED_UCLASS_BODY()
 public:
 	// Build a possible list of overrides from a blueprint and pull in any existing overrides.
 	void InitialiseFromBlueprint(UAnimBlueprint* Blueprint);
diff --git a/Engine/Source/Editor/UnrealEd/Classes/Animation/EditorSkeletonNotifyObj.h b/Engine/Source/Editor/UnrealEd/Classes/Animation/EditorSkeletonNotifyObj.h
--- a/Engine/Source/Editor/UnrealEd/Classes/Animation/EditorSkeletonNotifyObj.h
+++ b/Engine/Source/Editor/UnrealEd/Classes/Animation/EditorSkeletonNotifyObj.h
@@ -13,7 +13,7 @@ DECLARE_DELEGATE_TwoParams( FOnAnimObjectChange, class UObject*, bool)
 USTRUCT()
 struct FSkeletonNotifyDependentAnimations
 {
-	GENERATED_BODY()
+	GENERATED_USTRUCT_BODY()
 
 	UPROPERTY(EditAnywhere, BlueprintReadWrite, Category=SkeletonNotifies)
 	TArray<UAnimSequenceBase*> Animations;
@@ -26,9 +26,7 @@ struct FSkeletonNotifyDependentAnimations
 UCLASS(MinimalAPI)
 class UEditorSkeletonNotifyObj : public UObject
 {
-	GENERATED_BODY()
-public:
-	UNREALED_API UEditorSkeletonNotifyObj(const FObjectInitializer& ObjectInitializer = FObjectInitializer::Get());
+	GENERATED_UCLASS_BODY()
 public:
 	
 	UObject* EditedObject;
diff --git a/Engine/Source/Editor/UnrealEd/Classes/Builders/ConeBuilder.h b/Engine/Source/Editor/UnrealEd/Classes/Builders/ConeBuilder.h
--- a/Engine/Source/Editor/UnrealEd/Classes/Builders/ConeBuilder.h
+++ b/Engine/Source/Editor/UnrealEd/Classes/Builders/ConeBuilder.h
@@ -10,9 +10,11 @@
 UCLASS(MinimalAPI, autoexpandcategories=BrushSettings, EditInlineNew, meta=(DisplayName="Cone"))
 class UConeBuilder : public UEditorBrushBuilder
 {
+public:
 	GENERATED_BODY()
+
 public:
-	UNREALED_API UConeBuilder(const FObjectInitializer& ObjectInitializer = FObjectInitializer::Get());
+	UConeBuilder(const FObjectInitializer& ObjectInitializer = FObjectInitializer::Get());
 
 	/** Distance from base to tip of cone */
 	UPROPERTY(EditAnywhere, Category=BrushSettings, meta=(ClampMin = "0.000001"))
diff --git a/Engine/Source/Editor/UnrealEd/Classes/Builders/CubeBuilder.h b/Engine/Source/Editor/UnrealEd/Classes/Builders/CubeBuilder.h
--- a/Engine/Source/Editor/UnrealEd/Classes/Builders/CubeBuilder.h
+++ b/Engine/Source/Editor/UnrealEd/Classes/Builders/CubeBuilder.h
@@ -10,9 +10,11 @@
 UCLASS(MinimalAPI, autoexpandcategories=BrushSettings, EditInlineNew, meta=(DisplayName="Box"))
 class UCubeBuilder : public UEditorBrushBuilder
 {
+public:
 	GENERATED_BODY()
+
 public:
-	UNREALED_API UCubeBuilder(const FObjectInitializer& ObjectInitializer = FObjectInitializer::Get());
+	UCubeBuilder(const FObjectInitializer& ObjectInitializer = FObjectInitializer::Get());
 
 	/** The size of the cube in the X dimension */
 	UPROPERTY(EditAnywhere, Category=BrushSettings, meta=(ClampMin = "0.000001"))
diff --git a/Engine/Source/Editor/UnrealEd/Classes/Builders/CurvedStairBuilder.h b/Engine/Source/Editor/UnrealEd/Classes/Builders/CurvedStairBuilder.h
--- a/Engine/Source/Editor/UnrealEd/Classes/Builders/CurvedStairBuilder.h
+++ b/Engine/Source/Editor/UnrealEd/Classes/Builders/CurvedStairBuilder.h
@@ -10,9 +10,11 @@
 UCLASS(MinimalAPI, autoexpandcategories=BrushSettings, EditInlineNew,meta=(DisplayName="Curved Stair"))
 class UCurvedStairBuilder : public UEditorBrushBuilder
 {
+public:
 	GENERATED_BODY()
+
 public:
-	UNREALED_API UCurvedStairBuilder(const FObjectInitializer& ObjectInitializer = FObjectInitializer::Get());
+	UCurvedStairBuilder(const FObjectInitializer& ObjectInitializer = FObjectInitializer::Get());
 
 	/** The radius of the inner curve of the stair */
 	UPROPERTY(EditAnywhere, Category=BrushSettings, meta=(ClampMin = "1"))
diff --git a/Engine/Source/Editor/UnrealEd/Classes/Builders/CylinderBuilder.h b/Engine/Source/Editor/UnrealEd/Classes/Builders/CylinderBuilder.h
--- a/Engine/Source/Editor/UnrealEd/Classes/Builders/CylinderBuilder.h
+++ b/Engine/Source/Editor/UnrealEd/Classes/Builders/CylinderBuilder.h
@@ -15,9 +15,11 @@
 UCLASS(MinimalAPI, autoexpandcategories=BrushSettings, EditInlineNew, meta=(DisplayName="Cylinder"))
 class UCylinderBuilder : public UEditorBrushBuilder
 {
+public:
 	GENERATED_BODY()
+
 public:
-	UNREALED_API UCylinderBuilder(const FObjectInitializer& ObjectInitializer = FObjectInitializer::Get());
+	UCylinderBuilder(const FObjectInitializer& ObjectInitializer = FObjectInitializer::Get());
 
 	/** Distance from base to tip of cylinder */
 	UPROPERTY(EditAnywhere, Category=BrushSettings, meta=(ClampMin = "0.000001"))
diff --git a/Engine/Source/Editor/UnrealEd/Classes/Builders/EditorBrushBuilder.h b/Engine/Source/Editor/UnrealEd/Classes/Builders/EditorBrushBuilder.h
--- a/Engine/Source/Editor/UnrealEd/Classes/Builders/EditorBrushBuilder.h
+++ b/Engine/Source/Editor/UnrealEd/Classes/Builders/EditorBrushBuilder.h
@@ -10,7 +10,8 @@ class UEditorBrushBuilder : public UBrushBuilder
 public:
 	GENERATED_BODY()
 public:
-	UNREALED_API UEditorBrushBuilder(const FObjectInitializer& ObjectInitializer = FObjectInitializer::Get());
+
+	UEditorBrushBuilder(const FObjectInitializer& ObjectInitializer = FObjectInitializer::Get());
 
 	/** UBrushBuilder interface */
 	virtual void BeginBrush( bool InMergeCoplanars, FName InLayer ) override;
diff --git a/Engine/Source/Editor/UnrealEd/Classes/Builders/LinearStairBuilder.h b/Engine/Source/Editor/UnrealEd/Classes/Builders/LinearStairBuilder.h
--- a/Engine/Source/Editor/UnrealEd/Classes/Builders/LinearStairBuilder.h
+++ b/Engine/Source/Editor/UnrealEd/Classes/Builders/LinearStairBuilder.h
@@ -10,9 +10,11 @@
 UCLASS(MinimalAPI, autoexpandcategories=BrushSettings, EditInlineNew, meta=(DisplayName="Linear Stair"))
 class ULinearStairBuilder : public UEditorBrushBuilder
 {
+public:
 	GENERATED_BODY()
+
 public:
-	UNREALED_API ULinearStairBuilder(const FObjectInitializer& ObjectInitializer = FObjectInitializer::Get());
+	ULinearStairBuilder(const FObjectInitializer& ObjectInitializer = FObjectInitializer::Get());
 
 	/** The length of each step */
 	UPROPERTY(EditAnywhere, Category=BrushSettings, meta=(ClampMin = "1"))
diff --git a/Engine/Source/Editor/UnrealEd/Classes/Builders/SpiralStairBuilder.h b/Engine/Source/Editor/UnrealEd/Classes/Builders/SpiralStairBuilder.h
--- a/Engine/Source/Editor/UnrealEd/Classes/Builders/SpiralStairBuilder.h
+++ b/Engine/Source/Editor/UnrealEd/Classes/Builders/SpiralStairBuilder.h
@@ -10,9 +10,11 @@
 UCLASS(MinimalAPI, autoexpandcategories=BrushSettings, EditInlineNew, meta=(DisplayName="Spiral Stair"))
 class USpiralStairBuilder : public UEditorBrushBuilder
 {
+public:
 	GENERATED_BODY()
+
 public:
-	UNREALED_API USpiralStairBuilder(const FObjectInitializer& ObjectInitializer = FObjectInitializer::Get());
+	USpiralStairBuilder(const FObjectInitializer& ObjectInitializer = FObjectInitializer::Get());
 
 	/** The radius of the inner curve of the stair */
 	UPROPERTY(EditAnywhere, Category=BrushSettings, meta=(ClampMin = "1"))
diff --git a/Engine/Source/Editor/UnrealEd/Classes/Builders/TetrahedronBuilder.h b/Engine/Source/Editor/UnrealEd/Classes/Builders/TetrahedronBuilder.h
--- a/Engine/Source/Editor/UnrealEd/Classes/Builders/TetrahedronBuilder.h
+++ b/Engine/Source/Editor/UnrealEd/Classes/Builders/TetrahedronBuilder.h
@@ -10,9 +10,11 @@
 UCLASS(MinimalAPI, autoexpandcategories=BrushSettings, EditInlineNew, meta=(DisplayName="Sphere"))
 class UTetrahedronBuilder : public UEditorBrushBuilder
 {
+public:
 	GENERATED_BODY()
+
 public:
-	UNREALED_API UTetrahedronBuilder(const FObjectInitializer& ObjectInitializer = FObjectInitializer::Get());
+	UTetrahedronBuilder(const FObjectInitializer& ObjectInitializer = FObjectInitializer::Get());
 
 	/** The radius of this sphere */
 	UPROPERTY(EditAnywhere, Category=BrushSettings, meta=(ClampMin = "0.000001"))
diff --git a/Engine/Source/Editor/UnrealEd/Classes/Builders/VolumetricBuilder.h b/Engine/Source/Editor/UnrealEd/Classes/Builders/VolumetricBuilder.h
--- a/Engine/Source/Editor/UnrealEd/Classes/Builders/VolumetricBuilder.h
+++ b/Engine/Source/Editor/UnrealEd/Classes/Builders/VolumetricBuilder.h
@@ -10,10 +10,12 @@
 UCLASS(MinimalAPI, autoexpandcategories=BrushSettings, EditInlineNew, NotPlaceable, meta=(DisplayName="Volumetric"))
 class UVolumetricBuilder : public UEditorBrushBuilder
 {
-	GENERATED_BODY()
 public:
-	UNREALED_API UVolumetricBuilder(const FObjectInitializer& ObjectInitializer = FObjectInitializer::Get());
+	GENERATED_BODY()
 
+public:
+	UVolumetricBuilder(const FObjectInitializer& ObjectInitializer = FObjectInitializer::Get());
+	
 	UPROPERTY(EditAnywhere, Category=BrushSettings, meta=(ClampMin = "0.000001"))
 	float Z;
 
diff --git a/Engine/Source/Editor/UnrealEd/Classes/Commandlets/AutomatedLauncherCommandlet.h b/Engine/Source/Editor/UnrealEd/Classes/Commandlets/AutomatedLauncherCommandlet.h
--- a/Engine/Source/Editor/UnrealEd/Classes/Commandlets/AutomatedLauncherCommandlet.h
+++ b/Engine/Source/Editor/UnrealEd/Classes/Commandlets/AutomatedLauncherCommandlet.h
@@ -9,11 +9,9 @@
 UCLASS()
 class UAutomatedLauncherCommandlet : public UCommandlet
 {
-	GENERATED_BODY()
-public:
-	UAutomatedLauncherCommandlet(const FObjectInitializer& ObjectInitializer = FObjectInitializer::Get());
+	GENERATED_UCLASS_BODY()
 
 	// Begin UCommandlet Interface
 	virtual int32 Main(const FString& Params) override;
 	// End UCommandlet Interface
-};
+};
\ No newline at end of file
diff --git a/Engine/Source/Editor/UnrealEd/Classes/Commandlets/CompressAnimationsCommandlet.h b/Engine/Source/Editor/UnrealEd/Classes/Commandlets/CompressAnimationsCommandlet.h
--- a/Engine/Source/Editor/UnrealEd/Classes/Commandlets/CompressAnimationsCommandlet.h
+++ b/Engine/Source/Editor/UnrealEd/Classes/Commandlets/CompressAnimationsCommandlet.h
@@ -7,9 +7,7 @@
 UCLASS()
 class UCompressAnimationsCommandlet : public UCommandlet
 {
-	GENERATED_BODY()
-public:
-	UCompressAnimationsCommandlet(const FObjectInitializer& ObjectInitializer = FObjectInitializer::Get());
+	GENERATED_UCLASS_BODY()
 	// Begin UCommandlet Interface
 	virtual int32 Main(const FString& Params) override;
 	// End UCommandlet Interface
diff --git a/Engine/Source/Editor/UnrealEd/Classes/Commandlets/DerivedDataCacheCommandlet.h b/Engine/Source/Editor/UnrealEd/Classes/Commandlets/DerivedDataCacheCommandlet.h
--- a/Engine/Source/Editor/UnrealEd/Classes/Commandlets/DerivedDataCacheCommandlet.h
+++ b/Engine/Source/Editor/UnrealEd/Classes/Commandlets/DerivedDataCacheCommandlet.h
@@ -12,9 +12,7 @@
 UCLASS()
 class UDerivedDataCacheCommandlet : public UCommandlet
 {
-	GENERATED_BODY()
-public:
-	UDerivedDataCacheCommandlet(const FObjectInitializer& ObjectInitializer = FObjectInitializer::Get());
+	GENERATED_UCLASS_BODY()
 	// Begin UCommandlet Interface
 	virtual int32 Main(const FString& Params) override;
 	// End UCommandlet Interface
diff --git a/Engine/Source/Editor/UnrealEd/Classes/Commandlets/DiffAssetsCommandlet.h b/Engine/Source/Editor/UnrealEd/Classes/Commandlets/DiffAssetsCommandlet.h
--- a/Engine/Source/Editor/UnrealEd/Classes/Commandlets/DiffAssetsCommandlet.h
+++ b/Engine/Source/Editor/UnrealEd/Classes/Commandlets/DiffAssetsCommandlet.h
@@ -11,9 +11,7 @@
 UCLASS()
 class UDiffAssetsCommandlet : public UCommandlet
 {
-    GENERATED_BODY()
-public:
-    UDiffAssetsCommandlet(const FObjectInitializer& ObjectInitializer = FObjectInitializer::Get());
+    GENERATED_UCLASS_BODY()
 	// Begin UCommandlet Interface
 	virtual int32 Main(const FString& Params) override
 	{
diff --git a/Engine/Source/Editor/UnrealEd/Classes/Commandlets/DiffPackagesCommandlet.h b/Engine/Source/Editor/UnrealEd/Classes/Commandlets/DiffPackagesCommandlet.h
--- a/Engine/Source/Editor/UnrealEd/Classes/Commandlets/DiffPackagesCommandlet.h
+++ b/Engine/Source/Editor/UnrealEd/Classes/Commandlets/DiffPackagesCommandlet.h
@@ -31,9 +31,7 @@ enum EObjectDiff
 UCLASS()
 class UDiffPackagesCommandlet : public UCommandlet
 {
-	GENERATED_BODY()
-public:
-	UDiffPackagesCommandlet(const FObjectInitializer& ObjectInitializer = FObjectInitializer::Get());
+	GENERATED_UCLASS_BODY()
 
 	UPROPERTY()
 	UPackage* Packages[3];
diff --git a/Engine/Source/Editor/UnrealEd/Classes/Commandlets/DumpBlueprintsInfoCommandlet.h b/Engine/Source/Editor/UnrealEd/Classes/Commandlets/DumpBlueprintsInfoCommandlet.h
--- a/Engine/Source/Editor/UnrealEd/Classes/Commandlets/DumpBlueprintsInfoCommandlet.h
+++ b/Engine/Source/Editor/UnrealEd/Classes/Commandlets/DumpBlueprintsInfoCommandlet.h
@@ -7,9 +7,7 @@
 UCLASS()
 class UDumpBlueprintsInfoCommandlet : public UCommandlet
 {
-	GENERATED_BODY()
-public:
-	UDumpBlueprintsInfoCommandlet(const FObjectInitializer& ObjectInitializer = FObjectInitializer::Get());
+	GENERATED_UCLASS_BODY()
 
 public:		
 	// Begin UCommandlet Interface
diff --git a/Engine/Source/Editor/UnrealEd/Classes/Commandlets/DumpHiddenCategoriesCommandlet.h b/Engine/Source/Editor/UnrealEd/Classes/Commandlets/DumpHiddenCategoriesCommandlet.h
--- a/Engine/Source/Editor/UnrealEd/Classes/Commandlets/DumpHiddenCategoriesCommandlet.h
+++ b/Engine/Source/Editor/UnrealEd/Classes/Commandlets/DumpHiddenCategoriesCommandlet.h
@@ -7,9 +7,7 @@
 UCLASS()
 class UDumpHiddenCategoriesCommandlet : public UCommandlet
 {
-	GENERATED_BODY()
-public:
-	UDumpHiddenCategoriesCommandlet(const FObjectInitializer& ObjectInitializer = FObjectInitializer::Get());
+	GENERATED_UCLASS_BODY()
 
 public:		
 	// Begin UCommandlet Interface
diff --git a/Engine/Source/Editor/UnrealEd/Classes/Commandlets/FileServerCommandlet.h b/Engine/Source/Editor/UnrealEd/Classes/Commandlets/FileServerCommandlet.h
--- a/Engine/Source/Editor/UnrealEd/Classes/Commandlets/FileServerCommandlet.h
+++ b/Engine/Source/Editor/UnrealEd/Classes/Commandlets/FileServerCommandlet.h
@@ -16,9 +16,7 @@ UCLASS()
 class UFileServerCommandlet
 	: public UCommandlet
 {
-	GENERATED_BODY()
-public:
-	UFileServerCommandlet(const FObjectInitializer& ObjectInitializer = FObjectInitializer::Get());
+	GENERATED_UCLASS_BODY()
 
 
 public:
@@ -34,4 +32,4 @@ private:
 
 	// Holds the instance identifier.
 	FGuid InstanceId;
-};
+};
\ No newline at end of file
diff --git a/Engine/Source/Editor/UnrealEd/Classes/Commandlets/GatherTextCommandlet.h b/Engine/Source/Editor/UnrealEd/Classes/Commandlets/GatherTextCommandlet.h
--- a/Engine/Source/Editor/UnrealEd/Classes/Commandlets/GatherTextCommandlet.h
+++ b/Engine/Source/Editor/UnrealEd/Classes/Commandlets/GatherTextCommandlet.h
@@ -18,9 +18,7 @@ namespace EOutputJson
 UCLASS()
 class UGatherTextCommandlet : public UGatherTextCommandletBase
 {
-    GENERATED_BODY()
-public:
-    UGatherTextCommandlet(const FObjectInitializer& ObjectInitializer = FObjectInitializer::Get());
+    GENERATED_UCLASS_BODY()
 public:
 	// Begin UCommandlet Interface
 	virtual int32 Main(const FString& Params) override;
diff --git a/Engine/Source/Editor/UnrealEd/Classes/Commandlets/GatherTextCommandletBase.h b/Engine/Source/Editor/UnrealEd/Classes/Commandlets/GatherTextCommandletBase.h
--- a/Engine/Source/Editor/UnrealEd/Classes/Commandlets/GatherTextCommandletBase.h
+++ b/Engine/Source/Editor/UnrealEd/Classes/Commandlets/GatherTextCommandletBase.h
@@ -248,9 +248,7 @@ class FJsonObject;
 UCLASS()
 class UGatherTextCommandletBase : public UCommandlet
 {
-	GENERATED_BODY()
-public:
-	UGatherTextCommandletBase(const FObjectInitializer& ObjectInitializer = FObjectInitializer::Get());
+	GENERATED_UCLASS_BODY()
 
 
 public:
diff --git a/Engine/Source/Editor/UnrealEd/Classes/Commandlets/GatherTextFromAssetsCommandlet.h b/Engine/Source/Editor/UnrealEd/Classes/Commandlets/GatherTextFromAssetsCommandlet.h
--- a/Engine/Source/Editor/UnrealEd/Classes/Commandlets/GatherTextFromAssetsCommandlet.h
+++ b/Engine/Source/Editor/UnrealEd/Classes/Commandlets/GatherTextFromAssetsCommandlet.h
@@ -30,9 +30,7 @@ class UDialogueWave;
 UCLASS()
 class UGatherTextFromAssetsCommandlet : public UGatherTextCommandletBase
 {
-	GENERATED_BODY()
-public:
-	UGatherTextFromAssetsCommandlet(const FObjectInitializer& ObjectInitializer = FObjectInitializer::Get());
+	GENERATED_UCLASS_BODY()
 
 	void ProcessPackages( const TArray< UPackage* >& PackagesToProcess );
 	void ProcessDialogueWave( const UDialogueWave* DialogueWave );
diff --git a/Engine/Source/Editor/UnrealEd/Classes/Commandlets/GatherTextFromMetadataCommandlet.h b/Engine/Source/Editor/UnrealEd/Classes/Commandlets/GatherTextFromMetadataCommandlet.h
--- a/Engine/Source/Editor/UnrealEd/Classes/Commandlets/GatherTextFromMetadataCommandlet.h
+++ b/Engine/Source/Editor/UnrealEd/Classes/Commandlets/GatherTextFromMetadataCommandlet.h
@@ -10,9 +10,7 @@
 UCLASS()
 class UGatherTextFromMetaDataCommandlet : public UGatherTextCommandletBase
 {
-	GENERATED_BODY()
-public:
-	UGatherTextFromMetaDataCommandlet(const FObjectInitializer& ObjectInitializer = FObjectInitializer::Get());
+	GENERATED_UCLASS_BODY()
 
 public:
 	// Begin UCommandlet Interface
@@ -30,4 +28,4 @@ private:
 private:
 	void GatherTextFromUObjects(const TArray<FString>& IncludePaths, const TArray<FString>& ExcludePaths, const FGatherParameters& Arguments);
 	void GatherTextFromUObject(UField* const Field, const FGatherParameters& Arguments);
-};
+};
\ No newline at end of file
diff --git a/Engine/Source/Editor/UnrealEd/Classes/Commandlets/GatherTextFromSourceCommandlet.h b/Engine/Source/Editor/UnrealEd/Classes/Commandlets/GatherTextFromSourceCommandlet.h
--- a/Engine/Source/Editor/UnrealEd/Classes/Commandlets/GatherTextFromSourceCommandlet.h
+++ b/Engine/Source/Editor/UnrealEd/Classes/Commandlets/GatherTextFromSourceCommandlet.h
@@ -10,9 +10,7 @@
 UCLASS()
 class UGatherTextFromSourceCommandlet : public UGatherTextCommandletBase
 {
-	GENERATED_BODY()
-public:
-	UGatherTextFromSourceCommandlet(const FObjectInitializer& ObjectInitializer = FObjectInitializer::Get());
+	GENERATED_UCLASS_BODY()
 
 private:
 #define LOC_DEFINE_REGION
diff --git a/Engine/Source/Editor/UnrealEd/Classes/Commandlets/GenerateBlueprintAPICommandlet.h b/Engine/Source/Editor/UnrealEd/Classes/Commandlets/GenerateBlueprintAPICommandlet.h
--- a/Engine/Source/Editor/UnrealEd/Classes/Commandlets/GenerateBlueprintAPICommandlet.h
+++ b/Engine/Source/Editor/UnrealEd/Classes/Commandlets/GenerateBlueprintAPICommandlet.h
@@ -7,9 +7,7 @@
 UCLASS()
 class UGenerateBlueprintAPICommandlet : public UCommandlet
 {
-	GENERATED_BODY()
-public:
-	UGenerateBlueprintAPICommandlet(const FObjectInitializer& ObjectInitializer = FObjectInitializer::Get());
+	GENERATED_UCLASS_BODY()
 
 public:		
 	// Begin UCommandlet Interface
diff --git a/Engine/Source/Editor/UnrealEd/Classes/Commandlets/GenerateDistillFileSetsCommandlet.h b/Engine/Source/Editor/UnrealEd/Classes/Commandlets/GenerateDistillFileSetsCommandlet.h
--- a/Engine/Source/Editor/UnrealEd/Classes/Commandlets/GenerateDistillFileSetsCommandlet.h
+++ b/Engine/Source/Editor/UnrealEd/Classes/Commandlets/GenerateDistillFileSetsCommandlet.h
@@ -11,9 +11,7 @@
 UCLASS()
 class UGenerateDistillFileSetsCommandlet : public UCommandlet
 {
-	GENERATED_BODY()
-public:
-	UGenerateDistillFileSetsCommandlet(const FObjectInitializer& ObjectInitializer = FObjectInitializer::Get());
+	GENERATED_UCLASS_BODY()
 
 	// Begin UCommandlet Interface
 	virtual int32 Main(const FString& Params) override;
diff --git a/Engine/Source/Editor/UnrealEd/Classes/Commandlets/GenerateGatherArchiveCommandlet.h b/Engine/Source/Editor/UnrealEd/Classes/Commandlets/GenerateGatherArchiveCommandlet.h
--- a/Engine/Source/Editor/UnrealEd/Classes/Commandlets/GenerateGatherArchiveCommandlet.h
+++ b/Engine/Source/Editor/UnrealEd/Classes/Commandlets/GenerateGatherArchiveCommandlet.h
@@ -14,9 +14,7 @@
 UCLASS()
 class UGenerateGatherArchiveCommandlet : public UGatherTextCommandletBase
 {
-    GENERATED_BODY()
-public:
-    UGenerateGatherArchiveCommandlet(const FObjectInitializer& ObjectInitializer = FObjectInitializer::Get());
+    GENERATED_UCLASS_BODY()
 #if CPP || UE_BUILD_DOCS
 public:
 	// Begin UCommandlet Interface
diff --git a/Engine/Source/Editor/UnrealEd/Classes/Commandlets/GenerateGatherManifestCommandlet.h b/Engine/Source/Editor/UnrealEd/Classes/Commandlets/GenerateGatherManifestCommandlet.h
--- a/Engine/Source/Editor/UnrealEd/Classes/Commandlets/GenerateGatherManifestCommandlet.h
+++ b/Engine/Source/Editor/UnrealEd/Classes/Commandlets/GenerateGatherManifestCommandlet.h
@@ -10,9 +10,7 @@
 UCLASS()
 class UGenerateGatherManifestCommandlet : public UGatherTextCommandletBase
 {
-    GENERATED_BODY()
-public:
-    UGenerateGatherManifestCommandlet(const FObjectInitializer& ObjectInitializer = FObjectInitializer::Get());
+    GENERATED_UCLASS_BODY()
 #if CPP || UE_BUILD_DOCS
 public:
 	// Begin UCommandlet Interface
diff --git a/Engine/Source/Editor/UnrealEd/Classes/Commandlets/GenerateTextLocalizationReportCommandlet.h b/Engine/Source/Editor/UnrealEd/Classes/Commandlets/GenerateTextLocalizationReportCommandlet.h
--- a/Engine/Source/Editor/UnrealEd/Classes/Commandlets/GenerateTextLocalizationReportCommandlet.h
+++ b/Engine/Source/Editor/UnrealEd/Classes/Commandlets/GenerateTextLocalizationReportCommandlet.h
@@ -9,9 +9,7 @@
 UCLASS()
 class UGenerateTextLocalizationReportCommandlet : public UGatherTextCommandletBase
 {
-	GENERATED_BODY()
-public:
-	UGenerateTextLocalizationReportCommandlet(const FObjectInitializer& ObjectInitializer = FObjectInitializer::Get());
+	GENERATED_UCLASS_BODY()
 
 
 
@@ -50,4 +48,4 @@ private:
 	FString SectionName;
 	FString CmdlineTimeStamp;
 
-};
+};
\ No newline at end of file
diff --git a/Engine/Source/Editor/UnrealEd/Classes/Commandlets/GenerateTextLocalizationResourceCommandlet.h b/Engine/Source/Editor/UnrealEd/Classes/Commandlets/GenerateTextLocalizationResourceCommandlet.h
--- a/Engine/Source/Editor/UnrealEd/Classes/Commandlets/GenerateTextLocalizationResourceCommandlet.h
+++ b/Engine/Source/Editor/UnrealEd/Classes/Commandlets/GenerateTextLocalizationResourceCommandlet.h
@@ -12,12 +12,10 @@
 UCLASS()
 class UGenerateTextLocalizationResourceCommandlet : public UGatherTextCommandletBase
 {
-	GENERATED_BODY()
-public:
-	UGenerateTextLocalizationResourceCommandlet(const FObjectInitializer& ObjectInitializer = FObjectInitializer::Get());
+	GENERATED_UCLASS_BODY()
 
 public:
 	// Begin UCommandlet Interface
 	virtual int32 Main(const FString& Params) override;
 	// End UCommandlet Interface
-};
+};
\ No newline at end of file
diff --git a/Engine/Source/Editor/UnrealEd/Classes/Commandlets/InternationalizationConditioningCommandlet.h b/Engine/Source/Editor/UnrealEd/Classes/Commandlets/InternationalizationConditioningCommandlet.h
--- a/Engine/Source/Editor/UnrealEd/Classes/Commandlets/InternationalizationConditioningCommandlet.h
+++ b/Engine/Source/Editor/UnrealEd/Classes/Commandlets/InternationalizationConditioningCommandlet.h
@@ -13,9 +13,7 @@
 UCLASS()
 class UInternationalizationConditioningCommandlet : public UGatherTextCommandletBase
 {
-    GENERATED_BODY()
-public:
-    UInternationalizationConditioningCommandlet(const FObjectInitializer& ObjectInitializer = FObjectInitializer::Get());
+    GENERATED_UCLASS_BODY()
 
 	
 	/**
@@ -196,4 +194,4 @@ private:
 	FString SectionName;
 	FConfigCacheIni LegacyLocalizationCacheIni;
 
-};
+};
\ No newline at end of file
diff --git a/Engine/Source/Editor/UnrealEd/Classes/Commandlets/InternationalizationExportCommandlet.h b/Engine/Source/Editor/UnrealEd/Classes/Commandlets/InternationalizationExportCommandlet.h
--- a/Engine/Source/Editor/UnrealEd/Classes/Commandlets/InternationalizationExportCommandlet.h
+++ b/Engine/Source/Editor/UnrealEd/Classes/Commandlets/InternationalizationExportCommandlet.h
@@ -381,9 +381,7 @@ private:
 UCLASS()
 class UInternationalizationExportCommandlet : public UGatherTextCommandletBase
 {
-    GENERATED_BODY()
-public:
-    UInternationalizationExportCommandlet(const FObjectInitializer& ObjectInitializer = FObjectInitializer::Get());
+    GENERATED_UCLASS_BODY()
 
 public:
 	// Begin UCommandlet Interface
@@ -397,4 +395,4 @@ private:
 	FString ConfigPath;
 	FString SectionName;
 
-};
+};
\ No newline at end of file
diff --git a/Engine/Source/Editor/UnrealEd/Classes/Commandlets/ListMaterialsUsedWithMeshEmittersCommandlet.h b/Engine/Source/Editor/UnrealEd/Classes/Commandlets/ListMaterialsUsedWithMeshEmittersCommandlet.h
--- a/Engine/Source/Editor/UnrealEd/Classes/Commandlets/ListMaterialsUsedWithMeshEmittersCommandlet.h
+++ b/Engine/Source/Editor/UnrealEd/Classes/Commandlets/ListMaterialsUsedWithMeshEmittersCommandlet.h
@@ -7,9 +7,7 @@
 UCLASS()
 class UListMaterialsUsedWithMeshEmittersCommandlet : public UCommandlet
 {
-	GENERATED_BODY()
-public:
-	UListMaterialsUsedWithMeshEmittersCommandlet(const FObjectInitializer& ObjectInitializer = FObjectInitializer::Get());
+	GENERATED_UCLASS_BODY()
 	
 	/** Process the given particle system. For mesh emitters, output material paths that don't have bUsedWithParticleSprites flagged.*/
 	void ProcessParticleSystem(UParticleSystem* ParticleSystem , TArray<FString> &OutMaterials);
@@ -17,4 +15,4 @@ public:
 	// Begin UCommandlet Interface
 	virtual int32 Main(const FString& Params) override;
 	// End UCommandlet Interface
-};
+};
\ No newline at end of file
diff --git a/Engine/Source/Editor/UnrealEd/Classes/Commandlets/ListStaticMeshesImportedFromSpeedTreesCommandlet.h b/Engine/Source/Editor/UnrealEd/Classes/Commandlets/ListStaticMeshesImportedFromSpeedTreesCommandlet.h
--- a/Engine/Source/Editor/UnrealEd/Classes/Commandlets/ListStaticMeshesImportedFromSpeedTreesCommandlet.h
+++ b/Engine/Source/Editor/UnrealEd/Classes/Commandlets/ListStaticMeshesImportedFromSpeedTreesCommandlet.h
@@ -7,11 +7,9 @@
 UCLASS()
 class UListStaticMeshesImportedFromSpeedTreesCommandlet : public UCommandlet
 {
-	GENERATED_BODY()
-public:
-	UListStaticMeshesImportedFromSpeedTreesCommandlet(const FObjectInitializer& ObjectInitializer = FObjectInitializer::Get());
+	GENERATED_UCLASS_BODY()
 
 	// Begin UCommandlet Interface
 	virtual int32 Main(const FString& Params) override;
 	// End UCommandlet Interface
-};
+};
\ No newline at end of file
diff --git a/Engine/Source/Editor/UnrealEd/Classes/Commandlets/LoadPackageCommandlet.h b/Engine/Source/Editor/UnrealEd/Classes/Commandlets/LoadPackageCommandlet.h
--- a/Engine/Source/Editor/UnrealEd/Classes/Commandlets/LoadPackageCommandlet.h
+++ b/Engine/Source/Editor/UnrealEd/Classes/Commandlets/LoadPackageCommandlet.h
@@ -7,9 +7,7 @@
 UCLASS()
 class ULoadPackageCommandlet : public UCommandlet
 {
-	GENERATED_BODY()
-public:
-	ULoadPackageCommandlet(const FObjectInitializer& ObjectInitializer = FObjectInitializer::Get());
+	GENERATED_UCLASS_BODY()
 	/**
 	 *	Parse the given load list file, placing the entries in the given Tokens array.
 	 *
diff --git a/Engine/Source/Editor/UnrealEd/Classes/Commandlets/ParticleSystemAuditCommandlet.h b/Engine/Source/Editor/UnrealEd/Classes/Commandlets/ParticleSystemAuditCommandlet.h
--- a/Engine/Source/Editor/UnrealEd/Classes/Commandlets/ParticleSystemAuditCommandlet.h
+++ b/Engine/Source/Editor/UnrealEd/Classes/Commandlets/ParticleSystemAuditCommandlet.h
@@ -9,9 +9,7 @@
 UCLASS(config=Editor)
 class UParticleSystemAuditCommandlet : public UCommandlet
 {
-	GENERATED_BODY()
-public:
-	UParticleSystemAuditCommandlet(const FObjectInitializer& ObjectInitializer = FObjectInitializer::Get());
+	GENERATED_UCLASS_BODY()
 
 	/** Helper struct for tracking disabled LOD level issues */
 	struct FParticleSystemLODInfo
diff --git a/Engine/Source/Editor/UnrealEd/Classes/Commandlets/PkgInfoCommandlet.h b/Engine/Source/Editor/UnrealEd/Classes/Commandlets/PkgInfoCommandlet.h
--- a/Engine/Source/Editor/UnrealEd/Classes/Commandlets/PkgInfoCommandlet.h
+++ b/Engine/Source/Editor/UnrealEd/Classes/Commandlets/PkgInfoCommandlet.h
@@ -7,9 +7,7 @@
 UCLASS()
 class UPkgInfoCommandlet : public UCommandlet
 {
-	GENERATED_BODY()
-public:
-	UPkgInfoCommandlet(const FObjectInitializer& ObjectInitializer = FObjectInitializer::Get());
+	GENERATED_UCLASS_BODY()
 	// Begin UCommandlet Interface
 	virtual int32 Main(const FString& Params) override;
 	// End UCommandlet Interface
diff --git a/Engine/Source/Editor/UnrealEd/Classes/Commandlets/RepairLocalizationDataCommandlet.h b/Engine/Source/Editor/UnrealEd/Classes/Commandlets/RepairLocalizationDataCommandlet.h
--- a/Engine/Source/Editor/UnrealEd/Classes/Commandlets/RepairLocalizationDataCommandlet.h
+++ b/Engine/Source/Editor/UnrealEd/Classes/Commandlets/RepairLocalizationDataCommandlet.h
@@ -9,12 +9,10 @@
 UCLASS()
 class URepairLocalizationDataCommandlet : public UGatherTextCommandletBase
 {
-	GENERATED_BODY()
-public:
-	URepairLocalizationDataCommandlet(const FObjectInitializer& ObjectInitializer = FObjectInitializer::Get());
+	GENERATED_UCLASS_BODY()
 
 public:
 	// Begin UCommandlet Interface
 	virtual int32 Main(const FString& Params) override;
 	// End UCommandlet Interface
-};
+};
\ No newline at end of file
diff --git a/Engine/Source/Editor/UnrealEd/Classes/Commandlets/ReplaceActorCommandlet.h b/Engine/Source/Editor/UnrealEd/Classes/Commandlets/ReplaceActorCommandlet.h
--- a/Engine/Source/Editor/UnrealEd/Classes/Commandlets/ReplaceActorCommandlet.h
+++ b/Engine/Source/Editor/UnrealEd/Classes/Commandlets/ReplaceActorCommandlet.h
@@ -8,9 +8,7 @@
 UCLASS()
 class UReplaceActorCommandlet : public UCommandlet
 {
-	GENERATED_BODY()
-public:
-	UReplaceActorCommandlet(const FObjectInitializer& ObjectInitializer = FObjectInitializer::Get());
+	GENERATED_UCLASS_BODY()
 	// Begin UCommandlet Interface
 	virtual int32 Main(const FString& Params) override;
 	// End UCommandlet Interface
diff --git a/Engine/Source/Editor/UnrealEd/Classes/Commandlets/ResavePackagesCommandlet.h b/Engine/Source/Editor/UnrealEd/Classes/Commandlets/ResavePackagesCommandlet.h
--- a/Engine/Source/Editor/UnrealEd/Classes/Commandlets/ResavePackagesCommandlet.h
+++ b/Engine/Source/Editor/UnrealEd/Classes/Commandlets/ResavePackagesCommandlet.h
@@ -8,9 +8,7 @@ UCLASS()
 // Added UNREALED_API to expose this to the save packages test
 class UNREALED_API UResavePackagesCommandlet : public UCommandlet
 {
-    GENERATED_BODY()
-public:
-    UResavePackagesCommandlet(const FObjectInitializer& ObjectInitializer = FObjectInitializer::Get());
+    GENERATED_UCLASS_BODY()
 
 	enum EBrevity
 	{
diff --git a/Engine/Source/Editor/UnrealEd/Classes/Commandlets/UpdateGameProjectCommandlet.h b/Engine/Source/Editor/UnrealEd/Classes/Commandlets/UpdateGameProjectCommandlet.h
--- a/Engine/Source/Editor/UnrealEd/Classes/Commandlets/UpdateGameProjectCommandlet.h
+++ b/Engine/Source/Editor/UnrealEd/Classes/Commandlets/UpdateGameProjectCommandlet.h
@@ -7,9 +7,7 @@
 UCLASS()
 class UUpdateGameProjectCommandlet : public UCommandlet
 {
-	GENERATED_BODY()
-public:
-	UUpdateGameProjectCommandlet(const FObjectInitializer& ObjectInitializer = FObjectInitializer::Get());
+	GENERATED_UCLASS_BODY()
 
 	// Begin UCommandlet Interface
 	virtual int32 Main(const FString& Params) override;
diff --git a/Engine/Source/Editor/UnrealEd/Classes/Commandlets/WrangleContentCommandlet.h b/Engine/Source/Editor/UnrealEd/Classes/Commandlets/WrangleContentCommandlet.h
--- a/Engine/Source/Editor/UnrealEd/Classes/Commandlets/WrangleContentCommandlet.h
+++ b/Engine/Source/Editor/UnrealEd/Classes/Commandlets/WrangleContentCommandlet.h
@@ -7,9 +7,7 @@
 UCLASS()
 class UWrangleContentCommandlet : public UCommandlet
 {
-	GENERATED_BODY()
-public:
-	UWrangleContentCommandlet(const FObjectInitializer& ObjectInitializer = FObjectInitializer::Get());
+	GENERATED_UCLASS_BODY()
 	// Begin UCommandlet Interface
 	virtual int32 Main(const FString& Params) override;
 	// End UCommandlet Interface
diff --git a/Engine/Source/Editor/UnrealEd/Classes/CookOnTheSide/CookOnTheFlyServer.h b/Engine/Source/Editor/UnrealEd/Classes/CookOnTheSide/CookOnTheFlyServer.h
--- a/Engine/Source/Editor/UnrealEd/Classes/CookOnTheSide/CookOnTheFlyServer.h
+++ b/Engine/Source/Editor/UnrealEd/Classes/CookOnTheSide/CookOnTheFlyServer.h
@@ -52,7 +52,7 @@ UCLASS()
 class UNREALED_API UCookOnTheFlyServer : public UObject, public FTickableEditorObject
 {
 	GENERATED_BODY()
-public:
+
 	UCookOnTheFlyServer(const FObjectInitializer& ObjectInitializer = FObjectInitializer::Get());
 
 private:
diff --git a/Engine/Source/Editor/UnrealEd/Classes/Editor/EditorEngine.h b/Engine/Source/Editor/UnrealEd/Classes/Editor/EditorEngine.h
--- a/Engine/Source/Editor/UnrealEd/Classes/Editor/EditorEngine.h
+++ b/Engine/Source/Editor/UnrealEd/Classes/Editor/EditorEngine.h
@@ -42,7 +42,7 @@ enum EPasteTo
 USTRUCT()
 struct FSlatePlayInEditorInfo
 {
-	GENERATED_BODY()
+	GENERATED_USTRUCT_BODY()
 
 	/** The spawned player for updating viewport location from player when pie closes */
 	TWeakObjectPtr<class ULocalPlayer>	EditorPlayer;
@@ -67,9 +67,9 @@ struct FSlatePlayInEditorInfo
 USTRUCT()
 struct FPIELoginInfo
 {
-	GENERATED_BODY()
 public:
 
+	GENERATED_USTRUCT_BODY()
 
 	/** Type of account. Needed to identity the auth method to use (epic, internal, facebook, etc) */
 	UPROPERTY()
@@ -120,7 +120,7 @@ struct FPieLoginStruct
 USTRUCT()
 struct FCopySelectedInfo
 {
-	GENERATED_BODY()
+	GENERATED_USTRUCT_BODY()
 
 	/** Do not cache this info, it is only valid after a call to CanCopySelectedActorsToClipboard has been made, and becomes redundant
 	when the current selection changes. Used to determine whether a copy can be performed based on the current selection state */
diff --git a/Engine/Source/Editor/UnrealEd/Classes/Editor/EditorGameAgnosticSettings.h b/Engine/Source/Editor/UnrealEd/Classes/Editor/EditorGameAgnosticSettings.h
--- a/Engine/Source/Editor/UnrealEd/Classes/Editor/EditorGameAgnosticSettings.h
+++ b/Engine/Source/Editor/UnrealEd/Classes/Editor/EditorGameAgnosticSettings.h
@@ -7,9 +7,7 @@
 UCLASS(config=EditorGameAgnostic)
 class UEditorGameAgnosticSettings : public UObject
 {
-	GENERATED_BODY()
-public:
-	UEditorGameAgnosticSettings(const FObjectInitializer& ObjectInitializer = FObjectInitializer::Get());
+	GENERATED_UCLASS_BODY()
 
 	/** When checked, the most recently loaded project will be auto-loaded at editor startup if no other project was specified on the command line */
 	UPROPERTY(EditAnywhere, Category=General)
@@ -60,4 +58,4 @@ public:
 	// Begin UObject Interface
 	virtual void PostEditChangeProperty( struct FPropertyChangedEvent& PropertyChangedEvent) override;
 	// End UObject Interface
-};
+};
\ No newline at end of file
diff --git a/Engine/Source/Editor/UnrealEd/Classes/Editor/EditorUserSettings.h b/Engine/Source/Editor/UnrealEd/Classes/Editor/EditorUserSettings.h
--- a/Engine/Source/Editor/UnrealEd/Classes/Editor/EditorUserSettings.h
+++ b/Engine/Source/Editor/UnrealEd/Classes/Editor/EditorUserSettings.h
@@ -10,9 +10,7 @@ UCLASS(minimalapi, autoexpandcategories=(ViewportControls, ViewportLookAndFeel,
 	   config=EditorUserSettings)
 class UEditorUserSettings : public UObject
 {
-	GENERATED_BODY()
-public:
-	UEditorUserSettings(const FObjectInitializer& ObjectInitializer = FObjectInitializer::Get());
+	GENERATED_UCLASS_BODY()
 
 	// =====================================================================
 	// The following options are exposed in the Preferences Editor 
diff --git a/Engine/Source/Editor/UnrealEd/Classes/Editor/GroupActor.h b/Engine/Source/Editor/UnrealEd/Classes/Editor/GroupActor.h
--- a/Engine/Source/Editor/UnrealEd/Classes/Editor/GroupActor.h
+++ b/Engine/Source/Editor/UnrealEd/Classes/Editor/GroupActor.h
@@ -11,9 +11,7 @@
 UCLASS(hidedropdown,MinimalAPI, notplaceable, NotBlueprintable)
 class AGroupActor : public AActor
 {
-	GENERATED_BODY()
-public:
-	UNREALED_API AGroupActor(const FObjectInitializer& ObjectInitializer = FObjectInitializer::Get());
+	GENERATED_UCLASS_BODY()
 
 #if WITH_EDITORONLY_DATA
 	UPROPERTY()
diff --git a/Engine/Source/Editor/UnrealEd/Classes/Editor/PropertyEditorTestObject.h b/Engine/Source/Editor/UnrealEd/Classes/Editor/PropertyEditorTestObject.h
--- a/Engine/Source/Editor/UnrealEd/Classes/Editor/PropertyEditorTestObject.h
+++ b/Engine/Source/Editor/UnrealEd/Classes/Editor/PropertyEditorTestObject.h
@@ -38,7 +38,7 @@ enum ArrayLabelEnum
 USTRUCT()
 struct FPropertyEditorTestSubStruct
 {
-	GENERATED_BODY()
+	GENERATED_USTRUCT_BODY()
 
 	FPropertyEditorTestSubStruct()
 		: FirstProperty( 7897789 )
@@ -65,7 +65,7 @@ struct FPropertyEditorTestSubStruct
 USTRUCT()
 struct FPropertyEditorTestBasicStruct
 {
-	GENERATED_BODY()
+	GENERATED_USTRUCT_BODY()
 
 	FPropertyEditorTestBasicStruct()
 		: IntPropertyInsideAStruct( 0 )
@@ -90,9 +90,7 @@ struct FPropertyEditorTestBasicStruct
 UCLASS(transient)
 class UPropertyEditorTestObject : public UObject
 {
-    GENERATED_BODY()
-public:
-    UPropertyEditorTestObject(const FObjectInitializer& ObjectInitializer = FObjectInitializer::Get());
+    GENERATED_UCLASS_BODY()
 
 	// Integer
 	UPROPERTY(EditAnywhere, Category=BasicProperties)
@@ -263,4 +261,4 @@ public:
 
 	UPROPERTY(EditAnywhere, Category=AssetPropertyTests)
 	AActor* OnlyActorsAllowed;
-};
+};
\ No newline at end of file
diff --git a/Engine/Source/Editor/UnrealEd/Classes/Editor/TemplateMapMetadata.h b/Engine/Source/Editor/UnrealEd/Classes/Editor/TemplateMapMetadata.h
--- a/Engine/Source/Editor/UnrealEd/Classes/Editor/TemplateMapMetadata.h
+++ b/Engine/Source/Editor/UnrealEd/Classes/Editor/TemplateMapMetadata.h
@@ -7,8 +7,6 @@
 UCLASS(deprecated)
 class UDEPRECATED_TemplateMapMetadata : public UObject
 {
-	GENERATED_BODY()
-public:
-	UDEPRECATED_TemplateMapMetadata(const FObjectInitializer& ObjectInitializer = FObjectInitializer::Get());
+	GENERATED_UCLASS_BODY()
 };
 
diff --git a/Engine/Source/Editor/UnrealEd/Classes/Editor/Transactor.h b/Engine/Source/Editor/UnrealEd/Classes/Editor/Transactor.h
--- a/Engine/Source/Editor/UnrealEd/Classes/Editor/Transactor.h
+++ b/Engine/Source/Editor/UnrealEd/Classes/Editor/Transactor.h
@@ -410,9 +410,7 @@ public:
 UCLASS(transient)
 class UTransactor : public UObject
 {
-    GENERATED_BODY()
-public:
-    UTransactor(const FObjectInitializer& ObjectInitializer = FObjectInitializer::Get());
+    GENERATED_UCLASS_BODY()
 	/**
 	 * Begins a new undo transaction.  An undo transaction is defined as all actions
 	 * which take place when the user selects "undo" a single time.
diff --git a/Engine/Source/Editor/UnrealEd/Classes/Editor/UnrealEdEngine.h b/Engine/Source/Editor/UnrealEd/Classes/Editor/UnrealEdEngine.h
--- a/Engine/Source/Editor/UnrealEd/Classes/Editor/UnrealEdEngine.h
+++ b/Engine/Source/Editor/UnrealEd/Classes/Editor/UnrealEdEngine.h
@@ -37,7 +37,7 @@ enum EWriteDisallowedWarningState
 USTRUCT()
 struct FClassMoveInfo
 {
-	GENERATED_BODY()
+	GENERATED_USTRUCT_BODY()
 
 	/** The type of asset this MoveInfo applies to. */
 	UPROPERTY(config)
@@ -67,7 +67,7 @@ struct FClassMoveInfo
 USTRUCT()
 struct FTemplateMapInfo
 {
-	GENERATED_BODY()
+	GENERATED_USTRUCT_BODY()
 
 	/** The Texture2D associated with this map template */
 	UPROPERTY()
@@ -88,6 +88,7 @@ class FPerformanceMonitor;
 UCLASS(config=Engine, transient)
 class UNREALED_API UUnrealEdEngine : public UEditorEngine, public FNotifyHook
 {
+public:
 	GENERATED_BODY()
 public:
 
diff --git a/Engine/Source/Editor/UnrealEd/Classes/Editor/UnrealEdTypes.h b/Engine/Source/Editor/UnrealEd/Classes/Editor/UnrealEdTypes.h
--- a/Engine/Source/Editor/UnrealEd/Classes/Editor/UnrealEdTypes.h
+++ b/Engine/Source/Editor/UnrealEd/Classes/Editor/UnrealEdTypes.h
@@ -14,7 +14,7 @@
 USTRUCT()
 struct FLightmassParameterValue
 {
-	GENERATED_BODY()
+	GENERATED_USTRUCT_BODY()
 
 	/** If true, override the given parameter with the given settings */
 	UPROPERTY(EditAnywhere, BlueprintReadWrite, Category=LightmassParameterValue)
@@ -32,7 +32,7 @@ struct FLightmassParameterValue
 USTRUCT()
 struct FLightmassBooleanParameterValue : public FLightmassParameterValue
 {
-	GENERATED_BODY()
+	GENERATED_USTRUCT_BODY()
 
 	/** The boolean value to override the parent value with */
 	UPROPERTY(EditAnywhere, BlueprintReadWrite, Category=LightmassBooleanParameterValue)
@@ -50,7 +50,7 @@ struct FLightmassBooleanParameterValue : public FLightmassParameterValue
 USTRUCT()
 struct FLightmassScalarParameterValue : public FLightmassParameterValue
 {
-	GENERATED_BODY()
+	GENERATED_USTRUCT_BODY()
 
 	/** The scalar value to override the parent value with */
 	UPROPERTY(EditAnywhere, BlueprintReadWrite, Category=LightmassScalarParameterValue)
@@ -71,7 +71,7 @@ struct FLightmassScalarParameterValue : public FLightmassParameterValue
 USTRUCT()
 struct FLightmassParameterizedMaterialSettings
 {
-	GENERATED_BODY()
+	GENERATED_USTRUCT_BODY()
 
 	/** If true, forces translucency to cast static shadows as if the material were masked. */
 	UPROPERTY(EditAnywhere, BlueprintReadWrite, Category=LightmassParameterizedMaterialSettings)
@@ -132,9 +132,7 @@ enum ELevelViewportType
 UCLASS(abstract, config=UnrealEd)
 class UUnrealEdTypes : public UObject
 {
-	GENERATED_BODY()
-public:
-	UUnrealEdTypes(const FObjectInitializer& ObjectInitializer = FObjectInitializer::Get());
+	GENERATED_UCLASS_BODY()
 
 };
 
diff --git a/Engine/Source/Editor/UnrealEd/Classes/Exporters/AnimSequenceExporterFBX.h b/Engine/Source/Editor/UnrealEd/Classes/Exporters/AnimSequenceExporterFBX.h
--- a/Engine/Source/Editor/UnrealEd/Classes/Exporters/AnimSequenceExporterFBX.h
+++ b/Engine/Source/Editor/UnrealEd/Classes/Exporters/AnimSequenceExporterFBX.h
@@ -11,9 +11,7 @@
 UCLASS()
 class UAnimSequenceExporterFBX : public UExporter
 {
-	GENERATED_BODY()
-public:
-	UAnimSequenceExporterFBX(const FObjectInitializer& ObjectInitializer = FObjectInitializer::Get());
+	GENERATED_UCLASS_BODY()
 
 
 	// Begin UExporter Interface
diff --git a/Engine/Source/Editor/UnrealEd/Classes/Exporters/ExportTextContainer.h b/Engine/Source/Editor/UnrealEd/Classes/Exporters/ExportTextContainer.h
--- a/Engine/Source/Editor/UnrealEd/Classes/Exporters/ExportTextContainer.h
+++ b/Engine/Source/Editor/UnrealEd/Classes/Exporters/ExportTextContainer.h
@@ -7,9 +7,7 @@
 UCLASS(MinimalAPI)
 class UExportTextContainer : public UObject
 {
-	GENERATED_BODY()
-public:
-	UNREALED_API UExportTextContainer(const FObjectInitializer& ObjectInitializer = FObjectInitializer::Get());
+	GENERATED_UCLASS_BODY()
 
 	/** ExportText representation of one or more objects */
 	UPROPERTY()
diff --git a/Engine/Source/Editor/UnrealEd/Classes/Exporters/LevelExporterFBX.h b/Engine/Source/Editor/UnrealEd/Classes/Exporters/LevelExporterFBX.h
--- a/Engine/Source/Editor/UnrealEd/Classes/Exporters/LevelExporterFBX.h
+++ b/Engine/Source/Editor/UnrealEd/Classes/Exporters/LevelExporterFBX.h
@@ -10,10 +10,11 @@
 UCLASS()
 class ULevelExporterFBX : public UExporter
 {
-	GENERATED_BODY()
 public:
-	ULevelExporterFBX(const FObjectInitializer& ObjectInitializer = FObjectInitializer::Get());
+	GENERATED_BODY()
 
+public:
+	ULevelExporterFBX(const FObjectInitializer& ObjectInitializer = FObjectInitializer());
 
 	// Begin UExporter Interface
 	virtual bool ExportBinary( UObject* Object, const TCHAR* Type, FArchive& Ar, FFeedbackContext* Warn, int32 FileIndex = 0, uint32 PortFlags=0 ) override;
diff --git a/Engine/Source/Editor/UnrealEd/Classes/Exporters/LevelExporterOBJ.h b/Engine/Source/Editor/UnrealEd/Classes/Exporters/LevelExporterOBJ.h
--- a/Engine/Source/Editor/UnrealEd/Classes/Exporters/LevelExporterOBJ.h
+++ b/Engine/Source/Editor/UnrealEd/Classes/Exporters/LevelExporterOBJ.h
@@ -10,10 +10,11 @@
 UCLASS()
 class ULevelExporterOBJ : public UExporter
 {
-	GENERATED_BODY()
 public:
-	ULevelExporterOBJ(const FObjectInitializer& ObjectInitializer = FObjectInitializer::Get());
+	GENERATED_BODY()
 
+public:
+	ULevelExporterOBJ(const FObjectInitializer& ObjectInitializer = FObjectInitializer());
 
 	// Begin UExporter Interface
 	virtual bool ExportText( const FExportObjectInnerContext* Context, UObject* Object, const TCHAR* Type, FOutputDevice& Ar, FFeedbackContext* Warn, uint32 PortFlags=0 ) override;
diff --git a/Engine/Source/Editor/UnrealEd/Classes/Exporters/LevelExporterSTL.h b/Engine/Source/Editor/UnrealEd/Classes/Exporters/LevelExporterSTL.h
--- a/Engine/Source/Editor/UnrealEd/Classes/Exporters/LevelExporterSTL.h
+++ b/Engine/Source/Editor/UnrealEd/Classes/Exporters/LevelExporterSTL.h
@@ -10,10 +10,11 @@
 UCLASS()
 class ULevelExporterSTL : public UExporter
 {
-	GENERATED_BODY()
 public:
-	ULevelExporterSTL(const FObjectInitializer& ObjectInitializer = FObjectInitializer::Get());
+	GENERATED_BODY()
 
+public:
+	ULevelExporterSTL(const FObjectInitializer& ObjectInitializer = FObjectInitializer());
 
 	// Begin UExporter Interface
 	virtual bool ExportText( const FExportObjectInnerContext* Context, UObject* Object, const TCHAR* Type, FOutputDevice& Ar, FFeedbackContext* Warn, uint32 PortFlags=0 ) override;
diff --git a/Engine/Source/Editor/UnrealEd/Classes/Exporters/LevelExporterT3D.h b/Engine/Source/Editor/UnrealEd/Classes/Exporters/LevelExporterT3D.h
--- a/Engine/Source/Editor/UnrealEd/Classes/Exporters/LevelExporterT3D.h
+++ b/Engine/Source/Editor/UnrealEd/Classes/Exporters/LevelExporterT3D.h
@@ -10,10 +10,11 @@
 UCLASS()
 class ULevelExporterT3D : public UExporter
 {
-	GENERATED_BODY()
 public:
-	ULevelExporterT3D(const FObjectInitializer& ObjectInitializer = FObjectInitializer::Get());
+	GENERATED_BODY()
 
+public:
+	ULevelExporterT3D(const FObjectInitializer& ObjectInitializer = FObjectInitializer());
 
 	// Begin UExporter Interface
 	virtual bool ExportText( const FExportObjectInnerContext* Context, UObject* Object, const TCHAR* Type, FOutputDevice& Ar, FFeedbackContext* Warn, uint32 PortFlags=0 ) override;
diff --git a/Engine/Source/Editor/UnrealEd/Classes/Exporters/ModelExporterT3D.h b/Engine/Source/Editor/UnrealEd/Classes/Exporters/ModelExporterT3D.h
--- a/Engine/Source/Editor/UnrealEd/Classes/Exporters/ModelExporterT3D.h
+++ b/Engine/Source/Editor/UnrealEd/Classes/Exporters/ModelExporterT3D.h
@@ -10,10 +10,11 @@
 UCLASS()
 class UModelExporterT3D : public UExporter
 {
-	GENERATED_BODY()
 public:
-	UModelExporterT3D(const FObjectInitializer& ObjectInitializer = FObjectInitializer::Get());
+	GENERATED_BODY()
 
+public:
+	UModelExporterT3D(const FObjectInitializer& ObjectInitializer = FObjectInitializer());
 
 	// Begin UExporter Interface
 	virtual bool ExportText( const FExportObjectInnerContext* Context, UObject* Object, const TCHAR* Type, FOutputDevice& Ar, FFeedbackContext* Warn, uint32 PortFlags=0 ) override;
diff --git a/Engine/Source/Editor/UnrealEd/Classes/Exporters/ObjectExporterT3D.h b/Engine/Source/Editor/UnrealEd/Classes/Exporters/ObjectExporterT3D.h
--- a/Engine/Source/Editor/UnrealEd/Classes/Exporters/ObjectExporterT3D.h
+++ b/Engine/Source/Editor/UnrealEd/Classes/Exporters/ObjectExporterT3D.h
@@ -10,9 +10,11 @@
 UCLASS()
 class UObjectExporterT3D : public UExporter
 {
+public:
 	GENERATED_BODY()
+
 public:
-	UObjectExporterT3D(const FObjectInitializer& ObjectInitializer = FObjectInitializer::Get());
+	UObjectExporterT3D(const FObjectInitializer& ObjectInitializer = FObjectInitializer());
 
 
 	// Begin UExporter Interface
diff --git a/Engine/Source/Editor/UnrealEd/Classes/Exporters/SequenceExporterT3D.h b/Engine/Source/Editor/UnrealEd/Classes/Exporters/SequenceExporterT3D.h
--- a/Engine/Source/Editor/UnrealEd/Classes/Exporters/SequenceExporterT3D.h
+++ b/Engine/Source/Editor/UnrealEd/Classes/Exporters/SequenceExporterT3D.h
@@ -10,10 +10,11 @@
 UCLASS()
 class USequenceExporterT3D : public UExporter
 {
-	GENERATED_BODY()
 public:
-	USequenceExporterT3D(const FObjectInitializer& ObjectInitializer = FObjectInitializer::Get());
+	GENERATED_BODY()
 
+public:
+	USequenceExporterT3D(const FObjectInitializer& ObjectInitializer = FObjectInitializer());
 
 	// Begin UExporter Interface
 	virtual bool ExportText( const FExportObjectInnerContext* Context, UObject* Object, const TCHAR* Type, FOutputDevice& Ar, FFeedbackContext* Warn, uint32 PortFlags=0 ) override;
diff --git a/Engine/Source/Editor/UnrealEd/Classes/Exporters/SkeletalMeshExporterFBX.h b/Engine/Source/Editor/UnrealEd/Classes/Exporters/SkeletalMeshExporterFBX.h
--- a/Engine/Source/Editor/UnrealEd/Classes/Exporters/SkeletalMeshExporterFBX.h
+++ b/Engine/Source/Editor/UnrealEd/Classes/Exporters/SkeletalMeshExporterFBX.h
@@ -10,9 +10,7 @@
 UCLASS()
 class USkeletalMeshExporterFBX : public UExporter
 {
-	GENERATED_BODY()
-public:
-	USkeletalMeshExporterFBX(const FObjectInitializer& ObjectInitializer = FObjectInitializer::Get());
+	GENERATED_UCLASS_BODY()
 
 
 	// Begin UExporter Interface
diff --git a/Engine/Source/Editor/UnrealEd/Classes/Exporters/SoundExporterOGG.h b/Engine/Source/Editor/UnrealEd/Classes/Exporters/SoundExporterOGG.h
--- a/Engine/Source/Editor/UnrealEd/Classes/Exporters/SoundExporterOGG.h
+++ b/Engine/Source/Editor/UnrealEd/Classes/Exporters/SoundExporterOGG.h
@@ -6,9 +6,11 @@
 UCLASS()
 class USoundExporterOGG : public UExporter
 {
-    GENERATED_BODY()
 public:
-    USoundExporterOGG(const FObjectInitializer& ObjectInitializer = FObjectInitializer::Get());
+	GENERATED_BODY()
+
+public:
+	USoundExporterOGG(const FObjectInitializer& ObjectInitializer = FObjectInitializer());
 
 	// Begin UExporter Interface
 	virtual bool ExportBinary( UObject* Object, const TCHAR* Type, FArchive& Ar, FFeedbackContext* Warn, int32 FileIndex = 0, uint32 PortFlags=0 ) override;
diff --git a/Engine/Source/Editor/UnrealEd/Classes/Exporters/SoundExporterWAV.h b/Engine/Source/Editor/UnrealEd/Classes/Exporters/SoundExporterWAV.h
--- a/Engine/Source/Editor/UnrealEd/Classes/Exporters/SoundExporterWAV.h
+++ b/Engine/Source/Editor/UnrealEd/Classes/Exporters/SoundExporterWAV.h
@@ -6,9 +6,11 @@
 UCLASS()
 class USoundExporterWAV : public UExporter
 {
-    GENERATED_BODY()
 public:
-    USoundExporterWAV(const FObjectInitializer& ObjectInitializer = FObjectInitializer::Get());
+	GENERATED_BODY()
+
+public:
+	USoundExporterWAV(const FObjectInitializer& ObjectInitializer = FObjectInitializer());
 
 	// Begin UExporter Interface
 	virtual bool ExportBinary( UObject* Object, const TCHAR* Type, FArchive& Ar, FFeedbackContext* Warn, int32 FileIndex = 0, uint32 PortFlags=0 ) override;
diff --git a/Engine/Source/Editor/UnrealEd/Classes/Exporters/SoundSurroundExporterWAV.h b/Engine/Source/Editor/UnrealEd/Classes/Exporters/SoundSurroundExporterWAV.h
--- a/Engine/Source/Editor/UnrealEd/Classes/Exporters/SoundSurroundExporterWAV.h
+++ b/Engine/Source/Editor/UnrealEd/Classes/Exporters/SoundSurroundExporterWAV.h
@@ -10,10 +10,11 @@
 UCLASS()
 class USoundSurroundExporterWAV : public UExporter
 {
-	GENERATED_BODY()
 public:
-	USoundSurroundExporterWAV(const FObjectInitializer& ObjectInitializer = FObjectInitializer::Get());
+	GENERATED_BODY()
 
+public:
+	USoundSurroundExporterWAV(const FObjectInitializer& ObjectInitializer = FObjectInitializer());
 
 	// Begin UExporter Interface
 	virtual int32 GetFileCount( void ) const override;
diff --git a/Engine/Source/Editor/UnrealEd/Classes/Exporters/StaticMeshExporterFBX.h b/Engine/Source/Editor/UnrealEd/Classes/Exporters/StaticMeshExporterFBX.h
--- a/Engine/Source/Editor/UnrealEd/Classes/Exporters/StaticMeshExporterFBX.h
+++ b/Engine/Source/Editor/UnrealEd/Classes/Exporters/StaticMeshExporterFBX.h
@@ -10,9 +10,7 @@
 UCLASS()
 class UStaticMeshExporterFBX : public UExporter
 {
-	GENERATED_BODY()
-public:
-	UStaticMeshExporterFBX(const FObjectInitializer& ObjectInitializer = FObjectInitializer::Get());
+	GENERATED_UCLASS_BODY()
 
 
 	// Begin UExporter Interface
diff --git a/Engine/Source/Editor/UnrealEd/Classes/Exporters/TextBufferExporterTXT.h b/Engine/Source/Editor/UnrealEd/Classes/Exporters/TextBufferExporterTXT.h
--- a/Engine/Source/Editor/UnrealEd/Classes/Exporters/TextBufferExporterTXT.h
+++ b/Engine/Source/Editor/UnrealEd/Classes/Exporters/TextBufferExporterTXT.h
@@ -10,10 +10,11 @@
 UCLASS()
 class UTextBufferExporterTXT : public UExporter
 {
-	GENERATED_BODY()
 public:
-	UTextBufferExporterTXT(const FObjectInitializer& ObjectInitializer = FObjectInitializer::Get());
+	GENERATED_BODY()
 
+public:
+	UTextBufferExporterTXT(const FObjectInitializer& ObjectInitializer = FObjectInitializer());
 
 	// Begin UExporter Interface
 	virtual bool ExportText( const FExportObjectInnerContext* Context, UObject* Object, const TCHAR* Type, FOutputDevice& Ar, FFeedbackContext* Warn, uint32 PortFlags=0 ) override;
diff --git a/Engine/Source/Editor/UnrealEd/Classes/Exporters/TextureExporterBMP.h b/Engine/Source/Editor/UnrealEd/Classes/Exporters/TextureExporterBMP.h
--- a/Engine/Source/Editor/UnrealEd/Classes/Exporters/TextureExporterBMP.h
+++ b/Engine/Source/Editor/UnrealEd/Classes/Exporters/TextureExporterBMP.h
@@ -10,9 +10,7 @@
 UCLASS()
 class UTextureExporterBMP : public UExporter
 {
-	GENERATED_BODY()
-public:
-	UTextureExporterBMP(const FObjectInitializer& ObjectInitializer = FObjectInitializer::Get());
+	GENERATED_UCLASS_BODY()
 
 
 	bool ExportBinary( UObject* Object, const TCHAR* Type, FArchive& Ar, FFeedbackContext* Warn, int32 FileIndex = 0, uint32 PortFlags=0 );
diff --git a/Engine/Source/Editor/UnrealEd/Classes/Exporters/TextureExporterPCX.h b/Engine/Source/Editor/UnrealEd/Classes/Exporters/TextureExporterPCX.h
--- a/Engine/Source/Editor/UnrealEd/Classes/Exporters/TextureExporterPCX.h
+++ b/Engine/Source/Editor/UnrealEd/Classes/Exporters/TextureExporterPCX.h
@@ -10,9 +10,7 @@
 UCLASS()
 class UTextureExporterPCX : public UExporter
 {
-	GENERATED_BODY()
-public:
-	UTextureExporterPCX(const FObjectInitializer& ObjectInitializer = FObjectInitializer::Get());
+	GENERATED_UCLASS_BODY()
 
 
 	// Begin UExporter Interface
diff --git a/Engine/Source/Editor/UnrealEd/Classes/Exporters/TextureExporterTGA.h b/Engine/Source/Editor/UnrealEd/Classes/Exporters/TextureExporterTGA.h
--- a/Engine/Source/Editor/UnrealEd/Classes/Exporters/TextureExporterTGA.h
+++ b/Engine/Source/Editor/UnrealEd/Classes/Exporters/TextureExporterTGA.h
@@ -10,9 +10,7 @@
 UCLASS()
 class UNREALED_API UTextureExporterTGA : public UExporter
 {
-	GENERATED_BODY()
-public:
-	UTextureExporterTGA(const FObjectInitializer& ObjectInitializer = FObjectInitializer::Get());
+	GENERATED_UCLASS_BODY()
 
 
 	// Begin UExporter Interface
diff --git a/Engine/Source/Editor/UnrealEd/Classes/Factories/AimOffsetBlendSpaceFactory1D.h b/Engine/Source/Editor/UnrealEd/Classes/Factories/AimOffsetBlendSpaceFactory1D.h
--- a/Engine/Source/Editor/UnrealEd/Classes/Factories/AimOffsetBlendSpaceFactory1D.h
+++ b/Engine/Source/Editor/UnrealEd/Classes/Factories/AimOffsetBlendSpaceFactory1D.h
@@ -10,9 +10,7 @@
 UCLASS(hidecategories=Object, MinimalAPI)
 class UAimOffsetBlendSpaceFactory1D : public UBlendSpaceFactory1D
 {
-	GENERATED_BODY()
-public:
-	UNREALED_API UAimOffsetBlendSpaceFactory1D(const FObjectInitializer& ObjectInitializer = FObjectInitializer::Get());
+	GENERATED_UCLASS_BODY()
 
 	// Begin UFactory Interface
 	virtual UObject* FactoryCreateNew(UClass* Class,UObject* InParent,FName Name,EObjectFlags Flags,UObject* Context,FFeedbackContext* Warn) override;
diff --git a/Engine/Source/Editor/UnrealEd/Classes/Factories/AimOffsetBlendSpaceFactoryNew.h b/Engine/Source/Editor/UnrealEd/Classes/Factories/AimOffsetBlendSpaceFactoryNew.h
--- a/Engine/Source/Editor/UnrealEd/Classes/Factories/AimOffsetBlendSpaceFactoryNew.h
+++ b/Engine/Source/Editor/UnrealEd/Classes/Factories/AimOffsetBlendSpaceFactoryNew.h
@@ -10,9 +10,7 @@
 UCLASS(hidecategories=Object, MinimalAPI)
 class UAimOffsetBlendSpaceFactoryNew : public UBlendSpaceFactoryNew
 {
-	GENERATED_BODY()
-public:
-	UNREALED_API UAimOffsetBlendSpaceFactoryNew(const FObjectInitializer& ObjectInitializer = FObjectInitializer::Get());
+	GENERATED_UCLASS_BODY()
 
 	// Begin UFactory Interface
 	virtual UObject* FactoryCreateNew(UClass* Class,UObject* InParent,FName Name,EObjectFlags Flags,UObject* Context,FFeedbackContext* Warn) override;
diff --git a/Engine/Source/Editor/UnrealEd/Classes/Factories/AnimBlueprintFactory.h b/Engine/Source/Editor/UnrealEd/Classes/Factories/AnimBlueprintFactory.h
--- a/Engine/Source/Editor/UnrealEd/Classes/Factories/AnimBlueprintFactory.h
+++ b/Engine/Source/Editor/UnrealEd/Classes/Factories/AnimBlueprintFactory.h
@@ -11,9 +11,7 @@
 UCLASS(HideCategories=Object,MinimalAPI)
 class UAnimBlueprintFactory : public UFactory
 {
-	GENERATED_BODY()
-public:
-	UNREALED_API UAnimBlueprintFactory(const FObjectInitializer& ObjectInitializer = FObjectInitializer::Get());
+	GENERATED_UCLASS_BODY()
 
 	// The type of blueprint that will be created
 	UPROPERTY(EditAnywhere, Category=AnimBlueprintFactory)
diff --git a/Engine/Source/Editor/UnrealEd/Classes/Factories/AnimCompositeFactory.h b/Engine/Source/Editor/UnrealEd/Classes/Factories/AnimCompositeFactory.h
--- a/Engine/Source/Editor/UnrealEd/Classes/Factories/AnimCompositeFactory.h
+++ b/Engine/Source/Editor/UnrealEd/Classes/Factories/AnimCompositeFactory.h
@@ -10,9 +10,7 @@
 UCLASS(HideCategories=Object,MinimalAPI)
 class UAnimCompositeFactory : public UFactory
 {
-	GENERATED_BODY()
-public:
-	UNREALED_API UAnimCompositeFactory(const FObjectInitializer& ObjectInitializer = FObjectInitializer::Get());
+	GENERATED_UCLASS_BODY()
 
 	UPROPERTY()
 	class USkeleton* TargetSkeleton;
diff --git a/Engine/Source/Editor/UnrealEd/Classes/Factories/AnimMontageFactory.h b/Engine/Source/Editor/UnrealEd/Classes/Factories/AnimMontageFactory.h
--- a/Engine/Source/Editor/UnrealEd/Classes/Factories/AnimMontageFactory.h
+++ b/Engine/Source/Editor/UnrealEd/Classes/Factories/AnimMontageFactory.h
@@ -10,9 +10,7 @@
 UCLASS(HideCategories=Object,MinimalAPI)
 class UAnimMontageFactory : public UFactory
 {
-	GENERATED_BODY()
-public:
-	UNREALED_API UAnimMontageFactory(const FObjectInitializer& ObjectInitializer = FObjectInitializer::Get());
+	GENERATED_UCLASS_BODY()
 
 	UPROPERTY()
 	class USkeleton* TargetSkeleton;
diff --git a/Engine/Source/Editor/UnrealEd/Classes/Factories/AnimSequenceFactory.h b/Engine/Source/Editor/UnrealEd/Classes/Factories/AnimSequenceFactory.h
--- a/Engine/Source/Editor/UnrealEd/Classes/Factories/AnimSequenceFactory.h
+++ b/Engine/Source/Editor/UnrealEd/Classes/Factories/AnimSequenceFactory.h
@@ -10,9 +10,7 @@
 UCLASS(HideCategories=Object,MinimalAPI)
 class UAnimSequenceFactory : public UFactory
 {
-	GENERATED_BODY()
-public:
-	UNREALED_API UAnimSequenceFactory(const FObjectInitializer& ObjectInitializer = FObjectInitializer::Get());
+	GENERATED_UCLASS_BODY()
 
 	UPROPERTY()
 	class USkeleton* TargetSkeleton;
diff --git a/Engine/Source/Editor/UnrealEd/Classes/Factories/BlendSpaceFactory1D.h b/Engine/Source/Editor/UnrealEd/Classes/Factories/BlendSpaceFactory1D.h
--- a/Engine/Source/Editor/UnrealEd/Classes/Factories/BlendSpaceFactory1D.h
+++ b/Engine/Source/Editor/UnrealEd/Classes/Factories/BlendSpaceFactory1D.h
@@ -10,9 +10,7 @@
 UCLASS(hidecategories=Object, MinimalAPI)
 class UBlendSpaceFactory1D : public UFactory
 {
-	GENERATED_BODY()
-public:
-	UNREALED_API UBlendSpaceFactory1D(const FObjectInitializer& ObjectInitializer = FObjectInitializer::Get());
+	GENERATED_UCLASS_BODY()
 
 	UPROPERTY()
 	class USkeleton*	TargetSkeleton;
diff --git a/Engine/Source/Editor/UnrealEd/Classes/Factories/BlendSpaceFactoryNew.h b/Engine/Source/Editor/UnrealEd/Classes/Factories/BlendSpaceFactoryNew.h
--- a/Engine/Source/Editor/UnrealEd/Classes/Factories/BlendSpaceFactoryNew.h
+++ b/Engine/Source/Editor/UnrealEd/Classes/Factories/BlendSpaceFactoryNew.h
@@ -10,9 +10,7 @@
 UCLASS(hidecategories=Object, MinimalAPI)
 class UBlendSpaceFactoryNew : public UFactory
 {
-	GENERATED_BODY()
-public:
-	UNREALED_API UBlendSpaceFactoryNew(const FObjectInitializer& ObjectInitializer = FObjectInitializer::Get());
+	GENERATED_UCLASS_BODY()
 
 	UPROPERTY()
 	class USkeleton*	TargetSkeleton;
diff --git a/Engine/Source/Editor/UnrealEd/Classes/Factories/BlueprintFactory.h b/Engine/Source/Editor/UnrealEd/Classes/Factories/BlueprintFactory.h
--- a/Engine/Source/Editor/UnrealEd/Classes/Factories/BlueprintFactory.h
+++ b/Engine/Source/Editor/UnrealEd/Classes/Factories/BlueprintFactory.h
@@ -10,9 +10,7 @@
 UCLASS(hidecategories=Object, collapsecategories)
 class UNREALED_API UBlueprintFactory : public UFactory
 {
-	GENERATED_BODY()
-public:
-	UBlueprintFactory(const FObjectInitializer& ObjectInitializer = FObjectInitializer::Get());
+	GENERATED_UCLASS_BODY()
 
 	// The parent class of the created blueprint
 	UPROPERTY(EditAnywhere, Category=BlueprintFactory, meta=(AllowAbstract = "", BlueprintBaseOnly = ""))
diff --git a/Engine/Source/Editor/UnrealEd/Classes/Factories/BlueprintFunctionLibraryFactory.h b/Engine/Source/Editor/UnrealEd/Classes/Factories/BlueprintFunctionLibraryFactory.h
--- a/Engine/Source/Editor/UnrealEd/Classes/Factories/BlueprintFunctionLibraryFactory.h
+++ b/Engine/Source/Editor/UnrealEd/Classes/Factories/BlueprintFunctionLibraryFactory.h
@@ -10,9 +10,7 @@
 UCLASS(hidecategories=Object, collapsecategories)
 class UBlueprintFunctionLibraryFactory : public UBlueprintFactory
 {
-	GENERATED_BODY()
-public:
-	UBlueprintFunctionLibraryFactory(const FObjectInitializer& ObjectInitializer = FObjectInitializer::Get());
+	GENERATED_UCLASS_BODY()
 
 	// UFactory interface
 	virtual FText GetDisplayName() const override;
diff --git a/Engine/Source/Editor/UnrealEd/Classes/Factories/BlueprintInterfaceFactory.h b/Engine/Source/Editor/UnrealEd/Classes/Factories/BlueprintInterfaceFactory.h
--- a/Engine/Source/Editor/UnrealEd/Classes/Factories/BlueprintInterfaceFactory.h
+++ b/Engine/Source/Editor/UnrealEd/Classes/Factories/BlueprintInterfaceFactory.h
@@ -10,9 +10,7 @@
 UCLASS(hidecategories=Object, collapsecategories)
 class UBlueprintInterfaceFactory : public UFactory
 {
-	GENERATED_BODY()
-public:
-	UBlueprintInterfaceFactory(const FObjectInitializer& ObjectInitializer = FObjectInitializer::Get());
+	GENERATED_UCLASS_BODY()
 
 	// UFactory interface
 	virtual FText GetDisplayName() const override;
diff --git a/Engine/Source/Editor/UnrealEd/Classes/Factories/BlueprintMacroFactory.h b/Engine/Source/Editor/UnrealEd/Classes/Factories/BlueprintMacroFactory.h
--- a/Engine/Source/Editor/UnrealEd/Classes/Factories/BlueprintMacroFactory.h
+++ b/Engine/Source/Editor/UnrealEd/Classes/Factories/BlueprintMacroFactory.h
@@ -10,9 +10,7 @@
 UCLASS(hidecategories=Object, collapsecategories)
 class UBlueprintMacroFactory : public UBlueprintFactory
 {
-	GENERATED_BODY()
-public:
-	UBlueprintMacroFactory(const FObjectInitializer& ObjectInitializer = FObjectInitializer::Get());
+	GENERATED_UCLASS_BODY()
 
 	// UFactory interface
 	virtual FText GetDisplayName() const override;
diff --git a/Engine/Source/Editor/UnrealEd/Classes/Factories/CSVImportFactory.h b/Engine/Source/Editor/UnrealEd/Classes/Factories/CSVImportFactory.h
--- a/Engine/Source/Editor/UnrealEd/Classes/Factories/CSVImportFactory.h
+++ b/Engine/Source/Editor/UnrealEd/Classes/Factories/CSVImportFactory.h
@@ -12,9 +12,7 @@ DECLARE_LOG_CATEGORY_EXTERN(LogCSVImportFactory, Log, All);
 UCLASS(hidecategories=Object)
 class UNREALED_API UCSVImportFactory : public UFactory
 {
-	GENERATED_BODY()
-public:
-	UCSVImportFactory(const FObjectInitializer& ObjectInitializer = FObjectInitializer::Get());
+	GENERATED_UCLASS_BODY()
 
 	// Begin UFactory Interface
 	virtual FText GetDisplayName() const override;
diff --git a/Engine/Source/Editor/UnrealEd/Classes/Factories/CameraAnimFactory.h b/Engine/Source/Editor/UnrealEd/Classes/Factories/CameraAnimFactory.h
--- a/Engine/Source/Editor/UnrealEd/Classes/Factories/CameraAnimFactory.h
+++ b/Engine/Source/Editor/UnrealEd/Classes/Factories/CameraAnimFactory.h
@@ -7,9 +7,7 @@
 UCLASS(hidecategories=Object, MinimalAPI)
 class UCameraAnimFactory : public UFactory
 {
-	GENERATED_BODY()
-public:
-	UNREALED_API UCameraAnimFactory(const FObjectInitializer& ObjectInitializer = FObjectInitializer::Get());
+	GENERATED_UCLASS_BODY()
 
 	// Begin UFactory Interface
 	virtual FText GetDisplayName() const override;
diff --git a/Engine/Source/Editor/UnrealEd/Classes/Factories/CurveFactory.h b/Engine/Source/Editor/UnrealEd/Classes/Factories/CurveFactory.h
--- a/Engine/Source/Editor/UnrealEd/Classes/Factories/CurveFactory.h
+++ b/Engine/Source/Editor/UnrealEd/Classes/Factories/CurveFactory.h
@@ -21,9 +21,7 @@ enum ECurveType
 UCLASS(hidecategories=Object)
 class UCurveFactory : public UFactory
 {
-	GENERATED_BODY()
-public:
-	UCurveFactory(const FObjectInitializer& ObjectInitializer = FObjectInitializer::Get());
+	GENERATED_UCLASS_BODY()
 
 	/** The type of blueprint that will be created */
 	UPROPERTY(EditAnywhere, Category=CurveFactory)
diff --git a/Engine/Source/Editor/UnrealEd/Classes/Factories/CurveImportFactory.h b/Engine/Source/Editor/UnrealEd/Classes/Factories/CurveImportFactory.h
--- a/Engine/Source/Editor/UnrealEd/Classes/Factories/CurveImportFactory.h
+++ b/Engine/Source/Editor/UnrealEd/Classes/Factories/CurveImportFactory.h
@@ -10,9 +10,7 @@
 UCLASS(hidecategories=Object)
 class UCurveImportFactory : public UFactory
 {
-	GENERATED_BODY()
-public:
-	UCurveImportFactory(const FObjectInitializer& ObjectInitializer = FObjectInitializer::Get());
+	GENERATED_UCLASS_BODY()
 
 	// Begin UFactory Interface
 	virtual UObject* FactoryCreateText( UClass* InClass, UObject* InParent, FName InName, EObjectFlags Flags, UObject* Context, const TCHAR* Type, const TCHAR*& Buffer, const TCHAR* BufferEnd, FFeedbackContext* Warn ) override;
diff --git a/Engine/Source/Editor/UnrealEd/Classes/Factories/DataAssetFactory.h b/Engine/Source/Editor/UnrealEd/Classes/Factories/DataAssetFactory.h
--- a/Engine/Source/Editor/UnrealEd/Classes/Factories/DataAssetFactory.h
+++ b/Engine/Source/Editor/UnrealEd/Classes/Factories/DataAssetFactory.h
@@ -9,9 +9,7 @@ class UDataAsset;
 UCLASS(hidecategories=Object, MinimalAPI)
 class UDataAssetFactory : public UFactory
 {
-	GENERATED_BODY()
-public:
-	UNREALED_API UDataAssetFactory(const FObjectInitializer& ObjectInitializer = FObjectInitializer::Get());
+	GENERATED_UCLASS_BODY()
 
 	UPROPERTY(EditAnywhere, Category=DataAsset)
 	TSubclassOf<UDataAsset> DataAssetClass;
diff --git a/Engine/Source/Editor/UnrealEd/Classes/Factories/DataTableFactory.h b/Engine/Source/Editor/UnrealEd/Classes/Factories/DataTableFactory.h
--- a/Engine/Source/Editor/UnrealEd/Classes/Factories/DataTableFactory.h
+++ b/Engine/Source/Editor/UnrealEd/Classes/Factories/DataTableFactory.h
@@ -7,9 +7,7 @@
 UCLASS(hidecategories=Object, MinimalAPI)
 class UDataTableFactory : public UFactory
 {
-	GENERATED_BODY()
-public:
-	UNREALED_API UDataTableFactory(const FObjectInitializer& ObjectInitializer = FObjectInitializer::Get());
+	GENERATED_UCLASS_BODY()
 
 	UPROPERTY()
 	class UScriptStruct* Struct;
diff --git a/Engine/Source/Editor/UnrealEd/Classes/Factories/DestructibleMeshFactory.h b/Engine/Source/Editor/UnrealEd/Classes/Factories/DestructibleMeshFactory.h
--- a/Engine/Source/Editor/UnrealEd/Classes/Factories/DestructibleMeshFactory.h
+++ b/Engine/Source/Editor/UnrealEd/Classes/Factories/DestructibleMeshFactory.h
@@ -8,9 +8,7 @@
 UCLASS(hideCategories=Object, MinimalAPI)
 class UDestructibleMeshFactory : public UFactory
 {
-    GENERATED_BODY()
-public:
-    UNREALED_API UDestructibleMeshFactory(const FObjectInitializer& ObjectInitializer = FObjectInitializer::Get());
+    GENERATED_UCLASS_BODY()
 	// Begin UFactory Interface
 	virtual FText GetDisplayName() const override;
 #if WITH_APEX
diff --git a/Engine/Source/Editor/UnrealEd/Classes/Factories/DialogueVoiceFactory.h b/Engine/Source/Editor/UnrealEd/Classes/Factories/DialogueVoiceFactory.h
--- a/Engine/Source/Editor/UnrealEd/Classes/Factories/DialogueVoiceFactory.h
+++ b/Engine/Source/Editor/UnrealEd/Classes/Factories/DialogueVoiceFactory.h
@@ -10,9 +10,7 @@
 UCLASS(hidecategories=Object)
 class UDialogueVoiceFactory : public UFactory
 {
-	GENERATED_BODY()
-public:
-	UDialogueVoiceFactory(const FObjectInitializer& ObjectInitializer = FObjectInitializer::Get());
+	GENERATED_UCLASS_BODY()
 
 	// Begin UFactory Interface
 	virtual UObject* FactoryCreateNew(UClass* Class,UObject* InParent,FName Name,EObjectFlags Flags,UObject* Context,FFeedbackContext* Warn) override;
diff --git a/Engine/Source/Editor/UnrealEd/Classes/Factories/DialogueWaveFactory.h b/Engine/Source/Editor/UnrealEd/Classes/Factories/DialogueWaveFactory.h
--- a/Engine/Source/Editor/UnrealEd/Classes/Factories/DialogueWaveFactory.h
+++ b/Engine/Source/Editor/UnrealEd/Classes/Factories/DialogueWaveFactory.h
@@ -10,9 +10,7 @@
 UCLASS(hidecategories=Object)
 class UDialogueWaveFactory : public UFactory
 {
-	GENERATED_BODY()
-public:
-	UDialogueWaveFactory(const FObjectInitializer& ObjectInitializer = FObjectInitializer::Get());
+	GENERATED_UCLASS_BODY()
 
 	// Begin UFactory Interface
 	virtual UObject* FactoryCreateNew(UClass* Class,UObject* InParent,FName Name,EObjectFlags Flags,UObject* Context,FFeedbackContext* Warn) override;
diff --git a/Engine/Source/Editor/UnrealEd/Classes/Factories/EnumFactory.h b/Engine/Source/Editor/UnrealEd/Classes/Factories/EnumFactory.h
--- a/Engine/Source/Editor/UnrealEd/Classes/Factories/EnumFactory.h
+++ b/Engine/Source/Editor/UnrealEd/Classes/Factories/EnumFactory.h
@@ -7,9 +7,7 @@
 UCLASS(hidecategories=Object, MinimalAPI)
 class UEnumFactory : public UFactory
 {
-	GENERATED_BODY()
-public:
-	UNREALED_API UEnumFactory(const FObjectInitializer& ObjectInitializer = FObjectInitializer::Get());
+	GENERATED_UCLASS_BODY()
 
 	// Begin UFactory Interface
 	virtual UObject* FactoryCreateNew(UClass* Class,UObject* InParent,FName Name,EObjectFlags Flags,UObject* Context,FFeedbackContext* Warn) override;
diff --git a/Engine/Source/Editor/UnrealEd/Classes/Factories/Factory.h b/Engine/Source/Editor/UnrealEd/Classes/Factories/Factory.h
--- a/Engine/Source/Editor/UnrealEd/Classes/Factories/Factory.h
+++ b/Engine/Source/Editor/UnrealEd/Classes/Factories/Factory.h
@@ -12,9 +12,7 @@
 UCLASS(abstract)
 class UNREALED_API UFactory : public UObject
 {
-	GENERATED_BODY()
-public:
-	UFactory(const FObjectInitializer& ObjectInitializer = FObjectInitializer::Get());
+	GENERATED_UCLASS_BODY()
 
 	/** The class manufactured by this factory. */
 	UPROPERTY()
diff --git a/Engine/Source/Editor/UnrealEd/Classes/Factories/FbxAnimSequenceImportData.h b/Engine/Source/Editor/UnrealEd/Classes/Factories/FbxAnimSequenceImportData.h
--- a/Engine/Source/Editor/UnrealEd/Classes/Factories/FbxAnimSequenceImportData.h
+++ b/Engine/Source/Editor/UnrealEd/Classes/Factories/FbxAnimSequenceImportData.h
@@ -26,9 +26,7 @@ enum EFBXAnimationLengthImportType
 UCLASS()
 class UNREALED_API UFbxAnimSequenceImportData : public UFbxAssetImportData
 {
-	GENERATED_BODY()
-public:
-	UFbxAnimSequenceImportData(const FObjectInitializer& ObjectInitializer = FObjectInitializer::Get());
+	GENERATED_UCLASS_BODY()
 
 	/** Type of asset to import from the FBX file */
 	UPROPERTY(EditAnywhere, Category=ImportSettings, meta=(DisplayName = "Animation Length"))
@@ -63,4 +61,4 @@ public:
 
 	virtual bool CanEditChange(const UProperty* InProperty) const;
 
-};
+};
\ No newline at end of file
diff --git a/Engine/Source/Editor/UnrealEd/Classes/Factories/FbxAssetImportData.h b/Engine/Source/Editor/UnrealEd/Classes/Factories/FbxAssetImportData.h
--- a/Engine/Source/Editor/UnrealEd/Classes/Factories/FbxAssetImportData.h
+++ b/Engine/Source/Editor/UnrealEd/Classes/Factories/FbxAssetImportData.h
@@ -11,9 +11,7 @@
 UCLASS(config=EditorUserSettings, HideCategories=Object, abstract)
 class UNREALED_API UFbxAssetImportData : public UAssetImportData
 {
-	GENERATED_BODY()
-public:
-	UFbxAssetImportData(const FObjectInitializer& ObjectInitializer = FObjectInitializer::Get());
+	GENERATED_UCLASS_BODY()
 
 	UPROPERTY(EditAnywhere, Category=Transform, meta=(ImportType="StaticMesh|SkeletalMesh|Animation", ImportCategory="Transform"))
 	FVector ImportTranslation;
@@ -23,4 +21,4 @@ public:
 
 	UPROPERTY(EditAnywhere, Category=Transform, meta=(ImportType="StaticMesh|SkeletalMesh|Animation", ImportCategory="Transform"))
 	float ImportUniformScale;
-};
+};
\ No newline at end of file
diff --git a/Engine/Source/Editor/UnrealEd/Classes/Factories/FbxFactory.h b/Engine/Source/Editor/UnrealEd/Classes/Factories/FbxFactory.h
--- a/Engine/Source/Editor/UnrealEd/Classes/Factories/FbxFactory.h
+++ b/Engine/Source/Editor/UnrealEd/Classes/Factories/FbxFactory.h
@@ -6,9 +6,7 @@
 UCLASS(hidecategories=Object)
 class UNREALED_API UFbxFactory : public UFactory
 {
-	GENERATED_BODY()
-public:
-	UFbxFactory(const FObjectInitializer& ObjectInitializer = FObjectInitializer::Get());
+	GENERATED_UCLASS_BODY()
 
 	UPROPERTY()
 	class UFbxImportUI* ImportUI;
diff --git a/Engine/Source/Editor/UnrealEd/Classes/Factories/FbxMeshImportData.h b/Engine/Source/Editor/UnrealEd/Classes/Factories/FbxMeshImportData.h
--- a/Engine/Source/Editor/UnrealEd/Classes/Factories/FbxMeshImportData.h
+++ b/Engine/Source/Editor/UnrealEd/Classes/Factories/FbxMeshImportData.h
@@ -18,9 +18,7 @@ enum EFBXNormalImportMethod
 UCLASS(config=EditorUserSettings, configdonotcheckdefaults, abstract)
 class UFbxMeshImportData : public UFbxAssetImportData
 {
-	GENERATED_BODY()
-public:
-	UFbxMeshImportData(const FObjectInitializer& ObjectInitializer = FObjectInitializer::Get());
+	GENERATED_UCLASS_BODY()
 
 	/** Enables importing of mesh LODs from FBX LOD groups, if present in the FBX file */
 	UPROPERTY(EditAnywhere, AdvancedDisplay, config, Category=ImportSettings, meta=(OBJRestrict="true", ImportType="Mesh", ToolTip="If enabled, creates LOD models for Unreal meshes from LODs in the import file; If not enabled, only the base mesh from the LOD group is imported"))
@@ -31,4 +29,4 @@ public:
 	TEnumAsByte<enum EFBXNormalImportMethod> NormalImportMethod;
 
 	bool CanEditChange( const UProperty* InProperty ) const;
-};
+};
\ No newline at end of file
diff --git a/Engine/Source/Editor/UnrealEd/Classes/Factories/FbxSkeletalMeshImportData.h b/Engine/Source/Editor/UnrealEd/Classes/Factories/FbxSkeletalMeshImportData.h
--- a/Engine/Source/Editor/UnrealEd/Classes/Factories/FbxSkeletalMeshImportData.h
+++ b/Engine/Source/Editor/UnrealEd/Classes/Factories/FbxSkeletalMeshImportData.h
@@ -9,9 +9,7 @@
 UCLASS()
 class UFbxSkeletalMeshImportData : public UFbxMeshImportData
 {
-	GENERATED_BODY()
-public:
-	UFbxSkeletalMeshImportData(const FObjectInitializer& ObjectInitializer = FObjectInitializer::Get());
+	GENERATED_UCLASS_BODY()
 
 	/** Enable this option to update Skeleton (of the mesh)'s reference pose. Mesh's reference pose is always updated.  */
 	UPROPERTY(EditAnywhere, AdvancedDisplay, Category=ImportSettings, meta=(ImportType="SkeletalMesh", ToolTip="If enabled, update the Skeleton (of the mesh being imported)'s reference pose."))
@@ -41,4 +39,4 @@ public:
 	static UFbxSkeletalMeshImportData* GetImportDataForSkeletalMesh(USkeletalMesh* SkeletalMesh, UFbxSkeletalMeshImportData* TemplateForCreation);
 
 	bool CanEditChange( const UProperty* InProperty ) const;
-};
+};
\ No newline at end of file
diff --git a/Engine/Source/Editor/UnrealEd/Classes/Factories/FbxStaticMeshImportData.h b/Engine/Source/Editor/UnrealEd/Classes/Factories/FbxStaticMeshImportData.h
--- a/Engine/Source/Editor/UnrealEd/Classes/Factories/FbxStaticMeshImportData.h
+++ b/Engine/Source/Editor/UnrealEd/Classes/Factories/FbxStaticMeshImportData.h
@@ -24,9 +24,7 @@ namespace EVertexColorImportOption
 UCLASS(config=EditorUserSettings, AutoExpandCategories=(Options), MinimalAPI)
 class UFbxStaticMeshImportData : public UFbxMeshImportData
 {
-	GENERATED_BODY()
-public:
-	UNREALED_API UFbxStaticMeshImportData(const FObjectInitializer& ObjectInitializer = FObjectInitializer::Get());
+	GENERATED_UCLASS_BODY()
 
 	/** For static meshes, enabling this option will combine all meshes in the FBX into a single monolithic mesh in Unreal */
 	UPROPERTY(EditAnywhere, AdvancedDisplay, Category=ImportSettings, meta=(ImportType="StaticMesh"))
diff --git a/Engine/Source/Editor/UnrealEd/Classes/Factories/FbxTextureImportData.h b/Engine/Source/Editor/UnrealEd/Classes/Factories/FbxTextureImportData.h
--- a/Engine/Source/Editor/UnrealEd/Classes/Factories/FbxTextureImportData.h
+++ b/Engine/Source/Editor/UnrealEd/Classes/Factories/FbxTextureImportData.h
@@ -9,13 +9,11 @@
 UCLASS(AutoExpandCategories=(Texture))
 class UNREALED_API UFbxTextureImportData : public UFbxAssetImportData
 {
-	GENERATED_BODY()
-public:
-	UFbxTextureImportData(const FObjectInitializer& ObjectInitializer = FObjectInitializer::Get());
+	GENERATED_UCLASS_BODY()
 
 	/** If either importing of textures (or materials) is enabled, this option will cause normal map values to be inverted */
 	UPROPERTY(EditAnywhere, AdvancedDisplay, config, Category=ImportSettings, meta=(OBJRestrict="true"))
 	uint32 bInvertNormalMaps:1;
 
 	bool CanEditChange( const UProperty* InProperty ) const;
-};
+};
\ No newline at end of file
diff --git a/Engine/Source/Editor/UnrealEd/Classes/Factories/FontFactory.h b/Engine/Source/Editor/UnrealEd/Classes/Factories/FontFactory.h
--- a/Engine/Source/Editor/UnrealEd/Classes/Factories/FontFactory.h
+++ b/Engine/Source/Editor/UnrealEd/Classes/Factories/FontFactory.h
@@ -10,9 +10,7 @@
 UCLASS()
 class UNREALED_API UFontFactory : public UFactory
 {
-	GENERATED_BODY()
-public:
-	UFontFactory(const FObjectInitializer& ObjectInitializer = FObjectInitializer::Get());
+	GENERATED_UCLASS_BODY()
 
 	// Begin UFactory Interface
 	virtual UObject* FactoryCreateNew(UClass* InClass, UObject* InParent, FName InName, EObjectFlags Flags, UObject* Context, FFeedbackContext* Warn) override;
diff --git a/Engine/Source/Editor/UnrealEd/Classes/Factories/FontFileImportFactory.h b/Engine/Source/Editor/UnrealEd/Classes/Factories/FontFileImportFactory.h
--- a/Engine/Source/Editor/UnrealEd/Classes/Factories/FontFileImportFactory.h
+++ b/Engine/Source/Editor/UnrealEd/Classes/Factories/FontFileImportFactory.h
@@ -10,9 +10,7 @@
 UCLASS()
 class UNREALED_API UFontFileImportFactory : public UFactory
 {
-	GENERATED_BODY()
-public:
-	UFontFileImportFactory(const FObjectInitializer& ObjectInitializer = FObjectInitializer::Get());
+	GENERATED_UCLASS_BODY()
 
 	// Begin UFactory Interface
 	virtual UObject* FactoryCreateBinary(UClass* Class, UObject* InParent, FName Name, EObjectFlags Flags, UObject* Context, const TCHAR* Type, const uint8*& Buffer, const uint8* BufferEnd, FFeedbackContext* Warn) override;
diff --git a/Engine/Source/Editor/UnrealEd/Classes/Factories/ForceFeedbackEffectFactory.h b/Engine/Source/Editor/UnrealEd/Classes/Factories/ForceFeedbackEffectFactory.h
--- a/Engine/Source/Editor/UnrealEd/Classes/Factories/ForceFeedbackEffectFactory.h
+++ b/Engine/Source/Editor/UnrealEd/Classes/Factories/ForceFeedbackEffectFactory.h
@@ -6,9 +6,7 @@
 UCLASS(hidecategories=Object)
 class UForceFeedbackEffectFactory : public UFactory
 {
-	GENERATED_BODY()
-public:
-	UForceFeedbackEffectFactory(const FObjectInitializer& ObjectInitializer = FObjectInitializer::Get());
+	GENERATED_UCLASS_BODY()
 
 
 	// Begin UFactory Interface
diff --git a/Engine/Source/Editor/UnrealEd/Classes/Factories/InterpDataFactoryNew.h b/Engine/Source/Editor/UnrealEd/Classes/Factories/InterpDataFactoryNew.h
--- a/Engine/Source/Editor/UnrealEd/Classes/Factories/InterpDataFactoryNew.h
+++ b/Engine/Source/Editor/UnrealEd/Classes/Factories/InterpDataFactoryNew.h
@@ -10,9 +10,7 @@
 UCLASS(hidecategories=Object)
 class UInterpDataFactoryNew : public UFactory
 {
-	GENERATED_BODY()
-public:
-	UInterpDataFactoryNew(const FObjectInitializer& ObjectInitializer = FObjectInitializer::Get());
+	GENERATED_UCLASS_BODY()
 
 
 	// Begin UFactory Interface
diff --git a/Engine/Source/Editor/UnrealEd/Classes/Factories/LevelFactory.h b/Engine/Source/Editor/UnrealEd/Classes/Factories/LevelFactory.h
--- a/Engine/Source/Editor/UnrealEd/Classes/Factories/LevelFactory.h
+++ b/Engine/Source/Editor/UnrealEd/Classes/Factories/LevelFactory.h
@@ -10,9 +10,7 @@
 UCLASS(MinimalAPI)
 class ULevelFactory : public UFactory
 {
-	GENERATED_BODY()
-public:
-	UNREALED_API ULevelFactory(const FObjectInitializer& ObjectInitializer = FObjectInitializer::Get());
+	GENERATED_UCLASS_BODY()
 
 
 	// Begin UFactory Interface
diff --git a/Engine/Source/Editor/UnrealEd/Classes/Factories/MaterialFunctionFactoryNew.h b/Engine/Source/Editor/UnrealEd/Classes/Factories/MaterialFunctionFactoryNew.h
--- a/Engine/Source/Editor/UnrealEd/Classes/Factories/MaterialFunctionFactoryNew.h
+++ b/Engine/Source/Editor/UnrealEd/Classes/Factories/MaterialFunctionFactoryNew.h
@@ -10,9 +10,7 @@
 UCLASS(hidecategories=Object, collapsecategories)
 class UMaterialFunctionFactoryNew : public UFactory
 {
-	GENERATED_BODY()
-public:
-	UMaterialFunctionFactoryNew(const FObjectInitializer& ObjectInitializer = FObjectInitializer::Get());
+	GENERATED_UCLASS_BODY()
 
 
 	// Begin UFactory Interface
diff --git a/Engine/Source/Editor/UnrealEd/Classes/Factories/MaterialInstanceConstantFactoryNew.h b/Engine/Source/Editor/UnrealEd/Classes/Factories/MaterialInstanceConstantFactoryNew.h
--- a/Engine/Source/Editor/UnrealEd/Classes/Factories/MaterialInstanceConstantFactoryNew.h
+++ b/Engine/Source/Editor/UnrealEd/Classes/Factories/MaterialInstanceConstantFactoryNew.h
@@ -10,9 +10,7 @@
 UCLASS(hidecategories=Object, collapsecategories,MinimalAPI)
 class UMaterialInstanceConstantFactoryNew : public UFactory
 {
-	GENERATED_BODY()
-public:
-	UNREALED_API UMaterialInstanceConstantFactoryNew(const FObjectInitializer& ObjectInitializer = FObjectInitializer::Get());
+	GENERATED_UCLASS_BODY()
 
 	UPROPERTY()
 	class UMaterialInterface* InitialParent;
diff --git a/Engine/Source/Editor/UnrealEd/Classes/Factories/MaterialParameterCollectionFactoryNew.h b/Engine/Source/Editor/UnrealEd/Classes/Factories/MaterialParameterCollectionFactoryNew.h
--- a/Engine/Source/Editor/UnrealEd/Classes/Factories/MaterialParameterCollectionFactoryNew.h
+++ b/Engine/Source/Editor/UnrealEd/Classes/Factories/MaterialParameterCollectionFactoryNew.h
@@ -10,9 +10,7 @@
 UCLASS(hidecategories=Object, collapsecategories)
 class UMaterialParameterCollectionFactoryNew : public UFactory
 {
-	GENERATED_BODY()
-public:
-	UMaterialParameterCollectionFactoryNew(const FObjectInitializer& ObjectInitializer = FObjectInitializer::Get());
+	GENERATED_UCLASS_BODY()
 
 	// Begin UFactory Interface
 	virtual UObject* FactoryCreateNew(UClass* Class,UObject* InParent,FName Name,EObjectFlags Flags,UObject* Context,FFeedbackContext* Warn) override;
diff --git a/Engine/Source/Editor/UnrealEd/Classes/Factories/ModelFactory.h b/Engine/Source/Editor/UnrealEd/Classes/Factories/ModelFactory.h
--- a/Engine/Source/Editor/UnrealEd/Classes/Factories/ModelFactory.h
+++ b/Engine/Source/Editor/UnrealEd/Classes/Factories/ModelFactory.h
@@ -10,9 +10,7 @@
 UCLASS()
 class UModelFactory : public UFactory
 {
-	GENERATED_BODY()
-public:
-	UModelFactory(const FObjectInitializer& ObjectInitializer = FObjectInitializer::Get());
+	GENERATED_UCLASS_BODY()
 
 
 	// Begin UFactory Interface
diff --git a/Engine/Source/Editor/UnrealEd/Classes/Factories/NiagaraScriptFactoryNew.h b/Engine/Source/Editor/UnrealEd/Classes/Factories/NiagaraScriptFactoryNew.h
--- a/Engine/Source/Editor/UnrealEd/Classes/Factories/NiagaraScriptFactoryNew.h
+++ b/Engine/Source/Editor/UnrealEd/Classes/Factories/NiagaraScriptFactoryNew.h
@@ -6,9 +6,7 @@
 UCLASS(hidecategories=Object)
 class UNiagaraScriptFactoryNew : public UFactory
 {
-	GENERATED_BODY()
-public:
-	UNiagaraScriptFactoryNew(const FObjectInitializer& ObjectInitializer = FObjectInitializer::Get());
+	GENERATED_UCLASS_BODY()
 
 	// Begin UFactory Interface
 	virtual UObject* FactoryCreateNew(UClass* Class,UObject* InParent,FName Name,EObjectFlags Flags,UObject* Context,FFeedbackContext* Warn) override;
diff --git a/Engine/Source/Editor/UnrealEd/Classes/Factories/ObjectLibraryFactory.h b/Engine/Source/Editor/UnrealEd/Classes/Factories/ObjectLibraryFactory.h
--- a/Engine/Source/Editor/UnrealEd/Classes/Factories/ObjectLibraryFactory.h
+++ b/Engine/Source/Editor/UnrealEd/Classes/Factories/ObjectLibraryFactory.h
@@ -6,9 +6,7 @@
 UCLASS(hidecategories=Object)
 class UObjectLibraryFactory : public UFactory
 {
-	GENERATED_BODY()
-public:
-	UObjectLibraryFactory(const FObjectInitializer& ObjectInitializer = FObjectInitializer::Get());
+	GENERATED_UCLASS_BODY()
 
 	// Begin UFactory Interface
 	virtual UObject* FactoryCreateNew(UClass* Class,UObject* InParent,FName Name,EObjectFlags Flags,UObject* Context,FFeedbackContext* Warn) override;
diff --git a/Engine/Source/Editor/UnrealEd/Classes/Factories/PackFactory.h b/Engine/Source/Editor/UnrealEd/Classes/Factories/PackFactory.h
--- a/Engine/Source/Editor/UnrealEd/Classes/Factories/PackFactory.h
+++ b/Engine/Source/Editor/UnrealEd/Classes/Factories/PackFactory.h
@@ -9,11 +9,9 @@
 UCLASS(hidecategories=Object)
 class UPackFactory : public UFactory
 {
-	GENERATED_BODY()
-public:
-	UPackFactory(const FObjectInitializer& ObjectInitializer = FObjectInitializer::Get());
+	GENERATED_UCLASS_BODY()
 
 	// Begin UFactory Interface
 	virtual UObject* FactoryCreateBinary( UClass* Class, UObject* InParent, FName Name, EObjectFlags Flags, UObject* Context, const TCHAR* Type, const uint8*& Buffer, const uint8* BufferEnd, FFeedbackContext* Warn ) override;
 	// End UFactory Interface
-};
+};
\ No newline at end of file
diff --git a/Engine/Source/Editor/UnrealEd/Classes/Factories/PackageFactory.h b/Engine/Source/Editor/UnrealEd/Classes/Factories/PackageFactory.h
--- a/Engine/Source/Editor/UnrealEd/Classes/Factories/PackageFactory.h
+++ b/Engine/Source/Editor/UnrealEd/Classes/Factories/PackageFactory.h
@@ -10,9 +10,7 @@
 UCLASS()
 class UPackageFactory : public UFactory
 {
-	GENERATED_BODY()
-public:
-	UPackageFactory(const FObjectInitializer& ObjectInitializer = FObjectInitializer::Get());
+	GENERATED_UCLASS_BODY()
 
 
 	// Begin UFactory Interface
diff --git a/Engine/Source/Editor/UnrealEd/Classes/Factories/ParticleSystemFactoryNew.h b/Engine/Source/Editor/UnrealEd/Classes/Factories/ParticleSystemFactoryNew.h
--- a/Engine/Source/Editor/UnrealEd/Classes/Factories/ParticleSystemFactoryNew.h
+++ b/Engine/Source/Editor/UnrealEd/Classes/Factories/ParticleSystemFactoryNew.h
@@ -10,9 +10,7 @@
 UCLASS(hidecategories=Object)
 class UParticleSystemFactoryNew : public UFactory
 {
-	GENERATED_BODY()
-public:
-	UParticleSystemFactoryNew(const FObjectInitializer& ObjectInitializer = FObjectInitializer::Get());
+	GENERATED_UCLASS_BODY()
 
 
 	// Begin UFactory Interface
diff --git a/Engine/Source/Editor/UnrealEd/Classes/Factories/PhysicalMaterialFactoryNew.h b/Engine/Source/Editor/UnrealEd/Classes/Factories/PhysicalMaterialFactoryNew.h
--- a/Engine/Source/Editor/UnrealEd/Classes/Factories/PhysicalMaterialFactoryNew.h
+++ b/Engine/Source/Editor/UnrealEd/Classes/Factories/PhysicalMaterialFactoryNew.h
@@ -10,9 +10,7 @@
 UCLASS(hidecategories=Object)
 class UPhysicalMaterialFactoryNew : public UFactory
 {
-	GENERATED_BODY()
-public:
-	UPhysicalMaterialFactoryNew(const FObjectInitializer& ObjectInitializer = FObjectInitializer::Get());
+	GENERATED_UCLASS_BODY()
 
 
 	// Begin UFactory Interface
diff --git a/Engine/Source/Editor/UnrealEd/Classes/Factories/PolysFactory.h b/Engine/Source/Editor/UnrealEd/Classes/Factories/PolysFactory.h
--- a/Engine/Source/Editor/UnrealEd/Classes/Factories/PolysFactory.h
+++ b/Engine/Source/Editor/UnrealEd/Classes/Factories/PolysFactory.h
@@ -10,9 +10,7 @@
 UCLASS()
 class UPolysFactory : public UFactory
 {
-	GENERATED_BODY()
-public:
-	UPolysFactory(const FObjectInitializer& ObjectInitializer = FObjectInitializer::Get());
+	GENERATED_UCLASS_BODY()
 
 
 	// Begin UFactory Interface
diff --git a/Engine/Source/Editor/UnrealEd/Classes/Factories/ReimportCurveFactory.h b/Engine/Source/Editor/UnrealEd/Classes/Factories/ReimportCurveFactory.h
--- a/Engine/Source/Editor/UnrealEd/Classes/Factories/ReimportCurveFactory.h
+++ b/Engine/Source/Editor/UnrealEd/Classes/Factories/ReimportCurveFactory.h
@@ -6,9 +6,7 @@
 UCLASS()
 class UReimportCurveFactory : public UCSVImportFactory, public FReimportHandler
 {
-	GENERATED_BODY()
-public:
-	UReimportCurveFactory(const FObjectInitializer& ObjectInitializer = FObjectInitializer::Get());
+	GENERATED_UCLASS_BODY()
 
 
 	// Begin FReimportHandler interface
diff --git a/Engine/Source/Editor/UnrealEd/Classes/Factories/ReimportDataTableFactory.h b/Engine/Source/Editor/UnrealEd/Classes/Factories/ReimportDataTableFactory.h
--- a/Engine/Source/Editor/UnrealEd/Classes/Factories/ReimportDataTableFactory.h
+++ b/Engine/Source/Editor/UnrealEd/Classes/Factories/ReimportDataTableFactory.h
@@ -6,9 +6,7 @@
 UCLASS()
 class UReimportDataTableFactory : public UCSVImportFactory, public FReimportHandler
 {
-	GENERATED_BODY()
-public:
-	UReimportDataTableFactory(const FObjectInitializer& ObjectInitializer = FObjectInitializer::Get());
+	GENERATED_UCLASS_BODY()
 
 	// Begin FReimportHandler interface
 	virtual bool CanReimport( UObject* Obj, TArray<FString>& OutFilenames ) override;
diff --git a/Engine/Source/Editor/UnrealEd/Classes/Factories/ReimportDestructibleMeshFactory.h b/Engine/Source/Editor/UnrealEd/Classes/Factories/ReimportDestructibleMeshFactory.h
--- a/Engine/Source/Editor/UnrealEd/Classes/Factories/ReimportDestructibleMeshFactory.h
+++ b/Engine/Source/Editor/UnrealEd/Classes/Factories/ReimportDestructibleMeshFactory.h
@@ -8,9 +8,7 @@
 UCLASS(hideCategories=Object)
 class UReimportDestructibleMeshFactory : public UDestructibleMeshFactory, public FReimportHandler
 {
-    GENERATED_BODY()
-public:
-    UReimportDestructibleMeshFactory(const FObjectInitializer& ObjectInitializer = FObjectInitializer::Get());
+    GENERATED_UCLASS_BODY()
 public:
 	virtual FText GetDisplayName() const override;
 #if WITH_APEX
diff --git a/Engine/Source/Editor/UnrealEd/Classes/Factories/ReimportFbxAnimSequenceFactory.h b/Engine/Source/Editor/UnrealEd/Classes/Factories/ReimportFbxAnimSequenceFactory.h
--- a/Engine/Source/Editor/UnrealEd/Classes/Factories/ReimportFbxAnimSequenceFactory.h
+++ b/Engine/Source/Editor/UnrealEd/Classes/Factories/ReimportFbxAnimSequenceFactory.h
@@ -10,9 +10,7 @@
 UCLASS(collapsecategories)
 class UReimportFbxAnimSequenceFactory : public UFbxFactory, public FReimportHandler
 {
-	GENERATED_BODY()
-public:
-	UReimportFbxAnimSequenceFactory(const FObjectInitializer& ObjectInitializer = FObjectInitializer::Get());
+	GENERATED_UCLASS_BODY()
 
 
 	// Begin FReimportHandler interface
diff --git a/Engine/Source/Editor/UnrealEd/Classes/Factories/ReimportFbxSkeletalMeshFactory.h b/Engine/Source/Editor/UnrealEd/Classes/Factories/ReimportFbxSkeletalMeshFactory.h
--- a/Engine/Source/Editor/UnrealEd/Classes/Factories/ReimportFbxSkeletalMeshFactory.h
+++ b/Engine/Source/Editor/UnrealEd/Classes/Factories/ReimportFbxSkeletalMeshFactory.h
@@ -10,9 +10,7 @@
 UCLASS(collapsecategories)
 class UReimportFbxSkeletalMeshFactory : public UFbxFactory, public FReimportHandler
 {
-	GENERATED_BODY()
-public:
-	UReimportFbxSkeletalMeshFactory(const FObjectInitializer& ObjectInitializer = FObjectInitializer::Get());
+	GENERATED_UCLASS_BODY()
 
 
 	// Begin FReimportHandler interface
