Unify Slate's FInputGesture and Engine's FInputChord as FInputChord defined in Slate

diff --git a/Engine/Source/Editor/UnrealEd/Private/EditorViewportClient.cpp b/Engine/Source/Editor/UnrealEd/Private/EditorViewportClient.cpp
--- a/Engine/Source/Editor/UnrealEd/Private/EditorViewportClient.cpp
+++ b/Engine/Source/Editor/UnrealEd/Private/EditorViewportClient.cpp
@@ -1237,13 +1237,13 @@ void FEditorViewportClient::UpdateCameraMovement( float DeltaTime )
 		GestureMoveForwardBackwardImpulse = 0.0f;
 
 		// Forward/back
-		if( ( bRemapWASDKeys && Viewport->KeyState( FViewportNavigationCommands::Get().Forward->GetActiveGesture()->Key ) ) ||
+		if( ( bRemapWASDKeys && Viewport->KeyState( FViewportNavigationCommands::Get().Forward->GetActiveChord()->Key ) ) ||
 			( bRemapArrowKeys && Viewport->KeyState( EKeys::Up ) ) ||
 			Viewport->KeyState( EKeys::NumPadEight ) )
 		{
 			CameraUserImpulseData->MoveForwardBackwardImpulse += 1.0f;
 		}
-		if( ( bRemapWASDKeys && Viewport->KeyState( FViewportNavigationCommands::Get().Backward->GetActiveGesture()->Key ) ) ||
+		if( ( bRemapWASDKeys && Viewport->KeyState( FViewportNavigationCommands::Get().Backward->GetActiveChord()->Key ) ) ||
 			( bRemapArrowKeys && Viewport->KeyState( EKeys::Down ) ) ||
 			Viewport->KeyState( EKeys::NumPadTwo ) )
 		{
@@ -1251,13 +1251,13 @@ void FEditorViewportClient::UpdateCameraMovement( float DeltaTime )
 		}
 
 		// Right/left
-		if ( ( bRemapWASDKeys && Viewport->KeyState( FViewportNavigationCommands::Get().Right->GetActiveGesture()->Key) ) ||
+		if ( ( bRemapWASDKeys && Viewport->KeyState( FViewportNavigationCommands::Get().Right->GetActiveChord()->Key) ) ||
 			( bRemapArrowKeys && Viewport->KeyState( EKeys::Right ) ) ||
 			Viewport->KeyState( EKeys::NumPadSix ) )
 		{
 			CameraUserImpulseData->MoveRightLeftImpulse += 1.0f;
 		}
-		if( ( bRemapWASDKeys && Viewport->KeyState( FViewportNavigationCommands::Get().Left->GetActiveGesture()->Key ) ) ||
+		if( ( bRemapWASDKeys && Viewport->KeyState( FViewportNavigationCommands::Get().Left->GetActiveChord()->Key ) ) ||
 			( bRemapArrowKeys && Viewport->KeyState( EKeys::Left ) ) ||
 			Viewport->KeyState( EKeys::NumPadFour ) )
 		{
@@ -1265,24 +1265,24 @@ void FEditorViewportClient::UpdateCameraMovement( float DeltaTime )
 		}
 
 		// Up/down
-		if( ( bRemapWASDKeys && Viewport->KeyState( FViewportNavigationCommands::Get().Up->GetActiveGesture()->Key ) ) ||
+		if( ( bRemapWASDKeys && Viewport->KeyState( FViewportNavigationCommands::Get().Up->GetActiveChord()->Key ) ) ||
 			Viewport->KeyState( EKeys::PageUp ) || Viewport->KeyState( EKeys::NumPadNine ) || Viewport->KeyState( EKeys::Add ) )
 		{
 			CameraUserImpulseData->MoveUpDownImpulse += 1.0f;
 		}
-		if( ( bRemapWASDKeys && Viewport->KeyState( FViewportNavigationCommands::Get().Down->GetActiveGesture()->Key) ) ||
+		if( ( bRemapWASDKeys && Viewport->KeyState( FViewportNavigationCommands::Get().Down->GetActiveChord()->Key) ) ||
 			Viewport->KeyState( EKeys::PageDown ) || Viewport->KeyState( EKeys::NumPadSeven ) || Viewport->KeyState( EKeys::Subtract ) )
 		{
 			CameraUserImpulseData->MoveUpDownImpulse -= 1.0f;
 		}
 
 		// Zoom FOV out/in
-		if( ( bRemapWASDKeys && Viewport->KeyState( FViewportNavigationCommands::Get().FovZoomOut->GetActiveGesture()->Key ) ) ||
+		if( ( bRemapWASDKeys && Viewport->KeyState( FViewportNavigationCommands::Get().FovZoomOut->GetActiveChord()->Key ) ) ||
 			Viewport->KeyState( EKeys::NumPadOne ) )
 		{
 			CameraUserImpulseData->ZoomOutInImpulse += 1.0f;
 		}
-		if( ( bRemapWASDKeys && Viewport->KeyState( FViewportNavigationCommands::Get().FovZoomIn->GetActiveGesture()->Key ) ) ||
+		if( ( bRemapWASDKeys && Viewport->KeyState( FViewportNavigationCommands::Get().FovZoomIn->GetActiveChord()->Key ) ) ||
 			Viewport->KeyState( EKeys::NumPadThree ) )
 		{
 			CameraUserImpulseData->ZoomOutInImpulse -= 1.0f;
@@ -1888,14 +1888,14 @@ void FEditorViewportClient::SetMatineeRecordingWindow (IMatineeBase* InInterpEd)
 bool FEditorViewportClient::IsFlightCameraActive() const
 {
 	bool bIsFlightMovementKey = 
-		( Viewport->KeyState( FViewportNavigationCommands::Get().Forward->GetActiveGesture()->Key )
-		|| Viewport->KeyState( FViewportNavigationCommands::Get().Backward->GetActiveGesture()->Key )
-		|| Viewport->KeyState( FViewportNavigationCommands::Get().Left->GetActiveGesture()->Key )
-		|| Viewport->KeyState( FViewportNavigationCommands::Get().Right->GetActiveGesture()->Key )
-		|| Viewport->KeyState( FViewportNavigationCommands::Get().Up->GetActiveGesture()->Key )
-		|| Viewport->KeyState( FViewportNavigationCommands::Get().Down->GetActiveGesture()->Key )
-		|| Viewport->KeyState( FViewportNavigationCommands::Get().FovZoomIn->GetActiveGesture()->Key )
-		|| Viewport->KeyState( FViewportNavigationCommands::Get().FovZoomOut->GetActiveGesture()->Key ) );
+		( Viewport->KeyState( FViewportNavigationCommands::Get().Forward->GetActiveChord()->Key )
+		|| Viewport->KeyState( FViewportNavigationCommands::Get().Backward->GetActiveChord()->Key )
+		|| Viewport->KeyState( FViewportNavigationCommands::Get().Left->GetActiveChord()->Key )
+		|| Viewport->KeyState( FViewportNavigationCommands::Get().Right->GetActiveChord()->Key )
+		|| Viewport->KeyState( FViewportNavigationCommands::Get().Up->GetActiveChord()->Key )
+		|| Viewport->KeyState( FViewportNavigationCommands::Get().Down->GetActiveChord()->Key )
+		|| Viewport->KeyState( FViewportNavigationCommands::Get().FovZoomIn->GetActiveChord()->Key )
+		|| Viewport->KeyState( FViewportNavigationCommands::Get().FovZoomOut->GetActiveChord()->Key ) );
 
 	const bool bIsUsingTrackpad = FSlateApplication::Get().IsUsingTrackpad();
 
@@ -4790,16 +4790,16 @@ FViewportNavigationCommands::FViewportNavigationCommands()
 
 void FViewportNavigationCommands::RegisterCommands()
 {
-	UI_COMMAND(Forward, "Forward", "Moves the camera Forward", EUserInterfaceActionType::Button, FInputGesture(EKeys::W));
-	UI_COMMAND(Backward, "Backward", "Moves the camera Backward", EUserInterfaceActionType::Button, FInputGesture(EKeys::S));
-	UI_COMMAND(Left, "Left", "Moves the camera Left", EUserInterfaceActionType::Button, FInputGesture(EKeys::A));
-	UI_COMMAND(Right, "Right", "Moves the camera Right", EUserInterfaceActionType::Button, FInputGesture(EKeys::D));
+	UI_COMMAND(Forward, "Forward", "Moves the camera Forward", EUserInterfaceActionType::Button, FInputChord(EKeys::W));
+	UI_COMMAND(Backward, "Backward", "Moves the camera Backward", EUserInterfaceActionType::Button, FInputChord(EKeys::S));
+	UI_COMMAND(Left, "Left", "Moves the camera Left", EUserInterfaceActionType::Button, FInputChord(EKeys::A));
+	UI_COMMAND(Right, "Right", "Moves the camera Right", EUserInterfaceActionType::Button, FInputChord(EKeys::D));
 
-	UI_COMMAND(Up, "Up", "Moves the camera Up", EUserInterfaceActionType::Button, FInputGesture(EKeys::E));
-	UI_COMMAND(Down, "Down", "Moves the camera Down", EUserInterfaceActionType::Button, FInputGesture(EKeys::Q));
+	UI_COMMAND(Up, "Up", "Moves the camera Up", EUserInterfaceActionType::Button, FInputChord(EKeys::E));
+	UI_COMMAND(Down, "Down", "Moves the camera Down", EUserInterfaceActionType::Button, FInputChord(EKeys::Q));
 
-	UI_COMMAND(FovZoomIn, "FOV Zoom In", "Narrows the camers FOV", EUserInterfaceActionType::Button, FInputGesture(EKeys::C));
-	UI_COMMAND(FovZoomOut, "FOV Zoom Out", "Widens the camera FOV", EUserInterfaceActionType::Button, FInputGesture(EKeys::Z));
+	UI_COMMAND(FovZoomIn, "FOV Zoom In", "Narrows the camers FOV", EUserInterfaceActionType::Button, FInputChord(EKeys::C));
+	UI_COMMAND(FovZoomOut, "FOV Zoom Out", "Widens the camera FOV", EUserInterfaceActionType::Button, FInputChord(EKeys::Z));
 }
 
 #undef LOCTEXT_NAMESPACE 
diff --git a/Engine/Source/Editor/UnrealEd/Private/EditorViewportCommands.cpp b/Engine/Source/Editor/UnrealEd/Private/EditorViewportCommands.cpp
--- a/Engine/Source/Editor/UnrealEd/Private/EditorViewportCommands.cpp
+++ b/Engine/Source/Editor/UnrealEd/Private/EditorViewportCommands.cpp
@@ -8,76 +8,76 @@
 
 void FEditorViewportCommands::RegisterCommands()
 {
-	UI_COMMAND( Perspective, "Perspective", "Switches the viewport to perspective view", EUserInterfaceActionType::RadioButton, FInputGesture( EModifierKey::Alt, EKeys::G ) );
-	UI_COMMAND( Front, "Front", "Switches the viewport to front view", EUserInterfaceActionType::RadioButton, FInputGesture( EModifierKey::Alt, EKeys::H ) );
-	UI_COMMAND( Back, "Back", "Switches the viewport to back view", EUserInterfaceActionType::RadioButton, FInputGesture( EModifierKey::Alt | EModifierKey::Shift, EKeys::H ) );
-	UI_COMMAND( Top, "Top", "Switches the viewport to top view", EUserInterfaceActionType::RadioButton, FInputGesture( EModifierKey::Alt, EKeys::J ) );
-	UI_COMMAND( Bottom, "Bottom", "Switches the viewport to bottom view", EUserInterfaceActionType::RadioButton, FInputGesture( EModifierKey::Alt | EModifierKey::Shift, EKeys::J ) );
-	UI_COMMAND( Left, "Left", "Switches the viewport to left view", EUserInterfaceActionType::RadioButton, FInputGesture( EModifierKey::Alt, EKeys::K ) );
-	UI_COMMAND( Right, "Right", "Switches the viewport to right view", EUserInterfaceActionType::RadioButton, FInputGesture( EModifierKey::Alt | EModifierKey::Shift, EKeys::K ) );
+	UI_COMMAND( Perspective, "Perspective", "Switches the viewport to perspective view", EUserInterfaceActionType::RadioButton, FInputChord( EModifierKey::Alt, EKeys::G ) );
+	UI_COMMAND( Front, "Front", "Switches the viewport to front view", EUserInterfaceActionType::RadioButton, FInputChord( EModifierKey::Alt, EKeys::H ) );
+	UI_COMMAND( Back, "Back", "Switches the viewport to back view", EUserInterfaceActionType::RadioButton, FInputChord( EModifierKey::Alt | EModifierKey::Shift, EKeys::H ) );
+	UI_COMMAND( Top, "Top", "Switches the viewport to top view", EUserInterfaceActionType::RadioButton, FInputChord( EModifierKey::Alt, EKeys::J ) );
+	UI_COMMAND( Bottom, "Bottom", "Switches the viewport to bottom view", EUserInterfaceActionType::RadioButton, FInputChord( EModifierKey::Alt | EModifierKey::Shift, EKeys::J ) );
+	UI_COMMAND( Left, "Left", "Switches the viewport to left view", EUserInterfaceActionType::RadioButton, FInputChord( EModifierKey::Alt, EKeys::K ) );
+	UI_COMMAND( Right, "Right", "Switches the viewport to right view", EUserInterfaceActionType::RadioButton, FInputChord( EModifierKey::Alt | EModifierKey::Shift, EKeys::K ) );
 
-	UI_COMMAND( WireframeMode, "Brush Wireframe View Mode", "Renders the scene in brush wireframe", EUserInterfaceActionType::RadioButton, FInputGesture( EModifierKey::Alt, EKeys::Two ) );
-	UI_COMMAND( UnlitMode, "Unlit View Mode", "Renders the scene with no lights", EUserInterfaceActionType::RadioButton, FInputGesture( EModifierKey::Alt, EKeys::Three ) );
-	UI_COMMAND( LitMode, "Lit View Mode", "Renders the scene with normal lighting", EUserInterfaceActionType::RadioButton, FInputGesture( EModifierKey::Alt, EKeys::Four ) );
-	UI_COMMAND( DetailLightingMode, "Detail Lighting View Mode", "Renders the scene with detailed lighting only", EUserInterfaceActionType::RadioButton, FInputGesture( EModifierKey::Alt, EKeys::Five ) );
-	UI_COMMAND( LightingOnlyMode, "Lighting Only View Mode", "Renders the scene with lights only, no textures", EUserInterfaceActionType::RadioButton, FInputGesture( EModifierKey::Alt, EKeys::Six ) );
-	UI_COMMAND( LightComplexityMode, "Light Complexity View Mode", "Renders the scene with light complexity visualization", EUserInterfaceActionType::RadioButton, FInputGesture( EModifierKey::Alt, EKeys::Seven ) );
-	UI_COMMAND( ShaderComplexityMode, "Shader Complexity View Mode", "Renders the scene with shader complexity visualization", EUserInterfaceActionType::RadioButton, FInputGesture( EModifierKey::Alt, EKeys::Eight ) );
-	UI_COMMAND( StationaryLightOverlapMode, "Stationary Light Overlap View Mode", "Visualizes overlap of stationary lights", EUserInterfaceActionType::RadioButton, FInputGesture() );
-	UI_COMMAND( LightmapDensityMode, "Lightmap Density View Mode", "Renders the scene with lightmap density visualization", EUserInterfaceActionType::RadioButton, FInputGesture( EModifierKey::Alt, EKeys::Zero ) );
-	UI_COMMAND( VisualizeBufferMode, "Buffer Visualization View Mode", "Renders a set of selected post process materials, which visualize various intermediate render buffers (material attributes)", EUserInterfaceActionType::RadioButton, FInputGesture() );
-	UI_COMMAND( ReflectionOverrideMode, "Reflections View Mode", "Renders the scene with reflections only", EUserInterfaceActionType::RadioButton, FInputGesture() );
-	UI_COMMAND( CollisionPawn, "Player Collision", "Renders player collision visualization", EUserInterfaceActionType::RadioButton, FInputGesture() );
-	UI_COMMAND( CollisionVisibility, "Visibility Collision", "Renders visibility collision visualization", EUserInterfaceActionType::RadioButton, FInputGesture() );
+	UI_COMMAND( WireframeMode, "Brush Wireframe View Mode", "Renders the scene in brush wireframe", EUserInterfaceActionType::RadioButton, FInputChord( EModifierKey::Alt, EKeys::Two ) );
+	UI_COMMAND( UnlitMode, "Unlit View Mode", "Renders the scene with no lights", EUserInterfaceActionType::RadioButton, FInputChord( EModifierKey::Alt, EKeys::Three ) );
+	UI_COMMAND( LitMode, "Lit View Mode", "Renders the scene with normal lighting", EUserInterfaceActionType::RadioButton, FInputChord( EModifierKey::Alt, EKeys::Four ) );
+	UI_COMMAND( DetailLightingMode, "Detail Lighting View Mode", "Renders the scene with detailed lighting only", EUserInterfaceActionType::RadioButton, FInputChord( EModifierKey::Alt, EKeys::Five ) );
+	UI_COMMAND( LightingOnlyMode, "Lighting Only View Mode", "Renders the scene with lights only, no textures", EUserInterfaceActionType::RadioButton, FInputChord( EModifierKey::Alt, EKeys::Six ) );
+	UI_COMMAND( LightComplexityMode, "Light Complexity View Mode", "Renders the scene with light complexity visualization", EUserInterfaceActionType::RadioButton, FInputChord( EModifierKey::Alt, EKeys::Seven ) );
+	UI_COMMAND( ShaderComplexityMode, "Shader Complexity View Mode", "Renders the scene with shader complexity visualization", EUserInterfaceActionType::RadioButton, FInputChord( EModifierKey::Alt, EKeys::Eight ) );
+	UI_COMMAND( StationaryLightOverlapMode, "Stationary Light Overlap View Mode", "Visualizes overlap of stationary lights", EUserInterfaceActionType::RadioButton, FInputChord() );
+	UI_COMMAND( LightmapDensityMode, "Lightmap Density View Mode", "Renders the scene with lightmap density visualization", EUserInterfaceActionType::RadioButton, FInputChord( EModifierKey::Alt, EKeys::Zero ) );
+	UI_COMMAND( VisualizeBufferMode, "Buffer Visualization View Mode", "Renders a set of selected post process materials, which visualize various intermediate render buffers (material attributes)", EUserInterfaceActionType::RadioButton, FInputChord() );
+	UI_COMMAND( ReflectionOverrideMode, "Reflections View Mode", "Renders the scene with reflections only", EUserInterfaceActionType::RadioButton, FInputChord() );
+	UI_COMMAND( CollisionPawn, "Player Collision", "Renders player collision visualization", EUserInterfaceActionType::RadioButton, FInputChord() );
+	UI_COMMAND( CollisionVisibility, "Visibility Collision", "Renders visibility collision visualization", EUserInterfaceActionType::RadioButton, FInputChord() );
 
-	UI_COMMAND( ToggleRealTime, "Realtime", "Toggles real time rendering in this viewport", EUserInterfaceActionType::ToggleButton, FInputGesture( EModifierKey::Control, EKeys::R ) );
-	UI_COMMAND( ToggleStats, "Show Stats", "Toggles the ability to show stats in this viewport (enables realtime)", EUserInterfaceActionType::ToggleButton, FInputGesture( EModifierKey::Shift, EKeys::L ) );
-	UI_COMMAND( ToggleFPS, "Show FPS", "Toggles showing frames per second in this viewport (enables realtime)", EUserInterfaceActionType::ToggleButton, FInputGesture( EModifierKey::Control|EModifierKey::Shift, EKeys::H ) );
+	UI_COMMAND( ToggleRealTime, "Realtime", "Toggles real time rendering in this viewport", EUserInterfaceActionType::ToggleButton, FInputChord( EModifierKey::Control, EKeys::R ) );
+	UI_COMMAND( ToggleStats, "Show Stats", "Toggles the ability to show stats in this viewport (enables realtime)", EUserInterfaceActionType::ToggleButton, FInputChord( EModifierKey::Shift, EKeys::L ) );
+	UI_COMMAND( ToggleFPS, "Show FPS", "Toggles showing frames per second in this viewport (enables realtime)", EUserInterfaceActionType::ToggleButton, FInputChord( EModifierKey::Control|EModifierKey::Shift, EKeys::H ) );
 
-	UI_COMMAND( ScreenCapture, "Screen Capture", "Take a screenshot of the active viewport.", EUserInterfaceActionType::Button, FInputGesture(EKeys::F9) );
-	UI_COMMAND( ScreenCaptureForProjectThumbnail, "Update Project Thumbnail", "Take a screenshot of the active viewport for use as the project thumbnail.", EUserInterfaceActionType::Button, FInputGesture() );
+	UI_COMMAND( ScreenCapture, "Screen Capture", "Take a screenshot of the active viewport.", EUserInterfaceActionType::Button, FInputChord(EKeys::F9) );
+	UI_COMMAND( ScreenCaptureForProjectThumbnail, "Update Project Thumbnail", "Take a screenshot of the active viewport for use as the project thumbnail.", EUserInterfaceActionType::Button, FInputChord() );
 
-	UI_COMMAND( IncrementPositionGridSize, "Grid Size (Position): Increment", "Increases the position grid size setting by one", EUserInterfaceActionType::Button, FInputGesture( EKeys::RightBracket ) );
-	UI_COMMAND( DecrementPositionGridSize, "Grid Size (Position): Decrement", "Decreases the position grid size setting by one", EUserInterfaceActionType::Button, FInputGesture( EKeys::LeftBracket ) );
-	UI_COMMAND( IncrementRotationGridSize, "Grid Size (Rotation): Increment", "Increases the rotation grid size setting by one", EUserInterfaceActionType::Button, FInputGesture( EModifierKey::Shift, EKeys::RightBracket ) );
-	UI_COMMAND( DecrementRotationGridSize, "Grid Size (Rotation): Decrement", "Decreases the rotation grid size setting by one", EUserInterfaceActionType::Button, FInputGesture( EModifierKey::Shift, EKeys::LeftBracket ) );
+	UI_COMMAND( IncrementPositionGridSize, "Grid Size (Position): Increment", "Increases the position grid size setting by one", EUserInterfaceActionType::Button, FInputChord( EKeys::RightBracket ) );
+	UI_COMMAND( DecrementPositionGridSize, "Grid Size (Position): Decrement", "Decreases the position grid size setting by one", EUserInterfaceActionType::Button, FInputChord( EKeys::LeftBracket ) );
+	UI_COMMAND( IncrementRotationGridSize, "Grid Size (Rotation): Increment", "Increases the rotation grid size setting by one", EUserInterfaceActionType::Button, FInputChord( EModifierKey::Shift, EKeys::RightBracket ) );
+	UI_COMMAND( DecrementRotationGridSize, "Grid Size (Rotation): Decrement", "Decreases the rotation grid size setting by one", EUserInterfaceActionType::Button, FInputChord( EModifierKey::Shift, EKeys::LeftBracket ) );
 
 	
-	UI_COMMAND( TranslateMode, "Translate Mode", "Select and translate objects", EUserInterfaceActionType::ToggleButton, FInputGesture(EKeys::W) );
-	UI_COMMAND( RotateMode, "Rotate Mode", "Select and rotate objects", EUserInterfaceActionType::ToggleButton, FInputGesture(EKeys::E) );
-	UI_COMMAND( ScaleMode, "Scale Mode", "Select and scale objects", EUserInterfaceActionType::ToggleButton, FInputGesture(EKeys::R) );
-	UI_COMMAND( TranslateRotateMode, "Combined Translate and Rotate Mode", "Select and translate or rotate objects", EUserInterfaceActionType::ToggleButton, FInputGesture() );
+	UI_COMMAND( TranslateMode, "Translate Mode", "Select and translate objects", EUserInterfaceActionType::ToggleButton, FInputChord(EKeys::W) );
+	UI_COMMAND( RotateMode, "Rotate Mode", "Select and rotate objects", EUserInterfaceActionType::ToggleButton, FInputChord(EKeys::E) );
+	UI_COMMAND( ScaleMode, "Scale Mode", "Select and scale objects", EUserInterfaceActionType::ToggleButton, FInputChord(EKeys::R) );
+	UI_COMMAND( TranslateRotateMode, "Combined Translate and Rotate Mode", "Select and translate or rotate objects", EUserInterfaceActionType::ToggleButton, FInputChord() );
 
-	UI_COMMAND( ShrinkTransformWidget, "Shrink Transform Widget", "Shrink the level editor transform widget", EUserInterfaceActionType::Button, FInputGesture(EModifierKey::Alt, EKeys::LeftBracket) );
-	UI_COMMAND( ExpandTransformWidget, "Expand Transform Widget", "Expand the level editor transform widget", EUserInterfaceActionType::Button, FInputGesture(EModifierKey::Alt, EKeys::RightBracket) );
+	UI_COMMAND( ShrinkTransformWidget, "Shrink Transform Widget", "Shrink the level editor transform widget", EUserInterfaceActionType::Button, FInputChord(EModifierKey::Alt, EKeys::LeftBracket) );
+	UI_COMMAND( ExpandTransformWidget, "Expand Transform Widget", "Expand the level editor transform widget", EUserInterfaceActionType::Button, FInputChord(EModifierKey::Alt, EKeys::RightBracket) );
 
-	UI_COMMAND( RelativeCoordinateSystem_World, "World-relative Transform", "Move and rotate objects relative to the cardinal world axes", EUserInterfaceActionType::RadioButton, FInputGesture() );
-	UI_COMMAND( RelativeCoordinateSystem_Local, "Local-relative Transform", "Move and rotate objects relative to the object's local axes", EUserInterfaceActionType::RadioButton, FInputGesture() );
+	UI_COMMAND( RelativeCoordinateSystem_World, "World-relative Transform", "Move and rotate objects relative to the cardinal world axes", EUserInterfaceActionType::RadioButton, FInputChord() );
+	UI_COMMAND( RelativeCoordinateSystem_Local, "Local-relative Transform", "Move and rotate objects relative to the object's local axes", EUserInterfaceActionType::RadioButton, FInputChord() );
 
 #if PLATFORM_MAC
-	UI_COMMAND( CycleTransformGizmoCoordSystem, "Cycle Transform Coordinate System", "Cycles the transform gizmo coordinate systems between world and local (object) space", EUserInterfaceActionType::Button, FInputGesture(EModifierKey::Command, EKeys::Tilde));
+	UI_COMMAND( CycleTransformGizmoCoordSystem, "Cycle Transform Coordinate System", "Cycles the transform gizmo coordinate systems between world and local (object) space", EUserInterfaceActionType::Button, FInputChord(EModifierKey::Command, EKeys::Tilde));
 #else
-	UI_COMMAND( CycleTransformGizmoCoordSystem, "Cycle Transform Coordinate System", "Cycles the transform gizmo coordinate systems between world and local (object) space", EUserInterfaceActionType::Button, FInputGesture(EModifierKey::Control, EKeys::Tilde));
+	UI_COMMAND( CycleTransformGizmoCoordSystem, "Cycle Transform Coordinate System", "Cycles the transform gizmo coordinate systems between world and local (object) space", EUserInterfaceActionType::Button, FInputChord(EModifierKey::Control, EKeys::Tilde));
 #endif
-	UI_COMMAND( CycleTransformGizmos, "Cycle Between Translate, Rotate, and Scale", "Cycles the transform gizmos between translate, rotate, and scale", EUserInterfaceActionType::Button, FInputGesture(EKeys::SpaceBar) );
+	UI_COMMAND( CycleTransformGizmos, "Cycle Between Translate, Rotate, and Scale", "Cycles the transform gizmos between translate, rotate, and scale", EUserInterfaceActionType::Button, FInputChord(EKeys::SpaceBar) );
 	
-	UI_COMMAND( FocusViewportToSelection, "Focus Selected", "Moves the camera in front of the selection", EUserInterfaceActionType::Button, FInputGesture( EKeys::F ) );
+	UI_COMMAND( FocusViewportToSelection, "Focus Selected", "Moves the camera in front of the selection", EUserInterfaceActionType::Button, FInputChord( EKeys::F ) );
 
-	UI_COMMAND( LocationGridSnap, "Grid Snap", "Enables or disables snapping to the grid when dragging objects around", EUserInterfaceActionType::ToggleButton, FInputGesture() );
-	UI_COMMAND( RotationGridSnap, "Rotation Snap", "Enables or disables snapping objects to a rotation grid", EUserInterfaceActionType::ToggleButton, FInputGesture() );
-	UI_COMMAND( ScaleGridSnap, "Scale Snap", "Enables or disables snapping objects to a scale grid", EUserInterfaceActionType::ToggleButton, FInputGesture() );
-	UI_COMMAND( SurfaceSnapping, "Surface Snapping", "If enabled, actors will snap to surfaces in the world when dragging", EUserInterfaceActionType::ToggleButton, FInputGesture() );
+	UI_COMMAND( LocationGridSnap, "Grid Snap", "Enables or disables snapping to the grid when dragging objects around", EUserInterfaceActionType::ToggleButton, FInputChord() );
+	UI_COMMAND( RotationGridSnap, "Rotation Snap", "Enables or disables snapping objects to a rotation grid", EUserInterfaceActionType::ToggleButton, FInputChord() );
+	UI_COMMAND( ScaleGridSnap, "Scale Snap", "Enables or disables snapping objects to a scale grid", EUserInterfaceActionType::ToggleButton, FInputChord() );
+	UI_COMMAND( SurfaceSnapping, "Surface Snapping", "If enabled, actors will snap to surfaces in the world when dragging", EUserInterfaceActionType::ToggleButton, FInputChord() );
 
-	UI_COMMAND( ToggleAutoExposure, "Automatic (Default in-game)", "Enable automatic expose", EUserInterfaceActionType::RadioButton, FInputGesture() );
-	UI_COMMAND( FixedExposure4m, "Fixed Exposure: -4", "Set the fixed exposure to -4", EUserInterfaceActionType::RadioButton, FInputGesture() );
-	UI_COMMAND( FixedExposure3m, "Fixed Exposure: -3", "Set the fixed exposure to -3", EUserInterfaceActionType::RadioButton, FInputGesture() );
-	UI_COMMAND( FixedExposure2m, "Fixed Exposure: -2", "Set the fixed exposure to -2", EUserInterfaceActionType::RadioButton, FInputGesture() );
-	UI_COMMAND( FixedExposure1m, "Fixed Exposure: -1", "Set the fixed exposure to -1", EUserInterfaceActionType::RadioButton, FInputGesture() );
-	UI_COMMAND( FixedExposure0, "Fixed Exposure: 0 (Indoor)", "Set the fixed exposure to 0", EUserInterfaceActionType::RadioButton, FInputGesture() );
-	UI_COMMAND( FixedExposure1p, "Fixed Exposure: +1", "Set the fixed exposure to 1", EUserInterfaceActionType::RadioButton, FInputGesture() );
-	UI_COMMAND( FixedExposure2p, "Fixed Exposure: +2", "Set the fixed exposure to 2", EUserInterfaceActionType::RadioButton, FInputGesture() );
-	UI_COMMAND( FixedExposure3p, "Fixed Exposure: +3", "Set the fixed exposure to 3", EUserInterfaceActionType::RadioButton, FInputGesture() );
-	UI_COMMAND( FixedExposure4p, "Fixed Exposure: +4", "Set the fixed exposure to 4", EUserInterfaceActionType::RadioButton, FInputGesture() );
+	UI_COMMAND( ToggleAutoExposure, "Automatic (Default in-game)", "Enable automatic expose", EUserInterfaceActionType::RadioButton, FInputChord() );
+	UI_COMMAND( FixedExposure4m, "Fixed Exposure: -4", "Set the fixed exposure to -4", EUserInterfaceActionType::RadioButton, FInputChord() );
+	UI_COMMAND( FixedExposure3m, "Fixed Exposure: -3", "Set the fixed exposure to -3", EUserInterfaceActionType::RadioButton, FInputChord() );
+	UI_COMMAND( FixedExposure2m, "Fixed Exposure: -2", "Set the fixed exposure to -2", EUserInterfaceActionType::RadioButton, FInputChord() );
+	UI_COMMAND( FixedExposure1m, "Fixed Exposure: -1", "Set the fixed exposure to -1", EUserInterfaceActionType::RadioButton, FInputChord() );
+	UI_COMMAND( FixedExposure0, "Fixed Exposure: 0 (Indoor)", "Set the fixed exposure to 0", EUserInterfaceActionType::RadioButton, FInputChord() );
+	UI_COMMAND( FixedExposure1p, "Fixed Exposure: +1", "Set the fixed exposure to 1", EUserInterfaceActionType::RadioButton, FInputChord() );
+	UI_COMMAND( FixedExposure2p, "Fixed Exposure: +2", "Set the fixed exposure to 2", EUserInterfaceActionType::RadioButton, FInputChord() );
+	UI_COMMAND( FixedExposure3p, "Fixed Exposure: +3", "Set the fixed exposure to 3", EUserInterfaceActionType::RadioButton, FInputChord() );
+	UI_COMMAND( FixedExposure4p, "Fixed Exposure: +4", "Set the fixed exposure to 4", EUserInterfaceActionType::RadioButton, FInputChord() );
 }
 
 #undef LOCTEXT_NAMESPACE
\ No newline at end of file
diff --git a/Engine/Source/Editor/UnrealEd/Private/Kismet2/DebuggerCommands.cpp b/Engine/Source/Editor/UnrealEd/Private/Kismet2/DebuggerCommands.cpp
--- a/Engine/Source/Editor/UnrealEd/Private/Kismet2/DebuggerCommands.cpp
+++ b/Engine/Source/Editor/UnrealEd/Private/Kismet2/DebuggerCommands.cpp
@@ -218,35 +218,35 @@ FPlayWorldCommands::FPlayWorldCommands()
 void FPlayWorldCommands::RegisterCommands()
 {
 	// SIE
-	UI_COMMAND( Simulate, "Simulate", "Start simulating the game", EUserInterfaceActionType::Check, FInputGesture( EKeys::S, EModifierKey::Alt ) );
+	UI_COMMAND( Simulate, "Simulate", "Start simulating the game", EUserInterfaceActionType::Check, FInputChord( EKeys::S, EModifierKey::Alt ) );
 
 	// PIE
-	UI_COMMAND( RepeatLastPlay, "Play", "Launches a game preview session in the same mode as the last game preview session launched from the Game Preview Modes dropdown next to the Play button on the level editor toolbar", EUserInterfaceActionType::Button, FInputGesture( EKeys::P, EModifierKey::Alt ) )
-	UI_COMMAND( PlayInViewport, "Selected Viewport", "Play this level in the active level editor viewport", EUserInterfaceActionType::Check, FInputGesture() );
-	UI_COMMAND( PlayInEditorFloating, "New Editor Window", "Play this level in a new window", EUserInterfaceActionType::Check, FInputGesture() );
-	UI_COMMAND( PlayInVR, "VR Preview", "Play this level in VR", EUserInterfaceActionType::Check, FInputGesture() );
-	UI_COMMAND( PlayInMobilePreview, "Mobile Preview", "Play this level as a mobile device preview (runs in its own process)", EUserInterfaceActionType::Check, FInputGesture() );
-	UI_COMMAND( PlayInNewProcess, "Standalone Game", "Play this level in a new window that runs in its own process", EUserInterfaceActionType::Check, FInputGesture() );
-	UI_COMMAND( PlayInCameraLocation, "Current Camera Location", "Spawn the player at the current camera location", EUserInterfaceActionType::RadioButton, FInputGesture() );
-	UI_COMMAND( PlayInDefaultPlayerStart, "Default Player Start", "Spawn the player at the map's default player start", EUserInterfaceActionType::RadioButton, FInputGesture() );
-	UI_COMMAND( PlayInNetworkSettings, "Network Settings...", "Open the settings for the 'Play In' feature", EUserInterfaceActionType::Button, FInputGesture() );
-	UI_COMMAND( PlayInNetworkDedicatedServer, "Run Dedicated Server", "If checked, a separate dedicated server will be launched. Otherwise the first player will act as a listen server that all other players connect to.", EUserInterfaceActionType::ToggleButton, FInputGesture() );
-	UI_COMMAND( PlayInSettings, "Advanced Settings...", "Open the settings for the 'Play In' feature", EUserInterfaceActionType::Button, FInputGesture() );
+	UI_COMMAND( RepeatLastPlay, "Play", "Launches a game preview session in the same mode as the last game preview session launched from the Game Preview Modes dropdown next to the Play button on the level editor toolbar", EUserInterfaceActionType::Button, FInputChord( EKeys::P, EModifierKey::Alt ) )
+	UI_COMMAND( PlayInViewport, "Selected Viewport", "Play this level in the active level editor viewport", EUserInterfaceActionType::Check, FInputChord() );
+	UI_COMMAND( PlayInEditorFloating, "New Editor Window", "Play this level in a new window", EUserInterfaceActionType::Check, FInputChord() );
+	UI_COMMAND( PlayInVR, "VR Preview", "Play this level in VR", EUserInterfaceActionType::Check, FInputChord() );
+	UI_COMMAND( PlayInMobilePreview, "Mobile Preview", "Play this level as a mobile device preview (runs in its own process)", EUserInterfaceActionType::Check, FInputChord() );
+	UI_COMMAND( PlayInNewProcess, "Standalone Game", "Play this level in a new window that runs in its own process", EUserInterfaceActionType::Check, FInputChord() );
+	UI_COMMAND( PlayInCameraLocation, "Current Camera Location", "Spawn the player at the current camera location", EUserInterfaceActionType::RadioButton, FInputChord() );
+	UI_COMMAND( PlayInDefaultPlayerStart, "Default Player Start", "Spawn the player at the map's default player start", EUserInterfaceActionType::RadioButton, FInputChord() );
+	UI_COMMAND( PlayInNetworkSettings, "Network Settings...", "Open the settings for the 'Play In' feature", EUserInterfaceActionType::Button, FInputChord() );
+	UI_COMMAND( PlayInNetworkDedicatedServer, "Run Dedicated Server", "If checked, a separate dedicated server will be launched. Otherwise the first player will act as a listen server that all other players connect to.", EUserInterfaceActionType::ToggleButton, FInputChord() );
+	UI_COMMAND( PlayInSettings, "Advanced Settings...", "Open the settings for the 'Play In' feature", EUserInterfaceActionType::Button, FInputChord() );
 
 	// SIE & PIE controls
-	UI_COMMAND( StopPlaySession, "Stop", "Stop simulation", EUserInterfaceActionType::Button, FInputGesture() );
-	UI_COMMAND( ResumePlaySession, "Resume", "Resume simulation", EUserInterfaceActionType::Button, FInputGesture() );
-	UI_COMMAND( PausePlaySession, "Pause", "Pause simulation", EUserInterfaceActionType::Button, FInputGesture() );
-	UI_COMMAND( SingleFrameAdvance, "Frame Skip", "Advances a single frame", EUserInterfaceActionType::Button, FInputGesture() );
-	UI_COMMAND( TogglePlayPauseOfPlaySession, "Toggle Play/Pause", "Resume playing if paused, or pause if playing", EUserInterfaceActionType::Button, FInputGesture( EKeys::Pause ) );
-	UI_COMMAND( PossessEjectPlayer, "Possess or Eject Player", "Possesses or ejects the player from the camera", EUserInterfaceActionType::Button, FInputGesture( EKeys::F8 ) );
-	UI_COMMAND( ShowCurrentStatement, "Find Node", "Show the current node", EUserInterfaceActionType::Button, FInputGesture() );
-	UI_COMMAND( StepInto, "Step", "Step to the next node to be executed", EUserInterfaceActionType::Button, FInputGesture( EKeys::F10) );
+	UI_COMMAND( StopPlaySession, "Stop", "Stop simulation", EUserInterfaceActionType::Button, FInputChord() );
+	UI_COMMAND( ResumePlaySession, "Resume", "Resume simulation", EUserInterfaceActionType::Button, FInputChord() );
+	UI_COMMAND( PausePlaySession, "Pause", "Pause simulation", EUserInterfaceActionType::Button, FInputChord() );
+	UI_COMMAND( SingleFrameAdvance, "Frame Skip", "Advances a single frame", EUserInterfaceActionType::Button, FInputChord() );
+	UI_COMMAND( TogglePlayPauseOfPlaySession, "Toggle Play/Pause", "Resume playing if paused, or pause if playing", EUserInterfaceActionType::Button, FInputChord( EKeys::Pause ) );
+	UI_COMMAND( PossessEjectPlayer, "Possess or Eject Player", "Possesses or ejects the player from the camera", EUserInterfaceActionType::Button, FInputChord( EKeys::F8 ) );
+	UI_COMMAND( ShowCurrentStatement, "Find Node", "Show the current node", EUserInterfaceActionType::Button, FInputChord() );
+	UI_COMMAND( StepInto, "Step", "Step to the next node to be executed", EUserInterfaceActionType::Button, FInputChord( EKeys::F10) );
 
 	// Launch
-	UI_COMMAND( RepeatLastLaunch, "Launch", "Launches the game on the device as the last session launched from the dropdown next to the Play on Device button on the level editor toolbar", EUserInterfaceActionType::Button, FInputGesture( EKeys::P, EModifierKey::Alt | EModifierKey::Shift ) )
-	UI_COMMAND( OpenProjectLauncher, "Project Launcher...", "Open the Project Launcher for advanced packaging, deploying and launching of your projects", EUserInterfaceActionType::Button, FInputGesture());
-	UI_COMMAND( OpenDeviceManager, "Device Manager...", "View and manage connected devices.", EUserInterfaceActionType::Button, FInputGesture() );
+	UI_COMMAND( RepeatLastLaunch, "Launch", "Launches the game on the device as the last session launched from the dropdown next to the Play on Device button on the level editor toolbar", EUserInterfaceActionType::Button, FInputChord( EKeys::P, EModifierKey::Alt | EModifierKey::Shift ) )
+	UI_COMMAND( OpenProjectLauncher, "Project Launcher...", "Open the Project Launcher for advanced packaging, deploying and launching of your projects", EUserInterfaceActionType::Button, FInputChord());
+	UI_COMMAND( OpenDeviceManager, "Device Manager...", "View and manage connected devices.", EUserInterfaceActionType::Button, FInputChord() );
 }
 
 
diff --git a/Engine/Source/Editor/UnrealEd/Private/LevelEditorViewport.cpp b/Engine/Source/Editor/UnrealEd/Private/LevelEditorViewport.cpp
--- a/Engine/Source/Editor/UnrealEd/Private/LevelEditorViewport.cpp
+++ b/Engine/Source/Editor/UnrealEd/Private/LevelEditorViewport.cpp
@@ -2289,13 +2289,13 @@ TSharedPtr<FDragTool> FLevelEditorViewportClient::MakeDragTool( EDragTool::Type
 
 static bool CommandAcceptsInput( FLevelEditorViewportClient& ViewportClient, FKey Key, const TSharedPtr<FUICommandInfo> Command )
 {
-	const FInputGesture& Gesture = *Command->GetActiveGesture();
+	const FInputChord& Chord = *Command->GetActiveChord();
 
-	return (!Gesture.NeedsControl()	|| ViewportClient.IsCtrlPressed() ) 
-		&& (!Gesture.NeedsAlt()		|| ViewportClient.IsAltPressed() ) 
-		&& (!Gesture.NeedsShift()	|| ViewportClient.IsShiftPressed() ) 
-		&& (!Gesture.NeedsCommand()		|| ViewportClient.IsCmdPressed() )
-		&& Gesture.Key == Key;
+	return (!Chord.NeedsControl()	|| ViewportClient.IsCtrlPressed() ) 
+		&& (!Chord.NeedsAlt()		|| ViewportClient.IsAltPressed() ) 
+		&& (!Chord.NeedsShift()	|| ViewportClient.IsShiftPressed() ) 
+		&& (!Chord.NeedsCommand()		|| ViewportClient.IsCmdPressed() )
+		&& Chord.Key == Key;
 }
 
 static const FLevelViewportCommands& GetLevelViewportCommands()
diff --git a/Engine/Source/Editor/UnrealEd/Private/RichCurveEditorCommands.cpp b/Engine/Source/Editor/UnrealEd/Private/RichCurveEditorCommands.cpp
--- a/Engine/Source/Editor/UnrealEd/Private/RichCurveEditorCommands.cpp
+++ b/Engine/Source/Editor/UnrealEd/Private/RichCurveEditorCommands.cpp
@@ -7,18 +7,18 @@
 
 void FRichCurveEditorCommands::RegisterCommands()
 {
-	UI_COMMAND(ZoomToFitHorizontal, "Fit Horizontal", "Zoom to Fit - Horizontal", EUserInterfaceActionType::Button, FInputGesture());
-	UI_COMMAND(ZoomToFitVertical, "Fit Vertical", "Zoom to Fit - Vertical", EUserInterfaceActionType::Button, FInputGesture());
-	UI_COMMAND(ZoomToFitAll, "Fit All", "Zoom to Fit", EUserInterfaceActionType::Button, FInputGesture());
-	UI_COMMAND(ZoomToFitSelected, "Fit Selected", "Zoom to Fit - Selected", EUserInterfaceActionType::Button, FInputGesture(EKeys::F));
+	UI_COMMAND(ZoomToFitHorizontal, "Fit Horizontal", "Zoom to Fit - Horizontal", EUserInterfaceActionType::Button, FInputChord());
+	UI_COMMAND(ZoomToFitVertical, "Fit Vertical", "Zoom to Fit - Vertical", EUserInterfaceActionType::Button, FInputChord());
+	UI_COMMAND(ZoomToFitAll, "Fit All", "Zoom to Fit", EUserInterfaceActionType::Button, FInputChord());
+	UI_COMMAND(ZoomToFitSelected, "Fit Selected", "Zoom to Fit - Selected", EUserInterfaceActionType::Button, FInputChord(EKeys::F));
 
-	UI_COMMAND(ToggleSnapping, "Snapping", "Toggle Snapping", EUserInterfaceActionType::ToggleButton, FInputGesture());
+	UI_COMMAND(ToggleSnapping, "Snapping", "Toggle Snapping", EUserInterfaceActionType::ToggleButton, FInputChord());
 
-	UI_COMMAND(InterpolationConstant, "Constant", "Constant interpolation", EUserInterfaceActionType::ToggleButton, FInputGesture());
-	UI_COMMAND(InterpolationLinear, "Linear", "Linear interpolation", EUserInterfaceActionType::ToggleButton, FInputGesture());
-	UI_COMMAND(InterpolationCubicAuto, "Auto", "Cubic interpolation - Automatic tangents", EUserInterfaceActionType::ToggleButton, FInputGesture());
-	UI_COMMAND(InterpolationCubicUser, "User", "Cubic interpolation - User flat tangents", EUserInterfaceActionType::ToggleButton, FInputGesture());
-	UI_COMMAND(InterpolationCubicBreak, "Break", "Cubic interpolation - User broken tangents", EUserInterfaceActionType::ToggleButton, FInputGesture());
+	UI_COMMAND(InterpolationConstant, "Constant", "Constant interpolation", EUserInterfaceActionType::ToggleButton, FInputChord());
+	UI_COMMAND(InterpolationLinear, "Linear", "Linear interpolation", EUserInterfaceActionType::ToggleButton, FInputChord());
+	UI_COMMAND(InterpolationCubicAuto, "Auto", "Cubic interpolation - Automatic tangents", EUserInterfaceActionType::ToggleButton, FInputChord());
+	UI_COMMAND(InterpolationCubicUser, "User", "Cubic interpolation - User flat tangents", EUserInterfaceActionType::ToggleButton, FInputChord());
+	UI_COMMAND(InterpolationCubicBreak, "Break", "Cubic interpolation - User broken tangents", EUserInterfaceActionType::ToggleButton, FInputChord());
 }
 
 #undef LOCTEXT_NAMESPACE
\ No newline at end of file
diff --git a/Engine/Source/Editor/UnrealEd/Private/SnappingUtils.cpp b/Engine/Source/Editor/UnrealEd/Private/SnappingUtils.cpp
--- a/Engine/Source/Editor/UnrealEd/Private/SnappingUtils.cpp
+++ b/Engine/Source/Editor/UnrealEd/Private/SnappingUtils.cpp
@@ -134,12 +134,12 @@ bool FEditorViewportSnapping::IsSnapToVertexEnabled()
 		FLevelEditorModule& LevelEditor = FModuleManager::GetModuleChecked<FLevelEditorModule>( TEXT("LevelEditor") );
 		const FLevelEditorCommands& Commands = LevelEditor.GetLevelEditorCommands();
 
-		const FInputGesture& Gesture = *Commands.HoldToEnableVertexSnapping->GetActiveGesture();
+		const FInputChord& Chord = *Commands.HoldToEnableVertexSnapping->GetActiveChord();
 
-		return (Gesture.NeedsControl() == GCurrentLevelEditingViewportClient->IsCtrlPressed() ) 
-			&& (Gesture.NeedsAlt() ==  GCurrentLevelEditingViewportClient->IsAltPressed() ) 
-			&& (Gesture.NeedsShift() == GCurrentLevelEditingViewportClient->IsShiftPressed() ) 
-			&& GCurrentLevelEditingViewportClient->Viewport->KeyState(Gesture.Key) == true;
+		return (Chord.NeedsControl() == GCurrentLevelEditingViewportClient->IsCtrlPressed() ) 
+			&& (Chord.NeedsAlt() ==  GCurrentLevelEditingViewportClient->IsAltPressed() ) 
+			&& (Chord.NeedsShift() == GCurrentLevelEditingViewportClient->IsShiftPressed() ) 
+			&& GCurrentLevelEditingViewportClient->Viewport->KeyState(Chord.Key) == true;
 	}
 	else
 	{
diff --git a/Engine/Source/Editor/UnrealEd/Private/SoundCueGraphEditorCommands.cpp b/Engine/Source/Editor/UnrealEd/Private/SoundCueGraphEditorCommands.cpp
--- a/Engine/Source/Editor/UnrealEd/Private/SoundCueGraphEditorCommands.cpp
+++ b/Engine/Source/Editor/UnrealEd/Private/SoundCueGraphEditorCommands.cpp
@@ -7,14 +7,14 @@
 
 void FSoundCueGraphEditorCommands::RegisterCommands()
 {
-	UI_COMMAND(PlayCue, "Play Cue", "Plays the SoundCue", EUserInterfaceActionType::Button, FInputGesture());
-	UI_COMMAND(PlayNode, "Play Node", "Plays the currently selected node", EUserInterfaceActionType::Button, FInputGesture());
-	UI_COMMAND(StopCueNode, "Stop", "Stops the currently playing cue/node", EUserInterfaceActionType::Button, FInputGesture());
-	UI_COMMAND(TogglePlayback, "Toggle Playback", "Plays the SoundCue or stops the currently playing cue/node", EUserInterfaceActionType::Button, FInputGesture(EKeys::SpaceBar));
+	UI_COMMAND(PlayCue, "Play Cue", "Plays the SoundCue", EUserInterfaceActionType::Button, FInputChord());
+	UI_COMMAND(PlayNode, "Play Node", "Plays the currently selected node", EUserInterfaceActionType::Button, FInputChord());
+	UI_COMMAND(StopCueNode, "Stop", "Stops the currently playing cue/node", EUserInterfaceActionType::Button, FInputChord());
+	UI_COMMAND(TogglePlayback, "Toggle Playback", "Plays the SoundCue or stops the currently playing cue/node", EUserInterfaceActionType::Button, FInputChord(EKeys::SpaceBar));
 
-	UI_COMMAND(BrowserSync, "Sync to Browser", "Selects the SoundWave in the content browser", EUserInterfaceActionType::Button, FInputGesture());
-	UI_COMMAND(AddInput, "Add Input", "Adds an input to the node", EUserInterfaceActionType::Button, FInputGesture());
-	UI_COMMAND(DeleteInput, "Delete Input", "Removes an input from the node", EUserInterfaceActionType::Button, FInputGesture());
+	UI_COMMAND(BrowserSync, "Sync to Browser", "Selects the SoundWave in the content browser", EUserInterfaceActionType::Button, FInputChord());
+	UI_COMMAND(AddInput, "Add Input", "Adds an input to the node", EUserInterfaceActionType::Button, FInputChord());
+	UI_COMMAND(DeleteInput, "Delete Input", "Removes an input from the node", EUserInterfaceActionType::Button, FInputChord());
 }
 
 #undef LOCTEXT_NAMESPACE
\ No newline at end of file
diff --git a/Engine/Source/Editor/UnrealEd/Private/Tests/EditorBuildPromotionTests.cpp b/Engine/Source/Editor/UnrealEd/Private/Tests/EditorBuildPromotionTests.cpp
--- a/Engine/Source/Editor/UnrealEd/Private/Tests/EditorBuildPromotionTests.cpp
+++ b/Engine/Source/Editor/UnrealEd/Private/Tests/EditorBuildPromotionTests.cpp
@@ -164,15 +164,15 @@ namespace EditorBuildPromotionTestUtils
 	*
 	* @param CommandContext - The context of the command
 	* @param Command - The command name to set
-	* @param NewGesture - The new input gesture to assign
+	* @param NewChord - The new input chord to assign
 	*/
-	static void SetEditorKeybinding(const FString& CommandContext, const FString& Command, const FInputGesture& NewGesture)
+	static void SetEditorKeybinding(const FString& CommandContext, const FString& Command, const FInputChord& NewChord)
 	{
 		TSharedPtr<FUICommandInfo> UICommand = FInputBindingManager::Get().FindCommandInContext(*CommandContext, *Command);
 		if (UICommand.IsValid())
 		{
-			UICommand->SetActiveGesture(NewGesture);
-			FInputBindingManager::Get().NotifyActiveGestureChanged(*UICommand.Get());
+			UICommand->SetActiveChord(NewChord);
+			FInputBindingManager::Get().NotifyActiveChordChanged(*UICommand.Get());
 			FInputBindingManager::Get().SaveInputBindings();
 		}
 		else
@@ -186,15 +186,15 @@ namespace EditorBuildPromotionTestUtils
 	*
 	* @param CommandContext - The context of the command
 	* @param Command - The command name to get
-	* @param NewGesture - The current input gesture that is assigned
+	* @param NewChord - The current input chord that is assigned
 	*/
-	static bool GetEditorKeybinding(const FString& CommandContext, const FString& Command, FInputGesture& CurrentGesture)
+	static bool GetEditorKeybinding(const FString& CommandContext, const FString& Command, FInputChord& CurrentChord)
 	{
 		TSharedPtr<FUICommandInfo> UICommand = FInputBindingManager::Get().FindCommandInContext(*CommandContext, *Command);
 		if (UICommand.IsValid())
 		{
-			TSharedRef<const FInputGesture> ActiveGesture = UICommand->GetActiveGesture();
-			CurrentGesture = ActiveGesture.Get();
+			TSharedRef<const FInputChord> ActiveChord = UICommand->GetActiveChord();
+			CurrentChord = ActiveChord.Get();
 		}
 		else
 		{
@@ -344,10 +344,10 @@ namespace EditorBuildPromotionTestUtils
 	/**
 	* Sends a UI command to the active top level window after focusing on a widget of a given type
 	*
-	* @param InGesture - The gesture to send to the window
+	* @param InChord - The chord to send to the window
 	* @param WidgetTypeToFocus - The widget type to find and focus on
 	*/
-	static void SendCommandToCurrentEditor(const FInputGesture& InGesture, const FName& WidgetTypeToFocus)
+	static void SendCommandToCurrentEditor(const FInputChord& InChord, const FName& WidgetTypeToFocus)
 	{
 		//Focus the asset Editor / Graph 
 		TSharedRef<SWindow> EditorWindow = FSlateApplication::Get().GetActiveTopLevelWindow().ToSharedRef();
@@ -358,8 +358,8 @@ namespace EditorBuildPromotionTestUtils
 			FSlateApplication::Get().SetKeyboardFocus(FocusWidget.ToSharedRef(), EFocusCause::SetDirectly);
 
 			//Send the command
-			FModifierKeysState ModifierKeys(InGesture.NeedsShift(), false, InGesture.NeedsControl(), false, InGesture.NeedsAlt(), false, InGesture.NeedsCommand(), false, false);
-			FKeyEvent KeyEvent(InGesture.Key, ModifierKeys, 0/*UserIndex*/, false, 0, 0);
+			FModifierKeysState ModifierKeys(InChord.NeedsShift(), false, InChord.NeedsControl(), false, InChord.NeedsAlt(), false, InChord.NeedsCommand(), false, false);
+			FKeyEvent KeyEvent(InChord.Key, ModifierKeys, 0/*UserIndex*/, false, 0, 0);
 			FSlateApplication::Get().ProcessKeyDownEvent(KeyEvent);
 			FSlateApplication::Get().ProcessKeyUpEvent(KeyEvent);
 		}
@@ -370,25 +370,25 @@ namespace EditorBuildPromotionTestUtils
 	}
 
 	/**
-	* Gets the current input gesture or sets a new one if it doesn't exist
+	* Gets the current input chord or sets a new one if it doesn't exist
 	*
 	* @param Context - The context of the UI Command
 	* @param Command - The name of the UI command
 	*/
-	static FInputGesture GetOrSetUICommand(const FString& Context, const FString& Command)
+	static FInputChord GetOrSetUICommand(const FString& Context, const FString& Command)
 	{
-		FInputGesture CurrentGesture;
-		GetEditorKeybinding(Context, Command, CurrentGesture);
+		FInputChord CurrentChord;
+		GetEditorKeybinding(Context, Command, CurrentChord);
 
 		//If there is no current keybinding, set one
-		if (!CurrentGesture.Key.IsValid())
+		if (!CurrentChord.Key.IsValid())
 		{
-			FInputGesture NewGesture(EKeys::J, EModifierKey::Control);
-			SetEditorKeybinding(Context, Command, NewGesture);
-			CurrentGesture = NewGesture;
+			FInputChord NewChord(EKeys::J, EModifierKey::Control);
+			SetEditorKeybinding(Context, Command, NewChord);
+			CurrentChord = NewChord;
 		}
 
-		return CurrentGesture;
+		return CurrentChord;
 	}
 
 	/**
@@ -398,10 +398,10 @@ namespace EditorBuildPromotionTestUtils
 	{
 		const FString Context = TEXT("MaterialEditor");
 		const FString Command = TEXT("Apply");
-		FInputGesture CurrentApplyGesture = GetOrSetUICommand(Context, Command);
+		FInputChord CurrentApplyChord = GetOrSetUICommand(Context, Command);
 
 		const FName FocusWidgetType(TEXT("SGraphEditor"));
-		SendCommandToCurrentEditor(CurrentApplyGesture, FocusWidgetType);
+		SendCommandToCurrentEditor(CurrentApplyChord, FocusWidgetType);
 	}
 
 	/**
@@ -412,10 +412,10 @@ namespace EditorBuildPromotionTestUtils
 		const FString Context = TEXT("AssetEditor");
 		const FString Command = TEXT("SaveAsset");
 
-		FInputGesture CurrentSaveGesture = GetOrSetUICommand(Context, Command);
+		FInputChord CurrentSaveChord = GetOrSetUICommand(Context, Command);
 
 		const FName FocusWidgetType(TEXT("SCascadeEmitterCanvas"));
-		SendCommandToCurrentEditor(CurrentSaveGesture, FocusWidgetType);
+		SendCommandToCurrentEditor(CurrentSaveChord, FocusWidgetType);
 	}
 
 	/**
@@ -426,10 +426,10 @@ namespace EditorBuildPromotionTestUtils
 		const FString Context = TEXT("BlueprintEditor");
 		const FString Command = TEXT("ResetCamera");
 
-		FInputGesture CurrentSaveGesture = GetOrSetUICommand(Context, Command);
+		FInputChord CurrentSaveChord = GetOrSetUICommand(Context, Command);
 
 		const FName FocusWidgetType(TEXT("SSCSEditorViewport"));
-		SendCommandToCurrentEditor(CurrentSaveGesture, FocusWidgetType);
+		SendCommandToCurrentEditor(CurrentSaveChord, FocusWidgetType);
 	}
 
 	/**
@@ -4349,18 +4349,18 @@ bool FBuildPromotionSettingsTest::RunTest(const FString& Parameters)
 	UE_LOG(LogEditorBuildPromotionTests, Display, TEXT("Binding create empty layer shortcut"));
 
 	//Bind H to CreateEmptyLayer keybinding
-	FInputGesture NewCreateGesture(EKeys::H, EModifierKey::None);
-	EditorBuildPromotionTestUtils::SetEditorKeybinding(TEXT("LayersView"), TEXT("CreateEmptyLayer"), NewCreateGesture);
+	FInputChord NewCreateChord(EKeys::H, EModifierKey::None);
+	EditorBuildPromotionTestUtils::SetEditorKeybinding(TEXT("LayersView"), TEXT("CreateEmptyLayer"), NewCreateChord);
 
 	UE_LOG(LogEditorBuildPromotionTests, Display, TEXT("Binding request rename layer shortcut"));
 
 	//Bind J to RequestRenameLayer
-	FInputGesture NewRenameGesture(EKeys::J, EModifierKey::None);
-	EditorBuildPromotionTestUtils::SetEditorKeybinding(TEXT("LayersView"), TEXT("RequestRenameLayer"), NewRenameGesture);
+	FInputChord NewRenameChord(EKeys::J, EModifierKey::None);
+	EditorBuildPromotionTestUtils::SetEditorKeybinding(TEXT("LayersView"), TEXT("RequestRenameLayer"), NewRenameChord);
 
 	UE_LOG(LogEditorBuildPromotionTests, Display, TEXT("Binding play shortcut (PIE)"));
-	FInputGesture NewPIEGesture(EKeys::L, EModifierKey::Control);
-	EditorBuildPromotionTestUtils::SetEditorKeybinding(TEXT("PlayWorld"), TEXT("RepeatLastPlay"), NewPIEGesture);
+	FInputChord NewPIEChord(EKeys::L, EModifierKey::Control);
+	EditorBuildPromotionTestUtils::SetEditorKeybinding(TEXT("PlayWorld"), TEXT("RepeatLastPlay"), NewPIEChord);
 
 	//Export the keybindings
 	const FString TargetKeybindFile = FString::Printf(TEXT("%s/BuildPromotion/Keybindings-%d.ini"), *FPaths::AutomationDir(), GEngineVersion.GetChangelist());
@@ -4386,9 +4386,9 @@ bool FBuildPromotionSettingsTest::RunTest(const FString& Parameters)
 	//Change the setting back
 	EditorBuildPromotionTestUtils::SetPropertyByName(EditorStyleSettings, TEXT("bUseSmallToolBarIcons"), OldStyleSetting);
 
-	FInputGesture CurrentCreateGesture;
-	EditorBuildPromotionTestUtils::GetEditorKeybinding(TEXT("LayersView"), TEXT("CreateEmptyLayer"), CurrentCreateGesture);
-	if (CurrentCreateGesture.Key == NewCreateGesture.Key)
+	FInputChord CurrentCreateChord;
+	EditorBuildPromotionTestUtils::GetEditorKeybinding(TEXT("LayersView"), TEXT("CreateEmptyLayer"), CurrentCreateChord);
+	if (CurrentCreateChord.Key == NewCreateChord.Key)
 	{
 		UE_LOG(LogEditorBuildPromotionTests, Display, TEXT("CreateEmptyLayer keybinding correct."));
 	}
@@ -4397,9 +4397,9 @@ bool FBuildPromotionSettingsTest::RunTest(const FString& Parameters)
 		UE_LOG(LogEditorBuildPromotionTests, Error, TEXT("CreateEmptyLayer keybinding incorrect."));
 	}
 
-	FInputGesture CurrentRenameGesture;
-	EditorBuildPromotionTestUtils::GetEditorKeybinding(TEXT("LayersView"), TEXT("RequestRenameLayer"), CurrentRenameGesture);
-	if (CurrentRenameGesture.Key == NewRenameGesture.Key)
+	FInputChord CurrentRenameChord;
+	EditorBuildPromotionTestUtils::GetEditorKeybinding(TEXT("LayersView"), TEXT("RequestRenameLayer"), CurrentRenameChord);
+	if (CurrentRenameChord.Key == NewRenameChord.Key)
 	{
 		UE_LOG(LogEditorBuildPromotionTests, Display, TEXT("RequestRenameLayer keybinding correct."));
 	}
diff --git a/Engine/Source/Editor/UnrealEd/Private/Toolkits/AssetEditorCommonCommands.cpp b/Engine/Source/Editor/UnrealEd/Private/Toolkits/AssetEditorCommonCommands.cpp
--- a/Engine/Source/Editor/UnrealEd/Private/Toolkits/AssetEditorCommonCommands.cpp
+++ b/Engine/Source/Editor/UnrealEd/Private/Toolkits/AssetEditorCommonCommands.cpp
@@ -7,10 +7,10 @@
 
 void FAssetEditorCommonCommands::RegisterCommands()
 {
-	UI_COMMAND( SaveAsset, "Save", "Saves this asset", EUserInterfaceActionType::Button, FInputGesture() );
-	UI_COMMAND( ReimportAsset, "Reimport", "Reimports the asset being edited", EUserInterfaceActionType::Button, FInputGesture());
-	UI_COMMAND( SwitchToStandaloneEditor, "Switch to Standalone Editor", "Closes the level-centric asset editor and reopens it in 'standalone' mode", EUserInterfaceActionType::Button, FInputGesture() );
-	UI_COMMAND( SwitchToWorldCentricEditor, "Switch to World-Centric Editor", "Closes the standalone asset editor and reopens it in 'world-centric' mode, docked within the level editor that it was originally opened in.", EUserInterfaceActionType::Button, FInputGesture() );
+	UI_COMMAND( SaveAsset, "Save", "Saves this asset", EUserInterfaceActionType::Button, FInputChord() );
+	UI_COMMAND( ReimportAsset, "Reimport", "Reimports the asset being edited", EUserInterfaceActionType::Button, FInputChord());
+	UI_COMMAND( SwitchToStandaloneEditor, "Switch to Standalone Editor", "Closes the level-centric asset editor and reopens it in 'standalone' mode", EUserInterfaceActionType::Button, FInputChord() );
+	UI_COMMAND( SwitchToWorldCentricEditor, "Switch to World-Centric Editor", "Closes the standalone asset editor and reopens it in 'world-centric' mode, docked within the level editor that it was originally opened in.", EUserInterfaceActionType::Button, FInputChord() );
 }
 
 #undef LOCTEXT_NAMESPACE
\ No newline at end of file
diff --git a/Engine/Source/Editor/UnrealEd/Private/Toolkits/GlobalEditorCommonCommands.cpp b/Engine/Source/Editor/UnrealEd/Private/Toolkits/GlobalEditorCommonCommands.cpp
--- a/Engine/Source/Editor/UnrealEd/Private/Toolkits/GlobalEditorCommonCommands.cpp
+++ b/Engine/Source/Editor/UnrealEd/Private/Toolkits/GlobalEditorCommonCommands.cpp
@@ -21,17 +21,17 @@ FGlobalEditorCommonCommands::FGlobalEditorCommonCommands()
 
 void FGlobalEditorCommonCommands::RegisterCommands()
 {
-	UI_COMMAND( SummonControlTabNavigation, "Tab Navigation", "Summons a list of open assets and tabs", EUserInterfaceActionType::Button, FInputGesture(EModifierKey::Control, EKeys::Tab) );
-	UI_COMMAND( SummonControlTabNavigationAlternate, "Tab Navigation", "Summons a list of open assets and tabs", EUserInterfaceActionType::Button, FInputGesture(EModifierKey::Command, EKeys::Tab) );
+	UI_COMMAND( SummonControlTabNavigation, "Tab Navigation", "Summons a list of open assets and tabs", EUserInterfaceActionType::Button, FInputChord(EModifierKey::Control, EKeys::Tab) );
+	UI_COMMAND( SummonControlTabNavigationAlternate, "Tab Navigation", "Summons a list of open assets and tabs", EUserInterfaceActionType::Button, FInputChord(EModifierKey::Command, EKeys::Tab) );
 
-	UI_COMMAND( SummonOpenAssetDialog, "Open Asset...", "Summons an asset picker", EUserInterfaceActionType::Button, FInputGesture(EModifierKey::Control, EKeys::P) );
-	UI_COMMAND( SummonOpenAssetDialogAlternate, "Open Asset...", "Summons an asset picker", EUserInterfaceActionType::Button, FInputGesture(EModifierKey::Alt | EModifierKey::Shift, EKeys::O));
-	UI_COMMAND( FindInContentBrowser, "Find in Content Browser", "Summons the Content Browser and navigates to the selected asset", EUserInterfaceActionType::Button, FInputGesture(EModifierKey::Control, EKeys::B));
-	UI_COMMAND( ViewReferences, "View References", "Launches the reference viewer showing the selected assets' references", EUserInterfaceActionType::Button, FInputGesture(EModifierKey::Shift | EModifierKey::Alt, EKeys::R));
+	UI_COMMAND( SummonOpenAssetDialog, "Open Asset...", "Summons an asset picker", EUserInterfaceActionType::Button, FInputChord(EModifierKey::Control, EKeys::P) );
+	UI_COMMAND( SummonOpenAssetDialogAlternate, "Open Asset...", "Summons an asset picker", EUserInterfaceActionType::Button, FInputChord(EModifierKey::Alt | EModifierKey::Shift, EKeys::O));
+	UI_COMMAND( FindInContentBrowser, "Find in Content Browser", "Summons the Content Browser and navigates to the selected asset", EUserInterfaceActionType::Button, FInputChord(EModifierKey::Control, EKeys::B));
+	UI_COMMAND( ViewReferences, "View References", "Launches the reference viewer showing the selected assets' references", EUserInterfaceActionType::Button, FInputChord(EModifierKey::Shift | EModifierKey::Alt, EKeys::R));
 	
-	UI_COMMAND( OpenConsoleCommandBox, "Open Console Command Box", "Opens an edit box where you can type in a console command", EUserInterfaceActionType::Button, FInputGesture(EKeys::Tilde));
+	UI_COMMAND( OpenConsoleCommandBox, "Open Console Command Box", "Opens an edit box where you can type in a console command", EUserInterfaceActionType::Button, FInputChord(EKeys::Tilde));
 
-	UI_COMMAND( OpenDocumentation, "Open Documentation...", "Opens documentation for this tool", EUserInterfaceActionType::Button, FInputGesture(EKeys::F1) );
+	UI_COMMAND( OpenDocumentation, "Open Documentation...", "Opens documentation for this tool", EUserInterfaceActionType::Button, FInputChord(EKeys::F1) );
 }
 
 void FGlobalEditorCommonCommands::MapActions(TSharedRef<FUICommandList>& ToolkitCommands)
@@ -66,7 +66,7 @@ void FGlobalEditorCommonCommands::OnPressedCtrlTab(TSharedPtr<FUICommandInfo> Tr
 		const FVector2D TabListSize(700.0f, 486.0f);
 
 		// Create the contents of the popup
-		TSharedRef<SWidget> ActualWidget = SNew(SGlobalTabSwitchingDialog, TabListSize, *TriggeringCommand->GetActiveGesture());
+		TSharedRef<SWidget> ActualWidget = SNew(SGlobalTabSwitchingDialog, TabListSize, *TriggeringCommand->GetActiveChord());
 
 		TSharedPtr<SWindow> NewWindow = OpenPopup(ActualWidget, TabListSize);
 	}
diff --git a/Engine/Source/Editor/UnrealEd/Private/Toolkits/SGlobalTabSwitchingDialog.cpp b/Engine/Source/Editor/UnrealEd/Private/Toolkits/SGlobalTabSwitchingDialog.cpp
--- a/Engine/Source/Editor/UnrealEd/Private/Toolkits/SGlobalTabSwitchingDialog.cpp
+++ b/Engine/Source/Editor/UnrealEd/Private/Toolkits/SGlobalTabSwitchingDialog.cpp
@@ -344,12 +344,12 @@ void SGlobalTabSwitchingDialog::DismissDialog()
 	FSlateApplication::Get().DismissAllMenus();
 }
 
-void SGlobalTabSwitchingDialog::Construct(const FArguments& InArgs, FVector2D InSize, FInputGesture InTriggerGesture)
+void SGlobalTabSwitchingDialog::Construct(const FArguments& InArgs, FVector2D InSize, FInputChord InTriggerChord)
 {
 	check(!bIsAlreadyOpen);
 	bIsAlreadyOpen = true;
 
-	TriggerGesture = InTriggerGesture;
+	TriggerChord = InTriggerChord;
 
 	AssetThumbnailPool = MakeShareable(new FAssetThumbnailPool(128));
 
@@ -475,9 +475,9 @@ void SGlobalTabSwitchingDialog::Construct(const FArguments& InArgs, FVector2D In
 FReply SGlobalTabSwitchingDialog::OnKeyUp(const FGeometry& MyGeometry, const FKeyEvent& InKeyEvent)
 {
 	// Check to see if the trigger modifier key was released, which should close the dialog
-	const bool bCloseViaControl = TriggerGesture.NeedsControl() && ((InKeyEvent.GetKey() == EKeys::LeftControl) || (InKeyEvent.GetKey() == EKeys::RightControl));
-	const bool bCloseViaCommand = TriggerGesture.NeedsCommand() && ((InKeyEvent.GetKey() == EKeys::LeftCommand) || (InKeyEvent.GetKey() == EKeys::RightCommand));
-	const bool bCloseViaAlt = TriggerGesture.NeedsAlt() && ((InKeyEvent.GetKey() == EKeys::LeftAlt) || (InKeyEvent.GetKey() == EKeys::RightAlt));
+	const bool bCloseViaControl = TriggerChord.NeedsControl() && ((InKeyEvent.GetKey() == EKeys::LeftControl) || (InKeyEvent.GetKey() == EKeys::RightControl));
+	const bool bCloseViaCommand = TriggerChord.NeedsCommand() && ((InKeyEvent.GetKey() == EKeys::LeftCommand) || (InKeyEvent.GetKey() == EKeys::RightCommand));
+	const bool bCloseViaAlt = TriggerChord.NeedsAlt() && ((InKeyEvent.GetKey() == EKeys::LeftAlt) || (InKeyEvent.GetKey() == EKeys::RightAlt));
 	
 	if (bCloseViaControl || bCloseViaCommand || bCloseViaAlt)
 	{
@@ -490,7 +490,7 @@ FReply SGlobalTabSwitchingDialog::OnKeyUp(const FGeometry& MyGeometry, const FKe
 
 FReply SGlobalTabSwitchingDialog::OnKeyDown(const FGeometry& MyGeometry, const FKeyEvent& InKeyEvent)
 {
-	if (TriggerGesture.Key == InKeyEvent.GetKey())
+	if (TriggerChord.Key == InKeyEvent.GetKey())
 	{
 		const bool bCycleForward = !InKeyEvent.IsShiftDown();
 		CycleSelection(bCycleForward);
diff --git a/Engine/Source/Editor/UnrealEd/Private/Toolkits/SGlobalTabSwitchingDialog.h b/Engine/Source/Editor/UnrealEd/Private/Toolkits/SGlobalTabSwitchingDialog.h
--- a/Engine/Source/Editor/UnrealEd/Private/Toolkits/SGlobalTabSwitchingDialog.h
+++ b/Engine/Source/Editor/UnrealEd/Private/Toolkits/SGlobalTabSwitchingDialog.h
@@ -12,7 +12,7 @@ public:
 	SLATE_END_ARGS()
 
 public:
-	void Construct(const FArguments& InArgs, FVector2D InSize, FInputGesture InTriggerGesture);
+	void Construct(const FArguments& InArgs, FVector2D InSize, FInputChord InTriggerChord);
 
 	~SGlobalTabSwitchingDialog()
 	{
@@ -48,8 +48,8 @@ private:
 	FTabListItemPtr GetMainTabListSelectedItem() const;
 
 private:
-	// The gesture that triggered the dialog (so we can handle the correct Tab/`/etc... key repeat, and dismiss on the correct control/command modifier release
-	FInputGesture TriggerGesture;
+	// The chord that triggered the dialog (so we can handle the correct Tab/`/etc... key repeat, and dismiss on the correct control/command modifier release
+	FInputChord TriggerChord;
 
 	// Pool for maintaining and rendering thumbnails
 	TSharedPtr<class FAssetThumbnailPool> AssetThumbnailPool;
diff --git a/Engine/Source/Editor/UnrealEd/Private/UnrealEdMisc.cpp b/Engine/Source/Editor/UnrealEd/Private/UnrealEdMisc.cpp
--- a/Engine/Source/Editor/UnrealEd/Private/UnrealEdMisc.cpp
+++ b/Engine/Source/Editor/UnrealEd/Private/UnrealEdMisc.cpp
@@ -426,7 +426,7 @@ void FUnrealEdMisc::OnInit()
 	FAssetNameToken::OnGotoAsset().BindRaw(this, &FUnrealEdMisc::OnGotoAsset);
 
 	// Register to receive notification of new key bindings
-	OnUserDefinedGestureChangedDelegateHandle = FInputBindingManager::Get().RegisterUserDefinedGestureChanged(FOnUserDefinedGestureChanged::FDelegate::CreateRaw( this, &FUnrealEdMisc::OnUserDefinedGestureChanged ));
+	OnUserDefinedChordChangedDelegateHandle = FInputBindingManager::Get().RegisterUserDefinedChordChanged(FOnUserDefinedChordChanged::FDelegate::CreateRaw( this, &FUnrealEdMisc::OnUserDefinedChordChanged ));
 
 	SlowTask.EnterProgressFrame(10);
 
@@ -756,7 +756,7 @@ void FUnrealEdMisc::OnExit()
 		FEditorViewportStats::SendUsageData();
 	}
 
-	FInputBindingManager::Get().UnregisterUserDefinedGestureChanged(OnUserDefinedGestureChangedDelegateHandle);
+	FInputBindingManager::Get().UnregisterUserDefinedChordChanged(OnUserDefinedChordChangedDelegateHandle);
 	FMessageLog::OnMessageSelectionChanged().Unbind();
 	FUObjectToken::DefaultOnMessageTokenActivated().Unbind();
 	FUObjectToken::DefaultOnGetObjectDisplayName().Unbind();
@@ -1606,17 +1606,17 @@ FString FUnrealEdMisc::GetExecutableForCommandlets() const
 	return ExecutableName;
 }
 
-void FUnrealEdMisc::OnUserDefinedGestureChanged(const FUICommandInfo& CommandInfo)
+void FUnrealEdMisc::OnUserDefinedChordChanged(const FUICommandInfo& CommandInfo)
 {
 	if( FEngineAnalytics::IsAvailable() )
 	{
-		FString GestureName = FString::Printf(TEXT("%s.%s"), *CommandInfo.GetBindingContext().ToString(), *CommandInfo.GetCommandName().ToString());
+		FString ChordName = FString::Printf(TEXT("%s.%s"), *CommandInfo.GetBindingContext().ToString(), *CommandInfo.GetCommandName().ToString());
 
 		//@todo This shouldn't be using a localized value; GetInputText() [10/11/2013 justin.sargent]
-		TArray< FAnalyticsEventAttribute > GestureAttribs;
-		GestureAttribs.Add(FAnalyticsEventAttribute(TEXT("Context"), GestureName));
-		GestureAttribs.Add(FAnalyticsEventAttribute(TEXT("Shortcut"), CommandInfo.GetActiveGesture()->GetInputText().ToString()));
-		FEngineAnalytics::GetProvider().RecordEvent(TEXT("Editor.Usage.KeyboardShortcut"), GestureAttribs);
+		TArray< FAnalyticsEventAttribute > ChordAttribs;
+		ChordAttribs.Add(FAnalyticsEventAttribute(TEXT("Context"), ChordName));
+		ChordAttribs.Add(FAnalyticsEventAttribute(TEXT("Shortcut"), CommandInfo.GetActiveChord()->GetInputText().ToString()));
+		FEngineAnalytics::GetProvider().RecordEvent(TEXT("Editor.Usage.KeyboardShortcut"), ChordAttribs);
 	}
 }
 
diff --git a/Engine/Source/Editor/UnrealEd/Public/EditorShowFlags.h b/Engine/Source/Editor/UnrealEd/Public/EditorShowFlags.h
--- a/Engine/Source/Editor/UnrealEd/Public/EditorShowFlags.h
+++ b/Engine/Source/Editor/UnrealEd/Public/EditorShowFlags.h
@@ -20,11 +20,11 @@ public:
 	/** Which group the flags should show up */
 	EShowFlagGroup	Group;
 
-	/** Gesture for activating the show flag */
-	FInputGesture InputGesture;
+	/** Chord for activating the show flag */
+	FInputChord InputChord;
 
 	/**
-	 *  FShowFlagData constructor with no gesture associated
+	 *  FShowFlagData constructor with no chord associated
 	 *
 	 * @param InName				The non-localized flag name string
 	 * @param InDisplayName			The localized display name for the flag
@@ -34,15 +34,15 @@ public:
 	FShowFlagData(const FString& InName, const FText& InDisplayName, const uint32 InEngineShowFlagIndex, EShowFlagGroup InGroup = SFG_Normal);
 
 	/**
-	 *  FShowFlagData constructor with no gesture associated
+	 *  FShowFlagData constructor with no chord associated
 	 *
 	 * @param InName				The non-localized flag name string
 	 * @param InDisplayName			The localized display name for the flag
 	 * @param InEngineShowFlagIndex	Index on this flag in the array of flags in FEngineShowFlags
 	 * @param InGroup				The group that the flag belongs to (default: SFG_Normal)
-	 * @param InInputGesture		InputGesture to be used for the flag
+	 * @param InInputChord			InputChord to be used for the flag
 	 */
-	FShowFlagData(const FString& InName, const FText& InDisplayName, const uint32 InEngineShowFlagIndex, EShowFlagGroup InGroup, const FInputGesture& InInputGesture);
+	FShowFlagData(const FString& InName, const FText& InDisplayName, const uint32 InEngineShowFlagIndex, EShowFlagGroup InGroup, const FInputChord& InInputChord);
 
 	bool IsEnabled(const FLevelEditorViewportClient* ViewportClient) const;
 
diff --git a/Engine/Source/Editor/UnrealEd/Public/GraphEditor.h b/Engine/Source/Editor/UnrealEd/Public/GraphEditor.h
--- a/Engine/Source/Editor/UnrealEd/Public/GraphEditor.h
+++ b/Engine/Source/Editor/UnrealEd/Public/GraphEditor.h
@@ -74,7 +74,7 @@ public:
 
 	DECLARE_DELEGATE_RetVal_FiveParams( FActionMenuContent, FOnCreateActionMenu, UEdGraph*, const FVector2D&, const TArray<UEdGraphPin*>&, bool, FActionMenuClosed );
 
-	DECLARE_DELEGATE_RetVal_TwoParams( FReply, FOnSpawnNodeByShortcut, FInputGesture, const FVector2D& );
+	DECLARE_DELEGATE_RetVal_TwoParams( FReply, FOnSpawnNodeByShortcut, FInputChord, const FVector2D& );
 
 	DECLARE_DELEGATE( FOnNodeSpawnedByKeymap );
 
diff --git a/Engine/Source/Editor/UnrealEd/Public/UnrealEdMisc.h b/Engine/Source/Editor/UnrealEd/Public/UnrealEdMisc.h
--- a/Engine/Source/Editor/UnrealEd/Public/UnrealEdMisc.h
+++ b/Engine/Source/Editor/UnrealEd/Public/UnrealEdMisc.h
@@ -243,8 +243,8 @@ private:
 	 */
 	void OnMapChanged( UWorld* World, EMapChangeType::Type MapChangeType );
 
-	/** Called when the input manager records a user-defined gesture */
-	void OnUserDefinedGestureChanged(const FUICommandInfo& CommandInfo);
+	/** Called when the input manager records a user-defined chord */
+	void OnUserDefinedChordChanged(const FUICommandInfo& CommandInfo);
 
 	/** Delegate for (default) message log UObject token activation - selects the object that the token refers to (if any) */
 	void OnMessageTokenActivated(const TSharedRef<class IMessageToken>& Token);
@@ -323,8 +323,8 @@ private:
 	/** Package names and the number of times they have been updated */
 	TMap<FName, uint32> NumUpdatesByAssetName;
 
-	/** Handle to the registered OnUserDefinedGestureChanged delegate. */
-	FDelegateHandle OnUserDefinedGestureChangedDelegateHandle;
+	/** Handle to the registered OnUserDefinedChordChanged delegate. */
+	FDelegateHandle OnUserDefinedChordChangedDelegateHandle;
 
 	/** Handle to the registered OnMapChanged delegate. */
 	FDelegateHandle OnMapChangedDelegateHandle;
diff --git a/Engine/Source/Editor/WorldBrowser/Private/LevelCollectionCommands.h b/Engine/Source/Editor/WorldBrowser/Private/LevelCollectionCommands.h
--- a/Engine/Source/Editor/WorldBrowser/Private/LevelCollectionCommands.h
+++ b/Engine/Source/Editor/WorldBrowser/Private/LevelCollectionCommands.h
@@ -25,94 +25,94 @@ public:
 	/** Initialize the commands */
 	virtual void RegisterCommands() override
 	{
-		UI_COMMAND( RefreshBrowser,	"Refresh",	"Refreshes opened world", EUserInterfaceActionType::Button, FInputGesture(EKeys::F5) );
+		UI_COMMAND( RefreshBrowser,	"Refresh",	"Refreshes opened world", EUserInterfaceActionType::Button, FInputChord(EKeys::F5) );
 
 		//invalid selected level
-		UI_COMMAND( FixUpInvalidReference, "Replace Selected Level","Removes the broken level and replaces it with the level browsed to", EUserInterfaceActionType::Button, FInputGesture() );
-		UI_COMMAND( RemoveInvalidReference, "Remove Selected Level", "Removes the reference to the missing level from the map", EUserInterfaceActionType::Button, FInputGesture() );
+		UI_COMMAND( FixUpInvalidReference, "Replace Selected Level","Removes the broken level and replaces it with the level browsed to", EUserInterfaceActionType::Button, FInputChord() );
+		UI_COMMAND( RemoveInvalidReference, "Remove Selected Level", "Removes the reference to the missing level from the map", EUserInterfaceActionType::Button, FInputChord() );
 		
 		//levels
-		UI_COMMAND( World_MakeLevelCurrent, "Make Current", "Make this Level the Current Level", EUserInterfaceActionType::Button, FInputGesture( EKeys::Enter ) );
-		UI_COMMAND( World_LoadLevel, "Load", "Load selected level into world", EUserInterfaceActionType::Button, FInputGesture() );
-		UI_COMMAND(World_UnloadLevel, "Unload", "Unload selected level from world", EUserInterfaceActionType::Button, FInputGesture(EKeys::Platform_Delete));
-		UI_COMMAND( World_SaveSelectedLevels, "Save", "Saves selected levels", EUserInterfaceActionType::Button, FInputGesture() );
-		UI_COMMAND( World_SaveSelectedLevelAs, "Save As...", "Save the selected level as...", EUserInterfaceActionType::Button, FInputGesture() );
-		UI_COMMAND( World_MigrateSelectedLevels, "Migrate...", "Copies the selected levels and all their dependencies to a different game", EUserInterfaceActionType::Button, FInputGesture() );
-		UI_COMMAND( World_MergeSelectedLevels, "Merge...", "Merges the selected levels into a new level, removing the original levels from the persistent", EUserInterfaceActionType::Button, FInputGesture() );
-		UI_COMMAND( World_CreateEmptyLevel, "Create New...", "Creates a new empty Level", EUserInterfaceActionType::Button, FInputGesture() );
-		UI_COMMAND( World_AddExistingLevel, "Add Existing...", "Adds an existing level", EUserInterfaceActionType::Button, FInputGesture() );
-		UI_COMMAND( World_AddSelectedActorsToNewLevel, "Create New with Selected Actors...", "Adds the actors currently selected in the active viewport to a new Level", EUserInterfaceActionType::Button, FInputGesture() );
-		UI_COMMAND( World_RemoveSelectedLevels, "Remove Selected", "Removes selected levels from the disk", EUserInterfaceActionType::Button, FInputGesture() );
-
-		UI_COMMAND( MoveWorldOrigin, "Move World Origin to Level Position", "Moves world origin to level position", EUserInterfaceActionType::Button, FInputGesture() );
-		UI_COMMAND( ResetWorldOrigin, "Reset World Origin", "Moves world origin to zero", EUserInterfaceActionType::Button, FInputGesture() );
-		UI_COMMAND( ResetLevelOrigin, "Reset Level Position", "Moves level to zero", EUserInterfaceActionType::Button, FInputGesture() );
+		UI_COMMAND( World_MakeLevelCurrent, "Make Current", "Make this Level the Current Level", EUserInterfaceActionType::Button, FInputChord( EKeys::Enter ) );
+		UI_COMMAND( World_LoadLevel, "Load", "Load selected level into world", EUserInterfaceActionType::Button, FInputChord() );
+		UI_COMMAND(World_UnloadLevel, "Unload", "Unload selected level from world", EUserInterfaceActionType::Button, FInputChord(EKeys::Platform_Delete));
+		UI_COMMAND( World_SaveSelectedLevels, "Save", "Saves selected levels", EUserInterfaceActionType::Button, FInputChord() );
+		UI_COMMAND( World_SaveSelectedLevelAs, "Save As...", "Save the selected level as...", EUserInterfaceActionType::Button, FInputChord() );
+		UI_COMMAND( World_MigrateSelectedLevels, "Migrate...", "Copies the selected levels and all their dependencies to a different game", EUserInterfaceActionType::Button, FInputChord() );
+		UI_COMMAND( World_MergeSelectedLevels, "Merge...", "Merges the selected levels into a new level, removing the original levels from the persistent", EUserInterfaceActionType::Button, FInputChord() );
+		UI_COMMAND( World_CreateEmptyLevel, "Create New...", "Creates a new empty Level", EUserInterfaceActionType::Button, FInputChord() );
+		UI_COMMAND( World_AddExistingLevel, "Add Existing...", "Adds an existing level", EUserInterfaceActionType::Button, FInputChord() );
+		UI_COMMAND( World_AddSelectedActorsToNewLevel, "Create New with Selected Actors...", "Adds the actors currently selected in the active viewport to a new Level", EUserInterfaceActionType::Button, FInputChord() );
+		UI_COMMAND( World_RemoveSelectedLevels, "Remove Selected", "Removes selected levels from the disk", EUserInterfaceActionType::Button, FInputChord() );
+
+		UI_COMMAND( MoveWorldOrigin, "Move World Origin to Level Position", "Moves world origin to level position", EUserInterfaceActionType::Button, FInputChord() );
+		UI_COMMAND( ResetWorldOrigin, "Reset World Origin", "Moves world origin to zero", EUserInterfaceActionType::Button, FInputChord() );
+		UI_COMMAND( ResetLevelOrigin, "Reset Level Position", "Moves level to zero", EUserInterfaceActionType::Button, FInputChord() );
 		
 		//level selection
-		UI_COMMAND( SelectAllLevels, "Select All Levels", "Selects all levels", EUserInterfaceActionType::Button, FInputGesture() );
-		UI_COMMAND( DeselectAllLevels, "De-select All Levels", "De-selects all levels", EUserInterfaceActionType::Button, FInputGesture() );
-		UI_COMMAND( InvertLevelSelection, "Invert Level Selection", "Inverts level selection", EUserInterfaceActionType::Button, FInputGesture() );
+		UI_COMMAND( SelectAllLevels, "Select All Levels", "Selects all levels", EUserInterfaceActionType::Button, FInputChord() );
+		UI_COMMAND( DeselectAllLevels, "De-select All Levels", "De-selects all levels", EUserInterfaceActionType::Button, FInputChord() );
+		UI_COMMAND( InvertLevelSelection, "Invert Level Selection", "Inverts level selection", EUserInterfaceActionType::Button, FInputChord() );
 
 		// Source Control
-		UI_COMMAND( SCCCheckOut, "Check Out", "Checks out the selected asset from source control.", EUserInterfaceActionType::Button, FInputGesture() );
-		UI_COMMAND( SCCCheckIn, "Check In", "Checks in the selected asset to source control.", EUserInterfaceActionType::Button, FInputGesture() );
-		UI_COMMAND( SCCOpenForAdd, "Mark For Add", "Adds the selected asset to source control.", EUserInterfaceActionType::Button, FInputGesture() );
-		UI_COMMAND( SCCHistory, "History", "Displays the source control revision history of the selected asset.", EUserInterfaceActionType::Button, FInputGesture() );
-		UI_COMMAND( SCCRefresh, "Refresh", "Updates the source control status of the asset.", EUserInterfaceActionType::Button, FInputGesture() );
-		UI_COMMAND( SCCDiffAgainstDepot, "Diff Against Depot", "Look at differences between your version of the asset and that in source control.", EUserInterfaceActionType::Button, FInputGesture() );
-		UI_COMMAND( SCCConnect, "Connect To Source Control", "Connect to source control to allow source control operations to be performed on content and levels.", EUserInterfaceActionType::Button, FInputGesture() );
+		UI_COMMAND( SCCCheckOut, "Check Out", "Checks out the selected asset from source control.", EUserInterfaceActionType::Button, FInputChord() );
+		UI_COMMAND( SCCCheckIn, "Check In", "Checks in the selected asset to source control.", EUserInterfaceActionType::Button, FInputChord() );
+		UI_COMMAND( SCCOpenForAdd, "Mark For Add", "Adds the selected asset to source control.", EUserInterfaceActionType::Button, FInputChord() );
+		UI_COMMAND( SCCHistory, "History", "Displays the source control revision history of the selected asset.", EUserInterfaceActionType::Button, FInputChord() );
+		UI_COMMAND( SCCRefresh, "Refresh", "Updates the source control status of the asset.", EUserInterfaceActionType::Button, FInputChord() );
+		UI_COMMAND( SCCDiffAgainstDepot, "Diff Against Depot", "Look at differences between your version of the asset and that in source control.", EUserInterfaceActionType::Button, FInputChord() );
+		UI_COMMAND( SCCConnect, "Connect To Source Control", "Connect to source control to allow source control operations to be performed on content and levels.", EUserInterfaceActionType::Button, FInputChord() );
 
 		// set new streaming method
-		UI_COMMAND( SetAddStreamingMethod_Blueprint, "Set Blueprint Streaming Method", "Sets the streaming method for new or added levels to Blueprint streaming", EUserInterfaceActionType::RadioButton, FInputGesture() );
-		UI_COMMAND( SetAddStreamingMethod_AlwaysLoaded, "Set Streaming to Always Loaded", "Sets the streaming method new or added selected levels to be always loaded", EUserInterfaceActionType::RadioButton, FInputGesture() );
+		UI_COMMAND( SetAddStreamingMethod_Blueprint, "Set Blueprint Streaming Method", "Sets the streaming method for new or added levels to Blueprint streaming", EUserInterfaceActionType::RadioButton, FInputChord() );
+		UI_COMMAND( SetAddStreamingMethod_AlwaysLoaded, "Set Streaming to Always Loaded", "Sets the streaming method new or added selected levels to be always loaded", EUserInterfaceActionType::RadioButton, FInputChord() );
 
 		// change streaming method
-		UI_COMMAND( SetStreamingMethod_Blueprint, "Change Blueprint Streaming Method", "Changes the streaming method for the selected levels to Blueprint streaming", EUserInterfaceActionType::Check, FInputGesture() );
-		UI_COMMAND( SetStreamingMethod_AlwaysLoaded, "Change Streaming to Always Loaded", "Changes the streaming method for the selected levels to be always loaded", EUserInterfaceActionType::Check, FInputGesture() );
+		UI_COMMAND( SetStreamingMethod_Blueprint, "Change Blueprint Streaming Method", "Changes the streaming method for the selected levels to Blueprint streaming", EUserInterfaceActionType::Check, FInputChord() );
+		UI_COMMAND( SetStreamingMethod_AlwaysLoaded, "Change Streaming to Always Loaded", "Changes the streaming method for the selected levels to be always loaded", EUserInterfaceActionType::Check, FInputChord() );
 
 		//layers
-		UI_COMMAND( AssignLevelToLayer, "Assign to layer", "Assign selected levels to current layer", EUserInterfaceActionType::Button, FInputGesture() );
+		UI_COMMAND( AssignLevelToLayer, "Assign to layer", "Assign selected levels to current layer", EUserInterfaceActionType::Button, FInputChord() );
 
 		//actors
-		UI_COMMAND( AddsActors, "Select Actors", "Adds the Actors in the selected Levels from the viewport's existing selection", EUserInterfaceActionType::Button, FInputGesture() );
-		UI_COMMAND( RemovesActors, "Deselect Actors", "Removes the Actors in the selected Levels from the viewport's existing selection", EUserInterfaceActionType::Button, FInputGesture() );
-		UI_COMMAND( MoveActorsToSelected, "Move Selected Actors to Level", "Moves the selected actors to this level", EUserInterfaceActionType::Button, FInputGesture() );
-		UI_COMMAND( MoveFoliageToSelected, "Move Selected Foliage to Level", "Moves the selected foliage instances to this level. Keeps cross-level references to original bases", EUserInterfaceActionType::Button, FInputGesture() );
-		UI_COMMAND( SelectStreamingVolumes, "Select Streaming Volumes", "Selects the streaming volumes associated with the selected levels", EUserInterfaceActionType::Button, FInputGesture() );
+		UI_COMMAND( AddsActors, "Select Actors", "Adds the Actors in the selected Levels from the viewport's existing selection", EUserInterfaceActionType::Button, FInputChord() );
+		UI_COMMAND( RemovesActors, "Deselect Actors", "Removes the Actors in the selected Levels from the viewport's existing selection", EUserInterfaceActionType::Button, FInputChord() );
+		UI_COMMAND( MoveActorsToSelected, "Move Selected Actors to Level", "Moves the selected actors to this level", EUserInterfaceActionType::Button, FInputChord() );
+		UI_COMMAND( MoveFoliageToSelected, "Move Selected Foliage to Level", "Moves the selected foliage instances to this level. Keeps cross-level references to original bases", EUserInterfaceActionType::Button, FInputChord() );
+		UI_COMMAND( SelectStreamingVolumes, "Select Streaming Volumes", "Selects the streaming volumes associated with the selected levels", EUserInterfaceActionType::Button, FInputChord() );
 
 		//visibility
-		UI_COMMAND( World_ShowSelectedLevels, "Show Selected", "Toggles selected levels to a visible state in the viewports", EUserInterfaceActionType::Button, FInputGesture() );
-		UI_COMMAND( World_HideSelectedLevels, "Hide Selected", "Toggles selected levels to an invisible state in the viewports", EUserInterfaceActionType::Button, FInputGesture() );
-		UI_COMMAND( World_ShowOnlySelectedLevels, "Show Only Selected", "Toggles the selected levels to a visible state; toggles all other levels to an invisible state", EUserInterfaceActionType::Button, FInputGesture() );
-		UI_COMMAND( World_ShowAllLevels, "Show All", "Toggles all levels to a visible state in the viewports", EUserInterfaceActionType::Button, FInputGesture() );
-		UI_COMMAND( World_HideAllLevels, "Hide All", "Hides all levels to an invisible state in the viewports", EUserInterfaceActionType::Button, FInputGesture() );
+		UI_COMMAND( World_ShowSelectedLevels, "Show Selected", "Toggles selected levels to a visible state in the viewports", EUserInterfaceActionType::Button, FInputChord() );
+		UI_COMMAND( World_HideSelectedLevels, "Hide Selected", "Toggles selected levels to an invisible state in the viewports", EUserInterfaceActionType::Button, FInputChord() );
+		UI_COMMAND( World_ShowOnlySelectedLevels, "Show Only Selected", "Toggles the selected levels to a visible state; toggles all other levels to an invisible state", EUserInterfaceActionType::Button, FInputChord() );
+		UI_COMMAND( World_ShowAllLevels, "Show All", "Toggles all levels to a visible state in the viewports", EUserInterfaceActionType::Button, FInputChord() );
+		UI_COMMAND( World_HideAllLevels, "Hide All", "Hides all levels to an invisible state in the viewports", EUserInterfaceActionType::Button, FInputChord() );
 
 		//lock
-		UI_COMMAND( World_LockSelectedLevels, "Lock Selected", "Locks selected levels", EUserInterfaceActionType::Button, FInputGesture() );
-		UI_COMMAND( World_UnockSelectedLevels, "Unlock Selected", "Unlocks selected levels", EUserInterfaceActionType::Button, FInputGesture() );
-		UI_COMMAND( World_LockAllLevels, "Lock All", "Locks all levels", EUserInterfaceActionType::Button, FInputGesture() );
-		UI_COMMAND( World_UnockAllLevels, "Unlock All", "Unlocks all levels", EUserInterfaceActionType::Button, FInputGesture() );
-		UI_COMMAND( World_LockReadOnlyLevels, "Lock Read-Only Levels", "Locks all read-only levels", EUserInterfaceActionType::Button, FInputGesture() );
-		UI_COMMAND( World_UnlockReadOnlyLevels, "Unlock Read-Only Levels", "Unlocks all read-only levels", EUserInterfaceActionType::Button, FInputGesture() );
+		UI_COMMAND( World_LockSelectedLevels, "Lock Selected", "Locks selected levels", EUserInterfaceActionType::Button, FInputChord() );
+		UI_COMMAND( World_UnockSelectedLevels, "Unlock Selected", "Unlocks selected levels", EUserInterfaceActionType::Button, FInputChord() );
+		UI_COMMAND( World_LockAllLevels, "Lock All", "Locks all levels", EUserInterfaceActionType::Button, FInputChord() );
+		UI_COMMAND( World_UnockAllLevels, "Unlock All", "Unlocks all levels", EUserInterfaceActionType::Button, FInputChord() );
+		UI_COMMAND( World_LockReadOnlyLevels, "Lock Read-Only Levels", "Locks all read-only levels", EUserInterfaceActionType::Button, FInputChord() );
+		UI_COMMAND( World_UnlockReadOnlyLevels, "Unlock Read-Only Levels", "Unlocks all read-only levels", EUserInterfaceActionType::Button, FInputChord() );
 		
 		//view
-		UI_COMMAND( FitToSelection, "Fit to Selection", "Fits view to selected levels", EUserInterfaceActionType::Button, FInputGesture(EKeys::Home) );
-		UI_COMMAND( ExpandSelectedItems, "Expand Selected", "Expands all children of a selected items", EUserInterfaceActionType::Button, FInputGesture() );
+		UI_COMMAND( FitToSelection, "Fit to Selection", "Fits view to selected levels", EUserInterfaceActionType::Button, FInputChord(EKeys::Home) );
+		UI_COMMAND( ExpandSelectedItems, "Expand Selected", "Expands all children of a selected items", EUserInterfaceActionType::Button, FInputChord() );
 		
 		// Parent->child
-		UI_COMMAND( ClearParentLink, "Clear Parent Link", "Clears parent link for selected items", EUserInterfaceActionType::Button, FInputGesture() );
+		UI_COMMAND( ClearParentLink, "Clear Parent Link", "Clears parent link for selected items", EUserInterfaceActionType::Button, FInputChord() );
 
-		UI_COMMAND( MoveLevelLeft,	"Move Level Left",		"Moves level to the left by 1 unit", EUserInterfaceActionType::Button, FInputGesture(EKeys::Left) );
-		UI_COMMAND( MoveLevelRight, "Move Level Right",		"Moves level to the right by 1 unit", EUserInterfaceActionType::Button, FInputGesture(EKeys::Right) );
-		UI_COMMAND( MoveLevelUp,	"Move Level Up",		"Moves level up by 1 unit", EUserInterfaceActionType::Button, FInputGesture(EKeys::Up) );
-		UI_COMMAND( MoveLevelDown,	"Move Level Down",		"Moves level down by 1 unit", EUserInterfaceActionType::Button, FInputGesture(EKeys::Down) );
+		UI_COMMAND( MoveLevelLeft,	"Move Level Left",		"Moves level to the left by 1 unit", EUserInterfaceActionType::Button, FInputChord(EKeys::Left) );
+		UI_COMMAND( MoveLevelRight, "Move Level Right",		"Moves level to the right by 1 unit", EUserInterfaceActionType::Button, FInputChord(EKeys::Right) );
+		UI_COMMAND( MoveLevelUp,	"Move Level Up",		"Moves level up by 1 unit", EUserInterfaceActionType::Button, FInputChord(EKeys::Up) );
+		UI_COMMAND( MoveLevelDown,	"Move Level Down",		"Moves level down by 1 unit", EUserInterfaceActionType::Button, FInputChord(EKeys::Down) );
 		
 		// Landscape operations
-		UI_COMMAND( ImportTiledLandscape, "Import Tiled Landscape...", "Imports landscape from a tiled heightmap files (<name>X<n>_Y<n>.r16)", EUserInterfaceActionType::Button, FInputGesture() );
-		UI_COMMAND( AddLandscapeLevelXNegative,	"-X",	"Add a new adjacent level with landscape proxy", EUserInterfaceActionType::Button, FInputGesture() );
-		UI_COMMAND( AddLandscapeLevelXPositive,	"+X",	"Add a new adjacent level with landscape proxy", EUserInterfaceActionType::Button, FInputGesture() );
-		UI_COMMAND( AddLandscapeLevelYNegative,	"-Y",	"Add a new adjacent level with landscape proxy", EUserInterfaceActionType::Button, FInputGesture() );
-		UI_COMMAND( AddLandscapeLevelYPositive,	"+Y",	"Add a new adjacent level with landscape proxy", EUserInterfaceActionType::Button, FInputGesture() );
+		UI_COMMAND( ImportTiledLandscape, "Import Tiled Landscape...", "Imports landscape from a tiled heightmap files (<name>X<n>_Y<n>.r16)", EUserInterfaceActionType::Button, FInputChord() );
+		UI_COMMAND( AddLandscapeLevelXNegative,	"-X",	"Add a new adjacent level with landscape proxy", EUserInterfaceActionType::Button, FInputChord() );
+		UI_COMMAND( AddLandscapeLevelXPositive,	"+X",	"Add a new adjacent level with landscape proxy", EUserInterfaceActionType::Button, FInputChord() );
+		UI_COMMAND( AddLandscapeLevelYNegative,	"-Y",	"Add a new adjacent level with landscape proxy", EUserInterfaceActionType::Button, FInputChord() );
+		UI_COMMAND( AddLandscapeLevelYPositive,	"+Y",	"Add a new adjacent level with landscape proxy", EUserInterfaceActionType::Button, FInputChord() );
 	}
 
 public:
diff --git a/Engine/Source/Runtime/AppFramework/Private/Widgets/Testing/STestSuite.cpp b/Engine/Source/Runtime/AppFramework/Private/Widgets/Testing/STestSuite.cpp
--- a/Engine/Source/Runtime/AppFramework/Private/Widgets/Testing/STestSuite.cpp
+++ b/Engine/Source/Runtime/AppFramework/Private/Widgets/Testing/STestSuite.cpp
@@ -3084,18 +3084,18 @@ public:
 public:
 	void RegisterCommands()
 	{
-		UI_COMMAND( FirstCommandInfo, "First Test", "This is the first test menu item", EUserInterfaceActionType::ToggleButton, FInputGesture() );
-		UI_COMMAND( SecondCommandInfo, "Second Test", "This is the second test menu item. Shows a keybinding", EUserInterfaceActionType::ToggleButton, FInputGesture( EModifierKey::Shift, EKeys::A ) );
+		UI_COMMAND( FirstCommandInfo, "First Test", "This is the first test menu item", EUserInterfaceActionType::ToggleButton, FInputChord() );
+		UI_COMMAND( SecondCommandInfo, "Second Test", "This is the second test menu item. Shows a keybinding", EUserInterfaceActionType::ToggleButton, FInputChord( EModifierKey::Shift, EKeys::A ) );
 
-		UI_COMMAND( ThirdCommandInfo, "Third Test", "This is the thrid test menu item", EUserInterfaceActionType::ToggleButton, FInputGesture() );
-		UI_COMMAND( FourthCommandInfo, "Fourth Test", "This is the fourth test menu item", EUserInterfaceActionType::ToggleButton, FInputGesture() );
+		UI_COMMAND( ThirdCommandInfo, "Third Test", "This is the thrid test menu item", EUserInterfaceActionType::ToggleButton, FInputChord() );
+		UI_COMMAND( FourthCommandInfo, "Fourth Test", "This is the fourth test menu item", EUserInterfaceActionType::ToggleButton, FInputChord() );
 
 
-		UI_COMMAND( FifthCommandInfo, "Fifth Test", "This is the fifth test menu item", EUserInterfaceActionType::ToggleButton, FInputGesture() );
-		UI_COMMAND( SixthCommandInfo, "Sixth Test", "This is the sixth test menu item. Shows a keybinding", EUserInterfaceActionType::ToggleButton, FInputGesture() );
+		UI_COMMAND( FifthCommandInfo, "Fifth Test", "This is the fifth test menu item", EUserInterfaceActionType::ToggleButton, FInputChord() );
+		UI_COMMAND( SixthCommandInfo, "Sixth Test", "This is the sixth test menu item. Shows a keybinding", EUserInterfaceActionType::ToggleButton, FInputChord() );
 
-		UI_COMMAND( SeventhCommandInfo, "Seventh Test", "This is the seventh test menu item", EUserInterfaceActionType::ToggleButton, FInputGesture() );
-		UI_COMMAND( EighthCommandInfo, "Eighth Test", "This is the eighth test menu item", EUserInterfaceActionType::ToggleButton, FInputGesture() );
+		UI_COMMAND( SeventhCommandInfo, "Seventh Test", "This is the seventh test menu item", EUserInterfaceActionType::ToggleButton, FInputChord() );
+		UI_COMMAND( EighthCommandInfo, "Eighth Test", "This is the eighth test menu item", EUserInterfaceActionType::ToggleButton, FInputChord() );
 	}
 };
 
diff --git a/Engine/Source/Runtime/Engine/Classes/Components/InputComponent.h b/Engine/Source/Runtime/Engine/Classes/Components/InputComponent.h
--- a/Engine/Source/Runtime/Engine/Classes/Components/InputComponent.h
+++ b/Engine/Source/Runtime/Engine/Classes/Components/InputComponent.h
@@ -2,6 +2,7 @@
 
 #pragma once
 #include "Runtime/InputCore/Classes/InputCoreTypes.h"
+#include "InputChord.h"
 #include "InputComponent.generated.h"
 
 
@@ -76,72 +77,6 @@ protected:
 };
 
 
-/** An Input Chord is a key and the modifier keys that are to be held with it. */
-USTRUCT()
-struct FInputChord
-{
-	GENERATED_USTRUCT_BODY()
-
-	/** The Key is the core of the chord. */
-	UPROPERTY()
-	FKey Key;
-
-	/** Whether the shift key is part of the chord.  */
-	UPROPERTY()
-	uint32 bShift:1;
-
-	/** Whether the control key is part of the chord.  */
-	UPROPERTY()
-	uint32 bCtrl:1;
-
-	/** Whether the alt key is part of the chord.  */
-	UPROPERTY()
-	uint32 bAlt:1;
-
-	/** Whether the command key is part of the chord.  */
-	UPROPERTY()
-	uint32 bCmd:1;
-
-	/** 
-	  * The ways two chords can be related to each other. A chord is considered masking 
-	  * when it has all the same modifier keys as another chord plus more.
-	  */
-	enum ERelationshipType
-	{
-		None,
-		Same,
-		Masked,
-		Masks
-	};
-
-	/** Returns the relationship between this chord and another. */
-	ERelationshipType GetRelationship(const FInputChord& OtherChord) const;
-
-	FInputChord()
-		: bShift(false)
-		, bCtrl(false)
-		, bAlt(false)
-		, bCmd(false)
-	{ }
-
-	FInputChord(const FKey InKey, const bool bInShift, const bool bInCtrl, const bool bInAlt, const bool bInCmd)
-		: Key(InKey)
-		, bShift(bInShift)
-		, bCtrl(bInCtrl)
-		, bAlt(bInAlt)
-		, bCmd(bInCmd)
-	{ }
-
-	bool operator==(const FInputChord& Other) const
-	{
-		return (   Key == Other.Key
-				&& bShift == Other.bShift
-				&& bCtrl == Other.bCtrl
-				&& bAlt == Other.bAlt
-				&& bCmd == Other.bCmd);
-	}
-};
-
 /** Base class for the different binding types. */
 struct FInputBinding
 {
diff --git a/Engine/Source/Runtime/Engine/Private/Components/InputComponent.cpp b/Engine/Source/Runtime/Engine/Private/Components/InputComponent.cpp
--- a/Engine/Source/Runtime/Engine/Private/Components/InputComponent.cpp
+++ b/Engine/Source/Runtime/Engine/Private/Components/InputComponent.cpp
@@ -3,42 +3,6 @@
 #include "EnginePrivate.h"
 
 
-/* FInputChord interface
- *****************************************************************************/
-
-FInputChord::ERelationshipType FInputChord::GetRelationship( const FInputChord& OtherChord ) const
-{
-	ERelationshipType Relationship = None;
-
-	if (Key == OtherChord.Key)
-	{
-		if ((bAlt == OtherChord.bAlt) &&
-			(bCtrl == OtherChord.bCtrl) &&
-			(bShift == OtherChord.bShift) &&
-			(bCmd == OtherChord.bCmd))
-		{
-			Relationship = Same;
-		}
-		else if ((bAlt || !OtherChord.bAlt) &&
-				(bCtrl || !OtherChord.bCtrl) &&
-				(bShift || !OtherChord.bShift) &&
-				(bCmd || !OtherChord.bCmd))
-		{
-			Relationship = Masks;
-		}
-		else if ((!bAlt || OtherChord.bAlt) &&
-				(!bCtrl || OtherChord.bCtrl) &&
-				(!bShift || OtherChord.bShift) &&
-				(!bCmd || OtherChord.bCmd))
-		{
-			Relationship = Masked;
-		}
-	}
-
-	return Relationship;
-}
-
-
 /* UInputComponent interface
  *****************************************************************************/
 
diff --git a/Engine/Source/Runtime/Slate/Private/Framework/Commands/Commands.cpp b/Engine/Source/Runtime/Slate/Private/Framework/Commands/Commands.cpp
--- a/Engine/Source/Runtime/Slate/Private/Framework/Commands/Commands.cpp
+++ b/Engine/Source/Runtime/Slate/Private/Framework/Commands/Commands.cpp
@@ -4,7 +4,7 @@
 
 #define LOC_DEFINE_REGION
 
-void UI_COMMAND_Function( FBindingContext* This, TSharedPtr< FUICommandInfo >& OutCommand, const TCHAR* OutSubNamespace, const TCHAR* OutCommandName, const TCHAR* OutCommandNameUnderscoreTooltip, const ANSICHAR* DotOutCommandName, const TCHAR* FriendlyName, const TCHAR* InDescription, const EUserInterfaceActionType::Type CommandType, const FInputGesture& InDefaultGesture )
+void UI_COMMAND_Function( FBindingContext* This, TSharedPtr< FUICommandInfo >& OutCommand, const TCHAR* OutSubNamespace, const TCHAR* OutCommandName, const TCHAR* OutCommandNameUnderscoreTooltip, const ANSICHAR* DotOutCommandName, const TCHAR* FriendlyName, const TCHAR* InDescription, const EUserInterfaceActionType::Type CommandType, const FInputChord& InDefaultChord )
 {
 	static const FString UICommandsStr(TEXT("UICommands"));
 	const FString Namespace = OutSubNamespace && FCString::Strlen(OutSubNamespace) > 0 ?  UICommandsStr + TEXT(".") + OutSubNamespace : UICommandsStr;
@@ -17,7 +17,7 @@ void UI_COMMAND_Function( FBindingContext* This, TSharedPtr< FUICommandInfo >& O
 		FInternationalization::ForUseOnlyByLocMacroAndGraphNodeTextLiterals_CreateText( InDescription, *Namespace, OutCommandNameUnderscoreTooltip ),
 		FSlateIcon( This->GetStyleSetName(), ISlateStyle::Join( This->GetContextName(), DotOutCommandName ) ),
 		CommandType,
-		InDefaultGesture
+		InDefaultChord
 	);
 }
 
