Fixed issue where cancelling from the "Check Out" dialog during saving assets would not cancel the whole operation.

diff --git a/Engine/Source/Editor/UnrealEd/Private/FileHelpers.cpp b/Engine/Source/Editor/UnrealEd/Private/FileHelpers.cpp
--- a/Engine/Source/Editor/UnrealEd/Private/FileHelpers.cpp
+++ b/Engine/Source/Editor/UnrealEd/Private/FileHelpers.cpp
@@ -2914,21 +2914,18 @@ FEditorFileUtils::EPromptReturnCode FEditorFileUtils::PromptForCheckoutAndSave(
 		// Prompt to check-out any packages under source control
 		const bool UserResponse = bAlreadyCheckedOut || FEditorFileUtils::PromptToCheckoutPackages( false, PackagesToSave, &PackagesCheckedOutOrMadeWritable, &PackagesNotNeedingCheckout );
 
-		if( UserResponse || PackagesNotNeedingCheckout.Num() > 0 )
+		if( UserResponse )
 		{
-			// Even if the user cancelled the checkout dialog, still save packages not needing checkout
-			TArray<UPackage*> FinalSaveList = PackagesNotNeedingCheckout;
-
-			if ( UserResponse )
+			TArray<UPackage*> FinalSaveList;
+			
+			if (bAlreadyCheckedOut)
 			{
-				if (bAlreadyCheckedOut)
-				{
-					FinalSaveList.Append(PackagesToSave);
-				}
-				else
-				{
-					FinalSaveList.Append(PackagesCheckedOutOrMadeWritable);
-				}
+				FinalSaveList = PackagesToSave;
+			}
+			else
+			{
+				FinalSaveList = PackagesNotNeedingCheckout;
+				FinalSaveList.Append(PackagesCheckedOutOrMadeWritable);
 			}
 
 			const FScopedBusyCursor BusyCursor;
