Checking Layers availability in sub-level managment before acessing it, could be null in commandlets

diff --git a/Engine/Source/Editor/UnrealEd/Private/EditorLevelUtils.cpp b/Engine/Source/Editor/UnrealEd/Private/EditorLevelUtils.cpp
--- a/Engine/Source/Editor/UnrealEd/Private/EditorLevelUtils.cpp
+++ b/Engine/Source/Editor/UnrealEd/Private/EditorLevelUtils.cpp
@@ -399,8 +399,11 @@ namespace EditorLevelUtils
 			}
 		}
 
-		GEditor->Layers->RemoveLevelLayerInformation( InLevel );
-
+		if (GEditor->Layers.IsValid())
+		{
+			GEditor->Layers->RemoveLevelLayerInformation( InLevel );
+		}
+		
 		GEditor->CloseEditedWorldAssets(CastChecked<UWorld>(InLevel->GetOuter()));
 
 		UWorld* OwningWorld = InLevel->OwningWorld;
@@ -729,7 +732,7 @@ namespace EditorLevelUtils
 				StreamingLevel->bShouldBeVisibleInEditor = bShouldBeVisible;
 			}
 
-			if( !bShouldBeVisible )
+			if (!bShouldBeVisible && GEditor->Layers.IsValid())
 			{
 				GEditor->Layers->RemoveLevelLayerInformation( Level );
 			}
@@ -766,7 +769,7 @@ namespace EditorLevelUtils
 				check(Level->bIsVisible == bShouldBeVisible);
 			}
 
-			if( bShouldBeVisible )
+			if (bShouldBeVisible && GEditor->Layers.IsValid())
 			{
 				GEditor->Layers->AddLevelLayerInformation( Level );
 			}
@@ -782,7 +785,9 @@ namespace EditorLevelUtils
 				if ( Actor )
 				{
 					bool bModified = false;
-					if ( bShouldBeVisible && bForceLayersVisible && GEditor->Layers->IsActorValidForLayer( Actor ) )
+					if (bShouldBeVisible && bForceLayersVisible && 
+						GEditor->Layers.IsValid() && 
+						GEditor->Layers->IsActorValidForLayer( Actor ))
 					{
 						// Make the actor layer visible, if it's not already.
 						if ( Actor->bHiddenEdLayer )
