Added ability to update reimport paths for a given object

diff --git a/Engine/Source/Editor/UnrealEd/Private/Editor.cpp b/Engine/Source/Editor/UnrealEd/Private/Editor.cpp
--- a/Engine/Source/Editor/UnrealEd/Private/Editor.cpp
+++ b/Engine/Source/Editor/UnrealEd/Private/Editor.cpp
@@ -2609,6 +2609,20 @@ bool FReimportManager::CanReimport( UObject* Obj ) const
 	return false;
 }
 
+void FReimportManager::UpdateReimportPaths( UObject* Obj, const TArray<FString>& InFilenames )
+{
+	if (Obj)
+	{
+		TArray<FString> UnusedExistingFilenames;
+		auto* Handler = Handlers.FindByPredicate([&](FReimportHandler* Handler){ return Handler->CanReimport(Obj, UnusedExistingFilenames); });
+		if (Handler)
+		{
+			(*Handler)->SetReimportPaths(Obj, InFilenames);
+			Obj->MarkPackageDirty();
+		}
+	}
+}
+
 bool FReimportManager::Reimport( UObject* Obj, bool bAskForNewFileIfMissing, bool bShowNotification )
 {
 	// Warn that were about to reimport, so prep for it
diff --git a/Engine/Source/Editor/UnrealEd/Public/EditorReimportHandler.h b/Engine/Source/Editor/UnrealEd/Public/EditorReimportHandler.h
--- a/Engine/Source/Editor/UnrealEd/Public/EditorReimportHandler.h
+++ b/Engine/Source/Editor/UnrealEd/Public/EditorReimportHandler.h
@@ -51,6 +51,14 @@ public:
 	 */
 	UNREALED_API virtual bool Reimport( UObject* Obj, bool bAskForNewFileIfMissing = false, bool bShowNotification = true );
 
+	/**
+	 * Update the reimport paths for the specified object
+	 *
+	 * @param	Obj	Object to update
+	 * @param	InFilenames The files we want to set to its import paths
+	 */
+	UNREALED_API virtual void UpdateReimportPaths( UObject* Obj, const TArray<FString>& InFilenames );
+
 	/**
 	 * Convert a file path to be relative to the specified object, if it resides in the same package folder
 	 * 
