Fixed misspelling in comment.

diff --git a/Engine/Source/Editor/UnrealEd/Private/Tests/EditorOpenAssetAutomationTests.cpp b/Engine/Source/Editor/UnrealEd/Private/Tests/EditorOpenAssetAutomationTests.cpp
--- a/Engine/Source/Editor/UnrealEd/Private/Tests/EditorOpenAssetAutomationTests.cpp
+++ b/Engine/Source/Editor/UnrealEd/Private/Tests/EditorOpenAssetAutomationTests.cpp
@@ -371,7 +371,6 @@ bool FOpenBlendSpace1DAssetEditors::RunTest(const FString& Parameters)
 	return bDidTheTestPass;
 }
 
-
 //////////////////////////////////////////////////////////////////////////
 /**
 * This test opens each BLUEPRINT into its sub-editor.
@@ -380,7 +379,7 @@ IMPLEMENT_COMPLEX_AUTOMATION_TEST(FOpenBlueprintAssetEditors, "Tools.Blueprint.O
 
 void FOpenBlueprintAssetEditors::GetTests(TArray<FString>& OutBeautifiedNames, TArray<FString>& OutTestCommands) const
 {
-	//This grabs each BlendSpace asset in the Game/Content directory
+	//This grabs each Blueprints asset in the Game/Content directory
 	FEditorAutomationTestUtilities::CollectGameContentTestsByClass(UBlueprint::StaticClass(), true, OutBeautifiedNames, OutTestCommands);
 }
 
