Auto reipmort improvements

* Renames and moves within monitored directories are now detected and reacted to correctly
* Users can now specify virtual mapped paths for monitored directories. This allows new content files to prompt the auto creation of assets relative to the specified mapped directory.
* It's now possible to add wildcards to include and exclude specific files from monitored directories.

diff --git a/Engine/Source/Editor/UnrealEd/Classes/Settings/EditorLoadingSavingSettings.h b/Engine/Source/Editor/UnrealEd/Classes/Settings/EditorLoadingSavingSettings.h
--- a/Engine/Source/Editor/UnrealEd/Classes/Settings/EditorLoadingSavingSettings.h
+++ b/Engine/Source/Editor/UnrealEd/Classes/Settings/EditorLoadingSavingSettings.h
@@ -16,6 +16,51 @@ namespace ELoadLevelAtStartup
 	};
 }
 
+/** A filter used by the auto reimport manager to explicitly include/exclude files matching the specified wildcard */
+USTRUCT()
+struct FAutoReimportWildcard
+{
+	GENERATED_USTRUCT_BODY()
+
+	/** The wildcard filter as a string. Files that match this wildcard will be included/excluded according to the bInclude member */
+	UPROPERTY(EditAnywhere, config, Category=AutoReimport)
+	FString Wildcard;
+	
+	/** When true, files that match this wildcard will be included (if it doesn't fail any other filters), when false, matches will be excluded from the reimporter */
+	UPROPERTY(EditAnywhere, config, Category=AutoReimport)
+	bool bInclude;
+};
+
+/** Auto reimport settings for a specific directory */
+USTRUCT()
+struct FAutoReimportDirectoryConfig
+{
+	GENERATED_USTRUCT_BODY()
+
+	/** The source directory to monitor. Either an absolute directory on the file system, or a virtual mounted path */
+	UPROPERTY(EditAnywhere, config, Category=AutoReimport)
+	FString SourceDirectory;
+
+	/** Where SourceDirectory points to an ordinary file system path, MountPoint specifies the virtual mounted location to import new files to. */
+	UPROPERTY(EditAnywhere, config, Category=AutoReimport)
+	FString MountPoint;
+
+	/** A set of wildcard filters to apply to this directory */
+	UPROPERTY(EditAnywhere, config, Category=AutoReimport)
+	TArray<FAutoReimportWildcard> Wildcards;
+
+	struct UNREALED_API FParseContext
+	{
+		TArray<TPair<FString, FString>> MountedPaths;
+		bool bEnableLogging;
+		FParseContext(bool bInEnableLogging = true);
+	};
+
+	/** Parse and validate the specified source directory / mount point combination */
+	UNREALED_API static bool ParseSourceDirectoryAndMountPoint(FString& SourceDirectory, FString& MountPoint, const FParseContext& InContext = FParseContext());
+};
+
+
 /**
  * Implements the Level Editor's loading and saving settings.
  */
@@ -45,14 +90,19 @@ public:
 	UPROPERTY(EditAnywhere, config, Category=AutoReimport, meta=(DisplayName="Monitor Content Directories", ToolTip="When enabled, changes to made to source content files inside the content directories will automatically be reflected in the content browser."))
 	bool bMonitorContentDirectories;
 
+	UPROPERTY()
+	TArray<FString> AutoReimportDirectories_DEPRECATED;
+
 	/** Directories being monitored for Auto Reimport */
 	UPROPERTY(EditAnywhere, config, AdvancedDisplay,Category=AutoReimport, meta=(DisplayName="Directories to Monitor", ToolTip="Lists every directory to monitor for content changes. Can be virtual package paths (eg /Game/ or /MyPlugin/), or absolute paths on disk."))
-	TArray<FString> AutoReimportDirectories;
+	TArray<FAutoReimportDirectoryConfig> AutoReimportDirectorySettings;
 
 	UPROPERTY(EditAnywhere, config, AdvancedDisplay, Category=AutoReimport, meta=(DisplayName="Auto Create Assets", ToolTip="When enabled, newly added source content files will be automatically imported into new assets."))
 	bool bAutoCreateAssets;
 	UPROPERTY(EditAnywhere, config, AdvancedDisplay, Category=AutoReimport, meta=(DisplayName="Auto Delete Assets", ToolTip="When enabled, deleting a source content file will automatically prompt the deletion of any related assets."))
 	bool bAutoDeleteAssets;
+	UPROPERTY(EditAnywhere, config, AdvancedDisplay, Category=AutoReimport, meta=(DisplayName="Detect Changes On Restart", ToolTip="When enabled, changes to monitored directories since UE4 was closed will be detected on restart.\n(Not recommended when working in collaboration with others using source control)."))
+	bool bDetectChangesOnRestart;
 
 	/** Internal setting to control whether we should ask the user whether we should automatically delete source files when their assets are deleted */
 	UPROPERTY(config)
@@ -136,6 +186,7 @@ protected:
 	// UObject overrides
 
 	virtual void PostEditChangeProperty( struct FPropertyChangedEvent& PropertyChangedEvent ) override;
+	virtual void PostInitProperties() override;
 
 private:
 
diff --git a/Engine/Source/Editor/UnrealEd/Private/AssetDeleteModel.cpp b/Engine/Source/Editor/UnrealEd/Private/AssetDeleteModel.cpp
--- a/Engine/Source/Editor/UnrealEd/Private/AssetDeleteModel.cpp
+++ b/Engine/Source/Editor/UnrealEd/Private/AssetDeleteModel.cpp
@@ -200,7 +200,7 @@ void FAssetDeleteModel::DeleteSourceContentFiles()
 		}
 
 		// One way or another this file is going to be deleted, but we don't want the import manager to react to the deletion
-		GUnrealEd->AutoReimportManager->ReportExternalChange(Path, FFileChangeData::FCA_Removed);
+		GUnrealEd->AutoReimportManager->IgnoreDeletedFile(Path);
 
 		if (ISourceControlModule::Get().IsEnabled())
 		{
diff --git a/Engine/Source/Editor/UnrealEd/Private/AutoReimport/AutoReimportManager.cpp b/Engine/Source/Editor/UnrealEd/Private/AutoReimport/AutoReimportManager.cpp
--- a/Engine/Source/Editor/UnrealEd/Private/AutoReimport/AutoReimportManager.cpp
+++ b/Engine/Source/Editor/UnrealEd/Private/AutoReimport/AutoReimportManager.cpp
@@ -6,6 +6,7 @@
 #include "AutoReimport/AutoReimportManager.h"
 #include "ContentDirectoryMonitor.h"
 
+#include "PackageTools.h"
 #include "ObjectTools.h"
 #include "ARFilter.h"
 #include "ContentBrowserModule.h"
@@ -16,9 +17,6 @@
 #define LOCTEXT_NAMESPACE "AutoReimportManager"
 #define yield if (TimeLimit.Exceeded()) return
 
-DEFINE_LOG_CATEGORY_STATIC(LogAutoReimportManager, Log, All);
-
-
 enum class EStateMachineNode { CallOnce, CallMany };
 
 /** A simple state machine class that calls functions mapped on enum values. If any function returns a new enum type, it moves onto that function */
@@ -103,7 +101,10 @@ public:
 	TArray<FPathAndMountPoint> GetMonitoredDirectories() const;
 
 	/** Report an external change to the manager, such that a subsequent equal change reported by the os be ignored */
-	void ReportExternalChange(const FString& Filename, FFileChangeData::EFileChangeAction Action);
+	void IgnoreNewFile(const FString& Filename);
+	void IgnoreFileModification(const FString& Filename);
+	void IgnoreMovedFile(const FString& SrcFilename, const FString& DstFilename);
+	void IgnoreDeletedFile(const FString& Filename);
 
 	/** Destroy this manager */
 	void Destroy();
@@ -123,6 +124,9 @@ private:
 	/** Called when a new asset path has been mounted or unmounted */
 	void OnContentPathChanged(const FString& InAssetPath, const FString& FileSystemPath);
 
+	/** Called when an asset has been renamed */
+	void OnAssetRenamed(const FAssetData& AssetData, const FString& OldPath);
+
 private:
 
 	/** Callback for when an editor user setting has changed */
@@ -150,7 +154,7 @@ private:
 private:
 
 	/** Idle processing */
-	TOptional<ECurrentState> Idle(const FTimeLimit& Limit);
+	TOptional<ECurrentState> Idle();
 	
 	/** Process any remaining pending additions we have */
 	TOptional<ECurrentState> ProcessAdditions(const FTimeLimit& TimeLimit);
@@ -175,9 +179,6 @@ private:
 	/** Feedback context that can selectively override the global context to consume progress events for saving of assets */
 	TSharedPtr<FReimportFeedbackContext> FeedbackContextOverride;
 
-	/** Cached string of extensions that are supported by all available factories */
-	FString SupportedExtensions;
-
 	/** Array of objects that detect changes to directories */
 	TArray<FContentDirectoryMonitor> DirectoryMonitors;
 
@@ -189,15 +190,16 @@ private:
 
 	/** The cached number of unprocessed changes we currently have to process */
 	int32 CachedNumUnprocessedChanges;
+
+	/** Reentracy guard for when we are making changes to assets */
+	bool bGuardAssetChanges;
 };
 
 FAutoReimportManager::FAutoReimportManager()
 	: StateMachine(ECurrentState::Idle)
 	, StartProcessingDelay(0.5)
+	, bGuardAssetChanges(false)
 {
-	// @todo: arodham: update this when modules are reloaded or new factory types are available?
-	SupportedExtensions = GetAllFactoryExtensions();
-
 	auto* Settings = GetMutableDefault<UEditorLoadingSavingSettings>();
 	Settings->OnSettingChanged().AddRaw(this, &FAutoReimportManager::HandleLoadingSavingSettingChanged);
 
@@ -210,13 +212,16 @@ FAutoReimportManager::FAutoReimportManager()
 		MessageLogModule.RegisterLogListing("AssetReimport", LOCTEXT("AssetReimportLabel", "Asset Reimport"));
 	}
 
+	IAssetRegistry& AssetRegistry = FModuleManager::LoadModuleChecked<FAssetRegistryModule>("AssetRegistry").Get();
+	AssetRegistry.OnAssetRenamed().AddRaw(this, &FAutoReimportManager::OnAssetRenamed);
+
 	// Only set this up for content directories if the user has this enabled
 	if (Settings->bMonitorContentDirectories)
 	{
 		SetUpDirectoryMonitors();
 	}
 
-	StateMachine.Add(ECurrentState::Idle,					EStateMachineNode::CallOnce,	[this](const FTimeLimit& T){ return this->Idle(T);					});
+	StateMachine.Add(ECurrentState::Idle,					EStateMachineNode::CallOnce,	[this](const FTimeLimit& T){ return this->Idle();					});
 	StateMachine.Add(ECurrentState::PendingResponse,		EStateMachineNode::CallOnce,	[this](const FTimeLimit& T){ return this->PendingResponse();		});
 	StateMachine.Add(ECurrentState::ProcessAdditions,		EStateMachineNode::CallMany,	[this](const FTimeLimit& T){ return this->ProcessAdditions(T);		});
 	StateMachine.Add(ECurrentState::SavePackages,			EStateMachineNode::CallOnce,	[this](const FTimeLimit& T){ return this->SavePackages();			});
@@ -234,19 +239,68 @@ TArray<FPathAndMountPoint> FAutoReimportManager::GetMonitoredDirectories() const
 	return Dirs;
 }
 
-void FAutoReimportManager::ReportExternalChange(const FString& Filename, FFileChangeData::EFileChangeAction Action)
+void FAutoReimportManager::IgnoreNewFile(const FString& Filename)
+{
+	for (auto& Monitor : DirectoryMonitors)
+	{
+		if (Filename.StartsWith(Monitor.GetDirectory()))
+		{
+			Monitor.IgnoreNewFile(Filename);
+		}
+	}
+}
+
+void FAutoReimportManager::IgnoreFileModification(const FString& Filename)
+{
+	for (auto& Monitor : DirectoryMonitors)
+	{
+		if (Filename.StartsWith(Monitor.GetDirectory()))
+		{
+			Monitor.IgnoreFileModification(Filename);
+		}
+	}
+}
+
+void FAutoReimportManager::IgnoreMovedFile(const FString& SrcFilename, const FString& DstFilename)
+{
+	for (auto& Monitor : DirectoryMonitors)
+	{
+		const bool bSrcInFolder = SrcFilename.StartsWith(Monitor.GetDirectory());
+		const bool bDstInFolder = DstFilename.StartsWith(Monitor.GetDirectory());
+
+		if (bSrcInFolder && bDstInFolder)
+		{
+			Monitor.IgnoreMovedFile(SrcFilename, DstFilename);
+		}
+		else if (bSrcInFolder)
+		{
+			Monitor.IgnoreDeletedFile(SrcFilename);
+		}
+		else if (bDstInFolder)
+		{
+			Monitor.IgnoreNewFile(DstFilename);
+		}
+	}
+}
+
+void FAutoReimportManager::IgnoreDeletedFile(const FString& Filename)
 {
 	for (auto& Monitor : DirectoryMonitors)
 	{
 		if (Filename.StartsWith(Monitor.GetDirectory()))
 		{
-			Monitor.ReportExternalChange(Filename, Action);
+			Monitor.IgnoreDeletedFile(Filename);
 		}
 	}
 }
 
 void FAutoReimportManager::Destroy()
 {
+	if (auto* AssetRegistryModule = FModuleManager::GetModulePtr<FAssetRegistryModule>("AssetRegistry"))
+	{
+		AssetRegistryModule->Get().OnAssetRenamed().RemoveAll(this);
+	}
+	
 	if (auto* Settings = GetMutableDefault<UEditorLoadingSavingSettings>())
 	{
 		Settings->OnSettingChanged().RemoveAll(this);
@@ -264,6 +318,74 @@ void FAutoReimportManager::Destroy()
 	DirectoryMonitors.Empty();
 }
 
+void FAutoReimportManager::OnAssetRenamed(const FAssetData& AssetData, const FString& OldPath)
+{
+	if (bGuardAssetChanges)
+	{
+		return;
+	}
+
+	// This code moves a source content file that reside alongside assets when the assets are renamed. We do this under the following conditions:
+	// 	1. The sourcefile is solely referenced from the the asset that has been moved
+	//	2. Said asset only references a single file
+	//	3. The source file resides in the same folder as the asset
+	//
+	// Additionally, we rename the source file if it matched the name of the asset before the rename/move.
+	//	- If we rename the source file, then we also update the reimport paths for the asset
+
+	TArray<FString> SourceFilesRelativeToOldPath;
+	for (const auto& Pair : AssetData.TagsAndValues)
+	{
+		if (Pair.Key.IsEqual(UObject::SourceFileTagName(), ENameCase::IgnoreCase, false))
+		{
+			SourceFilesRelativeToOldPath.Add(Pair.Value);
+		}
+	}
+
+	// We move the file with the asset provided it is the only file referenced, and sits right beside the uasset file
+	if (SourceFilesRelativeToOldPath.Num() == 1 && !SourceFilesRelativeToOldPath[0].GetCharArray().ContainsByPredicate([](const TCHAR Char) { return Char == '/' || Char == '\\'; }))
+	{
+		const FString AbsoluteSrcPath = FPaths::ConvertRelativePathToFull(FPackageName::LongPackageNameToFilename(FPackageName::GetLongPackagePath(OldPath)));
+		const FString AbsoluteDstPath = FPaths::ConvertRelativePathToFull(FPackageName::LongPackageNameToFilename(AssetData.PackagePath.ToString()));
+
+		const FString OldAssetName = FPackageName::GetLongPackageAssetName(FPackageName::ObjectPathToPackageName(OldPath));
+		FString NewFileName = FPaths::GetBaseFilename(SourceFilesRelativeToOldPath[0]);
+
+		bool bRequireReimportPathUpdate = false;
+		if (PackageTools::SanitizePackageName(NewFileName) == OldAssetName)
+		{
+			NewFileName = AssetData.AssetName.ToString();
+			bRequireReimportPathUpdate = true;
+		}
+
+		const FString SrcFile = AbsoluteSrcPath / SourceFilesRelativeToOldPath[0];
+		const FString DstFile = AbsoluteDstPath / NewFileName + TEXT(".") + FPaths::GetExtension(SourceFilesRelativeToOldPath[0]);
+
+		// We can't do this if multiple assets reference the same file. We should be checking for > 1 referencing asset, but the asset registry
+		// filter lookup won't return the recently renamed package because it will be Empty by now, so we check for *anything* referencing the asset (assuming that we'll never find *this* asset).
+		const IAssetRegistry& Registry = FModuleManager::LoadModuleChecked<FAssetRegistryModule>("AssetRegistry").Get();
+		if (Utils::FindAssetsPertainingToFile(Registry, SrcFile).Num() > 0)
+		{
+			return;
+		}
+
+		if (!FPlatformFileManager::Get().GetPlatformFile().FileExists(*DstFile) &&
+			IFileManager::Get().Move(*DstFile, *SrcFile, false /*bReplace */, false, true /* attributes */, true /* don't retry */))
+		{
+			IgnoreMovedFile(SrcFile, DstFile);
+
+			if (bRequireReimportPathUpdate)
+			{
+				TArray<FString> Paths;
+				Paths.Add(DstFile);
+
+				// Update the reimport file names
+				FReimportManager::Instance()->UpdateReimportPaths(AssetData.GetAsset(), Paths);
+			}
+		}
+	}
+}
+
 int32 FAutoReimportManager::GetNumUnprocessedChanges() const
 {
 	return Utils::Reduce(DirectoryMonitors, [](const FContentDirectoryMonitor& Monitor, int32 Total){
@@ -297,6 +419,7 @@ TOptional<ECurrentState> FAutoReimportManager::ProcessAdditions(const FTimeLimit
 {
 	// Override the global feedback context while we do this to avoid popping up dialogs
 	TGuardValue<FFeedbackContext*> ScopedContextOverride(GWarn, FeedbackContextOverride.Get());
+	TGuardValue<bool> ScopedAssetChangesGuard(bGuardAssetChanges, true);
 
 	FeedbackContextOverride->GetContent()->SetMainText(GetProgressText());
 
@@ -348,6 +471,8 @@ TOptional<ECurrentState> FAutoReimportManager::SavePackages()
 	// We don't override the context specifically when saving packages so the user gets proper feedback
 	//TGuardValue<FFeedbackContext*> ScopedContextOverride(GWarn, FeedbackContextOverride.Get());
 
+	TGuardValue<bool> ScopedAssetChangesGuard(bGuardAssetChanges, true);
+
 	FeedbackContextOverride->GetContent()->SetMainText(GetProgressText());
 
 	if (PackagesToSave.Num() > 0)
@@ -367,6 +492,7 @@ TOptional<ECurrentState> FAutoReimportManager::ProcessModifications(const FTimeL
 {
 	// Override the global feedback context while we do this to avoid popping up dialogs
 	TGuardValue<FFeedbackContext*> ScopedContextOverride(GWarn, FeedbackContextOverride.Get());
+	TGuardValue<bool> ScopedAssetChangesGuard(bGuardAssetChanges, true);
 
 	FeedbackContextOverride->GetContent()->SetMainText(GetProgressText());
 
@@ -383,6 +509,8 @@ TOptional<ECurrentState> FAutoReimportManager::ProcessModifications(const FTimeL
 
 TOptional<ECurrentState> FAutoReimportManager::ProcessDeletions()
 {
+	TGuardValue<bool> ScopedAssetChangesGuard(bGuardAssetChanges, true);
+
 	const IAssetRegistry& Registry = FModuleManager::LoadModuleChecked<FAssetRegistryModule>("AssetRegistry").Get();
 
 	FeedbackContextOverride->GetContent()->SetMainText(GetProgressText());
@@ -404,6 +532,8 @@ TOptional<ECurrentState> FAutoReimportManager::ProcessDeletions()
 		ObjectTools::DeleteAssets(AssetsToDelete);
 	}
 
+	FeedbackContextOverride->GetContent()->SetMainText(GetProgressText());
+
 	Cleanup();
 	return ECurrentState::Idle;
 }
@@ -414,12 +544,11 @@ TOptional<ECurrentState> FAutoReimportManager::PendingResponse()
 	return TOptional<ECurrentState>();
 }
 
-TOptional<ECurrentState> FAutoReimportManager::Idle(const FTimeLimit& TimeLimit)
+TOptional<ECurrentState> FAutoReimportManager::Idle()
 {
 	for (auto& Monitor : DirectoryMonitors)
 	{
-		Monitor.Tick(TimeLimit);
-		yield TOptional<ECurrentState>();
+		Monitor.Tick();
 	}
 
 	const int32 NumUnprocessedChanges = GetNumUnprocessedChanges();
@@ -493,7 +622,7 @@ void FAutoReimportManager::AddReferencedObjects(FReferenceCollector& Collector)
 void FAutoReimportManager::HandleLoadingSavingSettingChanged(FName PropertyName)
 {
 	if (PropertyName == GET_MEMBER_NAME_CHECKED(UEditorLoadingSavingSettings, bMonitorContentDirectories) ||
-		PropertyName == GET_MEMBER_NAME_CHECKED(UEditorLoadingSavingSettings, AutoReimportDirectories))
+		PropertyName == GET_MEMBER_NAME_CHECKED(UEditorLoadingSavingSettings, AutoReimportDirectorySettings))
 	{
 		const auto* Settings = GetDefault<UEditorLoadingSavingSettings>();
 		if (Settings->bMonitorContentDirectories)
@@ -525,60 +654,54 @@ void FAutoReimportManager::OnContentPathChanged(const FString& InAssetPath, cons
 
 void FAutoReimportManager::SetUpDirectoryMonitors()
 {
-	TArray<FPathAndMountPoint> CollapsedDirectories;
-
-	// Build an array of directory / mounted path so we can match up absolute paths that the user has typed in
-	TArray<TPair<FString, FString>> MountedPaths;
+	struct FParsedSettings
 	{
-		TArray<FString> RootContentPaths;
-		FPackageName::QueryRootContentPaths( RootContentPaths );
-		for (FString& RootPath : RootContentPaths)
-		{
-			FString ContentFolder = FPaths::ConvertRelativePathToFull(FPackageName::LongPackageNameToFilename(RootPath));
-			MountedPaths.Add( TPairInitializer<FString, FString>(MoveTemp(ContentFolder), MoveTemp(RootPath)) );
-		}
-	}
-
+		FString SourceDirectory;
+		FString MountPoint;
+		FMatchRules Rules;
+	};
 
-	auto& FileManager = IFileManager::Get();
-	for (const FString& Path : GetDefault<UEditorLoadingSavingSettings>()->AutoReimportDirectories)
+	TArray<FParsedSettings> FinalArray;
+	auto SupportedExtensions = GetAllFactoryExtensions();
+	for (const auto& Setting : GetDefault<UEditorLoadingSavingSettings>()->AutoReimportDirectorySettings)
 	{
-		FPathAndMountPoint ThisPath;
+		FParsedSettings NewMapping;
+		NewMapping.SourceDirectory = Setting.SourceDirectory;
+		NewMapping.MountPoint = Setting.MountPoint;
 
-		const FName MountPoint = FPackageName::GetPackageMountPoint(Path);
-		if (!MountPoint.IsNone())
+		if (!FAutoReimportDirectoryConfig::ParseSourceDirectoryAndMountPoint(NewMapping.SourceDirectory, NewMapping.MountPoint))
 		{
-			ThisPath.Path = FPaths::ConvertRelativePathToFull(FPackageName::LongPackageNameToFilename(Path / TEXT(""))) / "";
+			continue;
 		}
-		else
+
+		// Only include extensions that match a factory
+		NewMapping.Rules.SetApplicableExtensions(SupportedExtensions);
+		for (const auto& WildcardConfig : Setting.Wildcards)
 		{
-			ThisPath.Path = FPaths::ConvertRelativePathToFull(Path) / "";
+			NewMapping.Rules.AddWildcardRule(WildcardConfig.Wildcard, WildcardConfig.bInclude);
 		}
+		
+		FinalArray.Add(NewMapping);
+	}
 
-		if (FileManager.DirectoryExists(*ThisPath.Path))
-		{
-			// Set the mounted path if necessary
-			auto* Pair = MountedPaths.FindByPredicate([&](const TPair<FString, FString>& InPair){
-				return ThisPath.Path.StartsWith(InPair.Key);
-			});
+	for (int32 Index = 0; Index < FinalArray.Num(); ++Index)
+	{
+		const auto& Mapping = FinalArray[Index];
 
-			if (Pair)
+		// We only create a directory monitor if there are no other's watching parent directories of this one
+		for (int32 OtherIndex = Index + 1; OtherIndex < FinalArray.Num(); ++OtherIndex)
+		{
+			if (FinalArray[Index].SourceDirectory.StartsWith(FinalArray[OtherIndex].SourceDirectory))
 			{
-				ThisPath.MountPoint = Pair->Value / ThisPath.Path.RightChop(Pair->Key.Len());
+				UE_LOG(LogAutoReimportManager, Warning, TEXT("Unable to watch directory %s as it will conflict with another watching %s."), *FinalArray[Index].SourceDirectory, *FinalArray[OtherIndex].SourceDirectory);
+				goto next;
 			}
-
-			CollapsedDirectories.Add(ThisPath);
 		}
-	}
 
-	Utils::RemoveDuplicates(CollapsedDirectories, [](const FPathAndMountPoint& A, const FPathAndMountPoint& B){
-		return A.Path.StartsWith(B.Path, ESearchCase::IgnoreCase) ||
-			   B.Path.StartsWith(A.Path, ESearchCase::IgnoreCase);
-	});
+		DirectoryMonitors.Emplace(Mapping.SourceDirectory, Mapping.Rules, Mapping.MountPoint);
 
-	for (const auto& Dir : CollapsedDirectories)
-	{
-		DirectoryMonitors.Emplace(Dir.Path, SupportedExtensions, Dir.MountPoint);
+	next:
+		continue;
 	}
 }
 
@@ -634,9 +757,24 @@ void UAutoReimportManager::Initialize()
 	Implementation = MakeShareable(new FAutoReimportManager);
 }
 
-void UAutoReimportManager::ReportExternalChange(const FString& Filename, FFileChangeData::EFileChangeAction Action)
+void UAutoReimportManager::IgnoreNewFile(const FString& Filename)
+{
+	Implementation->IgnoreNewFile(Filename);
+}
+
+void UAutoReimportManager::IgnoreFileModification(const FString& Filename)
+{
+	Implementation->IgnoreFileModification(Filename);
+}
+
+void UAutoReimportManager::IgnoreMovedFile(const FString& SrcFilename, const FString& DstFilename)
+{
+	Implementation->IgnoreMovedFile(SrcFilename, DstFilename);
+}
+
+void UAutoReimportManager::IgnoreDeletedFile(const FString& Filename)
 {
-	Implementation->ReportExternalChange(Filename, Action);
+	Implementation->IgnoreDeletedFile(Filename);
 }
 
 TArray<FPathAndMountPoint> UAutoReimportManager::GetMonitoredDirectories() const
diff --git a/Engine/Source/Editor/UnrealEd/Private/AutoReimport/AutoReimportTests.cpp b/Engine/Source/Editor/UnrealEd/Private/AutoReimport/AutoReimportTests.cpp
new file mode 100644
--- /dev/null
+++ b/Engine/Source/Editor/UnrealEd/Private/AutoReimport/AutoReimportTests.cpp
@@ -0,0 +1,770 @@
+// Copyright 1998-2015 Epic Games, Inc. All Rights Reserved.
+#include "UnrealEd.h"
+
+#include "FileCache.h"
+
+DEFINE_LOG_CATEGORY_STATIC(LogAutoReimportTests, Log, All);
+
+/** Simple latent copmmand that executes a function after a given delay */
+class FDelayedCallbackLatentCommand : public IAutomationLatentCommand
+{
+public:
+	FDelayedCallbackLatentCommand(TFunction<void()> InCallback, float InDelay = 0.1f)
+		: Callback(InCallback), Delay(InDelay)
+	{}
+
+	virtual bool Update() override
+	{
+		float NewTime = FPlatformTime::Seconds();
+		if (NewTime - StartTime >= Delay)
+		{
+			Callback();
+			return true;
+		}
+		return false;
+	}
+
+private:
+	TFunction<void()> Callback;
+	float Delay;
+};
+
+/** Generate a config from the specified options, to pass to FFileCache on construction */
+FFileCacheConfig GenerateFileCacheConfig(const FString& InWorkingDir)
+{
+	FString Directory = FPaths::ConvertRelativePathToFull(InWorkingDir);
+	FString CacheFilename = Directory / TEXT("Cache.bin");
+
+	FFileCacheConfig Config(Directory / TEXT("Content") / TEXT(""), MoveTemp(CacheFilename));
+	// We always store paths inside content folders relative to the folder
+	Config.PathType = EPathType::Relative;
+
+	Config.bDetectChangesSinceLastRun = false;
+	return Config;
+}
+
+/** Persistent test payload used for async testing, generally captured by lambdas */
+struct FAutoReimportTestPayload
+{
+	FFileCacheConfig Config;
+	TSharedPtr<FFileCache> FileCache;
+	FString WorkingDir;
+
+	FAutoReimportTestPayload(const FString& InWorkingDir)
+		: Config(GenerateFileCacheConfig(InWorkingDir))
+		, WorkingDir(InWorkingDir)
+	{	
+		IFileManager::Get().MakeDirectory(*Config.Directory, true);
+		// Just to be sure...
+		IFileManager::Get().MakeDirectory(*InWorkingDir, true);
+	}
+	~FAutoReimportTestPayload()
+	{
+		// Avoid writing out the file after we've nuiked the directory
+		FileCache = nullptr;
+		IFileManager::Get().DeleteDirectory(*WorkingDir, false, true);
+	}
+
+	void StartWatching()
+	{
+		if (!FileCache.IsValid())
+		{
+			FileCache = MakeShareable(new FFileCache(Config));
+		}
+	}
+
+	void WaitForStartup(TFunction<void()> Finished)
+	{
+		ADD_LATENT_AUTOMATION_COMMAND(FDelayedCallbackLatentCommand([=]{
+			FileCache->Tick();
+			if (FileCache->HasStartedUp())
+			{
+				Finished();
+			}
+			else
+			{
+				WaitForStartup(Finished);
+			}
+		}, 0.1));
+	}
+
+	void StopWatching()
+	{
+		FileCache = nullptr;
+	}
+};
+
+FString GetWorkingDir()
+{
+	return FPaths::ConvertRelativePathToFull(FPaths::AutomationTransientDir()) / TEXT("AutoReimport") / TEXT("");
+}
+
+FString GetTestSourceDir()
+{
+	return FPaths::ConvertRelativePathToFull(FPaths::EngineContentDir()) / TEXT("NotForLicensees") / TEXT("Automation") / TEXT("EditorAutomation") / TEXT("AutoReimportTestSource") / TEXT("");
+}
+
+struct FSrcDstFilenames
+{
+	FSrcDstFilenames(const TCHAR* InSrc, const TCHAR* InDst) : Src(InSrc), Dst(InDst) {}
+	const TCHAR* Src, *Dst;
+};
+
+/** Copy an array of test files from the source folder to the transient test folder */
+bool CopyTestFiles(const FAutoReimportTestPayload& Test, const TArray<FSrcDstFilenames>& Files)
+{
+	auto& PlatformFile = FPlatformFileManager::Get().GetPlatformFile();
+	const auto SourceDir = GetTestSourceDir();
+
+	for (const auto& File : Files)
+	{
+		const bool bCopied = PlatformFile.CopyFile(*(Test.Config.Directory / File.Dst), *(SourceDir / File.Src));
+		if (!bCopied)
+		{
+			UE_LOG(LogAutoReimportTests, Error, TEXT("Failed to copy source file %s to the test directory as %s."), File.Src, File.Dst);
+			return false;
+		}
+	}
+
+	return true;
+}
+
+/** Test that creating a new file gets reported correctly */
+IMPLEMENT_SIMPLE_AUTOMATION_TEST(FAutoReimportSimpleCreateTest, "Editor.Auto Reimport.Simple Create", EAutomationTestFlags::ATF_Editor)
+bool FAutoReimportSimpleCreateTest::RunTest(const FString& Parameters)
+{
+	const FString WorkingDir = GetWorkingDir();
+
+	static const TCHAR* Filename = TEXT("square.png");
+
+	// Start watching the directory
+	TSharedPtr<FAutoReimportTestPayload> Test = MakeShareable(new FAutoReimportTestPayload(WorkingDir));
+	Test->StartWatching();
+	Test->WaitForStartup([=]{
+		
+		TArray<FSrcDstFilenames> Files;
+		Files.Emplace(Filename, Filename);
+
+		if (!CopyTestFiles(*Test, Files))
+		{
+			return;
+		}
+
+		ADD_LATENT_AUTOMATION_COMMAND(FDelayedCallbackLatentCommand([=]{
+
+			auto Changes = Test->FileCache->GetOutstandingChanges();
+
+			if (Changes.Num() != 1)
+			{
+				UE_LOG(LogAutoReimportTests, Error, TEXT("Incorrect number of changes reported (%d != 1)."), Changes.Num());
+			}
+			// Check that we got a relative path
+			else if (Changes[0].Filename.Get() != Filename)
+			{
+				UE_LOG(LogAutoReimportTests, Error, TEXT("Path reported incorrectly (%s != %s)."), *Changes[0].Filename.Get(), Filename);
+			}
+			// Check it was the correct type
+			else if (Changes[0].Action != EFileAction::Added)
+			{
+				UE_LOG(LogAutoReimportTests, Error, TEXT("Incorrect action for added file %s."), Filename);
+			}
+			else
+			{
+				Test->FileCache->CompleteTransaction(MoveTemp(Changes[0]));
+
+				if (!Test->FileCache->FindFileData(Filename))
+				{
+					UE_LOG(LogAutoReimportTests, Error, TEXT("Add transaction was not applied correctly."));
+				}
+			}
+
+		}, 1));
+
+	});
+
+	return true;
+}
+
+/** Test that modifying an existing file gets reported correctly */
+IMPLEMENT_SIMPLE_AUTOMATION_TEST(FAutoReimportSimpleModifyTest, "Editor.Auto Reimport.Simple Modify", EAutomationTestFlags::ATF_Editor)
+bool FAutoReimportSimpleModifyTest::RunTest(const FString& Parameters)
+{
+	const FString WorkingDir = GetWorkingDir();
+
+	static const TCHAR* SrcFilename1 = TEXT("square.png");
+	static const TCHAR* SrcFilename2 = TEXT("red-square.png");
+	static const TCHAR* DstFilename = TEXT("square.png");
+
+	// Start watching the directory
+	TSharedPtr<FAutoReimportTestPayload> Test = MakeShareable(new FAutoReimportTestPayload(WorkingDir));
+
+	TArray<FSrcDstFilenames> Files;
+	Files.Emplace(SrcFilename1, DstFilename);
+
+	if (!CopyTestFiles(*Test, Files))
+	{
+		return false;
+	}
+
+	Test->StartWatching();
+	Test->WaitForStartup([=]{
+
+		TArray<FSrcDstFilenames> Files;
+		Files.Emplace(SrcFilename2, DstFilename);
+
+		if (!CopyTestFiles(*Test, Files))
+		{
+			return;
+		}
+
+		ADD_LATENT_AUTOMATION_COMMAND(FDelayedCallbackLatentCommand([=]{
+
+			auto Changes = Test->FileCache->GetOutstandingChanges();
+
+			if (Changes.Num() != 1)
+			{
+				UE_LOG(LogAutoReimportTests, Error, TEXT("Incorrect number of changes reported (%d != 1)."), Changes.Num());
+			}
+			// Check that we got a relative path
+			else if (Changes[0].Filename.Get() != DstFilename)
+			{
+				UE_LOG(LogAutoReimportTests, Error, TEXT("Path reported incorrectly (%s != %s)."), *Changes[0].Filename.Get(), DstFilename);
+			}
+			// Check it was the correct type
+			else if (Changes[0].Action != EFileAction::Modified)
+			{
+				UE_LOG(LogAutoReimportTests, Error, TEXT("Incorrect action for modified file %s."), DstFilename);
+			}
+			else
+			{
+				// Copy the file hash as it will be invalidated by the transaction being completed
+				auto FileHash = Changes[0].FileData.FileHash;
+
+				Test->FileCache->CompleteTransaction(MoveTemp(Changes[0]));
+				const auto* Data = Test->FileCache->FindFileData(DstFilename);
+				if (!Data || Data->FileHash != FileHash)
+				{
+					UE_LOG(LogAutoReimportTests, Error, TEXT("Modify transaction was not applied correctly."));
+				}
+			}
+
+		}, 1));
+
+	});
+
+	return true;
+}
+
+/** Test that deleting an existing file gets reported correctly */
+IMPLEMENT_SIMPLE_AUTOMATION_TEST(FAutoReimportSimpleDeleteTest, "Editor.Auto Reimport.Simple Delete", EAutomationTestFlags::ATF_Editor)
+bool FAutoReimportSimpleDeleteTest::RunTest(const FString& Parameters)
+{
+	const FString WorkingDir = GetWorkingDir();
+
+	static const TCHAR* Filename = TEXT("square.png");
+
+	TSharedPtr<FAutoReimportTestPayload> Test = MakeShareable(new FAutoReimportTestPayload(WorkingDir));
+
+	TArray<FSrcDstFilenames> Files;
+	Files.Emplace(Filename, Filename);
+
+	if (!CopyTestFiles(*Test, Files))
+	{
+		return false;
+	}
+
+	// Start watching the directory
+	Test->StartWatching();
+	Test->WaitForStartup([=]{
+
+		// Delete the file and check that it got reported as deleted
+		const bool bDeleted = IFileManager::Get().Delete(*(Test->Config.Directory / Filename));
+		if (!bDeleted)
+		{
+			UE_LOG(LogAutoReimportTests, Error, TEXT("Failed to delete source file from the test directory."));
+			return;
+		}
+
+		ADD_LATENT_AUTOMATION_COMMAND(FDelayedCallbackLatentCommand([=]{
+
+			auto Changes = Test->FileCache->GetOutstandingChanges();
+
+			if (Changes.Num() != 1)
+			{
+				UE_LOG(LogAutoReimportTests, Error, TEXT("Incorrect number of changes reported (%d != 1)."), Changes.Num());
+			}
+			// Check that we got a relative path
+			else if (Changes[0].Filename.Get() != Filename)
+			{
+				UE_LOG(LogAutoReimportTests, Error, TEXT("Path reported incorrectly (%s != %s)."), *Changes[0].Filename.Get(), Filename);
+			}
+			// Check it was the correct type
+			else if (Changes[0].Action != EFileAction::Removed)
+			{
+				UE_LOG(LogAutoReimportTests, Error, TEXT("Incorrect action for deleted file %s."), Filename);
+			}
+			else
+			{
+				Test->FileCache->CompleteTransaction(MoveTemp(Changes[0]));
+
+				if (Test->FileCache->FindFileData(Filename))
+				{
+					UE_LOG(LogAutoReimportTests, Error, TEXT("Remove transaction was not applied correctly."));
+				}
+			}
+
+		}, 1));
+
+	});
+
+	return true;
+}
+
+/** Test that renaming an existing file gets reported correctly */
+IMPLEMENT_SIMPLE_AUTOMATION_TEST(FAutoReimportSimpleRenameTest, "Editor.Auto Reimport.Simple Rename", EAutomationTestFlags::ATF_Editor)
+bool FAutoReimportSimpleRenameTest::RunTest(const FString& Parameters)
+{
+	const FString WorkingDir = GetWorkingDir();
+
+	static const TCHAR* SrcFilename = TEXT("square.png");
+	static const TCHAR* DstFilename = TEXT("square2.png");
+
+
+	TSharedPtr<FAutoReimportTestPayload> Test = MakeShareable(new FAutoReimportTestPayload(WorkingDir));
+
+	TArray<FSrcDstFilenames> Files;
+	Files.Emplace(SrcFilename, SrcFilename);
+
+	if (!CopyTestFiles(*Test, Files))
+	{
+		return false;
+	}
+
+	// Start watching the directory
+	Test->StartWatching();
+	Test->WaitForStartup([=]{
+
+		// Rename the file
+		IFileManager::Get().Move(*(Test->Config.Directory / DstFilename), *(Test->Config.Directory / SrcFilename));
+
+		ADD_LATENT_AUTOMATION_COMMAND(FDelayedCallbackLatentCommand([=]{
+
+			auto Changes = Test->FileCache->GetOutstandingChanges();
+
+			if (Changes.Num() != 1)
+			{
+				UE_LOG(LogAutoReimportTests, Error, TEXT("Incorrect number of changes reported (%d != 1)."), Changes.Num());
+			}
+			// Check it was the correct type
+			else if (Changes[0].Action != EFileAction::Moved)
+			{
+				UE_LOG(LogAutoReimportTests, Error, TEXT("Incorrect action for renamed file %s -> %s."), SrcFilename, DstFilename);
+			}
+			else if (Changes[0].Filename.Get() != DstFilename)
+			{
+				UE_LOG(LogAutoReimportTests, Error, TEXT("Source path reported incorrectly (%s != %s)."), *Changes[0].Filename.Get(), DstFilename);
+			}
+			else if (Changes[0].MovedFromFilename.Get() != SrcFilename)
+			{
+				UE_LOG(LogAutoReimportTests, Error, TEXT("Destination path reported incorrectly (%s != %s)."), *Changes[0].MovedFromFilename.Get(), SrcFilename);
+			}
+			else
+			{
+				Test->FileCache->CompleteTransaction(MoveTemp(Changes[0]));
+
+				if (Test->FileCache->FindFileData(SrcFilename) || !Test->FileCache->FindFileData(DstFilename))
+				{
+					UE_LOG(LogAutoReimportTests, Error, TEXT("Rename transaction was not applied correctly."));
+				}
+			}
+
+		}, 1));
+
+	});
+
+	return true;
+}
+
+/** Test that moving a file outside of the monitored directory gets reported correctly (should be reported as a delete) */
+IMPLEMENT_SIMPLE_AUTOMATION_TEST(FAutoReimportSimpleMoveExternallyTest, "Editor.Auto Reimport.Move Externally", EAutomationTestFlags::ATF_Editor)
+bool FAutoReimportSimpleMoveExternallyTest::RunTest(const FString& Parameters)
+{
+	const FString WorkingDir = GetWorkingDir();
+
+	static const TCHAR* SrcFilename = TEXT("square.png");
+	static const TCHAR* DstFilename = TEXT("../square.png");
+	
+	TSharedPtr<FAutoReimportTestPayload> Test = MakeShareable(new FAutoReimportTestPayload(WorkingDir));
+
+	TArray<FSrcDstFilenames> Files;
+	Files.Emplace(SrcFilename, SrcFilename);
+
+	if (!CopyTestFiles(*Test, Files))
+	{
+		return false;
+	}
+
+	// Start watching the directory
+	Test->StartWatching();
+	Test->WaitForStartup([=]{
+
+		if (!Test->FileCache->FindFileData(SrcFilename))
+		{
+			UE_LOG(LogAutoReimportTests, Error, TEXT("Could not find file data for initial file %s."), SrcFilename);
+			return;
+		}
+
+		// Rename the file
+		IFileManager::Get().Move(*(Test->Config.Directory / DstFilename), *(Test->Config.Directory / SrcFilename));
+
+		ADD_LATENT_AUTOMATION_COMMAND(FDelayedCallbackLatentCommand([=]{
+
+			auto Changes = Test->FileCache->GetOutstandingChanges();
+
+			if (Changes.Num() != 1)
+			{
+				UE_LOG(LogAutoReimportTests, Error, TEXT("Incorrect number of changes reported (%d != 1)."), Changes.Num());
+			}
+			// Check it was the correct type
+			else if (Changes[0].Action != EFileAction::Removed)
+			{
+				UE_LOG(LogAutoReimportTests, Error, TEXT("Incorrect action for removed file %s -> %s."), SrcFilename, DstFilename);
+			}
+			else if (Changes[0].Filename.Get() != SrcFilename)
+			{
+				UE_LOG(LogAutoReimportTests, Error, TEXT("Source path reported incorrectly (%s != %s)."), *Changes[0].Filename.Get(), SrcFilename);
+			}
+			else
+			{
+				Test->FileCache->CompleteTransaction(MoveTemp(Changes[0]));
+
+				if (Test->FileCache->FindFileData(SrcFilename))
+				{
+					UE_LOG(LogAutoReimportTests, Error, TEXT("Found data for file that should have been removed (%s)."), SrcFilename);
+				}
+			}
+
+		}, 1));
+
+	});
+
+	return true;
+}
+
+/** Test that bDetectChangesSinceLastRun works correctly when true and when false */
+IMPLEMENT_SIMPLE_AUTOMATION_TEST(FAutoReimportRestartDetectionTest, "Editor.Auto Reimport.Restart Detection", EAutomationTestFlags::ATF_Editor)
+bool FAutoReimportRestartDetectionTest::RunTest(const FString& Parameters)
+{
+	const FString WorkingDir = GetWorkingDir();
+	static const TCHAR* SrcFilename = TEXT("square.png");
+	static const TCHAR* SrcFilename2 = TEXT("red-square.png");
+	
+	TSharedPtr<FAutoReimportTestPayload> Test = MakeShareable(new FAutoReimportTestPayload(WorkingDir));
+
+	TArray<FSrcDstFilenames> Files;
+	Files.Emplace(SrcFilename, SrcFilename);
+	if (!CopyTestFiles(*Test, Files))
+	{
+		return false;
+	}
+
+	// Start watching the directory
+	Test->StartWatching();
+	Test->WaitForStartup([=]{
+
+		if (!Test->FileCache->FindFileData(SrcFilename))
+		{
+			UE_LOG(LogAutoReimportTests, Error, TEXT("Could not find file data for initial file %s."), SrcFilename);
+			return;
+		}
+
+		// Stop watching and write out the cache file
+		Test->StopWatching();
+		Test->Config.bDetectChangesSinceLastRun = true;
+
+		// Modify the file while the watcher isn't running
+		TArray<FSrcDstFilenames> Files;
+		Files.Emplace(SrcFilename2, SrcFilename);
+		if (!CopyTestFiles(*Test, Files))
+		{
+			return;
+		}
+
+		Test->StartWatching();
+		Test->WaitForStartup([=]{
+
+			auto Changes = Test->FileCache->GetOutstandingChanges();
+
+			for (auto& Change : Changes)
+			{
+				UE_LOG(LogAutoReimportTests, Log, TEXT("Change %d for file %s."), (uint8)Change.Action, *Change.Filename.Get());
+			}
+
+			if (Changes.Num() != 1)
+			{
+				UE_LOG(LogAutoReimportTests, Error, TEXT("Incorrect number of changes reported (%d != 1)."), Changes.Num());
+				return;
+			}
+			// Check it was the correct type
+			else if (Changes[0].Action != EFileAction::Modified)
+			{
+				UE_LOG(LogAutoReimportTests, Error, TEXT("Incorrect action for modified file %s."), SrcFilename);
+				return;
+			}
+			else
+			{
+				Test->FileCache->CompleteTransaction(MoveTemp(Changes[0]));
+			}
+
+			// Stop watching and write out the cache file
+			Test->StopWatching();
+			Test->Config.bDetectChangesSinceLastRun = false;
+
+			// Modify the file while the watcher isn't running
+			TArray<FSrcDstFilenames> Files;
+			Files.Emplace(SrcFilename, SrcFilename);
+			if (!CopyTestFiles(*Test, Files))
+			{
+				return;
+			}
+
+			Test->StartWatching();
+			Test->WaitForStartup([=]{
+
+				auto Changes = Test->FileCache->GetOutstandingChanges();
+				if (Changes.Num() != 0)
+				{
+					UE_LOG(LogAutoReimportTests, Error, TEXT("Shouldn't have reported changes when bDetectChangesSinceLastRun is false (%d change(s) received)."), Changes.Num());
+					return;
+				}
+
+			});
+		});
+
+	});
+
+	return true;
+}
+
+/** Test that making multiple changes to the same file gets picked up correctly */
+IMPLEMENT_SIMPLE_AUTOMATION_TEST(FAutoReimportMultipleChangesTest, "Editor.Auto Reimport.Multiple Changes", EAutomationTestFlags::ATF_Editor)
+bool FAutoReimportMultipleChangesTest::RunTest(const FString& Parameters)
+{
+	const FString WorkingDir = GetWorkingDir();
+	static const TCHAR* Filename1 = TEXT("square.png");
+	static const TCHAR* Filename2 = TEXT("red-square.png");
+	
+	TSharedPtr<FAutoReimportTestPayload> Test = MakeShareable(new FAutoReimportTestPayload(WorkingDir));
+
+	// Start watching the directory
+	Test->StartWatching();
+	Test->WaitForStartup([=]{
+
+		// Add Filename1 and Filename2
+		TArray<FSrcDstFilenames> Files;
+		Files.Emplace(Filename1, Filename1);
+		Files.Emplace(Filename2, Filename2);
+		if (!CopyTestFiles(*Test, Files))
+		{
+			return;
+		}
+
+		// Delete Filename1
+		IFileManager::Get().Delete(*(Test->Config.Directory / Filename1));
+
+		// Modify Filename2
+		Files.Reset();
+		Files.Emplace(Filename1, Filename2);
+		if (!CopyTestFiles(*Test, Files))
+		{
+			return;
+		}
+
+		ADD_LATENT_AUTOMATION_COMMAND(FDelayedCallbackLatentCommand([=]{
+
+			// The net change should just be a single added file
+			auto Changes = Test->FileCache->GetOutstandingChanges();
+
+			if (Changes.Num() != 1)
+			{
+				UE_LOG(LogAutoReimportTests, Error, TEXT("Incorrect number of changes reported (%d != 1)."), Changes.Num());
+			}
+			// Check it was the correct type
+			else if (Changes[0].Action != EFileAction::Added)
+			{
+				UE_LOG(LogAutoReimportTests, Error, TEXT("Incorrect action for added file %s."), Filename2);
+			}
+
+		}, 1));
+
+	});
+
+	return true;
+}
+
+/** Test that starting up a cache file with a different set of applicable extensions correctly ignores/updates the extensions, whilst reporting changes only for applicable extensions */
+IMPLEMENT_SIMPLE_AUTOMATION_TEST(FAutoReimportChangeExtensionsTest, "Editor.Auto Reimport.Change Extensions", EAutomationTestFlags::ATF_Editor)
+bool FAutoReimportChangeExtensionsTest::RunTest(const FString& Parameters)
+{
+	const FString WorkingDir = GetWorkingDir();
+	static const TCHAR* SrcFilename1 = TEXT("square.png");
+	static const TCHAR* SrcFilename2 = TEXT("red-square.png");
+	static const TCHAR* SrcFilename3 = TEXT("empty.txt");
+	
+	TSharedPtr<FAutoReimportTestPayload> Test = MakeShareable(new FAutoReimportTestPayload(WorkingDir));
+	Test->Config.Rules.SetApplicableExtensions(TEXT("txt;"));
+	Test->Config.bDetectChangesSinceLastRun = true;
+
+	// Start watching the directory
+	Test->StartWatching();
+	Test->WaitForStartup([=]{
+
+		// Add a txt file and a png - we should only be told about the txt
+		TArray<FSrcDstFilenames> Files;
+		Files.Emplace(SrcFilename1, SrcFilename1);
+		Files.Emplace(SrcFilename3, SrcFilename3);
+		if (!CopyTestFiles(*Test, Files))
+		{
+			return;
+		}
+
+		ADD_LATENT_AUTOMATION_COMMAND(FDelayedCallbackLatentCommand([=]{
+
+			auto Changes = Test->FileCache->GetOutstandingChanges();
+
+			if (Changes.Num() != 1)
+			{
+				UE_LOG(LogAutoReimportTests, Error, TEXT("Incorrect number of changes reported (%d != 1)."), Changes.Num());
+				return;
+			}
+			else if (Changes[0].Action != EFileAction::Added)
+			{
+				UE_LOG(LogAutoReimportTests, Error, TEXT("Incorrect action for added file %s."), SrcFilename3);
+				return;
+			}
+			else if (Changes[0].Filename.Get() != SrcFilename3)
+			{
+				UE_LOG(LogAutoReimportTests, Error, TEXT("Added file path reported incorrectly (%s != %s)."), *Changes[0].Filename.Get(), SrcFilename3);
+				return;
+			}
+			Test->FileCache->CompleteTransaction(MoveTemp(Changes[0]));
+
+			Test->StopWatching();
+			// Change the extensions we're watching, and add another png
+			Test->Config.Rules = FMatchRules();
+			Test->Config.Rules.SetApplicableExtensions(TEXT("png;"));
+
+			Test->StartWatching();
+			Test->WaitForStartup([=]{
+
+				auto Changes = Test->FileCache->GetOutstandingChanges();
+				if (Changes.Num() != 0)
+				{
+					UE_LOG(LogAutoReimportTests, Error, TEXT("Incorrect number of changes reported (%d != 0)."), Changes.Num());
+					return;
+				}
+
+				// Add another png - we should only be notified about this one
+				TArray<FSrcDstFilenames> Files;
+				Files.Emplace(SrcFilename2, SrcFilename2);
+				if (!CopyTestFiles(*Test, Files))
+				{
+					return;
+				}
+
+				ADD_LATENT_AUTOMATION_COMMAND(FDelayedCallbackLatentCommand([=]{
+					auto Changes = Test->FileCache->GetOutstandingChanges();
+
+					if (Changes.Num() != 1)
+					{
+						UE_LOG(LogAutoReimportTests, Error, TEXT("Incorrect number of changes reported (%d != 1)."), Changes.Num());
+					}
+				}, 1));
+
+			});
+
+		}, 1));
+
+	});
+
+	return true;
+}
+
+/** Test that starting up a cache file with a different set of applicable extensions correctly ignores/updates the extensions, whilst reporting changes only for applicable extensions */
+IMPLEMENT_SIMPLE_AUTOMATION_TEST(FAutoReimportWildcardFiltersTest, "Editor.Auto Reimport.Wildcard Filters", EAutomationTestFlags::ATF_Editor)
+bool FAutoReimportWildcardFiltersTest::RunTest(const FString& Parameters)
+{
+	const FString WorkingDir = GetWorkingDir();
+	static const TCHAR* SrcFilename1 = TEXT("square.png");
+	static const TCHAR* SrcFilename2 = TEXT("red-square.png");
+	static const TCHAR* SrcFilename3 = TEXT("empty.txt");
+
+	static const TCHAR* DstFilename1 = TEXT("sub-folder/square.png");
+	static const TCHAR* DstFilename2 = TEXT("red-square.png");
+	static const TCHAR* DstFilename3 = TEXT("sub-folder/empty.txt");
+	
+	TSharedPtr<FAutoReimportTestPayload> Test = MakeShareable(new FAutoReimportTestPayload(WorkingDir));
+	Test->Config.Rules.SetApplicableExtensions(TEXT("txt;png;"));
+	Test->Config.Rules.AddWildcardRule(TEXT("sub-folder/*"), true);
+	Test->Config.Rules.AddWildcardRule(TEXT("sub-folder/*.png"), false);
+	Test->Config.bDetectChangesSinceLastRun = true;
+
+	IFileManager::Get().MakeDirectory(*(GetWorkingDir() / TEXT("Content") / TEXT("sub-folder")), false);
+		
+	TArray<FSrcDstFilenames> Files;
+	Files.Emplace(SrcFilename1, DstFilename1);
+	Files.Emplace(SrcFilename2, DstFilename2);
+	Files.Emplace(SrcFilename3, DstFilename3);
+	if (!CopyTestFiles(*Test, Files))
+	{
+		return false;
+	}
+
+	// Start watching the directory
+	Test->StartWatching();
+	Test->WaitForStartup([=]{
+
+		if (!Test->FileCache->FindFileData(DstFilename3))
+		{
+			UE_LOG(LogAutoReimportTests, Error, TEXT("Couldn't find data for %s."), DstFilename3);
+		}
+		else if (Test->FileCache->FindFileData(DstFilename2))
+		{
+			UE_LOG(LogAutoReimportTests, Error, TEXT("Erroneously found data for %s."), DstFilename2);
+		}
+		else if (Test->FileCache->FindFileData(DstFilename1))
+		{
+			UE_LOG(LogAutoReimportTests, Error, TEXT("Erroneously found data for %s."), DstFilename1);
+		}
+		else
+		{
+			// Check that we get the same behaviour for live changes
+			TArray<FSrcDstFilenames> Files;
+			Files.Emplace(SrcFilename2, DstFilename1);
+			Files.Emplace(SrcFilename1, DstFilename2);
+			Files.Emplace(SrcFilename2, DstFilename3);
+			if (!CopyTestFiles(*Test, Files))
+			{
+				return;
+			}
+
+			ADD_LATENT_AUTOMATION_COMMAND(FDelayedCallbackLatentCommand([=]{
+				auto Changes = Test->FileCache->GetOutstandingChanges();
+
+				if (Changes.Num() != 1)
+				{
+					UE_LOG(LogAutoReimportTests, Error, TEXT("Incorrect number of changes reported (%d != 1)."), Changes.Num());
+				}
+				else if (Changes[0].Filename.Get() != DstFilename3)
+				{
+					UE_LOG(LogAutoReimportTests, Error, TEXT("Modified file path reported incorrectly (%s != %s)."), *Changes[0].Filename.Get(), DstFilename3);
+				}
+
+			}, 1));
+
+		}
+
+	});
+
+	return true;
+}
\ No newline at end of file
diff --git a/Engine/Source/Editor/UnrealEd/Private/AutoReimport/AutoReimportUtilities.cpp b/Engine/Source/Editor/UnrealEd/Private/AutoReimport/AutoReimportUtilities.cpp
--- a/Engine/Source/Editor/UnrealEd/Private/AutoReimport/AutoReimportUtilities.cpp
+++ b/Engine/Source/Editor/UnrealEd/Private/AutoReimport/AutoReimportUtilities.cpp
@@ -6,6 +6,184 @@
 #include "ARFilter.h"
 #include "IAssetRegistry.h"
 
+bool MatchExtensionString(const TCHAR* Filename, const TCHAR* Extensions)
+{
+	const int32 StrLength = FCString::Strlen(Filename);
+
+	const TCHAR* Ext = FCString::Strrchr(Filename, '.');
+	if (Ext && *(++Ext) != '\0')
+	{
+		const int32 ExtLength = StrLength - (Ext - Filename);
+
+		const TCHAR* Search = FCString::Strchr(Extensions, ';');
+		while(Search)
+		{
+			if (*(++Search) == '\0')
+			{
+				break;
+			}
+
+			if (FCString::Strnicmp(Search, Ext, ExtLength) == 0 && *(Search + ExtLength) == ';')
+			{
+				return true;
+			}
+
+			Search = FCString::Strchr(Search, ';');
+		}
+	}
+
+	return false;
+}
+
+/** Test MatchExtensionString */
+IMPLEMENT_SIMPLE_AUTOMATION_TEST(FMatchExtensionStringTest, "Editor.Auto Reimport.Extension Matching", EAutomationTestFlags::ATF_Editor)
+bool FMatchExtensionStringTest::RunTest(const FString& Parameters)
+{
+	auto Test = [](const TCHAR* Needle, const TCHAR* Haystack, bool bExpected) -> bool{
+
+		if (MatchExtensionString(Needle, Haystack) != bExpected)
+		{
+			if (bExpected)
+			{
+				UE_LOG(LogAutoReimportManager, Error, TEXT("Unable to match '%s' in '%s'"), Needle, Haystack);
+			}
+			else
+			{
+				UE_LOG(LogAutoReimportManager, Error, TEXT("Erroneously matched '%s' in '%s'"), Needle, Haystack);
+			}
+			return false;
+		}
+		return true;
+	};
+
+	bool Result = true;
+	Result = Test(TEXT("bla.txt"), TEXT(";;"), false) && Result;
+	Result = Test(TEXT("bla.txt"), TEXT(";"), false) && Result;
+	Result = Test(TEXT("bla.txt"), TEXT(""), false) && Result;
+	Result = Test(TEXT("bla.txt"), TEXT(";txt;"), true) && Result;
+	Result = Test(TEXT("bla.text"), TEXT(";txt;"), false) && Result;
+	Result = Test(TEXT("bla.txt1"), TEXT(";txt;"), false) && Result;
+	Result = Test(TEXT("bla."), TEXT(";bla;"), false) && Result;
+	Result = Test(TEXT("bla.png"), TEXT(";png;txt;"), true) && Result;
+	Result = Test(TEXT("bla.txt"), TEXT(";png;txt;"), true) && Result;
+	Result = Test(TEXT("/folder.bin/bla.txt"), TEXT(";png;txt;"), true) && Result;
+	Result = Test(TEXT("/folder.bin/bla"), TEXT(";png;bin;"), false) && Result;
+
+	return true;
+}
+
+struct FWildcardRule : IMatchRule
+{
+	FWildcardString WildcardString;
+	bool bInclude;
+
+	virtual ~FWildcardRule(){}
+	
+	virtual TOptional<bool> IsFileApplicable(const TCHAR* Filename) const override
+	{
+		if (WildcardString.IsMatch(Filename))
+		{
+			return bInclude;
+		}
+		return TOptional<bool>();
+	}
+
+	virtual void Serialize(FArchive& Ar) override
+	{
+		Ar << WildcardString;
+		Ar << bInclude;
+	}
+};
+
+void FMatchRules::AddWildcardRule(const FWildcardString& WildcardString, bool bInclude)
+{
+	FMatchRule Rule;
+	Rule.Type = FMatchRule::Wildcard;
+	auto* Impl = new FWildcardRule;
+	Impl->WildcardString = WildcardString;
+	Impl->bInclude = bInclude;
+	Rule.RuleImpl = MakeShareable(Impl);
+	Impls.Add(Rule);
+
+	// If there are any include patterns, we default to not matching all files
+	if (bInclude)
+	{
+		bDefaultIncludeState = false;
+	}
+}
+
+void FMatchRules::SetApplicableExtensions(const FString& InExtensions)
+{
+	ApplicableExtensions = InExtensions;
+
+	// Ensure that the extension strings are of the form ;ext1;ext2;ext3;
+	if (ApplicableExtensions[ApplicableExtensions.Len() - 1] != ';')
+	{
+		ApplicableExtensions += TEXT(";");
+	}
+	if (ApplicableExtensions[0] != ';')
+	{
+		ApplicableExtensions.InsertAt(0, ';');
+	}
+
+}
+
+void FMatchRules::FMatchRule::Serialize(FArchive& Ar)
+{
+	Ar << Type;
+	
+	if (Ar.IsLoading())
+	{
+		if (Type == Wildcard)
+		{
+			RuleImpl = MakeShareable(new FWildcardRule);
+		}
+	}
+
+	if (RuleImpl.IsValid())
+	{
+		RuleImpl->Serialize(Ar);
+	}
+}
+
+bool FMatchRules::IsFileApplicable(const TCHAR* Filename) const
+{
+	if (!ApplicableExtensions.IsEmpty() && !MatchExtensionString(Filename, *ApplicableExtensions))
+	{
+		return false;
+	}
+	
+	// If we have no rules, we match everything
+	if (Impls.Num() == 0)
+	{
+		return true;
+	}
+
+	bool bApplicable = bDefaultIncludeState;
+
+	// Otherwise we match any of the rules
+	for (const auto& Rule : Impls)
+	{
+		if (Rule.RuleImpl.IsValid())
+		{
+			auto IsApplicable = Rule.RuleImpl->IsFileApplicable(Filename);
+			if (IsApplicable.IsSet())
+			{
+				if (!IsApplicable.GetValue())
+				{
+					return false;
+				}
+				else
+				{
+					bApplicable = true;
+				}
+			}
+		}
+	}
+
+	return bApplicable;
+}
+
 namespace Utils
 {
 	TArray<FAssetData> FindAssetsPertainingToFile(const IAssetRegistry& Registry, const FString& AbsoluteFilename)
@@ -43,6 +221,13 @@ namespace Utils
 		return Assets;
 	}
 
+	TArray<FString> ExtractSourceFilePaths(UObject* Object)
+	{
+		TArray<FString> Temp;
+		ExtractSourceFilePaths(Object, Temp);
+		return Temp;
+	}
+
 	void ExtractSourceFilePaths(UObject* Object, TArray<FString>& OutSourceFiles)
 	{
 		TArray<UObject::FAssetRegistryTag> TagList;
diff --git a/Engine/Source/Editor/UnrealEd/Private/AutoReimport/AutoReimportUtilities.h b/Engine/Source/Editor/UnrealEd/Private/AutoReimport/AutoReimportUtilities.h
--- a/Engine/Source/Editor/UnrealEd/Private/AutoReimport/AutoReimportUtilities.h
+++ b/Engine/Source/Editor/UnrealEd/Private/AutoReimport/AutoReimportUtilities.h
@@ -5,6 +5,8 @@
 class FAssetData;
 class IAssetRegistry;
 
+DEFINE_LOG_CATEGORY_STATIC(LogAutoReimportManager, Log, All);
+
 /** An immutable string with a cached CRC for efficient comparison with other strings */
 struct FImmutableString
 {
@@ -12,6 +14,9 @@ struct FImmutableString
 	FImmutableString(FString InString = FString()) : String(MoveTemp(InString)), CachedHash(0) {}
 	FImmutableString& operator=(FString InString) { String = MoveTemp(InString); CachedHash = 0; return *this; }
 
+	FImmutableString(const TCHAR* InString) : String(InString), CachedHash(0) {}
+	FImmutableString& operator=(const TCHAR* InString) { String = InString; CachedHash = 0; return *this; }
+
 	FImmutableString(const FImmutableString&) = default;
 	FImmutableString& operator=(const FImmutableString&) = default;
 
@@ -84,59 +89,131 @@ private:
 	double StartTime;
 };
 
-namespace Utils
+/** Simple helper struct to ease the caching of MD5 hashes */
+struct FMD5Hash
 {
-	/** Reduce the array to the specified accumulator */
-	template<typename T, typename P, typename A>
-	A Reduce(const TArray<T>& InArray, P Predicate, A Accumulator)
+	/** Default constructor */
+	FMD5Hash() : bIsValid(false) {}
+
+	/** Check whether this has hash is valid or not */
+	bool IsValid() const { return bIsValid; }
+
+	/** Set up the MD5 hash from a container */
+	void Set(FMD5& MD5)
 	{
-		for (const T& X : InArray)
+		MD5.Final(Bytes);
+		bIsValid = true;
+	}
+
+	/** Compare one hash with another */
+	friend bool operator==(const FMD5Hash& LHS, const FMD5Hash& RHS)
+	{
+		return LHS.bIsValid == RHS.bIsValid && (!LHS.bIsValid || FMemory::Memcmp(LHS.Bytes, RHS.Bytes, 16) == 0);
+	}
+
+	/** Compare one hash with another */
+	friend bool operator!=(const FMD5Hash& LHS, const FMD5Hash& RHS)
+	{
+		return LHS.bIsValid != RHS.bIsValid || (LHS.bIsValid && FMemory::Memcmp(LHS.Bytes, RHS.Bytes, 16) != 0);
+	}
+
+	/** Serialise this hash */
+	friend FArchive& operator<<(FArchive& Ar, FMD5Hash& Hash)
+	{
+		Ar << Hash.bIsValid;
+		if (Hash.bIsValid)
 		{
-			Accumulator = Predicate(X, Accumulator);
+			Ar.Serialize(Hash.Bytes, 16);
 		}
-		return Accumulator;
+
+		return Ar;
 	}
 
-	/** Helper template function to remove duplicates from an array, given some predicate */
-	template<typename T, typename P>
-	void RemoveDuplicates(TArray<T>& Array, P Predicate)
+private:
+	/** Whether this hash is valid or not */
+	bool bIsValid;
+
+	/** The bytes this hash comprises */
+	uint8 Bytes[16];
+};
+
+/** A rule that checks whether a file is applicable or not */
+struct IMatchRule
+{
+	/** Serialize this rule */
+	virtual void Serialize(FArchive& Ar) = 0;
+
+	/** Test to see if a file is applicable based on this rule. Returns true if so, false of not, or empty if the file doesn't match this rule. */
+	virtual TOptional<bool> IsFileApplicable(const TCHAR* Filename) const = 0;
+};
+
+/** A set of rules that specifies what files apply to the auto reimporter */
+struct FMatchRules
+{
+	FMatchRules() : bDefaultIncludeState(true) {}
+
+	/** Specify a wildcard match to include or exclude */
+	void AddWildcardRule(const FWildcardString& WildcardString, bool bInclude);
+
+	/** Specify a set of applicable extensions, ; separated. Provided as an optimization to early-out on files we're not interested in. */
+	void SetApplicableExtensions(const FString& InExtensions);
+
+	/** Check whether the specified file is applicable based on these rules or not */
+	bool IsFileApplicable(const TCHAR* Filename) const;
+
+	friend void operator<<(FArchive& Ar, FMatchRules& Rules)
 	{
-		if (Array.Num() == 0)
-			return;
+		Ar << Rules.bDefaultIncludeState;
+		Ar << Rules.ApplicableExtensions;
+		Ar << Rules.Impls;
+	}
+
+private:
 
-		const int32 NumElements = Array.Num();
-		int32 Write = 0;
-		for (int32 Read = 0; Read < NumElements; ++Read)
+	/** Implementation of a match rule, wrapping up its type, and implementation */
+	struct FMatchRule
+	{
+		enum EType { Wildcard };
+		/** The type of this rule */
+		TEnumAsByte<EType> Type;
+		/** The runtime implementation */
+		TSharedPtr<IMatchRule> RuleImpl;
+
+	private:
+		void Serialize(FArchive& Ar);
+		friend void operator<<(FArchive& Ar, FMatchRule& Rule)
 		{
-			for (int32 DuplRead = Read + 1; DuplRead < NumElements; ++DuplRead)
-			{
-				if (Predicate(Array[Read], Array[DuplRead]))
-				{
-					goto next;
-				}
-			}
-
-			if (Write != Read)
-			{
-				Swap(Array[Write++], Array[Read]);
-			}
-			else
-			{
-				Write++;
-			}
-	next:
-			;
+			Rule.Serialize(Ar);
 		}
-		const int32 NumDuplicates = NumElements - Write;
-		if (NumDuplicates != 0)
+	};
+
+	/** Optimization to ignore files that don't match the given set of extensions. Expected to be of the form ;ext1;ext2;ext3; */
+	FString ApplicableExtensions;
+
+	/** Array of implementations */
+	TArray<FMatchRule> Impls;
+
+	/** Default inclusion state. True when only exclude rules exist, false if there are any Include rules */
+	bool bDefaultIncludeState;
+};
+
+namespace Utils
+{
+	/** Reduce the array to the specified accumulator */
+	template<typename T, typename P, typename A>
+	A Reduce(const TArray<T>& InArray, P Predicate, A Accumulator)
+	{
+		for (const T& X : InArray)
 		{
-			Array.RemoveAt(Write, NumDuplicates, false);
+			Accumulator = Predicate(X, Accumulator);
 		}
+		return Accumulator;
 	}
-
+	
 	/** Find a list of assets that were once imported from the specified filename */
 	TArray<FAssetData> FindAssetsPertainingToFile(const IAssetRegistry& Registry, const FString& AbsoluteFilename);
 
 	/** Extract any source file paths from the specified object */
+	TArray<FString> ExtractSourceFilePaths(UObject* Object);
 	void ExtractSourceFilePaths(UObject* Object, TArray<FString>& OutSourceFiles);
 }
\ No newline at end of file
diff --git a/Engine/Source/Editor/UnrealEd/Private/AutoReimport/ContentDirectoryMonitor.cpp b/Engine/Source/Editor/UnrealEd/Private/AutoReimport/ContentDirectoryMonitor.cpp
--- a/Engine/Source/Editor/UnrealEd/Private/AutoReimport/ContentDirectoryMonitor.cpp
+++ b/Engine/Source/Editor/UnrealEd/Private/AutoReimport/ContentDirectoryMonitor.cpp
@@ -7,12 +7,13 @@
 #include "AssetRegistryModule.h"
 #include "PackageTools.h"
 #include "ObjectTools.h"
+#include "AssetToolsModule.h"
 #include "ReimportFeedbackContext.h"
 
 #define LOCTEXT_NAMESPACE "ContentDirectoryMonitor"
 
 /** Generate a config from the specified options, to pass to FFileCache on construction */
-FFileCacheConfig GenerateFileCacheConfig(const FString& InPath, const FString& InSupportedExtensions, const FString& InMountedContentPath)
+FFileCacheConfig GenerateFileCacheConfig(const FString& InPath, const FMatchRules& InMatchRules, const FString& InMountedContentPath)
 {
 	FString Directory = FPaths::ConvertRelativePathToFull(InPath);
 
@@ -21,15 +22,16 @@ FFileCacheConfig GenerateFileCacheConfig(const FString& InPath, const FString& I
 	FString CacheFilename = FPaths::ConvertRelativePathToFull(FPaths::GameIntermediateDir()) / TEXT("ReimportCache") / FString::Printf(TEXT("%u.bin"), CRC);
 
 	FFileCacheConfig Config(MoveTemp(Directory), MoveTemp(CacheFilename));
-	Config.IncludeExtensions = InSupportedExtensions;
+	Config.Rules = InMatchRules;
 	// We always store paths inside content folders relative to the folder
 	Config.PathType = EPathType::Relative;
 
+	Config.bDetectChangesSinceLastRun = GetDefault<UEditorLoadingSavingSettings>()->bDetectChangesOnRestart;
 	return Config;
 }
 
-FContentDirectoryMonitor::FContentDirectoryMonitor(const FString& InDirectory, const FString& InSupportedExtensions, const FString& InMountedContentPath)
-	: Cache(GenerateFileCacheConfig(InDirectory, InSupportedExtensions, InMountedContentPath))
+FContentDirectoryMonitor::FContentDirectoryMonitor(const FString& InDirectory, const FMatchRules& InMatchRules, const FString& InMountedContentPath)
+	: Cache(GenerateFileCacheConfig(InDirectory, InMatchRules, InMountedContentPath))
 	, MountedContentPath(InMountedContentPath)
 	, LastSaveTime(0)
 {
@@ -41,14 +43,29 @@ void FContentDirectoryMonitor::Destroy()
 	Cache.Destroy();
 }
 
-void FContentDirectoryMonitor::ReportExternalChange(const FString& Filename, FFileChangeData::EFileChangeAction Action)
+void FContentDirectoryMonitor::IgnoreNewFile(const FString& Filename)
 {
-	Cache.ReportExternalChange(Filename, Action);
+	Cache.IgnoreNewFile(Filename);
 }
 
-void FContentDirectoryMonitor::Tick(const FTimeLimit& TimeLimit)
+void FContentDirectoryMonitor::IgnoreFileModification(const FString& Filename)
 {
-	Cache.Tick(TimeLimit);
+	Cache.IgnoreFileModification(Filename);
+}
+
+void FContentDirectoryMonitor::IgnoreMovedFile(const FString& SrcFilename, const FString& DstFilename)
+{
+	Cache.IgnoreMovedFile(SrcFilename, DstFilename);
+}
+
+void FContentDirectoryMonitor::IgnoreDeletedFile(const FString& Filename)
+{
+	Cache.IgnoreDeletedFile(Filename);
+}
+
+void FContentDirectoryMonitor::Tick()
+{
+	Cache.Tick();
 
 	const double Now = FPlatformTime::Seconds();
 	if (Now - LastSaveTime > ResaveIntervalS)
@@ -71,7 +88,7 @@ void FContentDirectoryMonitor::StartProcessing()
 		{
 			switch(Transaction.Action)
 			{
-				case FFileChangeData::FCA_Added:
+				case EFileAction::Added:
 					if (Settings->bAutoCreateAssets && !MountedContentPath.IsEmpty())
 					{
 						AddedFiles.Emplace(MoveTemp(Transaction));
@@ -82,11 +99,12 @@ void FContentDirectoryMonitor::StartProcessing()
 					}
 					break;
 
-				case FFileChangeData::FCA_Modified:
+				case EFileAction::Moved:
+				case EFileAction::Modified:
 					ModifiedFiles.Emplace(MoveTemp(Transaction));
 					break;
 
-				case FFileChangeData::FCA_Removed:
+				case EFileAction::Removed:
 					if (Settings->bAutoDeleteAssets && !MountedContentPath.IsEmpty())
 					{
 						DeletedFiles.Emplace(MoveTemp(Transaction));
@@ -103,6 +121,23 @@ void FContentDirectoryMonitor::StartProcessing()
 	}
 }
 
+UObject* AttemptImport(UClass* InFactoryType, UPackage* Package, FName InName, bool& bCancelled, const FString& FullFilename)
+{
+	UObject* Asset = nullptr;
+
+	if (UFactory* Factory = NewObject<UFactory>(GetTransientPackage(), InFactoryType))
+	{
+		Factory->AddToRoot();
+		if (Factory->ConfigureProperties())
+		{
+			Asset = UFactory::StaticImportObject(Factory->SupportedClass, Package, InName, RF_Public | RF_Standalone, bCancelled, *FullFilename, nullptr, Factory);
+		}
+		Factory->RemoveFromRoot();
+	}
+
+	return Asset;
+}
+
 void FContentDirectoryMonitor::ProcessAdditions(const IAssetRegistry& Registry, TArray<UPackage*>& OutPackagesToSave, const FTimeLimit& TimeLimit, const TMap<FString, TArray<UFactory*>>& InFactoriesByExtension, FReimportFeedbackContext& Context)
 {
 	bool bCancelled = false;
@@ -161,60 +196,56 @@ void FContentDirectoryMonitor::ProcessAdditions(const IAssetRegistry& Registry,
 			}
 			else
 			{
-				Context.AddMessage(EMessageSeverity::Info, FText::Format(LOCTEXT("Info_CreatingNewAsset", "Creating new asset {0}."), FText::FromString(PackagePath)));
+				Context.AddMessage(EMessageSeverity::Info, FText::Format(LOCTEXT("Info_CreatingNewAsset", "Importing new asset {0}."), FText::FromString(PackagePath)));
 
 				// Make sure the destination package is loaded
 				NewPackage->FullyLoad();
-
-				bool bSuccess = false;
 				
+				UObject* NewAsset = nullptr;
+
 				// Find a relevant factory for this file
 				const FString Ext = FPaths::GetExtension(Addition.Filename.Get(), false);
-				UFactory* FactoryType = nullptr;
-				if (auto* Factories = InFactoriesByExtension.Find(Ext))
+				auto* Factories = InFactoriesByExtension.Find(Ext);
+				if (Factories && Factories->Num() != 0)
 				{
-					FactoryType = (*Factories)[0];
-
-					UFactory* const* FoundCompatibleFactoryType = Factories->FindByPredicate([&](UFactory* InFactory){
-						return InFactory->FactoryCanImport(FullFilename);
-					});
-
-					if (FoundCompatibleFactoryType)
+					// Prefer a factory if it explicitly can import. UFactory::FactoryCanImport returns false by default, even if the factory supports the extension, so we can't use it directly.
+					UFactory* const * PreferredFactory = Factories->FindByPredicate([&](UFactory* F){ return F->FactoryCanImport(FullFilename); });
+					if (PreferredFactory)
 					{
-						FactoryType = *FoundCompatibleFactoryType;
+						NewAsset = AttemptImport((*PreferredFactory)->GetClass(), NewPackage, *NewAssetName, bCancelled, FullFilename);
 					}
-				}
-
-				UObject* NewAsset = nullptr;
-				UFactory* FactoryInstance = FactoryType ? NewObject<UFactory>(GetTransientPackage(), FactoryType->GetClass()) : nullptr;
-				if (FactoryInstance)
-				{
-					FactoryInstance->AddToRoot();
-					if (FactoryInstance->ConfigureProperties())
+					// If there was no preferred factory, just try them all until one succeeds
+					else for (UFactory* Factory : *Factories)
 					{
-						UClass* ImportAssetType = FactoryInstance->SupportedClass;
-						NewAsset = UFactory::StaticImportObject(ImportAssetType, NewPackage, FName(*NewAssetName), RF_Public | RF_Standalone, bCancelled, *FullFilename, nullptr, FactoryInstance);
+						NewAsset = AttemptImport(Factory->GetClass(), NewPackage, *NewAssetName, bCancelled, FullFilename);
+
+						if (bCancelled || NewAsset)
+						{
+							break;
+						}
 					}
-					FactoryInstance->RemoveFromRoot();
 				}
 
-				if (!bCancelled)
+				// If we didn't create an asset, unload and delete the package we just created
+				if (!NewAsset)
 				{
-					if (!NewAsset)
-					{
-						TArray<UPackage*> Packages;
-						Packages.Add(NewPackage);
+					TArray<UPackage*> Packages;
+					Packages.Add(NewPackage);
 
-						FText ErrorMessage;
-						if (!PackageTools::UnloadPackages(Packages, ErrorMessage))
-						{
-							Context.AddMessage(EMessageSeverity::Error, FText::Format(LOCTEXT("Error_UnloadingPackage", "There was an error unloading a package: {0}."), ErrorMessage));
-						}						
+					TGuardValue<bool> SuppressSlowTaskMessages(Context.bSuppressSlowTaskMessages, true);
 
-						Context.AddMessage(EMessageSeverity::Error, FText::Format(LOCTEXT("Error_FailedToImportAsset", "Failed to import file {0}."), FText::FromString(FullFilename)));
-						continue;
+					FText ErrorMessage;
+					if (!PackageTools::UnloadPackages(Packages, ErrorMessage))
+					{
+						Context.AddMessage(EMessageSeverity::Error, FText::Format(LOCTEXT("Error_UnloadingPackage", "There was an error unloading a package: {0}."), ErrorMessage));
 					}
 
+					// Just add the message to the message log rather than add it to the UI
+					// Factories may opt not to import the file, so we let them report errors if they do
+					Context.GetMessageLog().Message(EMessageSeverity::Info, FText::Format(LOCTEXT("Info_FailedToImportAsset", "Failed to import file {0}."), FText::FromString(FullFilename)));
+				}
+				else if (!bCancelled)
+				{
 					FAssetRegistryModule::AssetCreated(NewAsset);
 					GEditor->BroadcastObjectReimported(NewAsset);
 
@@ -223,7 +254,7 @@ void FContentDirectoryMonitor::ProcessAdditions(const IAssetRegistry& Registry,
 
 				// Refresh the supported class.  Some factories (e.g. FBX) only resolve their type after reading the file
 				// ImportAssetType = Factory->ResolveSupportedClass();
-				// @todo: analytics
+				// @todo: analytics?
 			}
 		}
 
@@ -248,19 +279,61 @@ void FContentDirectoryMonitor::ProcessModifications(const IAssetRegistry& Regist
 	for (int32 Index = 0; Index < ModifiedFiles.Num(); ++Index)
 	{
 		auto& Change = ModifiedFiles[Index];
-
 		const FString FullFilename = Cache.GetDirectory() + Change.Filename.Get();
 
-		for (const auto& AssetData : Utils::FindAssetsPertainingToFile(Registry, FullFilename))
+		// Move the asset before reimporting it. We always reimport moved assets to ensure that their import path is up to date
+		if (Change.Action == EFileAction::Moved)
 		{
-			UObject* Asset = AssetData.GetAsset();
-			if (!ReimportManager->Reimport(Asset, false /* Ask for new file */, false /* Show notification */))
+			const FString OldFilename = Cache.GetDirectory() + Change.MovedFromFilename.Get();
+			const auto Assets = Utils::FindAssetsPertainingToFile(Registry, OldFilename);
+
+			if (Assets.Num() == 1)
 			{
-				Context.AddMessage(EMessageSeverity::Error, FText::Format(LOCTEXT("Error_FailedToReimportAsset", "Failed to reimport asset {0}."), FText::FromString(Asset->GetName())));
+				UObject* Asset = Assets[0].GetAsset();
+				if (Asset && Utils::ExtractSourceFilePaths(Asset).Num() == 1)
+				{
+					FString NewAssetName = ObjectTools::SanitizeObjectName(FPaths::GetBaseFilename(Change.Filename.Get()));
+					FString PackagePath = PackageTools::SanitizePackageName(MountedContentPath / FPaths::GetPath(Change.Filename.Get()));
+
+					const FText SrcPathText = FText::FromString(Assets[0].PackagePath.ToString() / Assets[0].PackageName.ToString()),
+						DstPathText = FText::FromString(PackagePath / NewAssetName);
+
+					if (FEditorFileUtils::IsMapPackageAsset(PackagePath) || FindPackage(nullptr, *PackagePath))
+					{
+						Context.AddMessage(EMessageSeverity::Warning, FText::Format(LOCTEXT("MoveWarning_ExistingAsset", "Can't move {0} to {1} - one already exists."), SrcPathText, DstPathText));
+					}
+					else
+					{
+						TArray<FAssetRenameData> RenameData;
+						RenameData.Emplace(Asset, PackagePath, NewAssetName);
+
+						Context.AddMessage(EMessageSeverity::Info, FText::Format(LOCTEXT("Success_MovedAsset", "Moving asset {0} to {1}."),	SrcPathText, DstPathText));
+
+						FModuleManager::LoadModuleChecked<FAssetToolsModule>("AssetTools").Get().RenameAssets(RenameData);
+
+						TArray<FString> Filenames;
+						Filenames.Add(FullFilename);
+
+						// Update the reimport file names
+						FReimportManager::Instance()->UpdateReimportPaths(Asset, Filenames);
+						Asset->MarkPackageDirty();
+					}
+				}
 			}
-			else
+		}
+		else if (Change.Action == EFileAction::Modified)
+		{
+			for (const auto& AssetData : Utils::FindAssetsPertainingToFile(Registry, FullFilename))
 			{
-				Context.AddMessage(EMessageSeverity::Info, FText::Format(LOCTEXT("Success_CreatedNewAsset", "Reimported asset {0} from {1}."), FText::FromString(Asset->GetName()), FText::FromString(FullFilename)));
+				UObject* Asset = AssetData.GetAsset();
+				if (!ReimportManager->Reimport(Asset, false /* Ask for new file */, false /* Show notification */))
+				{
+					Context.AddMessage(EMessageSeverity::Error, FText::Format(LOCTEXT("Error_FailedToReimportAsset", "Failed to reimport asset {0}."), FText::FromString(Asset->GetName())));
+				}
+				else
+				{
+					Context.AddMessage(EMessageSeverity::Info, FText::Format(LOCTEXT("Success_CreatedNewAsset", "Reimported asset {0} from {1}."), FText::FromString(Asset->GetName()), FText::FromString(FullFilename)));
+				}
 			}
 		}
 
diff --git a/Engine/Source/Editor/UnrealEd/Private/AutoReimport/ContentDirectoryMonitor.h b/Engine/Source/Editor/UnrealEd/Private/AutoReimport/ContentDirectoryMonitor.h
--- a/Engine/Source/Editor/UnrealEd/Private/AutoReimport/ContentDirectoryMonitor.h
+++ b/Engine/Source/Editor/UnrealEd/Private/AutoReimport/ContentDirectoryMonitor.h
@@ -14,13 +14,13 @@ public:
 	/**
 	 * Constructor.
 	 * @param InDirectory			Content directory path to monitor. Assumed to be absolute.
-	 * @param InSupportedExtensions A string containing semi-colon separated extensions to monitor.
+	 * @param InMatchRules			Rules to select what will apply to this folder
 	 * @param InMountedContentPath	(optional) Mounted content path (eg /Engine/, /Game/) to which InDirectory maps.
 	 */
-	FContentDirectoryMonitor(const FString& InDirectory, const FString& InSupportedExtensions, const FString& InMountedContentPath = FString());
+	FContentDirectoryMonitor(const FString& InDirectory, const FMatchRules& InMatchRules, const FString& InMountedContentPath = FString());
 
 	/** Tick this monitor's cache to give it a chance to finish scanning for files */
-	void Tick(const FTimeLimit& TimeLimit);
+	void Tick();
 
 	/** Start processing any outstanding changes this monitor is aware of */
 	void StartProcessing();
@@ -35,7 +35,10 @@ public:
 	void ExtractAssetsToDelete(const IAssetRegistry& Registry, TArray<FAssetData>& OutAssetsToDelete);
 
 	/** Report an external change to the manager, such that a subsequent equal change reported by the os be ignored */
-	void ReportExternalChange(const FString& Filename, FFileChangeData::EFileChangeAction Action);
+	void IgnoreNewFile(const FString& Filename);
+	void IgnoreFileModification(const FString& Filename);
+	void IgnoreMovedFile(const FString& SrcFilename, const FString& DstFilename);
+	void IgnoreDeletedFile(const FString& Filename);
 
 	/** Destroy this monitor including its cache */
 	void Destroy();
diff --git a/Engine/Source/Editor/UnrealEd/Private/AutoReimport/FileCache.cpp b/Engine/Source/Editor/UnrealEd/Private/AutoReimport/FileCache.cpp
--- a/Engine/Source/Editor/UnrealEd/Private/AutoReimport/FileCache.cpp
+++ b/Engine/Source/Editor/UnrealEd/Private/AutoReimport/FileCache.cpp
@@ -7,34 +7,177 @@
 #include "DirectoryWatcherModule.h"
 #include "AutoReimportUtilities.h"
 
+/** Convert a FFileChangeData::EFileChangeAction into an EFileAction */
+EFileAction ToFileAction(FFileChangeData::EFileChangeAction InAction)
+{
+	switch (InAction)
+	{
+	case FFileChangeData::FCA_Added: 	return EFileAction::Added;
+	case FFileChangeData::FCA_Modified:	return EFileAction::Modified;
+	case FFileChangeData::FCA_Removed:	return EFileAction::Removed;
+	default:							return EFileAction::Modified;
+	}
+}
+
+/** Synchronously read a file's MD5 */
+FMD5Hash ReadFileMD5(const FString& Filename)
+{
+	FArchive* Ar = IFileManager::Get().CreateFileReader(*Filename);
+
+	FMD5Hash Hash;
+	if (Ar)
+	{
+		FMD5 MD5;
+
+		static const int64 BufferSize = 1024;
+		const int64 Size = Ar->TotalSize();
+		int64 Position = 0;
+		uint8 Buffer[BufferSize];
+
+		// Read in BufferSize chunks
+		while (Position < Size)
+		{
+			const auto ReadNum = FMath::Min(Size - Position, BufferSize);
+			Ar->Serialize(Buffer, ReadNum);
+			MD5.Update(Buffer, ReadNum);
+
+			Position += ReadNum;
+		}
+
+		Hash.Set(MD5);
+		delete Ar;
+	}
+
+	return Hash;
+}
+
 const FGuid FFileCacheCustomVersion::Key(0x8E7DDCB3, 0x80DA47BB, 0x9FD346A2, 0x93984DF6);
 FCustomVersionRegistration GRegisterFileCacheVersion(FFileCacheCustomVersion::Key, FFileCacheCustomVersion::Latest, TEXT("FileCacheVersion"));
 
+/** Single runnable thread used to parse file cache directories without blocking the main thread */
+struct FAsyncDirectoryReaderThread : public FRunnable
+{
+	typedef TArray<TWeakPtr<FAsyncDirectoryReader, ESPMode::ThreadSafe>> FReaderArray;
+
+	FAsyncDirectoryReaderThread() : Thread(nullptr) {}
+
+	/** Add a reader to this thread which will get ticked periodically until complete */
+	void AddReader(TSharedPtr<FAsyncDirectoryReader, ESPMode::ThreadSafe> InReader)
+	{
+		FScopeLock Lock(&ReaderArrayMutex);
+		Readers.Add(InReader);
+
+		if (!Thread)
+		{
+			static int32 Index = 0;
+			Thread = FRunnableThread::Create(this, *FString::Printf(TEXT("AsyncDirectoryReaders_%d"), ++Index));
+		}
+	}
+
+	/** Run this thread */
+	virtual uint32 Run()
+	{
+		for(;;)
+		{
+			// Copy the array while we tick the readers
+			FReaderArray Dupl;
+			{
+				FScopeLock Lock(&ReaderArrayMutex);
+				Dupl = Readers;
+			}
+
+			// Tick each one for a second
+			for (auto& Reader : Dupl)
+			{
+				auto PinnedReader = Reader.Pin();
+				if (PinnedReader.IsValid())
+				{
+					PinnedReader->Tick(FTimeLimit(1));
+				}
+			}
+
+			// Cleanup dead/finished readers
+			FScopeLock Lock(&ReaderArrayMutex);
+			for (int32 Index = 0; Index < Readers.Num(); )
+			{
+				auto Reader = Readers[Index].Pin();
+				if (!Reader.IsValid() || Reader->IsComplete())
+				{
+					Readers.RemoveAt(Index);
+				}
+				else
+				{
+					++Index;
+				}
+			}
+
+			// Shutdown the thread if we've nothing left to do
+			if (Readers.Num() == 0)
+			{
+				Thread = nullptr;
+				break;
+			}
+		}
+
+		return 0;
+	}
+
+private:
+	/** We start our own thread if one doesn't already exist. Guard it with a critical section */
+	FRunnableThread* Thread;
+
+	/** Array of things that need ticking, and a separate mutex to protect those */
+	FCriticalSection ReaderArrayMutex;
+	FReaderArray Readers;
+};
+
+FAsyncDirectoryReaderThread AsyncReaderThread;
+
 FAsyncDirectoryReader::FAsyncDirectoryReader(const FString& InDirectory, EPathType InPathType)
-	: RootPath(InDirectory), PathType(InPathType)
+	: RootPath(InDirectory), PathType(InPathType), StartTime(0)
 {
 	PendingDirectories.Add(InDirectory);
-	State.Emplace();
+	LiveState.Emplace();
+}
+
+TOptional<FDirectoryState> FAsyncDirectoryReader::GetLiveState()
+{
+	TOptional<FDirectoryState> OldState;
+	Swap(OldState, LiveState);
+
+	return OldState;
 }
 
-TOptional<FDirectoryState> FAsyncDirectoryReader::GetFinalState()
+
+TOptional<FDirectoryState> FAsyncDirectoryReader::GetCachedState()
 {
 	TOptional<FDirectoryState> OldState;
-	Swap(OldState, State);
+	Swap(OldState, CachedState);
 
 	return OldState;
 }
 
 bool FAsyncDirectoryReader::IsComplete() const
 {
-	return PendingFiles.Num() == 0 && PendingDirectories.Num() == 0;
+	return bIsComplete;
 }
 
 FAsyncDirectoryReader::EProgressResult FAsyncDirectoryReader::Tick(const FTimeLimit& TimeLimit)
 {
+	if (IsComplete())
+	{
+		return EProgressResult::Finished;
+	}
+
+	if (StartTime == 0)
+	{
+		StartTime = FPlatformTime::Seconds();
+	}
+
 	auto& FileManager = IFileManager::Get();
 	const int32 RootPathLen = RootPath.Len();
 
+	// Discover files
 	for (int32 Index = 0; Index < PendingDirectories.Num(); ++Index)
 	{
 		ScanDirectory(PendingDirectories[Index]);
@@ -47,15 +190,35 @@ FAsyncDirectoryReader::EProgressResult FAsyncDirectoryReader::Tick(const FTimeLi
 	}
 	PendingDirectories.Empty();
 
+	// Process files
 	for (int32 Index = 0; Index < PendingFiles.Num(); ++Index)
 	{
 		const auto& File = PendingFiles[Index];
 
 		// Store the file relative or absolute
 		FString Filename = (PathType == EPathType::Relative ? *File + RootPathLen : *File);
-		State->Files.Emplace(MoveTemp(Filename), FileManager.GetTimeStamp(*File));
 
-		if (Index % 100 == 0 && TimeLimit.Exceeded())
+		const auto Timestamp = FileManager.GetTimeStamp(*File);
+	
+		FMD5Hash MD5;
+		if (CachedState.IsSet())
+		{
+			const FFileData* CachedData = CachedState->Files.Find(Filename);
+			if (CachedData && CachedData->Timestamp == Timestamp && CachedData->FileHash.IsValid())
+			{
+				// Use the cached MD5 to avoid opening the file
+				MD5 = CachedData->FileHash;
+			}
+		}
+
+		if (!MD5.IsValid())
+		{
+			MD5 = ReadFileMD5(File);
+		}
+
+		LiveState->Files.Emplace(MoveTemp(Filename), FFileData(Timestamp, MD5));
+
+		if (TimeLimit.Exceeded())
 		{
 			// We've spent too long, bail
 			PendingFiles.RemoveAt(0, Index + 1, false);
@@ -64,6 +227,9 @@ FAsyncDirectoryReader::EProgressResult FAsyncDirectoryReader::Tick(const FTimeLi
 	}
 	PendingFiles.Empty();
 
+	bIsComplete = true;
+
+	UE_LOG(LogAutoReimportManager, Log, TEXT("Scanning file cache directory took %.2fs"), FPlatformTime::Seconds() - StartTime);
 	return EProgressResult::Finished;
 }
 
@@ -73,16 +239,19 @@ void FAsyncDirectoryReader::ScanDirectory(const FString& InDirectory)
 	{
 		TArray<FString>* PendingFiles;
 		TArray<FString>* PendingDirectories;
+		FMatchRules* Rules;
+		int32 RootPathLength;
 
 		virtual bool Visit(const TCHAR* FilenameOrDirectory, bool bIsDirectory)
 		{
+			FString FileStr(FilenameOrDirectory);
 			if (bIsDirectory)
 			{
-				PendingDirectories->Add(FString(FilenameOrDirectory));
+				PendingDirectories->Add(MoveTemp(FileStr));
 			}
-			else
+			else if (Rules->IsFileApplicable(FilenameOrDirectory + RootPathLength))
 			{
-				PendingFiles->Add(FString(FilenameOrDirectory));
+				PendingFiles->Add(MoveTemp(FileStr));
 			}
 			return true;
 		}
@@ -91,31 +260,31 @@ void FAsyncDirectoryReader::ScanDirectory(const FString& InDirectory)
 	FVisitor Visitor;
 	Visitor.PendingFiles = &PendingFiles;
 	Visitor.PendingDirectories = &PendingDirectories;
+	Visitor.Rules = &LiveState->Rules;
+	Visitor.RootPathLength = RootPath.Len();
+
 	IFileManager::Get().IterateDirectory(*InDirectory, Visitor);
 }
 
 FFileCache::FFileCache(const FFileCacheConfig& InConfig)
 	: Config(InConfig)
-	, DirectoryReader(InConfig.Directory, Config.PathType)
 	, bSavedCacheDirty(false)
 {
-	// Ensure that the extension strings are of the form ;ext1;ext2;ext3;
-	if (!Config.IncludeExtensions.IsEmpty() && Config.IncludeExtensions[0] != ';')
-	{
-		Config.IncludeExtensions.InsertAt(0, ';');
-	}
-	if (!Config.ExcludeExtensions.IsEmpty() && Config.ExcludeExtensions[0] != ';')
-	{
-		Config.ExcludeExtensions.InsertAt(0, ';');
-	}
+	// Ensure the directory has a trailing /
+	Config.Directory /= TEXT("");
+
+	DirectoryReader = MakeShareable(new FAsyncDirectoryReader(Config.Directory, Config.PathType));
+	DirectoryReader->SetMatchRules(Config.Rules);
 
 	// Attempt to load an existing cache file
 	auto ExistingCache = ReadCache();
 	if (ExistingCache.IsSet())
 	{
-		CachedDirectoryState = MoveTemp(ExistingCache.GetValue());
+		DirectoryReader->UseCachedState(MoveTemp(ExistingCache.GetValue()));
 	}
 
+	AsyncReaderThread.AddReader(DirectoryReader);
+
 	FDirectoryWatcherModule& Module = FModuleManager::LoadModuleChecked<FDirectoryWatcherModule>(TEXT("DirectoryWatcher"));
 	if (IDirectoryWatcher* DirectoryWatcher = Module.Get())
 	{
@@ -134,15 +303,30 @@ void FFileCache::Destroy()
 {
 	// Delete the cache file, and clear out everything
 	bSavedCacheDirty = false;
-	IFileManager::Get().Delete(*Config.CacheFile, false, true, true);
+	if (!Config.CacheFile.IsEmpty())
+	{
+		IFileManager::Get().Delete(*Config.CacheFile, false, true, true);
+	}
 
-	DirectoryReader = FAsyncDirectoryReader(Config.Directory, Config.PathType);
-	OutstandingChanges.Empty();
+	DirectoryReader = nullptr;
+	PendingTransactions.Empty();
+	DirtyFiles.Empty();
 	CachedDirectoryState = FDirectoryState();
 
 	UnbindWatcher();
 }
 
+const FFileData* FFileCache::FindFileData(FImmutableString InFilename) const
+{
+	if (!ensure(HasStartedUp()))
+	{
+		// It's invalid to call this while the cached state is still being updated on a thread.
+		return nullptr;
+	}
+
+	return CachedDirectoryState.Files.Find(InFilename);
+}
+
 void FFileCache::UnbindWatcher()
 {
 	if (!WatcherDelegate.IsValid())
@@ -164,17 +348,19 @@ void FFileCache::UnbindWatcher()
 TOptional<FDirectoryState> FFileCache::ReadCache() const
 {
 	TOptional<FDirectoryState> Optional;
+	if (!Config.CacheFile.IsEmpty())
+	{		
+		FArchive* Ar = IFileManager::Get().CreateFileReader(*Config.CacheFile);
+		if (Ar)
+		{
+			FDirectoryState Result;
 
-	FArchive* Ar = IFileManager::Get().CreateFileReader(*Config.CacheFile);
-	if (Ar)
-	{
-		FDirectoryState Result;
-
-		*Ar << Result;
-		Ar->Close();
-		delete Ar;
+			*Ar << Result;
+			Ar->Close();
+			delete Ar;
 
-		Optional.Emplace(MoveTemp(Result));
+			Optional.Emplace(MoveTemp(Result));
+		}
 	}
 
 	return Optional;
@@ -182,7 +368,7 @@ TOptional<FDirectoryState> FFileCache::ReadCache() const
 
 void FFileCache::WriteCache()
 {
-	if (bSavedCacheDirty)
+	if (bSavedCacheDirty && !Config.CacheFile.IsEmpty())
 	{
 		const FString ParentFolder = FPaths::GetPath(Config.CacheFile);
 		if (!IFileManager::Get().DirectoryExists(*ParentFolder))
@@ -205,48 +391,217 @@ void FFileCache::WriteCache()
 	}
 }
 
+FString FFileCache::GetAbsolutePath(const FString& InTransactionPath) const
+{
+	if (Config.PathType == EPathType::Relative)
+	{
+		return Config.Directory / InTransactionPath;
+	}
+	else
+	{
+		return InTransactionPath;
+	}
+}
+
+TOptional<FString> FFileCache::GetTransactionPath(const FString& InAbsolutePath) const
+{
+	FString Temp = FPaths::ConvertRelativePathToFull(InAbsolutePath);
+	FString RelativePath(*Temp + Config.Directory.Len());
+
+	// If it's a directory or is not applicable, ignore it
+	if (!Temp.StartsWith(Config.Directory) || IFileManager::Get().DirectoryExists(*Temp) || !Config.Rules.IsFileApplicable(*RelativePath))
+	{
+		return TOptional<FString>();
+	}
+
+	if (Config.PathType == EPathType::Relative)
+	{
+		return MoveTemp(RelativePath);
+	}
+	else
+	{
+		return MoveTemp(Temp);
+	}
+}
+
+void FFileCache::DiffDirtyFiles(const TSet<FImmutableString>& InDirtyFiles, TArray<FUpdateCacheTransaction>& InOutTransactions, const FDirectoryState* InFileSystemState) const
+{
+	TMap<FImmutableString, FFileData> AddedFiles, ModifiedFiles;
+	TSet<FImmutableString> RemovedFiles;
+
+	auto& FileManager = IFileManager::Get();
+	auto& PlatformFile = FPlatformFileManager::Get().GetPlatformFile();
+
+	for (const auto& File : InDirtyFiles)
+	{
+		FString AbsoluteFilename = GetAbsolutePath(File.Get());
+		
+		const auto* CachedState = CachedDirectoryState.Files.Find(File);
+
+		const bool bFileExists = InFileSystemState ? InFileSystemState->Files.Find(File) != nullptr : PlatformFile.FileExists(*AbsoluteFilename);
+		if (bFileExists)
+		{
+			const FFileData FileData = InFileSystemState ? *InFileSystemState->Files.Find(File) : FFileData(FileManager.GetTimeStamp(*AbsoluteFilename), ReadFileMD5(AbsoluteFilename));
+
+			// Do we think it exists in the cache?
+			if (CachedState)
+			{
+				// Has it changed?
+				if (*CachedState != FileData)
+				{
+					ModifiedFiles.Add(File, FileData);
+				}
+			}
+			else
+			{
+				AddedFiles.Add(File, FileData);
+			}
+		}
+		// We only report it as removed if it exists in the cache
+		else if (CachedState)
+		{
+			RemovedFiles.Add(File);
+		}
+	}
+
+	// Rename / move detection
+	for (auto RemoveIt = RemovedFiles.CreateIterator(); RemoveIt; ++RemoveIt)
+	{
+		const auto* CachedState = CachedDirectoryState.Files.Find(*RemoveIt);
+		if (CachedState)
+		{
+			for (auto AdIt = AddedFiles.CreateIterator(); AdIt; ++AdIt)
+			{				
+				if (AdIt.Value().FileHash == CachedState->FileHash)
+				{
+					// Found a move destination!
+					InOutTransactions.Add(FUpdateCacheTransaction(*RemoveIt, AdIt.Key(), AdIt.Value()));
+
+					AdIt.RemoveCurrent();
+					RemoveIt.RemoveCurrent();
+					break;
+				}
+			}
+		}
+	}
+
+	for (auto& RemovedFile : RemovedFiles)
+	{
+		InOutTransactions.Add(FUpdateCacheTransaction(MoveTemp(RemovedFile), EFileAction::Removed));
+	}
+	// RemovedFiles is now bogus
+
+	for (auto& Pair : AddedFiles)
+	{
+		InOutTransactions.Add(FUpdateCacheTransaction(MoveTemp(Pair.Key), EFileAction::Added, Pair.Value));
+	}
+	// AddedFiles is now bogus
+	
+	for (auto& Pair : ModifiedFiles)
+	{
+		InOutTransactions.Add(FUpdateCacheTransaction(MoveTemp(Pair.Key), EFileAction::Modified, Pair.Value));
+	}
+	// ModifiedFiles is now bogus
+}
+
 TArray<FUpdateCacheTransaction> FFileCache::GetOutstandingChanges()
 {
+	DiffDirtyFiles(DirtyFiles, PendingTransactions);
+	DirtyFiles.Empty();
+
 	TArray<FUpdateCacheTransaction> Moved;
-	Swap(OutstandingChanges, Moved);
+	Swap(PendingTransactions, Moved);
 	return Moved;
 }
 
-void FFileCache::ReportExternalChange(const FString& Filename, FFileChangeData::EFileChangeAction Action)
+void FFileCache::IgnoreNewFile(const FString& Filename)
+{
+	auto TransactionPath = GetTransactionPath(Filename);
+	if (TransactionPath.IsSet())
+	{
+		DirtyFiles.Remove(TransactionPath.GetValue());
+		
+		const FFileData FileData(IFileManager::Get().GetTimeStamp(*Filename), ReadFileMD5(Filename));
+		CompleteTransaction(FUpdateCacheTransaction(MoveTemp(TransactionPath.GetValue()), EFileAction::Added, FileData));	
+	}
+}
+
+void FFileCache::IgnoreFileModification(const FString& Filename)
 {
-	auto Transaction = ProcessChange(Filename, Action, IFileManager::Get());
-	if (Transaction.IsSet())
+	auto TransactionPath = GetTransactionPath(Filename);
+	if (TransactionPath.IsSet())
 	{
-		// Just complete the transaction immediately without notifying the client
-		CompleteTransaction(MoveTemp(Transaction.GetValue()));
+		DirtyFiles.Remove(TransactionPath.GetValue());
+		
+		const FFileData FileData(IFileManager::Get().GetTimeStamp(*Filename), ReadFileMD5(Filename));
+		CompleteTransaction(FUpdateCacheTransaction(MoveTemp(TransactionPath.GetValue()), EFileAction::Modified, FileData));	
+	}
+}
+
+void FFileCache::IgnoreMovedFile(const FString& SrcFilename, const FString& DstFilename)
+{
+	auto SrcTransactionPath = GetTransactionPath(SrcFilename);
+	auto DstTransactionPath = GetTransactionPath(DstFilename);
+
+	if (SrcTransactionPath.IsSet() && DstTransactionPath.IsSet())
+	{
+		DirtyFiles.Remove(SrcTransactionPath.GetValue());
+		DirtyFiles.Remove(DstTransactionPath.GetValue());
+
+		const FFileData FileData(IFileManager::Get().GetTimeStamp(*DstFilename), ReadFileMD5(DstFilename));
+		CompleteTransaction(FUpdateCacheTransaction(MoveTemp(SrcTransactionPath.GetValue()), MoveTemp(DstTransactionPath.GetValue()), FileData));	
+	}
+}
+
+void FFileCache::IgnoreDeletedFile(const FString& Filename)
+{
+	auto TransactionPath = GetTransactionPath(Filename);
+	if (TransactionPath.IsSet())
+	{
+		DirtyFiles.Remove(TransactionPath.GetValue());
+		CompleteTransaction(FUpdateCacheTransaction(MoveTemp(TransactionPath.GetValue()), EFileAction::Removed));	
 	}
 }
 
 void FFileCache::CompleteTransaction(FUpdateCacheTransaction&& Transaction)
 {
 	auto* CachedData = CachedDirectoryState.Files.Find(Transaction.Filename);
-
 	switch (Transaction.Action)
 	{
-	case FFileChangeData::FCA_Modified:
-		if (CachedData && CachedData->Timestamp < Transaction.Timestamp)
+	case EFileAction::Moved:
+		{
+			CachedDirectoryState.Files.Remove(Transaction.MovedFromFilename);
+			if (!CachedData)
+			{
+				CachedDirectoryState.Files.Add(Transaction.Filename, Transaction.FileData);
+			}
+			else
+			{
+				*CachedData = Transaction.FileData;
+			}
+
+			bSavedCacheDirty = true;
+		}
+		break;
+	case EFileAction::Modified:
+		if (CachedData)
 		{
 			// Update the timestamp
-			CachedData->Timestamp = Transaction.Timestamp;
+			*CachedData = Transaction.FileData;
 
 			bSavedCacheDirty = true;
 		}
 		break;
-	case FFileChangeData::FCA_Added:
+	case EFileAction::Added:
 		if (!CachedData)
 		{
 			// Add the file information to the cache
-			CachedDirectoryState.Files.Emplace(Transaction.Filename, Transaction.Timestamp);
+			CachedDirectoryState.Files.Emplace(Transaction.Filename, Transaction.FileData);
 
 			bSavedCacheDirty = true;
 		}
 		break;
-	case FFileChangeData::FCA_Removed:
+	case EFileAction::Removed:
 		if (CachedData)
 		{
 			// Remove the file information to the cache
@@ -262,194 +617,127 @@ void FFileCache::CompleteTransaction(FUpdateCacheTransaction&& Transaction)
 	}
 }
 
-void FFileCache::Tick(const FTimeLimit& Limit)
+void FFileCache::Tick()
 {
-	if (DirectoryReader.IsComplete() || DirectoryReader.Tick(Limit) == FAsyncDirectoryReader::EProgressResult::Pending)
+	if (!DirectoryReader.IsValid() || !DirectoryReader->IsComplete())
 	{
 		return;
 	}
 
 	// We should only ever get here once. The directory reader has finished scanning, and we can now diff the results with what we had saved in the cache file.
-	check(DirectoryReader.IsComplete());
+	check(DirectoryReader->IsComplete());
 	
-	TOptional<FDirectoryState> CurrentFileData = DirectoryReader.GetFinalState();
-	if (CachedDirectoryState.Files.Num() == 0)
+	TOptional<FDirectoryState> LiveState = DirectoryReader->GetLiveState();
+	TOptional<FDirectoryState> CachedState = DirectoryReader->GetCachedState();
+
+	// Null out our pointer to the directory reader to indicate that we've finished
+	DirectoryReader = nullptr;
+
+	if (!CachedState.IsSet() || !Config.bDetectChangesSinceLastRun)
 	{
 		// If we don't have any cached data yet, just use the file data we just harvested
-		CachedDirectoryState = MoveTemp(CurrentFileData.GetValue());
+		CachedDirectoryState = MoveTemp(LiveState.GetValue());
 		bSavedCacheDirty = true;
 		return;
 	}
+	else
+	{
+		// Use the cache that we gave to the directory reader
+		CachedDirectoryState = MoveTemp(CachedState.GetValue());
+	}
+
+	// Build up a list of dirty files
+	TSet<FImmutableString> LocalDirtyFiles;
 
 	// We already have cached data so we need to compare it with the harvested data
 	// to detect additions, modifications, and removals
-	for (const auto& FilenameAndData : CurrentFileData->Files)
+	for (const auto& FilenameAndData : LiveState->Files)
 	{
 		const FString& Filename = FilenameAndData.Key.Get();
-		if (!IsFileApplicable(Filename))
-		{
-			continue;
-		}
 
 		const auto* CachedData = CachedDirectoryState.Files.Find(Filename);
-		// Have we encountered this before?
-		if (CachedData)
-		{
-			if (CachedData->Timestamp != FilenameAndData.Value.Timestamp)
-			{
-				OutstandingChanges.Add(FUpdateCacheTransaction(Filename, FFileChangeData::FCA_Modified, FilenameAndData.Value.Timestamp));
-			}
-		}
-		else
+		if (!CachedData || *CachedData != FilenameAndData.Value)
 		{
-			// No file already - create one. We don't do this if the cache didn't exist before as we have no idea what is new and what isn't
-			OutstandingChanges.Add(FUpdateCacheTransaction(Filename, FFileChangeData::FCA_Added, FilenameAndData.Value.Timestamp));
+			LocalDirtyFiles.Add(Filename);
 		}
 	}
 
-	// Delete anything that doesn't exist on disk anymore
+	// Check for anything that doesn't exist on disk anymore
 	for (auto It = CachedDirectoryState.Files.CreateIterator(); It; ++It)
 	{
 		const FImmutableString& Filename = It.Key();
-		if (!CurrentFileData->Files.Contains(Filename) && IsFileApplicable(Filename.Get()))
+		if (!LiveState->Files.Contains(Filename))
 		{
-			OutstandingChanges.Add(FUpdateCacheTransaction(Filename, FFileChangeData::FCA_Removed));
+			LocalDirtyFiles.Add(Filename);
 		}
 	}
-}
-
-void FFileCache::OnDirectoryChanged(const TArray<FFileChangeData>& FileChanges)
-{
-	auto& FileManager = IFileManager::Get();
-
-	OutstandingChanges.Reserve(OutstandingChanges.Num() + FileChanges.Num());
 
-	for (const auto& ThisEntry : FileChanges)
+	if (LocalDirtyFiles.Num() != 0)
 	{
-		auto Transaction = ProcessChange(ThisEntry.Filename, ThisEntry.Action, FileManager);
-		if (Transaction.IsSet())
-		{
-			OutstandingChanges.Add(MoveTemp(Transaction.GetValue()));
-		}
-	}
-
-	Utils::RemoveDuplicates(OutstandingChanges, [](const FUpdateCacheTransaction& A, const FUpdateCacheTransaction& B){
-		return A.Action == B.Action && A.Filename == B.Filename;
-	});
-}
-
-TOptional<FUpdateCacheTransaction> FFileCache::ProcessChange(const FString& Filename, FFileChangeData::EFileChangeAction Action, IFileManager& FileManager)
-{
-	// If it's a directory or is not applicable, ignore it
-	if (FileManager.DirectoryExists(*Filename) || !IsFileApplicable(Filename))
-	{
-		return TOptional<FUpdateCacheTransaction>();
-	}
+		TArray<FUpdateCacheTransaction> TempTransactions;
+		DiffDirtyFiles(LocalDirtyFiles, TempTransactions, &LiveState.GetValue());
 
-	FImmutableString TransactionFilename;
-	{
-		FString Temp = FPaths::ConvertRelativePathToFull(Filename);
-		if (Config.PathType == EPathType::Relative)
-		{
-			TransactionFilename = FString(*Temp + Config.Directory.Len());
-		}
-		else
-		{
-			TransactionFilename = MoveTemp(Temp);
-		}
-	}
-	
-	switch (Action)
-	{
-	case FFileChangeData::FCA_Added:
+		// Any file changes that we've detected that are not relevant to the previously cached state, we just add to the cache immediately (without telling the user)
+		// This is because we have no idea whether they are new changes or not, because they were not previously cached.
+		for (auto& Transaction : TempTransactions)
 		{
-			// Collapse Remove -> Add into a modification
-			const int32 NumRemoved = OutstandingChanges.RemoveAll([&](const FUpdateCacheTransaction& X){
-				return X.Action == FFileChangeData::FCA_Removed && X.Filename == TransactionFilename;
-			});
-
-			const auto FileChangeAction = NumRemoved == 0 ? FFileChangeData::FCA_Added : FFileChangeData::FCA_Modified;
-
-			// If it's a modification (remove then add) or the file doesn't already exist in the cache, we can create a transaction
-			if (FileChangeAction == FFileChangeData::FCA_Modified || !CachedDirectoryState.Files.Find(TransactionFilename))
+			// Generate paths that are relative to the watch directory to pass to the match rules
+			const TCHAR* TransactionFilename = *Transaction.Filename.Get();
+			const TCHAR* MovedFromFilename = *Transaction.MovedFromFilename.Get();
+			if (Config.PathType == EPathType::Absolute)
 			{
-				return FUpdateCacheTransaction(TransactionFilename, FileChangeAction, FileManager.GetTimeStamp(*Filename));
-			}
-		}
-		break;
-
-	case FFileChangeData::FCA_Removed:
-		{
-			bool bPreviouslyAdded = false;
-			const int32 NumRemoved = OutstandingChanges.RemoveAll([&](const FUpdateCacheTransaction& X){
-				if (X.Filename == TransactionFilename)
+				if (Transaction.Filename.Get().Len() > Config.Directory.Len())
 				{
-					bPreviouslyAdded = X.Action == FFileChangeData::FCA_Added || bPreviouslyAdded;
-					return true;
+					TransactionFilename += Config.Directory.Len();
 				}
-				return false;
-			});
 
-			// We only create the transaction if there is net change, and the file actually exists in the cache
-			if (!bPreviouslyAdded && CachedDirectoryState.Files.Find(TransactionFilename))
-			{
-				return FUpdateCacheTransaction(TransactionFilename, FFileChangeData::FCA_Removed);
+				if (Transaction.MovedFromFilename.Get().Len() > Config.Directory.Len())
+				{
+					MovedFromFilename += Config.Directory.Len();
+				}
 			}
-		}
-		
-		break;
 
-	case FFileChangeData::FCA_Modified:
-		{
-			const bool bPreviouslyAdded = OutstandingChanges.ContainsByPredicate([&](const FUpdateCacheTransaction& X){
-				return X.Filename == TransactionFilename && X.Action == FFileChangeData::FCA_Added;
-			});
+			const bool bAppliesToOldCache = CachedDirectoryState.Rules.IsFileApplicable(TransactionFilename) &&
+				(Transaction.Action != EFileAction::Moved || CachedDirectoryState.Rules.IsFileApplicable(MovedFromFilename));
 
-			const FDateTime& Timestamp = FileManager.GetTimeStamp(*Filename);
+			const bool bAppliesToNewCache = LiveState->Rules.IsFileApplicable(TransactionFilename) &&
+				(Transaction.Action != EFileAction::Moved || LiveState->Rules.IsFileApplicable(MovedFromFilename));
 
-			// Don't report subsequent modifications if an addition is still waiting to be handled
-			if (!bPreviouslyAdded)
+			if (bAppliesToOldCache && bAppliesToNewCache)
 			{
-				auto* ExistingFileData = CachedDirectoryState.Files.Find(TransactionFilename);
-				if (ExistingFileData && ExistingFileData->Timestamp < Timestamp)
+				// Just throw away the transaction
+				continue;
+			}
+			else
+			{
+				// Remove it from our dirty files
+				LocalDirtyFiles.Remove(Transaction.Filename);
+				if (Transaction.Action == EFileAction::Moved)
 				{
-					return FUpdateCacheTransaction(TransactionFilename, FFileChangeData::FCA_Modified, Timestamp);
+					LocalDirtyFiles.Remove(Transaction.MovedFromFilename);
 				}
+
+				// Update the cache immediately and remove it from the pending transactions
+				CompleteTransaction(MoveTemp(Transaction));
 			}
 		}
-		break;
-
-	default:
-		break;
 	}
 
-	return TOptional<FUpdateCacheTransaction>();
-}
+	// Now the only files left in LocalDirtyFiles are things that definitely apply to the current cache, and used to apply to the old cache
+	DirtyFiles = DirtyFiles.Union(LocalDirtyFiles);
 
-bool FFileCache::IsFileApplicable(const FString& Filename)
-{
-	return (Config.ExcludeExtensions.IsEmpty() || !MatchExtensionString(Config.ExcludeExtensions, Filename)) &&
-		   (Config.IncludeExtensions.IsEmpty() ||  MatchExtensionString(Config.IncludeExtensions, Filename));
+	// Update the applicable extensions now that we've updated the cache
+	CachedDirectoryState.Rules = LiveState->Rules;
 }
 
-bool FFileCache::MatchExtensionString(const FString& InSource, const FString& InPath)
+void FFileCache::OnDirectoryChanged(const TArray<FFileChangeData>& FileChanges)
 {
-	int32 SlashPosition = InPath.Find(TEXT("/"), ESearchCase::CaseSensitive, ESearchDir::FromEnd);
-	// in case we are using backslashes on a platform that doesn't use backslashes
-	SlashPosition = FMath::Max(SlashPosition, InPath.Find(TEXT("\\"), ESearchCase::CaseSensitive, ESearchDir::FromEnd));
-
-	const int32 DotPosition = InPath.Find(TEXT("."), ESearchCase::CaseSensitive, ESearchDir::FromEnd);
-
-	if (DotPosition > INDEX_NONE && DotPosition > SlashPosition && DotPosition < InPath.Len() - 1)
+	for (const auto& ThisEntry : FileChanges)
 	{
-		FString Search;
-		Search.Reserve(InPath.Len() - DotPosition + 2);;
-		Search.InsertAt(0, ';');
-		Search += &InPath[DotPosition + 1];
-		Search.AppendChar(';');
-
-		return InSource.Find(Search, ESearchCase::IgnoreCase) != INDEX_NONE;
+		auto TransactionPath = GetTransactionPath(ThisEntry.Filename);
+		if (TransactionPath.IsSet())
+		{
+			DirtyFiles.Add(MoveTemp(TransactionPath.GetValue()));
+		}
 	}
-
-	return false;
 }
\ No newline at end of file
diff --git a/Engine/Source/Editor/UnrealEd/Private/AutoReimport/FileCache.h b/Engine/Source/Editor/UnrealEd/Private/AutoReimport/FileCache.h
--- a/Engine/Source/Editor/UnrealEd/Private/AutoReimport/FileCache.h
+++ b/Engine/Source/Editor/UnrealEd/Private/AutoReimport/FileCache.h
@@ -10,15 +10,25 @@
 struct UNREALED_API FFileCacheCustomVersion
 {
 	static const FGuid Key;
-	enum Type {	Initial, Latest = Initial };
+	enum Type {	Initial, IncludeFileHash, Latest = IncludeFileHash };
 };
 
 /** Structure representing specific information about a particular file */
 struct FFileData
 {
 	/** Constructiors */
-	FFileData() {}
-	FFileData(const FDateTime& InTimestamp) : Timestamp(InTimestamp) {}
+	FFileData() : Timestamp(0) {}
+	FFileData(const FDateTime& InTimestamp, const FMD5Hash& InFileHash) : Timestamp(InTimestamp), FileHash(InFileHash) {}
+
+	friend bool operator==(const FFileData& LHS, const FFileData& RHS)
+	{
+		return LHS.Timestamp == RHS.Timestamp && LHS.FileHash == RHS.FileHash;
+	}
+
+	friend bool operator!=(const FFileData& LHS, const FFileData& RHS)
+	{
+		return LHS.Timestamp != RHS.Timestamp || LHS.FileHash != RHS.FileHash;	
+	}
 
 	/** Serializer for this type */
 	friend FArchive& operator<<(FArchive& Ar, FFileData& Data)
@@ -28,10 +38,19 @@ struct FFileData
 			Ar << Data.Timestamp;
 		}
 
+		if (Ar.CustomVer(FFileCacheCustomVersion::Key) >= FFileCacheCustomVersion::IncludeFileHash)
+		{
+			Ar << Data.FileHash;
+		}
+
 		return Ar;
 	}
 
+	/** The cached timestamp of the file on disk */
 	FDateTime Timestamp;
+
+	/** The cached MD5 hash of the file on disk */
+	FMD5Hash FileHash;
 };
 
 /** Structure representing the file data for a number of files in a directory */
@@ -41,8 +60,11 @@ struct FDirectoryState
 	FDirectoryState() {}
 
 	/** Move construction */
-	FDirectoryState(FDirectoryState&& In) : Files(MoveTemp(In.Files)) {}
-	FDirectoryState& operator=(FDirectoryState&& In) { Swap(Files, In.Files); return *this; }
+	FDirectoryState(FDirectoryState&& In) : Rules(MoveTemp(In.Rules)), Files(MoveTemp(In.Files)) {}
+	FDirectoryState& operator=(FDirectoryState&& In) { Swap(Rules, In.Rules); Swap(Files, In.Files); return *this; }
+
+	/** The rules that define what this state applies to */
+	FMatchRules Rules;
 
 	/** Filename -> data map */
 	TMap<FImmutableString, FFileData> Files;
@@ -52,8 +74,11 @@ struct FDirectoryState
 	{
 		Ar.UsingCustomVersion(FFileCacheCustomVersion::Key);
 
-		if (Ar.CustomVer(FFileCacheCustomVersion::Key) >= FFileCacheCustomVersion::Initial)
+		// Ignore any old versions to data to ensure that we generate a new cache
+		if (Ar.CustomVer(FFileCacheCustomVersion::Key) >= FFileCacheCustomVersion::IncludeFileHash)
 		{
+			Ar << State.Rules;
+
 			// Number of files
 			int32 Num = State.Files.Num();
 			Ar << Num;
@@ -69,28 +94,41 @@ struct FDirectoryState
 	}
 };
 
+enum class EFileAction : uint8
+{
+	Added, Modified, Removed, Moved
+};
+
 /** A transaction issued by FFileCache to describe a change to the cache. The change is only committed once the transaction is returned to the cache (see FFileCache::CompleteTransaction). */
 struct FUpdateCacheTransaction
 {
 	/** The path of the file to which this transaction relates */
 	FImmutableString Filename;
 
-	/** The timestamp of the file at the time of this change */
-	FDateTime Timestamp;
+	/** In the case of a moved file, this represents the path the file was moved from */
+	FImmutableString MovedFromFilename;
+
+	/** File data pertaining to this change at the time of dispatch */
+	FFileData FileData;
 
 	/** The type of action that prompted this transaction */
-	FFileChangeData::EFileChangeAction Action;
+	EFileAction Action;
 
 	/** Publically moveable */
-	FUpdateCacheTransaction(FUpdateCacheTransaction&& In) : Filename(MoveTemp(In.Filename)), Timestamp(MoveTemp(In.Timestamp)), Action(MoveTemp(In.Action)) {}
-	FUpdateCacheTransaction& operator=(FUpdateCacheTransaction&& In) { Swap(Filename, In.Filename); Swap(Timestamp, In.Timestamp); Swap(Action, In.Action); return *this; }
+	FUpdateCacheTransaction(FUpdateCacheTransaction&& In) : Filename(MoveTemp(In.Filename)), MovedFromFilename(MoveTemp(In.MovedFromFilename)), FileData(MoveTemp(In.FileData)), Action(MoveTemp(In.Action)) {}
+	FUpdateCacheTransaction& operator=(FUpdateCacheTransaction&& In) { Swap(Filename, In.Filename); Swap(MovedFromFilename, In.MovedFromFilename); Swap(FileData, In.FileData); Swap(Action, In.Action); return *this; }
 
 private:
 	friend class FFileCache;
 
 	/** Construction responsibility is held by FFileCache */
-	FUpdateCacheTransaction(FImmutableString InFilename, FFileChangeData::EFileChangeAction InAction, const FDateTime& InTimestamp = 0)
-		: Filename(MoveTemp(InFilename)), Timestamp(InTimestamp), Action(InAction)
+	FUpdateCacheTransaction(FImmutableString InFilename, EFileAction InAction, const FFileData& InFileData = FFileData())
+		: Filename(MoveTemp(InFilename)), FileData(InFileData), Action(InAction)
+	{}
+
+	/** Construction responsibility is held by FFileCache */
+	FUpdateCacheTransaction(FImmutableString InMovedFromFilename, FImmutableString InMovedToFilename, const FFileData& InFileData)
+		: Filename(MoveTemp(InMovedToFilename)), MovedFromFilename(MoveTemp(InMovedFromFilename)), FileData(InFileData), Action(EFileAction::Moved)
 	{}
 	
 	/** Not Copyable */
@@ -120,7 +158,7 @@ enum EPathType
  *		}
  *		TOptional<FDirectoryState> State = Reader.GetFinalState();
  */
-struct FAsyncDirectoryReader
+struct FAsyncDirectoryReader : public TSharedFromThis<FAsyncDirectoryReader, ESPMode::ThreadSafe>
 {
 	enum class EProgressResult
 	{
@@ -130,11 +168,29 @@ struct FAsyncDirectoryReader
 	/** Constructor that sets up the directory reader to the specified directory */
 	FAsyncDirectoryReader(const FString& InDirectory, EPathType InPathType);
 
+	/** Set what files are relevant to this reader. Calling this once the reader starts results in undefined behaviour. */
+	void SetMatchRules(const FMatchRules& InRules)
+	{
+		if (LiveState.IsSet())
+		{
+			LiveState->Rules = InRules;
+		}
+	}
+
 	/**
 	 * Get the state of the directory once finished. Relinquishes the currently stored directory state to the client.
-	 * Returns nothing if incomplete, or if GetFinalState() has already been called.
+	 * Returns nothing if incomplete, or if GetLiveState() has already been called.
 	 */
-	TOptional<FDirectoryState> GetFinalState();
+	TOptional<FDirectoryState> GetLiveState();
+
+	/** Retrieve the cached state supplied to this class through UseCachedState(). */
+	TOptional<FDirectoryState> GetCachedState();
+
+	/** Instruct the directory reader to use the specified cached state to lookup file hashes, where timestamps haven't changed */
+	void UseCachedState(FDirectoryState InCachedState)
+	{
+		CachedState = MoveTemp(InCachedState);
+	}
 
 	/** Returns true when this directory reader has finished scanning the directory */
 	bool IsComplete() const;
@@ -152,37 +208,49 @@ private:
 	/** Whether we should return relative or absolute paths */
 	EPathType PathType;
 
-	/** The currently discovered state of the directory */
-	TOptional<FDirectoryState> State;
+	/** The currently discovered state of the directory - reset once relinquished to the client through GetLiveState */
+	TOptional<FDirectoryState> LiveState;
+
+	/** The previously cached  state of the directory, optional */
+	TOptional<FDirectoryState> CachedState;
 
 	/** A list of directories we have recursively found on our travels */
 	TArray<FString> PendingDirectories;
 
 	/** A list of files we have recursively found on our travels */
 	TArray<FString> PendingFiles;
+
+	/** Thread safe flag to signify when this class has finished reading */
+	FThreadSafeBool bIsComplete;
+
+	/** The time we started scanning the directory */
+	double StartTime;
+
+	/** Friendship required for the reader thread to flag when finished */
+	friend struct FAsyncDirectoryReaderThread;
 };
 
 /** Configuration structure required to construct a FFileCache */
 struct FFileCacheConfig
 {
 	FFileCacheConfig(FString InDirectory, FString InCacheFile)
-		: Directory(InDirectory), CacheFile(InCacheFile), PathType(EPathType::Relative)
+		: Directory(InDirectory), CacheFile(InCacheFile), PathType(EPathType::Relative), bDetectChangesSinceLastRun(false)
 	{}
 
 	/** String specifying the directory on disk that the cache should reflect */
 	FString Directory;
 
-	/** String specifying the file that the cache should be saved to */
+	/** String specifying the file that the cache should be saved to. When empty, no cache file will be maintained (thus only an in-memory cache is used) */
 	FString CacheFile;
 
-	/**
-	 * Strings specifying extensions to include/exclude in the cache (or empty to include everything).
-	 * Specified as a list of semi-colon separated extensions (eg "png;jpg;gif;").
-	 */
-	FString IncludeExtensions, ExcludeExtensions;
+	/** List of rules which define what we will be watching */
+	FMatchRules Rules;
 
 	/** Path type to return, relative to the directory or absolute. */
 	EPathType PathType;
+
+	/** When true, changes to the directory since the cache shutdown will be detected and reported. When false, said changes will silently be applied to the serialized cache. */
+	bool bDetectChangesSinceLastRun;
 };
 
 /**
@@ -209,8 +277,17 @@ public:
 	/** Get the absolute path of the directory this cache reflects */
 	const FString& GetDirectory() const { return Config.Directory; }
 
-	/** Tick this FileCache, optionally specifying a limiter to limit the amount of time spent in here */
-	void Tick(const FTimeLimit& Limit);
+	/** Check whether this file cache has finished starting up yet */
+	bool HasStartedUp() const { return !DirectoryReader.IsValid() || DirectoryReader->IsComplete(); }
+
+	/** Attempt to locate file data pertaining to the specified filename.
+	 *	@param		InFilename		The filename to find data for - either relative to the directory, or absolute, depending on Config.PathType.
+	 *	@return		The current cached file data pertaining to the specified filename, or nullptr. May not be completely up to date if there are outstanding transactions.
+	 */
+	const FFileData* FindFileData(FImmutableString InFilename) const;
+
+	/** Tick this FileCache */
+	void Tick();
 
 	/** Write out the cached file, if we have any changes to write */
 	void WriteCache();
@@ -222,10 +299,13 @@ public:
 	void CompleteTransaction(FUpdateCacheTransaction&& Transaction);
 
 	/** Report an external change to the manager, such that a subsequent equal change reported by the os be ignored */
-	void ReportExternalChange(const FString& Filename, FFileChangeData::EFileChangeAction Action);
+	void IgnoreNewFile(const FString& Filename);
+	void IgnoreFileModification(const FString& Filename);
+	void IgnoreMovedFile(const FString& SrcFilename, const FString& DstFilename);
+	void IgnoreDeletedFile(const FString& Filename);
 
 	/** Get the number of pending changes to the cache. */
-	int32 GetNumOutstandingChanges() const { return OutstandingChanges.Num(); }
+	int32 GetNumOutstandingChanges() const { return DirtyFiles.Num() + PendingTransactions.Num(); }
 
 	/** Get all pending changes to the cache. Transactions must be returned to CompleteTransaction to update the cache. */
 	TArray<FUpdateCacheTransaction> GetOutstandingChanges();
@@ -235,17 +315,16 @@ private:
 	/** Called when the directory we are monitoring has been changed in some way */
 	void OnDirectoryChanged(const TArray<FFileChangeData>& FileChanges);
 
-	/** Process a change, creating an outstanding change transaction if necessary */
-	TOptional<FUpdateCacheTransaction> ProcessChange(const FString& Filename, FFileChangeData::EFileChangeAction Action, IFileManager& FileManager);
+	/** Diff the specified set of dirty files (absolute paths, or relative to the monitor directory), adding transactions to the specified array if necessary.
+	 *	Optionally takes a directory state from which we can retrieve current file system state, without having to ask the FS directly.
+	 */
+	void DiffDirtyFiles(const TSet<FImmutableString>& InOutDirtyFiles, TArray<FUpdateCacheTransaction>& InOutTransactions, const FDirectoryState* InFileSystemState = nullptr) const;
 
-	/** Check whether a file is applicable to our cache */
-	bool IsFileApplicable(const FString& InPath);
+	/** Get the absolute path from a transaction filename */
+	FString GetAbsolutePath(const FString& InTransactionPath) const;
 
-	/**
-	 * Check whether a filename's extension matches the specified extension list.
-	 * InExtensions must be of the form ";ext1;ext2;ext3;". Note the leading and trailing ;
-	 */
-	static bool MatchExtensionString(const FString& InExtensions, const FString& InFilename);
+	/** Get the transaction path from an absolute filename. Returns nothing if the absolute path is invalid (not under the correct folder, or not applicable) */
+	TOptional<FString> GetTransactionPath(const FString& InAbsolutePath) const;
 
 	/** Unbind the watcher from the directory. Called on destruction. */
 	void UnbindWatcher();
@@ -259,10 +338,13 @@ private:
 	FFileCacheConfig Config;
 
 	/** 'Asynchronous' directory reader responsible for gathering all file/timestamp information recursively from our cache directory */
-	FAsyncDirectoryReader DirectoryReader;
+	TSharedPtr<FAsyncDirectoryReader, ESPMode::ThreadSafe> DirectoryReader;
+
+	/** A set of dirty files that we will use to report changes to the user */
+	TSet<FImmutableString> DirtyFiles;
 
-	/** Temporary array of transactions to be returned to the client through GetOutstandingChanges() */
-	TArray<FUpdateCacheTransaction> OutstandingChanges;
+	/** A list of pending change transactions to be returned to the user */
+	TArray<FUpdateCacheTransaction> PendingTransactions;
 
 	/** Our in-memory view of the cached directory state. */
 	FDirectoryState CachedDirectoryState;
diff --git a/Engine/Source/Editor/UnrealEd/Private/AutoReimport/ReimportFeedbackContext.cpp b/Engine/Source/Editor/UnrealEd/Private/AutoReimport/ReimportFeedbackContext.cpp
--- a/Engine/Source/Editor/UnrealEd/Private/AutoReimport/ReimportFeedbackContext.cpp
+++ b/Engine/Source/Editor/UnrealEd/Private/AutoReimport/ReimportFeedbackContext.cpp
@@ -288,7 +288,8 @@ void SReimportFeedback::Construct(const FArguments& InArgs, FText InMainText)
 }
 
 FReimportFeedbackContext::FReimportFeedbackContext()
-	: MessageLog("AssetReimport")
+	: bSuppressSlowTaskMessages(false)
+	, MessageLog("AssetReimport")
 {}
 
 void FReimportFeedbackContext::Initialize(TSharedRef<SReimportFeedback> Widget)
@@ -331,7 +332,7 @@ void FReimportFeedbackContext::StartSlowTask(const FText& Task, bool bShowCancel
 {
 	FFeedbackContext::StartSlowTask(Task, bShowCancelButton);
 
-	if (!Task.IsEmpty())
+	if (!bSuppressSlowTaskMessages && !Task.IsEmpty())
 	{
 		NotificationContent->Add(SNew(STextBlock).Text(Task));
 	}
diff --git a/Engine/Source/Editor/UnrealEd/Private/AutoReimport/ReimportFeedbackContext.h b/Engine/Source/Editor/UnrealEd/Private/AutoReimport/ReimportFeedbackContext.h
--- a/Engine/Source/Editor/UnrealEd/Private/AutoReimport/ReimportFeedbackContext.h
+++ b/Engine/Source/Editor/UnrealEd/Private/AutoReimport/ReimportFeedbackContext.h
@@ -61,6 +61,9 @@ public:
 
 	TSharedPtr<SReimportFeedback> GetContent() { return NotificationContent; }
 
+	/** Get the message log that this context is using */
+	FMessageLog& GetMessageLog() { return MessageLog; }
+
 	/** Add a message to the context (and message log) */
 	void AddMessage(EMessageSeverity::Type Severity, const FText& Message);
 
@@ -73,13 +76,13 @@ public:
 	virtual TSharedRef<SWidget> AsWidget() override { return NotificationContent.ToSharedRef(); };
 	virtual void StartSlowTask( const FText& Task, bool bShowCancelButton=false) override;
 
+	/** True if we will suppress slow task messages from being shown on the UI, false if otherwise */
+	bool bSuppressSlowTaskMessages;
+
 private:
 	/** Message log for output of errors and log messages */
 	FMessageLog MessageLog;
 
-	/** True if we need to notifiy about new messages on the message log */
-	bool NotifyMessageLog;
-
 	/** The notification that is shown when the context is active */
 	TSharedPtr<class SNotificationItem> Notification;
 
diff --git a/Engine/Source/Editor/UnrealEd/Private/Settings/SettingsClasses.cpp b/Engine/Source/Editor/UnrealEd/Private/Settings/SettingsClasses.cpp
--- a/Engine/Source/Editor/UnrealEd/Private/Settings/SettingsClasses.cpp
+++ b/Engine/Source/Editor/UnrealEd/Private/Settings/SettingsClasses.cpp
@@ -5,6 +5,7 @@
 #include "Components/BillboardComponent.h"
 #include "AI/Navigation/NavigationSystem.h"
 #include "Components/ArrowComponent.h"
+#include "AutoReimport/AutoReimportUtilities.h"
 
 /* UContentBrowserSettings interface
  *****************************************************************************/
@@ -96,10 +97,14 @@ UEditorLoadingSavingSettings::UEditorLoadingSavingSettings( const FObjectInitial
 	, bMonitorContentDirectories(false)
 	, bAutoCreateAssets(true)
 	, bAutoDeleteAssets(true)
+	, bDetectChangesOnRestart(false)
 	, bDeleteSourceFilesWithAssets(false)
 {
 	TextDiffToolPath.FilePath = TEXT("P4Merge.exe");
-	AutoReimportDirectories.Add("/Game/");
+
+	FAutoReimportDirectoryConfig Default;
+	Default.SourceDirectory = TEXT("/Game/");
+	AutoReimportDirectorySettings.Add(Default);
 }
 
 
@@ -113,7 +118,8 @@ void UEditorLoadingSavingSettings::PostEditChangeProperty( struct FPropertyChang
 {
 	Super::PostEditChangeProperty(PropertyChangedEvent);
 
-	const FName Name = (PropertyChangedEvent.Property != nullptr) ? PropertyChangedEvent.Property->GetFName() : NAME_None;
+	// Use MemberProperty here so we report the correct member name for nested changes
+	const FName Name = (PropertyChangedEvent.MemberProperty != nullptr) ? PropertyChangedEvent.MemberProperty->GetFName() : NAME_None;
 
 	if (Name == FName(TEXT("bSCCUseGlobalSettings")))
 	{
@@ -129,6 +135,93 @@ void UEditorLoadingSavingSettings::PostEditChangeProperty( struct FPropertyChang
 	SettingChangedEvent.Broadcast(Name);
 }
 
+void UEditorLoadingSavingSettings::PostInitProperties()
+{
+	for (const auto& String : AutoReimportDirectories_DEPRECATED)
+	{
+		FAutoReimportDirectoryConfig Config;
+		Config.SourceDirectory = String;
+		AutoReimportDirectorySettings.Add(Config);
+	}
+	Super::PostInitProperties();
+}
+
+FAutoReimportDirectoryConfig::FParseContext::FParseContext(bool bInEnableLogging)
+	: bEnableLogging(bInEnableLogging)
+{
+	TArray<FString> RootContentPaths;
+	FPackageName::QueryRootContentPaths( RootContentPaths );
+	for (FString& RootPath : RootContentPaths)
+	{
+		FString ContentFolder = FPaths::ConvertRelativePathToFull(FPackageName::LongPackageNameToFilename(RootPath));
+		MountedPaths.Add( TPairInitializer<FString, FString>(MoveTemp(ContentFolder), MoveTemp(RootPath)) );
+	}
+}
+
+bool FAutoReimportDirectoryConfig::ParseSourceDirectoryAndMountPoint(FString& SourceDirectory, FString& MountPoint, const FParseContext& InContext)
+{
+	SourceDirectory.ReplaceInline(TEXT("\\"), TEXT("/"));
+
+	// Check if the source directory is actually a mount point
+	if (!FPackageName::GetPackageMountPoint(SourceDirectory).IsNone())
+	{
+		MountPoint = SourceDirectory;
+		SourceDirectory = FString();
+	}
+
+	if (!SourceDirectory.IsEmpty() && !MountPoint.IsEmpty())
+	{
+		// We have both a source directory and a mount point. Verify that the source dir exists, and that the mount point is valid.
+		if (!IFileManager::Get().DirectoryExists(*SourceDirectory))
+		{
+			UE_CLOG(InContext.bEnableLogging, LogAutoReimportManager, Warning, TEXT("Unable to watch directory %s as it doesn't exist."), *SourceDirectory);
+			return false;
+		}
+
+		if (FPackageName::GetPackageMountPoint(MountPoint).IsNone())
+		{
+			UE_CLOG(InContext.bEnableLogging, LogAutoReimportManager, Warning, TEXT("Unable to setup directory %s to map to %s, as it's not a valid mounted path. Continuing without mounted path (auto reimports will still work, but auto add won't)."), *SourceDirectory, *MountPoint);
+		}
+	}
+	else if(!MountPoint.IsEmpty())
+	{
+		// We have just a mount point - validate it, and find its source directory
+		if (FPackageName::GetPackageMountPoint(MountPoint).IsNone())
+		{
+			UE_CLOG(InContext.bEnableLogging, LogAutoReimportManager, Warning, TEXT("Unable to setup directory monitor for %s, as it's not a valid mounted path."), *MountPoint);
+			return false;
+		}
+
+		SourceDirectory = FPackageName::LongPackageNameToFilename(MountPoint);
+	}
+	else if(!SourceDirectory.IsEmpty())
+	{
+		// We have just a source directory - verify whether it's a mounted path, and set up the mount point if so
+		if (!IFileManager::Get().DirectoryExists(*SourceDirectory))
+		{
+			UE_CLOG(InContext.bEnableLogging, LogAutoReimportManager, Warning, TEXT("Unable to watch directory %s as it doesn't exist."), *SourceDirectory);
+			return false;
+		}
+
+		// Set the mounted path if necessary
+		auto* Pair = InContext.MountedPaths.FindByPredicate([&](const TPair<FString, FString>& Pair){
+			return SourceDirectory.StartsWith(Pair.Key);
+		});
+
+		if (Pair)
+		{
+			MountPoint = Pair->Value / SourceDirectory.RightChop(Pair->Key.Len());
+			MountPoint.ReplaceInline(TEXT("\\"), TEXT("/"));
+		}
+	}
+	else
+	{
+		// Don't have any valid settings
+		return false;
+	}
+
+	return true;
+}
 
 /* UEditorMiscSettings interface
  *****************************************************************************/
diff --git a/Engine/Source/Editor/UnrealEd/Public/AutoReimport/AutoReimportManager.h b/Engine/Source/Editor/UnrealEd/Public/AutoReimport/AutoReimportManager.h
--- a/Engine/Source/Editor/UnrealEd/Public/AutoReimport/AutoReimportManager.h
+++ b/Engine/Source/Editor/UnrealEd/Public/AutoReimport/AutoReimportManager.h
@@ -34,7 +34,10 @@ public:
 	TArray<FPathAndMountPoint> GetMonitoredDirectories() const;
 
 	/** Report an external change to the manager, such that a subsequent equal change reported by the os be ignored */
-	void ReportExternalChange(const FString& Filename, FFileChangeData::EFileChangeAction Action);
+	void IgnoreNewFile(const FString& Filename);
+	void IgnoreFileModification(const FString& Filename);
+	void IgnoreMovedFile(const FString& SrcFilename, const FString& DstFilename);
+	void IgnoreDeletedFile(const FString& Filename);
 
 private:
 
