Disable prompting for checkout and save in unattended applications

diff --git a/Engine/Source/Editor/UnrealEd/Private/FileHelpers.cpp b/Engine/Source/Editor/UnrealEd/Private/FileHelpers.cpp
--- a/Engine/Source/Editor/UnrealEd/Private/FileHelpers.cpp
+++ b/Engine/Source/Editor/UnrealEd/Private/FileHelpers.cpp
@@ -2759,7 +2759,7 @@ bool FEditorFileUtils::SaveCurrentLevel()
 FEditorFileUtils::EPromptReturnCode FEditorFileUtils::PromptForCheckoutAndSave( const TArray<UPackage*>& InPackages, bool bCheckDirty, bool bPromptToSave, TArray<UPackage*>* OutFailedPackages, bool bAlreadyCheckedOut, bool bCanBeDeclined )
 {
 	// Check for re-entrance into this function
-	if ( bIsPromptingForCheckoutAndSave )
+	if ( bIsPromptingForCheckoutAndSave || FApp::IsUnattended() )
 	{
 		return PR_Cancelled;
 	}
