Allow pack files to specify additional loose files to copy in as part of the import step

diff --git a/Engine/Source/Editor/UnrealEd/Private/Factories/PackFactory.cpp b/Engine/Source/Editor/UnrealEd/Private/Factories/PackFactory.cpp
--- a/Engine/Source/Editor/UnrealEd/Private/Factories/PackFactory.cpp
+++ b/Engine/Source/Editor/UnrealEd/Private/Factories/PackFactory.cpp
@@ -132,9 +132,15 @@ namespace PackFactoryHelper
 		MemReader << FileContents;
 	}
 
+	struct FPackConfigParameters
+	{
+		FString GameName;
+		TArray<FString> AdditionalContentToAdd;
+	};
+
 	// Takes a string that represents the contents of a config file and sets up the supported config properties based on it
 	// Currently we support Action and Axis Mappings and a GameName (for setting up redirects)
-	void ProcessPackConfig(const FString& ConfigString, TMap<FString, FString>& ConfigParameters)
+	void ProcessPackConfig(const FString& ConfigString, FPackConfigParameters& ConfigParameters)
 	{
 		FConfigFile PackConfig;
 		PackConfig.ProcessInputFileContents(ConfigString);
@@ -228,7 +234,33 @@ namespace PackFactoryHelper
 		{	
 			if (FString* GameName = RedirectsSection->Find("GameName"))
 			{
-				ConfigParameters.Add("GameName", *GameName);
+				ConfigParameters.GameName = *GameName;
+			}
+		}
+
+		FConfigSection* AdditionalContentSection = PackConfig.Find("AdditionalContentToAdd");
+		if (AdditionalContentSection)
+		{
+			for (auto FilePair : *AdditionalContentSection)
+			{
+				if (FilePair.Key.ToString().Contains("Files"))
+				{
+					FString Filename = FPaths::GetCleanFilename(FilePair.Value);
+					FString Directory = FPaths::RootDir() / FPaths::GetPath(FilePair.Value);
+					FPaths::MakeStandardFilename(Directory);
+					FPakFile::MakeDirectoryFromPath(Directory);
+
+					if (Filename.Contains(TEXT("*")))
+					{
+						TArray<FString> FoundFiles;
+						IFileManager::Get().FindFilesRecursive(FoundFiles, *Directory, *Filename, true, false);
+						ConfigParameters.AdditionalContentToAdd.Append(FoundFiles);
+					}
+					else
+					{
+						ConfigParameters.AdditionalContentToAdd.Add(FilePair.Value);
+					}
+				}
 			}
 		}
 	}
@@ -255,7 +287,6 @@ UObject* UPackFactory::FactoryCreateBinary
 	if (PakFile.IsValid())
 	{
 		static FString ContentFolder(TEXT("/Content/"));
-		static FString SourceFolder(TEXT("/Source/"));
 		FString ContentDestinationRoot = FPaths::GameContentDir();
 
 		const int32 ChopIndex = PakFile.GetMountPoint().Find(ContentFolder);
@@ -272,7 +303,7 @@ UObject* UPackFactory::FactoryCreateBinary
 
 		bool bContainsSource = false;
 		FModuleContextInfo SourceModuleInfo;
-		TMap<FString, FString> ConfigParameters;
+		PackFactoryHelper::FPackConfigParameters ConfigParameters;
 
 		TArray<FString> WrittenFiles;
 		TArray<FString> WrittenSourceFiles;
@@ -337,7 +368,7 @@ UObject* UPackFactory::FactoryCreateBinary
 					SourceModuleInfo = ModuleInfo;
 
 					// Setup the game name redirect
-					if (FString* GameName = ConfigParameters.Find("GameName"))
+					if (!ConfigParameters.GameName.IsEmpty())
 					{
 						const FString EngineIniFilename = FPaths::ConvertRelativePathToFull(GetDefault<UEngine>()->GetDefaultConfigFilename());
 
@@ -352,7 +383,7 @@ UObject* UPackFactory::FactoryCreateBinary
 						}
 
 						const FString RedirectsSection(TEXT("/Script/Engine.Engine"));
-						const FString LongOldGameName = FString::Printf(TEXT("/Script/%s"), **GameName);
+						const FString LongOldGameName = FString::Printf(TEXT("/Script/%s"), *ConfigParameters.GameName);
 						const FString LongNewGameName = FString::Printf(TEXT("/Script/%s"), *ModuleInfo.ModuleName);
 						
 						FConfigCacheIni Config(EConfigCacheType::Temporary);
@@ -361,7 +392,7 @@ UObject* UPackFactory::FactoryCreateBinary
 						FConfigSection* PackageRedirects = Config.GetSectionPrivate(*RedirectsSection, true, false, EngineIniFilename);
 
 						PackageRedirects->Add(TEXT("+ActiveGameNameRedirects"), FString::Printf(TEXT("(OldGameName=\"%s\",NewGameName=\"%s\")"), *LongOldGameName, *LongNewGameName));
-						PackageRedirects->Add(TEXT("+ActiveGameNameRedirects"), FString::Printf(TEXT("(OldGameName=\"%s\",NewGameName=\"%s\")"), **GameName, *LongNewGameName));
+						PackageRedirects->Add(TEXT("+ActiveGameNameRedirects"), FString::Printf(TEXT("(OldGameName=\"%s\",NewGameName=\"%s\")"), *ConfigParameters.GameName, *LongNewGameName));
 
 						NewFile.UpdateSections(*EngineIniFilename, *RedirectsSection);
 
@@ -369,7 +400,7 @@ UObject* UPackFactory::FactoryCreateBinary
 						GConfig->LoadGlobalIniFile(FinalIniFileName, *RedirectsSection, NULL, NULL, true);
 
 						ULinkerLoad::AddGameNameRedirect(*LongOldGameName, *LongNewGameName);
-						ULinkerLoad::AddGameNameRedirect(**GameName, *LongNewGameName);
+						ULinkerLoad::AddGameNameRedirect(*ConfigParameters.GameName, *LongNewGameName);
 					}
 					break;
 				}
@@ -407,7 +438,7 @@ UObject* UPackFactory::FactoryCreateBinary
 					}
 					else 
 					{
-						const int32 SourceIndex = DestFilename.Find(SourceFolder);
+						const int32 SourceIndex = DestFilename.Find(TEXT("/Source/"));
 						if (SourceIndex != INDEX_NONE)
 						{
 							DestFilename = DestFilename.RightChop(SourceIndex + 8);
@@ -421,7 +452,7 @@ UObject* UPackFactory::FactoryCreateBinary
 					PackFactoryHelper::ExtractFileToString(Entry, PakReader, CopyBuffer, PersistentCompressionBuffer, SourceContents);
 
 					FGameProjectGenerationModule& GameProjectModule = FModuleManager::LoadModuleChecked<FGameProjectGenerationModule>(TEXT("GameProjectGeneration"));
-					FString StringToReplace = *ConfigParameters.Find("GameName");
+					FString StringToReplace = ConfigParameters.GameName;
 					StringToReplace += ".h";
 					SourceContents = SourceContents.Replace(*StringToReplace, *GameProjectModule.Get().DetermineModuleIncludePath(SourceModuleInfo, DestFilename), ESearchCase::CaseSensitive);
 
@@ -477,6 +508,48 @@ UObject* UPackFactory::FactoryCreateBinary
 
 		UE_LOG(LogPackFactory, Log, TEXT("Finished extracting %d files (including %d errors)."), FileCount, ErrorCount);
 
+		if (ConfigParameters.AdditionalContentToAdd.Num() > 0)
+		{
+			IPlatformFile& PlatformFile = FPlatformFileManager::Get().GetPlatformFile();
+
+			for (const FString& FileToCopy : ConfigParameters.AdditionalContentToAdd)
+			{
+				FString DestFilename = FileToCopy;
+				if (DestFilename.StartsWith(TEXT("Content/")))
+				{
+					DestFilename = DestFilename.RightChop(8);
+				}
+				else
+				{
+					const int32 ContentIndex = DestFilename.Find(ContentFolder);
+					if (ContentIndex != INDEX_NONE)
+					{
+						DestFilename = DestFilename.RightChop(ContentIndex + 9);
+					}
+				}
+				DestFilename = ContentDestinationRoot / DestFilename;
+
+				FString DestDirectory = FPaths::GetPath(DestFilename);
+
+				if (PlatformFile.CreateDirectoryTree(*DestDirectory))
+				{
+					if (PlatformFile.CopyFile(*DestFilename, *FileToCopy))
+					{
+						WrittenFiles.Add(DestFilename);
+						UE_LOG(LogPackFactory, Log, TEXT("Copied \"%s\" to \"%s\""), *FileToCopy, *DestFilename);
+					}
+					else
+					{
+						UE_LOG(LogPackFactory, Error, TEXT("Unable to copy file \"%s\" to \"%s\"."), *FileToCopy, *DestFilename);
+					}
+				}
+				else
+				{
+					UE_LOG(LogPackFactory, Error, TEXT("Unable to create directory \"%s\"."), *FileToCopy, *DestFilename);
+				}
+			}
+		}
+
 		if (WrittenFiles.Num() > 0)
 		{
 			// If we wrote out source files, kick off the hot reload process
