Support Source loose files from upacks as well

diff --git a/Engine/Source/Editor/UnrealEd/Private/Factories/PackFactory.cpp b/Engine/Source/Editor/UnrealEd/Private/Factories/PackFactory.cpp
--- a/Engine/Source/Editor/UnrealEd/Private/Factories/PackFactory.cpp
+++ b/Engine/Source/Editor/UnrealEd/Private/Factories/PackFactory.cpp
@@ -134,8 +134,14 @@ namespace PackFactoryHelper
 
 	struct FPackConfigParameters
 	{
+		FPackConfigParameters()
+			: bContainsSource(false)
+		{
+		}
+
+		bool bContainsSource;
 		FString GameName;
-		TArray<FString> AdditionalContentToAdd;
+		TArray<FString> AdditionalFilesToAdd;
 	};
 
 	// Takes a string that represents the contents of a config file and sets up the supported config properties based on it
@@ -238,10 +244,10 @@ namespace PackFactoryHelper
 			}
 		}
 
-		FConfigSection* AdditionalContentSection = PackConfig.Find("AdditionalContentToAdd");
-		if (AdditionalContentSection)
+		FConfigSection* AdditionalFilesSection = PackConfig.Find("AdditionalFilesToAdd");
+		if (AdditionalFilesSection)
 		{
-			for (auto FilePair : *AdditionalContentSection)
+			for (auto FilePair : *AdditionalFilesSection)
 			{
 				if (FilePair.Key.ToString().Contains("Files"))
 				{
@@ -254,11 +260,26 @@ namespace PackFactoryHelper
 					{
 						TArray<FString> FoundFiles;
 						IFileManager::Get().FindFilesRecursive(FoundFiles, *Directory, *Filename, true, false);
-						ConfigParameters.AdditionalContentToAdd.Append(FoundFiles);
+						ConfigParameters.AdditionalFilesToAdd.Append(FoundFiles);
+						if (!ConfigParameters.bContainsSource)
+						{
+							for (const FString& FoundFile : FoundFiles)
+							{
+								if (FoundFile.StartsWith(TEXT("Source/")) || FoundFile.Contains(TEXT("/Source/")))
+								{
+									ConfigParameters.bContainsSource = true;
+									break;
+								}
+							}
+						}
 					}
 					else
 					{
-						ConfigParameters.AdditionalContentToAdd.Add(FilePair.Value);
+						ConfigParameters.AdditionalFilesToAdd.Add(Directory / Filename);
+						if (!ConfigParameters.bContainsSource && (ConfigParameters.AdditionalFilesToAdd.Last().StartsWith(TEXT("Source/")) || ConfigParameters.AdditionalFilesToAdd.Last().Contains(TEXT("/Source/"))))
+						{
+							ConfigParameters.bContainsSource = true;
+						}
 					}
 				}
 			}
@@ -301,7 +322,6 @@ UObject* UPackFactory::FactoryCreateBinary
 		int32 ErrorCount = 0;
 		int32 FileCount = 0;
 
-		bool bContainsSource = false;
 		FModuleContextInfo SourceModuleInfo;
 		PackFactoryHelper::FPackConfigParameters ConfigParameters;
 
@@ -330,9 +350,9 @@ UObject* UPackFactory::FactoryCreateBinary
 					ErrorCount++;
 				}
 			}
-			else if (!bContainsSource && (It.Filename().StartsWith(TEXT("Source/")) || It.Filename().Contains(TEXT("/Source/"))))
+			else if (!ConfigParameters.bContainsSource && (It.Filename().StartsWith(TEXT("Source/")) || It.Filename().Contains(TEXT("/Source/"))))
 			{
-				bContainsSource = true;
+				ConfigParameters.bContainsSource = true;
 			}
 		}
 
@@ -340,7 +360,7 @@ UObject* UPackFactory::FactoryCreateBinary
 
 		// If we have source files, set up the project files if necessary and the game name redirects for blueprints saved with class
 		// references to the module name from the source template
-		if (bContainsSource)
+		if (ConfigParameters.bContainsSource)
 		{
 			FGameProjectGenerationModule& GameProjectModule = FModuleManager::LoadModuleChecked<FGameProjectGenerationModule>(TEXT("GameProjectGeneration"));
 			bProjectHadSourceFiles = GameProjectModule.Get().ProjectHasCodeFiles();
@@ -508,45 +528,94 @@ UObject* UPackFactory::FactoryCreateBinary
 
 		UE_LOG(LogPackFactory, Log, TEXT("Finished extracting %d files (including %d errors)."), FileCount, ErrorCount);
 
-		if (ConfigParameters.AdditionalContentToAdd.Num() > 0)
+		if (ConfigParameters.AdditionalFilesToAdd.Num() > 0)
 		{
 			IPlatformFile& PlatformFile = FPlatformFileManager::Get().GetPlatformFile();
 
-			for (const FString& FileToCopy : ConfigParameters.AdditionalContentToAdd)
+			for (const FString& FileToCopy : ConfigParameters.AdditionalFilesToAdd)
 			{
-				FString DestFilename = FileToCopy;
-				if (DestFilename.StartsWith(TEXT("Content/")))
+				if (FileToCopy.StartsWith(TEXT("Source/")) || FileToCopy.Contains(TEXT("/Source/")))
 				{
-					DestFilename = DestFilename.RightChop(8);
+					FString DestFilename = FileToCopy;
+					if (DestFilename.StartsWith(TEXT("Source/")))
+					{
+						DestFilename = DestFilename.RightChop(7);
+					}
+					else 
+					{
+						const int32 SourceIndex = DestFilename.Find(TEXT("/Source/"));
+						if (SourceIndex != INDEX_NONE)
+						{
+							DestFilename = DestFilename.RightChop(SourceIndex + 8);
+						}
+					}
+					DestFilename = SourceModuleInfo.ModuleSourcePath / DestFilename;
+
+					FString DestDirectory = FPaths::GetPath(DestFilename);
+
+					if (PlatformFile.CreateDirectoryTree(*DestDirectory))
+					{
+						FString SourceContents;
+						if (FFileHelper::LoadFileToString(SourceContents, *FileToCopy))
+						{
+							FGameProjectGenerationModule& GameProjectModule = FModuleManager::LoadModuleChecked<FGameProjectGenerationModule>(TEXT("GameProjectGeneration"));
+							FString StringToReplace = ConfigParameters.GameName;
+							StringToReplace += ".h";
+							SourceContents = SourceContents.Replace(*StringToReplace, *GameProjectModule.Get().DetermineModuleIncludePath(SourceModuleInfo, DestFilename), ESearchCase::CaseSensitive);
+
+							if (FFileHelper::SaveStringToFile(SourceContents, *DestFilename))
+							{
+								WrittenFiles.Add(*DestFilename);
+								WrittenSourceFiles.Add(*DestFilename);
+							}
+							else
+							{
+								UE_LOG(LogPackFactory, Error, TEXT("Unable to write file \"%s\"."), *DestFilename);
+								++ErrorCount;
+							}
+						}
+						else
+						{
+							UE_LOG(LogPackFactory, Error, TEXT("Unable to read file \"%s\"."), *FileToCopy);
+						}
+					}
 				}
 				else
 				{
-					const int32 ContentIndex = DestFilename.Find(ContentFolder);
-					if (ContentIndex != INDEX_NONE)
+					FString DestFilename = FileToCopy;
+					if (DestFilename.StartsWith(TEXT("Content/")))
 					{
-						DestFilename = DestFilename.RightChop(ContentIndex + 9);
+						DestFilename = DestFilename.RightChop(8);
 					}
-				}
-				DestFilename = ContentDestinationRoot / DestFilename;
+					else
+					{
+						const int32 ContentIndex = DestFilename.Find(ContentFolder);
+						if (ContentIndex != INDEX_NONE)
+						{
+							DestFilename = DestFilename.RightChop(ContentIndex + 9);
+						}
+					}
+					DestFilename = ContentDestinationRoot / DestFilename;
 
-				FString DestDirectory = FPaths::GetPath(DestFilename);
+					FString DestDirectory = FPaths::GetPath(DestFilename);
 
-				if (PlatformFile.CreateDirectoryTree(*DestDirectory))
-				{
-					if (PlatformFile.CopyFile(*DestFilename, *FileToCopy))
+					if (PlatformFile.CreateDirectoryTree(*DestDirectory))
 					{
-						WrittenFiles.Add(DestFilename);
-						UE_LOG(LogPackFactory, Log, TEXT("Copied \"%s\" to \"%s\""), *FileToCopy, *DestFilename);
+						if (PlatformFile.CopyFile(*DestFilename, *FileToCopy))
+						{
+							WrittenFiles.Add(DestFilename);
+							UE_LOG(LogPackFactory, Log, TEXT("Copied \"%s\" to \"%s\""), *FileToCopy, *DestFilename);
+						}
+						else
+						{
+							UE_LOG(LogPackFactory, Error, TEXT("Unable to copy file \"%s\" to \"%s\"."), *FileToCopy, *DestFilename);
+						}
 					}
 					else
 					{
-						UE_LOG(LogPackFactory, Error, TEXT("Unable to copy file \"%s\" to \"%s\"."), *FileToCopy, *DestFilename);
+						UE_LOG(LogPackFactory, Error, TEXT("Unable to create directory \"%s\"."), *FileToCopy, *DestFilename);
 					}
 				}
-				else
-				{
-					UE_LOG(LogPackFactory, Error, TEXT("Unable to create directory \"%s\"."), *FileToCopy, *DestFilename);
-				}
 			}
 		}
 
