UE-4871 Live: Crash when opening a blueprint who's component was edited in VS

Hot Reload reinstancer replaces old uclass.

diff --git a/Engine/Source/Developer/HotReload/Private/HotReloadClassReinstancer.h b/Engine/Source/Developer/HotReload/Private/HotReloadClassReinstancer.h
--- a/Engine/Source/Developer/HotReload/Private/HotReloadClassReinstancer.h
+++ b/Engine/Source/Developer/HotReload/Private/HotReloadClassReinstancer.h
@@ -114,6 +114,8 @@ public:
 	// FSerializableObject interface
 	virtual void AddReferencedObjects(FReferenceCollector& Collector) override;
 	// End of FSerializableObject interface
+
+	virtual bool IsClassObjectReplaced() const override  { return true; }
 };
 
 #endif // WITH_ENGINE
\ No newline at end of file
diff --git a/Engine/Source/Editor/UnrealEd/Private/Kismet2/KismetReinstanceUtilities.cpp b/Engine/Source/Editor/UnrealEd/Private/Kismet2/KismetReinstanceUtilities.cpp
--- a/Engine/Source/Editor/UnrealEd/Private/Kismet2/KismetReinstanceUtilities.cpp
+++ b/Engine/Source/Editor/UnrealEd/Private/Kismet2/KismetReinstanceUtilities.cpp
@@ -38,6 +38,17 @@ struct FReplaceReferenceHelper
 		}
 	}
 
+	static void IncludeClass(UClass* OldClass, UClass* NewClass, TMap<UObject*, UObject*> &OldToNewInstanceMap, TArray<UObject*> &SourceObjects, TArray<UObject*> &ObjectsToReplace)
+	{
+		OldToNewInstanceMap.Add(OldClass, NewClass);
+		SourceObjects.Add(OldClass);
+
+		if (auto OldCDO = OldClass->GetDefaultObject(false))
+		{
+			ObjectsToReplace.Add(OldCDO);
+		}
+	}
+
 	static void FindAndReplaceReferences(const TArray<UObject*>& SourceObjects, TSet<UObject*>* ObjectsThatShouldUseOldStuff, const TArray<UObject*>& ObjectsToReplace, const TMap<UObject*, UObject*>& OldToNewInstanceMap, const TMap<FStringAssetReference, UObject*>& ReinstancedObjectsWeakReferenceMap)
 	{
 		// Find everything that references these objects
@@ -321,6 +332,12 @@ void FBlueprintCompileReinstancer::ReinstanceFast()
 	TMap<UObject*, UObject*> OldToNewInstanceMap;
 	TMap<FStringAssetReference, UObject*> ReinstancedObjectsWeakReferenceMap;
 	FReplaceReferenceHelper::IncludeCDO(DuplicatedClass, ClassToReinstance, OldToNewInstanceMap, SourceObjects, OriginalCDO);
+
+	if (IsClassObjectReplaced())
+	{
+		FReplaceReferenceHelper::IncludeClass(DuplicatedClass, ClassToReinstance, OldToNewInstanceMap, SourceObjects, ObjectsToReplace);
+	}
+
 	FReplaceReferenceHelper::FindAndReplaceReferences(SourceObjects, &ObjectsThatShouldUseOldStuff, ObjectsToReplace, OldToNewInstanceMap, ReinstancedObjectsWeakReferenceMap);
 }
 
@@ -369,7 +386,7 @@ void FBlueprintCompileReinstancer::ReinstanceInner(bool bForceAlwaysReinstance)
 		if (bShouldReinstance)
 		{
 			UE_LOG(LogBlueprint, Log, TEXT("BlueprintCompileReinstancer: Doing a full reinstance on class '%s'"), *GetPathNameSafe(ClassToReinstance));
-			ReplaceInstancesOfClass(DuplicatedClass, ClassToReinstance, OriginalCDO, &ObjectsThatShouldUseOldStuff);
+			ReplaceInstancesOfClass(DuplicatedClass, ClassToReinstance, OriginalCDO, &ObjectsThatShouldUseOldStuff, IsClassObjectReplaced());
 		}
 		else if (ClassToReinstance->IsChildOf<UActorComponent>())
 		{
@@ -720,7 +737,7 @@ void FActorReplacementHelper::AttachChildActors(USceneComponent* RootComponent,
 	}
 }
 
-void FBlueprintCompileReinstancer::ReplaceInstancesOfClass(UClass* OldClass, UClass* NewClass, UObject*	OriginalCDO, TSet<UObject*>* ObjectsThatShouldUseOldStuff)
+void FBlueprintCompileReinstancer::ReplaceInstancesOfClass(UClass* OldClass, UClass* NewClass, UObject*	OriginalCDO, TSet<UObject*>* ObjectsThatShouldUseOldStuff, bool bClassObjectReplaced)
 {
 	USelection* SelectedActors;
 	bool bSelectionChanged = false;
@@ -945,6 +962,12 @@ void FBlueprintCompileReinstancer::ReplaceInstancesOfClass(UClass* OldClass, UCl
 	SourceObjects.Append(DstObjects);
 
 	FReplaceReferenceHelper::IncludeCDO(OldClass, NewClass, OldToNewInstanceMap, SourceObjects, OriginalCDO);
+	
+	if (bClassObjectReplaced)
+	{
+		FReplaceReferenceHelper::IncludeClass(OldClass, NewClass, OldToNewInstanceMap, SourceObjects, ObjectsToReplace);
+	}
+
 	FReplaceReferenceHelper::FindAndReplaceReferences(SourceObjects, ObjectsThatShouldUseOldStuff, ObjectsToReplace, OldToNewInstanceMap, ReinstancedObjectsWeakReferenceMap);
 
 	{ BP_SCOPED_COMPILER_EVENT_STAT(EKismetReinstancerStats_ReplacementConstruction);
diff --git a/Engine/Source/Editor/UnrealEd/Public/Kismet2/KismetReinstanceUtilities.h b/Engine/Source/Editor/UnrealEd/Public/Kismet2/KismetReinstanceUtilities.h
--- a/Engine/Source/Editor/UnrealEd/Public/Kismet2/KismetReinstanceUtilities.h
+++ b/Engine/Source/Editor/UnrealEd/Public/Kismet2/KismetReinstanceUtilities.h
@@ -72,7 +72,7 @@ public:
 	void UpdateBytecodeReferences();
 
 	/** Worker function to replace all instances of OldClass with a new instance of NewClass */
-	static void ReplaceInstancesOfClass(UClass* OldClass, UClass* NewClass, UObject* OriginalCDO = NULL, TSet<UObject*>* ObjectsThatShouldUseOldStuff = NULL);
+	static void ReplaceInstancesOfClass(UClass* OldClass, UClass* NewClass, UObject* OriginalCDO = NULL, TSet<UObject*>* ObjectsThatShouldUseOldStuff = NULL, bool bClassObjectReplaced = false);
 
 	/**
 	 * When re-instancing a component, we have to make sure all instance owners' 
@@ -86,6 +86,8 @@ public:
 	/** Verify that all instances of the duplicated class have been replaced and collected */
 	void VerifyReplacement();
 
+	virtual bool IsClassObjectReplaced() const { return false; }
+
 protected:
 	void ReinstanceInner(bool bForceAlwaysReinstance);
 
