Performance improvement for initial scanning of content directories for auto reimport

diff --git a/Engine/Source/Editor/UnrealEd/Private/AutoReimport/FileCache.cpp b/Engine/Source/Editor/UnrealEd/Private/AutoReimport/FileCache.cpp
--- a/Engine/Source/Editor/UnrealEd/Private/AutoReimport/FileCache.cpp
+++ b/Engine/Source/Editor/UnrealEd/Private/AutoReimport/FileCache.cpp
@@ -20,26 +20,30 @@ EFileAction ToFileAction(FFileChangeData::EFileChangeAction InAction)
 }
 
 /** Synchronously read a file's MD5 */
-FMD5Hash ReadFileMD5(const FString& Filename)
+FMD5Hash ReadFileMD5(const FString& Filename, TArray<uint8>* Buffer = nullptr)
 {
 	FArchive* Ar = IFileManager::Get().CreateFileReader(*Filename);
 
 	FMD5Hash Hash;
 	if (Ar)
 	{
+		TArray<uint8> LocalScratch;
+		if (!Buffer)
+		{
+			LocalScratch.Reserve(1024*64);
+			Buffer = &LocalScratch;
+		}
 		FMD5 MD5;
 
-		static const int64 BufferSize = 1024;
 		const int64 Size = Ar->TotalSize();
 		int64 Position = 0;
-		uint8 Buffer[BufferSize];
 
 		// Read in BufferSize chunks
 		while (Position < Size)
 		{
-			const auto ReadNum = FMath::Min(Size - Position, BufferSize);
-			Ar->Serialize(Buffer, ReadNum);
-			MD5.Update(Buffer, ReadNum);
+			const auto ReadNum = FMath::Min(Size - Position, (int64)Buffer->Num());
+			Ar->Serialize(Buffer->GetData(), ReadNum);
+			MD5.Update(Buffer->GetData(), ReadNum);
 
 			Position += ReadNum;
 		}
@@ -136,6 +140,9 @@ FAsyncDirectoryReaderThread AsyncReaderThread;
 FAsyncDirectoryReader::FAsyncDirectoryReader(const FString& InDirectory, EPathType InPathType)
 	: RootPath(InDirectory), PathType(InPathType), StartTime(0)
 {
+	// Read in files in 1MB chunks
+	ScratchBuffer.Reserve(1024*1024);
+
 	PendingDirectories.Add(InDirectory);
 	LiveState.Emplace();
 }
@@ -213,7 +220,7 @@ FAsyncDirectoryReader::EProgressResult FAsyncDirectoryReader::Tick(const FTimeLi
 
 		if (!MD5.IsValid())
 		{
-			MD5 = ReadFileMD5(File);
+			MD5 = ReadFileMD5(File, &ScratchBuffer);
 		}
 
 		LiveState->Files.Emplace(MoveTemp(Filename), FFileData(Timestamp, MD5));
diff --git a/Engine/Source/Editor/UnrealEd/Private/AutoReimport/FileCache.h b/Engine/Source/Editor/UnrealEd/Private/AutoReimport/FileCache.h
--- a/Engine/Source/Editor/UnrealEd/Private/AutoReimport/FileCache.h
+++ b/Engine/Source/Editor/UnrealEd/Private/AutoReimport/FileCache.h
@@ -228,6 +228,9 @@ private:
 
 	/** Friendship required for the reader thread to flag when finished */
 	friend struct FAsyncDirectoryReaderThread;
+
+	/** Scratch buffer used for reading in files */
+	TArray<uint8> ScratchBuffer;
 };
 
 /** Configuration structure required to construct a FFileCache */
