Fix possible case of StaticFindObject while saving.

https://jira.ol.epicgames.net/browse/UE-11234?

diff --git a/Engine/Source/Editor/UnrealEd/Public/GameModeInfoCustomizer.h b/Engine/Source/Editor/UnrealEd/Public/GameModeInfoCustomizer.h
--- a/Engine/Source/Editor/UnrealEd/Public/GameModeInfoCustomizer.h
+++ b/Engine/Source/Editor/UnrealEd/Public/GameModeInfoCustomizer.h
@@ -164,18 +164,9 @@ public:
 	/** Get the currently set GameMode class */
 	const UClass* GetCurrentGameModeClass() const
 	{
-		FString ClassName;
-		DefaultGameModeClassHandle->GetValueAsFormattedString(ClassName);
-
-		if (ClassName.IsEmpty() || ClassName == "None")
-		{
-			return NULL;
-		}
-
-		FString StrippedClassName = ClassName;
-		ConstructorHelpers::StripObjectClass(StrippedClassName);
-
-		const UClass* GameModeClass = FindObject<UClass>(ANY_PACKAGE, *StrippedClassName);
+		UObject* ClassValue = nullptr;
+		DefaultGameModeClassHandle->GetValue(ClassValue);
+		const UClass* GameModeClass = Cast<UClass>(ClassValue);
 		return GameModeClass;
 	}
 
