Fixed error caused by zero-sized scratch buffer

diff --git a/Engine/Source/Editor/UnrealEd/Private/AutoReimport/FileCache.cpp b/Engine/Source/Editor/UnrealEd/Private/AutoReimport/FileCache.cpp
--- a/Engine/Source/Editor/UnrealEd/Private/AutoReimport/FileCache.cpp
+++ b/Engine/Source/Editor/UnrealEd/Private/AutoReimport/FileCache.cpp
@@ -30,7 +30,7 @@ FMD5Hash ReadFileMD5(const FString& Filename, TArray<uint8>* Buffer = nullptr)
 		TArray<uint8> LocalScratch;
 		if (!Buffer)
 		{
-			LocalScratch.Reserve(1024*64);
+			LocalScratch.SetNumUninitialized(1024*64);
 			Buffer = &LocalScratch;
 		}
 		FMD5 MD5;
@@ -141,7 +141,7 @@ FAsyncDirectoryReader::FAsyncDirectoryReader(const FString& InDirectory, EPathTy
 	: RootPath(InDirectory), PathType(InPathType), StartTime(0)
 {
 	// Read in files in 1MB chunks
-	ScratchBuffer.Reserve(1024*1024);
+	ScratchBuffer.SetNumUninitialized(1024 * 1024);
 
 	PendingDirectories.Add(InDirectory);
 	LiveState.Emplace();
