Remove GWorld from TexAlign

diff --git a/Engine/Source/Editor/LevelEditor/Private/LevelEditorActions.cpp b/Engine/Source/Editor/LevelEditor/Private/LevelEditorActions.cpp
--- a/Engine/Source/Editor/LevelEditor/Private/LevelEditorActions.cpp
+++ b/Engine/Source/Editor/LevelEditor/Private/LevelEditorActions.cpp
@@ -1671,7 +1671,7 @@ void FLevelEditorActionCallbacks::OnSelectStationaryLightsExceedingOverlap()
 
 void FLevelEditorActionCallbacks::OnSurfaceAlignment( ETexAlign AlignmentMode )
 {
-	GTexAlignTools.GetAligner( AlignmentMode )->Align( AlignmentMode );
+	GTexAlignTools.GetAligner( AlignmentMode )->Align( GetWorld(), AlignmentMode );
 }
 
 void FLevelEditorActionCallbacks::RegroupActor_Clicked()
diff --git a/Engine/Source/Editor/UnrealEd/Classes/Editor/EditorEngine.h b/Engine/Source/Editor/UnrealEd/Classes/Editor/EditorEngine.h
--- a/Engine/Source/Editor/UnrealEd/Classes/Editor/EditorEngine.h
+++ b/Engine/Source/Editor/UnrealEd/Classes/Editor/EditorEngine.h
@@ -709,7 +709,7 @@ public:
 	 * Exec command handlers
 	 */
 	bool	HandleBlueprintifyFunction( const TCHAR* Str , FOutputDevice& Ar );
-	bool	HandleCallbackCommand( const TCHAR* Str, FOutputDevice& Ar );
+	bool	HandleCallbackCommand( UWorld* InWorld, const TCHAR* Str, FOutputDevice& Ar );
 	bool	HandleTestPropsCommand( const TCHAR* Str, FOutputDevice& Ar );
 	bool	HandleMapCommand( const TCHAR* Str, FOutputDevice& Ar, UWorld* InWorld );
 	bool	HandleSelectCommand( const TCHAR* Str, FOutputDevice& Ar, UWorld* InWorld );
diff --git a/Engine/Source/Editor/UnrealEd/Classes/TexAligner/TexAligner.h b/Engine/Source/Editor/UnrealEd/Classes/TexAligner/TexAligner.h
--- a/Engine/Source/Editor/UnrealEd/Classes/TexAligner/TexAligner.h
+++ b/Engine/Source/Editor/UnrealEd/Classes/TexAligner/TexAligner.h
@@ -63,10 +63,10 @@ class UTexAligner : public UObject
 	// End UObject Interface
 
 	// @todo document
-	UNREALED_API void Align( ETexAlign InTexAlignType );
+	UNREALED_API void Align( UWorld* InWorld, ETexAlign InTexAlignType );
 
 	// @todo document
-	void Align( ETexAlign InTexAlignType, UModel* InModel );
+	void Align( UWorld* InWorld, ETexAlign InTexAlignType, UModel* InModel );
 	
 	// Aligns a specific BSP surface
 	virtual void AlignSurf( ETexAlign InTexAlignType, UModel* InModel, FBspSurfIdx* InSurfIdx, FPoly* InPoly, FVector* InNormal );
diff --git a/Engine/Source/Editor/UnrealEd/Private/Editor.cpp b/Engine/Source/Editor/UnrealEd/Private/Editor.cpp
--- a/Engine/Source/Editor/UnrealEd/Private/Editor.cpp
+++ b/Engine/Source/Editor/UnrealEd/Private/Editor.cpp
@@ -110,7 +110,7 @@ FSimpleMulticastDelegate								FEditorDelegates::LayerChange;
 FEditorDelegates::FOnModeChanged						FEditorDelegates::ChangeEditorMode;
 FSimpleMulticastDelegate								FEditorDelegates::SurfProps;
 FSimpleMulticastDelegate								FEditorDelegates::SelectedProps;
-FSimpleMulticastDelegate								FEditorDelegates::FitTextureToSurface;
+FEditorDelegates::FOnFitTextureToSurface				FEditorDelegates::FitTextureToSurface;
 FSimpleMulticastDelegate								FEditorDelegates::ActorPropertiesChange;
 FSimpleMulticastDelegate								FEditorDelegates::RefreshEditor;
 FSimpleMulticastDelegate								FEditorDelegates::RefreshAllBrowsers;
diff --git a/Engine/Source/Editor/UnrealEd/Private/EditorServer.cpp b/Engine/Source/Editor/UnrealEd/Private/EditorServer.cpp
--- a/Engine/Source/Editor/UnrealEd/Private/EditorServer.cpp
+++ b/Engine/Source/Editor/UnrealEd/Private/EditorServer.cpp
@@ -5184,7 +5184,7 @@ bool UEditorEngine::Exec( UWorld* InWorld, const TCHAR* Stream, FOutputDevice& A
 	}
 	else if( FParse::Command(&Str,TEXT("EDCALLBACK")) )
 	{
-		HandleCallbackCommand( Str, Ar );
+		HandleCallbackCommand( InWorld, Str, Ar );
 	}
 	else if(FParse::Command(&Str,TEXT("STATICMESH")))
 	{
@@ -5535,7 +5535,7 @@ bool UEditorEngine::HandleBlueprintifyFunction( const TCHAR* Str , FOutputDevice
 	
 }
 
-bool UEditorEngine::HandleCallbackCommand( const TCHAR* Str , FOutputDevice& Ar )
+bool UEditorEngine::HandleCallbackCommand( UWorld* InWorld, const TCHAR* Str , FOutputDevice& Ar )
 {
 	bool bResult = true;
 	if ( FParse::Command(&Str,TEXT("SELECTEDPROPS")) )
@@ -5544,7 +5544,7 @@ bool UEditorEngine::HandleCallbackCommand( const TCHAR* Str , FOutputDevice& Ar
 	}
 	else if( FParse::Command( &Str, TEXT( "FITTEXTURETOSURFACE" ) ) )
 	{
-		FEditorDelegates::FitTextureToSurface.Broadcast();
+		FEditorDelegates::FitTextureToSurface.Broadcast(InWorld);
 	}
 	else
 	{
diff --git a/Engine/Source/Editor/UnrealEd/Private/TexAlignTools.cpp b/Engine/Source/Editor/UnrealEd/Private/TexAlignTools.cpp
--- a/Engine/Source/Editor/UnrealEd/Private/TexAlignTools.cpp
+++ b/Engine/Source/Editor/UnrealEd/Private/TexAlignTools.cpp
@@ -57,16 +57,16 @@ void UTexAligner::PostInitProperties()
 	DefTexAlign = TEXALIGN_Default;
 }
 
-void UTexAligner::Align( ETexAlign InTexAlignType )
+void UTexAligner::Align( UWorld* InWorld, ETexAlign InTexAlignType )
 {
-	for( int32 LevelIndex = 0; LevelIndex < GWorld->GetNumLevels(); ++LevelIndex )
+	for( int32 LevelIndex = 0; LevelIndex < InWorld->GetNumLevels(); ++LevelIndex )
 	{
-		ULevel* Level = GWorld->GetLevel(LevelIndex);
-		Align( InTexAlignType, Level->Model );
+		ULevel* Level = InWorld->GetLevel(LevelIndex);
+		Align( InWorld, InTexAlignType, Level->Model );
 	}
 }
 
-void UTexAligner::Align( ETexAlign InTexAlignType, UModel* InModel )
+void UTexAligner::Align( UWorld* InWorld, ETexAlign InTexAlignType, UModel* InModel )
 {
 	//
 	// Build an initial list of BSP surfaces to be aligned.
@@ -132,7 +132,7 @@ void UTexAligner::Align( ETexAlign InTexAlignType, UModel* InModel )
 
 	GEditor->RedrawLevelEditingViewports();
 
-	GWorld->MarkPackageDirty();
+	InWorld->MarkPackageDirty();
 	ULevel::LevelDirtiedEvent.Broadcast();
 }
 
@@ -490,12 +490,12 @@ UTexAligner* FTexAlignTools::GetAligner( ETexAlign InTexAlign )
 
 }
 
-void FTexAlignTools::OnEditorFitTextureToSurface()
+void FTexAlignTools::OnEditorFitTextureToSurface(UWorld* InWorld)
 {
 	UTexAligner* FitAligner = GTexAlignTools.Aligners[ 3 ];
-	for ( int32 LevelIndex = 0; LevelIndex < GWorld->GetNumLevels() ; ++LevelIndex )
+	for ( int32 LevelIndex = 0; LevelIndex < InWorld->GetNumLevels() ; ++LevelIndex )
 	{
-		ULevel* Level = GWorld->GetLevel(LevelIndex);
-		FitAligner->Align( TEXALIGN_None, Level->Model );
+		ULevel* Level = InWorld->GetLevel(LevelIndex);
+		FitAligner->Align( InWorld, TEXALIGN_None, Level->Model );
 	}
 }
diff --git a/Engine/Source/Editor/UnrealEd/Public/Editor.h b/Engine/Source/Editor/UnrealEd/Public/Editor.h
--- a/Engine/Source/Editor/UnrealEd/Public/Editor.h
+++ b/Engine/Source/Editor/UnrealEd/Public/Editor.h
@@ -96,6 +96,8 @@ struct UNREALED_API FEditorDelegates
 	DECLARE_MULTICAST_DELEGATE_OneParam(FOnAssetsDeleted, const TArray<UClass*>& /*DeletedAssetClasses*/);
 	/** delegate type for when a new level is added to the world */
 	DECLARE_MULTICAST_DELEGATE_OneParam(FOnAddLevelToWorld, ULevel*);
+	/** delegate type for when a texture is fit to surface  */
+	DECLARE_MULTICAST_DELEGATE_OneParam(FOnFitTextureToSurface, UWorld*);
 
 	/** Called when the CurrentLevel is switched to a new level.  Note that this event won't be fired for temporary
 		changes to the current level, such as when copying/pasting actors. */
@@ -109,7 +111,7 @@ struct UNREALED_API FEditorDelegates
 	/** Sent when requesting to display the properties of selected actors or BSP surfaces */
 	static FSimpleMulticastDelegate SelectedProps;
 	/** Fits the currently assigned texture to the selected surfaces */
-	static FSimpleMulticastDelegate FitTextureToSurface;
+	static FOnFitTextureToSurface FitTextureToSurface;
 	/** Called when the editor mode is changed */
 	static FOnModeChanged ChangeEditorMode;
 	/** Called when properties of an actor have changed */
diff --git a/Engine/Source/Editor/UnrealEd/Public/TexAlignTools.h b/Engine/Source/Editor/UnrealEd/Public/TexAlignTools.h
--- a/Engine/Source/Editor/UnrealEd/Public/TexAlignTools.h
+++ b/Engine/Source/Editor/UnrealEd/Public/TexAlignTools.h
@@ -55,7 +55,7 @@ private:
 	/** 
 	 * Delegate handlers
 	 **/
-	void OnEditorFitTextureToSurface();
+	void OnEditorFitTextureToSurface(UWorld* InWorld);
 
 };
 
