Fixed issues with 2- or 3-axis scaling on the Unreal Widget.

diff --git a/Engine/Source/Editor/UnrealEd/Private/UnrealWidget.cpp b/Engine/Source/Editor/UnrealEd/Private/UnrealWidget.cpp
--- a/Engine/Source/Editor/UnrealEd/Private/UnrealWidget.cpp
+++ b/Engine/Source/Editor/UnrealEd/Private/UnrealWidget.cpp
@@ -881,7 +881,8 @@ void FWidget::ConvertMouseMovementToAxisMovement( FEditorViewportClient* InViewp
 					AxisDir += ZAxisDir;
 				}
 
-				const float ScaleDelta = FVector2D::DotProduct(AxisDir.GetSafeNormal(), DragDir);
+				AxisDir.Normalize();
+				const float ScaleDelta = FVector2D::DotProduct(AxisDir, DragDir);
 
 				OutScale = FVector(
 					(CurrentAxis & EAxisList::X) ? ScaleDelta : 0.0f,
@@ -894,7 +895,10 @@ void FWidget::ConvertMouseMovementToAxisMovement( FEditorViewportClient* InViewp
 				FSnappingUtils::SnapScale(OutScale, GridSize);
 
 				// Convert to effective screen space delta, and replace input delta, adjusted for inverted screen space Y axis
-				const FVector2D EffectiveDelta = OutScale.X * XAxisDir + OutScale.Y * YAxisDir + OutScale.Z * ZAxisDir;
+				const float ScaleMax = OutScale.GetMax();
+				const float ScaleMin = OutScale.GetMin();
+				const float ScaleApplied = (ScaleMax > -ScaleMin) ? ScaleMax : ScaleMin;
+				const FVector2D EffectiveDelta = AxisDir * ScaleApplied;
 				InOutDelta = FVector(EffectiveDelta.X, -EffectiveDelta.Y, 0.0f);
 			}
 			break;
