Fixed issue where orbit camera delays hiding the mouse cursor until the next tick.

diff --git a/Engine/Source/Editor/UnrealEd/Private/EditorViewportClient.cpp b/Engine/Source/Editor/UnrealEd/Private/EditorViewportClient.cpp
--- a/Engine/Source/Editor/UnrealEd/Private/EditorViewportClient.cpp
+++ b/Engine/Source/Editor/UnrealEd/Private/EditorViewportClient.cpp
@@ -3768,7 +3768,7 @@ void FEditorViewportClient::UpdateRequiredCursorVisibility()
 		// Calc the raw delta from the mouse since we started dragging to detect if there was any movement
 		FVector RawMouseDelta = MouseDeltaTracker->GetScreenDelta();
 
-		if (bMouseButtonDown && (RawMouseDelta.SizeSquared() >= MOUSE_CLICK_DRAG_DELTA || IsFlightCameraActive()) && !MouseDeltaTracker->UsingDragTool())
+		if (bMouseButtonDown && (RawMouseDelta.SizeSquared() >= MOUSE_CLICK_DRAG_DELTA || IsFlightCameraActive() || ShouldOrbitCamera()) && !MouseDeltaTracker->UsingDragTool())
 		{
 			//current system - do not show cursor when mouse is down
 			RequiredCursorVisibiltyAndAppearance.bHardwareCursorVisible = false;
