UE-8787 Possible to have actor get deleted when trying to move it to a new level

diff --git a/Engine/Source/Editor/UnrealEd/Private/EditorActor.cpp b/Engine/Source/Editor/UnrealEd/Private/EditorActor.cpp
--- a/Engine/Source/Editor/UnrealEd/Private/EditorActor.cpp
+++ b/Engine/Source/Editor/UnrealEd/Private/EditorActor.cpp
@@ -115,6 +115,7 @@ void UUnrealEdEngine::edactCopySelected( UWorld* InWorld, FString* DestinationDa
 		// Before copying, deselect:
 		//		- Actors belonging to prefabs unless all actors in the prefab are selected.
 		//		- Builder brushes.
+		//      - World Settings.
 		TArray<AActor*> ActorsToDeselect;
 
 		bool bSomeSelectedActorsNotInCurrentLevel = false;
@@ -131,6 +132,12 @@ void UUnrealEdEngine::edactCopySelected( UWorld* InWorld, FString* DestinationDa
 				ActorsToDeselect.Add(Actor);
 			}
 
+			// Deselect world settings
+			if (Actor->IsA( AWorldSettings::StaticClass() ))
+			{
+				ActorsToDeselect.Add(Actor);
+			}
+
 			// If any selected actors are not in the current level, warn the user that some actors will not be copied.
 			if (!bSomeSelectedActorsNotInCurrentLevel && !Actor->GetLevel()->IsCurrentLevel())
 			{
diff --git a/Engine/Source/Editor/WorldBrowser/Private/LevelModel.cpp b/Engine/Source/Editor/WorldBrowser/Private/LevelModel.cpp
--- a/Engine/Source/Editor/WorldBrowser/Private/LevelModel.cpp
+++ b/Engine/Source/Editor/WorldBrowser/Private/LevelModel.cpp
@@ -9,6 +9,7 @@
 #include "Engine/Selection.h"
 #include "Engine/LevelStreaming.h"
 #include "AssetRegistryModule.h"
+#include "GameFramework/WorldSettings.h"
 
 #define LOCTEXT_NAMESPACE "WorldBrowser"
 
@@ -656,6 +657,14 @@ void FLevelModel::SelectActors(bool bSelect, bool bNotify, bool bSelectEvenIfHid
 				continue;
 			}
 			
+			//exclude the world settings and builder brush from actors selected
+			const bool bIsWorldSettings = Actor->IsA(AWorldSettings::StaticClass());
+			const bool bIsBuilderBrush = (Actor->IsA(ABrush::StaticClass()) && FActorEditorUtils::IsABuilderBrush(Actor));
+			if (bIsWorldSettings || bIsBuilderBrush)
+			{
+				continue;
+			}
+
 			bool bNotifyForActor = false;
 			Editor->GetSelectedActors()->Modify();
 			Editor->SelectActor(Actor, bSelect, bNotifyForActor, bSelectEvenIfHidden);
