Fix custom components with no class group not showing up in the add component drop down
https://jira.ol.epicgames.net/browse/UE-11338

diff --git a/Engine/Source/Editor/UnrealEd/Private/ComponentTypeRegistry.cpp b/Engine/Source/Editor/UnrealEd/Private/ComponentTypeRegistry.cpp
--- a/Engine/Source/Editor/UnrealEd/Private/ComponentTypeRegistry.cpp
+++ b/Engine/Source/Editor/UnrealEd/Private/ComponentTypeRegistry.cpp
@@ -236,6 +236,15 @@ void FComponentTypeRegistryData::ForceRefreshComponentList()
 					FComponentClassComboEntryPtr NewEntry(new FComponentClassComboEntry(ClassGroup, Class, bIncludeInFilter, EComponentCreateAction::SpawnExistingClass));
 					SortedClassList.Add(NewEntry);
 				}
+				else if(ClassGroupNames.Num() == 0 )
+				{
+					// No class group name found. Just add it to a "custom" category
+					
+					const bool bIncludeInFilter = true;
+					FString ClassGroup = LOCTEXT("CustomClassGroup", "Custom").ToString();
+					FComponentClassComboEntryPtr NewEntry(new FComponentClassComboEntry(ClassGroup, Class, bIncludeInFilter, EComponentCreateAction::SpawnExistingClass));
+					SortedClassList.Add(NewEntry);
+				}
 			}
 			
 			if (!bOutOfDateClass && !bBlueprintSkeletonClass)
