Auto Reimport - file hashes are now harvested on a background thread once the initial directory scan has completed.

This has the benefit of startup times being hugely improved in the worse case when no cache was present.

The downside is that rename/move detection on live changes may not work if the hashes have not been gathered yet, but this is only an issue if there was no hash, and the changes happen soon after the editor starts up. Rename/Move changes that apply to files with that haven't had a hash gathered will just be reported as removes and adds.

diff --git a/Engine/Source/Editor/UnrealEd/Private/AutoReimport/FileCache.cpp b/Engine/Source/Editor/UnrealEd/Private/AutoReimport/FileCache.cpp
--- a/Engine/Source/Editor/UnrealEd/Private/AutoReimport/FileCache.cpp
+++ b/Engine/Source/Editor/UnrealEd/Private/AutoReimport/FileCache.cpp
@@ -58,23 +58,24 @@ FMD5Hash ReadFileMD5(const FString& Filename, TArray<uint8>* Buffer = nullptr)
 const FGuid FFileCacheCustomVersion::Key(0x8E7DDCB3, 0x80DA47BB, 0x9FD346A2, 0x93984DF6);
 FCustomVersionRegistration GRegisterFileCacheVersion(FFileCacheCustomVersion::Key, FFileCacheCustomVersion::Latest, TEXT("FileCacheVersion"));
 
+
 /** Single runnable thread used to parse file cache directories without blocking the main thread */
-struct FAsyncDirectoryReaderThread : public FRunnable
+struct FAsyncTaskThread : public FRunnable
 {
-	typedef TArray<TWeakPtr<FAsyncDirectoryReader, ESPMode::ThreadSafe>> FReaderArray;
+	typedef TArray<TWeakPtr<IAsyncTask, ESPMode::ThreadSafe>> FTaskArray;
 
-	FAsyncDirectoryReaderThread() : Thread(nullptr) {}
+	FAsyncTaskThread() : Thread(nullptr) {}
 
 	/** Add a reader to this thread which will get ticked periodically until complete */
-	void AddReader(TSharedPtr<FAsyncDirectoryReader, ESPMode::ThreadSafe> InReader)
+	void AddTask(TSharedPtr<IAsyncTask, ESPMode::ThreadSafe> InTask)
 	{
-		FScopeLock Lock(&ReaderArrayMutex);
-		Readers.Add(InReader);
+		FScopeLock Lock(&TaskArrayMutex);
+		Tasks.Add(InTask);
 
 		if (!Thread)
 		{
 			static int32 Index = 0;
-			Thread = FRunnableThread::Create(this, *FString::Printf(TEXT("AsyncDirectoryReaders_%d"), ++Index));
+			Thread = FRunnableThread::Create(this, *FString::Printf(TEXT("AsyncTaskThread_%d"), ++Index));
 		}
 	}
 
@@ -84,30 +85,30 @@ struct FAsyncDirectoryReaderThread : public FRunnable
 		for(;;)
 		{
 			// Copy the array while we tick the readers
-			FReaderArray Dupl;
+			FTaskArray Dupl;
 			{
-				FScopeLock Lock(&ReaderArrayMutex);
-				Dupl = Readers;
+				FScopeLock Lock(&TaskArrayMutex);
+				Dupl = Tasks;
 			}
 
 			// Tick each one for a second
-			for (auto& Reader : Dupl)
+			for (auto& Task : Dupl)
 			{
-				auto PinnedReader = Reader.Pin();
-				if (PinnedReader.IsValid())
+				auto PinnedTask = Task.Pin();
+				if (PinnedTask.IsValid())
 				{
-					PinnedReader->Tick(FTimeLimit(1));
+					PinnedTask->Tick(FTimeLimit(1));
 				}
 			}
 
-			// Cleanup dead/finished readers
-			FScopeLock Lock(&ReaderArrayMutex);
-			for (int32 Index = 0; Index < Readers.Num(); )
+			// Cleanup dead/finished Tasks
+			FScopeLock Lock(&TaskArrayMutex);
+			for (int32 Index = 0; Index < Tasks.Num(); )
 			{
-				auto Reader = Readers[Index].Pin();
-				if (!Reader.IsValid() || Reader->IsComplete())
+				auto Task = Tasks[Index].Pin();
+				if (!Task.IsValid() || Task->IsComplete())
 				{
-					Readers.RemoveAt(Index);
+					Tasks.RemoveAt(Index);
 				}
 				else
 				{
@@ -116,7 +117,7 @@ struct FAsyncDirectoryReaderThread : public FRunnable
 			}
 
 			// Shutdown the thread if we've nothing left to do
-			if (Readers.Num() == 0)
+			if (Tasks.Num() == 0)
 			{
 				Thread = nullptr;
 				break;
@@ -127,22 +128,81 @@ struct FAsyncDirectoryReaderThread : public FRunnable
 	}
 
 private:
-	/** We start our own thread if one doesn't already exist. Guard it with a critical section */
+	/** We start our own thread if one doesn't already exist. */
 	FRunnableThread* Thread;
 
-	/** Array of things that need ticking, and a separate mutex to protect those */
-	FCriticalSection ReaderArrayMutex;
-	FReaderArray Readers;
+	/** Array of things that need ticking, and a mutex to protect them */
+	FCriticalSection TaskArrayMutex;
+	FTaskArray Tasks;
 };
 
-FAsyncDirectoryReaderThread AsyncReaderThread;
+FAsyncTaskThread AsyncTaskThread;
 
-FAsyncDirectoryReader::FAsyncDirectoryReader(const FString& InDirectory, EPathType InPathType)
-	: RootPath(InDirectory), PathType(InPathType), StartTime(0)
+/** Threading stategy for FAsyncFileHasher:
+ *	The task is constructed on the main thread with its Data.
+ *	The array 'Data' *never* changes size. The task thread moves along setting file hashes, while the main thread
+ * 	trails behind accessing the completed entries. We should thus never have 2 threads accessing the same memory,
+ *	except for the atomic 'CurrentIndex'
+ */
+
+FAsyncFileHasher::FAsyncFileHasher(TArray<FFilenameAndHash> InFilesThatNeedHashing)
+	: Data(MoveTemp(InFilesThatNeedHashing)), NumReturned(0)
 {
 	// Read in files in 1MB chunks
 	ScratchBuffer.SetNumUninitialized(1024 * 1024);
+}
+
+TArray<FFilenameAndHash> FAsyncFileHasher::GetCompletedData()
+{
+	TArray<FFilenameAndHash> Local;
+	const int32 CompletedIndex = CurrentIndex.GetValue();
+
+	if (NumReturned < CompletedIndex)
+	{
+		Local.Append(Data.GetData() + NumReturned, CompletedIndex - NumReturned);
+		NumReturned = CompletedIndex;
+
+		if (CompletedIndex == Data.Num())
+		{
+			Data.Empty();
+		}
+	}
+
+	return Local;
+}
+
+bool FAsyncFileHasher::IsComplete() const
+{
+	return CurrentIndex.GetValue() == Data.Num();
+}
+
+IAsyncTask::EProgressResult FAsyncFileHasher::Tick(const FTimeLimit& Limit)
+{
+	for (; CurrentIndex.GetValue() < Data.Num(); )
+	{
+		const auto Index = CurrentIndex.GetValue();
+		Data[Index].FileHash = ReadFileMD5(Data[Index].AbsoluteFilename, &ScratchBuffer);
 
+		CurrentIndex.Increment();
+
+		if (Limit.Exceeded())
+		{
+			return EProgressResult::Pending;
+		}
+	}
+
+	return EProgressResult::Finished;
+}
+
+/** Threading stategy for FAsyncDirectoryReader:
+ *	The directory reader owns the cached and live state until it has completely finished. Once IsComplete() is true, the main thread can 
+ *	have access to both the cached and farmed data.
+ */
+
+
+FAsyncDirectoryReader::FAsyncDirectoryReader(const FString& InDirectory, EPathType InPathType)
+	: RootPath(InDirectory), PathType(InPathType)
+{
 	PendingDirectories.Add(InDirectory);
 	LiveState.Emplace();
 }
@@ -150,7 +210,11 @@ FAsyncDirectoryReader::FAsyncDirectoryReader(const FString& InDirectory, EPathTy
 TOptional<FDirectoryState> FAsyncDirectoryReader::GetLiveState()
 {
 	TOptional<FDirectoryState> OldState;
-	Swap(OldState, LiveState);
+
+	if (ensureMsg(IsComplete(), TEXT("Invalid property access from thread before task completion")))
+	{
+		Swap(OldState, LiveState);
+	}
 
 	return OldState;
 }
@@ -159,7 +223,11 @@ TOptional<FDirectoryState> FAsyncDirectoryReader::GetLiveState()
 TOptional<FDirectoryState> FAsyncDirectoryReader::GetCachedState()
 {
 	TOptional<FDirectoryState> OldState;
-	Swap(OldState, CachedState);
+
+	if (ensureMsg(IsComplete(), TEXT("Invalid property access from thread before task completion")))
+	{
+		Swap(OldState, CachedState);
+	}
 
 	return OldState;
 }
@@ -176,11 +244,6 @@ FAsyncDirectoryReader::EProgressResult FAsyncDirectoryReader::Tick(const FTimeLi
 		return EProgressResult::Finished;
 	}
 
-	if (StartTime == 0)
-	{
-		StartTime = FPlatformTime::Seconds();
-	}
-
 	auto& FileManager = IFileManager::Get();
 	const int32 RootPathLen = RootPath.Len();
 
@@ -220,7 +283,7 @@ FAsyncDirectoryReader::EProgressResult FAsyncDirectoryReader::Tick(const FTimeLi
 
 		if (!MD5.IsValid())
 		{
-			MD5 = ReadFileMD5(File, &ScratchBuffer);
+			FilesThatNeedHashing.Emplace(File);
 		}
 
 		LiveState->Files.Emplace(MoveTemp(Filename), FFileData(Timestamp, MD5));
@@ -236,7 +299,7 @@ FAsyncDirectoryReader::EProgressResult FAsyncDirectoryReader::Tick(const FTimeLi
 
 	bIsComplete = true;
 
-	UE_LOG(LogAutoReimportManager, Log, TEXT("Scanning file cache directory took %.2fs"), FPlatformTime::Seconds() - StartTime);
+	UE_LOG(LogAutoReimportManager, Log, TEXT("Scanning file cache for directory '%s' took %.2fs"), *RootPath, GetAge());
 	return EProgressResult::Finished;
 }
 
@@ -276,6 +339,7 @@ void FAsyncDirectoryReader::ScanDirectory(const FString& InDirectory)
 FFileCache::FFileCache(const FFileCacheConfig& InConfig)
 	: Config(InConfig)
 	, bSavedCacheDirty(false)
+	, LastFileHashGetTime(0)
 {
 	// Ensure the directory has a trailing /
 	Config.Directory /= TEXT("");
@@ -290,7 +354,7 @@ FFileCache::FFileCache(const FFileCacheConfig& InConfig)
 		DirectoryReader->UseCachedState(MoveTemp(ExistingCache.GetValue()));
 	}
 
-	AsyncReaderThread.AddReader(DirectoryReader);
+	AsyncTaskThread.AddTask(DirectoryReader);
 
 	FDirectoryWatcherModule& Module = FModuleManager::LoadModuleChecked<FDirectoryWatcherModule>(TEXT("DirectoryWatcher"));
 	if (IDirectoryWatcher* DirectoryWatcher = Module.Get())
@@ -316,6 +380,8 @@ void FFileCache::Destroy()
 	}
 
 	DirectoryReader = nullptr;
+	AsyncFileHasher = nullptr;
+
 	PendingTransactions.Empty();
 	DirtyFiles.Empty();
 	CachedDirectoryState = FDirectoryState();
@@ -433,6 +499,9 @@ TOptional<FString> FFileCache::GetTransactionPath(const FString& InAbsolutePath)
 
 void FFileCache::DiffDirtyFiles(const TSet<FImmutableString>& InDirtyFiles, TArray<FUpdateCacheTransaction>& InOutTransactions, const FDirectoryState* InFileSystemState) const
 {
+	TArray<uint8> ScratchBuffer;
+	ScratchBuffer.SetNumUninitialized(1024*1024);
+
 	TMap<FImmutableString, FFileData> AddedFiles, ModifiedFiles;
 	TSet<FImmutableString> RemovedFiles;
 
@@ -448,7 +517,7 @@ void FFileCache::DiffDirtyFiles(const TSet<FImmutableString>& InDirtyFiles, TArr
 		const bool bFileExists = InFileSystemState ? InFileSystemState->Files.Find(File) != nullptr : PlatformFile.FileExists(*AbsoluteFilename);
 		if (bFileExists)
 		{
-			const FFileData FileData = InFileSystemState ? *InFileSystemState->Files.Find(File) : FFileData(FileManager.GetTimeStamp(*AbsoluteFilename), ReadFileMD5(AbsoluteFilename));
+			const FFileData FileData = InFileSystemState ? *InFileSystemState->Files.Find(File) : FFileData(FileManager.GetTimeStamp(*AbsoluteFilename), ReadFileMD5(AbsoluteFilename, &ScratchBuffer));
 
 			// Do we think it exists in the cache?
 			if (CachedState)
@@ -472,21 +541,24 @@ void FFileCache::DiffDirtyFiles(const TSet<FImmutableString>& InDirtyFiles, TArr
 	}
 
 	// Rename / move detection
-	for (auto RemoveIt = RemovedFiles.CreateIterator(); RemoveIt; ++RemoveIt)
+	if (Config.bDetectMoves)
 	{
-		const auto* CachedState = CachedDirectoryState.Files.Find(*RemoveIt);
-		if (CachedState)
+		for (auto RemoveIt = RemovedFiles.CreateIterator(); RemoveIt; ++RemoveIt)
 		{
-			for (auto AdIt = AddedFiles.CreateIterator(); AdIt; ++AdIt)
-			{				
-				if (AdIt.Value().FileHash == CachedState->FileHash)
-				{
-					// Found a move destination!
-					InOutTransactions.Add(FUpdateCacheTransaction(*RemoveIt, AdIt.Key(), AdIt.Value()));
-
-					AdIt.RemoveCurrent();
-					RemoveIt.RemoveCurrent();
-					break;
+			const auto* CachedState = CachedDirectoryState.Files.Find(*RemoveIt);
+			if (CachedState)
+			{
+				for (auto AdIt = AddedFiles.CreateIterator(); AdIt; ++AdIt)
+				{				
+					if (AdIt.Value().FileHash.IsValid() && AdIt.Value().FileHash == CachedState->FileHash)
+					{
+						// Found a move destination!
+						InOutTransactions.Add(FUpdateCacheTransaction(*RemoveIt, AdIt.Key(), AdIt.Value()));
+
+						AdIt.RemoveCurrent();
+						RemoveIt.RemoveCurrent();
+						break;
+					}
 				}
 			}
 		}
@@ -626,20 +698,74 @@ void FFileCache::CompleteTransaction(FUpdateCacheTransaction&& Transaction)
 
 void FFileCache::Tick()
 {
-	if (!DirectoryReader.IsValid() || !DirectoryReader->IsComplete())
+	/** Stage one: wait for the asynchronous directory reader to finish harvesting timestamps for the directory */
+	if (DirectoryReader.IsValid())
 	{
-		return;
+		if (!DirectoryReader->IsComplete())
+		{
+			return;
+		}
+		else
+		{
+			ReadStateFromAsyncReader();
+
+			if (Config.bDetectMoves)
+			{
+				auto FilesThatNeedHashing = DirectoryReader->GetFilesThatNeedHashing();
+				if (FilesThatNeedHashing.Num() > 0)
+				{
+					AsyncFileHasher = MakeShareable(new FAsyncFileHasher(MoveTemp(FilesThatNeedHashing)));
+
+					AsyncTaskThread.AddTask(AsyncFileHasher);
+				}
+			}
+
+			// Null out our pointer to the directory reader to indicate that we've finished
+			DirectoryReader = nullptr;
+		}
 	}
+	/** The file cache is now running, and will report changes. */
+	/** Keep harvesting file hashes from the file hashing task until complete. These are much slower to gather, and only required for rename/move detection. */
+	else if (AsyncFileHasher.IsValid())
+	{
+		double Now = FPlatformTime::Seconds();
+
+		if (Now - LastFileHashGetTime > 5.f)
+		{
+			LastFileHashGetTime = Now;
+			auto Hashes = AsyncFileHasher->GetCompletedData();
+			if (Hashes.Num() > 0)
+			{
+				bSavedCacheDirty = true;
+				for (const auto& Data : Hashes)
+				{
+					FImmutableString CachePath = (Config.PathType == EPathType::Relative) ? *Data.AbsoluteFilename + Config.Directory.Len() : *Data.AbsoluteFilename;
+
+					auto* FileData = CachedDirectoryState.Files.Find(CachePath);
+					if (FileData && !FileData->FileHash.IsValid())
+					{
+						FileData->FileHash = Data.FileHash;
+					}
+				}
+			}
 
+			if (AsyncFileHasher->IsComplete())
+			{
+				UE_LOG(LogAutoReimportManager, Log, TEXT("Retrieving MD5 hashes for directory '%s' took %.2fs"), *Config.Directory, AsyncFileHasher->GetAge());
+				AsyncFileHasher = nullptr;
+			}
+		}
+	}
+}
+
+void FFileCache::ReadStateFromAsyncReader()
+{
 	// We should only ever get here once. The directory reader has finished scanning, and we can now diff the results with what we had saved in the cache file.
 	check(DirectoryReader->IsComplete());
 	
 	TOptional<FDirectoryState> LiveState = DirectoryReader->GetLiveState();
 	TOptional<FDirectoryState> CachedState = DirectoryReader->GetCachedState();
 
-	// Null out our pointer to the directory reader to indicate that we've finished
-	DirectoryReader = nullptr;
-
 	if (!CachedState.IsSet() || !Config.bDetectChangesSinceLastRun)
 	{
 		// If we don't have any cached data yet, just use the file data we just harvested
diff --git a/Engine/Source/Editor/UnrealEd/Private/AutoReimport/FileCache.h b/Engine/Source/Editor/UnrealEd/Private/AutoReimport/FileCache.h
--- a/Engine/Source/Editor/UnrealEd/Private/AutoReimport/FileCache.h
+++ b/Engine/Source/Editor/UnrealEd/Private/AutoReimport/FileCache.h
@@ -146,8 +146,75 @@ enum EPathType
 	Absolute
 };
 
+struct IAsyncTask : public TSharedFromThis<IAsyncTask, ESPMode::ThreadSafe>
+{
+	enum class EProgressResult
+	{
+		Finished, Pending
+	};
+
+	IAsyncTask() : StartTime(FPlatformTime::Seconds()) {}
+
+	/** Tick this task. Only to be called on the task thread. */
+	virtual EProgressResult Tick(const FTimeLimit& TimeLimit) = 0;
+
+	/** Check whether this task is complete. Must be implemented in a thread-safe manner. */
+	virtual bool IsComplete() const = 0;
+
+	/** Get the age of this task in seconds */
+	double GetAge() const { return FPlatformTime::Seconds() - StartTime; }
+
+protected:
+
+	/** The time this task started */
+	double StartTime;
+};
+
+/** Simple struct that encapsulates a filename and its associated MD5 hash */
+struct FFilenameAndHash
+{
+	FString AbsoluteFilename;
+	FMD5Hash FileHash;
+
+	FFilenameAndHash(){}
+	FFilenameAndHash(const FString& File) : AbsoluteFilename(File) {}
+
+	FFilenameAndHash(FFilenameAndHash&& In) : AbsoluteFilename(MoveTemp(In.AbsoluteFilename)), FileHash(In.FileHash) {}
+	FFilenameAndHash& operator=(FFilenameAndHash&& In) { Swap(AbsoluteFilename, In.AbsoluteFilename); FileHash = In.FileHash; }
+};
+
+/** Async task responsible for MD5 hashing a number of files, reporting completed hashes to the client when done */
+struct FAsyncFileHasher : public IAsyncTask
+{
+	/** Constructor */
+	FAsyncFileHasher(TArray<FFilenameAndHash> InFilesThatNeedHashing);
+
+	/** Return any completed filenames and their corresponding hashes */
+	TArray<FFilenameAndHash> GetCompletedData();
+
+	/** Returns true when this directory reader has finished scanning the directory */
+	virtual bool IsComplete() const override;
+
+protected:
+
+	/** Tick this reader (discover new directories / files). Returns progress state. */
+	virtual EProgressResult Tick(const FTimeLimit& Limit) override;
+
+	/** The array of data that we will process */
+	TArray<FFilenameAndHash> Data;
+
+	/** The number of items we have returned to the client. Only accessed from the main thread. */
+	int32 NumReturned;
+
+	/** The number of files that we have hashed on the task thread. Atomic - safe to access from any thread. */
+	FThreadSafeCounter CurrentIndex;
+
+	/** Scratch buffer used for reading in files */
+	TArray<uint8> ScratchBuffer;
+};
+
 /**
- * Class responsible for 'asyncronously' scanning a folder for files and timestamps.
+ * Class responsible for 'asynchronously' scanning a folder for files and timestamps.
  * Example usage:
  *		FAsyncDirectoryReader Reader(TEXT("C:\\Path"), EPathType::Relative);
  *
@@ -158,13 +225,8 @@ enum EPathType
  *		}
  *		TOptional<FDirectoryState> State = Reader.GetFinalState();
  */
-struct FAsyncDirectoryReader : public TSharedFromThis<FAsyncDirectoryReader, ESPMode::ThreadSafe>
+struct FAsyncDirectoryReader : public IAsyncTask
 {
-	enum class EProgressResult
-	{
-		Finished, Pending
-	};
-
 	/** Constructor that sets up the directory reader to the specified directory */
 	FAsyncDirectoryReader(const FString& InDirectory, EPathType InPathType);
 
@@ -186,6 +248,14 @@ struct FAsyncDirectoryReader : public TSharedFromThis<FAsyncDirectoryReader, ESP
 	/** Retrieve the cached state supplied to this class through UseCachedState(). */
 	TOptional<FDirectoryState> GetCachedState();
 
+	/** Retrieve the cached state supplied to this class through UseCachedState(). */
+	TArray<FFilenameAndHash> GetFilesThatNeedHashing()
+	{
+		TArray<FFilenameAndHash> Swapped;
+		Swap(FilesThatNeedHashing, Swapped);
+		return Swapped;
+	}
+
 	/** Instruct the directory reader to use the specified cached state to lookup file hashes, where timestamps haven't changed */
 	void UseCachedState(FDirectoryState InCachedState)
 	{
@@ -193,10 +263,10 @@ struct FAsyncDirectoryReader : public TSharedFromThis<FAsyncDirectoryReader, ESP
 	}
 
 	/** Returns true when this directory reader has finished scanning the directory */
-	bool IsComplete() const;
+	virtual bool IsComplete() const override;
 
 	/** Tick this reader (discover new directories / files). Returns progress state. */
-	EProgressResult Tick(const FTimeLimit& Limit);
+	virtual EProgressResult Tick(const FTimeLimit& Limit) override;
 
 private:
 	/** Non-recursively scan a single directory for its contents. Adds results to Pending arrays. */
@@ -214,6 +284,9 @@ private:
 	/** The previously cached  state of the directory, optional */
 	TOptional<FDirectoryState> CachedState;
 
+	/** An array of files that need hashing */
+	TArray<FFilenameAndHash> FilesThatNeedHashing;
+
 	/** A list of directories we have recursively found on our travels */
 	TArray<FString> PendingDirectories;
 
@@ -222,22 +295,13 @@ private:
 
 	/** Thread safe flag to signify when this class has finished reading */
 	FThreadSafeBool bIsComplete;
-
-	/** The time we started scanning the directory */
-	double StartTime;
-
-	/** Friendship required for the reader thread to flag when finished */
-	friend struct FAsyncDirectoryReaderThread;
-
-	/** Scratch buffer used for reading in files */
-	TArray<uint8> ScratchBuffer;
 };
 
 /** Configuration structure required to construct a FFileCache */
 struct FFileCacheConfig
 {
 	FFileCacheConfig(FString InDirectory, FString InCacheFile)
-		: Directory(InDirectory), CacheFile(InCacheFile), PathType(EPathType::Relative), bDetectChangesSinceLastRun(false)
+		: Directory(InDirectory), CacheFile(InCacheFile), PathType(EPathType::Relative), bDetectChangesSinceLastRun(false), bDetectMoves(true)
 	{}
 
 	/** String specifying the directory on disk that the cache should reflect */
@@ -254,6 +318,10 @@ struct FFileCacheConfig
 
 	/** When true, changes to the directory since the cache shutdown will be detected and reported. When false, said changes will silently be applied to the serialized cache. */
 	bool bDetectChangesSinceLastRun;
+
+	/** True to detect moves and renames (based on file hash), false otherwise. When true, an additional thread will be launched on startup to harvest MD5 hashes.
+	 *  Changes that occur *before* hashes have been gatherd will not be detected as moves/renames. */
+	bool bDetectMoves;
 };
 
 /**
@@ -335,6 +403,9 @@ private:
 	/** Read the cache file data and return the contents */
 	TOptional<FDirectoryState> ReadCache() const;
 
+	/** Called when the initial async reader has finished harvesting file system timestamps */
+	void ReadStateFromAsyncReader();
+
 private:
 
 	/** Configuration settings applied on construction */
@@ -342,6 +413,7 @@ private:
 
 	/** 'Asynchronous' directory reader responsible for gathering all file/timestamp information recursively from our cache directory */
 	TSharedPtr<FAsyncDirectoryReader, ESPMode::ThreadSafe> DirectoryReader;
+	TSharedPtr<FAsyncFileHasher, ESPMode::ThreadSafe> AsyncFileHasher;
 
 	/** A set of dirty files that we will use to report changes to the user */
 	TSet<FImmutableString> DirtyFiles;
@@ -357,4 +429,7 @@ private:
 
 	/** True when the cached state we have in memory is more up to date than the serialized file. Enables WriteCache() when true. */
 	bool bSavedCacheDirty;
+
+	/** The time we last retrieved file hashes from the thread */
+	double LastFileHashGetTime;
 };
