Clang compile error fix

diff --git a/Engine/Source/Editor/UnrealEd/Private/AutoReimport/FileCache.h b/Engine/Source/Editor/UnrealEd/Private/AutoReimport/FileCache.h
--- a/Engine/Source/Editor/UnrealEd/Private/AutoReimport/FileCache.h
+++ b/Engine/Source/Editor/UnrealEd/Private/AutoReimport/FileCache.h
@@ -154,7 +154,8 @@ struct IAsyncTask : public TSharedFromThis<IAsyncTask, ESPMode::ThreadSafe>
 	};
 
 	IAsyncTask() : StartTime(FPlatformTime::Seconds()) {}
-
+	virtual ~IAsyncTask() {}
+	
 	/** Tick this task. Only to be called on the task thread. */
 	virtual EProgressResult Tick(const FTimeLimit& TimeLimit) = 0;
 
