Clang compile error fix

This type didn't need to be moveable anyway

diff --git a/Engine/Source/Editor/UnrealEd/Private/AutoReimport/FileCache.h b/Engine/Source/Editor/UnrealEd/Private/AutoReimport/FileCache.h
--- a/Engine/Source/Editor/UnrealEd/Private/AutoReimport/FileCache.h
+++ b/Engine/Source/Editor/UnrealEd/Private/AutoReimport/FileCache.h
@@ -179,9 +179,6 @@ struct FFilenameAndHash
 
 	FFilenameAndHash(){}
 	FFilenameAndHash(const FString& File) : AbsoluteFilename(File) {}
-
-	FFilenameAndHash(FFilenameAndHash&& In) : AbsoluteFilename(MoveTemp(In.AbsoluteFilename)), FileHash(In.FileHash) {}
-	FFilenameAndHash& operator=(FFilenameAndHash&& In) { Swap(AbsoluteFilename, In.AbsoluteFilename); FileHash = In.FileHash; }
 };
 
 /** Async task responsible for MD5 hashing a number of files, reporting completed hashes to the client when done */
