Fixed auto reimport automation tests

diff --git a/Engine/Source/Editor/UnrealEd/Private/AutoReimport/AutoReimportTests.cpp b/Engine/Source/Editor/UnrealEd/Private/AutoReimport/AutoReimportTests.cpp
--- a/Engine/Source/Editor/UnrealEd/Private/AutoReimport/AutoReimportTests.cpp
+++ b/Engine/Source/Editor/UnrealEd/Private/AutoReimport/AutoReimportTests.cpp
@@ -77,7 +77,7 @@ struct FAutoReimportTestPayload
 	{
 		ADD_LATENT_AUTOMATION_COMMAND(FDelayedCallbackLatentCommand([=]{
 			FileCache->Tick();
-			if (FileCache->HasStartedUp())
+			if (FileCache->MoveDetectionInitialized())
 			{
 				Finished();
 			}
diff --git a/Engine/Source/Editor/UnrealEd/Private/AutoReimport/FileCache.cpp b/Engine/Source/Editor/UnrealEd/Private/AutoReimport/FileCache.cpp
--- a/Engine/Source/Editor/UnrealEd/Private/AutoReimport/FileCache.cpp
+++ b/Engine/Source/Editor/UnrealEd/Private/AutoReimport/FileCache.cpp
@@ -165,6 +165,7 @@ TArray<FFilenameAndHash> FAsyncFileHasher::GetCompletedData()
 		if (CompletedIndex == Data.Num())
 		{
 			Data.Empty();
+			CurrentIndex.Set(0);
 		}
 	}
 
@@ -389,6 +390,28 @@ void FFileCache::Destroy()
 	UnbindWatcher();
 }
 
+bool FFileCache::HasStartedUp() const
+{
+	return !DirectoryReader.IsValid() || DirectoryReader->IsComplete();
+}
+
+bool FFileCache::MoveDetectionInitialized() const
+{
+	if (!HasStartedUp())
+	{
+		return false;
+	}
+	else if (!Config.bDetectMoves)
+	{
+		return true;
+	}
+	else
+	{
+		// We don't check AsyncFileHasher->IsComplete() here because that doesn't necessarily mean we've harvested the results off the thread
+		return !AsyncFileHasher.IsValid();
+	}
+}
+
 const FFileData* FFileCache::FindFileData(FImmutableString InFilename) const
 {
 	if (!ensure(HasStartedUp()))
diff --git a/Engine/Source/Editor/UnrealEd/Private/AutoReimport/FileCache.h b/Engine/Source/Editor/UnrealEd/Private/AutoReimport/FileCache.h
--- a/Engine/Source/Editor/UnrealEd/Private/AutoReimport/FileCache.h
+++ b/Engine/Source/Editor/UnrealEd/Private/AutoReimport/FileCache.h
@@ -346,8 +346,11 @@ public:
 	/** Get the absolute path of the directory this cache reflects */
 	const FString& GetDirectory() const { return Config.Directory; }
 
-	/** Check whether this file cache has finished starting up yet */
-	bool HasStartedUp() const { return !DirectoryReader.IsValid() || DirectoryReader->IsComplete(); }
+	/** Check whether this file cache has finished starting up yet. Does not imply rename/move detection is fully initialized. (see MoveDetectionInitialized()) */
+	bool HasStartedUp() const;
+
+	/** Check whether this move/rename detection has been initiated or not. This can take much longer than startup, so can be checked separately */
+	bool MoveDetectionInitialized() const;
 
 	/** Attempt to locate file data pertaining to the specified filename.
 	 *	@param		InFilename		The filename to find data for - either relative to the directory, or absolute, depending on Config.PathType.
