Auto reimported assets are now saved if the assets were previously not dirty

diff --git a/Engine/Source/Editor/UnrealEd/Private/AutoReimport/AutoReimportManager.cpp b/Engine/Source/Editor/UnrealEd/Private/AutoReimport/AutoReimportManager.cpp
--- a/Engine/Source/Editor/UnrealEd/Private/AutoReimport/AutoReimportManager.cpp
+++ b/Engine/Source/Editor/UnrealEd/Private/AutoReimport/AutoReimportManager.cpp
@@ -83,7 +83,7 @@ private:
 /** Enum and value to specify the current state of our processing */
 enum class ECurrentState
 {
-	Idle, PendingResponse, ProcessAdditions, SavePackages, ProcessModifications, ProcessDeletions
+	Idle, PendingResponse, ProcessAdditions, ProcessModifications, ProcessDeletions, SavePackages
 };
 uint32 GetTypeHash(ECurrentState State)
 {
@@ -224,9 +224,9 @@ FAutoReimportManager::FAutoReimportManager()
 	StateMachine.Add(ECurrentState::Idle,					EStateMachineNode::CallOnce,	[this](const FTimeLimit& T){ return this->Idle();					});
 	StateMachine.Add(ECurrentState::PendingResponse,		EStateMachineNode::CallOnce,	[this](const FTimeLimit& T){ return this->PendingResponse();		});
 	StateMachine.Add(ECurrentState::ProcessAdditions,		EStateMachineNode::CallMany,	[this](const FTimeLimit& T){ return this->ProcessAdditions(T);		});
-	StateMachine.Add(ECurrentState::SavePackages,			EStateMachineNode::CallOnce,	[this](const FTimeLimit& T){ return this->SavePackages();			});
 	StateMachine.Add(ECurrentState::ProcessModifications,	EStateMachineNode::CallMany,	[this](const FTimeLimit& T){ return this->ProcessModifications(T);	});
 	StateMachine.Add(ECurrentState::ProcessDeletions,		EStateMachineNode::CallMany,	[this](const FTimeLimit& T){ return this->ProcessDeletions();		});
+	StateMachine.Add(ECurrentState::SavePackages,			EStateMachineNode::CallOnce,	[this](const FTimeLimit& T){ return this->SavePackages();			});
 }
 
 TArray<FPathAndMountPoint> FAutoReimportManager::GetMonitoredDirectories() const
@@ -459,32 +459,10 @@ TOptional<ECurrentState> FAutoReimportManager::ProcessAdditions(const FTimeLimit
 
 	for (auto& Monitor : DirectoryMonitors)
 	{
-		Monitor.ProcessAdditions(Registry, PackagesToSave, TimeLimit, Factories, *FeedbackContextOverride);
+		Monitor.ProcessAdditions(Registry, TimeLimit, PackagesToSave, Factories, *FeedbackContextOverride);
 		yield TOptional<ECurrentState>();
 	}
 
-	return ECurrentState::SavePackages;
-}
-
-TOptional<ECurrentState> FAutoReimportManager::SavePackages()
-{
-	// We don't override the context specifically when saving packages so the user gets proper feedback
-	//TGuardValue<FFeedbackContext*> ScopedContextOverride(GWarn, FeedbackContextOverride.Get());
-
-	TGuardValue<bool> ScopedAssetChangesGuard(bGuardAssetChanges, true);
-
-	FeedbackContextOverride->GetContent()->SetMainText(GetProgressText());
-
-	if (PackagesToSave.Num() > 0)
-	{
-		const bool bAlreadyCheckedOut = false;
-		const bool bCheckDirty = false;
-		const bool bPromptToSave = false;
-		FEditorFileUtils::PromptForCheckoutAndSave(PackagesToSave, bCheckDirty, bPromptToSave, nullptr, bAlreadyCheckedOut);
-
-		PackagesToSave.Empty();
-	}
-
 	return ECurrentState::ProcessModifications;
 }
 
@@ -500,7 +478,7 @@ TOptional<ECurrentState> FAutoReimportManager::ProcessModifications(const FTimeL
 
 	for (auto& Monitor : DirectoryMonitors)
 	{
-		Monitor.ProcessModifications(Registry, TimeLimit, *FeedbackContextOverride);
+		Monitor.ProcessModifications(Registry, TimeLimit, PackagesToSave, *FeedbackContextOverride);
 		yield TOptional<ECurrentState>();
 	}
 
@@ -532,6 +510,29 @@ TOptional<ECurrentState> FAutoReimportManager::ProcessDeletions()
 		ObjectTools::DeleteAssets(AssetsToDelete);
 	}
 
+	return ECurrentState::SavePackages;
+}
+
+TOptional<ECurrentState> FAutoReimportManager::SavePackages()
+{
+	// We don't override the context specifically when saving packages so the user gets proper feedback
+	//TGuardValue<FFeedbackContext*> ScopedContextOverride(GWarn, FeedbackContextOverride.Get());
+
+	TGuardValue<bool> ScopedAssetChangesGuard(bGuardAssetChanges, true);
+
+	FeedbackContextOverride->GetContent()->SetMainText(GetProgressText());
+
+	if (PackagesToSave.Num() > 0)
+	{
+		const bool bAlreadyCheckedOut = false;
+		const bool bCheckDirty = false;
+		const bool bPromptToSave = false;
+		FEditorFileUtils::PromptForCheckoutAndSave(PackagesToSave, bCheckDirty, bPromptToSave, nullptr, bAlreadyCheckedOut);
+
+		PackagesToSave.Empty();
+	}
+
+	// Make sure the progress text is up to date before we close the notification
 	FeedbackContextOverride->GetContent()->SetMainText(GetProgressText());
 
 	Cleanup();
diff --git a/Engine/Source/Editor/UnrealEd/Private/AutoReimport/ContentDirectoryMonitor.cpp b/Engine/Source/Editor/UnrealEd/Private/AutoReimport/ContentDirectoryMonitor.cpp
--- a/Engine/Source/Editor/UnrealEd/Private/AutoReimport/ContentDirectoryMonitor.cpp
+++ b/Engine/Source/Editor/UnrealEd/Private/AutoReimport/ContentDirectoryMonitor.cpp
@@ -12,6 +12,12 @@
 
 #define LOCTEXT_NAMESPACE "ContentDirectoryMonitor"
 
+bool IsAssetDirty(UObject* Asset)
+{
+	UPackage* Package = Asset ? Asset->GetOutermost() : nullptr;
+	return Package ? Package->IsDirty() : false;
+}
+
 /** Generate a config from the specified options, to pass to FFileCache on construction */
 FFileCacheConfig GenerateFileCacheConfig(const FString& InPath, const FMatchRules& InMatchRules, const FString& InMountedContentPath)
 {
@@ -138,7 +144,7 @@ UObject* AttemptImport(UClass* InFactoryType, UPackage* Package, FName InName, b
 	return Asset;
 }
 
-void FContentDirectoryMonitor::ProcessAdditions(const IAssetRegistry& Registry, TArray<UPackage*>& OutPackagesToSave, const FTimeLimit& TimeLimit, const TMap<FString, TArray<UFactory*>>& InFactoriesByExtension, FReimportFeedbackContext& Context)
+void FContentDirectoryMonitor::ProcessAdditions(const IAssetRegistry& Registry, const FTimeLimit& TimeLimit, TArray<UPackage*>& OutPackagesToSave, const TMap<FString, TArray<UFactory*>>& InFactoriesByExtension, FReimportFeedbackContext& Context)
 {
 	bool bCancelled = false;
 	for (int32 Index = 0; Index < AddedFiles.Num(); ++Index)
@@ -272,7 +278,7 @@ void FContentDirectoryMonitor::ProcessAdditions(const IAssetRegistry& Registry,
 	AddedFiles.Empty();
 }
 
-void FContentDirectoryMonitor::ProcessModifications(const IAssetRegistry& Registry, const FTimeLimit& TimeLimit, FReimportFeedbackContext& Context)
+void FContentDirectoryMonitor::ProcessModifications(const IAssetRegistry& Registry, const FTimeLimit& TimeLimit, TArray<UPackage*>& OutPackagesToSave, FReimportFeedbackContext& Context)
 {
 	auto* ReimportManager = FReimportManager::Instance();
 
@@ -292,6 +298,8 @@ void FContentDirectoryMonitor::ProcessModifications(const IAssetRegistry& Regist
 				UObject* Asset = Assets[0].GetAsset();
 				if (Asset && Utils::ExtractSourceFilePaths(Asset).Num() == 1)
 				{
+					const bool bAssetWasDirty = IsAssetDirty(Asset);
+
 					FString NewAssetName = ObjectTools::SanitizeObjectName(FPaths::GetBaseFilename(Change.Filename.Get()));
 					FString PackagePath = PackageTools::SanitizePackageName(MountedContentPath / FPaths::GetPath(Change.Filename.Get()));
 
@@ -317,6 +325,11 @@ void FContentDirectoryMonitor::ProcessModifications(const IAssetRegistry& Regist
 						// Update the reimport file names
 						FReimportManager::Instance()->UpdateReimportPaths(Asset, Filenames);
 						Asset->MarkPackageDirty();
+
+						if (!bAssetWasDirty)
+						{
+							OutPackagesToSave.Add(Asset->GetOutermost());
+						}
 					}
 				}
 			}
@@ -326,13 +339,21 @@ void FContentDirectoryMonitor::ProcessModifications(const IAssetRegistry& Regist
 			for (const auto& AssetData : Utils::FindAssetsPertainingToFile(Registry, FullFilename))
 			{
 				UObject* Asset = AssetData.GetAsset();
-				if (!ReimportManager->Reimport(Asset, false /* Ask for new file */, false /* Show notification */))
-				{
-					Context.AddMessage(EMessageSeverity::Error, FText::Format(LOCTEXT("Error_FailedToReimportAsset", "Failed to reimport asset {0}."), FText::FromString(Asset->GetName())));
-				}
-				else
+				if (Asset)
 				{
-					Context.AddMessage(EMessageSeverity::Info, FText::Format(LOCTEXT("Success_CreatedNewAsset", "Reimported asset {0} from {1}."), FText::FromString(Asset->GetName()), FText::FromString(FullFilename)));
+					const bool bAssetWasDirty = IsAssetDirty(Asset);
+					if (!ReimportManager->Reimport(Asset, false /* Ask for new file */, false /* Show notification */))
+					{
+						Context.AddMessage(EMessageSeverity::Error, FText::Format(LOCTEXT("Error_FailedToReimportAsset", "Failed to reimport asset {0}."), FText::FromString(Asset->GetName())));
+					}
+					else
+					{
+						Context.AddMessage(EMessageSeverity::Info, FText::Format(LOCTEXT("Success_CreatedNewAsset", "Reimported asset {0} from {1}."), FText::FromString(Asset->GetName()), FText::FromString(FullFilename)));
+						if (!bAssetWasDirty)
+						{
+							OutPackagesToSave.Add(Asset->GetOutermost());
+						}
+					}
 				}
 			}
 		}
diff --git a/Engine/Source/Editor/UnrealEd/Private/AutoReimport/ContentDirectoryMonitor.h b/Engine/Source/Editor/UnrealEd/Private/AutoReimport/ContentDirectoryMonitor.h
--- a/Engine/Source/Editor/UnrealEd/Private/AutoReimport/ContentDirectoryMonitor.h
+++ b/Engine/Source/Editor/UnrealEd/Private/AutoReimport/ContentDirectoryMonitor.h
@@ -26,10 +26,10 @@ public:
 	void StartProcessing();
 
 	/** Extract the files we need to import from our outstanding changes (happens first)*/ 
-	void ProcessAdditions(const IAssetRegistry& Registry, TArray<UPackage*>& OutPackagesToSave, const FTimeLimit& TimeLimit, const TMap<FString, TArray<UFactory*>>& InFactoriesByExtension, class FReimportFeedbackContext& Context);
+	void ProcessAdditions(const IAssetRegistry& Registry, const FTimeLimit& TimeLimit, TArray<UPackage*>& OutPackagesToSave, const TMap<FString, TArray<UFactory*>>& InFactoriesByExtension, class FReimportFeedbackContext& Context);
 
 	/** Process the outstanding changes that we have cached */
-	void ProcessModifications(const IAssetRegistry& Registry, const FTimeLimit& TimeLimit, class FReimportFeedbackContext& Context);
+	void ProcessModifications(const IAssetRegistry& Registry, const FTimeLimit& TimeLimit, TArray<UPackage*>& OutPackagesToSave, class FReimportFeedbackContext& Context);
 
 	/** Extract the assets we need to delete from our outstanding changes (happens last) */ 
 	void ExtractAssetsToDelete(const IAssetRegistry& Registry, TArray<FAssetData>& OutAssetsToDelete);
