Auto reimport directories now wait for a little while before being updated as a result of a settings change

diff --git a/Engine/Source/Editor/UnrealEd/Private/AutoReimport/AutoReimportManager.cpp b/Engine/Source/Editor/UnrealEd/Private/AutoReimport/AutoReimportManager.cpp
--- a/Engine/Source/Editor/UnrealEd/Private/AutoReimport/AutoReimportManager.cpp
+++ b/Engine/Source/Editor/UnrealEd/Private/AutoReimport/AutoReimportManager.cpp
@@ -193,6 +193,9 @@ private:
 
 	/** Reentracy guard for when we are making changes to assets */
 	bool bGuardAssetChanges;
+
+	/** A timeout used to refresh directory monitors when the user has made an interactive change to the settings */
+	FTimeLimit ResetMonitorsTimeout;
 };
 
 FAutoReimportManager::FAutoReimportManager()
@@ -547,6 +550,29 @@ TOptional<ECurrentState> FAutoReimportManager::PendingResponse()
 
 TOptional<ECurrentState> FAutoReimportManager::Idle()
 {
+	// Check whether we need to reset the monitors or not
+	if (ResetMonitorsTimeout.Exceeded())
+	{
+		const auto* Settings = GetDefault<UEditorLoadingSavingSettings>();
+		if (Settings->bMonitorContentDirectories)
+		{
+			DirectoryMonitors.Empty();
+			SetUpDirectoryMonitors();
+		}
+		else
+		{
+			// Destroy all the existing monitors, including their file caches
+			for (auto& Monitor : DirectoryMonitors)
+			{
+				Monitor.Destroy();
+			}
+			DirectoryMonitors.Empty();
+		}
+
+		ResetMonitorsTimeout = FTimeLimit();
+		return TOptional<ECurrentState>();
+	}
+
 	for (auto& Monitor : DirectoryMonitors)
 	{
 		Monitor.Tick();
@@ -625,21 +651,7 @@ void FAutoReimportManager::HandleLoadingSavingSettingChanged(FName PropertyName)
 	if (PropertyName == GET_MEMBER_NAME_CHECKED(UEditorLoadingSavingSettings, bMonitorContentDirectories) ||
 		PropertyName == GET_MEMBER_NAME_CHECKED(UEditorLoadingSavingSettings, AutoReimportDirectorySettings))
 	{
-		const auto* Settings = GetDefault<UEditorLoadingSavingSettings>();
-		if (Settings->bMonitorContentDirectories)
-		{
-			DirectoryMonitors.Empty();
-			SetUpDirectoryMonitors();
-		}
-		else
-		{
-			// Destroy all the existing monitors, including their file caches
-			for (auto& Monitor : DirectoryMonitors)
-			{
-				Monitor.Destroy();
-			}
-			DirectoryMonitors.Empty();
-		}
+		ResetMonitorsTimeout = FTimeLimit(5.f);
 	}
 }
 
