Saving file caches now writes to a temporary intermediate file first to avoid corruption

diff --git a/Engine/Source/Editor/UnrealEd/Private/AutoReimport/FileCache.cpp b/Engine/Source/Editor/UnrealEd/Private/AutoReimport/FileCache.cpp
--- a/Engine/Source/Editor/UnrealEd/Private/AutoReimport/FileCache.cpp
+++ b/Engine/Source/Editor/UnrealEd/Private/AutoReimport/FileCache.cpp
@@ -471,8 +471,11 @@ void FFileCache::WriteCache()
 		{
 			IFileManager::Get().MakeDirectory(*ParentFolder, true);	
 		}
+
+		// Write to a temp file to avoid corruption
+		FString TempFile = Config.CacheFile + TEXT(".tmp");
 		
-		FArchive* Ar = IFileManager::Get().CreateFileWriter(*Config.CacheFile);
+		FArchive* Ar = IFileManager::Get().CreateFileWriter(*TempFile);
 		if (ensureMsgf(Ar, TEXT("Unable to write file-cache for '%s' to '%s'."), *Config.Directory, *Config.CacheFile))
 		{
 			*Ar << CachedDirectoryState;
@@ -483,6 +486,9 @@ void FFileCache::WriteCache()
 			CachedDirectoryState.Files.Shrink();
 
 			bSavedCacheDirty = false;
+
+			const bool bMoved = IFileManager::Get().Move(*Config.CacheFile, *TempFile, true, true);
+			ensureMsgf(bMoved, TEXT("Unable to move file-cache for '%s' from '%s' to '%s'."), *Config.Directory, *TempFile, *Config.CacheFile);
 		}
 	}
 }
