Checking for modified files now only checks file timestamps

Previously we could have compared 2 invalid file hashes as equal before they were generated successfully

diff --git a/Engine/Source/Editor/UnrealEd/Private/AutoReimport/FileCache.cpp b/Engine/Source/Editor/UnrealEd/Private/AutoReimport/FileCache.cpp
--- a/Engine/Source/Editor/UnrealEd/Private/AutoReimport/FileCache.cpp
+++ b/Engine/Source/Editor/UnrealEd/Private/AutoReimport/FileCache.cpp
@@ -552,7 +552,7 @@ void FFileCache::DiffDirtyFiles(const TSet<FImmutableString>& InDirtyFiles, TArr
 			if (CachedState)
 			{
 				// Has it changed?
-				if (*CachedState != FileData)
+				if (CachedState->Timestamp != FileData.Timestamp)
 				{
 					ModifiedFiles.Add(File, FileData);
 				}
@@ -575,11 +575,11 @@ void FFileCache::DiffDirtyFiles(const TSet<FImmutableString>& InDirtyFiles, TArr
 		for (auto RemoveIt = RemovedFiles.CreateIterator(); RemoveIt; ++RemoveIt)
 		{
 			const auto* CachedState = CachedDirectoryState.Files.Find(*RemoveIt);
-			if (CachedState)
+			if (CachedState && CachedState->FileHash.IsValid())
 			{
 				for (auto AdIt = AddedFiles.CreateIterator(); AdIt; ++AdIt)
 				{				
-					if (AdIt.Value().FileHash.IsValid() && AdIt.Value().FileHash == CachedState->FileHash)
+					if (AdIt.Value().FileHash == CachedState->FileHash)
 					{
 						// Found a move destination!
 						InOutTransactions.Add(FUpdateCacheTransaction(*RemoveIt, AdIt.Key(), AdIt.Value()));
@@ -818,7 +818,7 @@ void FFileCache::ReadStateFromAsyncReader()
 		const FString& Filename = FilenameAndData.Key.Get();
 
 		const auto* CachedData = CachedDirectoryState.Files.Find(Filename);
-		if (!CachedData || *CachedData != FilenameAndData.Value)
+		if (!CachedData || CachedData->Timestamp != FilenameAndData.Value.Timestamp)
 		{
 			LocalDirtyFiles.Add(Filename);
 		}
