UE-11961 GitHub #932 : About Dialog: change version text to be a read-only SEditableText

diff --git a/Engine/Source/Editor/UnrealEd/Private/AboutScreen.cpp b/Engine/Source/Editor/UnrealEd/Private/AboutScreen.cpp
--- a/Engine/Source/Editor/UnrealEd/Private/AboutScreen.cpp
+++ b/Engine/Source/Editor/UnrealEd/Private/AboutScreen.cpp
@@ -66,8 +66,9 @@ void SAboutScreen::Construct(const FArguments& InArgs)
 					.HAlign(HAlign_Right)
 					.Padding(FMargin(0.f, 52.f, 7.f, 0.f))
 					[
-						SNew(STextBlock)
+						SNew(SEditableText)
 						.ColorAndOpacity(FLinearColor(0.7f, 0.7f, 0.7f))
+						.IsReadOnly(true)
 						.Text( Version )
 					]
 				]
