Fixed custom versions not being written out into file caches

This addresses UE-12224 - Editor crashes when enabling monitor content directories

diff --git a/Engine/Source/Editor/UnrealEd/Private/AutoReimport/FileCache.cpp b/Engine/Source/Editor/UnrealEd/Private/AutoReimport/FileCache.cpp
--- a/Engine/Source/Editor/UnrealEd/Private/AutoReimport/FileCache.cpp
+++ b/Engine/Source/Editor/UnrealEd/Private/AutoReimport/FileCache.cpp
@@ -7,6 +7,48 @@
 #include "DirectoryWatcherModule.h"
 #include "AutoReimportUtilities.h"
 
+template<typename T>
+void ReadWithCustomVersions(FArchive& Ar, T& Data)
+{
+	int64 CustomVersionsOffset = 0;
+	Ar << CustomVersionsOffset;
+
+	const int64 DataStart = Ar.Tell();
+
+	Ar.Seek(CustomVersionsOffset);
+
+	// Serialize the custom versions
+	FCustomVersionContainer Vers = Ar.GetCustomVersions();
+	Vers.Serialize(Ar);
+	Ar.SetCustomVersions(Vers);
+
+	Ar.Seek(DataStart);
+
+	Ar << Data;
+}
+
+template<typename T>
+void WriteWithCustomVersions(FArchive& Ar, T& Data)
+{
+	const int64 CustomVersionsHeader = Ar.Tell();
+	int64 CustomVersionsOffset = CustomVersionsHeader;
+	// We'll come back later and fill this in
+	Ar << CustomVersionsOffset;
+
+	// Write out the data
+	Ar << Data;
+
+	CustomVersionsOffset = Ar.Tell();
+
+	// Serialize the custom versions
+	FCustomVersionContainer Vers = Ar.GetCustomVersions();
+	Vers.Serialize(Ar);
+
+	// Write out where the custom versions are in our header
+	Ar.Seek(CustomVersionsHeader);
+	Ar << CustomVersionsOffset;
+}
+
 /** Convert a FFileChangeData::EFileChangeAction into an EFileAction */
 EFileAction ToFileAction(FFileChangeData::EFileChangeAction InAction)
 {
@@ -58,6 +100,7 @@ FMD5Hash ReadFileMD5(const FString& Filename, TArray<uint8>* Buffer = nullptr)
 const FGuid FFileCacheCustomVersion::Key(0x8E7DDCB3, 0x80DA47BB, 0x9FD346A2, 0x93984DF6);
 FCustomVersionRegistration GRegisterFileCacheVersion(FFileCacheCustomVersion::Key, FFileCacheCustomVersion::Latest, TEXT("FileCacheVersion"));
 
+static const uint32 CacheFileMagicNumber = 0x03DCCB00;
 
 /** Single runnable thread used to parse file cache directories without blocking the main thread */
 struct FAsyncTaskThread : public FRunnable
@@ -449,13 +492,20 @@ TOptional<FDirectoryState> FFileCache::ReadCache() const
 		FArchive* Ar = IFileManager::Get().CreateFileReader(*Config.CacheFile);
 		if (Ar)
 		{
-			FDirectoryState Result;
+			// Serialize the magic number - the first iteration omitted version information, so we have a magic number to ignore this data
+			uint32 MagicNumber = 0;
+			*Ar << MagicNumber;
+
+			if (MagicNumber == CacheFileMagicNumber)
+			{
+				FDirectoryState Result;
+				ReadWithCustomVersions(*Ar, Result);
+
+				Optional.Emplace(MoveTemp(Result));
+			}
 
-			*Ar << Result;
 			Ar->Close();
 			delete Ar;
-
-			Optional.Emplace(MoveTemp(Result));
 		}
 	}
 
@@ -478,7 +528,11 @@ void FFileCache::WriteCache()
 		FArchive* Ar = IFileManager::Get().CreateFileWriter(*TempFile);
 		if (ensureMsgf(Ar, TEXT("Unable to write file-cache for '%s' to '%s'."), *Config.Directory, *Config.CacheFile))
 		{
-			*Ar << CachedDirectoryState;
+			// Serialize the magic number
+			uint32 MagicNumber = CacheFileMagicNumber;
+			*Ar << MagicNumber;
+
+			WriteWithCustomVersions(*Ar, CachedDirectoryState);
 
 			Ar->Close();
 			delete Ar;
