Fixed usage of incorrect iter variable

diff --git a/Engine/Source/Editor/UnrealEd/Private/Commandlets/InternationalizationExportCommandlet.cpp b/Engine/Source/Editor/UnrealEd/Private/Commandlets/InternationalizationExportCommandlet.cpp
--- a/Engine/Source/Editor/UnrealEd/Private/Commandlets/InternationalizationExportCommandlet.cpp
+++ b/Engine/Source/Editor/UnrealEd/Private/Commandlets/InternationalizationExportCommandlet.cpp
@@ -658,7 +658,7 @@ bool FPortableObjectFormatDOM::FromString( const FString& InStr )
 				bSuccess = false;
 				break;
 			}
-			for( uint32 NextLineIdx = LineIdx + 1; NextLineIdx < NumFileLines && LinesToProcess[NextLineIdx].Trim().TrimTrailing().StartsWith(TEXT("\"")); ++LineIdx )
+			for( uint32 NextLineIdx = LineIdx + 1; NextLineIdx < NumFileLines && LinesToProcess[NextLineIdx].Trim().TrimTrailing().StartsWith(TEXT("\"")); ++NextLineIdx)
 			{
 				const FString& NextLine = 
 				RawMsgCtxt += FindDelimitedString(Line, TEXT("\""), TEXT("\""), RawMsgCtxt);
