Fixed a vulnerability where a tooltip on the delete asset dialog could be created with an unbound size.

diff --git a/Engine/Source/Editor/UnrealEd/Private/Dialogs/SDeleteAssetsDialog.cpp b/Engine/Source/Editor/UnrealEd/Private/Dialogs/SDeleteAssetsDialog.cpp
--- a/Engine/Source/Editor/UnrealEd/Private/Dialogs/SDeleteAssetsDialog.cpp
+++ b/Engine/Source/Editor/UnrealEd/Private/Dialogs/SDeleteAssetsDialog.cpp
@@ -428,23 +428,47 @@ FText SDeleteAssetsDialog::GetHandleText() const
 
 FText SDeleteAssetsDialog::GetDeleteSourceContentTooltip() const
 {
-	const FText RootText = LOCTEXT("DeleteSourceFiles_Tooltip", "When checked, the following source content files will also be deleted along with the assets:\n\n{0}");
-
 	FString AllFiles;
-	for (const auto& PathAndAssetCount : DeleteModel->GetPendingDeletedSourceFileCounts())
+
+	static const int32 MaxNumPathsToShow = 25;
+	const auto& AllFileCounts = DeleteModel->GetPendingDeletedSourceFileCounts();
+	int32 TotalCount = 0, NumPrinted = 0;
+	for (const auto& PathAndAssetCount : AllFileCounts)
 	{
 		// If this path is no longer referenced by deleted files, it's toast.
 		if (PathAndAssetCount.Value == 0)
 		{
-			if (!AllFiles.IsEmpty())
+			++TotalCount;
+
+			if (TotalCount <= MaxNumPathsToShow)
 			{
-				AllFiles += TEXT("\n");
+				if (NumPrinted != 0)
+				{
+					AllFiles += TEXT("\n");
+				}
+
+				AllFiles += PathAndAssetCount.Key;
+				++NumPrinted;
 			}
-			AllFiles += PathAndAssetCount.Key;
 		}
 	}
 
-	return FText::Format(RootText, FText::FromString(AllFiles));
+	FText RootText;
+
+	FFormatOrderedArguments Args;
+	Args.Add(FText::FromString(AllFiles));
+
+	if (NumPrinted < TotalCount)
+	{
+		Args.Add(FText::AsNumber(TotalCount - NumPrinted));
+		RootText = LOCTEXT("DeleteSourceFiles_Tooltip", "When checked, the following source content files will also be deleted along with the assets:\n\n{0}\n... and {1} more.");
+	}
+	else
+	{
+		RootText = LOCTEXT("DeleteSourceFiles_Tooltip", "When checked, the following source content files will also be deleted along with the assets:\n\n{0}");
+	}
+	
+	return FText::Format(RootText, Args);
 }
 
 EVisibility SDeleteAssetsDialog::GetAssetReferencesVisiblity() const
