UE-6156: CreateEditorOnlyDefaultSubobject/NeedsLoadForClient/NeedsLoadForServer don't work correctly on editor executable with -game

The problem occured only for brushes. It's fixed, but only for newly created brushes. To fix the problem for old ones you need to recreate them from scratch.

diff --git a/Engine/Source/Editor/UnrealEd/Private/BSPOps.cpp b/Engine/Source/Editor/UnrealEd/Private/BSPOps.cpp
--- a/Engine/Source/Editor/UnrealEd/Private/BSPOps.cpp
+++ b/Engine/Source/Editor/UnrealEd/Private/BSPOps.cpp
@@ -521,8 +521,6 @@ void FBSPOps::csgPrepMovingBrush( ABrush* Actor )
 	check(Actor->Brush);
 	check(Actor->Brush->RootOutside);
 
-	Actor->ClearNotForClientOrServer();
-
 	RebuildBrush(Actor->Brush);
 
 	// Make sure simplified collision is up to date.
