Back out revision 77 from //depot/UE4/Engine/Source/Editor/UnrealEd/Public/Kismet2/KismetEditorUtilities.h

diff --git a/Engine/Source/Editor/UnrealEd/Public/Kismet2/KismetEditorUtilities.h b/Engine/Source/Editor/UnrealEd/Public/Kismet2/KismetEditorUtilities.h
--- a/Engine/Source/Editor/UnrealEd/Public/Kismet2/KismetEditorUtilities.h
+++ b/Engine/Source/Editor/UnrealEd/Public/Kismet2/KismetEditorUtilities.h
@@ -10,37 +10,6 @@
 
 class UBlueprintGeneratedClass;
 
-//////////////////////////////////////////////////////////////////////////
-// FBlueprintUnloader
-
-/** Utility struct, used to aid in unloading and replacing a specific blueprint. */
-struct UNREALED_API FBlueprintUnloader
-{
-public:
-	FBlueprintUnloader(UBlueprint* OldBlueprint);
-
-	/** 
-	 * Unloads the specified Blueprint (marking it pending-kill, and removing it 
-	 * from its outer package). Optionally, will unload the package as well.
-	 *
-	 * @param  bResetPackage	Whether or not this should unload the entire package.
-	 */
-	void UnloadBlueprint(const bool bResetPackage);
-	
-	/** 
-	 * Replaces all old references to the original blueprints (its class/CDO/etc.)
-	 * @param  NewBlueprint	The blueprint to replace old references with
-	 */
-	void ReplaceStaleRefs(UBlueprint* NewBlueprint);
-
-private:
-	TWeakObjectPtr<UBlueprint> OldBlueprint;
-	UClass*  OldGeneratedClass;
-	UObject* OldCDO;
-	UClass*  OldSkeletonClass;
-	UObject* OldSkelCDO;
-};
-
 //////////////////////////////////////////////////////////////////////////
 // FKismetEditorUtilities
 
