Slow task dialogs now only take focus if the application is currently active

This should prevent it from grabbing the foreground window when running in the background.
This addresses UE-10994.

diff --git a/Engine/Source/Editor/UnrealEd/Private/FeedbackContextEditor.cpp b/Engine/Source/Editor/UnrealEd/Private/FeedbackContextEditor.cpp
--- a/Engine/Source/Editor/UnrealEd/Private/FeedbackContextEditor.cpp
+++ b/Engine/Source/Editor/UnrealEd/Private/FeedbackContextEditor.cpp
@@ -381,12 +381,15 @@ void FFeedbackContextEditor::StartSlowTask( const FText& Task, bool bShowCancelB
 				OnCancelClicked = FOnCancelClickedDelegate::CreateRaw(this, &FFeedbackContextEditor::OnUserCancel);
 			}
 
+			const bool bFocusAndActivate = FPlatformProcess::IsThisApplicationForeground();
+
 			TSharedRef<SWindow> SlowTaskWindowRef = SNew(SWindow)
 				.SizingRule(ESizingRule::Autosized)
 				.AutoCenter(EAutoCenter::PreferredWorkArea)
 				.IsPopupWindow(true)
 				.CreateTitleBar(true)
-				.ActivateWhenFirstShown(true);
+				.ActivateWhenFirstShown(bFocusAndActivate)
+				.FocusWhenFirstShown(bFocusAndActivate);
 
 			SlowTaskWindowRef->SetContent(
 				SNew(SSlowTaskWidget)
