Improved progress reporting during autosave

diff --git a/Engine/Source/Editor/UnrealEd/Private/FileHelpers.cpp b/Engine/Source/Editor/UnrealEd/Private/FileHelpers.cpp
--- a/Engine/Source/Editor/UnrealEd/Private/FileHelpers.cpp
+++ b/Engine/Source/Editor/UnrealEd/Private/FileHelpers.cpp
@@ -2176,42 +2176,54 @@ EAutosaveContentPackagesResult::Type FEditorFileUtils::AutosaveContentPackagesEx
 	bool bSavedPkgs = false;
 	const UPackage* TransientPackage = GetTransientPackage();
 	
+	TArray<UPackage*> PackagesToSave;
+
 	// Check all packages for dirty, non-map, non-transient packages
 	for ( TObjectIterator<UPackage> PackageIter; PackageIter; ++PackageIter )
 	{
 		UPackage* CurPackage = *PackageIter;
-
 		// If the package is dirty and is not the transient package, we'd like to autosave it
 		if ( CurPackage && ( CurPackage != TransientPackage ) && CurPackage->IsDirty() && (bForceIfNotInList || DirtyPackagesForAutoSave.Contains(CurPackage)) )
 		{
 			UWorld* MapWorld = UWorld::FindWorldInPackage(CurPackage);
+
 			// Also, make sure this is not a map package
 			const bool bIsMapPackage = MapWorld != NULL;
 
 			// Ignore packages with long, invalid names. This culls out packages with paths in read-only roots such as /Temp.
 			const bool bInvalidLongPackageName = !FPackageName::IsShortPackageName(CurPackage->GetFName()) && !FPackageName::IsValidLongPackageName(CurPackage->GetName(), /*bIncludeReadOnlyRoots=*/false);
-			
+				
 			if ( !bIsMapPackage && !bInvalidLongPackageName )
 			{
-				// In order to save, the package must be fully-loaded first
-				if( !CurPackage->IsFullyLoaded() )
-				{
-					GWarn->BeginSlowTask( NSLOCTEXT("UnrealEd", "FullyLoadingPackages", "Fully loading packages"), true );
-					CurPackage->FullyLoad();
-					GWarn->EndSlowTask();
-				}
+				PackagesToSave.Add(CurPackage);
+			}
+		}
+	}
 
-				const FString AutosaveFilename = GetAutoSaveFilename(CurPackage, AbsoluteAutosaveDir, AutosaveIndex, FPackageName::GetAssetPackageExtension());
-				if (!GUnrealEd->Exec(nullptr, *FString::Printf(TEXT("OBJ SAVEPACKAGE PACKAGE=\"%s\" FILE=\"%s\" SILENT=false AUTOSAVING=true"), *CurPackage->GetName(), *AutosaveFilename)))
-				{
-					return EAutosaveContentPackagesResult::Failure;
-				}
+	FScopedSlowTask SlowTask(PackagesToSave.Num()*2, LOCTEXT("PerformingAutoSave_Caption", "Auto-saving out of date packages..."));
 
-				// Re-mark the package as dirty, because autosaving it will have cleared the dirty flag
-				CurPackage->MarkPackageDirty();
-				bSavedPkgs = true;
-			}
+	for (UPackage* CurPackage : PackagesToSave)
+	{
+		SlowTask.DefaultMessage = FText::Format(LOCTEXT("AutoSavingPackage", "Saving package {0}"), FText::FromString(CurPackage->GetName()));
+		SlowTask.EnterProgressFrame();
+
+		// In order to save, the package must be fully-loaded first
+		if( !CurPackage->IsFullyLoaded() )
+		{
+			CurPackage->FullyLoad();
 		}
+
+		SlowTask.EnterProgressFrame();
+
+		const FString AutosaveFilename = GetAutoSaveFilename(CurPackage, AbsoluteAutosaveDir, AutosaveIndex, FPackageName::GetAssetPackageExtension());
+		if (!GUnrealEd->Exec(nullptr, *FString::Printf(TEXT("OBJ SAVEPACKAGE PACKAGE=\"%s\" FILE=\"%s\" SILENT=false AUTOSAVING=true"), *CurPackage->GetName(), *AutosaveFilename)))
+		{
+			return EAutosaveContentPackagesResult::Failure;
+		}
+
+		// Re-mark the package as dirty, because autosaving it will have cleared the dirty flag
+		CurPackage->MarkPackageDirty();
+		bSavedPkgs = true;
 	}
 	
 	if ( bSavedPkgs )
diff --git a/Engine/Source/Editor/UnrealEd/Private/PackageAutoSaver.cpp b/Engine/Source/Editor/UnrealEd/Private/PackageAutoSaver.cpp
--- a/Engine/Source/Editor/UnrealEd/Private/PackageAutoSaver.cpp
+++ b/Engine/Source/Editor/UnrealEd/Private/PackageAutoSaver.cpp
@@ -151,6 +151,9 @@ void FPackageAutoSaver::AttemptAutoSave()
 		const bool bCanAutosave = CanAutoSave();
 		if (bCanAutosave)
 		{
+			FScopedSlowTask SlowTask(100.f, NSLOCTEXT("AutoSaveNotify", "PerformingAutoSave_Caption", "Auto-saving out of date packages..."));
+			SlowTask.MakeDialog();
+
 			bAutosaveHandled = true;
 
 			bIsAutoSaving = true;
@@ -168,6 +171,8 @@ void FPackageAutoSaver::AttemptAutoSave()
 			bool bLevelSaved = false;
 			auto AssetsSaveResults = EAutosaveContentPackagesResult::NothingToDo;
 
+			SlowTask.EnterProgressFrame(50);
+
 			if (LoadingSavingSettings->bAutoSaveMaps)
 			{
 				// If we weren't saving worlds when we last ran an auto-save, we need to forcibly save any that are dirty
@@ -176,6 +181,8 @@ void FPackageAutoSaver::AttemptAutoSave()
 				bLevelSaved = FEditorFileUtils::AutosaveMap(AutoSaveDir, NewAutoSaveIndex, bForceSaveWorlds, DirtyPackagesForAutoSave);
 			}
 
+			SlowTask.EnterProgressFrame(50);
+
 			if (LoadingSavingSettings->bAutoSaveContent && UnrealEdMisc.GetAutosaveState() != FUnrealEdMisc::EAutosaveState::Cancelled)
 			{
 				// If we weren't saving content packages when we last ran an auto-save, we need to forcibly save any that are dirty
