Updated the variable watch information so it can locate and display component variables watch info correctly
#UE-12427 - Variable out of scope in blueprint debugger variable watch
#Branch UE4
#Proj UnrealEd

diff --git a/Engine/Source/Editor/UnrealEd/Private/Kismet2/KismetDebugUtilities.cpp b/Engine/Source/Editor/UnrealEd/Private/Kismet2/KismetDebugUtilities.cpp
--- a/Engine/Source/Editor/UnrealEd/Private/Kismet2/KismetDebugUtilities.cpp
+++ b/Engine/Source/Editor/UnrealEd/Private/Kismet2/KismetDebugUtilities.cpp
@@ -897,6 +897,21 @@ FKismetDebugUtilities::EWatchTextResult FKismetDebugUtilities::GetWatchText(FStr
 			{
 				PropertyBase = ActiveObject;
 			}
+			else
+			{
+				// Try and locate the propertybase in actors components
+				if( AActor* Actor = Cast<AActor>( ActiveObject ))
+				{
+					for( auto ComponentIter : Actor->GetComponents() )
+					{
+						if( ComponentIter->GetClass()->IsChildOf(PropertyClass))
+						{
+							PropertyBase = ComponentIter;
+							break;
+						}
+					}
+				}
+			}
 
 			// Now either print out the variable value, or that it was out-of-scope
 			if (PropertyBase != nullptr)
