Removed CastChecked use that could fail on old Material Functions

diff --git a/Engine/Source/Editor/UnrealEd/Private/MaterialGraphSchema.cpp b/Engine/Source/Editor/UnrealEd/Private/MaterialGraphSchema.cpp
--- a/Engine/Source/Editor/UnrealEd/Private/MaterialGraphSchema.cpp
+++ b/Engine/Source/Editor/UnrealEd/Private/MaterialGraphSchema.cpp
@@ -782,9 +782,12 @@ bool UMaterialGraphSchema::HasCompatibleConnection(const FAssetData& FunctionAss
 		if (CombinedOutputTypes == 0)
 		{
 			// Need to load function to build combined output types
-			UMaterialFunction* MaterialFunction = CastChecked<UMaterialFunction>(FunctionAssetData.GetAsset());
-			CombinedInputTypes = MaterialFunction->CombinedInputTypes;
-			CombinedOutputTypes = MaterialFunction->CombinedOutputTypes;
+			UMaterialFunction* MaterialFunction = Cast<UMaterialFunction>(FunctionAssetData.GetAsset());
+			if (MaterialFunction != nullptr)
+			{
+				CombinedInputTypes = MaterialFunction->CombinedInputTypes;
+				CombinedOutputTypes = MaterialFunction->CombinedOutputTypes;
+			}
 		}
 
 		if (TestDirection == EGPD_Output)
