Fixed pivot translation so that the pivot can be moved by holding Alt+MMB and dragging any widget translation axis in both perspective and orthographic viewports.

diff --git a/Engine/Source/Editor/UnrealEd/Private/MouseDeltaTracker.cpp b/Engine/Source/Editor/UnrealEd/Private/MouseDeltaTracker.cpp
--- a/Engine/Source/Editor/UnrealEd/Private/MouseDeltaTracker.cpp
+++ b/Engine/Source/Editor/UnrealEd/Private/MouseDeltaTracker.cpp
@@ -188,9 +188,11 @@ void FMouseDeltaTracker::DetermineCurrentAxis(FEditorViewportClient* InViewportC
  */
 void FMouseDeltaTracker::StartTracking(FEditorViewportClient* InViewportClient, const int32 InX, const int32 InY, const FInputEventState& InInputState, bool bNudge, bool bResetDragToolState)
 {
+	DetermineCurrentAxis(InViewportClient);
+
 	// Initialize widget axis (in case it hasn't been set by the hovered hit proxy)
 
-	if (InViewportClient->Widget)
+	if (InViewportClient->Widget && InViewportClient->GetCurrentWidgetAxis() == EAxisList::None)
 	{
 		check(InViewportClient->Viewport);
 		HHitProxy* HitProxy = InViewportClient->Viewport->GetHitProxy(InX, InY);
@@ -201,8 +203,6 @@ void FMouseDeltaTracker::StartTracking(FEditorViewportClient* InViewportClient,
 		}
 	}
 
-	DetermineCurrentAxis( InViewportClient );
-
 	const bool AltDown = InViewportClient->IsAltPressed();
 	const bool ShiftDown = InViewportClient->IsShiftPressed();
 	const bool ControlDown = InViewportClient->IsCtrlPressed();
@@ -400,15 +400,15 @@ void FMouseDeltaTracker::AddDelta(FEditorViewportClient* InViewportClient, FKey
 	// We allow an exception for dragging with the left and/or right mouse button while holding control
 	// as that simulates moving objects with the gizmo
 	//
-	// We also allow the exception of the middle mouse button as it 
+	// We also allow the exception of the middle mouse button when Alt is pressed as it 
 	// allows movement of only the pivot.
 	const bool bIsOrthoObjectRotation = bControlDown && InViewportClient->IsOrtho();
-
-	const bool bUsingAxis = !UsingDragTool() && (LeftMouseButtonDown || ( InViewportClient->GetCurrentWidgetAxis() == EAxisList::Screen && MiddleMouseButtonDown ) || ( (bIsOrthoObjectRotation||bControlDown) && RightMouseButtonDown ));
+	const bool bUsingDragTool = UsingDragTool();
+	const bool bUsingAxis = !bUsingDragTool && (LeftMouseButtonDown || (bAltDown && MiddleMouseButtonDown) || ((bIsOrthoObjectRotation || bControlDown) && RightMouseButtonDown));
 
 	ConditionalBeginUsingDragTool( InViewportClient );
 
-	if( UsingDragTool() || !InViewportClient->IsTracking() || !bUsingAxis )
+	if( bUsingDragTool || !InViewportClient->IsTracking() || !bUsingAxis )
 	{
 		InViewportClient->SetCurrentWidgetAxis( EAxisList::None );
 	}
