UE-12136 This project crashes when one Blueprint is compiled if its child Blueprint is opened in another tab

Dependent blueprints refreshing after a BP compilation is moved to the reinstancer - to avoid any redundant and premature calls for Dependent->BroadcastChanged().

diff --git a/Engine/Source/Editor/UnrealEd/Private/Kismet2/Kismet2.cpp b/Engine/Source/Editor/UnrealEd/Private/Kismet2/Kismet2.cpp
--- a/Engine/Source/Editor/UnrealEd/Private/Kismet2/Kismet2.cpp
+++ b/Engine/Source/Editor/UnrealEd/Private/Kismet2/Kismet2.cpp
@@ -746,6 +746,19 @@ void FKismetEditorUtilities::CompileBlueprint(UBlueprint* BlueprintObj, bool bIs
 
 	ReinstanceHelper->UpdateBytecodeReferences();
 
+	const bool bIsInterface = FBlueprintEditorUtils::IsInterfaceBlueprint(BlueprintObj);
+	static const FBoolConfigValueHelper FirstCompileChildrenThenReinstance(TEXT("Kismet"), TEXT("bFirstCompileChildrenThenReinstance"), GEngineIni);
+	const bool bLetReinstancerRefreshDependBP = !bIsRegeneratingOnLoad && (OldClass != NULL)
+		&& FirstCompileChildrenThenReinstance && !bIsInterface;
+	if (bLetReinstancerRefreshDependBP)
+	{
+		BP_SCOPED_COMPILER_EVENT_STAT(EKismetCompilerStats_RefreshDependentBlueprints);
+
+		TArray<UBlueprint*> DependentBPs;
+		FBlueprintEditorUtils::GetDependentBlueprints(BlueprintObj, DependentBPs);
+		ReinstanceHelper->ListDependentBlueprintsToRefresh(DependentBPs);
+	}
+
 	if (!bIsRegeneratingOnLoad && (OldClass != NULL))
 	{
 		// Strip off any external components from the CDO, if needed because of reparenting, etc
@@ -795,7 +808,8 @@ void FKismetEditorUtilities::CompileBlueprint(UBlueprint* BlueprintObj, bool bIs
 		BlueprintObj->NewVariables[VarIndex].DefaultValue.Empty();
 	}
 
-	{ 
+	if (!bLetReinstancerRefreshDependBP && (bIsInterface || !BlueprintObj->bIsRegeneratingOnLoad))
+	{
 		BP_SCOPED_COMPILER_EVENT_STAT(EKismetCompilerStats_RefreshDependentBlueprints);
 
 		TArray<UBlueprint*> DependentBPs;
@@ -806,7 +820,7 @@ void FKismetEditorUtilities::CompileBlueprint(UBlueprint* BlueprintObj, bool bIs
 		{
 			// for interface changes, auto-refresh nodes on any dependent blueprints
 			// note: RefreshAllNodes() will internally send a change notification event to the dependent blueprint
-			if (FBlueprintEditorUtils::IsInterfaceBlueprint(BlueprintObj))
+			if (bIsInterface)
 			{
 				bool bPreviousRegenValue = Dependent->bIsRegeneratingOnLoad;
 				Dependent->bIsRegeneratingOnLoad = Dependent->bIsRegeneratingOnLoad || BlueprintObj->bIsRegeneratingOnLoad;
diff --git a/Engine/Source/Editor/UnrealEd/Private/Kismet2/KismetReinstanceUtilities.cpp b/Engine/Source/Editor/UnrealEd/Private/Kismet2/KismetReinstanceUtilities.cpp
--- a/Engine/Source/Editor/UnrealEd/Private/Kismet2/KismetReinstanceUtilities.cpp
+++ b/Engine/Source/Editor/UnrealEd/Private/Kismet2/KismetReinstanceUtilities.cpp
@@ -123,6 +123,8 @@ struct FReplaceReferenceHelper
 /////////////////////////////////////////////////////////////////////////////////
 // FBlueprintCompileReinstancer
 
+TSet<TWeakObjectPtr<UBlueprint>> FBlueprintCompileReinstancer::DependentBlueprintsToRefresh = TSet<TWeakObjectPtr<UBlueprint>>();
+
 FBlueprintCompileReinstancer::FBlueprintCompileReinstancer(UClass* InClassToReinstance, bool bIsBytecodeOnly, bool bSkipGC)
 	: ClassToReinstance(InClassToReinstance)
 	, DuplicatedClass(NULL)
@@ -396,6 +398,16 @@ void FBlueprintCompileReinstancer::ReinstanceInner(bool bForceAlwaysReinstance)
 	}
 }
 
+void FBlueprintCompileReinstancer::ListDependentBlueprintsToRefresh(const TArray<UBlueprint*>& DependentBPs)
+{
+	static const FBoolConfigValueHelper FirstCompileChildrenThenReinstance(TEXT("Kismet"), TEXT("bFirstCompileChildrenThenReinstance"), GEngineIni);
+	check(FirstCompileChildrenThenReinstance);
+	for (auto& Element : DependentBPs)
+	{
+		DependentBlueprintsToRefresh.Add(Element);
+	}
+}
+
 void FBlueprintCompileReinstancer::ReinstanceObjects(bool bForceAlwaysReinstance)
 {
 	BP_SCOPED_COMPILER_EVENT_STAT(EKismetCompilerStats_ReinstanceObjects);
@@ -419,6 +431,26 @@ void FBlueprintCompileReinstancer::ReinstanceObjects(bool bForceAlwaysReinstance
 
 			if (QueueToReinstance.Num() && (QueueToReinstance[0] == SharedThis))
 			{
+				{
+					BP_SCOPED_COMPILER_EVENT_STAT(EKismetCompilerStats_RefreshDependentBlueprints);
+					for (auto Reinst : QueueToReinstance)
+					{
+						UBlueprint* BP = Reinst->ClassToReinstance ? Cast<UBlueprint>(Reinst->ClassToReinstance->ClassGeneratedBy) : nullptr;
+						if (BP)
+						{
+							DependentBlueprintsToRefresh.Remove(BP);
+						}
+					}
+					for (auto BPPtr : DependentBlueprintsToRefresh)
+					{
+						if (BPPtr.IsValid())
+						{
+							BPPtr->BroadcastChanged();
+						}
+					}
+					DependentBlueprintsToRefresh.Empty();
+				}
+
 				// All children were recompiled. It's safe to reinstance.
 				for (int32 Idx = 0; Idx < QueueToReinstance.Num(); ++Idx)
 				{
diff --git a/Engine/Source/Editor/UnrealEd/Public/Kismet2/KismetReinstanceUtilities.h b/Engine/Source/Editor/UnrealEd/Public/Kismet2/KismetReinstanceUtilities.h
--- a/Engine/Source/Editor/UnrealEd/Public/Kismet2/KismetReinstanceUtilities.h
+++ b/Engine/Source/Editor/UnrealEd/Public/Kismet2/KismetReinstanceUtilities.h
@@ -14,6 +14,8 @@ class UNREALED_API FBlueprintCompileReinstancer : public TSharedFromThis<FBluepr
 {
 protected:
 
+	static TSet<TWeakObjectPtr<UBlueprint>> DependentBlueprintsToRefresh;
+
 	/** Reference to the class we're actively reinstancing */
 	UClass* ClassToReinstance;
 
@@ -52,6 +54,8 @@ public:
 	virtual void AddReferencedObjects(FReferenceCollector& Collector) override;
 	// End of FSerializableObject interface
 
+	void ListDependentBlueprintsToRefresh(const TArray<UBlueprint*>& DependentBPs);
+
 	static TSharedPtr<FBlueprintCompileReinstancer> Create(UClass* InClassToReinstance, bool bIsBytecodeOnly = false, bool bSkipGC = false)
 	{
 		return MakeShareable(new FBlueprintCompileReinstancer(InClassToReinstance, bIsBytecodeOnly, bSkipGC));
