Cis fix, new stat: EKismetCompilerStats_RefreshDependentBlueprintsInReinstancer

diff --git a/Engine/Source/Editor/UnrealEd/Private/Kismet2/KismetReinstanceUtilities.cpp b/Engine/Source/Editor/UnrealEd/Private/Kismet2/KismetReinstanceUtilities.cpp
--- a/Engine/Source/Editor/UnrealEd/Private/Kismet2/KismetReinstanceUtilities.cpp
+++ b/Engine/Source/Editor/UnrealEd/Private/Kismet2/KismetReinstanceUtilities.cpp
@@ -18,6 +18,7 @@ DECLARE_CYCLE_STAT(TEXT("Update Bytecode References"), EKismetReinstancerStats_U
 DECLARE_CYCLE_STAT(TEXT("Recompile Child Classes"), EKismetReinstancerStats_RecompileChildClasses, STATGROUP_KismetReinstancer );
 DECLARE_CYCLE_STAT(TEXT("Replace Classes Without Reinstancing"), EKismetReinstancerStats_ReplaceClassNoReinsancing, STATGROUP_KismetReinstancer );
 DECLARE_CYCLE_STAT(TEXT("Reinstance Objects"), EKismetCompilerStats_ReinstanceObjects, STATGROUP_KismetCompiler);
+DECLARE_CYCLE_STAT(TEXT("Refresh Dependent Blueprints In Reinstancer"), EKismetCompilerStats_RefreshDependentBlueprintsInReinstancer, STATGROUP_KismetCompiler);
 
 struct FReplaceReferenceHelper
 {
@@ -432,7 +433,7 @@ void FBlueprintCompileReinstancer::ReinstanceObjects(bool bForceAlwaysReinstance
 			if (QueueToReinstance.Num() && (QueueToReinstance[0] == SharedThis))
 			{
 				{
-					BP_SCOPED_COMPILER_EVENT_STAT(EKismetCompilerStats_RefreshDependentBlueprints);
+					BP_SCOPED_COMPILER_EVENT_STAT(EKismetCompilerStats_RefreshDependentBlueprintsInReinstancer);
 					for (auto Reinst : QueueToReinstance)
 					{
 						UBlueprint* BP = Reinst->ClassToReinstance ? Cast<UBlueprint>(Reinst->ClassToReinstance->ClassGeneratedBy) : nullptr;
