Editor: Fixed occasional shutdown crash in source code navigation database (due to static destruction order)

diff --git a/Engine/Source/Editor/UnrealEd/Private/SourceCodeNavigation.cpp b/Engine/Source/Editor/UnrealEd/Private/SourceCodeNavigation.cpp
--- a/Engine/Source/Editor/UnrealEd/Private/SourceCodeNavigation.cpp
+++ b/Engine/Source/Editor/UnrealEd/Private/SourceCodeNavigation.cpp
@@ -175,8 +175,6 @@ FSourceFileDatabase::FSourceFileDatabase()
 {
 	// Register to be notified when new .Build.cs files are added to the project
 	FSourceCodeNavigation::AccessOnNewModuleAdded().AddRaw(this, &FSourceFileDatabase::OnNewModuleAdded);
-
-	UpdateIfNeeded();
 }
 
 FSourceFileDatabase::~FSourceFileDatabase()
@@ -680,6 +678,10 @@ void FSourceCodeNavigationImpl::NavigateToFunctionSource( const FString& Functio
 }
 
 
+FCriticalSection FSourceCodeNavigation::CriticalSection;
+FSourceFileDatabase FSourceCodeNavigation::Instance;
+
+
 void FSourceCodeNavigation::Initialize()
 {
 	class FAsyncInitializeSourceFileDatabase : public FNonAbandonableTask
@@ -711,11 +713,9 @@ void FSourceCodeNavigation::Initialize()
 const FSourceFileDatabase& FSourceCodeNavigation::GetSourceFileDatabase()
 {
 	// Lock so that nothing may proceed while the AsyncTask is constructing the FSourceFileDatabase for the first time
-	static FCriticalSection CriticalSection;
 	FScopeLock Lock(&CriticalSection);
-
-	static FSourceFileDatabase Instance;
 	Instance.UpdateIfNeeded();
+
 	return Instance;
 }
 
diff --git a/Engine/Source/Editor/UnrealEd/Public/SourceCodeNavigation.h b/Engine/Source/Editor/UnrealEd/Public/SourceCodeNavigation.h
--- a/Engine/Source/Editor/UnrealEd/Public/SourceCodeNavigation.h
+++ b/Engine/Source/Editor/UnrealEd/Public/SourceCodeNavigation.h
@@ -200,6 +200,14 @@ public:
 
 	/** Call this to access the multi-cast delegate that you can register a callback with */
 	UNREALED_API static FOnNewModuleAdded& AccessOnNewModuleAdded();
+
+private:
+
+	/** Critical section for locking access to the source file database. */
+	static FCriticalSection CriticalSection;
+
+	/** Source file database instance. */
+	static FSourceFileDatabase Instance;
 };
 
 
