ue4# BP: KismetDebugUtilities::OnScriptExceptionWrong link to node in AccessViolation error message.

https://udn.unrealengine.com/questions/240576/blueprint-none-spamming-the-output-clickback-to-er.html

diff --git a/Engine/Source/Editor/UnrealEd/Private/Kismet2/KismetDebugUtilities.cpp b/Engine/Source/Editor/UnrealEd/Private/Kismet2/KismetDebugUtilities.cpp
--- a/Engine/Source/Editor/UnrealEd/Private/Kismet2/KismetDebugUtilities.cpp
+++ b/Engine/Source/Editor/UnrealEd/Private/Kismet2/KismetDebugUtilities.cpp
@@ -72,6 +72,8 @@ void FKismetDebugUtilities::OnScriptException(const UObject* ActiveObject, const
 		UObject* SavedObjectBeingDebugged = ObjectBeingDebugged;
 		UWorld* WorldBeingDebugged = BlueprintObj->GetWorldBeingDebugged();
 
+		const int32 BreakpointOffset = StackFrame.Code - StackFrame.Node->Script.GetData() - 1;
+
 		bool bShouldBreakExecution = false;
 		bool bForceToCurrentObject = false;
 
@@ -97,11 +99,13 @@ void FKismetDebugUtilities::OnScriptException(const UObject* ActiveObject, const
 				// it's not helpful to link to a generated-function, so we just provide the plain name)
 				TSharedRef<IMessageToken> TraceTargetMessageLog = FTextToken::Create(FText::FromString(GeneratedFuncName));
 
+				FString MsgInBlueprintStr(TEXT("in blueprint"));
+
 #if WITH_EDITORONLY_DATA // to protect access to GeneratedClass->DebugData
 				UBlueprintGeneratedClass* GeneratedClass = Cast<UBlueprintGeneratedClass>(ClassContainingCode);
 				if ((GeneratedClass != NULL) && GeneratedClass->DebugData.IsValid())
 				{
-					UEdGraphNode* BlueprintNode = GeneratedClass->DebugData.GetLastExecutedNode();
+					UEdGraphNode* BlueprintNode = GeneratedClass->DebugData.FindSourceNodeFromCodeLocation(StackFrame.Node, BreakpointOffset, true);
 					// if instead, there is a node we can point to...
 					if (BlueprintNode != NULL)
 					{
@@ -110,6 +114,8 @@ void FKismetDebugUtilities::OnScriptException(const UObject* ActiveObject, const
 						FText NodeTitle = BlueprintNode->GetNodeTitle(ENodeTitleType::ListView); // a more user friendly name
 						// link to the last executed node (the one throwing the exception, presumably)
 						TraceTargetMessageLog = FUObjectToken::Create(BlueprintNode, NodeTitle)->OnMessageTokenActivated(FOnMessageTokenActivated::CreateStatic(&Local::OnMessageLogLinkActivated));
+
+						MsgInBlueprintStr = FString::Printf(TEXT("in graph '%s' in blueprint"), *GetNameSafe(BlueprintNode->GetGraph()));
 					}
 				}
 #endif // WITH_EDITORONLY_DATA
@@ -117,7 +123,7 @@ void FKismetDebugUtilities::OnScriptException(const UObject* ActiveObject, const
 				FMessageLog("PIE").Error(FText::FromString(Info.GetDescription()))
 					->AddToken(TraceClassificationMessageLog)
 					->AddToken(TraceTargetMessageLog)
-					->AddToken(FTextToken::Create(FText::FromString(TEXT("in blueprint"))))
+					->AddToken(FTextToken::Create(FText::FromString(MsgInBlueprintStr)))
 					->AddToken(FUObjectToken::Create(BlueprintObj, FText::FromString(BlueprintObj->GetName()))->OnMessageTokenActivated(FOnMessageTokenActivated::CreateStatic(&Local::OnMessageLogLinkActivated)));
 			}
 			break;
@@ -178,13 +184,11 @@ void FKismetDebugUtilities::OnScriptException(const UObject* ActiveObject, const
 
 		if (BlueprintObj->GetObjectBeingDebugged() == ActiveObject)
 		{
-			const int32 BreakpointOffset = StackFrame.Code - StackFrame.Node->Script.GetData() - 1; //@TODO: Might want to make this a parameter of Info
-
 			// Record into the trace log
 			FKismetTraceSample& Tracer = TraceStackSamples.WriteNewElementUninitialized();
 			Tracer.Context = ActiveObject;
 			Tracer.Function = StackFrame.Node;
-			Tracer.Offset = BreakpointOffset;
+			Tracer.Offset = BreakpointOffset; //@TODO: Might want to make this a parameter of Info
 			Tracer.ObservationTime = FPlatformTime::Seconds();
 
 			// Find the node that generated the code which we hit
diff --git a/Engine/Source/Runtime/Engine/Classes/Engine/BlueprintGeneratedClass.h b/Engine/Source/Runtime/Engine/Classes/Engine/BlueprintGeneratedClass.h
--- a/Engine/Source/Runtime/Engine/Classes/Engine/BlueprintGeneratedClass.h
+++ b/Engine/Source/Runtime/Engine/Classes/Engine/BlueprintGeneratedClass.h
@@ -208,13 +208,6 @@ protected:
 	TMap<TWeakObjectPtr<UObject>, UProperty*> DebugPinToPropertyMap;
 
 public:
-	void Clear()
-	{
-		DebugNodeLineNumbers.Empty();
-		DebugNodeIndexLookup.Empty();
-		DebugNodesAllocatedUniqueIDsMap.Empty();
-		DebugPinToPropertyMap.Empty();
-	}
 
 	// Returns the UEdGraphNode associated with the UUID, or nullptr if there isn't one.
 	UEdGraphNode* FindNodeFromUUID(int32 UUID) const
@@ -393,27 +386,6 @@ public:
 			}
 		}
 	}
-
-	/**
-	 * Retrieves the most recently executed node (created to help pinpoint 
-	 * runtime errors).
-	 * 
-	 * @returns The most recently executed node (nullptr if there is not one).
-	 */
-	UEdGraphNode* GetLastExecutedNode()
-	{
-		UEdGraphNode* LastExecutedNode = nullptr;
-		if (IsValid())
-		{
-			TWeakObjectPtr<UEdGraphNode> NodePtr = DebugNodeLineNumbers.Top().Node;
-			if (NodePtr.IsValid())
-			{
-
-				LastExecutedNode = NodePtr.Get();
-			}
-		}
-		return LastExecutedNode;
-	}
 #endif
 };
 
