UE-12230 HideCategories() property is not held when the editor is closed

While DataOnly BP regeneration MeteData connected to class are refreshed.

diff --git a/Engine/Source/Editor/KismetCompiler/Private/KismetCompiler.cpp b/Engine/Source/Editor/KismetCompiler/Private/KismetCompiler.cpp
--- a/Engine/Source/Editor/KismetCompiler/Private/KismetCompiler.cpp
+++ b/Engine/Source/Editor/KismetCompiler/Private/KismetCompiler.cpp
@@ -18,7 +18,6 @@
 #include "MovieSceneBindings.h"
 #include "Kismet2/Kismet2NameValidators.h"
 #include "UserDefinedStructureCompilerUtils.h"
-#include "EditorCategoryUtils.h"
 #include "K2Node_EnumLiteral.h"
 #include "K2Node_SetVariableOnPersistentFrame.h"
 #include "EdGraph/EdGraphNode_Documentation.h"
@@ -1663,7 +1662,7 @@ void FKismetCompilerContext::FinishCompilingClass(UClass* Class)
 {
 	UClass* ParentClass = Class->GetSuperClass();
 
-	TArray<FString> AllHideCategories;
+	FBlueprintEditorUtils::RecreateClassMetaData(Blueprint, Class, false);
 
 	if (ParentClass != NULL)
 	{
@@ -1689,57 +1688,12 @@ void FKismetCompilerContext::FinishCompilingClass(UClass* Class)
 
 		// Copy the category info from the parent class
 #if WITH_EDITORONLY_DATA
-		if (!ParentClass->HasMetaData(FBlueprintMetadata::MD_IgnoreCategoryKeywordsInSubclasses))
-		{
-			FEditorCategoryUtils::GetClassHideCategories(ParentClass, AllHideCategories);
-			if (ParentClass->HasMetaData(TEXT("ShowCategories")))
-			{
-				Class->SetMetaData(TEXT("ShowCategories"), *ParentClass->GetMetaData("ShowCategories"));
-			}
-			if (ParentClass->HasMetaData(TEXT("AutoExpandCategories")))
-			{
-				Class->SetMetaData(TEXT("AutoExpandCategories"), *ParentClass->GetMetaData("AutoExpandCategories"));
-			}
-			if (ParentClass->HasMetaData(TEXT("AutoCollapseCategories")))
-			{
-				Class->SetMetaData(TEXT("AutoCollapseCategories"), *ParentClass->GetMetaData("AutoCollapseCategories"));
-			}
-		}
-
-		if (ParentClass->HasMetaData(TEXT("HideFunctions")))
-		{
-			Class->SetMetaData(TEXT("HideFunctions"), *ParentClass->GetMetaData("HideFunctions"));
-		}
 		
 		// Blueprinted Components are always Blueprint Spawnable
 		if (ParentClass->IsChildOf(UActorComponent::StaticClass()))
 		{
-			static const FName NAME_ClassGroupNames(TEXT("ClassGroupNames"));
-			Class->SetMetaData(FBlueprintMetadata::MD_BlueprintSpawnableComponent, TEXT("true"));
-
-			FString ClassGroupCategory = NSLOCTEXT("BlueprintableComponents", "CategoryName", "Custom").ToString();
-			if (!Blueprint->BlueprintCategory.IsEmpty())
-			{
-				ClassGroupCategory = Blueprint->BlueprintCategory;
-			}
-
-			Class->SetMetaData(NAME_ClassGroupNames, *ClassGroupCategory);
-
 			FComponentTypeRegistry::Get().InvalidateClass(Class);
 		}
-
-		// Add a category if one has been specified
-		if(Blueprint->BlueprintCategory.Len() > 0)
-		{
-			Class->SetMetaData(TEXT("Category"), *Blueprint->BlueprintCategory);
-		}
-
-		if ((Blueprint->BlueprintType == BPTYPE_Normal) || 
-			(Blueprint->BlueprintType == BPTYPE_Const)  || 
-			(Blueprint->BlueprintType == BPTYPE_Interface))
-		{
-			Class->SetMetaData(FBlueprintMetadata::MD_AllowableBlueprintVariableType, TEXT("true"));
-		}
 #endif
 
 		// Add in additional flags implied by the blueprint
@@ -1756,12 +1710,6 @@ void FKismetCompilerContext::FinishCompilingClass(UClass* Class)
 		//@TODO: Might want to be able to specify some of these here too
 	}
 
-	AllHideCategories.Append(Blueprint->HideCategories);
-	if (AllHideCategories.Num())
-	{
-		Class->SetMetaData(TEXT("HideCategories"), *FString::Join(AllHideCategories, TEXT(" ")));
-	}
-
 	// Add in any other needed flags
 	Class->ClassFlags |= (CLASS_Parsed | CLASS_CompiledFromBlueprint);
 
diff --git a/Engine/Source/Editor/UnrealEd/Private/Kismet2/BlueprintEditorUtils.cpp b/Engine/Source/Editor/UnrealEd/Private/Kismet2/BlueprintEditorUtils.cpp
--- a/Engine/Source/Editor/UnrealEd/Private/Kismet2/BlueprintEditorUtils.cpp
+++ b/Engine/Source/Editor/UnrealEd/Private/Kismet2/BlueprintEditorUtils.cpp
@@ -49,6 +49,8 @@
 #include "Editor/Blutility/Public/IBlutilityModule.h"
 
 #include "Engine/InheritableComponentHandler.h"
+
+#include "EditorCategoryUtils.h"
 #define LOCTEXT_NAMESPACE "Blueprint"
 
 DEFINE_LOG_CATEGORY(LogBlueprintDebug);
@@ -1241,6 +1243,11 @@ UClass* FBlueprintEditorUtils::RegenerateBlueprintClass(UBlueprint* Blueprint, U
 			// No actual compilation work to be done, but try to conform the class and fix up anything that might need to be updated if the native base class has changed in any way
 			FKismetEditorUtilities::ConformBlueprintFlagsAndComponents(Blueprint);
 
+			if (Blueprint->GeneratedClass)
+			{
+				FBlueprintEditorUtils::RecreateClassMetaData(Blueprint, Blueprint->GeneratedClass, true);
+			}
+
 			// Flag data only blueprints as being up-to-date
 			Blueprint->Status = BS_UpToDate;
 		}
@@ -1376,6 +1383,88 @@ UClass* FBlueprintEditorUtils::RegenerateBlueprintClass(UBlueprint* Blueprint, U
 	return bRegenerated ? Blueprint->GeneratedClass : NULL;
 }
 
+void FBlueprintEditorUtils::RecreateClassMetaData(UBlueprint* Blueprint, UClass* Class, bool bRemoveExistingMetaData)
+{
+	if (!ensure(Blueprint && Class))
+	{
+		return;
+	}
+
+	UClass* ParentClass = Class->GetSuperClass();
+	TArray<FString> AllHideCategories;
+
+	if (bRemoveExistingMetaData)
+	{
+		Class->RemoveMetaData("HideCategories");
+		Class->RemoveMetaData("ShowCategories");
+		Class->RemoveMetaData("HideFunctions");
+		Class->RemoveMetaData("AutoExpandCategories");
+		Class->RemoveMetaData("AutoCollapseCategories");
+		Class->RemoveMetaData("ClassGroupNames");
+		Class->RemoveMetaData("Category");
+		Class->RemoveMetaData(FBlueprintMetadata::MD_AllowableBlueprintVariableType);
+	}
+
+	if (ensure(ParentClass != NULL))
+	{
+		if (!ParentClass->HasMetaData(FBlueprintMetadata::MD_IgnoreCategoryKeywordsInSubclasses))
+		{
+			FEditorCategoryUtils::GetClassHideCategories(ParentClass, AllHideCategories);
+			if (ParentClass->HasMetaData(TEXT("ShowCategories")))
+			{
+				Class->SetMetaData(TEXT("ShowCategories"), *ParentClass->GetMetaData("ShowCategories"));
+			}
+			if (ParentClass->HasMetaData(TEXT("AutoExpandCategories")))
+			{
+				Class->SetMetaData(TEXT("AutoExpandCategories"), *ParentClass->GetMetaData("AutoExpandCategories"));
+			}
+			if (ParentClass->HasMetaData(TEXT("AutoCollapseCategories")))
+			{
+				Class->SetMetaData(TEXT("AutoCollapseCategories"), *ParentClass->GetMetaData("AutoCollapseCategories"));
+			}
+		}
+
+		if (ParentClass->HasMetaData(TEXT("HideFunctions")))
+		{
+			Class->SetMetaData(TEXT("HideFunctions"), *ParentClass->GetMetaData("HideFunctions"));
+		}
+
+		if (ParentClass->IsChildOf(UActorComponent::StaticClass()))
+		{
+			static const FName NAME_ClassGroupNames(TEXT("ClassGroupNames"));
+			Class->SetMetaData(FBlueprintMetadata::MD_BlueprintSpawnableComponent, TEXT("true"));
+
+			FString ClassGroupCategory = NSLOCTEXT("BlueprintableComponents", "CategoryName", "Custom").ToString();
+			if (!Blueprint->BlueprintCategory.IsEmpty())
+			{
+				ClassGroupCategory = Blueprint->BlueprintCategory;
+			}
+
+			Class->SetMetaData(NAME_ClassGroupNames, *ClassGroupCategory);
+		}
+	}
+
+	// Add a category if one has been specified
+	if (Blueprint->BlueprintCategory.Len() > 0)
+	{
+		Class->SetMetaData(TEXT("Category"), *Blueprint->BlueprintCategory);
+	}
+
+	if ((Blueprint->BlueprintType == BPTYPE_Normal) ||
+		(Blueprint->BlueprintType == BPTYPE_Const) ||
+		(Blueprint->BlueprintType == BPTYPE_Interface))
+	{
+		Class->SetMetaData(FBlueprintMetadata::MD_AllowableBlueprintVariableType, TEXT("true"));
+	}
+
+	AllHideCategories.Append(Blueprint->HideCategories);
+	if (AllHideCategories.Num())
+	{
+		Class->SetMetaData(TEXT("HideCategories"), *FString::Join(AllHideCategories, TEXT(" ")));
+	}
+}
+
+
 void FBlueprintEditorUtils::PatchCDOSubobjectsIntoExport(UObject* PreviousCDO, UObject* NewCDO)
 {
 	if (PreviousCDO && NewCDO)
diff --git a/Engine/Source/Editor/UnrealEd/Public/Kismet2/BlueprintEditorUtils.h b/Engine/Source/Editor/UnrealEd/Public/Kismet2/BlueprintEditorUtils.h
--- a/Engine/Source/Editor/UnrealEd/Public/Kismet2/BlueprintEditorUtils.h
+++ b/Engine/Source/Editor/UnrealEd/Public/Kismet2/BlueprintEditorUtils.h
@@ -88,6 +88,9 @@ public:
 	 */
 	static void PatchCDOSubobjectsIntoExport(UObject* PreviousCDO, UObject* NewCDO);
 
+	/** Recreates class meta data */
+	static void RecreateClassMetaData(UBlueprint* Blueprint, UClass* Class, bool bRemoveExistingMetaData);
+
 	/**
 	 * Copies the default properties of all parent blueprint classes in the chain to the specified blueprint's skeleton CDO
 	 */
