#UE4 Made PackagesPerGC, IdleTimeToGC, and MaxMemoryAllowance configurable.

diff --git a/Engine/Config/BaseEditor.ini b/Engine/Config/BaseEditor.ini
--- a/Engine/Config/BaseEditor.ini
+++ b/Engine/Config/BaseEditor.ini
@@ -293,6 +293,9 @@ FarLODDistanceTheshold=3000
 
 [CookSettings]
 bConsiderCompressedPackageFileLengthRequirements=true
+PackagesPerGC=50
+IdleTimeToGC=20
+MaxMemoryAllowance=8192
 +FullGCAssetClassNames=/Script/Engine.World
 +VersionedIniParams=Engine:/Script/Engine.RendererSettings:Compat.UseDXT5NormalMaps
 +VersionedIniParams=Engine:/Script/AndroidRuntimeSettings.AndroidRuntimeSettings:AndroidAudio
diff --git a/Engine/Source/Editor/UnrealEd/Classes/Commandlets/CookCommandlet.h b/Engine/Source/Editor/UnrealEd/Classes/Commandlets/CookCommandlet.h
--- a/Engine/Source/Editor/UnrealEd/Classes/Commandlets/CookCommandlet.h
+++ b/Engine/Source/Editor/UnrealEd/Classes/Commandlets/CookCommandlet.h
@@ -34,8 +34,6 @@ class UCookCommandlet
 	bool bUnversioned;
 	/** Generate manifests for building streaming install packages */
 	bool bGenerateStreamingInstallManifests;
-	/** Max memory the cooker should use before forcing a gc */
-	uint64 MaxMemoryAllowance;
 	/** All commandline tokens */
 	TArray<FString> Tokens;
 	/** All commandline switches */
@@ -119,7 +117,7 @@ private:
 	void MaybeMarkPackageAsAlreadyLoaded(UPackage *Package);
 
 	/** See if the cooker has exceeded max memory allowance in this case the cooker should force a garbage collection */
-	bool HasExceededMaxMemory() const;
+	bool HasExceededMaxMemory(uint64 MaxMemoryAllowance) const;
 
 	/** Gets the output directory respecting any command line overrides */
 	FString GetOutputDirectoryOverride() const;
diff --git a/Engine/Source/Editor/UnrealEd/Classes/CookOnTheSide/CookOnTheFlyServer.h b/Engine/Source/Editor/UnrealEd/Classes/CookOnTheSide/CookOnTheFlyServer.h
--- a/Engine/Source/Editor/UnrealEd/Classes/CookOnTheSide/CookOnTheFlyServer.h
+++ b/Engine/Source/Editor/UnrealEd/Classes/CookOnTheSide/CookOnTheFlyServer.h
@@ -623,6 +623,12 @@ private:
 	//////////////////////////////////////////////////////////////////////////
 	// General cook options
 	TArray<UClass*> FullGCAssetClasses;
+	/** Number of packages to load before performing a garbage collect. Set to 0 to never GC based on number of loaded packages */
+	uint32 PackagesPerGC;
+	/** Amount of time that is allowed to be idle before forcing a garbage collect. Set to 0 to never force GC due to idle time */
+	double IdleTimeToGC;
+	/** Max memory the cooker should use before forcing a gc */
+	uint64 MaxMemoryAllowance;
 
 	ECookInitializationFlags CookFlags;
 	TAutoPtr<class FSandboxPlatformFile> SandboxFile;
@@ -833,6 +839,15 @@ public:
 	 */
 	void SetFullGCAssetClasses( const TArray<UClass*>& InFullGCAssetClasses );
 
+	/** Returns the configured number of packages to process before GC */
+	uint32 GetPackagesPerGC() const;
+
+	/** Returns the configured amount of idle time before forcing a GC */
+	double GetIdleTimeToGC() const;
+
+	/** Returns the configured amount of memory allowed before forcing a GC */
+	uint64 GetMaxMemoryAllowance() const;
+
 	/**
 	 * Callbacks from editor 
 	 */
diff --git a/Engine/Source/Editor/UnrealEd/Private/Commandlets/CookCommandlet.cpp b/Engine/Source/Editor/UnrealEd/Private/Commandlets/CookCommandlet.cpp
--- a/Engine/Source/Editor/UnrealEd/Private/Commandlets/CookCommandlet.cpp
+++ b/Engine/Source/Editor/UnrealEd/Private/Commandlets/CookCommandlet.cpp
@@ -111,18 +111,19 @@ bool UCookCommandlet::CookOnTheFly( FGuid InstanceId, int32 Timeout, bool bForce
 	}
 
 	// Garbage collection should happen when either
-	//	1. We have cooked a map
+	//	1. We have cooked a map (configurable asset type)
 	//	2. We have cooked non-map packages and...
-	//		a. we have accumulated 50 of these since the last GC.
-	//		b. we have been idle for 20 seconds.
+	//		a. we have accumulated 50 (configurable) of these since the last GC.
+	//		b. we have been idle for 20 (configurable) seconds.
 	bool bShouldGC = true;
 
 	// megamoth
 	uint32 NonMapPackageCountSinceLastGC = 0;
 	
-	const int32 PackagesPerGC = 50;
-	
-	const double IdleTimeToGC = 20.0;
+	const uint32 PackagesPerGC = CookOnTheFlyServer->GetPackagesPerGC();
+	const double IdleTimeToGC = CookOnTheFlyServer->GetIdleTimeToGC();
+	const uint64 MaxMemoryAllowance = CookOnTheFlyServer->GetMaxMemoryAllowance();
+
 	double LastCookActionTime = FPlatformTime::Seconds();
 
 	FDateTime LastConnectionTime = FDateTime::UtcNow();
@@ -143,13 +144,13 @@ bool UCookCommandlet::CookOnTheFly( FGuid InstanceId, int32 Timeout, bool bForce
 
 		if (NonMapPackageCountSinceLastGC > 0)
 		{
-			if ( NonMapPackageCountSinceLastGC > PackagesPerGC)
+			if ((PackagesPerGC > 0) && (NonMapPackageCountSinceLastGC > PackagesPerGC))
 			{
 				UE_LOG(LogCookCommandlet, Display, TEXT("Cooker has exceeded max number of non map packages since last gc"));
 				bShouldGC |= true;
 			}
 
-			if (((FPlatformTime::Seconds() - LastCookActionTime) >= IdleTimeToGC))
+			if ((IdleTimeToGC > 0) && ((FPlatformTime::Seconds() - LastCookActionTime) >= IdleTimeToGC))
 			{
 				UE_LOG(LogCookCommandlet, Display, TEXT("Cooker has been idle for long time gc"));
 				bShouldGC |= true;
@@ -162,7 +163,7 @@ bool UCookCommandlet::CookOnTheFly( FGuid InstanceId, int32 Timeout, bool bForce
 			bShouldGC |= true;
 		}
 
-		if ( !bShouldGC && HasExceededMaxMemory() )
+		if ( !bShouldGC && HasExceededMaxMemory(MaxMemoryAllowance) )
 		{
 			UE_LOG(LogCookCommandlet, Display, TEXT("Cooker has exceeded max memory usage collecting garbage"));
 			bShouldGC |= true;
@@ -522,8 +523,6 @@ int32 UCookCommandlet::Main(const FString& CmdLineParams)
 	bIterativeCooking = Switches.Contains(TEXT("ITERATE"));
 	bSkipEditorContent = Switches.Contains(TEXT("SKIPEDITORCONTENT")); // This won't save out any packages in Engine/COntent/Editor*
 
-	MaxMemoryAllowance = 8LL * 1024LL * 1024LL * 1024LL;
-
 	if (bLeakTest)
 	{
 		for (FObjectIterator It; It; ++It)
@@ -1187,18 +1186,19 @@ bool UCookCommandlet::NewCook( const TArray<ITargetPlatform*>& Platforms, TArray
 	CookOnTheFlyServer->StartCookByTheBook( StartupOptions );
 
 	// Garbage collection should happen when either
-	//	1. We have cooked a map
+	//	1. We have cooked a map (configurable asset type)
 	//	2. We have cooked non-map packages and...
-	//		a. we have accumulated 50 of these since the last GC.
-	//		b. we have been idle for 20 seconds.
+	//		a. we have accumulated 50 (configurable) of these since the last GC.
+	//		b. we have been idle for 20 (configurable) seconds.
 	bool bShouldGC = true;
 
 	// megamoth
 	uint32 NonMapPackageCountSinceLastGC = 0;
 
-	const int32 PackagesPerGC = 50;
+	const uint32 PackagesPerGC = CookOnTheFlyServer->GetPackagesPerGC();
+	const double IdleTimeToGC = CookOnTheFlyServer->GetIdleTimeToGC();
+	const uint64 MaxMemoryAllowance = CookOnTheFlyServer->GetMaxMemoryAllowance();
 
-	const double IdleTimeToGC = 20.0;
 	double LastCookActionTime = FPlatformTime::Seconds();
 
 	FDateTime LastConnectionTime = FDateTime::UtcNow();
@@ -1221,13 +1221,14 @@ bool UCookCommandlet::NewCook( const TArray<ITargetPlatform*>& Platforms, TArray
 		if (NonMapPackageCountSinceLastGC > 0)
 		{
 			// We should GC if we have packages to collect and we've been idle for some time.
-			bShouldGC |= (NonMapPackageCountSinceLastGC > PackagesPerGC) || 
-				((FPlatformTime::Seconds() - LastCookActionTime) >= IdleTimeToGC);
+			const bool bExceededPackagesPerGC = (PackagesPerGC > 0) && (NonMapPackageCountSinceLastGC > PackagesPerGC);
+			const bool bExceededIdleTimeToGC = (IdleTimeToGC > 0) && ((FPlatformTime::Seconds() - LastCookActionTime) >= IdleTimeToGC);
+			bShouldGC |= bExceededPackagesPerGC || bExceededIdleTimeToGC;
 		}
 
 		bShouldGC |= (TickResults & UCookOnTheFlyServer::COSR_RequiresGC)!=0;
 
-		bShouldGC |= HasExceededMaxMemory();
+		bShouldGC |= HasExceededMaxMemory(MaxMemoryAllowance);
 
 
 		// don't clean up if we are waiting on cache of cooked data
@@ -1255,7 +1256,7 @@ bool UCookCommandlet::NewCook( const TArray<ITargetPlatform*>& Platforms, TArray
 	return true;
 }
 
-bool UCookCommandlet::HasExceededMaxMemory() const
+bool UCookCommandlet::HasExceededMaxMemory(uint64 MaxMemoryAllowance) const
 {
 	const FPlatformMemoryStats MemStats = FPlatformMemory::GetStats();
 
diff --git a/Engine/Source/Editor/UnrealEd/Private/CookOnTheFlyServer.cpp b/Engine/Source/Editor/UnrealEd/Private/CookOnTheFlyServer.cpp
--- a/Engine/Source/Editor/UnrealEd/Private/CookOnTheFlyServer.cpp
+++ b/Engine/Source/Editor/UnrealEd/Private/CookOnTheFlyServer.cpp
@@ -1595,6 +1595,21 @@ void UCookOnTheFlyServer::SetFullGCAssetClasses( const TArray<UClass*>& InFullGC
 	FullGCAssetClasses = InFullGCAssetClasses;
 }
 
+uint32 UCookOnTheFlyServer::GetPackagesPerGC() const
+{
+	return PackagesPerGC;
+}
+
+double UCookOnTheFlyServer::GetIdleTimeToGC() const
+{
+	return IdleTimeToGC;
+}
+
+uint64 UCookOnTheFlyServer::GetMaxMemoryAllowance() const
+{
+	return MaxMemoryAllowance;
+}
+
 void UCookOnTheFlyServer::BeginDestroy()
 {
 	EndNetworkFileServer();
@@ -1940,6 +1955,21 @@ void UCookOnTheFlyServer::Initialize( ECookMode::Type DesiredCookMode, ECookInit
 		// default to UWorld
 		FullGCAssetClasses.Add( UWorld::StaticClass() );
 	}
+
+	PackagesPerGC = 50;
+	int32 ConfigPackagesPerGC = 0;
+	if (GConfig->GetInt( TEXT("CookSettings"), TEXT("PackagesPerGC"), ConfigPackagesPerGC, GEditorIni ))
+	{
+		// Going unsigned. Make negative values 0
+		PackagesPerGC = ConfigPackagesPerGC > 0 ? ConfigPackagesPerGC : 0;
+	}
+
+	IdleTimeToGC = 20.0;
+	GConfig->GetDouble( TEXT("CookSettings"), TEXT("IdleTimeToGC"), IdleTimeToGC, GEditorIni );
+
+	int32 MaxMemoryAllowanceInMB = 8 * 1024;
+	GConfig->GetInt( TEXT("CookSettings"), TEXT("MaxMemoryAllowance"), MaxMemoryAllowanceInMB, GEditorIni );
+	MaxMemoryAllowance = MaxMemoryAllowanceInMB * 1024LL * 1024LL;
 	
 	if ( IsCookByTheBookMode() )
 	{
