Fixed: Material flattening produces black diffuse textures on some materails
UE-10959

diff --git a/Engine/Source/Developer/MeshUtilities/Private/MeshUtilities.cpp b/Engine/Source/Developer/MeshUtilities/Private/MeshUtilities.cpp
--- a/Engine/Source/Developer/MeshUtilities/Private/MeshUtilities.cpp
+++ b/Engine/Source/Developer/MeshUtilities/Private/MeshUtilities.cpp
@@ -3668,6 +3668,7 @@ void FMeshUtilities::CreateProxyMesh(
 	
 	const FString AssetBaseName = FPackageName::GetShortName(ProxyBasePackageName);
 	const FString AssetBasePath = InOuter ? TEXT("") : FPackageName::GetLongPackagePath(ProxyBasePackageName) + TEXT("/");
+	UWorld* InWorld = SourceActors.Num() ? SourceActors[0]->GetWorld() : nullptr;
 	
 	TArray<UStaticMeshComponent*> ComponentsToMerge;
 	
@@ -3731,7 +3732,7 @@ void FMeshUtilities::CreateProxyMesh(
 		for (UMaterialInterface* Material : StaticMeshMaterials)
 		{
 			UniqueMaterials.Add(FlattenMaterial);
-			MaterialExportUtils::ExportMaterial(Material, UniqueMaterials.Last());
+			MaterialExportUtils::ExportMaterial(InWorld, Material, UniqueMaterials.Last());
 		}
 	}
 		
diff --git a/Engine/Source/Editor/UnrealEd/Private/EditorExporters.cpp b/Engine/Source/Editor/UnrealEd/Private/EditorExporters.cpp
--- a/Engine/Source/Editor/UnrealEd/Private/EditorExporters.cpp
+++ b/Engine/Source/Editor/UnrealEd/Private/EditorExporters.cpp
@@ -37,6 +37,7 @@
 #include "Engine/TextureCube.h"
 #include "EngineUtils.h"
 #include "DeviceProfiles/DeviceProfileManager.h"
+#include "TileRendering.h"
 
 DEFINE_LOG_CATEGORY_STATIC(LogEditorExporters, Log, All);
 
@@ -2447,19 +2448,87 @@ UExportTextContainer::UExportTextContainer(const FObjectInitializer& ObjectIniti
 {
 }
 
-
-
-
 /*-----------------------------------------------------------------------------
 	Material export utilities
 -----------------------------------------------------------------------------*/
 namespace MaterialExportUtils
 {
+	void RenderMaterialTile(UWorld* InWorld, FMaterialRenderProxy* InMaterialProxy, UTextureRenderTarget2D* InRenderTarget)
+	{
+		float CurrentRealTime = 0.f;
+		float CurrentWorldTime = 0.f;
+		float DeltaWorldTime = 0.f;
+
+		if (InWorld)
+		{
+			CurrentRealTime = InWorld->GetRealTimeSeconds();
+			CurrentWorldTime = InWorld->GetTimeSeconds();
+			DeltaWorldTime = InWorld->GetDeltaSeconds();
+		}
+		else
+		{
+			CurrentRealTime = FApp::GetCurrentTime() - GStartTime;
+			CurrentWorldTime = FApp::GetDeltaTime();
+			DeltaWorldTime = FApp::GetCurrentTime() - GStartTime;
+		}
+
+		const FRenderTarget* RenderTargetResource = InRenderTarget->GameThread_GetRenderTargetResource();
+		FSceneViewFamily* ViewFamily = new FSceneViewFamily(FSceneViewFamily::ConstructionValues(
+			RenderTargetResource,
+			InWorld ? InWorld->Scene : nullptr,
+			FEngineShowFlags(ESFIM_Game))
+			.SetWorldTimes(CurrentWorldTime, DeltaWorldTime, CurrentRealTime)
+			.SetGammaCorrection(RenderTargetResource->GetDisplayGamma()));
+
+		FIntPoint ViewSize = RenderTargetResource->GetSizeXY();
+		FIntRect ViewRect(FIntPoint(0, 0), ViewSize);
+		// YZ - Front view
+		FMatrix ViewMatrix = FMatrix(
+						FPlane(1, 0, 0, 0),
+						FPlane(0, 0, -1, 0),
+						FPlane(0, 1, 0, 0),
+						FPlane(0, 0, 0, 1));
+					
+		// make a temporary view
+		FSceneViewInitOptions ViewInitOptions;
+		ViewInitOptions.ViewFamily = ViewFamily;
+		ViewInitOptions.SetViewRectangle(ViewRect);
+		ViewInitOptions.ViewOrigin = FVector(ViewSize.X/2.f, ViewSize.Y/2.f, 0.0f);
+		ViewInitOptions.ViewRotationMatrix = ViewMatrix;
+		ViewInitOptions.ProjectionMatrix = FReversedZOrthoMatrix(ViewSize.X/2.f, ViewSize.Y/2.f, 0.5f/HALF_WORLD_MAX, HALF_WORLD_MAX);
+		ViewInitOptions.BackgroundColor = FLinearColor::Black;
+		ViewInitOptions.OverlayColor = FLinearColor::White;
+				
+		FSceneView* View = new FSceneView(ViewInitOptions);
+						
+		ENQUEUE_UNIQUE_RENDER_COMMAND_FOURPARAMETER(
+			RenderMaterialTileCommand,
+			FSceneView*, InView, View,
+			FMaterialRenderProxy*, InProxy, InMaterialProxy,
+			const FRenderTarget*, InRenderTarget, RenderTargetResource,
+			FIntPoint, InSize, ViewSize,
+			{
+				::SetRenderTarget(RHICmdList, InRenderTarget->GetRenderTargetTexture(), FTextureRHIRef());
+				FIntRect ViewportRect = FIntRect(FIntPoint::ZeroValue, InSize);
+				RHICmdList.SetViewport(ViewportRect.Min.X, ViewportRect.Min.Y, 0.0f, ViewportRect.Max.X, ViewportRect.Max.Y, 1.0f);
+				
+				// We want to render material into tile that has vertices on YZ plane
+				// By default tile mesh created on XY plane, so rotate it
+				FQuat Rotation = FRotator(0.f, 0.f, 90).Quaternion();
+				FTileRenderer::DrawRotatedTile(RHICmdList, *InView, InProxy, false, Rotation, 0.f, 0.f, InSize.X, InSize.Y,  0.f, 0.f, 1.f, 1.f);
+			});
+
+		FlushRenderingCommands();
+
+		delete View;
+	}
 	
-	bool ExportMaterialProperty(UMaterialInterface* InMaterial, 
-								EMaterialProperty InMaterialProperty, 
-								UTextureRenderTarget2D* InRenderTarget, 
-								TArray<FColor>& OutBMP)
+	bool ExportMaterialProperty(UMaterialInterface* InMaterial,	EMaterialProperty InMaterialProperty, UTextureRenderTarget2D* InRenderTarget, TArray<FColor>& OutBMP)
+	{
+		return ExportMaterialProperty(nullptr, InMaterial, InMaterialProperty, InRenderTarget, OutBMP);
+	}
+		
+	bool ExportMaterialProperty(UWorld* InWorld, UMaterialInterface* InMaterial, EMaterialProperty InMaterialProperty, UTextureRenderTarget2D* InRenderTarget, TArray<FColor>& OutBMP)
 	{
 		TScopedPointer<FExportMaterialProxy> MaterialProxy(new FExportMaterialProxy(InMaterial, InMaterialProperty));
 		if (MaterialProxy == NULL)
@@ -2475,31 +2544,25 @@ namespace MaterialExportUtils
 			OutBMP.Add(UniformValue);
 			return true;
 		}
-	
-		check(InRenderTarget);
-		FTextureRenderTargetResource* RTResource = InRenderTarget->GameThread_GetRenderTargetResource();
 
-		{
-			// Create a canvas for the render target and clear it to black
-			FCanvas Canvas(RTResource, NULL, FApp::GetCurrentTime() - GStartTime, FApp::GetDeltaTime(), FApp::GetCurrentTime() - GStartTime, GMaxRHIFeatureLevel);
-			Canvas.Clear(FLinearColor::Black);
-			FCanvasTileItem TileItem(FVector2D(0.0f, 0.0f), MaterialProxy, FVector2D(InRenderTarget->SizeX, InRenderTarget->SizeY));
-			TileItem.bFreezeTime = true;
-			Canvas.DrawItem( TileItem );
-			Canvas.Flush_GameThread();
-			FlushRenderingCommands();
-			Canvas.SetRenderTarget_GameThread(NULL);
-			FlushRenderingCommands();
-		}
-				
+		check(InRenderTarget);
+		
+		RenderMaterialTile(InWorld, MaterialProxy, InRenderTarget);
+		
 		bool bNormalmap = (InMaterialProperty == MP_Normal);
 		FReadSurfaceDataFlags ReadPixelFlags(bNormalmap ? RCM_SNorm : RCM_UNorm);
 		ReadPixelFlags.SetLinearToGamma(false);
-
+		
+		FTextureRenderTargetResource* RTResource = InRenderTarget->GameThread_GetRenderTargetResource();
 		return RTResource->ReadPixels(OutBMP, ReadPixelFlags);
 	}
 	
 	bool ExportMaterial(UMaterialInterface* InMaterial, FFlattenMaterial& OutFlattenMaterial)
+	{
+		return ExportMaterial(nullptr, InMaterial, OutFlattenMaterial);
+	}
+	
+	bool ExportMaterial(UWorld* InWorld, UMaterialInterface* InMaterial, FFlattenMaterial& OutFlattenMaterial)
 	{
 		// Render diffuse property
 		if (OutFlattenMaterial.DiffuseSize.X > 0 && 
diff --git a/Engine/Source/Editor/UnrealEd/Public/MaterialExportUtils.h b/Engine/Source/Editor/UnrealEd/Public/MaterialExportUtils.h
--- a/Engine/Source/Editor/UnrealEd/Public/MaterialExportUtils.h
+++ b/Engine/Source/Editor/UnrealEd/Public/MaterialExportUtils.h
@@ -35,23 +35,25 @@ namespace MaterialExportUtils
 	/**
 	 * Renders specified material property into texture
 	 *
+	 * @param InWorld				World object to use for material property rendering
 	 * @param InMaterial			Target material
 	 * @param InMaterialProperty	Material property to render
 	 * @param InRenderTarget		Render target to render to
 	 * @param OutBMP				Output array of rendered samples 
 	 * @return						Whether operation was successful
 	 */
-	UNREALED_API bool ExportMaterialProperty(UMaterialInterface* InMaterial, 
-												EMaterialProperty InMaterialProperty, 
-												UTextureRenderTarget2D* InRenderTarget, 
-												TArray<FColor>& OutBMP);
+	UNREALED_API bool ExportMaterialProperty(UWorld* InWorld, UMaterialInterface* InMaterial, EMaterialProperty InMaterialProperty, UTextureRenderTarget2D* InRenderTarget, TArray<FColor>& OutBMP);
+	UNREALED_API bool ExportMaterialProperty(UMaterialInterface* InMaterial, EMaterialProperty InMaterialProperty, UTextureRenderTarget2D* InRenderTarget, TArray<FColor>& OutBMP);
+
 	/**
 	 * Flattens specified material
 	 *
+	 * @param InWorld				World object to use for material rendering
 	 * @param InMaterial			Target material
 	 * @param OutFlattenMaterial	Output flattened material
 	 * @return						Whether operation was successful
 	 */
+	UNREALED_API bool ExportMaterial(UWorld* InWorld, UMaterialInterface* InMaterial, FFlattenMaterial& OutFlattenMaterial);
 	UNREALED_API bool ExportMaterial(UMaterialInterface* InMaterial, FFlattenMaterial& OutFlattenMaterial);
 	
 	/**
diff --git a/Engine/Source/Runtime/Engine/Private/TileRendering.cpp b/Engine/Source/Runtime/Engine/Private/TileRendering.cpp
--- a/Engine/Source/Runtime/Engine/Private/TileRendering.cpp
+++ b/Engine/Source/Runtime/Engine/Private/TileRendering.cpp
@@ -139,11 +139,8 @@ public:
 };
 TGlobalResource<FTileMesh> GTileMesh;
 
-void FTileRenderer::DrawTile(FRHICommandListImmediate& RHICmdList, const class FSceneView& View, const FMaterialRenderProxy* MaterialRenderProxy, bool bNeedsToSwitchVerticalAxis, float X, float Y, float SizeX, float SizeY, float U, float V, float SizeU, float SizeV, bool bIsHitTesting, const FHitProxyId HitProxyId, const FColor InVertexColor)
+static void CreateTileVerices(FMaterialTileVertex DestVertex[4], const class FSceneView& View, bool bNeedsToSwitchVerticalAxis, float X, float Y, float SizeX, float SizeY, float U, float V, float SizeU, float SizeV, const FColor InVertexColor)
 {
-	FMaterialTileVertex DestVertex[4];
-
-	// create verts
 	if (bNeedsToSwitchVerticalAxis)
 	{
 		DestVertex[0].Initialize(X + SizeX, View.ViewRect.Height() - (Y + SizeY), U + SizeU, V + SizeV);
@@ -163,6 +160,38 @@ void FTileRenderer::DrawTile(FRHICommandListImmediate& RHICmdList, const class F
 	DestVertex[1].Color = InVertexColor.DWColor();
 	DestVertex[2].Color = InVertexColor.DWColor();
 	DestVertex[3].Color = InVertexColor.DWColor();
+}
+
+void FTileRenderer::DrawTile(FRHICommandListImmediate& RHICmdList, const class FSceneView& View, const FMaterialRenderProxy* MaterialRenderProxy, bool bNeedsToSwitchVerticalAxis, float X, float Y, float SizeX, float SizeY, float U, float V, float SizeU, float SizeV, bool bIsHitTesting, const FHitProxyId HitProxyId, const FColor InVertexColor)
+{
+	// create verts
+	FMaterialTileVertex DestVertex[4];
+	CreateTileVerices(DestVertex, View, bNeedsToSwitchVerticalAxis, X, Y, SizeX, SizeY, U, V, SizeU, SizeV, InVertexColor);
+
+	// update the FMeshBatch
+	FMeshBatch& Mesh = GTileMesh.MeshElement;
+	Mesh.UseDynamicData = true;
+	Mesh.DynamicVertexData = DestVertex;
+	Mesh.MaterialRenderProxy = MaterialRenderProxy;
+
+	GetRendererModule().DrawTileMesh(RHICmdList, View, Mesh, bIsHitTesting, HitProxyId);
+}
+
+void FTileRenderer::DrawRotatedTile(FRHICommandListImmediate& RHICmdList, const class FSceneView& View, const FMaterialRenderProxy* MaterialRenderProxy, bool bNeedsToSwitchVerticalAxis, const FQuat& Rotation, float X, float Y, float SizeX, float SizeY, float U, float V, float SizeU, float SizeV, bool bIsHitTesting, const FHitProxyId HitProxyId, const FColor InVertexColor)
+{
+	// create verts
+	FMaterialTileVertex DestVertex[4];
+	CreateTileVerices(DestVertex, View, bNeedsToSwitchVerticalAxis, X, Y, SizeX, SizeY, U, V, SizeU, SizeV, InVertexColor);
+	
+	// rotate tile using view center as origin
+	FIntPoint ViewRectSize = View.ViewRect.Size();
+	FVector ViewRectOrigin = FVector(ViewRectSize.X*0.5f, ViewRectSize.Y*0.5f, 0.f);
+	for (int32 i = 0; i < 4; ++i)
+	{
+		DestVertex[i].Position = Rotation.RotateVector(DestVertex[i].Position - ViewRectOrigin) + ViewRectOrigin;
+		DestVertex[i].TangentX = Rotation.RotateVector(DestVertex[i].TangentX);
+		DestVertex[i].TangentZ = Rotation.RotateVector(DestVertex[i].TangentZ);
+	}
 
 	// update the FMeshBatch
 	FMeshBatch& Mesh = GTileMesh.MeshElement;
diff --git a/Engine/Source/Runtime/Engine/Public/TileRendering.h b/Engine/Source/Runtime/Engine/Public/TileRendering.h
--- a/Engine/Source/Runtime/Engine/Public/TileRendering.h
+++ b/Engine/Source/Runtime/Engine/Public/TileRendering.h
@@ -16,6 +16,12 @@ public:
 	 * (UV = [0..1]
 	 */
 	ENGINE_API static void DrawTile(FRHICommandListImmediate& RHICmdList, const class FSceneView& View, const FMaterialRenderProxy* MaterialRenderProxy, bool bNeedsToSwitchVerticalAxis, float X, float Y, float SizeX, float SizeY, float U, float V, float SizeU, float SizeV, bool bIsHitTesting = false, const FHitProxyId HitProxyId = FHitProxyId(), const FColor InVertexColor = FColor(255, 255, 255, 255));
+
+	/**
+	 * Draw a rotated tile at the given location and size, using the given UVs
+	 * (UV = [0..1]
+	 */
+	ENGINE_API static void DrawRotatedTile(FRHICommandListImmediate& RHICmdList, const class FSceneView& View, const FMaterialRenderProxy* MaterialRenderProxy, bool bNeedsToSwitchVerticalAxis, const FQuat& Rotation, float X, float Y, float SizeX, float SizeY, float U, float V, float SizeU, float SizeV, bool bIsHitTesting = false, const FHitProxyId HitProxyId = FHitProxyId(), const FColor InVertexColor = FColor(255, 255, 255, 255));
 	
 private:
 
