Drag-moving an actor in a viewport will be aborted and will allow PIE to start. PIE will now fail to start and notify the user when an edit/transaction is in progress.
[UE-11173] - Hitting S while Alt + Dragging a BSP and then Stopping PIE will crash the editor

diff --git a/Engine/Source/Editor/UnrealEd/Private/Editor.cpp b/Engine/Source/Editor/UnrealEd/Private/Editor.cpp
--- a/Engine/Source/Editor/UnrealEd/Private/Editor.cpp
+++ b/Engine/Source/Editor/UnrealEd/Private/Editor.cpp
@@ -120,6 +120,7 @@ FSimpleMulticastDelegate								FEditorDelegates::LoadSelectedAssetsIfNeeded;
 FSimpleMulticastDelegate								FEditorDelegates::DisplayLoadErrors;
 FEditorDelegates::FOnEditorModeTransitioned				FEditorDelegates::EditorModeEnter;
 FEditorDelegates::FOnEditorModeTransitioned				FEditorDelegates::EditorModeExit;
+FEditorDelegates::FOnPIEEvent							FEditorDelegates::PreBeginPIE;
 FEditorDelegates::FOnPIEEvent							FEditorDelegates::BeginPIE;
 FEditorDelegates::FOnPIEEvent							FEditorDelegates::EndPIE;
 FEditorDelegates::FOnPIEEvent							FEditorDelegates::PausePIE;
diff --git a/Engine/Source/Editor/UnrealEd/Private/LevelEditorViewport.cpp b/Engine/Source/Editor/UnrealEd/Private/LevelEditorViewport.cpp
--- a/Engine/Source/Editor/UnrealEd/Private/LevelEditorViewport.cpp
+++ b/Engine/Source/Editor/UnrealEd/Private/LevelEditorViewport.cpp
@@ -1506,6 +1506,9 @@ FLevelEditorViewportClient::FLevelEditorViewportClient(const TSharedPtr<SLevelVi
 	// Register for editor cleanse events so we can release references to hovered actors
 	FEditorSupportDelegates::CleanseEditor.AddRaw(this, &FLevelEditorViewportClient::OnEditorCleanse);
 
+	// Register for editor PIE event that allows you to clean up states that might block PIE
+	FEditorDelegates::PreBeginPIE.AddRaw(this, &FLevelEditorViewportClient::OnPreBeginPIE);
+
 	// Add a delegate so we get informed when an actor has moved.
 	GEngine->OnActorMoved().AddRaw(this, &FLevelEditorViewportClient::OnActorMoved);
 
@@ -1526,6 +1529,7 @@ FLevelEditorViewportClient::~FLevelEditorViewportClient()
 {
 	// Unregister for all global callbacks to this object
 	FEditorSupportDelegates::CleanseEditor.RemoveAll(this);
+	FEditorDelegates::PreBeginPIE.RemoveAll(this);
 
 	// Remove our move delegate
 	GEngine->OnActorMoved().RemoveAll(this);
@@ -4194,6 +4198,12 @@ void FLevelEditorViewportClient::OnEditorCleanse()
 	ClearHoverFromObjects();
 }
 
+void FLevelEditorViewportClient::OnPreBeginPIE(const bool bIsSimulating)
+{
+	// Called before PIE attempts to start, allowing the viewport to cancel processes, like dragging, that will block PIE from beginning
+	AbortTracking();
+}
+
 bool FLevelEditorViewportClient::GetSpriteCategoryVisibility( const FName& InSpriteCategory ) const
 {
 	const int32 CategoryIndex = GEngine->GetSpriteCategoryIndex( InSpriteCategory );
diff --git a/Engine/Source/Editor/UnrealEd/Private/PlayLevel.cpp b/Engine/Source/Editor/UnrealEd/Private/PlayLevel.cpp
--- a/Engine/Source/Editor/UnrealEd/Private/PlayLevel.cpp
+++ b/Engine/Source/Editor/UnrealEd/Private/PlayLevel.cpp
@@ -1885,8 +1885,34 @@ bool UEditorEngine::SpawnPlayFromHereStart( UWorld* World, AActor*& PlayerStart,
 
 void UEditorEngine::PlayInEditor( UWorld* InWorld, bool bInSimulateInEditor )
 {
+	// Broadcast PreBeginPIE before checks that might block PIE below (BeginPIE is broadcast below after the checks)
+	FEditorDelegates::PreBeginPIE.Broadcast(bInSimulateInEditor);
+
 	double PIEStartTime = FPlatformTime::Seconds();
 
+	// Block PIE when there is a transaction recording into the undo buffer
+	if (GEditor->IsTransactionActive())
+	{
+		FFormatNamedArguments Args;
+		Args.Add(TEXT("TransactionName"), GEditor->GetTransactionName());
+
+		FText NotificationText;
+		if (bInSimulateInEditor)
+		{
+			NotificationText = FText::Format(NSLOCTEXT("UnrealEd", "SIECantStartDuringTransaction", "Can't Simulate when performing {TransactionName} operation"), Args);
+		}
+		else
+		{
+			NotificationText = FText::Format(NSLOCTEXT("UnrealEd", "PIECantStartDuringTransaction", "Can't Play In Editor when performing {TransactionName} operation"), Args);
+		}
+
+		FNotificationInfo Info(NotificationText);
+		Info.ExpireDuration = 5.0f;
+		Info.bUseLargeFont = true;
+		FSlateNotificationManager::Get().AddNotification(Info);
+		return;
+	}
+
 	// Prompt the user that Matinee must be closed before PIE can occur.
 	if( GLevelEditorModeTools().IsModeActive(FBuiltinEditorModes::EM_InterpEdit) )
 	{
@@ -1903,6 +1929,7 @@ void UEditorEngine::PlayInEditor( UWorld* InWorld, bool bInSimulateInEditor )
 
 	FBlueprintEditorUtils::FindAndSetDebuggableBlueprintInstances();
 
+	// Broadcast BeginPIE after checks that might block PIE above (PreBeginPIE is broadcast above before the checks)
 	FEditorDelegates::BeginPIE.Broadcast(bInSimulateInEditor);
 
 	// let navigation know PIE starts so it can avoid any blueprint creation/deletion/instantiation affect editor map's navmesh changes
diff --git a/Engine/Source/Editor/UnrealEd/Public/Editor.h b/Engine/Source/Editor/UnrealEd/Public/Editor.h
--- a/Engine/Source/Editor/UnrealEd/Public/Editor.h
+++ b/Engine/Source/Editor/UnrealEd/Public/Editor.h
@@ -135,6 +135,8 @@ struct UNREALED_API FEditorDelegates
 	static FOnEditorModeTransitioned EditorModeEnter;
 	/** Called when an editor mode is being exited */
 	static FOnEditorModeTransitioned EditorModeExit;
+	/** Sent when a PIE session is beginning (before we decide if PIE can run - allows clients to avoid blocking PIE) */
+	static FOnPIEEvent PreBeginPIE;
 	/** Sent when a PIE session is beginning */
 	static FOnPIEEvent BeginPIE;
 	/** Sent when a PIE session is ending */
diff --git a/Engine/Source/Editor/UnrealEd/Public/LevelEditorViewport.h b/Engine/Source/Editor/UnrealEd/Public/LevelEditorViewport.h
--- a/Engine/Source/Editor/UnrealEd/Public/LevelEditorViewport.h
+++ b/Engine/Source/Editor/UnrealEd/Public/LevelEditorViewport.h
@@ -497,6 +497,9 @@ protected:
 	/** Called when editor cleanse event is triggered */
 	void OnEditorCleanse();
 
+	/** Called before the editor tries to begin PIE */
+	void OnPreBeginPIE(const bool bIsSimulating);
+
 	/** Callback for when an editor user setting has changed */
 	void HandleViewportSettingChanged(FName PropertyName);
 
