Tweaked coupld of some meta tags used by tutorials (viewport and toolbar)

diff --git a/Engine/Source/Editor/LevelEditor/Private/SLevelEditor.cpp b/Engine/Source/Editor/LevelEditor/Private/SLevelEditor.cpp
--- a/Engine/Source/Editor/LevelEditor/Private/SLevelEditor.cpp
+++ b/Engine/Source/Editor/LevelEditor/Private/SLevelEditor.cpp
@@ -560,7 +560,7 @@ TSharedRef<SDockTab> SLevelEditor::SpawnLevelEditorTab( const FSpawnTabArgs& Arg
 			.Icon( FEditorStyle::GetBrush("ToolBar.Icon") )
 			[
 				SNew(SHorizontalBox)
-				.AddMetaData<FTagMetaData>(FTagMetaData(TEXT("LevelToolbar")))
+				.AddMetaData<FTagMetaData>(FTagMetaData(TEXT("LevelEditorToolbar")))
 				+SHorizontalBox::Slot()
 				.FillWidth(1)
 				.VAlign(VAlign_Bottom)
diff --git a/Engine/Source/Editor/UnrealEd/Private/SEditorViewport.cpp b/Engine/Source/Editor/UnrealEd/Private/SEditorViewport.cpp
--- a/Engine/Source/Editor/UnrealEd/Private/SEditorViewport.cpp
+++ b/Engine/Source/Editor/UnrealEd/Private/SEditorViewport.cpp
@@ -41,7 +41,7 @@ void SEditorViewport::Construct( const FArguments& InArgs )
 		SAssignNew( ViewportWidget, SViewport )
 		.ShowEffectWhenDisabled( false )
 		.EnableGammaCorrection( false ) // Scene rendering handles this
-		.AddMetaData(InArgs.MetaData.Num() > 0 ? InArgs.MetaData[0] : MakeShareable(new FTagMetaData(TEXT("EditorViewports"))))
+		.AddMetaData(InArgs.MetaData.Num() > 0 ? InArgs.MetaData[0] : MakeShareable(new FTagMetaData(TEXT("LevelEditorViewport"))))
 		[
 			SAssignNew( ViewportOverlay, SOverlay )
 			+SOverlay::Slot()
