Fixed cloth assets with lower resolution simulation meshes crashing as mapping data seems to only be stored for sim verts.
#UE-10674

diff --git a/Engine/Source/Editor/UnrealEd/Private/ApexClothingUtils.cpp b/Engine/Source/Editor/UnrealEd/Private/ApexClothingUtils.cpp
--- a/Engine/Source/Editor/UnrealEd/Private/ApexClothingUtils.cpp
+++ b/Engine/Source/Editor/UnrealEd/Private/ApexClothingUtils.cpp
@@ -268,8 +268,6 @@ bool LoadPhysicalMeshFromClothingAsset(NxClothingAsset& ApexClothingAsset,
 	// Mapping
 	uint32 EndTotalIndex = StartSimulIndex + NumTotalVertices;
 
-	check(EndTotalIndex <= MapSize);
-
 	RenderToPhysicalMapping.Empty(NumTotalVertices);
 
 	uint16 MaxSimulVertIndex = 0;
@@ -277,18 +275,21 @@ bool LoadPhysicalMeshFromClothingAsset(NxClothingAsset& ApexClothingAsset,
 	for(uint32 i=StartSimulIndex; i < EndTotalIndex; i++)
 	{
 		FApexClothPhysToRenderVertData& Mapping = RenderToPhysicalMapping[RenderToPhysicalMapping.AddZeroed()];
-		Mapping.PositionBaryCoordsAndDist = P2U4BaryCoord(SkinningMap[i].positionBary);
-		Mapping.NormalBaryCoordsAndDist = P2U4BaryCoord(SkinningMap[i].normalBary);
-		Mapping.TangentBaryCoordsAndDist = P2U4BaryCoord(SkinningMap[i].tangentBary);
-		Mapping.SimulMeshVertIndices[0] = (uint16)SkinningMap[i].vertexIndex0;
-		Mapping.SimulMeshVertIndices[1] = (uint16)SkinningMap[i].vertexIndex1;
-		Mapping.SimulMeshVertIndices[2] = (uint16)SkinningMap[i].vertexIndex2;
-		Mapping.SimulMeshVertIndices[3] = (uint16)0; // store valid index to distinguish from fixed vertex
-
-		// to extract only clothing section's simulation index, because apex file includes non-clothing sections info as well.
-		for(int ArrayIdx=0; ArrayIdx < 3; ArrayIdx++)
+		if(i < NumRealSimulVertices)
 		{
-			MaxSimulVertIndex = FMath::Max(MaxSimulVertIndex, Mapping.SimulMeshVertIndices[ArrayIdx]);
+			Mapping.PositionBaryCoordsAndDist = P2U4BaryCoord(SkinningMap[i].positionBary);
+			Mapping.NormalBaryCoordsAndDist = P2U4BaryCoord(SkinningMap[i].normalBary);
+			Mapping.TangentBaryCoordsAndDist = P2U4BaryCoord(SkinningMap[i].tangentBary);
+			Mapping.SimulMeshVertIndices[0] = (uint16)SkinningMap[i].vertexIndex0;
+			Mapping.SimulMeshVertIndices[1] = (uint16)SkinningMap[i].vertexIndex1;
+			Mapping.SimulMeshVertIndices[2] = (uint16)SkinningMap[i].vertexIndex2;
+			Mapping.SimulMeshVertIndices[3] = (uint16)0; // store valid index to distinguish from fixed vertex
+
+			// to extract only clothing section's simulation index, because apex file includes non-clothing sections info as well.
+			for(int ArrayIdx = 0; ArrayIdx < 3; ArrayIdx++)
+			{
+				MaxSimulVertIndex = FMath::Max(MaxSimulVertIndex, Mapping.SimulMeshVertIndices[ArrayIdx]);
+			}
 		}
 	}
 
