Set default value for vector nodes when compiling

diff --git a/Engine/Source/Editor/UnrealEd/Private/Niagara/NiagaraScriptSource.cpp b/Engine/Source/Editor/UnrealEd/Private/Niagara/NiagaraScriptSource.cpp
--- a/Engine/Source/Editor/UnrealEd/Private/Niagara/NiagaraScriptSource.cpp
+++ b/Engine/Source/Editor/UnrealEd/Private/Niagara/NiagaraScriptSource.cpp
@@ -114,6 +114,7 @@ void UNiagaraScriptSource::Compile()
 			{
 				EditorExposedVectorConstant *Const = new EditorExposedVectorConstant();
 				Const->ConstName = Node->Input.Name;
+				Const->Value = Node->VectorDefault;
 				ExposedVectorConstants.Add( MakeShareable(Const) );
 			}
 			else if (Node->Input.Type == ENiagaraDataType::Curve)
