Editor - Added FRunPIEAndExit automation test

diff --git a/Engine/Source/Editor/UnrealEd/Private/Tests/EditorAutomationTests.cpp b/Engine/Source/Editor/UnrealEd/Private/Tests/EditorAutomationTests.cpp
--- a/Engine/Source/Editor/UnrealEd/Private/Tests/EditorAutomationTests.cpp
+++ b/Engine/Source/Editor/UnrealEd/Private/Tests/EditorAutomationTests.cpp
@@ -2117,5 +2117,29 @@ bool FLaunchOnTest::RunTest(const FString& Parameters)
 	ADD_LATENT_AUTOMATION_COMMAND(FDeleteDirCommand(StagedBuildsLocation));
 	ADD_LATENT_AUTOMATION_COMMAND(FDeleteDirCommand(TempMapLocation));
 
+	return true;
+}
+
+/**
+ * Pie Test
+ * Verification PIE works
+ */
+IMPLEMENT_SIMPLE_AUTOMATION_TEST( FRunPIEAndExitTest, "RunPIEAndExit", EAutomationTestFlags::ATF_Editor )
+
+/** 
+ * Execute the loading of one map to verify PIE works
+ *
+ * @param Parameters - Unused for this test
+ * @return	TRUE if the test was successful, FALSE otherwise
+ */
+ bool FRunPIEAndExitTest::RunTest( const FString& Parameters )
+{
+	ADD_LATENT_AUTOMATION_COMMAND( FStartPIECommand( true ) );
+	ADD_LATENT_AUTOMATION_COMMAND( FWaitLatentCommand( 3.0f ) );
+	ADD_LATENT_AUTOMATION_COMMAND( FEndPlayMapCommand() );
+
+	ADD_LATENT_AUTOMATION_COMMAND( FWaitLatentCommand( 3.0f ) );
+	ADD_LATENT_AUTOMATION_COMMAND( FRequestExitCommand() );
+
 	return true;
 }
\ No newline at end of file
