#ue4 - Convert uses of FRotationMatrix(GetActorRotation()).TransformVector(...) to GetActorQuat().RotateVector(...) which is equivalent but faster.

UE-6964

diff --git a/Engine/Source/Editor/UnrealEd/Private/LevelEditorViewport.cpp b/Engine/Source/Editor/UnrealEd/Private/LevelEditorViewport.cpp
--- a/Engine/Source/Editor/UnrealEd/Private/LevelEditorViewport.cpp
+++ b/Engine/Source/Editor/UnrealEd/Private/LevelEditorViewport.cpp
@@ -2134,9 +2134,6 @@ void FLevelEditorViewportClient::ProjectActorsIntoWorld(const TArray<AActor*>& A
 
 			const auto* Factory = GEditor->FindActorFactoryForActorClass(Actor->GetClass());
 			
-			const auto ActorLocation = Actor->GetActorLocation();
-			const auto ActorRotation = Actor->GetActorRotation().Quaternion();
-
 			const FTransform* PreDragActorTransform = PreDragActorTransforms.Find(Actor);
 			check(PreDragActorTransform);
 
