Editor: Renamed IAsyncTask to IAsyncFileCacheTask, so the name can be used for more generic purposes.

diff --git a/Engine/Source/Editor/UnrealEd/Private/AutoReimport/FileCache.cpp b/Engine/Source/Editor/UnrealEd/Private/AutoReimport/FileCache.cpp
--- a/Engine/Source/Editor/UnrealEd/Private/AutoReimport/FileCache.cpp
+++ b/Engine/Source/Editor/UnrealEd/Private/AutoReimport/FileCache.cpp
@@ -105,12 +105,12 @@ static const uint32 CacheFileMagicNumber = 0x03DCCB00;
 /** Single runnable thread used to parse file cache directories without blocking the main thread */
 struct FAsyncTaskThread : public FRunnable
 {
-	typedef TArray<TWeakPtr<IAsyncTask, ESPMode::ThreadSafe>> FTaskArray;
+	typedef TArray<TWeakPtr<IAsyncFileCacheTask, ESPMode::ThreadSafe>> FTaskArray;
 
 	FAsyncTaskThread() : Thread(nullptr) {}
 
 	/** Add a reader to this thread which will get ticked periodically until complete */
-	void AddTask(TSharedPtr<IAsyncTask, ESPMode::ThreadSafe> InTask)
+	void AddTask(TSharedPtr<IAsyncFileCacheTask, ESPMode::ThreadSafe> InTask)
 	{
 		FScopeLock Lock(&TaskArrayMutex);
 		Tasks.Add(InTask);
@@ -220,7 +220,7 @@ bool FAsyncFileHasher::IsComplete() const
 	return CurrentIndex.GetValue() == Data.Num();
 }
 
-IAsyncTask::EProgressResult FAsyncFileHasher::Tick(const FTimeLimit& Limit)
+IAsyncFileCacheTask::EProgressResult FAsyncFileHasher::Tick(const FTimeLimit& Limit)
 {
 	for (; CurrentIndex.GetValue() < Data.Num(); )
 	{
diff --git a/Engine/Source/Editor/UnrealEd/Private/AutoReimport/FileCache.h b/Engine/Source/Editor/UnrealEd/Private/AutoReimport/FileCache.h
--- a/Engine/Source/Editor/UnrealEd/Private/AutoReimport/FileCache.h
+++ b/Engine/Source/Editor/UnrealEd/Private/AutoReimport/FileCache.h
@@ -146,15 +146,15 @@ enum EPathType
 	Absolute
 };
 
-struct IAsyncTask : public TSharedFromThis<IAsyncTask, ESPMode::ThreadSafe>
+struct IAsyncFileCacheTask : public TSharedFromThis<IAsyncFileCacheTask, ESPMode::ThreadSafe>
 {
 	enum class EProgressResult
 	{
 		Finished, Pending
 	};
 
-	IAsyncTask() : StartTime(FPlatformTime::Seconds()) {}
-	virtual ~IAsyncTask() {}
+	IAsyncFileCacheTask() : StartTime(FPlatformTime::Seconds()) {}
+	virtual ~IAsyncFileCacheTask() {}
 	
 	/** Tick this task. Only to be called on the task thread. */
 	virtual EProgressResult Tick(const FTimeLimit& TimeLimit) = 0;
@@ -182,7 +182,7 @@ struct FFilenameAndHash
 };
 
 /** Async task responsible for MD5 hashing a number of files, reporting completed hashes to the client when done */
-struct FAsyncFileHasher : public IAsyncTask
+struct FAsyncFileHasher : public IAsyncFileCacheTask
 {
 	/** Constructor */
 	FAsyncFileHasher(TArray<FFilenameAndHash> InFilesThatNeedHashing);
@@ -223,7 +223,7 @@ protected:
  *		}
  *		TOptional<FDirectoryState> State = Reader.GetFinalState();
  */
-struct FAsyncDirectoryReader : public IAsyncTask
+struct FAsyncDirectoryReader : public IAsyncFileCacheTask
 {
 	/** Constructor that sets up the directory reader to the specified directory */
 	FAsyncDirectoryReader(const FString& InDirectory, EPathType InPathType);
