#UE4 A couple cooking optimizations for large projects. (High object count improvements)

diff --git a/Engine/Source/Editor/UnrealEd/Private/CookOnTheFlyServer.cpp b/Engine/Source/Editor/UnrealEd/Private/CookOnTheFlyServer.cpp
--- a/Engine/Source/Editor/UnrealEd/Private/CookOnTheFlyServer.cpp
+++ b/Engine/Source/Editor/UnrealEd/Private/CookOnTheFlyServer.cpp
@@ -588,7 +588,7 @@ private:
 			if ( Exclude.Find(Object) == INDEX_NONE )
 			{
 				// remove this check later because don't want this happening in development builds
-				check(RootSetArray.Find(Object)==INDEX_NONE);
+				//check(RootSetArray.Find(Object)==INDEX_NONE);
 
 				RootSetArray.Add( Object );
 				RootSet.Add(Object);
@@ -662,18 +662,14 @@ void UCookOnTheFlyServer::GetDependencies( const TSet<UPackage*>& Packages, TSet
 
 	TSet<UObject*> RootSet;
 
-
-	// Iterate through the object list
-	for( FObjectIterator It; It; ++It )
+	for (UPackage* Package : Packages)
 	{
-		if (*It)
+		TArray<UObject*> ObjectsInPackage;
+		GetObjectsWithOuter(Package, ObjectsInPackage, true);
+		for (UObject* Obj : ObjectsInPackage)
 		{
-		UPackage* Package = It->GetOutermost();
-		if ( Packages.Find( Package ) )
-		{
-			RootSet.Add(*It);
-			Found.Add(*It);
-		}
+			RootSet.Add(Obj);
+			Found.Add(Obj);
 		}
 	}
 
