'Build lighting only' triggers a BSP build first, fixes crash in RenderTestMap where Lightmass got sent a request for a 0 size mapping due to out of date BSP

diff --git a/Engine/Source/Editor/UnrealEd/Private/EditorBuildUtils.cpp b/Engine/Source/Editor/UnrealEd/Private/EditorBuildUtils.cpp
--- a/Engine/Source/Editor/UnrealEd/Private/EditorBuildUtils.cpp
+++ b/Engine/Source/Editor/UnrealEd/Private/EditorBuildUtils.cpp
@@ -266,6 +266,10 @@ bool FEditorBuildUtils::EditorBuild( UWorld* InWorld, EBuildOptions::Type Id, co
 				// We can't set the busy cursor for all windows, because lighting
 				// needs a cursor for the lighting options dialog.
 				const FScopedBusyCursor BusyCursor;
+
+				// BSP export to lightmass relies on current BSP state
+				GUnrealEd->Exec( InWorld, TEXT("MAP REBUILD ALLVISIBLE") );
+
 				GUnrealEd->BuildLighting( LightingBuildOptions );
 				bShouldMapCheck = false;
 			}
