Fixed: No way to copy Foliage from one Landscape to another
UE-9993
Foliage instances can be copy/pasted together with component they are painted on

diff --git a/Engine/Source/Editor/UnrealEd/Private/EditorExporters.cpp b/Engine/Source/Editor/UnrealEd/Private/EditorExporters.cpp
--- a/Engine/Source/Editor/UnrealEd/Private/EditorExporters.cpp
+++ b/Engine/Source/Editor/UnrealEd/Private/EditorExporters.cpp
@@ -516,19 +516,18 @@ void ULevelExporterT3D::ExportComponentExtra(const FExportObjectInnerContext* Co
 {
 	for (UActorComponent* ActorComponent : Components)
 	{
-		if (ActorComponent != NULL)
+		if (ActorComponent != nullptr && ActorComponent->GetWorld() != nullptr)
 		{
-			AActor* ActorOwner = ActorComponent->GetOwner();
-			if (ActorOwner != NULL)
+			// Go through all FoliageActors in the world, since we support cross-level bases
+			for (TActorIterator<AInstancedFoliageActor> It(ActorComponent->GetWorld()); It; ++It)
 			{
-				ULevel* ComponentLevel = Cast<ULevel>(ActorOwner->GetOuter());
-				AInstancedFoliageActor* IFA = AInstancedFoliageActor::GetInstancedFoliageActorForLevel(ComponentLevel);
-				if (IFA)
+				AInstancedFoliageActor* IFA = *It;
+				if (!IFA->IsPendingKill())
 				{
 					auto FoliageInstanceMap = IFA->GetInstancesForComponent(ActorComponent);
 					for (const auto& MapEntry : FoliageInstanceMap)
 					{
-						Ar.Logf(TEXT("%sBegin Foliage StaticMesh=%s Component=%s%s"), FCString::Spc(TextIndent), *MapEntry.Key->GetPathName(), *ActorComponent->GetName(), LINE_TERMINATOR);
+						Ar.Logf(TEXT("%sBegin Foliage FoliageType=%s Component=%s%s"), FCString::Spc(TextIndent), *MapEntry.Key->GetPathName(), *ActorComponent->GetName(), LINE_TERMINATOR);
 						for (const FFoliageInstancePlacementInfo* Inst : MapEntry.Value)
 						{
 							Ar.Logf(TEXT("%sLocation=%f,%f,%f Rotation=%f,%f,%f PreAlignRotation=%f,%f,%f DrawScale3D=%f,%f,%f Flags=%u%s"), FCString::Spc(TextIndent+3),
diff --git a/Engine/Source/Editor/UnrealEd/Private/EditorObject.cpp b/Engine/Source/Editor/UnrealEd/Private/EditorObject.cpp
--- a/Engine/Source/Editor/UnrealEd/Private/EditorObject.cpp
+++ b/Engine/Source/Editor/UnrealEd/Private/EditorObject.cpp
@@ -240,61 +240,54 @@ static const TCHAR* ImportProperties(
 		}
 		else if (GetBEGIN(&Str, TEXT("Foliage")))
 		{
-			UStaticMesh* StaticMesh;
+			UFoliageType* SourceFoliageType;
 			FName ComponentName;
 			if (SubobjectRoot &&
-				ParseObject<UStaticMesh>(Str, TEXT("StaticMesh="), StaticMesh, ANY_PACKAGE) &&
+				ParseObject<UFoliageType>(Str, TEXT("FoliageType="), SourceFoliageType, ANY_PACKAGE) &&
 				FParse::Value(Str, TEXT("Component="), ComponentName) )
 			{
 				UPrimitiveComponent* ActorComponent = FindObjectFast<UPrimitiveComponent>(SubobjectRoot, ComponentName);
 
-				if (ActorComponent)
+				if (ActorComponent && ActorComponent->GetComponentLevel())
 				{
-					ULevel* ComponentLevel = CastChecked<ULevel>(SubobjectRoot->GetOuter());
-					if (ComponentLevel->IsCurrentLevel())
-					{
-						AInstancedFoliageActor* IFA = AInstancedFoliageActor::GetInstancedFoliageActorForCurrentLevel(ComponentLevel->OwningWorld, true);
+					AInstancedFoliageActor* IFA = AInstancedFoliageActor::GetInstancedFoliageActorForLevel(ActorComponent->GetComponentLevel(), true);
 
-						const TCHAR* StrPtr;
-						FString TextLine;
-						while (FParse::Line(&SourceText, TextLine))
-						{
-							StrPtr = *TextLine;
-							if (GetEND(&StrPtr, TEXT("Foliage")))
-							{
-								break;
-							}
+					FFoliageMeshInfo* MeshInfo = nullptr;
+					UFoliageType* FoliageType = IFA->AddFoliageType(SourceFoliageType, &MeshInfo);
 
-							// Parse the instance properties
-							FFoliageInstance Instance;
-							FString Temp;
-							if (FParse::Value(StrPtr, TEXT("Location="), Temp, false))
-							{
-								GetFVECTOR(*Temp, Instance.Location);
-							}
-							if (FParse::Value(StrPtr, TEXT("Rotation="), Temp, false))
-							{
-								GetFROTATOR(*Temp, Instance.Rotation, 1);
-							}
-							if (FParse::Value(StrPtr, TEXT("PreAlignRotation="), Temp, false))
-							{
-								GetFROTATOR(*Temp, Instance.PreAlignRotation, 1);
-							}
-							if (FParse::Value(StrPtr, TEXT("DrawScale3D="), Temp, false))
-							{
-								GetFVECTOR(*Temp, Instance.DrawScale3D);
-							}
-							FParse::Value(StrPtr, TEXT("Flags="), Instance.Flags);
+					const TCHAR* StrPtr;
+					FString TextLine;
+					while (MeshInfo && FParse::Line(&SourceText, TextLine))
+					{
+						StrPtr = *TextLine;
+						if (GetEND(&StrPtr, TEXT("Foliage")))
+						{
+							break;
+						}
 
-							// Add the instance
-							FFoliageMeshInfo* MeshInfo;
-							UFoliageType* Type = IFA->GetSettingsForMesh(StaticMesh, &MeshInfo);
-							if (Type == NULL)
-							{
-								MeshInfo = IFA->AddMesh(StaticMesh, &Type);
-							}
-							MeshInfo->AddInstance(IFA, Type, Instance, ActorComponent);
+						// Parse the instance properties
+						FFoliageInstance Instance;
+						FString Temp;
+						if (FParse::Value(StrPtr, TEXT("Location="), Temp, false))
+						{
+							GetFVECTOR(*Temp, Instance.Location);
+						}
+						if (FParse::Value(StrPtr, TEXT("Rotation="), Temp, false))
+						{
+							GetFROTATOR(*Temp, Instance.Rotation, 1);
+						}
+						if (FParse::Value(StrPtr, TEXT("PreAlignRotation="), Temp, false))
+						{
+							GetFROTATOR(*Temp, Instance.PreAlignRotation, 1);
 						}
+						if (FParse::Value(StrPtr, TEXT("DrawScale3D="), Temp, false))
+						{
+							GetFVECTOR(*Temp, Instance.DrawScale3D);
+						}
+						FParse::Value(StrPtr, TEXT("Flags="), Instance.Flags);
+
+						// Add the instance
+						MeshInfo->AddInstance(IFA, FoliageType, Instance, ActorComponent);
 					}
 				}
 			}
diff --git a/Engine/Source/Runtime/Foliage/Private/InstancedFoliage.cpp b/Engine/Source/Runtime/Foliage/Private/InstancedFoliage.cpp
--- a/Engine/Source/Runtime/Foliage/Private/InstancedFoliage.cpp
+++ b/Engine/Source/Runtime/Foliage/Private/InstancedFoliage.cpp
@@ -1863,6 +1863,8 @@ void AInstancedFoliageActor::PostEditUndo()
 
 bool AInstancedFoliageActor::ShouldExport()
 {
+	// We don't support exporting/importing InstancedFoliageActor itself
+	// Instead foliage instances exported/imported together with components it's painted on
 	return false;
 }
 
