Stopped the data table editor from being able to insert a new row in the middle of the data table (rather than the end) if another row had previously been deleted.

The map is now compacted, which removes the free slot and forces the next add to appear at the end.

diff --git a/Engine/Source/Editor/UnrealEd/Private/DataTableEditorUtils.cpp b/Engine/Source/Editor/UnrealEd/Private/DataTableEditorUtils.cpp
--- a/Engine/Source/Editor/UnrealEd/Private/DataTableEditorUtils.cpp
+++ b/Engine/Source/Editor/UnrealEd/Private/DataTableEditorUtils.cpp
@@ -24,6 +24,9 @@ bool FDataTableEditorUtils::RemoveRow(UDataTable* DataTable, FName Name)
 			DataTable->RowStruct->DestroyStruct(RowData);
 			FMemory::Free(RowData);
 			bResult = true;
+
+			// Compact the map so that a subsequent add goes at the end of the table
+			DataTable->RowMap.Compact();
 		}
 		BroadcastPostChange(DataTable, EDataTableChangeInfo::RowList);
 	}
