Added a RawDelta property to MouseDeltaTracker which holds a simple mouse delta in pixels, untransformed.

Fixed issue where the drag start position in zoomed ortho modes is incorrect, due to the trigger for drag start being determined by a scaled delta.  Now uses the new RawDelta property.

diff --git a/Engine/Source/Editor/UnrealEd/Private/EditorViewportClient.cpp b/Engine/Source/Editor/UnrealEd/Private/EditorViewportClient.cpp
--- a/Engine/Source/Editor/UnrealEd/Private/EditorViewportClient.cpp
+++ b/Engine/Source/Editor/UnrealEd/Private/EditorViewportClient.cpp
@@ -2241,7 +2241,7 @@ void FEditorViewportClient::ProcessClickInViewport( const FInputEventState& Inpu
 		const int32	HitY = InputStateViewport->GetMouseY();
 		
 		// Calc the raw delta from the mouse to detect if there was any movement
-		FVector RawMouseDelta = MouseDeltaTracker->GetScreenDelta();
+		FVector RawMouseDelta = MouseDeltaTracker->GetRawDelta();
 
 		// Note: We are using raw mouse movement to double check distance moved in low performance situations.  In low performance situations its possible 
 		// that we would get a mouse down and a mouse up before the next tick where GEditor->MouseMovment has not been updated.  
@@ -3767,7 +3767,7 @@ void FEditorViewportClient::UpdateRequiredCursorVisibility()
 	else
 	{
 		// Calc the raw delta from the mouse since we started dragging to detect if there was any movement
-		FVector RawMouseDelta = MouseDeltaTracker->GetScreenDelta();
+		FVector RawMouseDelta = MouseDeltaTracker->GetRawDelta();
 
 		if (bMouseButtonDown && (RawMouseDelta.SizeSquared() >= MOUSE_CLICK_DRAG_DELTA || IsFlightCameraActive() || ShouldOrbitCamera()) && !MouseDeltaTracker->UsingDragTool())
 		{
diff --git a/Engine/Source/Editor/UnrealEd/Private/MouseDeltaTracker.cpp b/Engine/Source/Editor/UnrealEd/Private/MouseDeltaTracker.cpp
--- a/Engine/Source/Editor/UnrealEd/Private/MouseDeltaTracker.cpp
+++ b/Engine/Source/Editor/UnrealEd/Private/MouseDeltaTracker.cpp
@@ -24,6 +24,7 @@ FMouseDeltaTracker::FMouseDeltaTracker()
 	, End( FVector::ZeroVector )
 	, EndSnapped( FVector::ZeroVector )
 	, EndScreen( FVector::ZeroVector )
+	, RawDelta( FVector::ZeroVector )
 	, ReductionAmount( FVector::ZeroVector )
 	, DragTool( NULL )
 	, bHasAttemptedDragTool(false)
@@ -243,6 +244,7 @@ void FMouseDeltaTracker::StartTracking(FEditorViewportClient* InViewportClient,
 	ensure( !DragTool.IsValid() );
 
 	StartSnapped = Start = StartScreen = FVector( InX, InY, 0 );
+	RawDelta = FVector::ZeroVector;
 	TrackingWidgetMode = InViewportClient->GetWidgetMode();
 
 	// No drag tool is active, so handle snapping.
@@ -300,7 +302,7 @@ bool FMouseDeltaTracker::EndTracking(FEditorViewportClient* InViewportClient)
 
 	InViewportClient->Widget->ResetDeltaRotation();
 
-	Start = StartSnapped = StartScreen = End = EndSnapped = EndScreen = ReductionAmount = FVector::ZeroVector;
+	Start = StartSnapped = StartScreen = End = EndSnapped = EndScreen = RawDelta = ReductionAmount = FVector::ZeroVector;
 
 	// Delete the drag tool if one exists.
 	if( DragTool.IsValid() )
@@ -329,7 +331,7 @@ void FMouseDeltaTracker::ConditionalBeginUsingDragTool( FEditorViewportClient* I
 	const bool bControlDown = InViewportClient->IsCtrlPressed();
 
 	// Has there been enough mouse movement to begin using a drag tool. We don't want to start using a tool for clicks(could have very small mouse movements)
-	bool bEnoughMouseMovement = GetScreenDelta().SizeSquared() > MOUSE_CLICK_DRAG_DELTA;
+	bool bEnoughMouseMovement = GetRawDelta().SizeSquared() > MOUSE_CLICK_DRAG_DELTA;
 
 	if( bEnoughMouseMovement )
 	{
@@ -386,6 +388,11 @@ void FMouseDeltaTracker::AddDelta(FEditorViewportClient* InViewportClient, FKey
 		return;
 	}
 
+	// Accumulate raw delta
+	RawDelta += FVector(InKey == EKeys::MouseX ? InDelta : 0,
+						InKey == EKeys::MouseY ? InDelta : 0,
+						0);
+
 	// Note that AddDelta has been called since StartTracking
 	bHasReceivedAddDelta = true;
 
@@ -537,6 +544,14 @@ void FMouseDeltaTracker::AddDelta(FEditorViewportClient* InViewportClient, FKey
 
 }
 
+/**
+ * Returns the raw mouse delta, in pixels.
+ */
+const FVector FMouseDeltaTracker::GetRawDelta() const
+{
+	return RawDelta;
+}
+
 /**
  * Returns the current delta.
  */
diff --git a/Engine/Source/Editor/UnrealEd/Public/MouseDeltaTracker.h b/Engine/Source/Editor/UnrealEd/Public/MouseDeltaTracker.h
--- a/Engine/Source/Editor/UnrealEd/Public/MouseDeltaTracker.h
+++ b/Engine/Source/Editor/UnrealEd/Public/MouseDeltaTracker.h
@@ -60,6 +60,11 @@ public:
 	 */
 	const FVector GetScreenDelta() const;
 
+	/**
+	 * Returns the raw mouse delta in pixels since dragging started.
+	 */
+	const FVector GetRawDelta() const;
+
 	/**
 	 * Returns the unsnapped start position of the current mouse drag. (This will be zero if there is no drag in progress)
 	 */
@@ -132,14 +137,16 @@ private:
 	FVector Start;
 	/** The snapped start position of the current mouse drag. */
 	FVector StartSnapped;
-	/** The screen space start position of the current mouse drag. */
+	/** The screen space start position of the current mouse drag (may be scaled or rotated according to the ortho zoom or view). */
 	FVector StartScreen;
 	/** The unsnapped end position of the current mouse drag. */
 	FVector End;
 	/** The snapped end position of the current mouse drag. */
 	FVector EndSnapped;
-	/** The screen space end position of the current mouse drag. */
+	/** The screen space end position of the current mouse drag (may be scaled or rotated according to the ortho zoom or view). */
 	FVector EndScreen;
+	/** The raw unscaled mouse delta in pixels */
+	FVector RawDelta;
 
 	/** The amount that the End vectors have been reduced by since dragging started, this is added to the deltas to get an absolute delta. */
 	FVector ReductionAmount;
