Can now export 64 bpp textures, although the bit depth of the exported texture will be reduced (as neither .tga nor .bmp support 64-bit textures).

diff --git a/Engine/Source/Editor/UnrealEd/Private/Factories/EditorFactories.cpp b/Engine/Source/Editor/UnrealEd/Private/Factories/EditorFactories.cpp
--- a/Engine/Source/Editor/UnrealEd/Private/Factories/EditorFactories.cpp
+++ b/Engine/Source/Editor/UnrealEd/Private/Factories/EditorFactories.cpp
@@ -4737,11 +4737,23 @@ bool UTextureExporterBMP::ExportBinary( UObject* Object, const TCHAR* Type, FArc
 {
 	UTexture2D* Texture = CastChecked<UTexture2D>( Object );
 
-	if( !Texture->Source.IsValid() || Texture->Source.GetFormat() != TSF_BGRA8 )
+	if( !Texture->Source.IsValid() || ( Texture->Source.GetFormat() != TSF_BGRA8 && Texture->Source.GetFormat() != TSF_RGBA16 ) )
 	{
 		return false;
 	}
 
+	const bool bIsRGBA16 = Texture->Source.GetFormat() == TSF_RGBA16;
+	const int32 SourceBytesPerPixel = bIsRGBA16 ? 8 : 4;
+
+	if (bIsRGBA16)
+	{
+		FMessageLog ExportWarning("EditorErrors");
+		FFormatNamedArguments Arguments;
+		Arguments.Add(TEXT("Name"), FText::FromString(Texture->GetName()));
+		ExportWarning.Warning(FText::Format(LOCTEXT("BitDepthWarning", "{Name}: Texture is RGBA16 and cannot be represented at such high bit depth in .bmp. Color will be scaled to RGBA8."), Arguments));
+		ExportWarning.Open(EMessageSeverity::Warning);
+	}
+
 	int32 SizeX = Texture->Source.GetSizeX();
 	int32 SizeY = Texture->Source.GetSizeY();
 	TArray<uint8> RawData;
@@ -4754,9 +4766,9 @@ bool UTextureExporterBMP::ExportBinary( UObject* Object, const TCHAR* Type, FArc
 	bmf.bfType      = 'B' + (256*(int32)'M');
 	bmf.bfReserved1 = 0;
 	bmf.bfReserved2 = 0;
-	int32 biSizeImage	= SizeX * SizeY * 3;
+	int32 biSizeImage = SizeX * SizeY * 3;
 	bmf.bfOffBits   = sizeof(FBitmapFileHeader) + sizeof(FBitmapInfoHeader);
-	bmhdr.biBitCount= 24;
+	bmhdr.biBitCount = 24;
 
 	bmf.bfSize		= bmf.bfOffBits + biSizeImage;
 	Ar << bmf;
@@ -4778,13 +4790,23 @@ bool UTextureExporterBMP::ExportBinary( UObject* Object, const TCHAR* Type, FArc
 	// Upside-down scanlines.
 	for( int32 i=SizeY-1; i>=0; i-- )
 	{
-		uint8* ScreenPtr = &RawData[i*SizeX*4];
+		uint8* ScreenPtr = &RawData[i*SizeX*SourceBytesPerPixel];
 		for( int32 j=SizeX; j>0; j-- )
 		{
-			Ar << *ScreenPtr++;
-			Ar << *ScreenPtr++;
-			Ar << *ScreenPtr++;
-			ScreenPtr++;
+			if (bIsRGBA16)
+			{
+				Ar << ScreenPtr[1];
+				Ar << ScreenPtr[3];
+				Ar << ScreenPtr[5];
+				ScreenPtr += 8;
+			}
+			else
+			{
+				Ar << ScreenPtr[0];
+				Ar << ScreenPtr[1];
+				Ar << ScreenPtr[2];
+				ScreenPtr += 4;
+			}
 		}
 	}
 	return true;
@@ -5098,11 +5120,24 @@ bool UTextureExporterTGA::ExportBinary( UObject* Object, const TCHAR* Type, FArc
 {
 	UTexture2D* Texture = CastChecked<UTexture2D>( Object );
 
-	if( !Texture->Source.IsValid() || Texture->Source.GetFormat() != TSF_BGRA8 )
+	if (!Texture->Source.IsValid() || (Texture->Source.GetFormat() != TSF_BGRA8 && Texture->Source.GetFormat() != TSF_RGBA16))
 	{
 		return false;
 	}
 
+	const bool bIsRGBA16 = Texture->Source.GetFormat() == TSF_RGBA16;
+
+	if (bIsRGBA16)
+	{
+		FMessageLog ExportWarning("EditorErrors");
+		FFormatNamedArguments Arguments;
+		Arguments.Add(TEXT("Name"), FText::FromString(Texture->GetName()));
+		ExportWarning.Warning(FText::Format(LOCTEXT("BitDepthWarning", "{Name}: Texture is RGBA16 and cannot be represented at such high bit depth in .tga. Color will be scaled to RGBA8."), Arguments));
+		ExportWarning.Open(EMessageSeverity::Warning);
+	}
+
+	const int32 BytesPerPixel = bIsRGBA16 ? 8 : 4;
+
 	int32 SizeX = Texture->Source.GetSizeX();
 	int32 SizeY = Texture->Source.GetSizeY();
 	TArray<uint8> RawData;
@@ -5116,13 +5151,14 @@ bool UTextureExporterTGA::ExportBinary( UObject* Object, const TCHAR* Type, FArc
 		// If the texture isn't compressed with no alpha scan the texture to see if the alpha values are all 255 which means we can skip exporting it.
 		// This is a relatively slow process but we are just exporting textures 
 		bExportWithAlpha = false;
+		const int32 AlphaOffset = bIsRGBA16 ? 7 : 3;
 		for( int32 Y = SizeY - 1; Y >= 0; --Y )
 		{
-			uint8* Color = &RawData[Y * SizeX * 4];
+			uint8* Color = &RawData[Y * SizeX * BytesPerPixel];
 			for( int32 X = SizeX; X > 0; --X )
 			{
 				// Skip color info
-				Color+=3;
+				Color += AlphaOffset;
 				// Get Alpha value then increment the pointer past it for the next pixel
 				uint8 Alpha = *Color++;
 				if( Alpha != 255 )
@@ -5151,7 +5187,7 @@ bool UTextureExporterTGA::ExportBinary( UObject* Object, const TCHAR* Type, FArc
 	TGA.Width = OriginalWidth;
 	Ar.Serialize( &TGA, sizeof(TGA) );
 
-	if( bExportWithAlpha )
+	if( bExportWithAlpha && !bIsRGBA16)
 	{
 		for( int32 Y=0;Y < OriginalHeight;Y++ )
 		{
@@ -5164,14 +5200,28 @@ bool UTextureExporterTGA::ExportBinary( UObject* Object, const TCHAR* Type, FArc
 		// Serialize each pixel
 		for( int32 Y = OriginalHeight - 1; Y >= 0; --Y )
 		{
-			uint8* Color = &RawData[Y * OriginalWidth * 4];
+			uint8* Color = &RawData[Y * OriginalWidth * BytesPerPixel];
 			for( int32 X = OriginalWidth; X > 0; --X )
 			{
-				Ar << *Color++;
-				Ar << *Color++;
-				Ar << *Color++;
-				// Skip alpha channel since we are exporting with no alpha
-				Color++;
+				if (bIsRGBA16)
+				{
+					Ar << Color[1];
+					Ar << Color[3];
+					Ar << Color[5];
+					if (bExportWithAlpha)
+					{
+						Ar << Color[7];
+					}
+					Color += 8;
+				}
+				else
+				{
+					Ar << *Color++;
+					Ar << *Color++;
+					Ar << *Color++;
+					// Skip alpha channel since we are exporting with no alpha
+					Color++;
+				}
 			}
 		}
 	}
