Group actor pivot location now correctly updates when adding or removing actors to/from the group.

diff --git a/Engine/Source/Editor/UnrealEd/Private/EditorGroup.cpp b/Engine/Source/Editor/UnrealEd/Private/EditorGroup.cpp
--- a/Engine/Source/Editor/UnrealEd/Private/EditorGroup.cpp
+++ b/Engine/Source/Editor/UnrealEd/Private/EditorGroup.cpp
@@ -189,10 +189,12 @@ void UUnrealEdEngine::edactRemoveFromGroup()
 				if(ActorGroupParent)
 				{
 					ActorGroupParent->Add(*Actor);
+					ActorGroupParent->CenterGroupLocation();
 				}
 				else
 				{
 					ActorGroup->Remove(*Actor);
+					ActorGroup->CenterGroupLocation();
 				}
 			}
 		}
diff --git a/Engine/Source/Editor/UnrealEd/Private/GroupActor.cpp b/Engine/Source/Editor/UnrealEd/Private/GroupActor.cpp
--- a/Engine/Source/Editor/UnrealEd/Private/GroupActor.cpp
+++ b/Engine/Source/Editor/UnrealEd/Private/GroupActor.cpp
@@ -547,6 +547,8 @@ void AGroupActor::AddSelectedActorsToSelectedGroup()
 							SelectedGroup->Add( *ActorsToAdd[ActorIndex] );
 						}
 					}
+
+					SelectedGroup->CenterGroupLocation();
 				}
 			}
 			else
