Autosave is now deferred if shaders are compiling.

diff --git a/Engine/Source/Editor/UnrealEd/Private/PackageAutoSaver.cpp b/Engine/Source/Editor/UnrealEd/Private/PackageAutoSaver.cpp
--- a/Engine/Source/Editor/UnrealEd/Private/PackageAutoSaver.cpp
+++ b/Engine/Source/Editor/UnrealEd/Private/PackageAutoSaver.cpp
@@ -7,6 +7,7 @@
 #include "SNotificationList.h"
 #include "NotificationManager.h"
 #include "AutoSaveUtils.h"
+#include "ShaderCompiler.h"
 
 namespace PackageAutoSaverJson
 {
@@ -386,10 +387,11 @@ bool FPackageAutoSaver::CanAutoSave() const
 	const bool bPlayWorldValid = GUnrealEd->PlayWorld != nullptr;
 	const bool bAnyMenusVisible	= FSlateApplication::Get().AnyMenusVisible();
 	const bool bAutomationTesting = GIsAutomationTesting;
-	const bool bIsInteratcting = FSlateApplication::Get().HasAnyMouseCaptor() || GUnrealEd->IsUserInteracting() || bDidTypeInATextBlockRecently;
+	const bool bIsInteracting = FSlateApplication::Get().HasAnyMouseCaptor() || GUnrealEd->IsUserInteracting() || bDidTypeInATextBlockRecently;
 	const bool bHasGameOrProjectLoaded = FApp::HasGameName();
+	const bool bAreShadersCompiling = GShaderCompilingManager->IsCompiling();
 
-	return (bAutosaveEnabled && !bSlowTask && !bInterpEditMode && !bPlayWorldValid && !bAnyMenusVisible && !bAutomationTesting && !bIsInteratcting && !GIsDemoMode && bHasGameOrProjectLoaded);
+	return (bAutosaveEnabled && !bSlowTask && !bInterpEditMode && !bPlayWorldValid && !bAnyMenusVisible && !bAutomationTesting && !bIsInteracting && !GIsDemoMode && bHasGameOrProjectLoaded && !bAreShadersCompiling);
 }
 
 bool FPackageAutoSaver::DoPackagesNeedAutoSave() const
