Allow EditorUtilities::CopyActorProperties to be used for targets that are a subclass of the source
clean up some dead property checks
make remaining property checks use get member name checked

diff --git a/Engine/Source/Editor/UnrealEd/Private/Editor.cpp b/Engine/Source/Editor/UnrealEd/Private/Editor.cpp
--- a/Engine/Source/Editor/UnrealEd/Private/Editor.cpp
+++ b/Engine/Source/Editor/UnrealEd/Private/Editor.cpp
@@ -7255,7 +7255,7 @@ namespace EditorUtilities
 
 		// The actor's classes should be compatible, right?
 		UClass* ActorClass = SourceActor->GetClass();
-		check( ActorClass == TargetActor->GetClass() );
+		check( TargetActor->GetClass()->IsChildOf(ActorClass) );
 
 		// Get archetype instances for propagation (if requested)
 		TArray<UObject*> ArchetypeInstances;
@@ -7278,7 +7278,7 @@ namespace EditorUtilities
 				const bool bIsBlueprintReadonly = !!(Options & ECopyOptions::FilterBlueprintReadOnly) && !!( Property->PropertyFlags & CPF_BlueprintReadOnly );
 				const bool bIsIdentical = Property->Identical_InContainer( SourceActor, TargetActor );
 
-				if( !bIsTransient && !bIsIdentical && !bIsComponentContainer && !bIsComponentProp && !bIsBlueprintReadonly && Property->GetName() != TEXT( "Tag" ) )
+				if( !bIsTransient && !bIsIdentical && !bIsComponentContainer && !bIsComponentProp && !bIsBlueprintReadonly)
 				{
 					const bool bIsSafeToCopy = !( Options & ECopyOptions::OnlyCopyEditOrInterpProperties ) || ( Property->HasAnyPropertyFlags( CPF_Edit | CPF_Interp ) );
 					if( bIsSafeToCopy )
@@ -7339,13 +7339,6 @@ namespace EditorUtilities
 						}
 
 						++CopiedPropertyCount;
-
-						if( Property->GetFName() == FName( TEXT( "RelativeScale3D" ) ) ||
-							Property->GetFName() == FName( TEXT( "RelativeLocation" ) ) ||
-							Property->GetFName() == FName( TEXT( "RelativeRotation" ) ) )
-						{
-							bTransformChanged = true;
-						}
 					}
 				}
 			}
@@ -7399,9 +7392,9 @@ namespace EditorUtilities
 					const bool bIsIdentical = Property->Identical_InContainer( SourceComponent, TargetComponent );
 					const bool bIsComponent = !!( Property->PropertyFlags & ( CPF_InstancedReference | CPF_ContainsInstancedReference ) );
 					const bool bIsTransform =
-						Property->GetFName() == FName( TEXT( "RelativeScale3D" ) ) ||
-						Property->GetFName() == FName( TEXT( "RelativeLocation" ) ) ||
-						Property->GetFName() == FName( TEXT( "RelativeRotation" ) );
+						Property->GetFName() == GET_MEMBER_NAME_CHECKED(USceneComponent, RelativeScale3D) ||
+						Property->GetFName() == GET_MEMBER_NAME_CHECKED(USceneComponent, RelativeLocation) ||
+						Property->GetFName() == GET_MEMBER_NAME_CHECKED(USceneComponent, RelativeRotation);
 
 					if( !bIsTransient && !bIsIdentical && !bIsComponent && !SourceUCSModifiedProperties.Contains(Property)
 						&& ( !bIsTransform || SourceComponent != SourceActor->GetRootComponent() || ( !SourceActor->HasAnyFlags( RF_ClassDefaultObject | RF_ArchetypeObject ) && !TargetActor->HasAnyFlags( RF_ClassDefaultObject | RF_ArchetypeObject ) ) ) )
