Use EditorUtilities::CopyActorProperties instead of CopyPropertiesForUnrelatedObjects when creating a blueprint from an Actor in the world

diff --git a/Engine/Source/Editor/UnrealEd/Private/Kismet2/Kismet2.cpp b/Engine/Source/Editor/UnrealEd/Private/Kismet2/Kismet2.cpp
--- a/Engine/Source/Editor/UnrealEd/Private/Kismet2/Kismet2.cpp
+++ b/Engine/Source/Editor/UnrealEd/Private/Kismet2/Kismet2.cpp
@@ -1163,33 +1163,21 @@ UBlueprint* FKismetEditorUtilities::CreateBlueprintFromActor(const FName Bluepri
 
 			if (NewBlueprint->GeneratedClass != nullptr)
 			{
-				// Since we already created SCS Nodes for the instance components, temporarily cache and clear the
-				// array to avoid creating duplicates in the new CDO
-				const TArray<UActorComponent*> TempInstanceComponents(Actor->GetInstanceComponents());
-				Actor->ClearInstanceComponents(false);
+				AActor* CDO = CastChecked<AActor>(NewBlueprint->GeneratedClass->GetDefaultObject());
+				const auto CopyOptions = (EditorUtilities::ECopyOptions::Type)(EditorUtilities::ECopyOptions::OnlyCopyEditOrInterpProperties | EditorUtilities::ECopyOptions::PropagateChangesToArchetypeInstances);
+				EditorUtilities::CopyActorProperties(Actor, CDO, CopyOptions);
 
-				UObject* CDO = NewBlueprint->GeneratedClass->GetDefaultObject();
-				UEditorEngine::CopyPropertiesForUnrelatedObjects(Actor, CDO);
-
-				for (UActorComponent* Component : TempInstanceComponents)
-				{
-					Actor->AddInstanceComponent(Component);
-				}
-
-				if (AActor* CDOAsActor = Cast<AActor>(CDO))
+				if (USceneComponent* Scene = CDO->GetRootComponent())
 				{
-					if (USceneComponent* Scene = CDOAsActor->GetRootComponent())
-					{
-						Scene->RelativeLocation = FVector::ZeroVector;
-						Scene->RelativeRotation = FRotator::ZeroRotator;
+					Scene->RelativeLocation = FVector::ZeroVector;
+					Scene->RelativeRotation = FRotator::ZeroRotator;
 
-						// Clear out the attachment info after having copied the properties from the source actor
-						Scene->AttachParent = NULL;
-						Scene->AttachChildren.Empty();
+					// Clear out the attachment info after having copied the properties from the source actor
+					Scene->AttachParent = NULL;
+					Scene->AttachChildren.Empty();
 
-						// Ensure the light mass information is cleaned up
-						Scene->InvalidateLightingCache();
-					}
+					// Ensure the light mass information is cleaned up
+					Scene->InvalidateLightingCache();
 				}
 			}
 
