Extented root array test cases for json serialization

diff --git a/Engine/Source/Editor/UnrealEd/Private/Tests/EditorAutomationTests.cpp b/Engine/Source/Editor/UnrealEd/Private/Tests/EditorAutomationTests.cpp
--- a/Engine/Source/Editor/UnrealEd/Private/Tests/EditorAutomationTests.cpp
+++ b/Engine/Source/Editor/UnrealEd/Private/Tests/EditorAutomationTests.cpp
@@ -275,6 +275,185 @@ bool FJsonAutomationTest::RunTest(const FString& Parameters)
 		check( InputString == OutputString );
 	}
 
+	// Object Array Case
+	{
+		const FString InputString =
+			TEXT("[")
+			TEXT(	"{")
+			TEXT(		"\"Value\":\"Some String1\"")
+			TEXT(	"},")
+			TEXT(	"{")
+			TEXT(		"\"Value\":\"Some String2\"")
+			TEXT(	"},")
+			TEXT(	"{")
+			TEXT(		"\"Value\":\"Some String3\"")
+			TEXT(	"}")
+			TEXT("]");
+
+		TSharedRef< TJsonReader<> > Reader = TJsonReaderFactory<>::Create(InputString);
+
+		TArray< TSharedPtr<FJsonValue> > Array;
+		check(FJsonSerializer::Deserialize(Reader, Array));
+		check(Array.Num() == 3);
+		check(Array[0].IsValid());
+		check(Array[1].IsValid());
+		check(Array[2].IsValid());
+
+		TSharedPtr< FJsonObject > Object = Array[0]->AsObject();
+		check(Object.IsValid());
+		check(Object->GetStringField(TEXT("Value")) == TEXT("Some String1"));
+
+		Object = Array[1]->AsObject();
+		check(Object.IsValid());
+		check(Object->GetStringField(TEXT("Value")) == TEXT("Some String2"));
+
+		Object = Array[2]->AsObject();
+		check(Object.IsValid());
+		check(Object->GetStringField(TEXT("Value")) == TEXT("Some String3"));
+
+		FString OutputString;
+		TSharedRef< FCondensedJsonStringWriter > Writer = FCondensedJsonStringWriterFactory::Create(&OutputString);
+		check(FJsonSerializer::Serialize(Array, Writer));
+		check(InputString == OutputString);
+	}
+
+	// Number Array Case
+	{
+		const FString InputString =
+			TEXT("[")
+			TEXT("10,")
+			TEXT("20,")
+			TEXT("30,")
+			TEXT("40")
+			TEXT("]");
+
+		TSharedRef< TJsonReader<> > Reader = TJsonReaderFactory<>::Create(InputString);
+
+		TArray< TSharedPtr<FJsonValue> > Array;
+		check(FJsonSerializer::Deserialize(Reader, Array));
+		check(Array.Num() == 4);
+		check(Array[0].IsValid());
+		check(Array[1].IsValid());
+		check(Array[2].IsValid());
+		check(Array[3].IsValid());
+
+		double Number = Array[0]->AsNumber();
+		check(Number == 10);
+
+		Number = Array[1]->AsNumber();
+		check(Number == 20);
+
+		Number = Array[2]->AsNumber();
+		check(Number == 30);
+
+		Number = Array[3]->AsNumber();
+		check(Number == 40);
+
+		FString OutputString;
+		TSharedRef< FCondensedJsonStringWriter > Writer = FCondensedJsonStringWriterFactory::Create(&OutputString);
+		check(FJsonSerializer::Serialize(Array, Writer));
+		check(InputString == OutputString);
+	}
+
+	// String Array Case
+	{
+		const FString InputString =
+			TEXT("[")
+			TEXT("\"Some String1\",")
+			TEXT("\"Some String2\",")
+			TEXT("\"Some String3\",")
+			TEXT("\"Some String4\"")
+			TEXT("]");
+
+		TSharedRef< TJsonReader<> > Reader = TJsonReaderFactory<>::Create(InputString);
+
+		TArray< TSharedPtr<FJsonValue> > Array;
+		check(FJsonSerializer::Deserialize(Reader, Array));
+		check(Array.Num() == 4);
+		check(Array[0].IsValid());
+		check(Array[1].IsValid());
+		check(Array[2].IsValid());
+		check(Array[3].IsValid());
+
+		FString Text = Array[0]->AsString();
+		check(Text == TEXT("Some String1"));
+
+		Text = Array[1]->AsString();
+		check(Text == TEXT("Some String2"));
+
+		Text = Array[2]->AsString();
+		check(Text == TEXT("Some String3"));
+
+		Text = Array[3]->AsString();
+		check(Text == TEXT("Some String4"));
+
+		FString OutputString;
+		TSharedRef< FCondensedJsonStringWriter > Writer = FCondensedJsonStringWriterFactory::Create(&OutputString);
+		check(FJsonSerializer::Serialize(Array, Writer));
+		check(InputString == OutputString);
+	}
+
+	// Complex Array Case
+	{
+		const FString InputString =
+			TEXT("[")
+			TEXT(	"\"Some String1\",")
+			TEXT(	"10,")
+			TEXT(	"{")
+			TEXT(		"\"Value\":\"Some String3\"")
+			TEXT(	"},")
+			TEXT(	"[")
+			TEXT(		"\"Some String4\",")
+			TEXT(		"\"Some String5\"")
+			TEXT(	"],")
+			TEXT(	"true,")
+			TEXT(	"null")
+			TEXT("]");
+
+		TSharedRef< TJsonReader<> > Reader = TJsonReaderFactory<>::Create(InputString);
+
+		TArray< TSharedPtr<FJsonValue> > Array;
+		check(FJsonSerializer::Deserialize(Reader, Array));
+		check(Array.Num() == 6);
+		check(Array[0].IsValid());
+		check(Array[1].IsValid());
+		check(Array[2].IsValid());
+		check(Array[3].IsValid());
+		check(Array[4].IsValid());
+		check(Array[5].IsValid());
+
+		FString Text = Array[0]->AsString();
+		check(Text == TEXT("Some String1"));
+
+		double Number = Array[1]->AsNumber();
+		check(Number == 10);
+
+		TSharedPtr< FJsonObject > Object = Array[2]->AsObject();
+		check(Object.IsValid());
+		check(Object->GetStringField(TEXT("Value")) == TEXT("Some String3"));
+
+		const TArray<TSharedPtr< FJsonValue >>& InnerArray = Array[3]->AsArray();
+		check(InnerArray.Num() == 2);
+		check(Array[0].IsValid());
+		check(Array[1].IsValid());
+
+		Text = InnerArray[0]->AsString();
+		check(Text == TEXT("Some String4"));
+
+		Text = InnerArray[1]->AsString();
+		check(Text == TEXT("Some String5"));
+
+		bool Boolean = Array[4]->AsBool();
+		check(Boolean == true);
+
+		check(Array[5]->IsNull() == true);
+
+		FString OutputString;
+		TSharedRef< FCondensedJsonStringWriter > Writer = FCondensedJsonStringWriterFactory::Create(&OutputString);
+		check(FJsonSerializer::Serialize(Array, Writer));
+		check(InputString == OutputString);
+	}
+
 	// String Test
 	{
 		const FString InputString =
