Small change to ensure that auto reimport directories are upgraded to the new format correctly.

diff --git a/Engine/Source/Editor/UnrealEd/Classes/Settings/EditorLoadingSavingSettings.h b/Engine/Source/Editor/UnrealEd/Classes/Settings/EditorLoadingSavingSettings.h
--- a/Engine/Source/Editor/UnrealEd/Classes/Settings/EditorLoadingSavingSettings.h
+++ b/Engine/Source/Editor/UnrealEd/Classes/Settings/EditorLoadingSavingSettings.h
@@ -90,7 +90,7 @@ public:
 	UPROPERTY(EditAnywhere, config, Category=AutoReimport, meta=(DisplayName="Monitor Content Directories", ToolTip="When enabled, changes to made to source content files inside the content directories will automatically be reflected in the content browser."))
 	bool bMonitorContentDirectories;
 
-	UPROPERTY()
+	UPROPERTY(config)
 	TArray<FString> AutoReimportDirectories_DEPRECATED;
 
 	/** Directories being monitored for Auto Reimport */
diff --git a/Engine/Source/Editor/UnrealEd/Private/Settings/SettingsClasses.cpp b/Engine/Source/Editor/UnrealEd/Private/Settings/SettingsClasses.cpp
--- a/Engine/Source/Editor/UnrealEd/Private/Settings/SettingsClasses.cpp
+++ b/Engine/Source/Editor/UnrealEd/Private/Settings/SettingsClasses.cpp
@@ -137,11 +137,15 @@ void UEditorLoadingSavingSettings::PostEditChangeProperty( struct FPropertyChang
 
 void UEditorLoadingSavingSettings::PostInitProperties()
 {
-	for (const auto& String : AutoReimportDirectories_DEPRECATED)
+	if (AutoReimportDirectories_DEPRECATED.Num() != 0)
 	{
-		FAutoReimportDirectoryConfig Config;
-		Config.SourceDirectory = String;
-		AutoReimportDirectorySettings.Add(Config);
+		AutoReimportDirectorySettings.Empty();
+		for (const auto& String : AutoReimportDirectories_DEPRECATED)
+		{
+			FAutoReimportDirectoryConfig Config;
+			Config.SourceDirectory = String;
+			AutoReimportDirectorySettings.Add(Config);
+		}
 	}
 	Super::PostInitProperties();
 }
