Fix issues with propagation to subclasses/instances of subclasses when applying instance properties to a blueprint

diff --git a/Engine/Source/Editor/UnrealEd/Private/Editor.cpp b/Engine/Source/Editor/UnrealEd/Private/Editor.cpp
--- a/Engine/Source/Editor/UnrealEd/Private/Editor.cpp
+++ b/Engine/Source/Editor/UnrealEd/Private/Editor.cpp
@@ -7036,20 +7036,20 @@ namespace EditorUtilities
 	template<class AllocatorType = FDefaultAllocator>
 	UActorComponent* FindMatchingComponentInstance( UActorComponent* SourceComponent, AActor* TargetActor, const TArray<UActorComponent*, AllocatorType>& TargetComponents, int32& StartIndex )
 	{
-		UActorComponent* TargetComponent = StartIndex < TargetComponents.Num() ? TargetComponents[ StartIndex ] : NULL;
+		const bool bSourceIsArchetype = SourceComponent->HasAnyFlags(RF_ArchetypeObject);
+		UActorComponent* TargetComponent = StartIndex < TargetComponents.Num() ? TargetComponents[ StartIndex ] : nullptr;
 
 		// If the source and target components do not match (e.g. context-specific), attempt to find a match in the target's array elsewhere
 		const int32 NumTargetComponents = TargetComponents.Num();
-		if( (SourceComponent != NULL) 
-			&& ((TargetComponent == NULL) 
+		if( (SourceComponent != nullptr) 
+			&& ((TargetComponent == nullptr) 
 				|| (SourceComponent->GetFName() != TargetComponent->GetFName()) ))
 		{
 			// Reset the target component since it doesn't match the source
-			TargetComponent = NULL;
+			TargetComponent = nullptr;
 
 			if (NumTargetComponents > 0)
 			{
-				const bool bSourceIsArchetype = SourceComponent->HasAnyFlags(RF_ArchetypeObject);
 				// Attempt to locate a match elsewhere in the target's component list
 				const int32 StartingIndex = (bSourceIsArchetype ? StartIndex : StartIndex + 1);
 				int32 FindTargetComponentIndex = (StartingIndex >= NumTargetComponents) ? 0 : StartingIndex;
@@ -7057,25 +7057,37 @@ namespace EditorUtilities
 				{
 					UActorComponent* FindTargetComponent = TargetComponents[ FindTargetComponentIndex ];
 
-					// In the case that the SourceComponent is an Archetype there is a better than even chance the name won't match due to the way the SCS
-					// is set up, so we're actually going to reverse search
-					if (bSourceIsArchetype)
+					if (FindTargetComponent->GetClass() == SourceComponent->GetClass())
 					{
-						if ( SourceComponent == FindTargetComponent->GetArchetype())
+						// In the case that the SourceComponent is an Archetype there is a better than even chance the name won't match due to the way the SCS
+						// is set up, so we're actually going to reverse search the archetype chain
+						if (bSourceIsArchetype)
 						{
-							TargetComponent = FindTargetComponent;
-							StartIndex = FindTargetComponentIndex;
-							break;
+							UActorComponent* CheckComponent = FindTargetComponent;
+							while (CheckComponent)
+							{
+								if ( SourceComponent == CheckComponent->GetArchetype())
+								{
+									TargetComponent = FindTargetComponent;
+									StartIndex = FindTargetComponentIndex;
+									break;
+								}
+								CheckComponent = CastChecked<UActorComponent>(CheckComponent->GetArchetype(), ECastCheckedType::NullAllowed);
+							}
+							if (TargetComponent)
+							{
+								break;
+							}
 						}
-					}
-					else
-					{
-						// If we found a match, update the target component and adjust the target index to the matching position
-						if( FindTargetComponent != NULL && SourceComponent->GetFName() == FindTargetComponent->GetFName() )
+						else
 						{
-							TargetComponent = FindTargetComponent;
-							StartIndex = FindTargetComponentIndex;
-							break;
+							// If we found a match, update the target component and adjust the target index to the matching position
+							if( FindTargetComponent != NULL && SourceComponent->GetFName() == FindTargetComponent->GetFName() )
+							{
+								TargetComponent = FindTargetComponent;
+								StartIndex = FindTargetComponentIndex;
+								break;
+							}
 						}
 					}
 
@@ -7089,38 +7101,57 @@ namespace EditorUtilities
 			}
 
 			// If we still haven't found a match and we're targeting a class default object what we're really looking
-			// for is the Archetype
-			if(TargetComponent == NULL && TargetActor->HasAnyFlags(RF_ClassDefaultObject|RF_ArchetypeObject))
+			// for is an Archetype
+			if(TargetComponent == nullptr && TargetActor->HasAnyFlags(RF_ClassDefaultObject|RF_ArchetypeObject))
 			{
-				TargetComponent = CastChecked<UActorComponent>(SourceComponent->GetArchetype(), ECastCheckedType::NullAllowed);
-
-				// If the returned target component is not from the direct class of the actor we're targeting, we need to insert an inheritable component
-				if (TargetComponent && (TargetComponent->GetOuter() != TargetActor->GetClass()))
+				if (bSourceIsArchetype)
 				{
-					// This component doesn't exist in the hierarchy anywhere and we're not going to modify the CDO, so we'll drop it
-					if (TargetComponent->HasAnyFlags(RF_ClassDefaultObject))
+					UBlueprintGeneratedClass* BPGC = Cast<UBlueprintGeneratedClass>(SourceComponent->GetOuter());
+
+					// If the target actor's class is a child of our owner and we're both archetypes, then we're actually looking for an overridden version of ourselves
+					if (BPGC && TargetActor->GetClass()->IsChildOf(BPGC))
 					{
-						TargetComponent = nullptr;
+						TargetComponent = Cast<UActorComponent>(TargetActor->GetClass()->FindArchetype(SourceComponent->GetClass(), SourceComponent->GetFName()));
+
+						// If it is us, then we're done, we don't need to find this
+						if (TargetComponent == SourceComponent)
+						{
+							TargetComponent = nullptr;
+						}
 					}
-					else
+				}
+				else
+				{
+					TargetComponent = CastChecked<UActorComponent>(SourceComponent->GetArchetype(), ECastCheckedType::NullAllowed);
+
+					// If the returned target component is not from the direct class of the actor we're targeting, we need to insert an inheritable component
+					if (TargetComponent && (TargetComponent->GetOuter() != TargetActor->GetClass()))
 					{
-						UBlueprintGeneratedClass* BPGC = CastChecked<UBlueprintGeneratedClass>(TargetActor->GetClass());
-						UBlueprint* Blueprint = CastChecked<UBlueprint>(BPGC->ClassGeneratedBy);
-						UInheritableComponentHandler* InheritableComponentHandler = Blueprint->GetInheritableComponentHandler(true);
-						if (InheritableComponentHandler)
+						// This component doesn't exist in the hierarchy anywhere and we're not going to modify the CDO, so we'll drop it
+						if (TargetComponent->HasAnyFlags(RF_ClassDefaultObject))
 						{
-							BPGC = Cast<UBlueprintGeneratedClass>(BPGC->GetSuperClass());
-							USCS_Node* SCSNode = nullptr;
-							while (BPGC)
+							TargetComponent = nullptr;
+						}
+						else
+						{
+							UBlueprintGeneratedClass* BPGC = CastChecked<UBlueprintGeneratedClass>(TargetActor->GetClass());
+							UBlueprint* Blueprint = CastChecked<UBlueprint>(BPGC->ClassGeneratedBy);
+							UInheritableComponentHandler* InheritableComponentHandler = Blueprint->GetInheritableComponentHandler(true);
+							if (InheritableComponentHandler)
 							{
-								SCSNode = BPGC->SimpleConstructionScript->FindSCSNode(SourceComponent->GetFName());
-								BPGC = (SCSNode ? nullptr : Cast<UBlueprintGeneratedClass>(BPGC->GetSuperClass()));
-							}
-							check(SCSNode);
+								BPGC = Cast<UBlueprintGeneratedClass>(BPGC->GetSuperClass());
+								USCS_Node* SCSNode = nullptr;
+								while (BPGC)
+								{
+									SCSNode = BPGC->SimpleConstructionScript->FindSCSNode(SourceComponent->GetFName());
+									BPGC = (SCSNode ? nullptr : Cast<UBlueprintGeneratedClass>(BPGC->GetSuperClass()));
+								}
+								check(SCSNode);
 
-							FComponentKey Key(SCSNode);
-							check(InheritableComponentHandler->GetOverridenComponentTemplate(Key) == nullptr);
-							TargetComponent = InheritableComponentHandler->CreateOverridenComponentTemplate(Key);
+								FComponentKey Key(SCSNode);
+								check(InheritableComponentHandler->GetOverridenComponentTemplate(Key) == nullptr);
+								TargetComponent = InheritableComponentHandler->CreateOverridenComponentTemplate(Key);
+							}
 						}
 					}
 				}
@@ -7147,7 +7178,7 @@ namespace EditorUtilities
 	}
 
 
-	void CopySingleActorPropertyRecursive(const void* const InSourcePtr, void* const InTargetPtr, UObject* const InTargetObject, UProperty* const InProperty)
+	void CopySinglePropertyRecursive(const void* const InSourcePtr, void* const InTargetPtr, UObject* const InTargetObject, UProperty* const InProperty)
 	{
 		// Properties that are *object* properties are tricky
 		// Sometimes the object will be a reference to a PIE-world object, and copying that reference back to an actor CDO asset is not a good idea
@@ -7205,7 +7236,7 @@ namespace EditorUtilities
 				for (TFieldIterator<UProperty> It(StructProperty->Struct); It; ++It)
 				{
 					UProperty* const InnerProperty = *It;
-					CopySingleActorPropertyRecursive(SourcePtr, TargetPtr, InTargetObject, InnerProperty);
+					CopySinglePropertyRecursive(SourcePtr, TargetPtr, InTargetObject, InnerProperty);
 				}
 			}
 
@@ -7224,7 +7255,7 @@ namespace EditorUtilities
 
 			for (int32 Index = 0; Index < Num; Index++)
 			{
-				CopySingleActorPropertyRecursive(SourceArrayHelper.GetRawPtr(Index), TargetArrayHelper.GetRawPtr(Index), InTargetObject, InnerProperty);
+				CopySinglePropertyRecursive(SourceArrayHelper.GetRawPtr(Index), TargetArrayHelper.GetRawPtr(Index), InTargetObject, InnerProperty);
 			}
 
 			bNeedsGenericCopy = false;
@@ -7237,14 +7268,14 @@ namespace EditorUtilities
 		}
 	}
 
-	void CopySingleActorProperty(const UObject* const InSourceObject, UObject* const InTargetObject, UProperty* const InProperty)
+	void CopySingleProperty(const UObject* const InSourceObject, UObject* const InTargetObject, UProperty* const InProperty)
 	{
-		CopySingleActorPropertyRecursive(InSourceObject, InTargetObject, InTargetObject, InProperty);
+		CopySinglePropertyRecursive(InSourceObject, InTargetObject, InTargetObject, InProperty);
 	}
 
 	int32 CopyActorProperties( AActor* SourceActor, AActor* TargetActor, const ECopyOptions::Type Options )
 	{
-		check( SourceActor != NULL && TargetActor != NULL );
+		check( SourceActor != nullptr && TargetActor != nullptr );
 
 		const bool bIsPreviewing = ( Options & ECopyOptions::PreviewOnly ) != 0;
 
@@ -7255,10 +7286,19 @@ namespace EditorUtilities
 		check( TargetActor->GetClass()->IsChildOf(ActorClass) );
 
 		// Get archetype instances for propagation (if requested)
-		TArray<UObject*> ArchetypeInstances;
+		TArray<AActor*> ArchetypeInstances;
 		if( Options & ECopyOptions::PropagateChangesToArchetypeInstances )
 		{
-			TargetActor->GetArchetypeInstances(ArchetypeInstances);
+			TArray<UObject*> ObjectArchetypeInstances;
+			TargetActor->GetArchetypeInstances(ObjectArchetypeInstances);
+
+			for (UObject* ObjectArchetype : ObjectArchetypeInstances)
+			{
+				if (AActor* ActorArchetype = Cast<AActor>(ObjectArchetype))
+				{
+					ArchetypeInstances.Add(ActorArchetype);
+				}			
+			}
 		}
 
 		bool bTransformChanged = false;
@@ -7267,7 +7307,7 @@ namespace EditorUtilities
 		// @todo sequencer: Most of this block of code was borrowed (pasted) from UEditorEngine::ConvertActors().  If we end up being able to share these code bodies, that would be nice!
 		{
 			bool bIsFirstModification = true;
-			for( UProperty* Property = ActorClass->PropertyLink; Property != NULL; Property = Property->PropertyLinkNext )
+			for( UProperty* Property = ActorClass->PropertyLink; Property != nullptr; Property = Property->PropertyLinkNext )
 			{
 				const bool bIsTransient = !!( Property->PropertyFlags & CPF_Transient );
 				const bool bIsComponentContainer = !!( Property->PropertyFlags & CPF_ContainsInstancedReference );
@@ -7297,17 +7337,16 @@ namespace EditorUtilities
 							TArray<UObject*> ArchetypeInstancesToChange;
 							if( Options & ECopyOptions::PropagateChangesToArchetypeInstances )
 							{
-								for( int32 InstanceIndex = 0; InstanceIndex < ArchetypeInstances.Num(); ++InstanceIndex )
+								for( AActor* ArchetypeInstance : ArchetypeInstances )
 								{
-									UObject* ArchetypeInstance = ArchetypeInstances[InstanceIndex];
-									if( ArchetypeInstance != NULL && Property->Identical_InContainer( ArchetypeInstance, TargetActor ) )
+									if( ArchetypeInstance != nullptr && Property->Identical_InContainer( ArchetypeInstance, TargetActor ) )
 									{
 										ArchetypeInstancesToChange.Add( ArchetypeInstance );
 									}
 								}
 							}
 
-							CopySingleActorProperty(SourceActor, TargetActor, Property);
+							CopySingleProperty(SourceActor, TargetActor, Property);
 
 							if( Options & ECopyOptions::CallPostEditChangeProperty )
 							{
@@ -7320,14 +7359,14 @@ namespace EditorUtilities
 								for( int32 InstanceIndex = 0; InstanceIndex < ArchetypeInstancesToChange.Num(); ++InstanceIndex )
 								{
 									UObject* ArchetypeInstance = ArchetypeInstancesToChange[InstanceIndex];
-									if( ArchetypeInstance != NULL )
+									if( ArchetypeInstance != nullptr )
 									{
 										if( bIsFirstModification )
 										{
 											ArchetypeInstance->Modify();
 										}
 
-										Property->CopyCompleteValue_InContainer( ArchetypeInstance, TargetActor );
+										CopySingleProperty( TargetActor, ArchetypeInstance, Property );
 									}
 								}
 							}
@@ -7355,7 +7394,7 @@ namespace EditorUtilities
 			UActorComponent* SourceComponent = *SourceComponentIter;
 			UActorComponent* TargetComponent = FindMatchingComponentInstance( SourceComponent, TargetActor, TargetComponents, TargetComponentIndex );
 
-			if( SourceComponent != NULL && TargetComponent != NULL )
+			if( SourceComponent != nullptr && TargetComponent != nullptr )
 			{
 				UClass* ComponentClass = SourceComponent->GetClass();
 				check( ComponentClass == TargetComponent->GetClass() );
@@ -7364,15 +7403,14 @@ namespace EditorUtilities
 				TArray<UActorComponent*> ComponentArchetypeInstances;
 				if( Options & ECopyOptions::PropagateChangesToArchetypeInstances )
 				{
-					for( int32 InstanceIndex = 0; InstanceIndex < ArchetypeInstances.Num(); ++InstanceIndex )
+					for( AActor* ArchetypeInstance : ArchetypeInstances )
 					{
-						AActor* ArchetypeInstance = Cast<AActor>(ArchetypeInstances[InstanceIndex]);
-						if( ArchetypeInstance != NULL )
+						if( ArchetypeInstance != nullptr )
 						{
 							UActorComponent* ComponentArchetypeInstance = FindMatchingComponentInstance( TargetComponent, ArchetypeInstance );
-							if( ComponentArchetypeInstance != NULL )
+							if( ComponentArchetypeInstance != nullptr )
 							{
-								ComponentArchetypeInstances.Add( ComponentArchetypeInstance );
+								ComponentArchetypeInstances.AddUnique( ComponentArchetypeInstance );
 							}
 						}
 					}
@@ -7383,7 +7421,7 @@ namespace EditorUtilities
 
 				// Copy component properties
 				bool bIsFirstModification = true;
-				for( UProperty* Property = ComponentClass->PropertyLink; Property != NULL; Property = Property->PropertyLinkNext )
+				for( UProperty* Property = ComponentClass->PropertyLink; Property != nullptr; Property = Property->PropertyLinkNext )
 				{
 					const bool bIsTransient = !!( Property->PropertyFlags & CPF_Transient );
 					const bool bIsIdentical = Property->Identical_InContainer( SourceComponent, TargetComponent );
@@ -7417,17 +7455,35 @@ namespace EditorUtilities
 								TArray<UActorComponent*> ComponentArchetypeInstancesToChange;
 								if( Options & ECopyOptions::PropagateChangesToArchetypeInstances )
 								{
-									for( int32 InstanceIndex = 0; InstanceIndex < ComponentArchetypeInstances.Num(); ++InstanceIndex )
+									for (UActorComponent* ComponentArchetypeInstance : ComponentArchetypeInstances)
 									{
-										UActorComponent* ComponentArchetypeInstance = Cast<UActorComponent>(ComponentArchetypeInstances[InstanceIndex]);
-										if( ComponentArchetypeInstance != NULL && Property->Identical_InContainer( ComponentArchetypeInstance, TargetComponent ) )
+										if( ComponentArchetypeInstance != nullptr && Property->Identical_InContainer( ComponentArchetypeInstance, TargetComponent ) )
 										{
-											ComponentArchetypeInstancesToChange.Add( ComponentArchetypeInstance );
+											bool bAdd = true;
+											// We also need to double check that either the direct archetype of the target is also identical
+											if (ComponentArchetypeInstance->GetArchetype() != TargetComponent)
+											{
+												UActorComponent* CheckComponent = CastChecked<UActorComponent>(ComponentArchetypeInstance->GetArchetype());
+												while (CheckComponent != ComponentArchetypeInstance)
+												{
+													if (!Property->Identical_InContainer( CheckComponent, TargetComponent ))
+													{
+														bAdd = false;
+														break;
+													}
+													CheckComponent = CastChecked<UActorComponent>(CheckComponent->GetArchetype());
+												}
+											}
+											
+											if (bAdd)
+											{
+												ComponentArchetypeInstancesToChange.Add( ComponentArchetypeInstance );
+											}
 										}
 									}
 								}
 
-								CopySingleActorProperty(SourceComponent, TargetComponent, Property);
+								CopySingleProperty(SourceComponent, TargetComponent, Property);
 
 								if( Options & ECopyOptions::CallPostEditChangeProperty )
 								{
@@ -7440,7 +7496,7 @@ namespace EditorUtilities
 									for( int32 InstanceIndex = 0; InstanceIndex < ComponentArchetypeInstancesToChange.Num(); ++InstanceIndex )
 									{
 										UActorComponent* ComponentArchetypeInstance = ComponentArchetypeInstancesToChange[InstanceIndex];
-										if( ComponentArchetypeInstance != NULL )
+										if( ComponentArchetypeInstance != nullptr )
 										{
 											if( bIsFirstModification )
 											{
@@ -7454,13 +7510,13 @@ namespace EditorUtilities
 
 												// We must also modify the owner, because we'll need script components to be reconstructed as part of an undo operation.
 												AActor* Owner = ComponentArchetypeInstance->GetOwner();
-												if( Owner != NULL )
+												if( Owner != nullptr )
 												{
 													Owner->Modify();
 												}
 											}
 
-											Property->CopyCompleteValue_InContainer( ComponentArchetypeInstance, TargetComponent );
+											CopySingleProperty( TargetComponent, ComponentArchetypeInstance, Property );
 
 											// Re-register the component with the scene
 											ComponentArchetypeInstance->ReregisterComponent();
