Detecting content changes for auto-reimport on restart now defaults to 'on'.
We also now pop up a warning explaining that this behaviour might be undesirable when running source control, giving the user the option to disable it.

diff --git a/Engine/Source/Developer/SourceControl/Private/SourceControlModule.cpp b/Engine/Source/Developer/SourceControl/Private/SourceControlModule.cpp
--- a/Engine/Source/Developer/SourceControl/Private/SourceControlModule.cpp
+++ b/Engine/Source/Developer/SourceControl/Private/SourceControlModule.cpp
@@ -8,6 +8,8 @@
 #include "SSourceControlLogin.h"
 
 #if WITH_EDITOR
+
+	#include "Settings/EditorLoadingSavingSettings.h"
 	#include "Runtime/Engine/Public/EngineAnalytics.h"
 	#include "Runtime/Analytics/Analytics/Public/Interfaces/IAnalyticsProvider.h"
 #endif
@@ -170,6 +172,8 @@ void FSourceControlModule::OnSourceControlDialogClosed(const TSharedRef<SWindow>
 	bTemporarilyDisabled = false;
 
 #if WITH_EDITOR
+	GetMutableDefault<UEditorLoadingSavingSettings>()->CheckSourceControlCompatability();
+
 	FString NewProvider = CurrentSourceControlProvider->GetName().ToString();
 	if( FEngineAnalytics::IsAvailable() && !ActiveProviderName.Equals( NewProvider, ESearchCase::IgnoreCase ))
 	{
diff --git a/Engine/Source/Editor/UnrealEd/Classes/Settings/EditorLoadingSavingSettings.h b/Engine/Source/Editor/UnrealEd/Classes/Settings/EditorLoadingSavingSettings.h
--- a/Engine/Source/Editor/UnrealEd/Classes/Settings/EditorLoadingSavingSettings.h
+++ b/Engine/Source/Editor/UnrealEd/Classes/Settings/EditorLoadingSavingSettings.h
@@ -84,6 +84,17 @@ public:
 	UPROPERTY(EditAnywhere, config, Category=Startup)
 	uint32 bRestoreOpenAssetTabsOnRestart:1;
 
+public:
+
+	/** Check that the current source control settings will play nicely with the auto-reimport feature */
+	void CheckSourceControlCompatability();
+
+private:
+
+	/** true when CheckSourceControlCompatability is enabled */
+	UPROPERTY(config)
+	bool bEnableSourceControlCompatabilityCheck;
+
 public:
 
 	/**Automatically reimports textures when a change to source content is detected */
diff --git a/Engine/Source/Editor/UnrealEd/Private/Settings/SettingsClasses.cpp b/Engine/Source/Editor/UnrealEd/Private/Settings/SettingsClasses.cpp
--- a/Engine/Source/Editor/UnrealEd/Private/Settings/SettingsClasses.cpp
+++ b/Engine/Source/Editor/UnrealEd/Private/Settings/SettingsClasses.cpp
@@ -7,6 +7,12 @@
 #include "Components/ArrowComponent.h"
 #include "AutoReimport/AutoReimportUtilities.h"
 
+#include "SNotificationList.h"
+#include "NotificationManager.h"
+#include "ISettingsModule.h"
+
+#define LOCTEXT_NAMESPACE "SettingsClasses"
+
 /* UContentBrowserSettings interface
  *****************************************************************************/
 
@@ -94,10 +100,11 @@ void UEditorExperimentalSettings::PostEditChangeProperty( struct FPropertyChange
 
 UEditorLoadingSavingSettings::UEditorLoadingSavingSettings( const FObjectInitializer& ObjectInitializer )
 	: Super(ObjectInitializer)
+	, bEnableSourceControlCompatabilityCheck(true)
 	, bMonitorContentDirectories(false)
 	, bAutoCreateAssets(true)
 	, bAutoDeleteAssets(true)
-	, bDetectChangesOnRestart(false)
+	, bDetectChangesOnRestart(true)
 	, bDeleteSourceFilesWithAssets(false)
 {
 	TextDiffToolPath.FilePath = TEXT("P4Merge.exe");
@@ -150,6 +157,60 @@ void UEditorLoadingSavingSettings::PostInitProperties()
 	Super::PostInitProperties();
 }
 
+void UEditorLoadingSavingSettings::CheckSourceControlCompatability()
+{
+	if (!bEnableSourceControlCompatabilityCheck || !bMonitorContentDirectories)
+	{
+		return;
+	}
+
+	if (ISourceControlModule::Get().IsEnabled() && bDetectChangesOnRestart)
+	{
+		// Persistent shared payload captured by the lambdas below
+		struct FPersistentPayload { TSharedPtr<SNotificationItem> Notification; };
+		TSharedRef<FPersistentPayload> Payload = MakeShareable(new FPersistentPayload);
+
+		FNotificationInfo Info(LOCTEXT("AutoReimport_NotificationTitle", "We noticed that your auto-reimport settings are set up to detect source content changes on restart.\nThis might cause unexpected behavior when starting up after getting latest from source control.\n\nWe recommend disabling this specific behavior."));
+
+ 		auto OnTurnOffClicked = [=]{
+			auto* Settings = GetMutableDefault<UEditorLoadingSavingSettings>();
+			Settings->bDetectChangesOnRestart = false;
+			Settings->SaveConfig();
+
+			Payload->Notification->SetEnabled(false);
+			Payload->Notification->Fadeout();
+		};
+		Info.ButtonDetails.Emplace(LOCTEXT("AutoReimport_TurnOff", "Don't detect changes on start-up"), FText(), FSimpleDelegate::CreateLambda(OnTurnOffClicked), SNotificationItem::ECompletionState::CS_None);
+
+ 		auto OnIgnoreClicked = [=]{
+
+			Payload->Notification->SetEnabled(false);
+ 			Payload->Notification->Fadeout();
+		};
+		Info.ButtonDetails.Emplace(LOCTEXT("AutoReimport_Ignore", "Ignore"), FText(), FSimpleDelegate::CreateLambda(OnIgnoreClicked), SNotificationItem::ECompletionState::CS_None);
+
+		Info.bUseLargeFont = false;
+		Info.bFireAndForget = false;
+
+		Info.CheckBoxStateChanged = FOnCheckStateChanged::CreateLambda([](ECheckBoxState State){
+			GetMutableDefault<UEditorLoadingSavingSettings>()->bEnableSourceControlCompatabilityCheck = (State != ECheckBoxState::Checked);
+		});
+		Info.CheckBoxText = LOCTEXT("AutoReimport_DontShowAgain", "Don't show again");
+
+		Info.Hyperlink = FSimpleDelegate::CreateLambda([]{
+			// Open Settings
+			ISettingsModule* SettingsModule = FModuleManager::GetModulePtr<ISettingsModule>("Settings");
+			if (SettingsModule != nullptr)
+			{
+				SettingsModule->ShowViewer("Editor", "General", "LoadingSaving");
+			}
+		});
+		Info.HyperlinkText = LOCTEXT("AutoReimport_OpenSettings", "Settings");
+
+		Payload->Notification = FSlateNotificationManager::Get().AddNotification(Info);
+	}
+}
+
 FAutoReimportDirectoryConfig::FParseContext::FParseContext(bool bInEnableLogging)
 	: bEnableLogging(bInEnableLogging)
 {
@@ -439,4 +500,6 @@ bool UProjectPackagingSettings::CanEditChange( const UProperty* InProperty ) con
 	}
 
 	return Super::CanEditChange(InProperty);
-}
\ No newline at end of file
+}
+
+#undef LOCTEXT_NAMESPACE
\ No newline at end of file
diff --git a/Engine/Source/Editor/UnrealEd/Private/UnrealEdEngine.cpp b/Engine/Source/Editor/UnrealEd/Private/UnrealEdEngine.cpp
--- a/Engine/Source/Editor/UnrealEd/Private/UnrealEdEngine.cpp
+++ b/Engine/Source/Editor/UnrealEd/Private/UnrealEdEngine.cpp
@@ -106,6 +106,8 @@ void UUnrealEdEngine::Init(IEngineLoop* InEngineLoop)
 	{
 		AutoReimportManager = NewObject<UAutoReimportManager>();
 		AutoReimportManager->Initialize();
+		
+		GetMutableDefault<UEditorLoadingSavingSettings>()->CheckSourceControlCompatability();
 	}
 
 	// register details panel customizations
