Ensure that advanced settings are visible when opening auto-reimport settings from notification

diff --git a/Engine/Source/Editor/UnrealEd/Private/Settings/SettingsClasses.cpp b/Engine/Source/Editor/UnrealEd/Private/Settings/SettingsClasses.cpp
--- a/Engine/Source/Editor/UnrealEd/Private/Settings/SettingsClasses.cpp
+++ b/Engine/Source/Editor/UnrealEd/Private/Settings/SettingsClasses.cpp
@@ -153,6 +153,7 @@ void UEditorLoadingSavingSettings::PostInitProperties()
 			Config.SourceDirectory = String;
 			AutoReimportDirectorySettings.Add(Config);
 		}
+		AutoReimportDirectories_DEPRECATED.Empty();
 	}
 	Super::PostInitProperties();
 }
@@ -202,6 +203,8 @@ void UEditorLoadingSavingSettings::CheckSourceControlCompatability()
 			ISettingsModule* SettingsModule = FModuleManager::GetModulePtr<ISettingsModule>("Settings");
 			if (SettingsModule != nullptr)
 			{
+				// Ensure that the advanced properties are visible
+				GConfig->SetBool(TEXT("DetailCategoriesAdvanced"), TEXT("EditorLoadingSavingSettings.AutoReimport"), true, GEditorUserSettingsIni);
 				SettingsModule->ShowViewer("Editor", "General", "LoadingSaving");
 			}
 		});
@@ -502,4 +505,4 @@ bool UProjectPackagingSettings::CanEditChange( const UProperty* InProperty ) con
 	return Super::CanEditChange(InProperty);
 }
 
-#undef LOCTEXT_NAMESPACE
\ No newline at end of file
+#undef LOCTEXT_NAMESPACE
