Autosaving is now prevented if the user makes any interation in the last 15 seconds.  This is to avoid annoying autosave pop ups happening instantly after a user finishes interacting

diff --git a/Engine/Source/Editor/UnrealEd/Private/PackageAutoSaver.cpp b/Engine/Source/Editor/UnrealEd/Private/PackageAutoSaver.cpp
--- a/Engine/Source/Editor/UnrealEd/Private/PackageAutoSaver.cpp
+++ b/Engine/Source/Editor/UnrealEd/Private/PackageAutoSaver.cpp
@@ -372,14 +372,11 @@ bool FPackageAutoSaver::CanAutoSave() const
 	// Don't allow auto-saving if the auto-save wouldn't save anything
 	const bool bPackagesNeedAutoSave = DoPackagesNeedAutoSave();
 
-	static FName EditableText("SEditableText");
-
 	double LastInteractionTime = FSlateApplication::Get().GetLastUserInteractionTime();
 	
-	const float EditableTextDelay = 15.0f;
-	TSharedPtr<SWidget> KeyboardFocusedWidget = FSlateApplication::Get().GetKeyboardFocusedWidget();
+	const float InteractionDelay = 15.0f;
 
-	bool bDidTypeInATextBlockRecently = KeyboardFocusedWidget.IsValid() && KeyboardFocusedWidget->GetType() == EditableText && (FApp::GetCurrentTime() - LastInteractionTime) < EditableTextDelay;
+	bool bDidInteractRecently = (FApp::GetCurrentTime() - LastInteractionTime) < InteractionDelay;
 
 	const bool bAutosaveEnabled	= GetDefault<UEditorLoadingSavingSettings>()->bAutoSaveEnable && bPackagesNeedAutoSave;
 	const bool bSlowTask = GIsSlowTask;
@@ -387,7 +384,7 @@ bool FPackageAutoSaver::CanAutoSave() const
 	const bool bPlayWorldValid = GUnrealEd->PlayWorld != nullptr;
 	const bool bAnyMenusVisible	= FSlateApplication::Get().AnyMenusVisible();
 	const bool bAutomationTesting = GIsAutomationTesting;
-	const bool bIsInteracting = FSlateApplication::Get().HasAnyMouseCaptor() || GUnrealEd->IsUserInteracting() || bDidTypeInATextBlockRecently;
+	const bool bIsInteracting = FSlateApplication::Get().HasAnyMouseCaptor() || GUnrealEd->IsUserInteracting() || bDidInteractRecently;
 	const bool bHasGameOrProjectLoaded = FApp::HasGameName();
 	const bool bAreShadersCompiling = GShaderCompilingManager->IsCompiling();
 
