Generate map dependencies seperately if needed for tools outside of editor.
UEPLAT-685 WIP

diff --git a/Engine/Source/Editor/UnrealEd/Classes/CookOnTheSide/CookOnTheFlyServer.h b/Engine/Source/Editor/UnrealEd/Classes/CookOnTheSide/CookOnTheFlyServer.h
--- a/Engine/Source/Editor/UnrealEd/Classes/CookOnTheSide/CookOnTheFlyServer.h
+++ b/Engine/Source/Editor/UnrealEd/Classes/CookOnTheSide/CookOnTheFlyServer.h
@@ -588,7 +588,8 @@ private:
 		FCookByTheBookOptions() : bGenerateStreamingInstallManifests(false),
 			bRunning(false),
 			CookTime( 0.0 ),
-			CookStartTime( 0.0 )
+			CookStartTime( 0.0 ), 
+			bGenerateDependeciesForMaps(false)
 		{ }
 
 		/** Should we test for UObject leaks */
@@ -611,6 +612,8 @@ private:
 		TArray<FFilePlatformRequest> PreviousCookRequests; 
 		double CookTime;
 		double CookStartTime;
+		/** Generate Map dependencies */
+		bool bGenerateDependeciesForMaps;
 	};
 	FCookByTheBookOptions* CookByTheBookOptions;
 	
@@ -634,6 +637,10 @@ private:
 	TAutoPtr<class FSandboxPlatformFile> SandboxFile;
 	bool bIsSavingPackage; // used to stop recursive mark package dirty functions
 
+	//////////////////////////////////////////////////////////////////////////
+	// Dependency graph of maps as root objects. 
+	TMap< FName, TSet <FName> > MapDependencyGraph; 
+
 	// data about the current package being processed
 	struct FReentryData
 	{
@@ -738,10 +745,12 @@ public:
 		FString DLCName;
 		FString CreateReleaseVersion;
 		FString BasedOnReleaseVersion;
+		bool bGenerateDependeciesForMaps; 
 
 		FCookByTheBookStartupOptions() :
 			CookOptions(ECookByTheBookOptions::None),
-			DLCName(FString())
+			DLCName(FString()),
+			bGenerateDependeciesForMaps(false)
 		{ }
 	};
 
diff --git a/Engine/Source/Editor/UnrealEd/Private/Commandlets/CookCommandlet.cpp b/Engine/Source/Editor/UnrealEd/Private/Commandlets/CookCommandlet.cpp
--- a/Engine/Source/Editor/UnrealEd/Private/Commandlets/CookCommandlet.cpp
+++ b/Engine/Source/Editor/UnrealEd/Private/Commandlets/CookCommandlet.cpp
@@ -1182,6 +1182,7 @@ bool UCookCommandlet::NewCook( const TArray<ITargetPlatform*>& Platforms, TArray
 	Swap( StartupOptions.BasedOnReleaseVersion, BasedOnReleaseVersion );
 	Swap( StartupOptions.CreateReleaseVersion, CreateReleaseVersion );
 	StartupOptions.CookOptions = CookOptions;
+	StartupOptions.bGenerateDependeciesForMaps = Switches.Contains(TEXT("GenerateDependenciesForMaps"));
 
 	CookOnTheFlyServer->StartCookByTheBook( StartupOptions );
 
diff --git a/Engine/Source/Editor/UnrealEd/Private/CookOnTheFlyServer.cpp b/Engine/Source/Editor/UnrealEd/Private/CookOnTheFlyServer.cpp
--- a/Engine/Source/Editor/UnrealEd/Private/CookOnTheFlyServer.cpp
+++ b/Engine/Source/Editor/UnrealEd/Private/CookOnTheFlyServer.cpp
@@ -789,6 +789,19 @@ void UCookOnTheFlyServer::GenerateManifestInfo( UPackage* Package, const TArray<
 		Packages.Add( Object->GetOutermost() );
 	}
 
+	if (CookByTheBookOptions->bGenerateDependeciesForMaps
+		&& Package->ContainsMap()
+		&& (World != nullptr)
+		)
+	{
+		TSet <FName> Names;
+		for (auto Object : Packages)
+		{
+			Names.Add(Object->GetFName());
+		}
+		MapDependencyGraph.Add(Package->GetFName(), Names);
+	}
+
 	// update the manifests with generated dependencies
 	for ( const auto& PlatformName : TargetPlatformNames )
 	{
@@ -2658,6 +2671,7 @@ void UCookOnTheFlyServer::CollectFilesToCook(TArray<FString>& FilesInPath, const
 	// the correct folder; etc.
 	{
 		TArray<FString> UIContentPaths;
+		TSet <FName> ContentDirectoryAssets; 
 		if (GConfig->GetArray(TEXT("UI"), TEXT("ContentDirectories"), UIContentPaths, GEditorIni) > 0)
 		{
 			for (int32 DirIdx = 0; DirIdx < UIContentPaths.Num(); DirIdx++)
@@ -2669,11 +2683,16 @@ void UCookOnTheFlyServer::CollectFilesToCook(TArray<FString>& FilesInPath, const
 				for (int32 Index = 0; Index < Files.Num(); Index++)
 				{
 					FString StdFile = Files[Index];
+					FName PackageName = FName(*FPackageName::FilenameToLongPackageName(StdFile));
+					ContentDirectoryAssets.Add(PackageName);
 					FPaths::MakeStandardFilename(StdFile);
 					AddFileToCook( FilesInPath, StdFile);
 				}
 			}
 		}
+
+		if (CookByTheBookOptions->bGenerateDependeciesForMaps) 
+			MapDependencyGraph.Add(FName(TEXT("ContentDirectoryAssets")), ContentDirectoryAssets);
 	}
 }
 
@@ -2801,6 +2820,37 @@ void UCookOnTheFlyServer::CookByTheBookFinished()
 		}
 	}
 
+	if (CookByTheBookOptions->bGenerateDependeciesForMaps)
+	{
+		FString MapDependencyGraphFile = FPaths::GameDir()/ TEXT("MapDependencyGraph.json");
+		// dump dependency graph. 
+		FString DependencyString; 
+		DependencyString += "{";
+		for (auto& Ele : MapDependencyGraph)
+		{
+			TSet <FName>& Deps = Ele.Value;
+			FName MapName = Ele.Key;
+			DependencyString += TEXT("\t\"") + MapName.ToString() + TEXT("\" : \n\t[\n ") ;
+			for (auto& Val : Deps)
+			{
+				DependencyString += TEXT("\t\t\"") + Val.ToString() + TEXT("\",\n");
+			}
+			DependencyString.RemoveFromEnd(TEXT(",\n"));
+			DependencyString += TEXT("\n\t],\n") ;
+		}
+		DependencyString.RemoveFromEnd(TEXT(",\n"));
+		DependencyString += "\n}";
+
+		ITargetPlatformManagerModule& TPM = GetTargetPlatformManagerRef();
+		const TArray<ITargetPlatform*>& Platforms = TPM.GetCookingTargetPlatforms();
+
+		for (ITargetPlatform* Platform : Platforms)
+		{
+			FString CookedMapDependencyGraphFilePlatform =  ConvertToFullSandboxPath(MapDependencyGraphFile, true).Replace(TEXT("[Platform]"), *Platform->PlatformName());
+			FFileHelper::SaveStringToFile(DependencyString, *CookedMapDependencyGraphFilePlatform, FFileHelper::EEncodingOptions::ForceUnicode);
+		}
+	}
+
 	CookByTheBookOptions->LastGCItems.Empty();
 	const float TotalCookTime = (float)(FPlatformTime::Seconds() - CookByTheBookOptions->CookStartTime);
 	UE_LOG(LogCook, Display, TEXT("Cook by the book total time in tick %fs total time %f"), CookByTheBookOptions->CookTime, TotalCookTime);
@@ -2947,6 +2997,7 @@ void UCookOnTheFlyServer::StartCookByTheBook( const FCookByTheBookStartupOptions
 	CookByTheBookOptions->bCancel = false;
 	CookByTheBookOptions->CookTime = 0.0f;
 	CookByTheBookOptions->CookStartTime = FPlatformTime::Seconds();
+	CookByTheBookOptions->bGenerateDependeciesForMaps = CookByTheBookStartupOptions.bGenerateDependeciesForMaps;
 	
 	CookByTheBookOptions->CreateReleaseVersion = CreateReleaseVersion;
 
