Paper2D: 2D Mode
- Combined 2D Rotate / Translate widget
- 2D Layer Snap planes
- Some UI stuff to go with it

diff --git a/Engine/Source/Editor/EditorSettingsViewer/Private/EditorSettingsViewerModule.cpp b/Engine/Source/Editor/EditorSettingsViewer/Private/EditorSettingsViewerModule.cpp
--- a/Engine/Source/Editor/EditorSettingsViewer/Private/EditorSettingsViewerModule.cpp
+++ b/Engine/Source/Editor/EditorSettingsViewer/Private/EditorSettingsViewerModule.cpp
@@ -174,6 +174,14 @@ protected:
 			GetMutableDefault<ULevelEditorViewportSettings>()
 		);
 
+		// view port settings
+		SettingsModule.RegisterSettings("Editor", "LevelEditor", "2D",
+			LOCTEXT("LevelEditor2DSettingsName", "2D"),
+			LOCTEXT("LevelEditor2DSettingsDescription", "Configure the settings for the 2D Level Editor."),
+			GetMutableDefault<ULevelEditor2DSettings>()
+		);
+
+
 		// miscellaneous settings
 		SettingsModule.RegisterSettings("Editor", "LevelEditor", "Misc",
 			LOCTEXT("LevelEditorMiscSettingsName", "Miscellaneous"),
diff --git a/Engine/Source/Editor/EditorStyle/Private/SlateEditorStyle.cpp b/Engine/Source/Editor/EditorStyle/Private/SlateEditorStyle.cpp
--- a/Engine/Source/Editor/EditorStyle/Private/SlateEditorStyle.cpp
+++ b/Engine/Source/Editor/EditorStyle/Private/SlateEditorStyle.cpp
@@ -4117,10 +4117,13 @@ void FSlateEditorStyle::FStyle::SetupLevelEditorStyle()
 		Set( "EditorViewport.ScaleMode.Small", new IMAGE_BRUSH( "Icons/icon_scaleb_16x", Icon16x16 ) );
 		Set( "EditorViewport.TranslateRotateMode", new IMAGE_BRUSH( "Icons/icon_translate_rotate_40x", Icon20x20 ) );
 		Set( "EditorViewport.TranslateRotateMode.Small", new IMAGE_BRUSH( "Icons/icon_translate_rotate_40x", Icon20x20 ) );
-		Set( "EditorViewport.ToggleRealTime", new IMAGE_BRUSH( "Icons/icon_MatEd_Realtime_40x", Icon40x40 ) );
+		Set( "EditorViewport.TranslateRotate2DMode", new IMAGE_BRUSH("Icons/icon_translate_rotate_2d_40x", Icon20x20));
+		Set( "EditorViewport.TranslateRotate2DMode.Small", new IMAGE_BRUSH("Icons/icon_translate_rotate_2d_40x", Icon20x20));
+		Set("EditorViewport.ToggleRealTime", new IMAGE_BRUSH("Icons/icon_MatEd_Realtime_40x", Icon40x40));
 		Set( "EditorViewport.ToggleRealTime.Small", new IMAGE_BRUSH( "Icons/icon_MatEd_Realtime_40x", Icon20x20 ) );
 		Set( "EditorViewport.LocationGridSnap", new IMAGE_BRUSH( "Old/LevelEditor/LocationGridSnap", Icon14x14, IconColor) );
 		Set( "EditorViewport.RotationGridSnap", new IMAGE_BRUSH( "Old/LevelEditor/RotationGridSnap", Icon14x14, IconColor ) );
+		Set( "EditorViewport.Layer2DSnap", new IMAGE_BRUSH("Old/LevelEditor/Layer2DSnap", Icon14x14, IconColor));
 		Set( "EditorViewport.ScaleGridSnap", new IMAGE_BRUSH( "Old/LevelEditor/ScaleGridSnap", Icon14x14, IconColor ) );
 		Set( "EditorViewport.ToggleSurfaceSnapping", new IMAGE_BRUSH( "Icons/icon_surface_snapping_14px", Icon14x14 ) );
 		Set( "EditorViewport.RelativeCoordinateSystem_Local", new IMAGE_BRUSH( "Icons/icon_axis_local_16px", Icon16x16, IconColor ) );
diff --git a/Engine/Source/Editor/UnrealEd/Classes/Settings/LevelEditorViewportSettings.h b/Engine/Source/Editor/UnrealEd/Classes/Settings/LevelEditorViewportSettings.h
--- a/Engine/Source/Editor/UnrealEd/Classes/Settings/LevelEditorViewportSettings.h
+++ b/Engine/Source/Editor/UnrealEd/Classes/Settings/LevelEditorViewportSettings.h
@@ -181,6 +181,55 @@ struct UNREALED_API FSnapToSurfaceSettings
 };
 
 
+/**
+* 2D layer settings
+*/
+USTRUCT()
+struct UNREALED_API FMode2DLayer
+{
+	GENERATED_USTRUCT_BODY()
+
+	FMode2DLayer()
+		: Name(TEXT("Default"))
+		, LayerOffset(0)
+	{ }
+
+	/** Whether snapping to surfaces in the world is enabled */
+	UPROPERTY(EditAnywhere, config, Category=Layer)
+	FString Name;
+
+	/** The amount of offset to apply when snapping to surfaces */
+	UPROPERTY(EditAnywhere, config, Category=Layer)
+	float LayerOffset;
+};
+
+/**
+* Implements the Level Editor's 2D view port settings.
+*/
+UCLASS(config = EditorUserSettings)
+class UNREALED_API ULevelEditor2DSettings
+	: public UObject
+{
+	GENERATED_UCLASS_BODY()
+
+public:
+	/** If enabled will allow 2D mode */
+	UPROPERTY(EditAnywhere, config, Category = General, meta = (DisplayName = "Enable 2D Mode"))
+	bool bMode2DEnabled;
+
+	/** When enabled and active snap layer is valid, items dropped in the viewport will be dropped on the appropriate Y distance away */
+	UPROPERTY(EditAnywhere, config, Category = LayerSnapping, AdvancedDisplay)
+	bool bEnableLayerSnap;
+
+	/** Will try to use snap settings from SnapLayers */
+	UPROPERTY(EditAnywhere, config, Category = LayerSnapping, AdvancedDisplay)
+	int32 ActiveSnapLayerIndex;
+
+	/** Snap layers that are displayed in the viewport toolbar */
+	UPROPERTY(EditAnywhere, config, Category = LayerSnapping)
+	TArray<FMode2DLayer> SnapLayers;
+};
+
 /**
  * Implements the Level Editor's view port settings.
  */
diff --git a/Engine/Source/Editor/UnrealEd/Private/EditorModeManager.cpp b/Engine/Source/Editor/UnrealEd/Private/EditorModeManager.cpp
--- a/Engine/Source/Editor/UnrealEd/Private/EditorModeManager.cpp
+++ b/Engine/Source/Editor/UnrealEd/Private/EditorModeManager.cpp
@@ -37,6 +37,7 @@ FEditorModeTools::FEditorModeTools()
 	,	Snapping( 0 )
 	,	SnappedActor( 0 )
 	,	TranslateRotateXAxisAngle(0)
+	,	TranslateRotate2DAngle(0)
 	,	DefaultID(FBuiltinEditorModes::EM_Default)
 	,	WidgetMode( FWidget::WM_None )
 	,	OverrideWidgetMode( FWidget::WM_None )
diff --git a/Engine/Source/Editor/UnrealEd/Private/EditorSelectUtils.cpp b/Engine/Source/Editor/UnrealEd/Private/EditorSelectUtils.cpp
--- a/Engine/Source/Editor/UnrealEd/Private/EditorSelectUtils.cpp
+++ b/Engine/Source/Editor/UnrealEd/Private/EditorSelectUtils.cpp
@@ -179,6 +179,7 @@ void UUnrealEdEngine::SetPivot( FVector NewPivot, bool bSnapPivotToGrid, bool bI
 
 	//default to using the x axis for the translate rotate widget
 	EditorModeTools.TranslateRotateXAxisAngle = 0.0f;
+	EditorModeTools.TranslateRotate2DAngle = 0.0f;
 	FVector TranslateRotateWidgetWorldXAxis;
 
 	AActor* LastSelectedActor = NULL;
diff --git a/Engine/Source/Editor/UnrealEd/Private/EditorViewportClient.cpp b/Engine/Source/Editor/UnrealEd/Private/EditorViewportClient.cpp
--- a/Engine/Source/Editor/UnrealEd/Private/EditorViewportClient.cpp
+++ b/Engine/Source/Editor/UnrealEd/Private/EditorViewportClient.cpp
@@ -833,7 +833,8 @@ bool FLevelEditorViewportClient::ShouldUseMoveCanvasMovement()
 		//OBJECT MOVEMENT CODE
 		if ( ( AltDown == false && ShiftDown == false && ( LeftMouseButtonDown ^ RightMouseButtonDown ) ) && 
 			( ( GetWidgetMode() == FWidget::WM_Translate && Widget->GetCurrentAxis() != EAxisList::None ) ||
-			( GetWidgetMode() == FWidget::WM_TranslateRotateZ && Widget->GetCurrentAxis() != EAxisList::ZRotation &&  Widget->GetCurrentAxis() != EAxisList::None ) ) )
+			( GetWidgetMode() == FWidget::WM_TranslateRotateZ && Widget->GetCurrentAxis() != EAxisList::ZRotation &&  Widget->GetCurrentAxis() != EAxisList::None ) ||
+			( GetWidgetMode() == FWidget::WM_2D && Widget->GetCurrentAxis() != EAxisList::Rotate2D &&  Widget->GetCurrentAxis() != EAxisList::None ) ) )
 		{
 			return true;
 		}
@@ -1127,7 +1128,7 @@ EMouseCursor::Type FEditorViewportClient::GetCursor(FViewport* InViewport,int32
 	}
 	else if (bMoveCanvasMovement && 
 		bHasMouseMovedSinceClick &&
-		( GetWidgetMode() == FWidget::WM_Translate || GetWidgetMode() == FWidget::WM_TranslateRotateZ ) )
+		(GetWidgetMode() == FWidget::WM_Translate || GetWidgetMode() == FWidget::WM_TranslateRotateZ || GetWidgetMode() == FWidget::WM_2D))
 	{
 		MouseCursor = EMouseCursor::CardinalCross;
 	}
@@ -2370,7 +2371,8 @@ bool FEditorViewportClient::ShouldUseMoveCanvasMovement() const
 		//OBJECT MOVEMENT CODE
 		if ( ( AltDown == false && ShiftDown == false && ( LeftMouseButtonDown ^ RightMouseButtonDown ) ) && 
 			( ( GetWidgetMode() == FWidget::WM_Translate && Widget->GetCurrentAxis() != EAxisList::None ) ||
-			( GetWidgetMode() == FWidget::WM_TranslateRotateZ && Widget->GetCurrentAxis() != EAxisList::ZRotation &&  Widget->GetCurrentAxis() != EAxisList::None ) ) )
+			( GetWidgetMode() == FWidget::WM_TranslateRotateZ && Widget->GetCurrentAxis() != EAxisList::ZRotation &&  Widget->GetCurrentAxis() != EAxisList::None ) ||
+			( GetWidgetMode() == FWidget::WM_2D && Widget->GetCurrentAxis() != EAxisList::Rotate2D &&  Widget->GetCurrentAxis() != EAxisList::None ) ) )
 		{
 			return true;
 		}
@@ -2808,6 +2810,10 @@ void FEditorViewportClient::ApplyDeltaToRotateWidget(const FRotator& InRot)
 		FRotator TranslateRotateWidgetRotation(0, ModeTools->TranslateRotateXAxisAngle, 0);
 		TranslateRotateWidgetRotation += InRot;
 		ModeTools->TranslateRotateXAxisAngle = TranslateRotateWidgetRotation.Yaw;
+
+		FRotator Widget2DRotation(ModeTools->TranslateRotate2DAngle, 0, 0);
+		Widget2DRotation += InRot;
+		ModeTools->TranslateRotate2DAngle = Widget2DRotation.Pitch;
 	}
 }
 
@@ -3765,7 +3771,8 @@ void FEditorViewportClient::UpdateRequiredCursorVisibility()
 		//Translating an object, but NOT moving the camera AND the object (shift)
 		if ( ( AltDown == false && ShiftDown == false && ( LeftMouseButtonDown ^ RightMouseButtonDown ) ) && 
 			( ( GetWidgetMode() == FWidget::WM_Translate && Widget->GetCurrentAxis() != EAxisList::None ) ||
-			(  GetWidgetMode() == FWidget::WM_TranslateRotateZ && Widget->GetCurrentAxis() != EAxisList::ZRotation &&  Widget->GetCurrentAxis() != EAxisList::None ) ) )
+			(  GetWidgetMode() == FWidget::WM_TranslateRotateZ && Widget->GetCurrentAxis() != EAxisList::ZRotation &&  Widget->GetCurrentAxis() != EAxisList::None ) ||
+			( GetWidgetMode() == FWidget::WM_2D && Widget->GetCurrentAxis() != EAxisList::Rotate2D &&  Widget->GetCurrentAxis() != EAxisList::None ) ) )
 		{
 			RequiredCursorVisibiltyAndAppearance.bHardwareCursorVisible = false;
 			RequiredCursorVisibiltyAndAppearance.bSoftwareCursorVisible = true;								
diff --git a/Engine/Source/Editor/UnrealEd/Private/EditorViewportCommands.cpp b/Engine/Source/Editor/UnrealEd/Private/EditorViewportCommands.cpp
--- a/Engine/Source/Editor/UnrealEd/Private/EditorViewportCommands.cpp
+++ b/Engine/Source/Editor/UnrealEd/Private/EditorViewportCommands.cpp
@@ -47,6 +47,7 @@ void FEditorViewportCommands::RegisterCommands()
 	UI_COMMAND( RotateMode, "Rotate Mode", "Select and rotate objects", EUserInterfaceActionType::ToggleButton, FInputChord(EKeys::E) );
 	UI_COMMAND( ScaleMode, "Scale Mode", "Select and scale objects", EUserInterfaceActionType::ToggleButton, FInputChord(EKeys::R) );
 	UI_COMMAND( TranslateRotateMode, "Combined Translate and Rotate Mode", "Select and translate or rotate objects", EUserInterfaceActionType::ToggleButton, FInputChord() );
+	UI_COMMAND( TranslateRotate2DMode, "2D Mode", "Select and translate or rotate objects in 2D", EUserInterfaceActionType::ToggleButton, FInputGesture());
 
 	UI_COMMAND( ShrinkTransformWidget, "Shrink Transform Widget", "Shrink the level editor transform widget", EUserInterfaceActionType::Button, FInputChord(EModifierKey::Alt, EKeys::LeftBracket) );
 	UI_COMMAND( ExpandTransformWidget, "Expand Transform Widget", "Expand the level editor transform widget", EUserInterfaceActionType::Button, FInputChord(EModifierKey::Alt, EKeys::RightBracket) );
@@ -65,6 +66,7 @@ void FEditorViewportCommands::RegisterCommands()
 
 	UI_COMMAND( LocationGridSnap, "Grid Snap", "Enables or disables snapping to the grid when dragging objects around", EUserInterfaceActionType::ToggleButton, FInputChord() );
 	UI_COMMAND( RotationGridSnap, "Rotation Snap", "Enables or disables snapping objects to a rotation grid", EUserInterfaceActionType::ToggleButton, FInputChord() );
+	UI_COMMAND( Layer2DSnap, "Layer2D Snap", "Enables or disables snapping objects to a 2D layer", EUserInterfaceActionType::ToggleButton, FInputGesture());
 	UI_COMMAND( ScaleGridSnap, "Scale Snap", "Enables or disables snapping objects to a scale grid", EUserInterfaceActionType::ToggleButton, FInputChord() );
 	UI_COMMAND( SurfaceSnapping, "Surface Snapping", "If enabled, actors will snap to surfaces in the world when dragging", EUserInterfaceActionType::ToggleButton, FInputChord() );
 
diff --git a/Engine/Source/Editor/UnrealEd/Private/LevelEditorViewport.cpp b/Engine/Source/Editor/UnrealEd/Private/LevelEditorViewport.cpp
--- a/Engine/Source/Editor/UnrealEd/Private/LevelEditorViewport.cpp
+++ b/Engine/Source/Editor/UnrealEd/Private/LevelEditorViewport.cpp
@@ -983,6 +983,41 @@ bool FLevelEditorViewportClient::HasDropPreviewActors() const
 	return DropPreviewActors.Num() > 0;
 }
 
+/* Helper functions to find a dropped position on a 2D layer */
+
+static bool IsDroppingOn2DLayer()
+{
+	const ULevelEditor2DSettings* Settings2D = GetDefault<ULevelEditor2DSettings>();
+	return Settings2D->bEnableLayerSnap && Settings2D->SnapLayers.IsValidIndex(Settings2D->ActiveSnapLayerIndex);
+}
+
+static FActorPositionTraceResult TraceForPositionOn2DLayer(const FViewportCursorLocation& Cursor)
+{
+	const ULevelEditor2DSettings* Settings2D = GetDefault<ULevelEditor2DSettings>();
+	check(Settings2D->SnapLayers.IsValidIndex(Settings2D->ActiveSnapLayerIndex));
+
+	float Offset = Settings2D->SnapLayers[Settings2D->ActiveSnapLayerIndex].LayerOffset;
+	FVector PlaneCenter(0, Offset, 0);
+	FVector PlaneNormal(0, -1, 0);
+
+	FActorPositionTraceResult Result;
+	float Numerator = FVector::DotProduct(PlaneCenter - Cursor.GetOrigin(), PlaneNormal);
+	float Denominator = FVector::DotProduct(PlaneNormal, Cursor.GetDirection());
+	if (FMath::Abs(Denominator) < SMALL_NUMBER)
+	{
+		Result.State = FActorPositionTraceResult::Failed;
+	}
+	else
+	{
+		Result.State = FActorPositionTraceResult::HitSuccess;
+		Result.SurfaceNormal = PlaneNormal;
+		float D = Numerator / Denominator;
+		Result.Location = Cursor.GetOrigin() + D * Cursor.GetDirection();
+	}
+
+	return Result;
+}
+
 bool FLevelEditorViewportClient::UpdateDropPreviewActors(int32 MouseX, int32 MouseY, const TArray<UObject*>& DroppedObjects, bool& out_bDroppedObjectsVisible, UActorFactory* FactoryToUse)
 {
 	out_bDroppedObjectsVisible = false;
@@ -1036,7 +1071,7 @@ bool FLevelEditorViewportClient::UpdateDropPreviewActors(int32 MouseX, int32 Mou
 	FSceneView* View = CalcSceneView( &ViewFamily );
 	FViewportCursorLocation Cursor(View, this, MouseX, MouseY);
 
-	const FActorPositionTraceResult TraceResult = FActorPositioning::TraceWorldForPositionWithDefault(Cursor, *View, &DraggingActors);
+	const FActorPositionTraceResult TraceResult = IsDroppingOn2DLayer() ? TraceForPositionOn2DLayer(Cursor) : FActorPositioning::TraceWorldForPositionWithDefault(Cursor, *View, &DraggingActors);
 
 	GEditor->UnsnappedClickLocation = TraceResult.Location;
 	GEditor->ClickLocation = TraceResult.Location;
@@ -1187,7 +1222,7 @@ bool FLevelEditorViewportClient::DropObjectsAtCoordinates(int32 MouseX, int32 Mo
 
 		HHitProxy* HitProxy = Viewport->GetHitProxy(Cursor.GetCursorPos().X, Cursor.GetCursorPos().Y);
 
-		const FActorPositionTraceResult TraceResult = FActorPositioning::TraceWorldForPositionWithDefault(Cursor, *View);
+		const FActorPositionTraceResult TraceResult = IsDroppingOn2DLayer() ? TraceForPositionOn2DLayer(Cursor) : FActorPositioning::TraceWorldForPositionWithDefault(Cursor, *View);
 		
 		GEditor->UnsnappedClickLocation = TraceResult.Location;
 		GEditor->ClickLocation = TraceResult.Location;
@@ -2530,6 +2565,9 @@ void FLevelEditorViewportClient::TrackingStarted( const FInputEventState& InInpu
 			case FWidget::WM_TranslateRotateZ:
 				TrackingDescription = FText::Format(LOCTEXT("TranslateRotateZTransaction", "Translate/RotateZ {0}"), ObjectTypeBeingTracked);
 				break;
+			case FWidget::WM_2D:
+				TrackingDescription = FText::Format(LOCTEXT("TranslateRotate2D", "Translate/Rotate2D {0}"), ObjectTypeBeingTracked);
+				break;
 			default:
 				if( bNudge )
 				{
diff --git a/Engine/Source/Editor/UnrealEd/Private/MouseDeltaTracker.cpp b/Engine/Source/Editor/UnrealEd/Private/MouseDeltaTracker.cpp
--- a/Engine/Source/Editor/UnrealEd/Private/MouseDeltaTracker.cpp
+++ b/Engine/Source/Editor/UnrealEd/Private/MouseDeltaTracker.cpp
@@ -104,6 +104,7 @@ void FMouseDeltaTracker::DetermineCurrentAxis(FEditorViewportClient* InViewportC
 
 				case FWidget::WM_Translate:
 				case FWidget::WM_TranslateRotateZ:
+				case FWidget::WM_2D:
 					switch( InViewportClient->ViewportType )
 					{
 						case LVT_Perspective:
@@ -266,6 +267,7 @@ void FMouseDeltaTracker::StartTracking(FEditorViewportClient* InViewportClient,
 		}
 		break;
 		case FWidget::WM_TranslateRotateZ:
+		case FWidget::WM_2D:
 			FSnappingUtils::SnapPointToGrid( StartSnapped, FVector(GEditor->GetGridSize(),GEditor->GetGridSize(),GEditor->GetGridSize()) );
 			break;
 
@@ -429,7 +431,9 @@ void FMouseDeltaTracker::AddDelta(FEditorViewportClient* InViewportClient, FKey
 		// Affect input delta by the camera speed
 
 		FWidget::EWidgetMode WidgetMode = InViewportClient->GetWidgetMode();
-		bool bIsRotation = (WidgetMode == FWidget::WM_Rotate) || ( ( WidgetMode == FWidget::WM_TranslateRotateZ ) && ( InViewportClient->GetCurrentWidgetAxis() == EAxisList::ZRotation ) );
+		bool bIsRotation = (WidgetMode == FWidget::WM_Rotate) 
+			|| ( ( WidgetMode == FWidget::WM_TranslateRotateZ ) && ( InViewportClient->GetCurrentWidgetAxis() == EAxisList::ZRotation ) )
+			|| ( ( WidgetMode == FWidget::WM_2D) && (InViewportClient->GetCurrentWidgetAxis() == EAxisList::Rotate2D ) );
 		if (bIsRotation)
 		{
 			Wk *= GetDefault<ULevelEditorViewportSettings>()->MouseSensitivty;
@@ -523,8 +527,9 @@ void FMouseDeltaTracker::AddDelta(FEditorViewportClient* InViewportClient, FKey
 			}
 			break;
 			case FWidget::WM_TranslateRotateZ:
+			case FWidget::WM_2D:
 			{
-				if (InViewportClient->GetCurrentWidgetAxis() == EAxisList::ZRotation)
+				if (InViewportClient->GetCurrentWidgetAxis() == EAxisList::Rotate2D)
 				{
 					FRotator Rotation( EndSnapped.X, EndSnapped.Y, EndSnapped.Z );
 					FSnappingUtils::SnapRotatorToGrid( Rotation );
diff --git a/Engine/Source/Editor/UnrealEd/Private/SEditorViewport.cpp b/Engine/Source/Editor/UnrealEd/Private/SEditorViewport.cpp
--- a/Engine/Source/Editor/UnrealEd/Private/SEditorViewport.cpp
+++ b/Engine/Source/Editor/UnrealEd/Private/SEditorViewport.cpp
@@ -248,7 +248,15 @@ void SEditorViewport::BindCommands()
 		FIsActionButtonVisible::CreateSP( this, &SEditorViewport::IsTranslateRotateModeVisible )
 		);
 
-	CommandListRef.MapAction( 
+	CommandListRef.MapAction(
+		Commands.TranslateRotate2DMode,
+		FExecuteAction::CreateSP(ClientRef, &FEditorViewportClient::SetWidgetMode, FWidget::WM_2D),
+		FCanExecuteAction::CreateSP(ClientRef, &FEditorViewportClient::CanSetWidgetMode, FWidget::WM_2D),
+		FIsActionChecked::CreateSP(this, &SEditorViewport::IsWidgetModeActive, FWidget::WM_2D),
+		FIsActionButtonVisible::CreateSP(this, &SEditorViewport::Is2DModeVisible)
+		);
+
+	CommandListRef.MapAction(
 		Commands.ShrinkTransformWidget,
 		FExecuteAction::CreateSP( ClientRef, &FEditorViewportClient::AdjustTransformWidgetSize, -1 )
 		);
@@ -489,7 +497,12 @@ bool SEditorViewport::IsTranslateRotateModeVisible() const
 	return GetDefault<ULevelEditorViewportSettings>()->bAllowTranslateRotateZWidget;
 }
 
-bool SEditorViewport::IsCoordSystemActive( ECoordSystem CoordSystem ) const
+bool SEditorViewport::Is2DModeVisible() const
+{
+	return GetDefault<ULevelEditor2DSettings>()->bMode2DEnabled;
+}
+
+bool SEditorViewport::IsCoordSystemActive(ECoordSystem CoordSystem) const
 {
 	return Client->GetWidgetCoordSystemSpace() == CoordSystem;
 }
@@ -509,6 +522,11 @@ void SEditorViewport::OnCycleWidgetMode()
 			++WidgetModeAsInt;
 		}
 
+		if ((WidgetModeAsInt == FWidget::WM_2D) && (!GetDefault<ULevelEditor2DSettings>()->bMode2DEnabled))
+		{
+			++WidgetModeAsInt;
+		}
+
 		if( WidgetModeAsInt == FWidget::WM_Max )
 		{
 			WidgetModeAsInt -= FWidget::WM_Max;
diff --git a/Engine/Source/Editor/UnrealEd/Private/STransformViewportToolbar.cpp b/Engine/Source/Editor/UnrealEd/Private/STransformViewportToolbar.cpp
--- a/Engine/Source/Editor/UnrealEd/Private/STransformViewportToolbar.cpp
+++ b/Engine/Source/Editor/UnrealEd/Private/STransformViewportToolbar.cpp
@@ -202,6 +202,10 @@ TSharedRef< SWidget > STransformViewportToolBar::MakeTransformToolBar( const TSh
 		// Scale Mode
 		static FName ScaleModeName = FName(TEXT("ScaleMode"));
 		ToolbarBuilder.AddToolBarButton( FEditorViewportCommands::Get().ScaleMode, NAME_None, TAttribute<FText>(), TAttribute<FText>(), TAttribute<FSlateIcon>(), ScaleModeName );
+
+		// 2D Mode
+		static FName TranslateRotate2DModeName = FName(TEXT("TranslateRotate2DMode"));
+		ToolbarBuilder.AddToolBarButton(FEditorViewportCommands::Get().TranslateRotate2DMode, NAME_None, TAttribute<FText>(), TAttribute<FText>(), TAttribute<FSlateIcon>(), TranslateRotate2DModeName);
 	}
 	ToolbarBuilder.EndBlockGroup();
 	ToolbarBuilder.EndSection();
@@ -264,7 +268,7 @@ TSharedRef< SWidget > STransformViewportToolBar::MakeTransformToolBar( const TSh
 									.IsChecked(this, &STransformViewportToolBar::IsRotationGridSnapChecked)
 									.OnCheckStateChanged(this, &STransformViewportToolBar::HandleToggleRotationGridSnap)
 									.Label(this, &STransformViewportToolBar::GetRotationGridLabel)
-									.OnGetMenuContent(this ,&STransformViewportToolBar::FillRotationGridSnapMenu)
+									.OnGetMenuContent(this, &STransformViewportToolBar::FillRotationGridSnapMenu)
 									.ToggleButtonToolTip(Command->GetDescription())
 									.MenuButtonToolTip(LOCTEXT("RotationGridSnap_ToolTip", "Set the Rotation Grid Snap value"))
 									.Icon(Command->GetIcon())
@@ -272,7 +276,32 @@ TSharedRef< SWidget > STransformViewportToolBar::MakeTransformToolBar( const TSh
 									, RotationSnapName );
 	}
 	ToolbarBuilder.EndSection();
-	
+
+	ToolbarBuilder.BeginSection("Layer2DSnap");
+	{
+		// Grab the existing UICommand 
+		FUICommandInfo* Command = FEditorViewportCommands::Get().Layer2DSnap.Get();
+
+		static FName Layer2DSnapName = FName(TEXT("Layer2DSnap"));
+
+		ToolbarBuilder.AddWidget(SNew(SViewportToolBarComboMenu)
+			.Cursor(EMouseCursor::Default)
+			.Style(ToolBarStyle)
+			.Visibility(this, &STransformViewportToolBar::IsLayer2DSnapVisible)
+			.IsChecked(this, &STransformViewportToolBar::IsLayer2DSnapChecked)
+			.OnCheckStateChanged(this, &STransformViewportToolBar::HandleToggleLayer2DSnap)
+			.Label(this, &STransformViewportToolBar::GetLayer2DLabel)
+			.OnGetMenuContent(this, &STransformViewportToolBar::FillLayer2DSnapMenu)
+			.ToggleButtonToolTip(Command->GetDescription())
+			.MenuButtonToolTip(LOCTEXT("Layer2DSnap_ToolTip", "Set the 2d layer snap value"))
+			.Icon(Command->GetIcon())
+			.ParentToolBar(SharedThis(this))
+			, Layer2DSnapName);
+	}
+	ToolbarBuilder.EndSection();
+
+
+
 	ToolbarBuilder.BeginSection("ScaleGridSnap");
 	{
 		// Grab the existing UICommand 
@@ -379,7 +408,18 @@ FText STransformViewportToolBar::GetLocationGridLabel() const
 
 FText STransformViewportToolBar::GetRotationGridLabel() const
 {
-	return FText::Format( LOCTEXT("GridRotation - Number - DegreeSymbol", "{0}\u00b0"), FText::AsNumber( GEditor->GetRotGridSize().Pitch ) );
+	return FText::Format(LOCTEXT("GridRotation - Number - DegreeSymbol", "{0}\u00b0"), FText::AsNumber(GEditor->GetRotGridSize().Pitch));
+}
+
+FText STransformViewportToolBar::GetLayer2DLabel() const
+{
+	const ULevelEditor2DSettings* Settings2D = GetDefault<ULevelEditor2DSettings>();
+	if (Settings2D->SnapLayers.IsValidIndex(Settings2D->ActiveSnapLayerIndex))
+	{
+		return FText::FromString(Settings2D->SnapLayers[Settings2D->ActiveSnapLayerIndex].Name);
+	}
+	
+	return FText();
 }
 
 FText STransformViewportToolBar::GetScaleGridLabel() const
@@ -460,6 +500,19 @@ void STransformViewportToolBar::SetScaleGridSize( int32 InIndex )
 	GEditor->SetScaleGridSize( InIndex );
 }
 
+/**
+	* Sets the scale grid size
+	*
+	* @param	InIndex	The new index of the scale grid size to use
+	*/
+void STransformViewportToolBar::SetLayer2D( int32 Layer2DIndex )
+{
+	ULevelEditor2DSettings* Settings2D = GetMutableDefault<ULevelEditor2DSettings>();
+	Settings2D->bEnableLayerSnap = true;
+	Settings2D->ActiveSnapLayerIndex = Layer2DIndex;
+	Settings2D->PostEditChange();
+}
+
 		
 /**
 	* Checks to see if the specified grid size index is the current grid size index
@@ -497,12 +550,17 @@ bool STransformViewportToolBar::IsRotationGridSizeChecked( int32 GridSizeIndex,
 	*
 	* @return	True if the specified scale grid size is the current one
 	*/
-bool STransformViewportToolBar::IsScaleGridSizeChecked( int32 GridSizeIndex )
+bool STransformViewportToolBar::IsScaleGridSizeChecked(int32 GridSizeIndex)
 {
 	const ULevelEditorViewportSettings* ViewportSettings = GetDefault<ULevelEditorViewportSettings>();
 	return (ViewportSettings->CurrentScalingGridSize == GridSizeIndex);
 }
 
+bool STransformViewportToolBar::IsLayer2DSelected(int32 LayerIndex)
+{
+	return (GetDefault<ULevelEditor2DSettings>()->ActiveSnapLayerIndex == LayerIndex);
+}
+
 void STransformViewportToolBar::TogglePreserveNonUniformScale()
 {
 	ULevelEditorViewportSettings* ViewportSettings = GetMutableDefault<ULevelEditorViewportSettings>();
@@ -554,7 +612,7 @@ TSharedRef<SWidget> STransformViewportToolBar::FillRotationGridSnapMenu()
 
 	return SNew(SUniformGridPanel)
 
-	+ SUniformGridPanel::Slot(0, 0)
+		+ SUniformGridPanel::Slot(0, 0)
 		[
 			BuildRotationGridCheckBoxList("Common", LOCTEXT("RotationCommonText", "Common"), ViewportSettings->CommonRotGridSizes, GridMode_Common)
 		]
@@ -565,6 +623,26 @@ TSharedRef<SWidget> STransformViewportToolBar::FillRotationGridSnapMenu()
 		];
 }
 
+TSharedRef<SWidget> STransformViewportToolBar::FillLayer2DSnapMenu()
+{
+	const ULevelEditor2DSettings* Settings2D = GetDefault<ULevelEditor2DSettings>();
+	int32 LayerCount = Settings2D->SnapLayers.Num();
+	const bool bInShouldCloseWindowAfterMenuSelection = true;
+	FMenuBuilder ShowMenuBuilder(bInShouldCloseWindowAfterMenuSelection, CommandList); // rename
+	for (int32 LayerIndex = 0; LayerIndex < LayerCount; ++LayerIndex)
+	{
+		FName LayerName(*Settings2D->SnapLayers[LayerIndex].Name);
+
+		FUIAction Action(FExecuteAction::CreateStatic(&STransformViewportToolBar::SetLayer2D, LayerIndex),
+			FCanExecuteAction(),
+			FIsActionChecked::CreateStatic(&STransformViewportToolBar::IsLayer2DSelected, LayerIndex));
+
+		ShowMenuBuilder.AddMenuEntry(FText::FromName(LayerName), FText::GetEmpty(), FSlateIcon(), Action, NAME_None, EUserInterfaceActionType::RadioButton);
+	}
+
+	return ShowMenuBuilder.MakeWidget();
+}
+
 TSharedRef<SWidget> STransformViewportToolBar::BuildRotationGridCheckBoxList(FName InExtentionHook, const FText& InHeading, const TArray<float>& InGridSizes, ERotationGridMode InGridMode) const
 {
 	const bool bShouldCloseWindowAfterMenuSelection = true;
@@ -661,6 +739,19 @@ ECheckBoxState STransformViewportToolBar::IsRotationGridSnapChecked() const
 	return GetDefault<ULevelEditorViewportSettings>()->RotGridEnabled ? ECheckBoxState::Checked : ECheckBoxState::Unchecked;
 }
 
+ECheckBoxState STransformViewportToolBar::IsLayer2DSnapChecked() const
+{
+	const ULevelEditor2DSettings* Settings2D = GetDefault<ULevelEditor2DSettings>();
+	bool bChecked = Settings2D->bEnableLayerSnap && Settings2D->SnapLayers.IsValidIndex(Settings2D->ActiveSnapLayerIndex);
+	return bChecked ? ECheckBoxState::Checked : ECheckBoxState::Unchecked;
+}
+
+EVisibility STransformViewportToolBar::IsLayer2DSnapVisible() const
+{
+	const ULevelEditor2DSettings* Settings2D = GetDefault<ULevelEditor2DSettings>();
+	return Settings2D->bMode2DEnabled ? EVisibility::Visible : EVisibility::Collapsed;
+}
+
 ECheckBoxState STransformViewportToolBar::IsScaleGridSnapChecked() const
 {
 	return GetDefault<ULevelEditorViewportSettings>()->SnapScaleEnabled ? ECheckBoxState::Checked : ECheckBoxState::Unchecked;
@@ -671,12 +762,27 @@ void STransformViewportToolBar::HandleToggleLocationGridSnap( ECheckBoxState InS
 	GUnrealEd->Exec( GEditor->GetEditorWorldContext().World(), *FString::Printf( TEXT("MODE GRID=%d"), !GetDefault<ULevelEditorViewportSettings>()->GridEnabled ? 1 : 0 ) );
 }
 
-void STransformViewportToolBar::HandleToggleRotationGridSnap( ECheckBoxState InState )
+void STransformViewportToolBar::HandleToggleRotationGridSnap(ECheckBoxState InState)
+{
+	GUnrealEd->Exec(GEditor->GetEditorWorldContext().World(), *FString::Printf(TEXT("MODE ROTGRID=%d"), !GetDefault<ULevelEditorViewportSettings>()->RotGridEnabled ? 1 : 0));
+}
+
+void STransformViewportToolBar::HandleToggleLayer2DSnap(ECheckBoxState InState)
 {
-	GUnrealEd->Exec( GEditor->GetEditorWorldContext().World(), *FString::Printf( TEXT("MODE ROTGRID=%d"), !GetDefault<ULevelEditorViewportSettings>()->RotGridEnabled ? 1 : 0 ) );
+	ULevelEditor2DSettings* Settings2D = GetMutableDefault<ULevelEditor2DSettings>();
+	if (!Settings2D->bEnableLayerSnap && Settings2D->SnapLayers.Num() > 0)
+	{
+		Settings2D->bEnableLayerSnap = true;
+		Settings2D->ActiveSnapLayerIndex = FMath::Clamp(Settings2D->ActiveSnapLayerIndex, 0, Settings2D->SnapLayers.Num() - 1);
+	}
+	else
+	{
+		Settings2D->bEnableLayerSnap = false;
+	}
+	Settings2D->PostEditChange();
 }
 
-void STransformViewportToolBar::HandleToggleScaleGridSnap( ECheckBoxState InState )
+void STransformViewportToolBar::HandleToggleScaleGridSnap(ECheckBoxState InState)
 {
 	GUnrealEd->Exec( GEditor->GetEditorWorldContext().World(), *FString::Printf( TEXT("MODE SCALEGRID=%d"), !GetDefault<ULevelEditorViewportSettings>()->SnapScaleEnabled ? 1 : 0 ) );
 }
diff --git a/Engine/Source/Editor/UnrealEd/Private/Settings/SettingsClasses.cpp b/Engine/Source/Editor/UnrealEd/Private/Settings/SettingsClasses.cpp
--- a/Engine/Source/Editor/UnrealEd/Private/Settings/SettingsClasses.cpp
+++ b/Engine/Source/Editor/UnrealEd/Private/Settings/SettingsClasses.cpp
@@ -344,6 +344,14 @@ ULevelEditorPlaySettings::ULevelEditorPlaySettings( const FObjectInitializer& Ob
 	RouteGamepadToSecondWindow = false;
 }
 
+/* ULevelEditor2DSettings 
+*****************************************************************************/
+
+ULevelEditor2DSettings::ULevelEditor2DSettings(const FObjectInitializer& ObjectInitializer)
+	: Super(ObjectInitializer)
+{
+}
+
 
 /* ULevelEditorViewportSettings interface
  *****************************************************************************/
diff --git a/Engine/Source/Editor/UnrealEd/Private/UnrealWidget.cpp b/Engine/Source/Editor/UnrealEd/Private/UnrealWidget.cpp
--- a/Engine/Source/Editor/UnrealEd/Private/UnrealWidget.cpp
+++ b/Engine/Source/Editor/UnrealEd/Private/UnrealWidget.cpp
@@ -10,6 +10,7 @@ IMPLEMENT_HIT_PROXY(HWidgetAxis,HHitProxy);
 
 static const float AXIS_LENGTH = 35.0f;
 static const float TRANSLATE_ROTATE_AXIS_CIRCLE_RADIUS = 20.0f;
+static const float TWOD_AXIS_CIRCLE_RADIUS = 10.0f;
 static const float INNER_AXIS_CIRCLE_RADIUS = 48.0f;
 static const float OUTER_AXIS_CIRCLE_RADIUS = 56.0f;
 static const float ROTATION_TEXT_RADIUS = 75.0f;
@@ -186,6 +187,10 @@ void FWidget::Render( const FSceneView* View,FPrimitiveDrawInterface* PDI, FEdit
 			Render_TranslateRotateZ( View, PDI, ViewportClient, Loc, bDrawWidget );
 			break;
 
+		case WM_2D:
+			Render_2D( View, PDI, ViewportClient, Loc, bDrawWidget );
+			break;
+
 		default:
 			break;
 	}
@@ -768,6 +773,160 @@ void FWidget::Render_TranslateRotateZ( const FSceneView* View, FPrimitiveDrawInt
 	}
 }
 
+/**
+* Draws the 2D widget.
+*/
+
+void FWidget::Render_2D(const FSceneView* View, FPrimitiveDrawInterface* PDI, FEditorViewportClient* ViewportClient, const FVector& InLocation, bool bDrawWidget)
+{
+	//////////////////////////////////////////////////////////////////////////
+	// Translation subwidget
+	//////////////////////////////////////////////////////////////////////////
+
+	// Figure out axis colors
+	const FLinearColor& XColor = (CurrentAxis&EAxisList::X ? (FLinearColor)CurrentColor : AxisColorX);
+	const FLinearColor& YColor = (CurrentAxis&EAxisList::Y ? (FLinearColor)CurrentColor : AxisColorY);
+	const FLinearColor& ZColor = (CurrentAxis&EAxisList::Z ? (FLinearColor)CurrentColor : AxisColorZ);
+	FColor CurrentScreenColor = (CurrentAxis & EAxisList::Screen ? CurrentColor : ScreenSpaceColor);
+
+	// Figure out axis matrices
+	FMatrix WidgetMatrix = CustomCoordSystem * FTranslationMatrix(InLocation);
+
+	bool bIsPerspective = (View->ViewMatrices.ProjMatrix.M[3][3] < 1.0f);
+	const bool bIsOrthoXY = !bIsPerspective && FMath::Abs(View->ViewMatrices.ViewMatrix.M[2][2]) > 0.0f;
+	const bool bIsOrthoXZ = !bIsPerspective && FMath::Abs(View->ViewMatrices.ViewMatrix.M[1][2]) > 0.0f;
+	const bool bIsOrthoYZ = !bIsPerspective && FMath::Abs(View->ViewMatrices.ViewMatrix.M[0][2]) > 0.0f;
+
+	// For local space widgets, we always want to draw all three axis, since they may not be aligned with
+	// the orthographic projection anyway.
+	const bool bIsLocalSpace = (ViewportClient->GetWidgetCoordSystemSpace() == COORD_Local);
+
+	EAxisList::Type DrawAxis = EAxisList::None;
+	if (bIsOrthoXY)
+	{
+		DrawAxis = EAxisList::X;
+	}
+	else if (bIsOrthoXZ)
+	{
+		DrawAxis = EAxisList::XZ;
+	}
+	else if (bIsOrthoYZ)
+	{
+		DrawAxis = EAxisList::Z;
+	}
+	else if (bIsPerspective)
+	{
+		// Find the best plane to move on
+		const FVector CameraZAxis = View->ViewMatrices.ViewMatrix.GetColumn(2);
+		const FVector LargestAxis = CameraZAxis.GetAbs();
+		if (LargestAxis.X > LargestAxis.Y)
+		{
+			if (LargestAxis.Z > LargestAxis.X)
+			{
+				DrawAxis = EAxisList::X;
+			}
+			else
+			{
+				DrawAxis = EAxisList::Z;
+			}
+		}
+		else
+		{
+			if (LargestAxis.Y > LargestAxis.Z)
+			{
+				DrawAxis = EAxisList::XZ;
+			}
+			else
+			{
+				DrawAxis = EAxisList::X;
+			}
+		}
+	}
+
+	const bool bDisabled = IsWidgetDisabled();
+
+	FVector Scale;
+	float UniformScale = View->WorldToScreen(InLocation).W * (4.0f / View->ViewRect.Width() / View->ViewMatrices.ProjMatrix.M[0][0]);
+
+	if (bIsOrthoXY)
+	{
+		Scale = FVector(UniformScale, UniformScale, 1.0f);
+	}
+	else if (bIsOrthoXZ)
+	{
+		Scale = FVector(UniformScale, 1.0f, UniformScale);
+	}
+	else if (bIsOrthoYZ)
+	{
+		Scale = FVector(1.0f, UniformScale, UniformScale);
+	}
+	else
+	{
+		Scale = FVector(UniformScale, UniformScale, UniformScale);
+	}
+
+	// Radius
+	const float ScaledRadius = (TWOD_AXIS_CIRCLE_RADIUS * UniformScale) + GetDefault<ULevelEditorViewportSettings>()->TransformWidgetSizeAdjustment;
+	const int32 CircleSides = (GetDefault<ULevelEditorViewportSettings>()->TransformWidgetSizeAdjustment > 0)
+		? AXIS_CIRCLE_SIDES + (GetDefault<ULevelEditorViewportSettings>()->TransformWidgetSizeAdjustment / 5)
+		: AXIS_CIRCLE_SIDES;
+
+	// Draw the grabbers
+	if (bDrawWidget)
+	{
+		FVector CornerPos = FVector(7, 0, 7) * UniformScale;
+		FVector AxisSize = FVector(12, 1.2, 12) * UniformScale;
+		float CornerLength = 1.2f * UniformScale;
+
+		// Draw the axis lines with arrow heads
+		if (DrawAxis&EAxisList::X && (bIsPerspective || bIsLocalSpace || !bIsOrthoYZ))
+		{
+			UMaterialInstanceDynamic* XMaterial = (CurrentAxis&EAxisList::X ? CurrentAxisMaterial : AxisMaterialX);
+			Render_Axis(View, PDI, EAxisList::X, WidgetMatrix, XMaterial, XColor, XAxisDir, Scale, bDrawWidget);
+		}
+
+		if (DrawAxis&EAxisList::Y && (bIsPerspective || bIsLocalSpace || !bIsOrthoXZ))
+		{
+			UMaterialInstanceDynamic* YMaterial = (CurrentAxis&EAxisList::Y ? CurrentAxisMaterial : AxisMaterialY);
+			Render_Axis(View, PDI, EAxisList::Y, WidgetMatrix, YMaterial, YColor, YAxisDir, Scale, bDrawWidget);
+		}
+
+		if (DrawAxis&EAxisList::Z && (bIsPerspective || bIsLocalSpace || !bIsOrthoXY))
+		{
+			UMaterialInstanceDynamic* ZMaterial = (CurrentAxis&EAxisList::Z ? CurrentAxisMaterial : AxisMaterialZ);
+			Render_Axis(View, PDI, EAxisList::Z, WidgetMatrix, ZMaterial, ZColor, ZAxisDir, Scale, bDrawWidget);
+		}
+
+		// After the primitives have been scaled and transformed, we apply this inverse scale that flattens the dimension
+		// that was scaled up to prevent it from intersecting with the near plane.  In perspective this won't have any effect,
+		// but in the ortho viewports it will prevent scaling in the direction of the camera and thus intersecting the near plane.
+		FVector FlattenScale = FVector(Scale.Component(0) == 1.0f ? 1.0f / UniformScale : 1.0f, Scale.Component(1) == 1.0f ? 1.0f / UniformScale : 1.0f, Scale.Component(2) == 1.0f ? 1.0f / UniformScale : 1.0f);
+		float ArrowRadius = ScaledRadius * 2.0f;
+
+		if (((DrawAxis&EAxisList::XZ) == EAxisList::XZ) && (bIsPerspective || bIsLocalSpace || bIsOrthoXZ)) // Front
+		{
+			uint8 Alpha = ((CurrentAxis&EAxisList::XZ) == EAxisList::XZ ? 0x3f : 0x0f);
+			PDI->SetHitProxy(new HWidgetAxis(EAxisList::XZ, bDisabled));
+			{
+				FColor Color = YColor;
+				DrawCircle(PDI, InLocation, CustomCoordSystem.TransformPosition(FVector(1, 0, 0)), CustomCoordSystem.TransformPosition(FVector(0, 0, 1)), Color, ScaledRadius, CircleSides, SDPG_Foreground);
+				Color.A = Alpha;
+				DrawDisc(PDI, InLocation, CustomCoordSystem.TransformPosition(FVector(1, 0, 0)), CustomCoordSystem.TransformPosition(FVector(0, 0, 1)), Color, ScaledRadius, CircleSides, TransparentPlaneMaterialXY->GetRenderProxy(false), SDPG_Foreground);
+			}
+			PDI->SetHitProxy(NULL);
+
+			PDI->SetHitProxy(new HWidgetAxis(EAxisList::Rotate2D, bDisabled));
+			{
+				FVector XAxis = CustomCoordSystem.TransformPosition(FVector(1, 0, 0).RotateAngleAxis((EditorModeTools ? EditorModeTools->TranslateRotate2DAngle : 0), FVector(0, -1, 0)));
+				FVector YAxis = CustomCoordSystem.TransformPosition(FVector(0, 0, 1).RotateAngleAxis((EditorModeTools ? EditorModeTools->TranslateRotate2DAngle : 0), FVector(0, -1, 0)));
+				FVector BaseArrowPoint = InLocation + XAxis * ScaledRadius;
+				DrawFlatArrow(PDI, BaseArrowPoint, XAxis, YAxis, YColor, ArrowRadius, ArrowRadius*.5f, TransparentPlaneMaterialXY->GetRenderProxy(false), SDPG_Foreground);
+			}
+			PDI->SetHitProxy(NULL);
+		}
+	}
+}
+
 /**
  * Converts mouse movement on the screen to widget axis movement/rotation.
  */
@@ -907,6 +1066,7 @@ void FWidget::ConvertMouseMovementToAxisMovement( FEditorViewportClient* InViewp
 			{
 				if( CurrentAxis == EAxisList::ZRotation )
 				{
+
 					const FVector2D AxisDir = bIsOrthoDrawingFullRing ? TangentDir : ZAxisDir;
 					FRotator Rotation = FRotator(0, FVector2D::DotProduct(AxisDir, DragDir), 0);
 					FSnappingUtils::SnapRotatorToGrid(Rotation);
@@ -939,6 +1099,51 @@ void FWidget::ConvertMouseMovementToAxisMovement( FEditorViewportClient* InViewp
 				}
 			}
 			break;
+
+		case WM_2D:
+		{
+			if (CurrentAxis == EAxisList::Rotate2D)
+			{
+				// TODO: Determine why -TangentDir is necessary here, and fix whatever is causing it
+				const FVector2D AxisDir = bIsOrthoDrawingFullRing ? -TangentDir : YAxisDir;
+
+				FRotator Rotation = FRotator(FVector2D::DotProduct(AxisDir, DragDir), 0, 0);
+				FSnappingUtils::SnapRotatorToGrid(Rotation);
+
+				CurrentDeltaRotation = Rotation.Pitch;
+				FVector2D EffectiveDelta = AxisDir * Rotation.Pitch;
+
+
+				// Adjust the input delta according to how much rotation was actually applied
+				InOutDelta = FVector(EffectiveDelta.X, -EffectiveDelta.Y, 0.0f);
+
+				// Need to get the delta rotation in the current coordinate space of the widget
+				OutRotation = (CustomCoordSystem.Inverse() * FRotationMatrix(Rotation) * CustomCoordSystem).Rotator();
+			}
+			else
+			{
+				// Get drag delta in widget axis space
+				OutDrag = FVector(
+					(CurrentAxis & EAxisList::X) ? FVector2D::DotProduct(XAxisDir, DragDir) : 0.0f,
+					(CurrentAxis & EAxisList::Y) ? FVector2D::DotProduct(YAxisDir, DragDir) : 0.0f,
+					(CurrentAxis & EAxisList::Z) ? FVector2D::DotProduct(ZAxisDir, DragDir) : 0.0f
+					);
+
+				// Snap to grid in widget axis space
+				const FVector GridSize = FVector(GEditor->GetGridSize());
+				FSnappingUtils::SnapPointToGrid(OutDrag, GridSize);
+
+				// Convert to effective screen space delta, and replace input delta, adjusted for inverted screen space Y axis
+				const FVector2D EffectiveDelta = OutDrag.X * XAxisDir + OutDrag.Y * YAxisDir + OutDrag.Z * ZAxisDir;
+				InOutDelta = FVector(EffectiveDelta.X, -EffectiveDelta.Y, 0.0f);
+
+				// Transform drag delta into world space
+				OutDrag = CustomCoordSystem.TransformPosition(OutDrag);
+			}
+		}
+			break;
+
+
 		default:
 			break;
 	}
@@ -1038,6 +1243,69 @@ void FWidget::AbsoluteTranslationConvertMouseMovementToAxisMovement(FSceneView*
 			break;
 		}
 
+		case WM_2D:
+		{
+			switch (CurrentAxis)
+			{
+				case EAxisList::X:
+				{
+					GetAxisPlaneNormalAndMask(InputCoordSystem, Params.XAxis, Params.CameraDir, Params.PlaneNormal, Params.NormalToRemove);
+					OutDrag = GetAbsoluteTranslationDelta(Params);
+					break;
+				}
+				case EAxisList::Z:
+				{
+					GetAxisPlaneNormalAndMask(InputCoordSystem, Params.ZAxis, Params.CameraDir, Params.PlaneNormal, Params.NormalToRemove);
+					OutDrag = GetAbsoluteTranslationDelta(Params);
+					break;
+				}
+				case EAxisList::XZ:
+				{
+					GetPlaneNormalAndMask(Params.YAxis, Params.PlaneNormal, Params.NormalToRemove);
+					OutDrag = GetAbsoluteTranslationDelta(Params);
+					break;
+				}
+			
+				//Rotate about the y-axis
+				case EAxisList::Rotate2D:
+				{
+					//no position snapping, we'll handle the rotation snapping elsewhere
+					Params.bPositionSnapping = false;
+
+					GetPlaneNormalAndMask(Params.YAxis, Params.PlaneNormal, Params.NormalToRemove);
+					//No DAMPING
+					Params.bMovementLockedToCamera = false;
+					//this is the one movement type where we want to always use the widget origin and 
+					//NOT the "first click" origin
+					FVector XZPlaneProjectedPosition = GetAbsoluteTranslationDelta(Params) + InitialTranslationOffset;
+
+					//remove the component along the normal we want to mute
+					float MovementAlongMutedAxis = XZPlaneProjectedPosition | Params.NormalToRemove;
+					XZPlaneProjectedPosition = XZPlaneProjectedPosition - (Params.NormalToRemove*MovementAlongMutedAxis);
+
+					if (!XZPlaneProjectedPosition.Normalize())
+					{
+						XZPlaneProjectedPosition = Params.YAxis;
+					}
+
+					//NOW, find the rotation around the PlaneNormal to make the xaxis point at InDrag
+					OutRotation = FRotator::ZeroRotator;
+
+					float PitchDegrees = -FMath::Atan2(-XZPlaneProjectedPosition.Z, XZPlaneProjectedPosition.X) * 180.f / PI;
+					OutRotation.Pitch = PitchDegrees - (EditorModeTools ? EditorModeTools->TranslateRotate2DAngle : 0);
+
+					if (bSnapEnabled)
+					{
+						FSnappingUtils::SnapRotatorToGrid(OutRotation);
+					}
+
+					break;
+				}
+			}
+
+			break;
+		}
+
 		case WM_TranslateRotateZ:
 		{
 			FVector LineToUse;
@@ -1116,7 +1384,7 @@ void FWidget::AbsoluteTranslationConvertMouseMovementToAxisMovement(FSceneView*
 /** Only some modes support Absolute Translation Movement */
 bool FWidget::AllowsAbsoluteTranslationMovement(EWidgetMode WidgetMode)
 {
-	if ((WidgetMode == WM_Translate) || (WidgetMode == WM_TranslateRotateZ))
+	if ((WidgetMode == WM_Translate) || (WidgetMode == WM_TranslateRotateZ) || (WidgetMode == WM_2D))
 	{
 		return true;
 	}
diff --git a/Engine/Source/Editor/UnrealEd/Public/EditorModeManager.h b/Engine/Source/Editor/UnrealEd/Public/EditorModeManager.h
--- a/Engine/Source/Editor/UnrealEd/Public/EditorModeManager.h
+++ b/Engine/Source/Editor/UnrealEd/Public/EditorModeManager.h
@@ -343,6 +343,9 @@ public:
 	/** The angle for the translate rotate widget */
 	float TranslateRotateXAxisAngle;
 
+	/** The angles for the 2d translate rotate widget */
+	float TranslateRotate2DAngle;
+
 	/** Draws in the top level corner of all FEditorViewportClient windows (can be used to relay info to the user). */
 	FString InfoString;
 
diff --git a/Engine/Source/Editor/UnrealEd/Public/EditorViewportCommands.h b/Engine/Source/Editor/UnrealEd/Public/EditorViewportCommands.h
--- a/Engine/Source/Editor/UnrealEd/Public/EditorViewportCommands.h
+++ b/Engine/Source/Editor/UnrealEd/Public/EditorViewportCommands.h
@@ -117,6 +117,9 @@ public:
 	/** TranslateRotate Mode */
 	TSharedPtr< FUICommandInfo > TranslateRotateMode;
 
+	/** TranslateRotate Mode */
+	TSharedPtr< FUICommandInfo > TranslateRotate2DMode;
+
 	/** Shrink the level editor transform widget */
 	TSharedPtr< FUICommandInfo > ShrinkTransformWidget;
 
@@ -161,6 +164,8 @@ public:
 	/** Enables or disables snapping to a rotational grid while rotating objects */
 	TSharedPtr< FUICommandInfo > RotationGridSnap;
 
+	/** Enables or disables snapping to a 2D layer */
+	TSharedPtr< FUICommandInfo > Layer2DSnap;
 
 	/** Enables or disables snapping to a scaling grid while scaling objects */
 	TSharedPtr< FUICommandInfo > ScaleGridSnap;
diff --git a/Engine/Source/Editor/UnrealEd/Public/SEditorViewport.h b/Engine/Source/Editor/UnrealEd/Public/SEditorViewport.h
--- a/Engine/Source/Editor/UnrealEd/Public/SEditorViewport.h
+++ b/Engine/Source/Editor/UnrealEd/Public/SEditorViewport.h
@@ -138,6 +138,11 @@ protected:
 	 */
 	virtual bool IsTranslateRotateModeVisible() const;
 
+	/**
+	* @return true if the 2d mode widget is visible
+	*/
+	virtual bool Is2DModeVisible() const;
+
 	/**
 	 * Moves between widget modes
 	 */
diff --git a/Engine/Source/Editor/UnrealEd/Public/STransformViewportToolbar.h b/Engine/Source/Editor/UnrealEd/Public/STransformViewportToolbar.h
--- a/Engine/Source/Editor/UnrealEd/Public/STransformViewportToolbar.h
+++ b/Engine/Source/Editor/UnrealEd/Public/STransformViewportToolbar.h
@@ -49,22 +49,26 @@ private:
 	/** Grid snap label callbacks */
 	FText GetLocationGridLabel() const;
 	FText GetRotationGridLabel() const;
+	FText GetLayer2DLabel() const;
 	FText GetScaleGridLabel() const;
 
 	/** GridSnap menu construction callbacks */
 	TSharedRef<SWidget> FillLocationGridSnapMenu();
 	TSharedRef<SWidget> FillRotationGridSnapMenu();
+	TSharedRef<SWidget> FillLayer2DSnapMenu();
 	TSharedRef<SWidget> FillScaleGridSnapMenu();
 
 	/** Grid snap setting callbacks */
 	static void SetGridSize( int32 InIndex );
 	static void SetRotationGridSize( int32 InIndex, ERotationGridMode InGridMode );
 	static void SetScaleGridSize( int32 InIndex );
+	static void SetLayer2D(int32 Layer2DIndex);
 
 	/** Grid snap is checked callbacks for the menu values */
 	static bool IsGridSizeChecked( int32 GridSizeIndex );
 	static bool IsRotationGridSizeChecked( int32 GridSizeIndex, ERotationGridMode GridMode );
 	static bool IsScaleGridSizeChecked( int32 GridSizeIndex );
+	static bool IsLayer2DSelected(int32 Later2DIndex);
 
 	/** Callbacks for preserving non-uniform scaling when snapping */
 	static void TogglePreserveNonUniformScale();
@@ -80,11 +84,15 @@ private:
 	/** Grid Snap checked state callbacks */
 	ECheckBoxState IsLocationGridSnapChecked() const;
 	ECheckBoxState IsRotationGridSnapChecked() const;
+	ECheckBoxState IsLayer2DSnapChecked() const;
 	ECheckBoxState IsScaleGridSnapChecked() const;
 
+	EVisibility IsLayer2DSnapVisible() const;
+
 	/** Grid snap toggle handlers */
 	void HandleToggleLocationGridSnap(ECheckBoxState InState);
 	void HandleToggleRotationGridSnap(ECheckBoxState InState);
+	void HandleToggleLayer2DSnap(ECheckBoxState InState);
 	void HandleToggleScaleGridSnap(ECheckBoxState InState);
 
 private:
diff --git a/Engine/Source/Editor/UnrealEd/Public/UnrealWidget.h b/Engine/Source/Editor/UnrealEd/Public/UnrealWidget.h
--- a/Engine/Source/Editor/UnrealEd/Public/UnrealWidget.h
+++ b/Engine/Source/Editor/UnrealEd/Public/UnrealWidget.h
@@ -18,6 +18,7 @@ public:
 		WM_TranslateRotateZ,
 		WM_Rotate,
 		WM_Scale,
+		WM_2D,
 		WM_Max,
 	};
 
@@ -71,6 +72,11 @@ public:
 	*/
 	void Render_TranslateRotateZ( const FSceneView* View, FPrimitiveDrawInterface* PDI, FEditorViewportClient* ViewportClient, const FVector& InLocation, bool bDrawWidget );
 
+	/**
+	* Draws the combined 2D widget.
+	*/
+	void Render_2D(const FSceneView* View, FPrimitiveDrawInterface* PDI, FEditorViewportClient* ViewportClient, const FVector& InLocation, bool bDrawWidget);
+
 	/**
 	 * Converts mouse movement on the screen to widget axis movement/rotation.
 	 */
diff --git a/Engine/Source/Runtime/Core/Public/Math/Axis.h b/Engine/Source/Runtime/Core/Public/Math/Axis.h
--- a/Engine/Source/Runtime/Core/Public/Math/Axis.h
+++ b/Engine/Source/Runtime/Core/Public/Math/Axis.h
@@ -32,7 +32,11 @@ namespace EAxisList
 		YZ			= Y | Z,
 		XYZ			= X | Y | Z,
 		All			= XYZ | Screen,
+
 		//alias over Axis YZ since it isn't used when the z-rotation widget is being used
 		ZRotation	= YZ,
+		
+		// alias over Screen since it isn't used when the 2d translate rotate widget is being used
+		Rotate2D	= Screen,
 	};
 }
