Fix static mesh importing not allocating enough material slots for the number of materials on the mesh.

https://jira.ol.epicgames.net/browse/UE-13026

diff --git a/Engine/Source/Editor/UnrealEd/Private/Fbx/FbxStaticMeshImport.cpp b/Engine/Source/Editor/UnrealEd/Private/Fbx/FbxStaticMeshImport.cpp
--- a/Engine/Source/Editor/UnrealEd/Private/Fbx/FbxStaticMeshImport.cpp
+++ b/Engine/Source/Editor/UnrealEd/Private/Fbx/FbxStaticMeshImport.cpp
@@ -1091,11 +1091,11 @@ UStaticMesh* UnFbx::FFbxImporter::ImportStaticMeshAsSingle(UObject* InParent, TA
 		SortedMaterialIndex.Sort();
 
 		UE_LOG(LogFbx,Log,TEXT("== After sorting:"));
-		TArray<UMaterialInterface*> SortedMaterials;
+		TArray<FFbxMaterial> SortedMaterials;
 		for (int32 SortedIndex = 0; SortedIndex < SortedMaterialIndex.Num(); ++SortedIndex)
 		{
 			int32 RemappedIndex = SortedMaterialIndex[SortedIndex] & 0xffff;
-			SortedMaterials.Add(UniqueMaterials[RemappedIndex].Material);
+			SortedMaterials.Add(UniqueMaterials[RemappedIndex]);
 			UE_LOG(LogFbx,Log,TEXT("%d: %s"),SortedIndex,*UniqueMaterials[RemappedIndex].GetName());
 		}
 		UE_LOG(LogFbx,Log,TEXT("== Mapping table:"));
@@ -1134,13 +1134,13 @@ UStaticMesh* UnFbx::FFbxImporter::ImportStaticMeshAsSingle(UObject* InParent, TA
 			// Also compact the sorted materials array.
 			if (RawMesh.MaterialIndexToImportIndex.Num() > 0)
 			{
-				TArray<UMaterialInterface*> OldSortedMaterials;
+				TArray<FFbxMaterial> OldSortedMaterials;
 
 				Exchange(OldSortedMaterials,SortedMaterials);
 				SortedMaterials.Empty(RawMesh.MaterialIndexToImportIndex.Num());
 				for (int32 MaterialIndex = 0; MaterialIndex < RawMesh.MaterialIndexToImportIndex.Num(); ++MaterialIndex)
 				{
-					UMaterialInterface* Material = NULL;
+					FFbxMaterial Material;
 					int32 ImportIndex = RawMesh.MaterialIndexToImportIndex[MaterialIndex];
 					if (OldSortedMaterials.IsValidIndex(ImportIndex))
 					{
@@ -1177,12 +1177,11 @@ UStaticMesh* UnFbx::FFbxImporter::ImportStaticMeshAsSingle(UObject* InParent, TA
 		for (int32 MaterialIndex = 0; MaterialIndex < NumMaterials; ++MaterialIndex)
 		{
 			FMeshSectionInfo Info = StaticMesh->SectionInfoMap.Get(LODIndex, MaterialIndex);
-			int32 Index = StaticMesh->Materials.Find( SortedMaterials[MaterialIndex] );
-			if( Index == INDEX_NONE )
-			{
-				Index = StaticMesh->Materials.Num();
-				StaticMesh->Materials.Add(SortedMaterials[MaterialIndex]);
-			}
+			//int32 Index = StaticMesh->Materials.Find( SortedMaterials[MaterialIndex] );
+			//if( Index == INDEX_NONE )
+			//{
+			int32 Index = StaticMesh->Materials.Add(SortedMaterials[MaterialIndex].Material);
+			//}
 			Info.MaterialIndex = Index;
 			StaticMesh->SectionInfoMap.Set(LODIndex, MaterialIndex, Info);
 			
