When ConditionalRecompileClass is called the RecompileBlueprintBytecode will recreate the CDO. The CDO can be referenced by Blueprint. For example as a default object value in a pin in call function node, when a static function from a library is called.

The problem was found while investigation UE-12028 - try to save Character_ShipB in attached project.

diff --git a/Engine/Source/Editor/UnrealEd/Private/Kismet2/KismetReinstanceUtilities.cpp b/Engine/Source/Editor/UnrealEd/Private/Kismet2/KismetReinstanceUtilities.cpp
--- a/Engine/Source/Editor/UnrealEd/Private/Kismet2/KismetReinstanceUtilities.cpp
+++ b/Engine/Source/Editor/UnrealEd/Private/Kismet2/KismetReinstanceUtilities.cpp
@@ -512,6 +512,12 @@ void FBlueprintCompileReinstancer::UpdateBytecodeReferences()
 					FArchiveReplaceObjectRef<UObject> ReplaceAr(CurrentFunction, FieldMappings, /*bNullPrivateRefs=*/ false, /*bIgnoreOuterRef=*/ true, /*bIgnoreArchetypeRef=*/ true);
 				}
 			}
+
+			FArchiveReplaceObjectRef<UObject> ReplaceInBPAr(*DependentBP, FieldMappings, false, true, true);
+			if (ReplaceInBPAr.GetCount())
+			{
+				UE_LOG(LogBlueprint, Log, TEXT("UpdateBytecodeReferences: %d references from %s was replaced in BP %s"), ReplaceInBPAr.GetCount(), *GetPathNameSafe(ClassToReinstance), *GetPathNameSafe(*DependentBP));
+			}
 		}
 	}
 }
