UE-7888 Get and Set nodes for variables from a BP's parent class will not compile without a reference to Self (4.8)

diff --git a/Engine/Source/Editor/KismetCompiler/Private/KismetCompilerModule.cpp b/Engine/Source/Editor/KismetCompiler/Private/KismetCompilerModule.cpp
--- a/Engine/Source/Editor/KismetCompiler/Private/KismetCompilerModule.cpp
+++ b/Engine/Source/Editor/KismetCompiler/Private/KismetCompilerModule.cpp
@@ -118,11 +118,6 @@ void FKismet2CompilerModule::CompileBlueprint(class UBlueprint* Blueprint, const
 	if (CompileOptions.CompileType != EKismetCompileType::Cpp)
 	{
 		BP_SCOPED_COMPILER_EVENT_STAT(EKismetCompilerStats_CompileSkeletonClass);
-
-		// Pre-cache the derived classes, before reinstancing can occur.
-		TArray<UClass*> ChildrenOfClass;
-		GetDerivedClasses(Blueprint->SkeletonGeneratedClass, ChildrenOfClass);
-
 		auto SkeletonReinstancer = FBlueprintCompileReinstancer::Create(Blueprint->SkeletonGeneratedClass);
 
 		FCompilerResultsLog SkeletonResults;
diff --git a/Engine/Source/Editor/UnrealEd/Private/Kismet2/BlueprintEditorUtils.cpp b/Engine/Source/Editor/UnrealEd/Private/Kismet2/BlueprintEditorUtils.cpp
--- a/Engine/Source/Editor/UnrealEd/Private/Kismet2/BlueprintEditorUtils.cpp
+++ b/Engine/Source/Editor/UnrealEd/Private/Kismet2/BlueprintEditorUtils.cpp
@@ -1892,6 +1892,30 @@ void FBlueprintEditorUtils::UpdateDelegatesInBlueprint(UBlueprint* Blueprint)
 // Blueprint has materially changed.  Recompile the skeleton, notify observers, and mark the package as dirty.
 void FBlueprintEditorUtils::MarkBlueprintAsStructurallyModified(UBlueprint* Blueprint)
 {
+	auto SkeletalRecompileChildren = [](TArray<UClass*> ChildrenOfClass)
+	{
+		for (auto SkelClass : ChildrenOfClass)
+		{
+			IKismetCompilerInterface& Compiler = FModuleManager::LoadModuleChecked<IKismetCompilerInterface>(KISMET_COMPILER_MODULENAME);
+
+			auto Blueprint = Cast<UBlueprint>(SkelClass->ClassGeneratedBy);
+			if (Blueprint
+				&& Blueprint->Status != BS_BeingCreated
+				&& !Blueprint->bBeingCompiled
+				&& !Blueprint->bIsRegeneratingOnLoad
+				&& Blueprint->bHasBeenRegenerated)
+			{
+				FCompilerResultsLog Results;
+				Results.bSilentMode = true;
+				Results.bLogInfoOnly = true;
+
+				FKismetCompilerOptions CompileOptions;
+				CompileOptions.CompileType = EKismetCompileType::SkeletonOnly;
+				Compiler.CompileBlueprint(Blueprint, CompileOptions, Results);
+			}
+		}
+	};
+
 	// The Blueprint has been structurally modified and this means that some node titles will need to be refreshed
 	GetDefault<UEdGraphSchema_K2>()->ForceVisualizationCacheClear();
 
@@ -1918,6 +1942,12 @@ void FBlueprintEditorUtils::MarkBlueprintAsStructurallyModified(UBlueprint* Blue
 			}
 		}
 
+		TArray<UClass*> ChildrenOfClass;
+		if (UClass* SkelClass = Blueprint->SkeletonGeneratedClass)
+		{
+			GetDerivedClasses(SkelClass, ChildrenOfClass);
+		}
+
 		{
 			// Invoke the compiler to update the skeleton class definition
 			IKismetCompilerInterface& Compiler = FModuleManager::LoadModuleChecked<IKismetCompilerInterface>(KISMET_COMPILER_MODULENAME);
@@ -1929,6 +1959,8 @@ void FBlueprintEditorUtils::MarkBlueprintAsStructurallyModified(UBlueprint* Blue
 		}
 		UpdateDelegatesInBlueprint(Blueprint);
 
+		SkeletalRecompileChildren(ChildrenOfClass);
+
 		{
 			BP_SCOPED_COMPILER_EVENT_STAT(EKismetCompilerStats_NotifyBlueprintChanged);
 
