Fix for watch values being reported as out of scope in the blueprint debugger.
#UE-12022 - Variables incorrectly returning as Out of Scope when watched
#Branch UE4
#Proj UnrealEd

#info Just wanted to run this by you Nick, I'm pretty sure the fix is good but this is still new code to me.

diff --git a/Engine/Source/Editor/UnrealEd/Private/Kismet2/KismetDebugUtilities.cpp b/Engine/Source/Editor/UnrealEd/Private/Kismet2/KismetDebugUtilities.cpp
--- a/Engine/Source/Editor/UnrealEd/Private/Kismet2/KismetDebugUtilities.cpp
+++ b/Engine/Source/Editor/UnrealEd/Private/Kismet2/KismetDebugUtilities.cpp
@@ -897,15 +897,15 @@ FKismetDebugUtilities::EWatchTextResult FKismetDebugUtilities::GetWatchText(FStr
 
 			// Try at member scope if it wasn't part of a current function scope
 			UClass* PropertyClass = Cast<UClass>(Property->GetOuter());
-			if ((PropertyBase == nullptr) && (PropertyClass != nullptr) && ActiveObject->GetClass()->IsChildOf(PropertyClass))
+			if( !PropertyBase && PropertyClass )
 			{
-				PropertyBase = ActiveObject;
-			}
-			else
-			{
-				// Try and locate the propertybase in actors components
-				if( AActor* Actor = Cast<AActor>( ActiveObject ))
+				if( ActiveObject->GetClass()->IsChildOf( PropertyClass ))
+				{
+					PropertyBase = ActiveObject;
+				}
+				else if( AActor* Actor = Cast<AActor>( ActiveObject ))
 				{
+					// Try and locate the propertybase in the actor components
 					for( auto ComponentIter : Actor->GetComponents() )
 					{
 						if( ComponentIter->GetClass()->IsChildOf(PropertyClass))
@@ -916,7 +916,18 @@ FKismetDebugUtilities::EWatchTextResult FKismetDebugUtilities::GetWatchText(FStr
 					}
 				}
 			}
-
+#if USE_UBER_GRAPH_PERSISTENT_FRAME
+			// Try find the propertybase in the persistent ubergraph frame
+			UFunction* OuterFunction = Cast<UFunction>( Property->GetOuter() );
+			if( !PropertyBase && OuterFunction )
+			{
+				if( UBlueprintGeneratedClass* BPGC = Cast<UBlueprintGeneratedClass>( Blueprint->GeneratedClass ))
+				{
+					PropertyBase = BPGC->GetPersistentUberGraphFrame( ActiveObject, OuterFunction );
+				}
+			}
+#endif // USE_UBER_GRAPH_PERSISTENT_FRAME
+			
 			// Now either print out the variable value, or that it was out-of-scope
 			if (PropertyBase != nullptr)
 			{
