#ANIM: CURVE IMPORT: fix the issue with tangent not importing correctly

https://jira.ol.epicgames.net/browse/UE-13353

diff --git a/Engine/Source/Editor/UnrealEd/Private/SkeletalMeshEdit.cpp b/Engine/Source/Editor/UnrealEd/Private/SkeletalMeshEdit.cpp
--- a/Engine/Source/Editor/UnrealEd/Private/SkeletalMeshEdit.cpp
+++ b/Engine/Source/Editor/UnrealEd/Private/SkeletalMeshEdit.cpp
@@ -742,22 +742,19 @@ bool UnFbx::FFbxImporter::ImportCurve(const FbxAnimCurve* FbxCurve, FFloatCurve
 				break;
 			}
 
-			switch (KeyTangentMode)
+			// when we import tangent, we only support break or user
+			// since it's modified by DCC and we only assume these two are valid
+			// auto does our own stuff, which doesn't work with what you see in DCC
+			if (KeyTangentMode & FbxAnimCurveDef::eTangentBreak)
 			{
-			case FbxAnimCurveDef::eTangentUser: //! Next slope at the left equal to slope at the right.
-				NewTangentMode = RCTM_User;
-				break;
-			case FbxAnimCurveDef::eTangentGenericBreak://! Independent left and right slopes.
-			case FbxAnimCurveDef::eTangentBreak://! Independent left and right slopes, with next slope at the left equal to slope at the right.
-			case FbxAnimCurveDef::eTangentAutoBreak://! Independent left and right slopes, with auto key.
 				NewTangentMode = RCTM_Break;
-				break;
-			default:
-				//@todo I don't think this works well
-				NewTangentMode = RCTM_Auto;
-				break;
+			}
+			else
+			{
+				NewTangentMode = RCTM_User;
 			}
 
+			// @fix me : weight of tangent is not used, but we'll just save this for future where we might use it. 
 			switch (KeyTangentWeightMode)
 			{
 			case FbxAnimCurveDef::eWeightedNone://! Tangent has default weights of 0.333; we define this state as not weighted.
@@ -800,15 +797,11 @@ bool UnFbx::FFbxImporter::ImportCurve(const FbxAnimCurve* FbxCurve, FFloatCurve
 			Curve->FloatCurve.SetKeyInterpMode(NewKeyHandle, NewInterpMode);
 			Curve->FloatCurve.SetKeyTangentMode(NewKeyHandle, NewTangentMode);
 			Curve->FloatCurve.SetKeyTangentWeightMode(NewKeyHandle, NewTangentWeightMode);
-			FRichCurveKey& NewKey = Curve->FloatCurve.GetKey(NewKeyHandle);
 
-			// update tangents - I don't know what the unit for Derivatives is
-			// so I'm just inquiring tangent like above
-			//FbxAnimCurveTangentInfo ArriveTangentInfo = FbxCurve->KeyGetLeftDerivativeInfo(KeyIndex);
-			//FbxAnimCurveTangentInfo LeaveTangentInfo = FbxCurve->KeyGetRightDerivativeInfo(KeyIndex);
-
-			NewKey.ArriveTangent = ArriveTangent;
-			NewKey.LeaveTangent = LeaveTangent;
+			FRichCurveKey& NewKey = Curve->FloatCurve.GetKey(NewKeyHandle);
+			// apply 1/100 - that seems like the tangent unit difference with FBX
+			NewKey.ArriveTangent = ArriveTangent * 0.01f;
+			NewKey.LeaveTangent = LeaveTangent * 0.01f;
 			NewKey.ArriveTangentWeight = ArriveTangentWeight;
 			NewKey.LeaveTangentWeight = LeaveTangentWeight;
 		}
