PR #1013: Fixing PVS-Studio warnings. (Contributed by [REDACTED_PERSON])
Some files have been omitted and will be submitted with modified corrections

diff --git a/Engine/Plugins/Runtime/Steam/SteamVR/Source/SteamVR/Private/SteamVRHMD.cpp b/Engine/Plugins/Runtime/Steam/SteamVR/Source/SteamVR/Private/SteamVRHMD.cpp
--- a/Engine/Plugins/Runtime/Steam/SteamVR/Source/SteamVR/Private/SteamVRHMD.cpp
+++ b/Engine/Plugins/Runtime/Steam/SteamVR/Source/SteamVR/Private/SteamVRHMD.cpp
@@ -146,12 +146,12 @@ void FSteamVRHMD::GetCurrentPose(FQuat& CurrentOrientation, FVector& CurrentPosi
 			TrackingFrame.bDeviceIsValid[i] = Poses[i].bDeviceIsValid;
 			TrackingFrame.bPoseIsValid[i] = Poses[i].bPoseIsValid;
 
-			FVector CurrentPosition;
-			FQuat CurrentOrientation;
-			PoseToOrientationAndPosition(Poses[i].mDeviceToAbsoluteTracking, CurrentOrientation, CurrentPosition);
+			FVector LocalCurrentPosition;
+			FQuat LocalCurrentOrientation;
+			PoseToOrientationAndPosition(Poses[i].mDeviceToAbsoluteTracking, LocalCurrentOrientation, LocalCurrentPosition);
 
-			TrackingFrame.DeviceOrientation[i] = CurrentOrientation;
-			TrackingFrame.DevicePosition[i] = CurrentPosition;
+			TrackingFrame.DeviceOrientation[i] = LocalCurrentOrientation;
+			TrackingFrame.DevicePosition[i] = LocalCurrentPosition;
 
 			TrackingFrame.RawPoses[i] = Poses[i].mDeviceToAbsoluteTracking;
 		}
diff --git a/Engine/Source/Developer/CrashDebugHelper/Private/Windows/WindowsPlatformStackWalkExt.cpp b/Engine/Source/Developer/CrashDebugHelper/Private/Windows/WindowsPlatformStackWalkExt.cpp
--- a/Engine/Source/Developer/CrashDebugHelper/Private/Windows/WindowsPlatformStackWalkExt.cpp
+++ b/Engine/Source/Developer/CrashDebugHelper/Private/Windows/WindowsPlatformStackWalkExt.cpp
@@ -553,7 +553,8 @@ int FWindowsPlatformStackWalkExt::GetCallstacks()
 	ULONG ContextUsed = 0;
 
 	// Get the context of the crashed thread
-	if( Control->GetStoredEventInformation( &DebugEvent, &ProcessID, &ThreadID, Context, ContextSize, &ContextUsed, NULL, 0, 0) )
+	HRESULT hr = Control->GetStoredEventInformation(&DebugEvent, &ProcessID, &ThreadID, Context, ContextSize, &ContextUsed, NULL, 0, 0);
+	if( FAILED(hr) )
 	{
 		return NumValidFunctionNames;
 	}
@@ -663,7 +664,8 @@ bool FWindowsPlatformStackWalkExt::OpenDumpFile( const FString& InCrashDumpFilen
 		return false;
 	}
 
-	if( Client->OpenDumpFileWide( *InCrashDumpFilename, NULL ) )
+	HRESULT hr = Client->OpenDumpFileWide(*InCrashDumpFilename, NULL);
+	if(FAILED(hr) )
 	{
 		UE_LOG( LogCrashDebugHelper, Warning, TEXT( "Failed to open minidump file: %s" ), *InCrashDumpFilename );
 		return false;
diff --git a/Engine/Source/Editor/AnimGraph/Private/AnimGraphNode_ModifyBone.cpp b/Engine/Source/Editor/AnimGraph/Private/AnimGraphNode_ModifyBone.cpp
--- a/Engine/Source/Editor/AnimGraph/Private/AnimGraphNode_ModifyBone.cpp
+++ b/Engine/Source/Editor/AnimGraph/Private/AnimGraphNode_ModifyBone.cpp
@@ -110,6 +110,7 @@ EBoneModificationMode UAnimGraphNode_ModifyBone::GetBoneModificationMode(int32 I
 		return Node.ScaleMode;
 		break;
 	case FWidget::WM_TranslateRotateZ:
+	case FWidget::WM_2D:
 		break;
 	}
 
@@ -128,6 +129,7 @@ int32 UAnimGraphNode_ModifyBone::GetNextWidgetMode(int32 InWidgetMode)
 	case FWidget::WM_Scale:
 		return FWidget::WM_Translate;
 	case FWidget::WM_TranslateRotateZ:
+	case FWidget::WM_2D:
 		break;
 	}
 
diff --git a/Engine/Source/Editor/BlueprintGraph/Private/K2Node_BreakStruct.cpp b/Engine/Source/Editor/BlueprintGraph/Private/K2Node_BreakStruct.cpp
--- a/Engine/Source/Editor/BlueprintGraph/Private/K2Node_BreakStruct.cpp
+++ b/Engine/Source/Editor/BlueprintGraph/Private/K2Node_BreakStruct.cpp
@@ -242,7 +242,7 @@ void UK2Node_BreakStruct::ValidateNodeDuringCompilation(class FCompilerResultsLo
 				const bool bIsBlueprintVisible = Property->HasAnyPropertyFlags(CPF_BlueprintVisible) || (Property->GetOwnerStruct() && Property->GetOwnerStruct()->IsA<UUserDefinedStruct>());
 				bHasAnyBlueprintVisibleProperty |= bIsBlueprintVisible;
 
-				const UEdGraphPin* Pin = Property ? FindPin(Property->GetName()) : NULL;
+				const UEdGraphPin* Pin = FindPin(Property->GetName());
 				const bool bIsLinked = Pin && Pin->LinkedTo.Num();
 
 				if (!bIsBlueprintVisible && bIsLinked)
diff --git a/Engine/Source/Editor/BlueprintGraph/Private/K2Node_MathExpression.cpp b/Engine/Source/Editor/BlueprintGraph/Private/K2Node_MathExpression.cpp
--- a/Engine/Source/Editor/BlueprintGraph/Private/K2Node_MathExpression.cpp
+++ b/Engine/Source/Editor/BlueprintGraph/Private/K2Node_MathExpression.cpp
@@ -232,7 +232,7 @@ public:
 	/** For debug purposes, constructs a textual representation of this expression */
 	virtual FString ToString() const override
 	{
-		if (Token.TokenType == FBasicToken::TOKEN_Identifier || FBasicToken::TOKEN_Guid)
+		if (Token.TokenType == FBasicToken::TOKEN_Identifier || Token.TokenType == FBasicToken::TOKEN_Guid)
 		{
 			return FString::Printf(TEXT("%s"), Token.Identifier);
 		}
diff --git a/Engine/Source/Editor/Kismet/Private/BlueprintEditor.cpp b/Engine/Source/Editor/Kismet/Private/BlueprintEditor.cpp
--- a/Engine/Source/Editor/Kismet/Private/BlueprintEditor.cpp
+++ b/Engine/Source/Editor/Kismet/Private/BlueprintEditor.cpp
@@ -5874,8 +5874,7 @@ void FBlueprintEditor::CollapseNodesIntoGraph(UEdGraphNode* InGatewayNode, UK2No
 					if(InterfaceTemplateNode && LocalPin->Direction == EGPD_Input)
 					{
 						// Find the pin on the entry node, we will use that pin's name to find the pin on the remote port
-						UEdGraphPin* EntryNodePin = nullptr;
-						EntryNodePin = InEntryNode->FindPin(LocalPin->LinkedTo[0]->PinName);
+						UEdGraphPin* EntryNodePin = InEntryNode->FindPin(LocalPin->LinkedTo[0]->PinName);
 						if(EntryNodePin)
 						{
 							LocalPin->BreakAllPinLinks();
diff --git a/Engine/Source/Editor/Kismet/Private/SBlueprintEditorToolbar.cpp b/Engine/Source/Editor/Kismet/Private/SBlueprintEditorToolbar.cpp
--- a/Engine/Source/Editor/Kismet/Private/SBlueprintEditorToolbar.cpp
+++ b/Engine/Source/Editor/Kismet/Private/SBlueprintEditorToolbar.cpp
@@ -84,7 +84,7 @@ void FKismet2Menu::FillFileMenuBlueprintSection( FMenuBuilder& MenuBuilder, FBlu
 				LOCTEXT("DeveloperMenu", "Developer"),
 				LOCTEXT("DeveloperMenu_ToolTip", "Open the developer menu"),
 				FNewMenuDelegate::CreateStatic( &FKismet2Menu::FillDeveloperMenu ),
-				"DeveloperMenu");
+				true);
 		}
 		MenuBuilder.EndSection();
 	}
diff --git a/Engine/Source/Editor/KismetCompiler/Private/KismetCompiler.cpp b/Engine/Source/Editor/KismetCompiler/Private/KismetCompiler.cpp
--- a/Engine/Source/Editor/KismetCompiler/Private/KismetCompiler.cpp
+++ b/Engine/Source/Editor/KismetCompiler/Private/KismetCompiler.cpp
@@ -1277,7 +1277,7 @@ void FKismetCompilerContext::PrecompileFunction(FKismetFunctionContext& Context)
 			if (!ParentFunction->IsSignatureCompatibleWith(Context.Function))
 			{
 				FString SignatureClassName("");
-				if (Context.EntryPoint && Context.EntryPoint->SignatureClass)
+				if (Context.EntryPoint->SignatureClass)
 				{
 					SignatureClassName = Context.EntryPoint->SignatureClass->GetName();
 				}
@@ -3296,8 +3296,7 @@ void FKismetCompilerContext::Compile()
 
 	NewClass->ClassGeneratedBy = Blueprint;
 
-	UClass* ParentClass = nullptr;
-	ParentClass = NewClass->ClassWithin;
+	UClass* ParentClass = NewClass->ClassWithin;
 	NewClass->SetSuperStruct(ParentClass);
 	NewClass->ClassFlags |= (ParentClass->ClassFlags & CLASS_Inherit);
 	NewClass->ClassCastFlags |= ParentClass->ClassCastFlags;
diff --git a/Engine/Source/Editor/Matinee/Private/MatineeMenus.cpp b/Engine/Source/Editor/Matinee/Private/MatineeMenus.cpp
--- a/Engine/Source/Editor/Matinee/Private/MatineeMenus.cpp
+++ b/Engine/Source/Editor/Matinee/Private/MatineeMenus.cpp
@@ -1537,11 +1537,16 @@ void FMatinee::OnContextTrackExportAnimFBX()
 		// Get the owning group of the track
 		UInterpGroup* Group = CastChecked<UInterpGroup>( SelectedTrack->GetOuter() ) ;
 
-		// Get the first group instance for this track.  In the case of animations there is only one instance usually
-		UInterpGroupInst* GroupInst = MatineeActor->FindFirstGroupInst( Group );
+		ASkeletalMeshActor* SkelMeshActor = NULL;
 
-		// Get the actor for this group.
-		ASkeletalMeshActor* SkelMeshActor = Cast<ASkeletalMeshActor>( GroupInst->GroupActor );
+		if ( MatineeActor != NULL )
+		{
+			// Get the first group instance for this track.  In the case of animations there is only one instance usually
+			UInterpGroupInst* GroupInst = MatineeActor->FindFirstGroupInst(Group);
+
+			// Get the actor for this group.
+			SkelMeshActor = Cast<ASkeletalMeshActor>(GroupInst->GroupActor);
+		}
 
 		// Someone could have hooked up an invalid actor.  In that case do nothing
 		if( SkelMeshActor )
@@ -1921,6 +1926,8 @@ void FMatinee::OnContextGroupExportAnimFBX()
 			{
 				// Find the skeletal mesh for this group
 				USkeletalMeshComponent* SkelMeshComponent = NULL;
+
+				if (MatineeActor != NULL)
 				{
 					// Get the first group instance for this group.  In the case of animations there is usually only one instance
 					UInterpGroupInst* GroupInst = MatineeActor->FindFirstGroupInst( SelectedGroup );
diff --git a/Engine/Source/Editor/PlacementMode/Private/PlacementMode.cpp b/Engine/Source/Editor/PlacementMode/Private/PlacementMode.cpp
--- a/Engine/Source/Editor/PlacementMode/Private/PlacementMode.cpp
+++ b/Engine/Source/Editor/PlacementMode/Private/PlacementMode.cpp
@@ -66,7 +66,10 @@ void FPlacementMode::Tick(FEditorViewportClient* ViewportClient,float DeltaTime)
 {
 	if ( IsCurrentlyPlacing() )
 	{
-		ViewportClient->SetRequiredCursorOverride( true, EMouseCursor::GrabHandClosed );
+		if ( ViewportClient )
+		{
+			ViewportClient->SetRequiredCursorOverride(true, EMouseCursor::GrabHandClosed);
+		}
 
 		bool HasValidFocusTarget = false;
 		for (int Index = ValidFocusTargetsForPlacement.Num() - 1; !HasValidFocusTarget && Index >= 0 ; Index--)
@@ -107,7 +110,7 @@ void FPlacementMode::Tick(FEditorViewportClient* ViewportClient,float DeltaTime)
 			StopPlacing();
 		}
 	}
-	else
+	else if (ViewportClient != NULL)
 	{
 		ViewportClient->SetRequiredCursorOverride( false );
 	}
diff --git a/Engine/Source/Editor/PropertyEditor/Private/UserInterface/PropertyEditor/SPropertyEditorColor.cpp b/Engine/Source/Editor/PropertyEditor/Private/UserInterface/PropertyEditor/SPropertyEditorColor.cpp
--- a/Engine/Source/Editor/PropertyEditor/Private/UserInterface/PropertyEditor/SPropertyEditorColor.cpp
+++ b/Engine/Source/Editor/PropertyEditor/Private/UserInterface/PropertyEditor/SPropertyEditorColor.cpp
@@ -109,7 +109,7 @@ FReply SPropertyEditorColor::ColorBlock_OnMouseButtonDown(const FGeometry& MyGeo
 	const UProperty* Property = PropertyEditor->GetProperty();
 	check(Property);
 
-	bool bRefreshOnlyOnOk = Property->GetOwnerClass()->IsChildOf(UMaterialExpressionConstant3Vector::StaticClass());
+	bool bRefreshOnlyOnOk = Property->GetOwnerClass() && Property->GetOwnerClass()->IsChildOf(UMaterialExpressionConstant3Vector::StaticClass());
 	//@todo Slate Property window: This should probably be controlled via metadata and then the actual alpha channel property hidden if its not used.
 	const bool bUseAlpha = !(Property->GetOwnerClass() && (Property->GetOwnerClass()->IsChildOf(ULightComponent::StaticClass()) || bRefreshOnlyOnOk));
 	
diff --git a/Engine/Source/Editor/TranslationEditor/Private/TranslationPickerEditWindow.cpp b/Engine/Source/Editor/TranslationEditor/Private/TranslationPickerEditWindow.cpp
--- a/Engine/Source/Editor/TranslationEditor/Private/TranslationPickerEditWindow.cpp
+++ b/Engine/Source/Editor/TranslationEditor/Private/TranslationPickerEditWindow.cpp
@@ -188,8 +188,7 @@ void STranslationPickerEditWidget::Construct(const FArguments& InArgs)
 	FString ManifestAndArchiveNameString;
 	if (NamespaceString && KeyString)
 	{
-		TSharedPtr<FString, ESPMode::ThreadSafe> TextTableName = TSharedPtr<FString, ESPMode::ThreadSafe>(nullptr);
-		TextTableName = FTextLocalizationManager::Get().GetTableName(*NamespaceString, *KeyString);
+		TSharedPtr<FString, ESPMode::ThreadSafe> TextTableName = FTextLocalizationManager::Get().GetTableName(*NamespaceString, *KeyString);
 		LocresFullPath = *TextTableName;
 		if (TextTableName.IsValid())
 		{
diff --git a/Engine/Source/Editor/UMGEditor/Private/Customizations/WidgetNavigationCustomization.cpp b/Engine/Source/Editor/UMGEditor/Private/Customizations/WidgetNavigationCustomization.cpp
--- a/Engine/Source/Editor/UMGEditor/Private/Customizations/WidgetNavigationCustomization.cpp
+++ b/Engine/Source/Editor/UMGEditor/Private/Customizations/WidgetNavigationCustomization.cpp
@@ -121,6 +121,8 @@ FText FWidgetNavigationCustomization::GetNavigationText(TWeakPtr<IPropertyHandle
 		return LOCTEXT("NavigationExplicit", "Explicit");
 	case EUINavigationRule::Invalid:
 		return LOCTEXT("NavigationMultipleValues", "Multiple Values");
+	case EUINavigationRule::Custom:
+		break;
 	}
 
 	return FText::GetEmpty();
diff --git a/Engine/Source/Editor/UMGEditor/Private/Designer/SDesignerView.cpp b/Engine/Source/Editor/UMGEditor/Private/Designer/SDesignerView.cpp
--- a/Engine/Source/Editor/UMGEditor/Private/Designer/SDesignerView.cpp
+++ b/Engine/Source/Editor/UMGEditor/Private/Designer/SDesignerView.cpp
@@ -713,6 +713,8 @@ void SDesignerView::GetPreviewAreaAndSize(FVector2D& Area, FVector2D& Size) cons
 		case EDesignPreviewSizeMode::DesiredOnScreen:
 			Size = CachedPreviewDesiredSize;
 			return;
+		case EDesignPreviewSizeMode::FillScreen:
+			break;
 		}
 	}
 }
diff --git a/Engine/Source/Editor/UnrealEd/Private/Editor.cpp b/Engine/Source/Editor/UnrealEd/Private/Editor.cpp
--- a/Engine/Source/Editor/UnrealEd/Private/Editor.cpp
+++ b/Engine/Source/Editor/UnrealEd/Private/Editor.cpp
@@ -3945,9 +3945,6 @@ void UEditorEngine::ConvertActorsFromClass( UClass* FromClass, UClass* ToClass )
 					if( Info.ActorGroup )
 					{
 						Info.ActorGroup->Add(*Actor);
-					}
-					if( Info.ActorGroup )
-					{
 						Info.ActorGroup->Add(*Actor);
 					}
 				}
@@ -6159,18 +6156,12 @@ AActor* UEditorEngine::AddActor(ULevel* InLevel, UClass* Class, const FTransform
 		}
 	}
 
-	///////////////////////////////
-	// If this actor is part of any layers (set in its default properties), add them into the visible layers list.
-
 	if( Actor )
 	{
+		// If this actor is part of any layers (set in its default properties), add them into the visible layers list.
 		GEditor->Layers->SetLayersVisibility( Actor->Layers, true );
-	}
 
-	///////////////////////////////
-	// Clean up.
-	if ( Actor )
-	{
+		// Clean up.
 		Actor->MarkPackageDirty();
 		ULevel::LevelDirtiedEvent.Broadcast();
 	}
@@ -7036,7 +7027,6 @@ namespace EditorUtilities
 	template<class AllocatorType = FDefaultAllocator>
 	UActorComponent* FindMatchingComponentInstance( UActorComponent* SourceComponent, AActor* TargetActor, const TArray<UActorComponent*, AllocatorType>& TargetComponents, int32& StartIndex )
 	{
-		const bool bSourceIsArchetype = SourceComponent->HasAnyFlags(RF_ArchetypeObject);
 		UActorComponent* TargetComponent = StartIndex < TargetComponents.Num() ? TargetComponents[ StartIndex ] : nullptr;
 
 		// If the source and target components do not match (e.g. context-specific), attempt to find a match in the target's array elsewhere
@@ -7045,6 +7035,7 @@ namespace EditorUtilities
 			&& ((TargetComponent == nullptr) 
 				|| (SourceComponent->GetFName() != TargetComponent->GetFName()) ))
 		{
+			const bool bSourceIsArchetype = SourceComponent->HasAnyFlags(RF_ArchetypeObject);
 			// Reset the target component since it doesn't match the source
 			TargetComponent = nullptr;
 
diff --git a/Engine/Source/Editor/UnrealEd/Private/Fbx/FbxMatineeImport.cpp b/Engine/Source/Editor/UnrealEd/Private/Fbx/FbxMatineeImport.cpp
--- a/Engine/Source/Editor/UnrealEd/Private/Fbx/FbxMatineeImport.cpp
+++ b/Engine/Source/Editor/UnrealEd/Private/Fbx/FbxMatineeImport.cpp
@@ -792,7 +792,7 @@ float FFbxImporter::ImportMatineeActor(FbxNode* Node, UInterpGroupInst* MatineeG
 			bIsCamera = true;
 		}
 
-		if( MovementTrack->SubTracks.Num() > 0 )
+		if (MovementTrack && MovementTrack->SubTracks.Num() > 0)
 		{
 			check (bIsCamera == false);
 			ImportMoveSubTrack(TransCurves[0], 0, SubTracks[0], 0, false, RealCurves[0], DefaultPos[0]);
@@ -817,7 +817,7 @@ float FFbxImporter::ImportMatineeActor(FbxNode* Node, UInterpGroupInst* MatineeG
 			// Scale the track timing to ensure that it is large enough
 			MovementTrack->GetTimeRange( StartTime, TimeLength );
 		}
-		else
+		else if (MovementTrack)
 		{
 			ImportMatineeAnimated(TransCurves[0], MovementTrack->PosTrack, 1, true, RealCurves[0], DefaultPos[0]);
 			ImportMatineeAnimated(TransCurves[1], MovementTrack->PosTrack, 0, true, RealCurves[0], DefaultPos[1]);
diff --git a/Engine/Source/Editor/UnrealEd/Private/LevelViewportClickHandlers.cpp b/Engine/Source/Editor/UnrealEd/Private/LevelViewportClickHandlers.cpp
--- a/Engine/Source/Editor/UnrealEd/Private/LevelViewportClickHandlers.cpp
+++ b/Engine/Source/Editor/UnrealEd/Private/LevelViewportClickHandlers.cpp
@@ -447,6 +447,13 @@ namespace ClickHandlers
 
 	bool ClickGeomPoly(FLevelEditorViewportClient* ViewportClient, HGeomPolyProxy* InHitProxy, const FViewportClick& Click)
 	{
+		//something is wrong with the hitproxy relating to this click - create a debug log to help identify what
+		if( InHitProxy == NULL )
+		{
+			UE_LOG(LogEditorViewport, Warning, TEXT("Invalid hitproxy"));
+			return false;
+		}
+
 		if( !InHitProxy->GeomObjectWeakPtr.IsValid() )
 		{
 			return false;
@@ -469,7 +476,7 @@ namespace ClickHandlers
 			FEdMode* Mode = GLevelEditorModeTools().GetActiveMode( FBuiltinEditorModes::EM_Geometry );
 			if( Mode )
 			{
-				if( ( InHitProxy != NULL) && ( InHitProxy->GetGeomObject() != NULL ) && (InHitProxy->GetGeomObject()->PolyPool.IsValidIndex( InHitProxy->PolyIndex ) == true ) )					
+				if( ( InHitProxy->GetGeomObject() != NULL ) && (InHitProxy->GetGeomObject()->PolyPool.IsValidIndex( InHitProxy->PolyIndex ) == true ) )					
 				{
 					Mode->GetCurrentTool()->StartTrans();
 
@@ -488,22 +495,13 @@ namespace ClickHandlers
 				}
 				else
 				{
-					//something is wrong with the hitproxy relating to this click - create a debug log to help identify what
-					if( InHitProxy == NULL) 
-					{
-						UE_LOG(LogEditorViewport, Warning, TEXT("Invalid hitproxy" ) );
-					}
-					else
+					//try to get the name of the object also
+					FString name = TEXT("UNKNOWN");
+					if( InHitProxy->GetGeomObject()->GetActualBrush() != NULL )
 					{
-						//try to get the name of the object also
-						FString name = TEXT("UNKNOWN");
-						if( InHitProxy->GetGeomObject()->GetActualBrush() != NULL )
-						{
-							name = InHitProxy->GetGeomObject()->GetActualBrush()->GetName();
-						}
-						UE_LOG(LogEditorViewport, Warning, TEXT("Invalid PolyIndex %d on %s" ) ,InHitProxy->PolyIndex , *name );
+						name = InHitProxy->GetGeomObject()->GetActualBrush()->GetName();
 					}
-
+					UE_LOG(LogEditorViewport, Warning, TEXT("Invalid PolyIndex %d on %s" ) ,InHitProxy->PolyIndex , *name );
 				}
 			}
 		}
diff --git a/Engine/Source/Editor/UnrealEd/Private/Lightmass/Lightmass.cpp b/Engine/Source/Editor/UnrealEd/Private/Lightmass/Lightmass.cpp
--- a/Engine/Source/Editor/UnrealEd/Private/Lightmass/Lightmass.cpp
+++ b/Engine/Source/Editor/UnrealEd/Private/Lightmass/Lightmass.cpp
@@ -362,6 +362,7 @@ void FLightmassProcessor::SwarmCallback( NSwarm::FMessage* CallbackMessage, void
 
 			switch (AlertMessage->AlertLevel)
 			{
+			case NSwarm::ALERT_LEVEL_INFO:	break;
 			case NSwarm::ALERT_LEVEL_WARNING:			CheckType = EMessageSeverity::Warning;			break;
 			case NSwarm::ALERT_LEVEL_ERROR:				CheckType = EMessageSeverity::Error;			break;
 			case NSwarm::ALERT_LEVEL_CRITICAL_ERROR:	CheckType = EMessageSeverity::CriticalError;	break;
diff --git a/Engine/Source/Editor/UnrealEd/Private/PlayLevel.cpp b/Engine/Source/Editor/UnrealEd/Private/PlayLevel.cpp
--- a/Engine/Source/Editor/UnrealEd/Private/PlayLevel.cpp
+++ b/Engine/Source/Editor/UnrealEd/Private/PlayLevel.cpp
@@ -1709,6 +1709,7 @@ void UEditorEngine::PlayForMovieCapture()
 			case EMatineeCaptureType::JPEG:
 				EditorCommandLine += TEXT(" -MATINEESSFORMAT=JPEG");
 				break;
+			default: break;
 			}
 
 			// If buffer visualization dumping is enabled, we need to tell capture process to enable it too
diff --git a/Engine/Source/Programs/CrashReporter/CrashReportClient/Private/CrashUpload.cpp b/Engine/Source/Programs/CrashReporter/CrashReportClient/Private/CrashUpload.cpp
--- a/Engine/Source/Programs/CrashReporter/CrashReportClient/Private/CrashUpload.cpp
+++ b/Engine/Source/Programs/CrashReporter/CrashReportClient/Private/CrashUpload.cpp
@@ -219,7 +219,7 @@ void FCrashUpload::CompressAndSendData()
 	// Copy compressed data into the array.
 	TArray<uint8> CompressedData;
 	CompressedData.Append( CompressedDataRaw, CompressedSize );
-	delete CompressedDataRaw;
+	delete [] CompressedDataRaw;
 	CompressedDataRaw = nullptr;
 
 	// Set up request for upload
diff --git a/Engine/Source/Runtime/Core/Public/Containers/LockFreeVoidPointerListBase.h b/Engine/Source/Runtime/Core/Public/Containers/LockFreeVoidPointerListBase.h
--- a/Engine/Source/Runtime/Core/Public/Containers/LockFreeVoidPointerListBase.h
+++ b/Engine/Source/Runtime/Core/Public/Containers/LockFreeVoidPointerListBase.h
@@ -535,8 +535,7 @@ private:
 				UE_CLOG(0/*MONITOR_LINK_ALLOCATION*/,LogLockFreeList, Log, TEXT("Number of links %d"),NumUsedLinks.GetValue());
 			}
 
-			FLink*  NewLink = nullptr;
-			NewLink = FLink::Unlink(&FreeLinks);
+			FLink*  NewLink = FLink::Unlink(&FreeLinks);
 			if (NewLink)
 			{
 				NumFreeLinks.Decrement();
diff --git a/Engine/Source/Runtime/Core/Public/HAL/MallocDebug.h b/Engine/Source/Runtime/Core/Public/HAL/MallocDebug.h
--- a/Engine/Source/Runtime/Core/Public/HAL/MallocDebug.h
+++ b/Engine/Source/Runtime/Core/Public/HAL/MallocDebug.h
@@ -55,8 +55,7 @@ public:
 	virtual void* Malloc( SIZE_T Size, uint32 Alignment ) override
 	{
 		check(Alignment == DEFAULT_ALIGNMENT && "Alignment currently unsupported in this allocator");
-		FMemDebug* Ptr = nullptr;
-		Ptr = (FMemDebug*)malloc( AllocatorOverhead + Size );
+		FMemDebug* Ptr = (FMemDebug*)malloc( AllocatorOverhead + Size );
 		check(Ptr);
 		uint8* AlignedPtr = Align( (uint8*) Ptr + sizeof(FMemDebug) + sizeof(FMemDebug*) + sizeof(int32), ALLOCATION_ALIGNMENT );
 		Ptr->RefCount = 1;
diff --git a/Engine/Source/Runtime/Core/Public/HAL/ThreadingBase.h b/Engine/Source/Runtime/Core/Public/HAL/ThreadingBase.h
--- a/Engine/Source/Runtime/Core/Public/HAL/ThreadingBase.h
+++ b/Engine/Source/Runtime/Core/Public/HAL/ThreadingBase.h
@@ -1147,7 +1147,7 @@ public:
 	 */
 	FORCEINLINE static T& Get()
 	{
-		return *(T*)FThreadSingletonInitializer::Get( [](){ return (FTlsAutoCleanup*)new T(); }, T::GetTlsSlot() );
+		return *(T*)FThreadSingletonInitializer::Get( [](){ return (FTlsAutoCleanup*)new T(); }, T::GetTlsSlot() ); //-V572
 	}
 };
 
diff --git a/Engine/Source/Runtime/Core/Public/Windows/WindowsPlatformCodeAnalysis.h b/Engine/Source/Runtime/Core/Public/Windows/WindowsPlatformCodeAnalysis.h
--- a/Engine/Source/Runtime/Core/Public/Windows/WindowsPlatformCodeAnalysis.h
+++ b/Engine/Source/Runtime/Core/Public/Windows/WindowsPlatformCodeAnalysis.h
@@ -133,10 +133,12 @@
 	//-V:SelectedEmitter:519
 	//-V:MAX_VERTS_PER_POLY:512
 	//-V:127:547
+	//-V:0x7F:547
 	//-V:FT_LOAD_TARGET_NORMAL:616
 	//-V:OPENGL_PERFORMANCE_DATA_INVALID:564
 	//-V:HLSLCC_VersionMajor:616
 	//-V:UWorld:623
+	//-V:TextureRHI:623
 	//-V:bIgnoreFieldReferences:519
 	//-V:CachedQueryInstance:519
 	//-V:MeshContext:519
@@ -150,4 +152,23 @@
 	//-V:SPAWN_INIT:595
 	//-V:BEGIN_UPDATE_LOOP:595
 	//-V:OPENGL_PERFORMANCE_DATA_INVALID:560
+	//-V:bSkipTranslationTrack:560
+	//-V:NumSelected>0:581
+	//-V:bTryPerTrackBitwiseCompression:581
+	//-V:DataStripped:581
+	//-V:FromInt:601
+	//-V:UE_CLOG(:501
+	//-V:UE_LOG(:501
+	//-V:UGL_REQUIRED_VOID:501
+	//-V:ENQUEUE_UNIQUE_RENDER_COMMAND_ONEPARAMETER:601
+	//-V:AnimScriptInstance:595
+	//-V:Driver:595
+	//-V:PSceneAsync->lockWrite:595
+	//-V:Context.World():595
+	//-V:ALLOCATE_VERTEX_DATA_TEMPLATE:501
+	//-V:UGL_REQUIRED:501
+	//-V:DEBUG_LOG_HTTP:523
+	//-V:GIsEditor:560
+	//-V:bHasEditorToken:560
+
 #endif
\ No newline at end of file
diff --git a/Engine/Source/Runtime/CoreUObject/Private/Serialization/AsyncLoading.cpp b/Engine/Source/Runtime/CoreUObject/Private/Serialization/AsyncLoading.cpp
--- a/Engine/Source/Runtime/CoreUObject/Private/Serialization/AsyncLoading.cpp
+++ b/Engine/Source/Runtime/CoreUObject/Private/Serialization/AsyncLoading.cpp
@@ -1498,7 +1498,7 @@ void FAsyncPackage::UpdateLoadPercentage()
 	}
 	else if (DeferredPostLoadObjects.Num() > 0)
 	{
-		NewLoadPercentage = DeferredPostLoadIndex / DeferredPostLoadObjects.Num();
+		NewLoadPercentage = static_cast<float>(DeferredPostLoadIndex) / DeferredPostLoadObjects.Num();
 	}
 	// It's also possible that we got so many objects to PostLoad that LoadPercantage will actually drop
 	LoadPercentage = FMath::Max(NewLoadPercentage, LoadPercentage);
diff --git a/Engine/Source/Runtime/Engine/Private/Animation/AnimNode_SpringBone.cpp b/Engine/Source/Runtime/Engine/Private/Animation/AnimNode_SpringBone.cpp
--- a/Engine/Source/Runtime/Engine/Private/Animation/AnimNode_SpringBone.cpp
+++ b/Engine/Source/Runtime/Engine/Private/Animation/AnimNode_SpringBone.cpp
@@ -63,7 +63,7 @@ void FAnimNode_SpringBone::EvaluateBoneTransforms(USkeletalMeshComponent* SkelCo
 
 	// Location of our bone in world space
 	FTransform SpaceBase = MeshBases.GetComponentSpaceTransform(SpringBone.BoneIndex);
-	FTransform  BoneTransformInWorldSpace = SpaceBase * SkelComp->GetComponentToWorld();
+	FTransform  BoneTransformInWorldSpace = (SkelComp != NULL) ? SpaceBase * SkelComp->GetComponentToWorld() : SpaceBase;
 
 	FVector const TargetPos = BoneTransformInWorldSpace.GetLocation();
 
diff --git a/Engine/Source/Runtime/Engine/Private/Components/SkeletalMeshComponent.cpp b/Engine/Source/Runtime/Engine/Private/Components/SkeletalMeshComponent.cpp
--- a/Engine/Source/Runtime/Engine/Private/Components/SkeletalMeshComponent.cpp
+++ b/Engine/Source/Runtime/Engine/Private/Components/SkeletalMeshComponent.cpp
@@ -1078,7 +1078,10 @@ void USkeletalMeshComponent::RefreshBoneTransforms(FActorComponentTickFunction*
 		Prerequistes.Add(ParallelAnimationEvaluationTask);
 		FGraphEventRef TickCompletionEvent = TGraphTask<FParallelAnimationCompletionTask>::CreateTask(&Prerequistes).ConstructAndDispatchWhenReady(this);
 
-		TickFunction->GetCompletionHandle()->DontCompleteUntil(TickCompletionEvent);
+		if ( TickFunction )
+		{
+			TickFunction->GetCompletionHandle()->DontCompleteUntil(TickCompletionEvent);
+		}
 	}
 	else
 	{
diff --git a/Engine/Source/Runtime/Engine/Private/LocalPlayer.cpp b/Engine/Source/Runtime/Engine/Private/LocalPlayer.cpp
--- a/Engine/Source/Runtime/Engine/Private/LocalPlayer.cpp
+++ b/Engine/Source/Runtime/Engine/Private/LocalPlayer.cpp
@@ -731,7 +731,7 @@ FSceneView* ULocalPlayer::CalcSceneView( class FSceneViewFamily* ViewFamily,
 		ViewInitOptions.bInCameraCut = PlayerController->PlayerCameraManager->bGameCameraCutThisFrame;
 	}
 	
-	check( PlayerController->GetWorld() );
+	check(PlayerController && PlayerController->GetWorld());
 
 	// Fill out the rest of the view init options
 	ViewInitOptions.ViewFamily = ViewFamily;
@@ -761,7 +761,7 @@ FSceneView* ULocalPlayer::CalcSceneView( class FSceneViewFamily* ViewFamily,
 		View->StartFinalPostprocessSettings(OutViewLocation);
 
 		// CameraAnim override
-		if (PlayerController && PlayerController->PlayerCameraManager)
+		if (PlayerController->PlayerCameraManager)
 		{
 			TArray<FPostProcessSettings> const* CameraAnimPPSettings;
 			TArray<float> const* CameraAnimPPBlendWeights;
diff --git a/Engine/Source/Runtime/Engine/Private/Particles/ParticleModules.cpp b/Engine/Source/Runtime/Engine/Private/Particles/ParticleModules.cpp
--- a/Engine/Source/Runtime/Engine/Private/Particles/ParticleModules.cpp
+++ b/Engine/Source/Runtime/Engine/Private/Particles/ParticleModules.cpp
@@ -3078,7 +3078,7 @@ void UParticleModuleKillBox::Render3DPreview(FParticleEmitterInstance* Owner, co
 		KillboxVerts[7].Set(CheckUR.X, CheckLL.Y, CheckUR.Z);
 	}
 	
-	if ((bAbsolute == false) && (Owner != NULL) && (Owner->Component != NULL))
+	if ((bAbsolute == false) && (Owner->Component != NULL))
 	{
 		if (bAxisAlignedAndFixedSize == false)
 		{
diff --git a/Engine/Source/Runtime/Engine/Private/Particles/ParticleModules_Velocity.cpp b/Engine/Source/Runtime/Engine/Private/Particles/ParticleModules_Velocity.cpp
--- a/Engine/Source/Runtime/Engine/Private/Particles/ParticleModules_Velocity.cpp
+++ b/Engine/Source/Runtime/Engine/Private/Particles/ParticleModules_Velocity.cpp
@@ -85,7 +85,7 @@ void UParticleModuleVelocity::SpawnEx(FParticleEmitterInstance* Owner, int32 Off
 		FVector FromOrigin = (Particle.Location - Owner->EmitterToSimulation.GetOrigin()).GetSafeNormal();
 
 		FVector OwnerScale(1.0f);
-		if ((bApplyOwnerScale == true) && Owner && Owner->Component)
+		if ((bApplyOwnerScale == true) && Owner->Component)
 		{
 			OwnerScale = Owner->Component->ComponentToWorld.GetScale3D();
 		}
@@ -270,7 +270,7 @@ void UParticleModuleVelocityOverLifetime::Update(FParticleEmitterInstance* Owner
 	UParticleLODLevel* LODLevel	= Owner->SpriteTemplate->GetCurrentLODLevel(Owner);
 	check(LODLevel);
 	FVector OwnerScale(1.0f);
-	if ((bApplyOwnerScale == true) && Owner && Owner->Component)
+	if ((bApplyOwnerScale == true) && Owner->Component)
 	{
 		OwnerScale = Owner->Component->ComponentToWorld.GetScale3D();
 	}
@@ -429,7 +429,7 @@ void UParticleModuleVelocityCone::SpawnEx(FParticleEmitterInstance* Owner, int32
 	UParticleLODLevel* LODLevel	= Owner->SpriteTemplate->GetCurrentLODLevel(Owner);
 	check(LODLevel);
 	FVector OwnerScale(1.0f);
-	if ((bApplyOwnerScale == true) && Owner && Owner->Component)
+	if ((bApplyOwnerScale == true) && Owner->Component)
 	{
 		OwnerScale = Owner->Component->ComponentToWorld.GetScale3D();
 	}
@@ -541,7 +541,7 @@ void UParticleModuleVelocityCone::Render3DPreview(FParticleEmitterInstance* Owne
 	FMatrix OwnerRotation(FMatrix::Identity);
 	FVector LocalToWorldOrigin(0.0f);
 	FMatrix LocalToWorld(FMatrix::Identity);
-	if (Owner && Owner->Component)
+	if (Owner->Component)
 	{
 		AActor* Actor = Owner->Component->GetOwner();
 		if (Actor)
diff --git a/Engine/Source/Runtime/Engine/Private/Particles/ParticleSystemRender.cpp b/Engine/Source/Runtime/Engine/Private/Particles/ParticleSystemRender.cpp
--- a/Engine/Source/Runtime/Engine/Private/Particles/ParticleSystemRender.cpp
+++ b/Engine/Source/Runtime/Engine/Private/Particles/ParticleSystemRender.cpp
@@ -5350,8 +5350,7 @@ int32 FDynamicRibbonEmitterData::FillVertexData(struct FAsyncBufferFillData& Dat
 	FVector	ViewOrigin	= CameraToWorld.GetOrigin();
 
 	int32 MaxTessellationBetweenParticles = FMath::Max<int32>(Source.MaxTessellationBetweenParticles, 1);
-	int32 Sheets = FMath::Max<int32>(Source.Sheets, 1);
-	Sheets = 1;
+	int32 Sheets = 1;
 
 	// The distance tracking for tiling the 2nd UV set
 	float CurrDistance = 0.0f;
@@ -6050,8 +6049,7 @@ int32 FDynamicAnimTrailEmitterData::FillVertexData(struct FAsyncBufferFillData&
 	uint8* TempDynamicParamData = (uint8*)Data.DynamicParameterData;
 	FParticleBeamTrailVertexDynamicParameter* DynParamVertex;
 
-	int32 Sheets = FMath::Max<int32>(Source.Sheets, 1);
-	Sheets = 1;
+	int32 Sheets = 1;
 
 	// The increment for going [0..1] along the complete trail
 	float TextureIncrement = 1.0f / (Data.VertexCount / 2.0f);
diff --git a/Engine/Source/Runtime/Engine/Private/PendingNetGame.cpp b/Engine/Source/Runtime/Engine/Private/PendingNetGame.cpp
--- a/Engine/Source/Runtime/Engine/Private/PendingNetGame.cpp
+++ b/Engine/Source/Runtime/Engine/Private/PendingNetGame.cpp
@@ -242,8 +242,7 @@ void UPendingNetGame::NotifyControlMessage(UNetConnection* Connection, uint8 Mes
 
 void UPendingNetGame::Tick( float DeltaTime )
 {
-	check(NetDriver);
-	check(NetDriver->ServerConnection);
+	check(NetDriver && NetDriver->ServerConnection);
 
 	// The following line disables checks for nullptr access on NetDriver. We have checked() it's validity above,
 	// but the TickDispatch call below may invalidate the ptr, thus we must null check after calling TickDispatch.
diff --git a/Engine/Source/Runtime/Engine/Private/ScriptPlatformInterface.cpp b/Engine/Source/Runtime/Engine/Private/ScriptPlatformInterface.cpp
--- a/Engine/Source/Runtime/Engine/Private/ScriptPlatformInterface.cpp
+++ b/Engine/Source/Runtime/Engine/Private/ScriptPlatformInterface.cpp
@@ -209,6 +209,8 @@ bool UCloudStorageBase::ReadKeyValue(const FString& KeyName, EPlatformInterfaceD
 				UE_LOG(LogScriptPlatformInterface, Log, TEXT("CONVERT SOME uint8 ARRAY TO UOBJECT HERE?"));
 			}
 			break;
+		case PIDT_Custom:
+			break;
 	}
 	GConfig->DisableFileOperations();
 
diff --git a/Engine/Source/Runtime/Launch/Private/LaunchEngineLoop.cpp b/Engine/Source/Runtime/Launch/Private/LaunchEngineLoop.cpp
--- a/Engine/Source/Runtime/Launch/Private/LaunchEngineLoop.cpp
+++ b/Engine/Source/Runtime/Launch/Private/LaunchEngineLoop.cpp
@@ -405,8 +405,7 @@ bool LaunchCheckForFileOverride(const TCHAR* CmdLine, bool& OutFileOverrideFound
 
 	// Try to create pak file wrapper
 	{
-		IPlatformFile* PlatformFile = nullptr;
-		PlatformFile = ConditionallyCreateFileWrapper(TEXT("PakFile"), CurrentPlatformFile, CmdLine);
+		IPlatformFile* PlatformFile = ConditionallyCreateFileWrapper(TEXT("PakFile"), CurrentPlatformFile, CmdLine);
 		if (PlatformFile)
 		{
 			CurrentPlatformFile = PlatformFile;
diff --git a/Engine/Source/Runtime/UMG/Private/Components/RichTextBlockDecorator.cpp b/Engine/Source/Runtime/UMG/Private/Components/RichTextBlockDecorator.cpp
--- a/Engine/Source/Runtime/UMG/Private/Components/RichTextBlockDecorator.cpp
+++ b/Engine/Source/Runtime/UMG/Private/Components/RichTextBlockDecorator.cpp
@@ -171,7 +171,7 @@ void FDefaultRichTextDecorator::ExplodeRunInfo(const FRunInfo& InRunInfo, FSlate
 		const FString& FontColorStringRef = *FontColorString;
 
 		// Is Hex color?
-		if ( ( !FontColorStringRef.IsEmpty() && FontColorStringRef[0] == TCHAR('#') ) )
+		if ( !FontColorStringRef.IsEmpty() && FontColorStringRef[0] == TCHAR('#') )
 		{
 			OutFontColor = FLinearColor(FColor::FromHex(FontColorStringRef));
 		}
diff --git a/Engine/Source/Runtime/UMG/Private/UWidgetNavigation.cpp b/Engine/Source/Runtime/UMG/Private/UWidgetNavigation.cpp
--- a/Engine/Source/Runtime/UMG/Private/UWidgetNavigation.cpp
+++ b/Engine/Source/Runtime/UMG/Private/UWidgetNavigation.cpp
@@ -24,6 +24,8 @@ FWidgetNavigationData& UWidgetNavigation::GetNavigationData(EUINavigation Nav)
 		return Left;
 	case EUINavigation::Right:
 		return Right;
+	default:
+		break;
 	}
 
 	// Should never happen
diff --git a/Engine/Source/Runtime/UnrealAudio/Private/UnrealAudioDeviceFormat.cpp b/Engine/Source/Runtime/UnrealAudio/Private/UnrealAudioDeviceFormat.cpp
--- a/Engine/Source/Runtime/UnrealAudio/Private/UnrealAudioDeviceFormat.cpp
+++ b/Engine/Source/Runtime/UnrealAudio/Private/UnrealAudioDeviceFormat.cpp
@@ -101,6 +101,9 @@ namespace UAudio
 			case EStreamFormat::INT_32:
 				ConvertToFloatType<int32, FloatType>(ToBuffer, FromBuffer, StreamInfo.BlockSize, ConvertInfo);
 				break;
+
+			default:
+				break;
 		}
 	}
 
@@ -187,6 +190,9 @@ namespace UAudio
 			case EStreamFormat::INT_32:
 				ConvertIntegerToIntegerType<int32, FloatType>(ToBuffer, FromBuffer, StreamInfo.BlockSize, ConvertInfo);
 				break;
+
+			default:
+				break;
 		}
 	}
 
@@ -215,6 +221,9 @@ namespace UAudio
 			case EStreamFormat::INT_32:
 				ConvertAllToIntegerType<double>(ToBuffer, FromBuffer, StreamInfo.BlockSize, ConvertInfo);
 				break;
+
+			default:
+				break;
 		}
 
 		return true;
diff --git a/Engine/Source/Runtime/Windows/D3D11RHI/Private/D3D11RHI.cpp b/Engine/Source/Runtime/Windows/D3D11RHI/Private/D3D11RHI.cpp
--- a/Engine/Source/Runtime/Windows/D3D11RHI/Private/D3D11RHI.cpp
+++ b/Engine/Source/Runtime/Windows/D3D11RHI/Private/D3D11RHI.cpp
@@ -617,7 +617,7 @@ void UpdateBufferStats(TRefCountPtr<ID3D11Buffer> Buffer, bool bAllocating)
 		}
 	}
 	else
-	{
+	{ //-V523
 		if (bUniformBuffer)
 		{
 			DEC_MEMORY_STAT_BY(STAT_UniformBufferMemory,Desc.ByteWidth);
