Beter fix for UE-7888

diff --git a/Engine/Source/Editor/UnrealEd/Private/Kismet2/BlueprintEditorUtils.cpp b/Engine/Source/Editor/UnrealEd/Private/Kismet2/BlueprintEditorUtils.cpp
--- a/Engine/Source/Editor/UnrealEd/Private/Kismet2/BlueprintEditorUtils.cpp
+++ b/Engine/Source/Editor/UnrealEd/Private/Kismet2/BlueprintEditorUtils.cpp
@@ -1892,26 +1892,34 @@ void FBlueprintEditorUtils::UpdateDelegatesInBlueprint(UBlueprint* Blueprint)
 // Blueprint has materially changed.  Recompile the skeleton, notify observers, and mark the package as dirty.
 void FBlueprintEditorUtils::MarkBlueprintAsStructurallyModified(UBlueprint* Blueprint)
 {
-	auto SkeletalRecompileChildren = [](TArray<UClass*> ChildrenOfClass)
+	struct FRefreshHelper
 	{
-		for (auto SkelClass : ChildrenOfClass)
+		static void SkeletalRecompileChildren(TArray<UClass*> SkelClassesToRecompile)
 		{
-			IKismetCompilerInterface& Compiler = FModuleManager::LoadModuleChecked<IKismetCompilerInterface>(KISMET_COMPILER_MODULENAME);
-
-			auto Blueprint = Cast<UBlueprint>(SkelClass->ClassGeneratedBy);
-			if (Blueprint
-				&& Blueprint->Status != BS_BeingCreated
-				&& !Blueprint->bBeingCompiled
-				&& !Blueprint->bIsRegeneratingOnLoad
-				&& Blueprint->bHasBeenRegenerated)
+			for (auto SkelClass : SkelClassesToRecompile)
 			{
-				FCompilerResultsLog Results;
-				Results.bSilentMode = true;
-				Results.bLogInfoOnly = true;
+				auto Blueprint = Cast<UBlueprint>(SkelClass->ClassGeneratedBy);
+				if (Blueprint
+					&& Blueprint->Status != BS_BeingCreated
+					&& !Blueprint->bBeingCompiled
+					&& !Blueprint->bIsRegeneratingOnLoad
+					&& Blueprint->bHasBeenRegenerated)
+				{
+					TArray<UClass*> ChildrenOfClass;
+					GetDerivedClasses(SkelClass, ChildrenOfClass, false);
+
+					IKismetCompilerInterface& Compiler = FModuleManager::LoadModuleChecked<IKismetCompilerInterface>(KISMET_COMPILER_MODULENAME);
 
-				FKismetCompilerOptions CompileOptions;
-				CompileOptions.CompileType = EKismetCompileType::SkeletonOnly;
-				Compiler.CompileBlueprint(Blueprint, CompileOptions, Results);
+					FCompilerResultsLog Results;
+					Results.bSilentMode = true;
+					Results.bLogInfoOnly = true;
+
+					FKismetCompilerOptions CompileOptions;
+					CompileOptions.CompileType = EKismetCompileType::SkeletonOnly;
+					Compiler.CompileBlueprint(Blueprint, CompileOptions, Results);
+
+					SkeletalRecompileChildren(ChildrenOfClass);
+				}
 			}
 		}
 	};
@@ -1945,7 +1953,7 @@ void FBlueprintEditorUtils::MarkBlueprintAsStructurallyModified(UBlueprint* Blue
 		TArray<UClass*> ChildrenOfClass;
 		if (UClass* SkelClass = Blueprint->SkeletonGeneratedClass)
 		{
-			GetDerivedClasses(SkelClass, ChildrenOfClass);
+			GetDerivedClasses(SkelClass, ChildrenOfClass, false);
 		}
 
 		{
@@ -1959,7 +1967,7 @@ void FBlueprintEditorUtils::MarkBlueprintAsStructurallyModified(UBlueprint* Blue
 		}
 		UpdateDelegatesInBlueprint(Blueprint);
 
-		SkeletalRecompileChildren(ChildrenOfClass);
+		FRefreshHelper::SkeletalRecompileChildren(ChildrenOfClass);
 
 		{
 			BP_SCOPED_COMPILER_EVENT_STAT(EKismetCompilerStats_NotifyBlueprintChanged);
