UDS: default values for properties of class type.

diff --git a/Engine/Source/Editor/BlueprintGraph/Private/K2Node_MakeStruct.cpp b/Engine/Source/Editor/BlueprintGraph/Private/K2Node_MakeStruct.cpp
--- a/Engine/Source/Editor/BlueprintGraph/Private/K2Node_MakeStruct.cpp
+++ b/Engine/Source/Editor/BlueprintGraph/Private/K2Node_MakeStruct.cpp
@@ -23,13 +23,17 @@ void UK2Node_MakeStruct::FMakeStructPinManager::CustomizePinData(UEdGraphPin* Pi
 {
 	struct FPinDefaultValueHelper
 	{
-		static void Set(UEdGraphPin& InPin, const FString& Value, bool bIsText)
+		static void Set(UEdGraphPin& InPin, const FString& Value, bool bIsText, bool bIsObject)
 		{
 			InPin.AutogeneratedDefaultValue = Value;
 			if (bIsText)
 			{
 				InPin.DefaultTextValue = FText::FromString(Value);
 			}
+			else if (bIsObject)
+			{
+				InPin.DefaultObject = StaticFindObject(UObject::StaticClass(), nullptr, *Value);
+			}
 			else
 			{
 				InPin.DefaultValue = Value;
@@ -47,10 +51,13 @@ void UK2Node_MakeStruct::FMakeStructPinManager::CustomizePinData(UEdGraphPin* Pi
 		const bool bIsText = Property->IsA<UTextProperty>();
 		checkSlow(bIsText == ((Schema->PC_Text == Pin->PinType.PinCategory) && !Pin->PinType.bIsArray));
 
+		const bool bIsObject = Property->IsA<UObjectPropertyBase>();
+		checkSlow(bIsObject == ((Schema->PC_Object == Pin->PinType.PinCategory || Schema->PC_Class == Pin->PinType.PinCategory) && !Pin->PinType.bIsArray));
+
 		const FString& MetadataDefaultValue = Property->GetMetaData(TEXT("MakeStructureDefaultValue"));
 		if (!MetadataDefaultValue.IsEmpty())
 		{
-			FPinDefaultValueHelper::Set(*Pin, MetadataDefaultValue, bIsText);
+			FPinDefaultValueHelper::Set(*Pin, MetadataDefaultValue, bIsText, bIsObject);
 			return;
 		}
 
@@ -61,7 +68,7 @@ void UK2Node_MakeStruct::FMakeStructPinManager::CustomizePinData(UEdGraphPin* Pi
 			{
 				if (Schema->IsPinDefaultValid(Pin, NewDefaultValue, NULL, FText::GetEmpty()).IsEmpty())
 				{
-					FPinDefaultValueHelper::Set(*Pin, NewDefaultValue, bIsText);
+					FPinDefaultValueHelper::Set(*Pin, NewDefaultValue, bIsText, bIsObject);
 					return;
 				}
 			}
@@ -221,7 +228,7 @@ bool UK2Node_MakeStruct::CanBeExposed(const UProperty* Property)
 
 bool UK2Node_MakeStruct::CanBeMade(const UScriptStruct* Struct)
 {
-	if(!Struct->HasMetaData(TEXT("HasNativeMake")))
+	if (Struct && !Struct->HasMetaData(TEXT("HasNativeMake")))
 	{
 		for (TFieldIterator<UProperty> It(Struct); It; ++It)
 		{
diff --git a/Engine/Source/Editor/KismetCompiler/Private/UserDefinedStructureCompilerUtils.cpp b/Engine/Source/Editor/KismetCompiler/Private/UserDefinedStructureCompilerUtils.cpp
--- a/Engine/Source/Editor/KismetCompiler/Private/UserDefinedStructureCompilerUtils.cpp
+++ b/Engine/Source/Editor/KismetCompiler/Private/UserDefinedStructureCompilerUtils.cpp
@@ -34,6 +34,7 @@ struct FUserDefinedStructureCompilerInner
 			DuplicatedStruct->Status = EUserDefinedStructureStatus::UDSS_Duplicate;
 			DuplicatedStruct->SetFlags(RF_Transient);
 			DuplicatedStruct->AddToRoot();
+			CastChecked<UUserDefinedStructEditorData>(DuplicatedStruct->EditorData)->RecreateDefaultInstance();
 
 			for (auto StructProperty : TObjectRange<UStructProperty>(RF_ClassDefaultObject | RF_PendingKill))
 			{
@@ -75,6 +76,11 @@ struct FUserDefinedStructureCompilerInner
 	{
 		check(StructToClean);
 
+		if (auto EditorData = Cast<UUserDefinedStructEditorData>(StructToClean->EditorData))
+		{
+			EditorData->CleanDefaultInstance();
+		}
+
 		const FString TransientString = FString::Printf(TEXT("TRASHSTRUCT_%s"), *StructToClean->GetName());
 		const FName TransientName = MakeUniqueObjectName(GetTransientPackage(), UUserDefinedStruct::StaticClass(), FName(*TransientString));
 		UUserDefinedStruct* TransientStruct = NewObject<UUserDefinedStruct>(GetTransientPackage(), TransientName, RF_Public | RF_Transient);
@@ -192,6 +198,13 @@ struct FUserDefinedStructureCompilerInner
 			LogError(Struct, MessageLog, FString::Printf(*LOCTEXT("StructurEmpty_Error", "Structure '%s' is empty ").ToString(), *Struct->GetFullName()));
 		}
 
+		FString DefaultInstanceError;
+		EditorData->RecreateDefaultInstance(&DefaultInstanceError);
+		if (!DefaultInstanceError.IsEmpty())
+		{
+			LogError(Struct, MessageLog, DefaultInstanceError);
+		}
+
 		const bool bNoErrorsDuringCompilation = (ErrorNum == MessageLog.NumErrors);
 		Struct->Status = bNoErrorsDuringCompilation ? EUserDefinedStructureStatus::UDSS_UpToDate : EUserDefinedStructureStatus::UDSS_Error;
 	}
diff --git a/Engine/Source/Editor/PropertyEditor/Private/PropertyEditorModule.cpp b/Engine/Source/Editor/PropertyEditor/Private/PropertyEditorModule.cpp
--- a/Engine/Source/Editor/PropertyEditor/Private/PropertyEditorModule.cpp
+++ b/Engine/Source/Editor/PropertyEditor/Private/PropertyEditorModule.cpp
@@ -667,7 +667,7 @@ TSharedRef<class IStructureDetailsView> FPropertyEditorModule::CreateStructureDe
 				const auto ArrayProperty = Cast<UArrayProperty>(&PropertyAndParent.Property);
 				const auto TrueProperty = ArrayProperty ? ArrayProperty->Inner : &PropertyAndParent.Property;
 				const bool bCanShow = TrueProperty 
-					&& !TrueProperty->IsA<UObjectPropertyBase>()
+					&& (!TrueProperty->IsA<UObjectPropertyBase>() || TrueProperty->IsA<UClassProperty>())
 					&& !TrueProperty->IsA<UInterfaceProperty>();
 				return bCanShow;
 			}
diff --git a/Engine/Source/Editor/UnrealEd/Classes/UserDefinedStructure/UserDefinedStructEditorData.h b/Engine/Source/Editor/UnrealEd/Classes/UserDefinedStructure/UserDefinedStructEditorData.h
--- a/Engine/Source/Editor/UnrealEd/Classes/UserDefinedStructure/UserDefinedStructEditorData.h
+++ b/Engine/Source/Editor/UnrealEd/Classes/UserDefinedStructure/UserDefinedStructEditorData.h
@@ -62,6 +62,19 @@ struct FStructVariableDescription
 	{ }
 };
 
+class FStructOnScopeMember : public FStructOnScope
+{
+public:
+	FStructOnScopeMember() : FStructOnScope() {}
+
+	void Recreate(const UStruct* InScriptStruct)
+	{
+		Destroy();
+		ScriptStruct = InScriptStruct;
+		Initialize();
+	}
+};
+
 UCLASS()
 class UNREALED_API UUserDefinedStructEditorData : public UObject
 {
@@ -79,9 +92,16 @@ public:
 	UPROPERTY()
 	FString ToolTip;
 
+	// optional super struct
 	UPROPERTY()
 	UScriptStruct* NativeBase;
 
+protected:
+	FStructOnScopeMember DefaultStructInstance;
+
+public:
+	static void AddReferencedObjects(UObject* InThis, FReferenceCollector& Collector);
+
 	// UObject interface.
 	virtual void PostEditUndo() override;
 	virtual void PostLoadSubobjects(struct FObjectInstancingGraph* OuterInstanceGraph) override;
@@ -89,4 +109,8 @@ public:
 
 	uint32 GenerateUniqueNameIdForMemberVariable();
 	class UUserDefinedStruct* GetOwnerStruct() const;
+
+	const uint8* GetDefaultInstance() const;
+	void RecreateDefaultInstance(FString* OutLog = nullptr);
+	void CleanDefaultInstance();
 };
\ No newline at end of file
diff --git a/Engine/Source/Editor/UnrealEd/Private/Kismet2/StructureEditorUtils.cpp b/Engine/Source/Editor/UnrealEd/Private/Kismet2/StructureEditorUtils.cpp
--- a/Engine/Source/Editor/UnrealEd/Private/Kismet2/StructureEditorUtils.cpp
+++ b/Engine/Source/Editor/UnrealEd/Private/Kismet2/StructureEditorUtils.cpp
@@ -365,11 +365,32 @@ bool FStructureEditorUtils::ChangeVariableType(UUserDefinedStruct* Struct, FGuid
 
 bool FStructureEditorUtils::ChangeVariableDefaultValue(UUserDefinedStruct* Struct, FGuid VarGuid, const FString& NewDefaultValue)
 {
-	const UEdGraphSchema_K2* K2Schema = GetDefault<UEdGraphSchema_K2>();
+	auto ValidateDefaultValue = [](const FStructVariableDescription& VarDesc, const FString& NewDefaultValue) -> bool
+	{
+		const UEdGraphSchema_K2* K2Schema = GetDefault<UEdGraphSchema_K2>();
+		const FEdGraphPinType PinType = VarDesc.ToPinType();
+
+		bool bResult = false;
+		//TODO: validation for values, that are not passed by string
+		if (PinType.PinCategory == K2Schema->PC_Text)
+		{
+			bResult = true;
+		}
+		else if ((PinType.PinCategory == K2Schema->PC_Object) || (PinType.PinCategory == K2Schema->PC_Interface) || (PinType.PinCategory == K2Schema->PC_Class))
+		{
+			bResult = true;
+		}
+		else
+		{
+			bResult = K2Schema->DefaultValueSimpleValidation(PinType, FString(), NewDefaultValue, NULL, FText::GetEmpty());
+		}
+		return bResult;
+	};
+
 	auto VarDesc = GetVarDescByGuid(Struct, VarGuid);
 	if (VarDesc 
 		&& (NewDefaultValue != VarDesc->DefaultValue)
-		&& K2Schema->DefaultValueSimpleValidation(VarDesc->ToPinType(), FString(), NewDefaultValue, NULL, FText::GetEmpty()))
+		&& ValidateDefaultValue(*VarDesc, NewDefaultValue))
 	{
 		bool bAdvancedValidation = true;
 		if (!NewDefaultValue.IsEmpty())
@@ -421,26 +442,32 @@ bool FStructureEditorUtils::UserDefinedStructEnabled()
 	return UseUserDefinedStructure;
 }
 
+void FStructureEditorUtils::RecreateDefaultInstanceInEditorData(UUserDefinedStruct* Struct)
+{
+	auto StructEditorData = Struct ? CastChecked<UUserDefinedStructEditorData>(Struct->EditorData) : nullptr;
+	if (StructEditorData)
+	{
+		StructEditorData->RecreateDefaultInstance();
+	}
+}
+
 bool FStructureEditorUtils::Fill_MakeStructureDefaultValue(const UUserDefinedStruct* Struct, uint8* StructData)
 {
 	bool bResult = true;
-	for (TFieldIterator<UProperty> It(Struct); It; ++It)
+	if (Struct && StructData)
 	{
-		if (const UProperty* Property = *It)
+		auto StructEditorData = CastChecked<UUserDefinedStructEditorData>(Struct->EditorData);
+		const uint8* DefaultInstance = StructEditorData->GetDefaultInstance();
+		if (DefaultInstance)
 		{
-			auto VarDesc = GetVarDesc(Struct).FindByPredicate(FFindByNameHelper<FStructVariableDescription>(Property->GetFName()));
-			if (VarDesc && !VarDesc->CurrentDefaultValue.IsEmpty())
-			{
-				bResult &= FBlueprintEditorUtils::PropertyValueFromString(Property, VarDesc->CurrentDefaultValue, StructData);
-			}
-			else
-			{
-				uint8* const InnerData = Property->ContainerPtrToValuePtr<uint8>(StructData);
-				bResult &= Fill_MakeStructureDefaultValue(Property, InnerData);
-			}
+			Struct->CopyScriptStruct(StructData, DefaultInstance);
+		}
+		else
+		{
+			bResult = false;
 		}
 	}
-
+	
 	return bResult;
 }
 
@@ -634,7 +661,7 @@ bool FStructureEditorUtils::MoveVariable(UUserDefinedStruct* Struct, FGuid VarGu
 		{
 			if (DescArray[Index].VarGuid == VarGuid)
 			{
-				const FScopedTransaction Transaction(LOCTEXT("ChangeVariableDefaultValue", "Varaibles reordered"));
+				const FScopedTransaction Transaction(LOCTEXT("ReorderVariables", "Varaibles reordered"));
 				ModifyStructData(Struct);
 
 				DescArray.Swap(Index, Index + (bMoveUp ? -1 : 1));
@@ -669,7 +696,7 @@ bool FStructureEditorUtils::Change3dWidgetEnabled(UUserDefinedStruct* Struct, FG
 	const auto PropertyStruct = VarDesc ? Cast<const UStruct>(VarDesc->SubCategoryObject.Get()) : NULL;
 	if (FEdMode::CanCreateWidgetForStructure(PropertyStruct) && (VarDesc->bEnable3dWidget != bIsEnabled))
 	{
-		const FScopedTransaction Transaction(LOCTEXT("ChangeVariableDefaultValue", "Change Variable Default Value"));
+		const FScopedTransaction Transaction(LOCTEXT("Change3dWidgetEnabled", "Change 3d Widget Enabled"));
 		ModifyStructData(Struct);
 
 		VarDesc->bEnable3dWidget = bIsEnabled;
diff --git a/Engine/Source/Editor/UnrealEd/Private/UserDefinedStructEditorData.cpp b/Engine/Source/Editor/UnrealEd/Private/UserDefinedStructEditorData.cpp
--- a/Engine/Source/Editor/UnrealEd/Private/UserDefinedStructEditorData.cpp
+++ b/Engine/Source/Editor/UnrealEd/Private/UserDefinedStructEditorData.cpp
@@ -3,6 +3,9 @@
 #include "UnrealEd.h"
 #include "StructureEditorUtils.h"
 #include "Engine/UserDefinedStruct.h"
+#include "BlueprintEditorUtils.h"
+
+#define LOCTEXT_NAMESPACE "UserDefinedStructEditorData"
 
 bool FStructVariableDescription::SetPinType(const FEdGraphPinType& VarType)
 {
@@ -50,4 +53,67 @@ void UUserDefinedStructEditorData::PostLoadSubobjects(FObjectInstancingGraph* Ou
 	{
 		VarDesc.bInvalidMember = !FStructureEditorUtils::CanHaveAMemberVariableOfType(GetOwnerStruct(), VarDesc.ToPinType());
 	}
-}
\ No newline at end of file
+}
+
+const uint8* UUserDefinedStructEditorData::GetDefaultInstance() const
+{
+	ensure(DefaultStructInstance.IsValid() && DefaultStructInstance.GetStruct() == GetOwnerStruct());
+	return DefaultStructInstance.GetStructMemory();
+}
+
+void UUserDefinedStructEditorData::RecreateDefaultInstance(FString* OutLog)
+{
+	UStruct* ScriptStruct = GetOwnerStruct();
+	DefaultStructInstance.Recreate(ScriptStruct);
+	uint8* StructData = DefaultStructInstance.GetStructMemory();
+	ensure(DefaultStructInstance.IsValid() && DefaultStructInstance.GetStruct() == ScriptStruct);
+	if (DefaultStructInstance.IsValid() && StructData && ScriptStruct)
+	{
+		for (TFieldIterator<UProperty> It(ScriptStruct); It; ++It)
+		{
+			UProperty* Property = *It;
+			if (Property)
+			{
+				auto VarDesc = VariablesDescriptions.FindByPredicate(FStructureEditorUtils::FFindByNameHelper<FStructVariableDescription>(Property->GetFName()));
+				if (VarDesc && !VarDesc->CurrentDefaultValue.IsEmpty())
+				{
+					if (!FBlueprintEditorUtils::PropertyValueFromString(Property, VarDesc->CurrentDefaultValue, StructData))
+					{
+						const FString Message = FString::Printf(TEXT("Cannot parse value. Property: %s String: \"%s\" ")
+							, (Property ? *Property->GetDisplayNameText().ToString() : TEXT("None"))
+							, *VarDesc->CurrentDefaultValue);
+						UE_LOG(LogClass, Warning, TEXT("UUserDefinedStructEditorData::RecreateDefaultInstance %s Struct: %s "), *Message, *GetPathNameSafe(ScriptStruct));
+						if (OutLog)
+						{
+							OutLog->Append(Message);
+						}
+					}
+				}
+			}
+		}
+	}
+}
+
+void UUserDefinedStructEditorData::CleanDefaultInstance()
+{
+	ensure(!DefaultStructInstance.IsValid() || DefaultStructInstance.GetStruct() == GetOwnerStruct());
+	DefaultStructInstance.Destroy();
+}
+
+void UUserDefinedStructEditorData::AddReferencedObjects(UObject* InThis, FReferenceCollector& Collector)
+{
+	UUserDefinedStructEditorData* This = CastChecked<UUserDefinedStructEditorData>(InThis);
+
+	UStruct* ScriptStruct = This->GetOwnerStruct();
+	ensure(!This->DefaultStructInstance.IsValid() || This->DefaultStructInstance.GetStruct() == ScriptStruct);
+	uint8* StructData = This->DefaultStructInstance.GetStructMemory();
+	if (StructData)
+	{
+		FSimpleObjectReferenceCollectorArchive ObjectReferenceCollector(This, Collector);
+		ScriptStruct->SerializeBin(ObjectReferenceCollector, StructData);
+	}
+
+	Super::AddReferencedObjects(This, Collector);
+}
+
+#undef LOCTEXT_NAMESPACE
\ No newline at end of file
diff --git a/Engine/Source/Editor/UnrealEd/Public/Kismet2/StructureEditorUtils.h b/Engine/Source/Editor/UnrealEd/Public/Kismet2/StructureEditorUtils.h
--- a/Engine/Source/Editor/UnrealEd/Public/Kismet2/StructureEditorUtils.h
+++ b/Engine/Source/Editor/UnrealEd/Public/Kismet2/StructureEditorUtils.h
@@ -132,6 +132,8 @@ public:
 
 	static bool Fill_MakeStructureDefaultValue(const UProperty* Property, uint8* PropertyData);
 
+	static void RecreateDefaultInstanceInEditorData(UUserDefinedStruct* Struct);
+
 	//VALIDATION
 	static bool CanHaveAMemberVariableOfType(const UUserDefinedStruct* Struct, const FEdGraphPinType& VarType, FString* OutMsg = NULL);
 
diff --git a/Engine/Source/Runtime/CoreUObject/Public/UObject/Class.h b/Engine/Source/Runtime/CoreUObject/Public/UObject/Class.h
--- a/Engine/Source/Runtime/CoreUObject/Public/UObject/Class.h
+++ b/Engine/Source/Runtime/CoreUObject/Public/UObject/Class.h
@@ -1098,10 +1098,7 @@ protected:
 
 	FStructOnScope() : SampleStructMemory(NULL) {}
 
-public:
-	FStructOnScope(const UStruct* InScriptStruct)
-		: ScriptStruct(InScriptStruct)
-		, SampleStructMemory(NULL)
+	void Initialize()
 	{
 		if (ScriptStruct.IsValid())
 		{
@@ -1110,6 +1107,14 @@ public:
 		}
 	}
 
+public:
+	FStructOnScope(const UStruct* InScriptStruct)
+		: ScriptStruct(InScriptStruct)
+		, SampleStructMemory(NULL)
+	{
+		Initialize();
+	}
+
 	virtual uint8* GetStructMemory() { return SampleStructMemory; }
 
 	virtual const uint8* GetStructMemory() const { return SampleStructMemory; }
diff --git a/Engine/Source/Runtime/Engine/Private/UserDefinedStruct.cpp b/Engine/Source/Runtime/Engine/Private/UserDefinedStruct.cpp
--- a/Engine/Source/Runtime/Engine/Private/UserDefinedStruct.cpp
+++ b/Engine/Source/Runtime/Engine/Private/UserDefinedStruct.cpp
@@ -25,6 +25,10 @@ void UUserDefinedStruct::Serialize(FArchive& Ar)
 		if ( EditorData != nullptr )
 		{
 			Ar.Preload(EditorData);
+			if (!(Ar.GetPortFlags() & PPF_Duplicate))
+			{
+				FStructureEditorUtils::RecreateDefaultInstanceInEditorData(this);
+			}
 		}
 
 		const FStructureEditorUtils::EStructureError Result = FStructureEditorUtils::IsStructureValid(this, NULL, &ErrorMessage);
