Newly created Blueprints will no longer be marked as dirty (compilation).

diff --git a/Engine/Source/Editor/UnrealEd/Private/Kismet2/Kismet2.cpp b/Engine/Source/Editor/UnrealEd/Private/Kismet2/Kismet2.cpp
--- a/Engine/Source/Editor/UnrealEd/Private/Kismet2/Kismet2.cpp
+++ b/Engine/Source/Editor/UnrealEd/Private/Kismet2/Kismet2.cpp
@@ -590,9 +590,6 @@ UK2Node_Event* FKismetEditorUtilities::AddDefaultEventNode(UBlueprint* InBluepri
 
 		InGraph->AddNode(NewEventNode);
 
-		// Rebuild the skeleton class so we can determine if we are overriding the function from the parent
-		FBlueprintEditorUtils::MarkBlueprintAsStructurallyModified(InBlueprint);
-
 		// Get the function that the event node or function entry represents
 		FFunctionFromNodeHelper FunctionFromNode(NewEventNode);
 		if (FunctionFromNode.Function && Schema->GetCallableParentFunction(FunctionFromNode.Function))
