Add actor array to the property test object

diff --git a/Engine/Source/Editor/UnrealEd/Classes/Editor/PropertyEditorTestObject.h b/Engine/Source/Editor/UnrealEd/Classes/Editor/PropertyEditorTestObject.h
--- a/Engine/Source/Editor/UnrealEd/Classes/Editor/PropertyEditorTestObject.h
+++ b/Engine/Source/Editor/UnrealEd/Classes/Editor/PropertyEditorTestObject.h
@@ -180,6 +180,9 @@ class UPropertyEditorTestObject : public UObject
 	UPROPERTY(EditAnywhere, Category=ArraysOfProperties)
 	TArray<UObject*> ObjectPropertyArray;
 
+	UPROPERTY(EditAnywhere, Category=ArraysOfProperties)
+	TArray<AActor*> ActorPropertyArray;
+
 	UPROPERTY(EditAnywhere, Category=ArraysOfProperties)
 	TArray<FLinearColor> LinearColorPropertyArray;
 
