#UE4 Logs/Warnings/Errors from Blueprint compiles now include the name of the blueprint that emitted the event.

diff --git a/Engine/Source/Editor/UnrealEd/Private/Kismet2/CompilerResultsLog.cpp b/Engine/Source/Editor/UnrealEd/Private/Kismet2/CompilerResultsLog.cpp
--- a/Engine/Source/Editor/UnrealEd/Private/Kismet2/CompilerResultsLog.cpp
+++ b/Engine/Source/Editor/UnrealEd/Private/Kismet2/CompilerResultsLog.cpp
@@ -280,15 +280,15 @@ void FCompilerResultsLog::InternalLogMessage(const EMessageSeverity::Type& Sever
 	{
 		if(Severity == EMessageSeverity::CriticalError || Severity == EMessageSeverity::Error)
 		{
-			UE_LOG(LogBlueprint, Error, TEXT("[compiler] %s"), *Line->ToText().ToString());
+			UE_LOG(LogBlueprint, Error, TEXT("[compiler %s] %s"), *SourceName, *Line->ToText().ToString());
 		}
 		else if(Severity == EMessageSeverity::Warning || Severity == EMessageSeverity::PerformanceWarning)
 		{
-			UE_LOG(LogBlueprint, Warning, TEXT("[compiler] %s"), *Line->ToText().ToString());
+			UE_LOG(LogBlueprint, Warning, TEXT("[compiler %s] %s"), *SourceName, *Line->ToText().ToString());
 		}
 		else
 		{
-			UE_LOG(LogBlueprint, Log, TEXT("[compiler] %s"), *Line->ToText().ToString());
+			UE_LOG(LogBlueprint, Log, TEXT("[compiler %s] %s"), *SourceName, *Line->ToText().ToString());
 		}
 	}
 }
