Make sure all objects get marked modify even if they have overriden the first changed property

diff --git a/Engine/Source/Editor/UnrealEd/Private/Editor.cpp b/Engine/Source/Editor/UnrealEd/Private/Editor.cpp
--- a/Engine/Source/Editor/UnrealEd/Private/Editor.cpp
+++ b/Engine/Source/Editor/UnrealEd/Private/Editor.cpp
@@ -7296,77 +7296,75 @@ namespace EditorUtilities
 
 		// Copy non-component properties from the old actor to the new actor
 		// @todo sequencer: Most of this block of code was borrowed (pasted) from UEditorEngine::ConvertActors().  If we end up being able to share these code bodies, that would be nice!
+		TSet<UObject*> ModifiedObjects;
+		for( UProperty* Property = ActorClass->PropertyLink; Property != nullptr; Property = Property->PropertyLinkNext )
 		{
-			bool bIsFirstModification = true;
-			for( UProperty* Property = ActorClass->PropertyLink; Property != nullptr; Property = Property->PropertyLinkNext )
-			{
-				const bool bIsTransient = !!( Property->PropertyFlags & CPF_Transient );
-				const bool bIsComponentContainer = !!( Property->PropertyFlags & CPF_ContainsInstancedReference );
-				const bool bIsComponentProp = !!( Property->PropertyFlags & ( CPF_InstancedReference | CPF_ContainsInstancedReference ) );
-				const bool bIsBlueprintReadonly = !!(Options & ECopyOptions::FilterBlueprintReadOnly) && !!( Property->PropertyFlags & CPF_BlueprintReadOnly );
-				const bool bIsIdentical = Property->Identical_InContainer( SourceActor, TargetActor );
+			const bool bIsTransient = !!( Property->PropertyFlags & CPF_Transient );
+			const bool bIsComponentContainer = !!( Property->PropertyFlags & CPF_ContainsInstancedReference );
+			const bool bIsComponentProp = !!( Property->PropertyFlags & ( CPF_InstancedReference | CPF_ContainsInstancedReference ) );
+			const bool bIsBlueprintReadonly = !!(Options & ECopyOptions::FilterBlueprintReadOnly) && !!( Property->PropertyFlags & CPF_BlueprintReadOnly );
+			const bool bIsIdentical = Property->Identical_InContainer( SourceActor, TargetActor );
 
-				if( !bIsTransient && !bIsIdentical && !bIsComponentContainer && !bIsComponentProp && !bIsBlueprintReadonly)
+			if( !bIsTransient && !bIsIdentical && !bIsComponentContainer && !bIsComponentProp && !bIsBlueprintReadonly)
+			{
+				const bool bIsSafeToCopy = !( Options & ECopyOptions::OnlyCopyEditOrInterpProperties ) || ( Property->HasAnyPropertyFlags( CPF_Edit | CPF_Interp ) );
+				if( bIsSafeToCopy )
 				{
-					const bool bIsSafeToCopy = !( Options & ECopyOptions::OnlyCopyEditOrInterpProperties ) || ( Property->HasAnyPropertyFlags( CPF_Edit | CPF_Interp ) );
-					if( bIsSafeToCopy )
+					if( !bIsPreviewing )
 					{
-						if( !bIsPreviewing )
+						if( !ModifiedObjects.Contains(TargetActor) )
 						{
-							if( bIsFirstModification )
-							{
-								// Start modifying the target object
-								TargetActor->Modify();
-							}
+							// Start modifying the target object
+							TargetActor->Modify();
+							ModifiedObjects.Add(TargetActor);
+						}
 
-							if( Options & ECopyOptions::CallPostEditChangeProperty )
-							{
-								TargetActor->PreEditChange( Property );
-							}
+						if( Options & ECopyOptions::CallPostEditChangeProperty )
+						{
+							TargetActor->PreEditChange( Property );
+						}
 
-							// Determine which archetype instances match the current property value of the target actor (before it gets changed). We only want to propagate the change to those instances.
-							TArray<UObject*> ArchetypeInstancesToChange;
-							if( Options & ECopyOptions::PropagateChangesToArchetypeInstances )
+						// Determine which archetype instances match the current property value of the target actor (before it gets changed). We only want to propagate the change to those instances.
+						TArray<UObject*> ArchetypeInstancesToChange;
+						if( Options & ECopyOptions::PropagateChangesToArchetypeInstances )
+						{
+							for( AActor* ArchetypeInstance : ArchetypeInstances )
 							{
-								for( AActor* ArchetypeInstance : ArchetypeInstances )
+								if( ArchetypeInstance != nullptr && Property->Identical_InContainer( ArchetypeInstance, TargetActor ) )
 								{
-									if( ArchetypeInstance != nullptr && Property->Identical_InContainer( ArchetypeInstance, TargetActor ) )
-									{
-										ArchetypeInstancesToChange.Add( ArchetypeInstance );
-									}
+									ArchetypeInstancesToChange.Add( ArchetypeInstance );
 								}
 							}
+						}
 
-							CopySingleProperty(SourceActor, TargetActor, Property);
+						CopySingleProperty(SourceActor, TargetActor, Property);
 
-							if( Options & ECopyOptions::CallPostEditChangeProperty )
-							{
-								FPropertyChangedEvent PropertyChangedEvent( Property );
-								TargetActor->PostEditChangeProperty( PropertyChangedEvent );
-							}
+						if( Options & ECopyOptions::CallPostEditChangeProperty )
+						{
+							FPropertyChangedEvent PropertyChangedEvent( Property );
+							TargetActor->PostEditChangeProperty( PropertyChangedEvent );
+						}
 
-							if( Options & ECopyOptions::PropagateChangesToArchetypeInstances )
+						if( Options & ECopyOptions::PropagateChangesToArchetypeInstances )
+						{
+							for( int32 InstanceIndex = 0; InstanceIndex < ArchetypeInstancesToChange.Num(); ++InstanceIndex )
 							{
-								for( int32 InstanceIndex = 0; InstanceIndex < ArchetypeInstancesToChange.Num(); ++InstanceIndex )
+								UObject* ArchetypeInstance = ArchetypeInstancesToChange[InstanceIndex];
+								if( ArchetypeInstance != nullptr )
 								{
-									UObject* ArchetypeInstance = ArchetypeInstancesToChange[InstanceIndex];
-									if( ArchetypeInstance != nullptr )
+									if( !ModifiedObjects.Contains(ArchetypeInstance) )
 									{
-										if( bIsFirstModification )
-										{
-											ArchetypeInstance->Modify();
-										}
-
-										CopySingleProperty( TargetActor, ArchetypeInstance, Property );
+										ArchetypeInstance->Modify();
+										ModifiedObjects.Add(ArchetypeInstance);
 									}
+
+									CopySingleProperty( TargetActor, ArchetypeInstance, Property );
 								}
 							}
-
-							bIsFirstModification = false;
 						}
-
-						++CopiedPropertyCount;
 					}
+
+					++CopiedPropertyCount;
 				}
 			}
 		}
@@ -7411,7 +7409,6 @@ namespace EditorUtilities
 				SourceComponent->GetUCSModifiedProperties(SourceUCSModifiedProperties);
 
 				// Copy component properties
-				bool bIsFirstModification = true;
 				for( UProperty* Property = ComponentClass->PropertyLink; Property != nullptr; Property = Property->PropertyLinkNext )
 				{
 					const bool bIsTransient = !!( Property->PropertyFlags & CPF_Transient );
@@ -7430,10 +7427,11 @@ namespace EditorUtilities
 						{
 							if( !bIsPreviewing )
 							{
-								if( bIsFirstModification )
+								if( !ModifiedObjects.Contains(TargetComponent) )
 								{
 									TargetComponent->SetFlags(RF_Transactional);
 									TargetComponent->Modify();
+									ModifiedObjects.Add(TargetComponent);
 								}
 
 								if( Options & ECopyOptions::CallPostEditChangeProperty )
@@ -7489,7 +7487,7 @@ namespace EditorUtilities
 										UActorComponent* ComponentArchetypeInstance = ComponentArchetypeInstancesToChange[InstanceIndex];
 										if( ComponentArchetypeInstance != nullptr )
 										{
-											if( bIsFirstModification )
+											if( !ModifiedObjects.Contains(ComponentArchetypeInstance) )
 											{
 												// Ensure that this instance will be included in any undo/redo operations, and record it into the transaction buffer.
 												// Note: We don't do this for components that originate from script, because they will be re-instanced from the template after an undo, so there is no need to record them.
@@ -7497,13 +7495,15 @@ namespace EditorUtilities
 												{
 													ComponentArchetypeInstance->SetFlags(RF_Transactional);
 													ComponentArchetypeInstance->Modify();
+													ModifiedObjects.Add(ComponentArchetypeInstance);
 												}
 
 												// We must also modify the owner, because we'll need script components to be reconstructed as part of an undo operation.
 												AActor* Owner = ComponentArchetypeInstance->GetOwner();
-												if( Owner != nullptr )
+												if( Owner != nullptr && !ModifiedObjects.Contains(Owner))
 												{
 													Owner->Modify();
+													ModifiedObjects.Add(Owner);
 												}
 											}
 
@@ -7514,8 +7514,6 @@ namespace EditorUtilities
 										}
 									}
 								}
-
-								bIsFirstModification = false;
 							}
 
 							++CopiedPropertyCount;
