Fixed circular widget reference in MergeActors dialog window
Changed default name for merged static mesh to be based on actor name instead of source mesh asset name
Replace with merged actor is inside transaction now so it can be undo

diff --git a/Engine/Source/Editor/UnrealEd/Private/Dialogs/SMeshMergingDialog.cpp b/Engine/Source/Editor/UnrealEd/Private/Dialogs/SMeshMergingDialog.cpp
--- a/Engine/Source/Editor/UnrealEd/Private/Dialogs/SMeshMergingDialog.cpp
+++ b/Engine/Source/Editor/UnrealEd/Private/Dialogs/SMeshMergingDialog.cpp
@@ -263,14 +263,23 @@ void SMeshMergingDialog::Construct(const FArguments& InArgs)
 
 FReply SMeshMergingDialog::OnCancelClicked()
 {
-	ParentWindow.Get()->RequestDestroyWindow();
+	auto ParentPtr = ParentWindow.Get().Pin();
+	if (ParentPtr.IsValid())
+	{
+		ParentPtr->RequestDestroyWindow();
+	}
 	return FReply::Handled();
 }
 
 FReply SMeshMergingDialog::OnMergeClicked()
 {
 	RunMerging();
-	ParentWindow.Get()->RequestDestroyWindow();
+	
+	auto ParentPtr = ParentWindow.Get().Pin();
+	if (ParentPtr.IsValid())
+	{
+		ParentPtr->RequestDestroyWindow();
+	}
 	return FReply::Handled();
 }
 
@@ -336,9 +345,9 @@ void SMeshMergingDialog::SetPlaceInWorld(ECheckBoxState NewValue)
 
 void SMeshMergingDialog::GenerateNewPackageName()
 {
-	USelection* SelectedActors = GEditor->GetSelectedActors();
-	MergedMeshPackageName.Empty();
+	MergedMeshPackageName = FPackageName::FilenameToLongPackageName(FPaths::GameContentDir() + TEXT("SM_MERGED"));
 
+	USelection* SelectedActors = GEditor->GetSelectedActors();
 	// Iterate through selected actors and find first static mesh asset
 	// Use this static mesh path as destination package name for a merged mesh
 	for(FSelectionIterator Iter(*SelectedActors); Iter; ++Iter)
@@ -346,28 +355,14 @@ void SMeshMergingDialog::GenerateNewPackageName()
 		AActor* Actor = Cast<AActor>(*Iter);
 		if (Actor)
 		{
-			TInlineComponentArray<UStaticMeshComponent*> SMComponets; 
-			Actor->GetComponents<UStaticMeshComponent>(SMComponets);
-			for (UStaticMeshComponent* Component : SMComponets)
-			{
-				if (Component->StaticMesh)
-				{
-					MergedMeshPackageName = FPackageName::GetLongPackagePath(Component->StaticMesh->GetOutermost()->GetName());
-					MergedMeshPackageName+= FString(TEXT("/SM_MERGED_")) + Component->StaticMesh->GetName();
-					break;
-				}
-			}
-		}
-
-		if (!MergedMeshPackageName.IsEmpty())
-		{
+			FString ActorName = Actor->GetName();
+			MergedMeshPackageName = FString::Printf(TEXT("%s_%s"), *MergedMeshPackageName, *ActorName);
 			break;
 		}
 	}
 
 	if (MergedMeshPackageName.IsEmpty())
 	{
-		MergedMeshPackageName = FPackageName::FilenameToLongPackageName(FPaths::GameContentDir() + TEXT("SM_MERGED"));
 		MergedMeshPackageName = MakeUniqueObjectName(NULL, UPackage::StaticClass(), *MergedMeshPackageName).ToString();
 	}
 }
@@ -442,6 +437,9 @@ void SMeshMergingDialog::RunMerging()
 		// Place new mesh in the world
 		if (bPlaceInWorld)
 		{
+			const FScopedTransaction Transaction(NSLOCTEXT("UnrealEd", "PlaceMergedActor", "Place Merged Actor"));
+			UniqueLevels[0]->Modify();
+			
 			UWorld* World = UniqueLevels[0]->OwningWorld;
 			FActorSpawnParameters Params;
 			Params.OverrideLevel = UniqueLevels[0];
@@ -454,6 +452,7 @@ void SMeshMergingDialog::RunMerging()
 			// Add source actors as children to merged actor and hide them
 			for (AActor* Actor : Actors)
 			{
+				Actor->Modify();
 				Actor->AttachRootComponentToActor(MergedActor, NAME_None, EAttachLocation::KeepWorldPosition);
 				Actor->SetActorHiddenInGame(true);
 				Actor->SetIsTemporarilyHiddenInEditor(true);
diff --git a/Engine/Source/Editor/UnrealEd/Public/Dialogs/SMeshMergingDialog.h b/Engine/Source/Editor/UnrealEd/Public/Dialogs/SMeshMergingDialog.h
--- a/Engine/Source/Editor/UnrealEd/Public/Dialogs/SMeshMergingDialog.h
+++ b/Engine/Source/Editor/UnrealEd/Public/Dialogs/SMeshMergingDialog.h
@@ -14,7 +14,7 @@ public:
 	}
 
 	// The parent window hosting this dialog
-	SLATE_ATTRIBUTE(TSharedPtr<SWindow>, ParentWindow)
+	SLATE_ATTRIBUTE(TWeakPtr<SWindow>, ParentWindow)
 
 	SLATE_END_ARGS()
 
@@ -71,7 +71,7 @@ private:
 
 private:
 	/** Pointer to the parent window */
-	TAttribute<TSharedPtr<SWindow>> ParentWindow;
+	TAttribute<TWeakPtr<SWindow>> ParentWindow;
 
 	/** Current mesh merging settings */
 	FMeshMergingSettings MergingSettings;
