#github https://github.com/EpicGames/UnrealEngine/pull/1009
Typo in description of Retarget Dialog
#1009
user: Manny-MADE

diff --git a/Engine/Source/Editor/UnrealEd/Private/SSkeletonWidget.cpp b/Engine/Source/Editor/UnrealEd/Private/SSkeletonWidget.cpp
--- a/Engine/Source/Editor/UnrealEd/Private/SSkeletonWidget.cpp
+++ b/Engine/Source/Editor/UnrealEd/Private/SSkeletonWidget.cpp
@@ -601,7 +601,7 @@ void SAnimationRemapSkeleton::Construct( const FArguments& InArgs )
 					SNew(STextBlock)
 					.AutoWrapText(true)
 					.Font(FEditorStyle::GetFontStyle("Persona.RetargetManager.BoldFont"))
-					.Text(LOCTEXT("RetargetBasePose_WarningMessage", "Make sure you have the similar retarget base pose. If they don't look alike here, you can edit your base pose in the Retarget Mangaer window to look alike."))
+					.Text(LOCTEXT("RetargetBasePose_WarningMessage", "Make sure you have the similar retarget base pose. If they don't look alike here, you can edit your base pose in the Retarget Manager window to look alike."))
 				]
 
 				+SVerticalBox::Slot()
