Fixed reentrancy issue in FMouseDeltaTracker::EndTracking which could cause a crash if FDragTool_ActorBoxSelect::EndDrag opened a modal dialog.

diff --git a/Engine/Source/Editor/UnrealEd/Private/MouseDeltaTracker.cpp b/Engine/Source/Editor/UnrealEd/Private/MouseDeltaTracker.cpp
--- a/Engine/Source/Editor/UnrealEd/Private/MouseDeltaTracker.cpp
+++ b/Engine/Source/Editor/UnrealEd/Private/MouseDeltaTracker.cpp
@@ -29,6 +29,7 @@ FMouseDeltaTracker::FMouseDeltaTracker()
 	, DragTool( NULL )
 	, bHasAttemptedDragTool(false)
 	, bUsedDragModifier(false)
+	, bIsDeletingDragTool(false)
 {
 }
 
@@ -306,15 +307,21 @@ bool FMouseDeltaTracker::EndTracking(FEditorViewportClient* InViewportClient)
 
 	Start = StartSnapped = StartScreen = End = EndSnapped = EndScreen = RawDelta = ReductionAmount = FVector::ZeroVector;
 
-	// Delete the drag tool if one exists.
-	if( DragTool.IsValid() )
+	if (!bIsDeletingDragTool)
 	{
-		if( DragTool->IsDragging() )
+		// Ending the drag tool may pop up a modal dialog which can cause unwanted reentrancy - protect against this.
+		TGuardValue<bool> RecursionGuard(bIsDeletingDragTool, true);
+
+		// Delete the drag tool if one exists.
+		if (DragTool.IsValid())
 		{
-			DragTool->EndDrag();
+			if (DragTool->IsDragging())
+			{
+				DragTool->EndDrag();
+			}
+			DragTool.Reset();
+			return false;
 		}
-		DragTool.Reset();
-		return false;
 	}
 
 	// Do not fade snapping indicators over time if this viewport is not realtime
diff --git a/Engine/Source/Editor/UnrealEd/Public/MouseDeltaTracker.h b/Engine/Source/Editor/UnrealEd/Public/MouseDeltaTracker.h
--- a/Engine/Source/Editor/UnrealEd/Public/MouseDeltaTracker.h
+++ b/Engine/Source/Editor/UnrealEd/Public/MouseDeltaTracker.h
@@ -170,6 +170,8 @@ private:
 	/** Tracks if the user used a modifier to drag a selected item. (Rather than using a widget handle).Reset on StartTracking */
 	bool bUsedDragModifier;
 
+	/** Tracks whether the drag tool is in the process of being deleted (to protect against reentrancy) */
+	bool bIsDeletingDragTool;
 
 	/** Stores the widget mode active when the tracker begins tracking to help stop it change mid-track */
 	FWidget::EWidgetMode TrackingWidgetMode;
