Fix so that duplicated actors maintain relative references.

diff --git a/Engine/Source/Editor/UnrealEd/Private/EditorObject.cpp b/Engine/Source/Editor/UnrealEd/Private/EditorObject.cpp
--- a/Engine/Source/Editor/UnrealEd/Private/EditorObject.cpp
+++ b/Engine/Source/Editor/UnrealEd/Private/EditorObject.cpp
@@ -96,11 +96,74 @@ void UEditorEngine::RenameObject(UObject* Object,UObject* NewOuter,const TCHAR*
 	Object->MarkPackageDirty();
 }
 
+
+static void RemapProperty(UProperty* Property, int32 Index, const TMap<FName, AActor*>& ActorRemapper, uint8* DestData)
+{
+	if (UObjectProperty* ObjectProperty = Cast<UObjectProperty>(Property))
+	{
+		// If there's a concrete index, use that, otherwise iterate all array members (for the case that this property is inside a struct, or there is exactly one element)
+		const int32 Num = (Index == INDEX_NONE) ? ObjectProperty->ArrayDim : 1;
+		const int32 StartIndex = (Index == INDEX_NONE) ? 0 : Index;
+		for (int32 Count = 0; Count < Num; Count++)
+		{
+			uint8* PropertyAddr = ObjectProperty->ContainerPtrToValuePtr<uint8>(DestData, StartIndex + Count);
+			UObject* Object = ObjectProperty->GetObjectPropertyValue(PropertyAddr);
+			if (Object)
+			{
+				AActor* const* RemappedObject = ActorRemapper.Find(Object->GetFName());
+				if (RemappedObject)
+				{
+					ObjectProperty->SetObjectPropertyValue(PropertyAddr, *RemappedObject);
+				}
+			}
+
+		}
+	}
+	else if (UArrayProperty* ArrayProperty = Cast<UArrayProperty>(Property))
+	{
+		FScriptArrayHelper ArrayHelper(ArrayProperty, ArrayProperty->ContainerPtrToValuePtr<void>(DestData));
+		if (Index != INDEX_NONE)
+		{
+			RemapProperty(ArrayProperty->Inner, INDEX_NONE, ActorRemapper, ArrayHelper.GetRawPtr(Index));
+		}
+		else
+		{
+			for (int32 ArrayIndex = 0; ArrayIndex < ArrayHelper.Num(); ArrayIndex++)
+			{
+				RemapProperty(ArrayProperty->Inner, INDEX_NONE, ActorRemapper, ArrayHelper.GetRawPtr(ArrayIndex));
+			}
+		}
+	}
+	else if (UStructProperty* StructProperty = Cast<UStructProperty>(Property))
+	{
+		if (Index != INDEX_NONE)
+		{
+			// If a concrete index was given, remap just that
+			for (TFieldIterator<UProperty> It(StructProperty->Struct); It; ++It)
+			{
+				RemapProperty(*It, INDEX_NONE, ActorRemapper, StructProperty->ContainerPtrToValuePtr<uint8>(DestData, Index));
+			}
+		}
+		else
+		{
+			// If no concrete index was given, either the ArrayDim is 1 (i.e. not a static array), or the struct is within
+			// a deeper structure (an array or another struct) and we cannot know which element was changed, so iterate through all elements.
+			for (int32 Count = 0; Count < StructProperty->ArrayDim; Count++)
+			{
+				for (TFieldIterator<UProperty> It(StructProperty->Struct); It; ++It)
+				{
+					RemapProperty(*It, INDEX_NONE, ActorRemapper, StructProperty->ContainerPtrToValuePtr<uint8>(DestData, Count));
+				}
+			}
+		}
+	}
+}
+
+
 //
 //	ImportProperties
 //
 
-
 /**
  * Parse and import text as property values for the object specified.  This function should never be called directly - use ImportObjectProperties instead.
  * 
@@ -123,8 +186,9 @@ static const TCHAR* ImportProperties(
 	UObject*					SubobjectRoot,
 	UObject*					SubobjectOuter,
 	FFeedbackContext*			Warn,
-	int32							Depth,
-	FObjectInstancingGraph&		InstanceGraph
+	int32						Depth,
+	FObjectInstancingGraph&		InstanceGraph,
+	const TMap<FName, AActor*>* ActorRemapper
 	)
 {
 	check(!GIsUCCMakeStandaloneHeaderGenerator);
@@ -362,7 +426,7 @@ static const TCHAR* ImportProperties(
 			{
 				// since we're redefining an object in the same text block, only need to import properties again
 				SourceText = ImportObjectProperties( (uint8*)BaseTemplate, SourceText, TemplateClass, SubobjectRoot, BaseTemplate,
-													Warn, Depth + 1, ContextSupplier ? ContextSupplier->CurrentLine : 0, &InstanceGraph );
+													Warn, Depth + 1, ContextSupplier ? ContextSupplier->CurrentLine : 0, &InstanceGraph, ActorRemapper );
 			}
 			else 
 			{
@@ -519,7 +583,8 @@ static const TCHAR* ImportProperties(
 					Warn, 
 					Depth+1,
 					ContextSupplier ? ContextSupplier->CurrentLine : 0,
-					&InstanceGraph
+					&InstanceGraph,
+					ActorRemapper
 					);
 			}
 		}
@@ -545,6 +610,14 @@ static const TCHAR* ImportProperties(
 		}
 	}
 
+	if (ActorRemapper)
+	{
+		for (const auto& DefinedProperty : DefinedProperties)
+		{
+			RemapProperty(DefinedProperty.Property, DefinedProperty.Index, *ActorRemapper, DestData);
+		}
+	}
+
 	// Prepare brush.
 	if( ImportedBrush && ObjectStruct->IsChildOf(ABrush::StaticClass()) )
 	{
@@ -627,7 +700,8 @@ const TCHAR* ImportObjectProperties( FImportObjectParams& InParams )
 			InParams.SubobjectOuter,
 			InParams.Warn,
 			InParams.Depth,
-			InstanceGraph
+			InstanceGraph,
+			InParams.ActorRemapper
 			);
 
 	if ( InParams.SubobjectOuter != NULL )
@@ -704,7 +778,8 @@ const TCHAR* ImportObjectProperties(
 	FFeedbackContext*	Warn,
 	int32					Depth,
 	int32					LineNumber,
-	FObjectInstancingGraph* InInstanceGraph
+	FObjectInstancingGraph* InInstanceGraph,
+	const TMap<FName, AActor*>* ActorRemapper
 	)
 {
 	FImportObjectParams Params;
@@ -718,6 +793,7 @@ const TCHAR* ImportObjectProperties(
 		Params.Depth = Depth;
 		Params.LineNumber = LineNumber;
 		Params.InInstanceGraph = InInstanceGraph;
+		Params.ActorRemapper = ActorRemapper;
 
 		// This implementation always calls PreEditChange/PostEditChange
 		Params.bShouldCallEditChange = true;
diff --git a/Engine/Source/Editor/UnrealEd/Private/Factories/EditorFactories.cpp b/Engine/Source/Editor/UnrealEd/Private/Factories/EditorFactories.cpp
--- a/Engine/Source/Editor/UnrealEd/Private/Factories/EditorFactories.cpp
+++ b/Engine/Source/Editor/UnrealEd/Private/Factories/EditorFactories.cpp
@@ -838,7 +838,7 @@ UObject* ULevelFactory::FactoryCreateText
 			if ( Actor->ShouldImport(PropText, bIsMoveToStreamingLevel) )
 			{
 				Actor->PreEditChange(nullptr);
-				ImportObjectProperties( (uint8*)Actor, **PropText, Actor->GetClass(), Actor, Actor, Warn, 0 );
+				ImportObjectProperties( (uint8*)Actor, **PropText, Actor->GetClass(), Actor, Actor, Warn, 0, INDEX_NONE, NULL, &NewActorsFNames );
 				bActorChanged = true;
 
 				GEditor->SelectActor( Actor, true, false, true );
diff --git a/Engine/Source/Editor/UnrealEd/Public/Editor.h b/Engine/Source/Editor/UnrealEd/Public/Editor.h
--- a/Engine/Source/Editor/UnrealEd/Public/Editor.h
+++ b/Engine/Source/Editor/UnrealEd/Public/Editor.h
@@ -341,6 +341,8 @@ struct FImportObjectParams
 		not necessary to specify a value when calling this function from other code */
 	FObjectInstancingGraph* InInstanceGraph;
 
+	/** provides a mapping from an actor name to a new instance to which it should be remapped */
+	const TMap<FName, AActor*>* ActorRemapper;
 
 	/** True if we should call PreEditChange/PostEditChange on the object as it's imported.  Pass false here
 		if you're going to do that on your own. */
@@ -358,6 +360,7 @@ struct FImportObjectParams
 		  Depth( 0 ),
 		  LineNumber( INDEX_NONE ),
 		  InInstanceGraph( NULL ),
+		  ActorRemapper( NULL ),
 		  bShouldCallEditChange( true )
 	{
 	}
@@ -390,6 +393,7 @@ const TCHAR* ImportObjectProperties( FImportObjectParams& InParams );
  * @param	LineNumber			used when importing defaults during script compilation for tracking which line the defaultproperties block begins on
  * @param	InstanceGraph		contains the mappings of instanced objects and components to their templates; used when recursively calling ImportObjectProperties; generally
  *								not necessary to specify a value when calling this function from other code
+ * @param	ActorRemapper		used when duplicating actors to remap references from a source actor to the duplicated actor
  *
  * @return	NULL if the default values couldn't be imported
  */
@@ -402,7 +406,8 @@ const TCHAR* ImportObjectProperties(
 	FFeedbackContext*	Warn,
 	int32					Depth,
 	int32					LineNumber = INDEX_NONE,
-	FObjectInstancingGraph* InstanceGraph=NULL
+	FObjectInstancingGraph* InstanceGraph = NULL,
+	const TMap<FName, AActor*>* ActorRemapper = NULL
 	);
 
 //
