Stopped group actor bounding brackets from being rendered on viewports in Game View.

diff --git a/Engine/Source/Editor/UnrealEd/Private/GroupActor.cpp b/Engine/Source/Editor/UnrealEd/Private/GroupActor.cpp
--- a/Engine/Source/Editor/UnrealEd/Private/GroupActor.cpp
+++ b/Engine/Source/Editor/UnrealEd/Private/GroupActor.cpp
@@ -351,6 +351,12 @@ void PrivateDrawBracketsForGroups( FPrimitiveDrawInterface* PDI, FViewport* View
 
 void AGroupActor::DrawBracketsForGroups( FPrimitiveDrawInterface* PDI, FViewport* Viewport, bool bMustBeSelected/*=true*/ )
 {
+	// Don't draw group actor brackets in game view
+	if (Viewport->GetClient()->IsInGameView())
+	{
+		return;
+	}
+
 	if( GUnrealEd->bGroupingActive )
 	{
 		check(PDI);
