Enable vHACD on OS X.

diff --git a/Engine/Source/Editor/UnrealEd/Private/ConvexDecompTool.cpp b/Engine/Source/Editor/UnrealEd/Private/ConvexDecompTool.cpp
--- a/Engine/Source/Editor/UnrealEd/Private/ConvexDecompTool.cpp
+++ b/Engine/Source/Editor/UnrealEd/Private/ConvexDecompTool.cpp
@@ -5,7 +5,7 @@
 =============================================================================*/
 
 // Only enabling on windows until other platforms can test!
-#define USE_VHACD (PLATFORM_WINDOWS || PLATFORM_LINUX)
+#define USE_VHACD (PLATFORM_WINDOWS || PLATFORM_LINUX || PLATFORM_MAC)
 
 #include "UnrealEd.h"
 
@@ -106,28 +106,45 @@ void DecomposeMeshToHulls(UBodySetup* InBodySetup, const TArray<FVector>& InVert
 	{
 		cl_platform_id* Platforms = new cl_platform_id[NumPlatforms];
 		Result = clGetPlatformIDs(NumPlatforms, Platforms, NULL);
+		ANSICHAR PlatformVersion[1024];
 		for(cl_uint i = 0; Result == CL_SUCCESS && i < NumPlatforms; i++)
 		{
-			cl_uint NumGPUs = 0;
-			Result = clGetDeviceIDs(Platforms[i], CL_DEVICE_TYPE_GPU, 0, NULL, &NumGPUs);
-			if ( Result == CL_SUCCESS )
+			clGetPlatformInfo(Platforms[i], CL_PLATFORM_VERSION, sizeof(PlatformVersion), PlatformVersion, nullptr);
+			
+			int32 Major = 0;
+			int32 Minor = 0;
+			sscanf(PlatformVersion, "OpenCL %d.%d", &Major, &Minor);
+			
+			if(Major > 1 || (Major == 1 && Minor >= 1))
 			{
-				cl_device_id* Devices = new cl_device_id[NumGPUs];
-				Result = clGetDeviceIDs(Platforms[i], CL_DEVICE_TYPE_GPU, NumGPUs, Devices, NULL);
-				for(cl_uint j = 0; j < NumGPUs; j++)
+				cl_uint NumGPUs = 0;
+				Result = clGetDeviceIDs(Platforms[i], CL_DEVICE_TYPE_GPU, 0, NULL, &NumGPUs);
+				if ( Result == CL_SUCCESS )
 				{
-					cl_uint Freq = 0;
-					clGetDeviceInfo(Devices[j], CL_DEVICE_MAX_CLOCK_FREQUENCY, sizeof(Freq), &Freq, NULL);
-					cl_uint Units = 0;
-					clGetDeviceInfo(Devices[j], CL_DEVICE_MAX_COMPUTE_UNITS, sizeof(Units), &Units, NULL);
-					int gflops = Units * Freq;
-					if (gflops > max_gflops)
+					cl_device_id* Devices = new cl_device_id[NumGPUs];
+					Result = clGetDeviceIDs(Platforms[i], CL_DEVICE_TYPE_GPU, NumGPUs, Devices, NULL);
+					for(cl_uint j = 0; j < NumGPUs; j++)
 					{
-						max_gflops = gflops;
-						max_gflops_device = Devices[j];
+						clGetDeviceInfo(Devices[j], CL_DEVICE_VERSION, sizeof(PlatformVersion), PlatformVersion, nullptr);
+						Major = 0;
+						Minor = 0;
+						sscanf(PlatformVersion, "OpenCL %d.%d", &Major, &Minor);
+						if(Major > 1 || (Major == 1 && Minor >= 1))
+						{
+							cl_uint Freq = 0;
+							clGetDeviceInfo(Devices[j], CL_DEVICE_MAX_CLOCK_FREQUENCY, sizeof(Freq), &Freq, NULL);
+							size_t Units = 0;
+							clGetDeviceInfo(Devices[j], CL_DEVICE_MAX_WORK_GROUP_SIZE, sizeof(Units), &Units, NULL);
+							int gflops = Units * Freq;
+							if (gflops > max_gflops)
+							{
+								max_gflops = gflops;
+								max_gflops_device = Devices[j];
+							}
+						}
 					}
+					delete [] Devices;
 				}
-				delete [] Devices;
 			}
 		}
 		delete [] Platforms;
