fixing non-unity

diff --git a/Engine/Source/Editor/UnrealEd/Private/Dialogs/SMeshMergingDialog.cpp b/Engine/Source/Editor/UnrealEd/Private/Dialogs/SMeshMergingDialog.cpp
--- a/Engine/Source/Editor/UnrealEd/Private/Dialogs/SMeshMergingDialog.cpp
+++ b/Engine/Source/Editor/UnrealEd/Private/Dialogs/SMeshMergingDialog.cpp
@@ -13,6 +13,7 @@
 #include "Engine/StaticMesh.h"
 #include "Engine/Selection.h"
 #include "SystemSettings.h"
+#include "ScopedTransaction.h"
 
 #define LOCTEXT_NAMESPACE "SMeshMergingDialog"
 
