Refactored the static mesh UV's test which also fixes a crash.

diff --git a/Engine/Source/Editor/UnrealEd/Private/Tests/EditorAutomationTests.cpp b/Engine/Source/Editor/UnrealEd/Private/Tests/EditorAutomationTests.cpp
--- a/Engine/Source/Editor/UnrealEd/Private/Tests/EditorAutomationTests.cpp
+++ b/Engine/Source/Editor/UnrealEd/Private/Tests/EditorAutomationTests.cpp
@@ -2010,172 +2010,37 @@ bool FTraceAllTimelinesAutomationTest::RunTest(const FString& BlueprintName)
 	return bPassed;
 }
 
-/** Static Mesh UV test specific data */
-namespace StaticMeshUVTest
-{
-	/** States the test can be in */
-	namespace EStaticMeshUVTestState
-	{
-		enum Type
-		{
-			Ready,
-			WaitingForPackage,
-			PackageLoaded
-		};
-	}
-
-	class FUVTestHelper : public TSharedFromThis<FUVTestHelper>
-	{
-	public:
-
-		/** Constructor */
-		FUVTestHelper(FAutomationTestBase* BaseTest) :
-			CurrentState(EStaticMeshUVTestState::Ready),
-			CurrentPackage(NULL),
-			LoadedPackageCount(0),
-			CurrentTest(BaseTest)
-		{
-			//FPackageName::FindPackagesInDirectory(ContentPackages, *FPaths::EngineContentDir());
-			FPackageName::FindPackagesInDirectory(ContentPackages, *FPaths::GameContentDir());
-		}
-
-		/**
-		 * Update the static mesh UV test
-		 * @return true if the test is complete
-		 */
-		bool Update()
-		{
-			bool bTestDone = false;
-			switch(CurrentState)
-			{
-			case EStaticMeshUVTestState::Ready:
-				LoadNextPackage();
-				break;
-			case EStaticMeshUVTestState::PackageLoaded:
-				bTestDone = ProcessCurrentPackage();
-				break;
-			case EStaticMeshUVTestState::WaitingForPackage:
-			default:
-				break;
-			}
-
-			if( bTestDone )
-			{
-				CollectGarbage(RF_Native);
-			}
-			return bTestDone;
-		}
-
-	private:
-
-		/** 
-		 * Callback when a package has finished loading
-		 */
-		void PackageLoadCallback(const FName& PackageName, UPackage* LoadedPackage, EAsyncLoadingResult::Type Result)
-		{
-			CurrentPackage = LoadedPackage;
-			CurrentState = EStaticMeshUVTestState::PackageLoaded;
-			ContentPackages.RemoveAt(0);
-		}
-
-		/** 
-		 * Loads the next package into memory
-		 */
-		void LoadNextPackage()
-		{
-			CurrentTest->SetSuppressLogs(true);
-			const FString PackageName = ContentPackages[0];
-			CurrentPackage =  FindPackage(NULL, *PackageName);
-			if( CurrentPackage )
-			{
-				//Already loaded, check for static meshes
-				CurrentState = EStaticMeshUVTestState::PackageLoaded;
-				ContentPackages.RemoveAt(0);
-			}
-			else
-			{
-				LoadPackageAsync(*PackageName, FLoadPackageAsyncDelegate::CreateSP(this, &FUVTestHelper::PackageLoadCallback));
-				CurrentState = EStaticMeshUVTestState::WaitingForPackage;
-			}
-		}
-
-		/** 
-		 * Performs the UV test on any static meshes in the current package
-		 */
-		bool ProcessCurrentPackage()
-		{
-			LoadedPackageCount++;
-			CurrentTest->SetSuppressLogs(false);
-
-			for (FObjectIterator ObjIt; ObjIt; ++ObjIt)
-			{
-				UObject* Object = *ObjIt;
-				if (Object != NULL)
-				{
-					//Call CheckLightMapUVs on all static meshes in this package
-					if( Object->IsA(UStaticMesh::StaticClass()) && Object->IsIn(CurrentPackage) )
-					{
-						UStaticMesh::CheckLightMapUVs((UStaticMesh*)Object,MissingUVMessages,BadUVMessages,ValidUVMessages,true);
-					}
-				}
-			}
-
-			if( ( (LoadedPackageCount % 10) == 0 ) )
-			{
-				CollectGarbage(RF_Native);
-			}
-
-			CurrentState = EStaticMeshUVTestState::Ready;
-
-			//We are done when we are out of packages to process
-			return ContentPackages.Num() == 0;
-		}
-
-		// List of packages to test
-		TArray<FString> ContentPackages;
-		// The state of the current test (Ready, WaitingForPackage, PackageLoaded) 
-		EStaticMeshUVTestState::Type CurrentState;
-		// The current package we are checking for static meshes
-		UPackage* CurrentPackage;
-		// How many packages we have loaded
-		uint32 LoadedPackageCount;
-		// Pointer to the current test
-		FAutomationTestBase* CurrentTest;
-
-		// Missing UV messages
-		TArray<FString> MissingUVMessages;
-		// Bas UV messages
-		TArray<FString> BadUVMessages;
-		// Valid UV messages
-		TArray<FString> ValidUVMessages;
-	};
-}
 
 /**
- * Latent command wrapper for the static mesh UV test
- */
-DEFINE_LATENT_AUTOMATION_COMMAND_ONE_PARAMETER(FPerformUVTestCommand, TSharedPtr<StaticMeshUVTest::FUVTestHelper>, TestInfo);
+* Tool to look for overlapping UV's in static meshes.
+*/
+IMPLEMENT_COMPLEX_AUTOMATION_TEST(FStaticMeshUVCheck, "Tools.Static Mesh.Static Mesh UVs Check", (EAutomationTestFlags::ATF_Editor | EAutomationTestFlags::ATF_RequiresUser));
 
-bool FPerformUVTestCommand::Update()
+void FStaticMeshUVCheck::GetTests(TArray<FString>& OutBeautifiedNames, TArray<FString>& OutTestCommands) const
 {
-	return TestInfo->Update();
+	//This grabs each Static Mesh in the Game/Content
+	FEditorAutomationTestUtilities::CollectGameContentTestsByClass(UStaticMesh::StaticClass(), true, OutBeautifiedNames, OutTestCommands);
 }
 
-/**
- * StaticMeshUVsTest
- */
-IMPLEMENT_SIMPLE_AUTOMATION_TEST(FStaticMeshUVsTest, "Tools.Static Mesh.Static Mesh UVs Check", EAutomationTestFlags::ATF_Editor | EAutomationTestFlags::ATF_RequiresUser)
-
-/**
- * Find all static meshes and check the lightmap UVs
- */
-bool FStaticMeshUVsTest::RunTest(const FString& Parameters)
+bool FStaticMeshUVCheck::RunTest(const FString& Parameters)
 {
-	TSharedPtr<StaticMeshUVTest::FUVTestHelper> TestHelper = MakeShareable(new StaticMeshUVTest::FUVTestHelper(this));
-	ADD_LATENT_AUTOMATION_COMMAND(FPerformUVTestCommand(TestHelper));
+	UObject* Object = StaticLoadObject(UObject::StaticClass(), NULL, *Parameters);
+
+	// Missing UV messages
+	TArray<FString> MissingUVMessages;
+	// Bas UV messages
+	TArray<FString> BadUVMessages;
+	// Valid UV messages
+	TArray<FString> ValidUVMessages;
+
+	UStaticMesh::CheckLightMapUVs((UStaticMesh*)Object, MissingUVMessages, BadUVMessages, ValidUVMessages, true);
+
+	CollectGarbage(GARBAGE_COLLECTION_KEEPFLAGS);
+
 	return true;
 }
 
+
 /**
 * Launches a map onto a specified device after making a change to it.
 */
@@ -2298,3 +2163,4 @@ bool FLaunchOnTest::RunTest(const FString& Parameters)
 
 	return true;
 }
+
