Merging 2501701

//depot/UE4-UT/...

to //depot/UE4/...

diff --git a/Engine/Source/Editor/UnrealEd/Private/NavigationBuildingNotification.cpp b/Engine/Source/Editor/UnrealEd/Private/NavigationBuildingNotification.cpp
--- a/Engine/Source/Editor/UnrealEd/Private/NavigationBuildingNotification.cpp
+++ b/Engine/Source/Editor/UnrealEd/Private/NavigationBuildingNotification.cpp
@@ -73,6 +73,9 @@ void FNavigationBuildingNotificationImpl::BuildFinished()
 
 	if (EEngine != NULL && (FEditorBuildUtils::IsBuildingNavigationFromUserRequest()))
 	{
+		// remove existing item, if any
+		ClearCompleteNotification();
+
 		FNotificationInfo Info( NSLOCTEXT("NavigationBuild", "NavigationBuildDoneMessage", "Navigation building completed.") );
 		Info.bFireAndForget = true;
 		Info.bUseThrobber = false;
