Merging 2502894

//depot/UE4-UT/...

to //depot/UE4/...

diff --git a/Engine/Source/Editor/UnrealEd/Private/EditorBuildUtils.cpp b/Engine/Source/Editor/UnrealEd/Private/EditorBuildUtils.cpp
--- a/Engine/Source/Editor/UnrealEd/Private/EditorBuildUtils.cpp
+++ b/Engine/Source/Editor/UnrealEd/Private/EditorBuildUtils.cpp
@@ -236,7 +236,10 @@ bool FEditorBuildUtils::EditorBuild( UWorld* InWorld, EBuildOptions::Type Id, co
 
 			GUnrealEd->Exec( InWorld, TEXT("MAP REBUILD") );
 
-			TriggerNavigationBuilder(InWorld, Id);
+			if (GetDefault<ULevelEditorMiscSettings>()->bNavigationAutoUpdate)
+			{
+				TriggerNavigationBuilder(InWorld, Id);
+			}
 
 			// No need to dirty the persient level if we're building BSP for a sub-level.
 			bDirtyPersistentLevel = false;
@@ -254,7 +257,10 @@ bool FEditorBuildUtils::EditorBuild( UWorld* InWorld, EBuildOptions::Type Id, co
 
 				GUnrealEd->Exec( InWorld, TEXT("MAP REBUILD ALLVISIBLE") );
 
-				TriggerNavigationBuilder(InWorld, Id);
+				if (GetDefault<ULevelEditorMiscSettings>()->bNavigationAutoUpdate)
+				{
+					TriggerNavigationBuilder(InWorld, Id);
+				}
 			}
 			break;
 		}
