Add a better dialog for whenever we need to output the log from an external process, and use it to display errors when GenerateProjectFiles fails.

diff --git a/Engine/Source/Editor/GameProjectGeneration/Private/GameProjectGenerationModule.cpp b/Engine/Source/Editor/GameProjectGeneration/Private/GameProjectGenerationModule.cpp
--- a/Engine/Source/Editor/GameProjectGeneration/Private/GameProjectGenerationModule.cpp
+++ b/Engine/Source/Editor/GameProjectGeneration/Private/GameProjectGenerationModule.cpp
@@ -84,12 +84,12 @@ bool FGameProjectGenerationModule::UpdateGameProject(const FString& ProjectFile,
 }
 
 
-bool FGameProjectGenerationModule::UpdateCodeProject(FText& OutFailReason)
+bool FGameProjectGenerationModule::UpdateCodeProject(FText& OutFailReason, FText& OutFailLog)
 {
 	FScopedSlowTask SlowTask(0, LOCTEXT( "UpdatingCodeProject", "Updating code project..." ) );
 	SlowTask.MakeDialog();
 
-	return GameProjectUtils::GenerateCodeProjectFiles(FPaths::GetProjectFilePath(), OutFailReason);
+	return GameProjectUtils::GenerateCodeProjectFiles(FPaths::GetProjectFilePath(), OutFailReason, OutFailLog);
 }
 
 bool FGameProjectGenerationModule::GenerateBasicSourceCode(TArray<FString>& OutCreatedFiles, FText& OutFailReason)
diff --git a/Engine/Source/Editor/GameProjectGeneration/Private/GameProjectUtils.cpp b/Engine/Source/Editor/GameProjectGeneration/Private/GameProjectUtils.cpp
--- a/Engine/Source/Editor/GameProjectGeneration/Private/GameProjectUtils.cpp
+++ b/Engine/Source/Editor/GameProjectGeneration/Private/GameProjectUtils.cpp
@@ -596,7 +596,7 @@ void GameProjectUtils::GetStarterContentFiles(TArray<FString>& OutFilenames)
 	IFileManager::Get().FindFilesRecursive(OutFilenames, *SrcFolder, TEXT("*.upack"), /*Files=*/true, /*Directories=*/false);
 }
 
-bool GameProjectUtils::CreateProject(const FProjectInformation& InProjectInfo, FText& OutFailReason)
+bool GameProjectUtils::CreateProject(const FProjectInformation& InProjectInfo, FText& OutFailReason, FText& OutFailLog)
 {
 	if ( !IsValidProjectFileForCreation(InProjectInfo.ProjectFilename, OutFailReason) )
 	{
@@ -610,12 +610,12 @@ bool GameProjectUtils::CreateProject(const FProjectInformation& InProjectInfo, F
 	FString TemplateName;
 	if ( InProjectInfo.TemplateFile.IsEmpty() )
 	{
-		bProjectCreationSuccessful = GenerateProjectFromScratch(InProjectInfo, OutFailReason);
+		bProjectCreationSuccessful = GenerateProjectFromScratch(InProjectInfo, OutFailReason, OutFailLog);
 		TemplateName = InProjectInfo.bShouldGenerateCode ? TEXT("Basic Code") : TEXT("Blank");
 	}
 	else
 	{
-		bProjectCreationSuccessful = CreateProjectFromTemplate(InProjectInfo, OutFailReason);
+		bProjectCreationSuccessful = CreateProjectFromTemplate(InProjectInfo, OutFailReason, OutFailLog);
 		TemplateName = FPaths::GetBaseFilename(InProjectInfo.TemplateFile);
 	}
 
@@ -1005,7 +1005,7 @@ UTemplateProjectDefs* GameProjectUtils::LoadTemplateDefs(const FString& ProjectD
 	return TemplateDefs;
 }
 
-bool GameProjectUtils::GenerateProjectFromScratch(const FProjectInformation& InProjectInfo, FText& OutFailReason)
+bool GameProjectUtils::GenerateProjectFromScratch(const FProjectInformation& InProjectInfo, FText& OutFailReason, FText& OutFailLog)
 {
 	FScopedSlowTask SlowTask(5);
 
@@ -1094,7 +1094,7 @@ bool GameProjectUtils::GenerateProjectFromScratch(const FProjectInformation& InP
 	if ( InProjectInfo.bShouldGenerateCode )
 	{
 		// Generate project files
-		if ( !GenerateCodeProjectFiles(InProjectInfo.ProjectFilename, OutFailReason) )
+		if ( !GenerateCodeProjectFiles(InProjectInfo.ProjectFilename, OutFailReason, OutFailLog) )
 		{
 			DeleteGeneratedProjectFiles(InProjectInfo.ProjectFilename);
 			DeleteCreatedFiles(NewProjectFolder, CreatedFiles);
@@ -1108,7 +1108,7 @@ bool GameProjectUtils::GenerateProjectFromScratch(const FProjectInformation& InP
 	return true;
 }
 
-bool GameProjectUtils::CreateProjectFromTemplate(const FProjectInformation& InProjectInfo, FText& OutFailReason)
+bool GameProjectUtils::CreateProjectFromTemplate(const FProjectInformation& InProjectInfo, FText& OutFailReason, FText& OutFailLog)
 {
 	FScopedSlowTask SlowTask(10);
 
@@ -1511,7 +1511,7 @@ bool GameProjectUtils::CreateProjectFromTemplate(const FProjectInformation& InPr
 	if ( InProjectInfo.bShouldGenerateCode )
 	{
 		// Generate project files
-		if ( !GenerateCodeProjectFiles(InProjectInfo.ProjectFilename, OutFailReason) )
+		if ( !GenerateCodeProjectFiles(InProjectInfo.ProjectFilename, OutFailReason, OutFailLog) )
 		{
 			DeleteGeneratedProjectFiles(InProjectInfo.ProjectFilename);
 			DeleteCreatedFiles(DestFolder, CreatedFiles);
@@ -1954,7 +1954,7 @@ bool GameProjectUtils::BuildCodeProject(const FString& ProjectFilename)
 	return bCompileSucceeded;
 }
 
-bool GameProjectUtils::GenerateCodeProjectFiles(const FString& ProjectFilename, FText& OutFailReason)
+bool GameProjectUtils::GenerateCodeProjectFiles(const FString& ProjectFilename, FText& OutFailReason, FText& OutFailLog)
 {
 	FStringOutputDevice OutputLog;
 	OutputLog.SetAutoEmitLineTerminator(true);
@@ -1964,9 +1964,8 @@ bool GameProjectUtils::GenerateCodeProjectFiles(const FString& ProjectFilename,
 
 	if(!bHaveProjectFiles)
 	{
-		FFormatNamedArguments Args;
-		Args.Add( TEXT("LogOutput"), FText::FromString(OutputLog) );
-		OutFailReason = FText::Format(LOCTEXT("CouldNotGenerateProjectFiles", "Failed to generate project files. Log output:\n{LogOutput}"), Args);
+		OutFailReason = LOCTEXT("ErrorWhileGeneratingProjectFiles", "An error occurred while trying to generate project files.");
+		OutFailLog = FText::FromString(OutputLog);
 		return false;
 	}
 
diff --git a/Engine/Source/Editor/GameProjectGeneration/Private/GameProjectUtils.h b/Engine/Source/Editor/GameProjectGeneration/Private/GameProjectUtils.h
--- a/Engine/Source/Editor/GameProjectGeneration/Private/GameProjectUtils.h
+++ b/Engine/Source/Editor/GameProjectGeneration/Private/GameProjectUtils.h
@@ -73,7 +73,7 @@ public:
 	static bool OpenCodeIDE(const FString& ProjectFile, FText& OutFailReason);
 
 	/** Creates the specified project file and all required folders. If TemplateFile is non-empty, it will be used as the template for creation. On failure, OutFailReason will be populated. */
-	static bool CreateProject(const FProjectInformation& InProjectInfo, FText& OutFailReason);
+	static bool CreateProject(const FProjectInformation& InProjectInfo, FText& OutFailReason, FText& OutFailLog);
 
 	/** Prompts the user to update his project file, if necessary. */
 	static void CheckForOutOfDateGameProjectFile();
@@ -129,8 +129,8 @@ public:
 	/** Compiles a project while showing a progress bar, and offers to open the IDE if it fails. */
 	static bool BuildCodeProject(const FString& ProjectFilename);
 
-	/** Creates code project files for a new game project. On failure, OutFailReason will be populated. */
-	static bool GenerateCodeProjectFiles(const FString& ProjectFilename, FText& OutFailReason);
+	/** Creates code project files for a new game project. On failure, OutFailReason and OutFailLog will be populated. */
+	static bool GenerateCodeProjectFiles(const FString& ProjectFilename, FText& OutFailReason, FText& OutFailLog);
 
 	/** Returns true if there are starter content files available for instancing into new projects. */
 	static bool IsStarterContentAvailableForNewProjects();
@@ -207,10 +207,10 @@ private:
 	static FString GetHardwareConfigString(const FProjectInformation& InProjectInfo);
 
 	/** Generates a new project without using a template project */
-	static bool GenerateProjectFromScratch(const FProjectInformation& InProjectInfo, FText& OutFailReason);
+	static bool GenerateProjectFromScratch(const FProjectInformation& InProjectInfo, FText& OutFailReason, FText& OutFailLog);
 
 	/** Generates a new project using a template project */
-	static bool CreateProjectFromTemplate(const FProjectInformation& InProjectInfo, FText& OutFailReason);
+	static bool CreateProjectFromTemplate(const FProjectInformation& InProjectInfo, FText& OutFailReason, FText& OutFailLog);
 
 	/** Sets the engine association for a new project. Handles foreign and non-foreign projects. */
 	static bool SetEngineAssociationForForeignProject(const FString& ProjectFileName, FText& OutFailReason);
diff --git a/Engine/Source/Editor/GameProjectGeneration/Private/SNewProjectWizard.cpp b/Engine/Source/Editor/GameProjectGeneration/Private/SNewProjectWizard.cpp
--- a/Engine/Source/Editor/GameProjectGeneration/Private/SNewProjectWizard.cpp
+++ b/Engine/Source/Editor/GameProjectGeneration/Private/SNewProjectWizard.cpp
@@ -14,6 +14,7 @@
 #include "Editor/Documentation/Public/IDocumentation.h"
 #include "BreakIterator.h"
 #include "SHyperlink.h"
+#include "SOutputLogDialog.h"
 
 #define LOCTEXT_NAMESPACE "NewProjectWizard"
 
@@ -1437,14 +1438,14 @@ bool SNewProjectWizard::CreateProject( const FString& ProjectFile )
 		return false;
 	}
 
-	FText FailReason;
+	FText FailReason, FailLog;
 
 	FProjectInformation ProjectInfo(ProjectFile, SelectedTemplate->bGenerateCode, bCopyStarterContent, SelectedTemplate->ProjectFile);
 	ProjectInfo.TargetedHardware = SelectedHardwareClassTarget;
 	ProjectInfo.DefaultGraphicsPerformance = SelectedGraphicsPreset;
-	if (!GameProjectUtils::CreateProject(ProjectInfo, FailReason))
+	if (!GameProjectUtils::CreateProject(ProjectInfo, FailReason, FailLog))
 	{
-		DisplayError(FailReason);
+		SOutputLogDialog::Open(LOCTEXT("CreateProject", "Create Project"), FailReason, FailLog, FText::GetEmpty());
 		return false;
 	}
 
diff --git a/Engine/Source/Editor/GameProjectGeneration/Private/Tests/GameProjectAutomationTests.cpp b/Engine/Source/Editor/GameProjectGeneration/Private/Tests/GameProjectAutomationTests.cpp
--- a/Engine/Source/Editor/GameProjectGeneration/Private/Tests/GameProjectAutomationTests.cpp
+++ b/Engine/Source/Editor/GameProjectGeneration/Private/Tests/GameProjectAutomationTests.cpp
@@ -116,8 +116,8 @@ bool FBuildPromotionNewProjectCreateTest::RunTest(const FString& Parameters)
 		}
 	}
 
-	FText FailReason;
-	if (GameProjectUtils::CreateProject(FProjectInformation(DesiredProjectFilename, false, true), FailReason))
+	FText FailReason, FailLog;
+	if (GameProjectUtils::CreateProject(FProjectInformation(DesiredProjectFilename, false, true), FailReason, FailLog))
 	{
 		UE_LOG(LogGameProjectGenerationTests, Display, TEXT("Generated a new project: %s"), *DesiredProjectFilename);
 		UE_LOG(LogGameProjectGenerationTests, Display, TEXT("Test successful!"));
@@ -125,7 +125,7 @@ bool FBuildPromotionNewProjectCreateTest::RunTest(const FString& Parameters)
 	}
 	else
 	{
-		UE_LOG(LogGameProjectGenerationTests, Error, TEXT("Could not generate new project: %s"), *FailReason.ToString());
+		UE_LOG(LogGameProjectGenerationTests, Error, TEXT("Could not generate new project: %s - %s"), *FailReason.ToString(), *FailLog.ToString());
 		UE_LOG(LogGameProjectGenerationTests, Display, TEXT("Test failed!"));
 	}
 
diff --git a/Engine/Source/Editor/GameProjectGeneration/Public/GameProjectGenerationModule.h b/Engine/Source/Editor/GameProjectGeneration/Public/GameProjectGenerationModule.h
--- a/Engine/Source/Editor/GameProjectGeneration/Public/GameProjectGenerationModule.h
+++ b/Engine/Source/Editor/GameProjectGeneration/Public/GameProjectGenerationModule.h
@@ -87,7 +87,7 @@ public:
 	virtual bool UpdateGameProject(const FString& ProjectFile, const FString& EngineIdentifier, FText& OutFailReason);
 
 	/** Updates the current code project */
-	virtual bool UpdateCodeProject(FText& OutFailReason);
+	virtual bool UpdateCodeProject(FText& OutFailReason, FText& OutFailLog);
 
 	/** Gets the current projects source file count */
 	virtual bool ProjectHasCodeFiles();
diff --git a/Engine/Source/Editor/MainFrame/Private/Frame/MainFrameActions.cpp b/Engine/Source/Editor/MainFrame/Private/Frame/MainFrameActions.cpp
--- a/Engine/Source/Editor/MainFrame/Private/Frame/MainFrameActions.cpp
+++ b/Engine/Source/Editor/MainFrame/Private/Frame/MainFrameActions.cpp
@@ -9,6 +9,7 @@
 #include "GenericCommands.h"
 #include "EngineBuildSettings.h"
 #include "SourceCodeNavigation.h"
+#include "SOutputLogDialog.h"
 
 
 #define LOCTEXT_NAMESPACE "MainFrameActions"
@@ -649,10 +650,10 @@ void FMainFrameActionCallbacks::RefreshCodeProject()
 		FSourceCodeNavigation::AccessOnCompilerNotFound().Broadcast();
 	}
 
-	FText FailReason;
-	if(!FGameProjectGenerationModule::Get().UpdateCodeProject(FailReason))
+	FText FailReason, FailLog;
+	if(!FGameProjectGenerationModule::Get().UpdateCodeProject(FailReason, FailLog))
 	{
-		FMessageDialog::Open(EAppMsgType::Ok, FailReason);
+		SOutputLogDialog::Open(LOCTEXT("RefreshProject", "Refresh Project"), FailReason, FailLog, FText::GetEmpty());
 	}
 }
 
diff --git a/Engine/Source/Editor/UnrealEd/Private/Dialogs/SOutputLogDialog.cpp b/Engine/Source/Editor/UnrealEd/Private/Dialogs/SOutputLogDialog.cpp
new file mode 100644
--- /dev/null
+++ b/Engine/Source/Editor/UnrealEd/Private/Dialogs/SOutputLogDialog.cpp
@@ -0,0 +1,169 @@
+// Copyright 1998-2015 Epic Games, Inc. All Rights Reserved.
+
+#include "UnrealEd.h"
+#include "SOutputLogDialog.h"
+#include "SHyperlink.h"
+
+void SOutputLogDialog::Open( const FText& InTitle, const FText& InHeader, const FText& InLog, const FText& InFooter )
+{
+	TSharedRef<SWindow> ModalWindow = SNew(SWindow)
+		.Title( InTitle )
+		.SizingRule(ESizingRule::Autosized)
+		.AutoCenter(EAutoCenter::PreferredWorkArea)
+		.SupportsMinimize(false) 
+		.SupportsMaximize(false);
+
+	TSharedRef<SOutputLogDialog> MessageBox = SNew(SOutputLogDialog)
+		.ParentWindow(ModalWindow)
+		.Header( InHeader )
+		.Log( InLog )
+		.Footer( InFooter );
+
+	ModalWindow->SetContent( MessageBox );
+
+	GEditor->EditorAddModalWindow(ModalWindow);
+}
+
+BEGIN_SLATE_FUNCTION_BUILD_OPTIMIZATION
+void SOutputLogDialog::Construct( const FArguments& InArgs )
+{
+	ParentWindow = InArgs._ParentWindow.Get();
+	ParentWindow->SetWidgetToFocusOnActivate(SharedThis(this));
+
+	FSlateFontInfo MessageFont( FEditorStyle::GetFontStyle("StandardDialog.LargeFont"));
+	Header = InArgs._Header;
+	Log = InArgs._Log;
+	Footer = InArgs._Footer;
+
+	MaxWidth = FSlateApplication::Get().GetPreferredWorkArea().GetSize().X * 0.8f;
+
+	TSharedPtr<SUniformGridPanel> ButtonBox;
+
+	this->ChildSlot
+		[	
+			SNew(SBorder)
+				.BorderImage(FEditorStyle::GetBrush("ToolPanel.GroupBorder"))
+				[
+					SNew(SVerticalBox)
+
+					+SVerticalBox::Slot()
+						.HAlign(HAlign_Fill)
+						.AutoHeight()
+						.Padding(12.0f)
+						[
+							SNew(STextBlock)
+								.Text(Header)
+								.Visibility(Header.Get().IsEmptyOrWhitespace()? EVisibility::Hidden : EVisibility::Visible)
+								.Font(MessageFont)
+								.WrapTextAt(MaxWidth - 50.0f)
+						]
+
+					+SVerticalBox::Slot()
+						.HAlign(HAlign_Fill)
+						.VAlign(VAlign_Fill)
+						.FillHeight(1.0f)
+						.MaxHeight(550)
+						.Padding(12.0f, 0.0f, 12.0f, 12.0f)
+						[
+							SNew(SMultiLineEditableTextBox)
+								.Style(FEditorStyle::Get(), "Log.TextBox")
+								.TextStyle(FEditorStyle::Get(), "Log.Normal")
+								.ForegroundColor(FLinearColor::Gray)
+								.Text(FText::TrimTrailing(Log.Get()))
+								.IsReadOnly(true)
+								.AlwaysShowScrollbars(true)
+						]
+
+					+SVerticalBox::Slot()
+						.HAlign(HAlign_Fill)
+						.AutoHeight()
+						.Padding(12.0f, 0.0f, 12.0f, Footer.Get().IsEmptyOrWhitespace()? 0.0f : 12.0f)
+						[
+							SNew(STextBlock)
+								.Text(Footer)
+								.Visibility(Footer.Get().IsEmptyOrWhitespace()? EVisibility::Hidden : EVisibility::Visible)
+								.Font(MessageFont)
+								.WrapTextAt(MaxWidth - 50.0f)
+						]
+
+					+SVerticalBox::Slot()
+						.AutoHeight()
+						.Padding(12.0f, 0.0f, 12.0f, 12.0f)
+						[
+							SNew(SHorizontalBox)
+
+							+SHorizontalBox::Slot()
+								.FillWidth(1.0f)
+								.HAlign(HAlign_Left)
+								.VAlign(VAlign_Center)
+								.Padding(0.0f)
+								[
+									SNew(SHyperlink)
+										.OnNavigate(this, &SOutputLogDialog::HandleCopyMessageHyperlinkNavigate)
+										.Text( NSLOCTEXT("SOutputLogDialog", "CopyMessageHyperlink", "Copy Message") )
+										.ToolTipText( NSLOCTEXT("SOutputLogDialog", "CopyMessageTooltip", "Copy the text in this message to the clipboard (CTRL+C)") )
+								]
+
+							+SHorizontalBox::Slot()
+								.AutoWidth()
+								.HAlign(HAlign_Right)
+								.VAlign(VAlign_Center)
+								.Padding(0.f)
+								[
+									SNew( SButton )
+										.Text( NSLOCTEXT("SOutputLogDialog", "Ok", "Ok") )
+										.OnClicked( this, &SOutputLogDialog::HandleButtonClicked )
+										.ContentPadding(FEditorStyle::GetMargin("StandardDialog.ContentPadding"))
+										.HAlign(HAlign_Center)
+								]
+						]
+				]
+		];
+}
+END_SLATE_FUNCTION_BUILD_OPTIMIZATION
+
+FReply SOutputLogDialog::OnKeyDown( const FGeometry& MyGeometry, const FKeyEvent& InKeyEvent )
+{
+	if( InKeyEvent.GetKey() == EKeys::Escape )
+	{
+		return HandleButtonClicked();
+	}
+	else if (InKeyEvent.GetKey() == EKeys::C && InKeyEvent.IsControlDown())
+	{
+		CopyMessageToClipboard();
+		return FReply::Handled();
+	}
+	else
+	{
+		return FReply::Unhandled();
+	}
+}
+
+bool SOutputLogDialog::SupportsKeyboardFocus() const
+{
+	return true;
+}
+
+FVector2D SOutputLogDialog::ComputeDesiredSize(float LayoutScaleMultiplier) const
+{
+	FVector2D DesiredSize = SCompoundWidget::ComputeDesiredSize(LayoutScaleMultiplier);
+	DesiredSize.X = FMath::Min(DesiredSize.X, MaxWidth);
+	return DesiredSize;
+}
+
+void SOutputLogDialog::CopyMessageToClipboard( )
+{
+	FString FullMessage = FString::Printf(TEXT("%s") LINE_TERMINATOR LINE_TERMINATOR TEXT("%s") LINE_TERMINATOR LINE_TERMINATOR TEXT("%s"), *Header.Get().ToString(), *Log.Get().ToString(), *Footer.Get().ToString()).Trim();
+	FPlatformMisc::ClipboardCopy( *FullMessage );
+}
+
+FReply SOutputLogDialog::HandleButtonClicked( )
+{
+	ParentWindow->RequestDestroyWindow();
+	return FReply::Handled();
+}
+
+void SOutputLogDialog::HandleCopyMessageHyperlinkNavigate( )
+{
+	CopyMessageToClipboard();
+}
diff --git a/Engine/Source/Editor/UnrealEd/Private/Factories/PackFactory.cpp b/Engine/Source/Editor/UnrealEd/Private/Factories/PackFactory.cpp
--- a/Engine/Source/Editor/UnrealEd/Private/Factories/PackFactory.cpp
+++ b/Engine/Source/Editor/UnrealEd/Private/Factories/PackFactory.cpp
@@ -15,6 +15,7 @@
 #include "GameProjectGenerationModule.h"
 #include "Factories/PackFactory.h"
 #include "GameFramework/InputSettings.h"
+#include "Dialogs/SOutputLogDialog.h"
 
 DEFINE_LOG_CATEGORY_STATIC(LogPackFactory, Log, All);
 
@@ -626,10 +627,10 @@ UObject* UPackFactory::FactoryCreateBinary
 			{
 				// Update the game projects before we attempt to build
 				FGameProjectGenerationModule& GameProjectModule = FModuleManager::LoadModuleChecked<FGameProjectGenerationModule>(TEXT("GameProjectGeneration"));
-				FText FailReason;
-				if (!GameProjectModule.UpdateCodeProject(FailReason))
+				FText FailReason, FailLog;
+				if (!GameProjectModule.UpdateCodeProject(FailReason, FailLog))
 				{
-					FMessageDialog::Open(EAppMsgType::Ok, FailReason);
+					SOutputLogDialog::Open(NSLOCTEXT("PackFactory", "CreateBinary", "Create binary"), FailReason, FailLog, FText::GetEmpty());
 				}
 
 				// Compile the new code, either using the in editor hot-reload (if an existing module), or as a brand new module (if no existing code)
diff --git a/Engine/Source/Editor/UnrealEd/Public/Dialogs/SOutputLogDialog.h b/Engine/Source/Editor/UnrealEd/Public/Dialogs/SOutputLogDialog.h
new file mode 100644
--- /dev/null
+++ b/Engine/Source/Editor/UnrealEd/Public/Dialogs/SOutputLogDialog.h
@@ -0,0 +1,46 @@
+// Copyright 1998-2015 Epic Games, Inc. All Rights Reserved.
+
+#pragma once
+
+class SOutputLogDialog : public SCompoundWidget
+{
+public:
+	SLATE_BEGIN_ARGS( SOutputLogDialog )	{}
+		SLATE_ATTRIBUTE(TSharedPtr<SWindow>, ParentWindow)
+		SLATE_ATTRIBUTE(FText, Header)	
+		SLATE_ATTRIBUTE(FText, Log)	
+		SLATE_ATTRIBUTE(FText, Footer)	
+	SLATE_END_ARGS()
+
+	/** Displays the modal dialog box */
+	static UNREALED_API void Open(const FText& InTitle, const FText& InHeader, const FText& InLog, const FText& InFooter = FText::GetEmpty());
+
+	/** Constructs the dialog */
+	void Construct( const FArguments& InArgs );
+
+	/** Override behavior when a key is pressed */
+	virtual	FReply OnKeyDown( const FGeometry& MyGeometry, const FKeyEvent& InKeyEvent ) override;
+
+	/** Override the base method to allow for keyboard focus */
+	virtual bool SupportsKeyboardFocus() const override;
+
+	/** Compute the desired size for this dialog */
+	virtual FVector2D ComputeDesiredSize(float LayoutScaleMultiplier) const override;
+
+protected:
+	/** Copies the message text to the clipboard. */
+	void CopyMessageToClipboard( );
+
+private:
+	/** Handles clicking a message box button. */
+	FReply HandleButtonClicked( );
+
+	/** Handles clicking the 'Copy Message' hyper link. */
+	void HandleCopyMessageHyperlinkNavigate( );
+
+	TSharedPtr<SWindow> ParentWindow;
+	TAttribute<FText> Header;
+	TAttribute<FText> Log;
+	TAttribute<FText> Footer;
+	float MaxWidth;
+};
