Editor: Changed global Undo/Redo routing to automatically invalidate all editor viewport clients
- Added UEditorEngine::InvalidateAllViewportsAndHitProxies
- Removed manual Invalidate calls in the Blueprint and Font editors during Undo/Redo handling
- Removed a redundant InvalidateDisplay after an Invalidate call in the Font editor

diff --git a/Engine/Source/Editor/FontEditor/Private/FontEditor.cpp b/Engine/Source/Editor/FontEditor/Private/FontEditor.cpp
--- a/Engine/Source/Editor/FontEditor/Private/FontEditor.cpp
+++ b/Engine/Source/Editor/FontEditor/Private/FontEditor.cpp
@@ -380,7 +380,6 @@ void FFontEditor::PostUndo(bool bSuccess)
 	UpdateLayout();
 
 	CompositeFontEditor->Refresh();
-	FontPreviewWidget->RefreshViewport();
 }
 
 void FFontEditor::NotifyPostChange( const FPropertyChangedEvent& PropertyChangedEvent, class FEditPropertyChain* PropertyThatChanged)
diff --git a/Engine/Source/Editor/FontEditor/Private/SFontEditorViewport.cpp b/Engine/Source/Editor/FontEditor/Private/SFontEditorViewport.cpp
--- a/Engine/Source/Editor/FontEditor/Private/SFontEditorViewport.cpp
+++ b/Engine/Source/Editor/FontEditor/Private/SFontEditorViewport.cpp
@@ -578,7 +578,6 @@ void SFontEditorViewport::Construct(const FArguments& InArgs)
 void SFontEditorViewport::RefreshViewport()
 {
 	Viewport->Invalidate();
-	Viewport->InvalidateDisplay();
 }
 
 int32 SFontEditorViewport::GetCurrentSelectedPage() const
diff --git a/Engine/Source/Editor/Kismet/Private/BlueprintEditor.cpp b/Engine/Source/Editor/Kismet/Private/BlueprintEditor.cpp
--- a/Engine/Source/Editor/Kismet/Private/BlueprintEditor.cpp
+++ b/Engine/Source/Editor/Kismet/Private/BlueprintEditor.cpp
@@ -2728,19 +2728,13 @@ void FBlueprintEditor::PostUndo(bool bSuccess)
 		}
 
 		// Transaction affects the Blueprint this editor handles, so react as necessary
-		if(bAffectsBlueprint)
+		if (bAffectsBlueprint)
 		{
 			SetUISelectionState(NAME_None);
 
 			// Will cause a call to RefreshEditors()
 			FBlueprintEditorUtils::MarkBlueprintAsStructurallyModified(GetBlueprintObj());
 
-			TSharedPtr<class SSCSEditorViewport> ViewportPtr = GetSCSViewport();
-			if (ViewportPtr.IsValid())
-			{
-				ViewportPtr->Invalidate();
-			}
-
 			FSlateApplication::Get().DismissAllMenus();
 		}
 	}
@@ -2768,7 +2762,7 @@ void FBlueprintEditor::PostRedo(bool bSuccess)
 		}
 
 		// Transaction affects the Blueprint this editor handles, so react as necessary
-		if(bAffectsBlueprint)
+		if (bAffectsBlueprint)
 		{
 			// Will cause a call to RefreshEditors()
 			FBlueprintEditorUtils::MarkBlueprintAsStructurallyModified( BlueprintObj );
diff --git a/Engine/Source/Editor/UnrealEd/Classes/Editor/EditorEngine.h b/Engine/Source/Editor/UnrealEd/Classes/Editor/EditorEngine.h
--- a/Engine/Source/Editor/UnrealEd/Classes/Editor/EditorEngine.h
+++ b/Engine/Source/Editor/UnrealEd/Classes/Editor/EditorEngine.h
@@ -2757,8 +2757,13 @@ protected:
 	/** Called when Matinee is opened */
 	virtual void OnOpenMatinee(){};
 
-	/** Invalidate all viewport client hit proxies immediately */
-	void InvalidateAllViewportClientHitProxies();
+	/**
+	 * Invalidates all editor viewports and hit proxies, used when global changes like Undo/Redo may have invalidated state everywhere
+	 */
+	void InvalidateAllViewportsAndHitProxies();
+
+	/** Invalidate all level editor viewport client hit proxies immediately */
+	void InvalidateAllLevelEditorViewportClientHitProxies();
 
 	/** Destroy any online subsystems generated by PIE */
 	void CleanupPIEOnlineSessions(TArray<FName> OnlineIdentifiers);
diff --git a/Engine/Source/Editor/UnrealEd/Private/Editor.cpp b/Engine/Source/Editor/UnrealEd/Private/Editor.cpp
--- a/Engine/Source/Editor/UnrealEd/Private/Editor.cpp
+++ b/Engine/Source/Editor/UnrealEd/Private/Editor.cpp
@@ -1821,7 +1821,15 @@ bool UEditorEngine::UpdateSingleViewportClient(FEditorViewportClient* InViewport
 	return bUpdatedNonRealtimeViewport;
 }
 
-void UEditorEngine::InvalidateAllViewportClientHitProxies()
+void UEditorEngine::InvalidateAllViewportsAndHitProxies()
+{
+	for (FEditorViewportClient* ViewportClient : AllViewportClients)
+	{
+		ViewportClient->Invalidate();
+	}
+}
+
+void UEditorEngine::InvalidateAllLevelEditorViewportClientHitProxies()
 {
 	for (const auto* LevelViewportClient : LevelViewportClients)
 	{
diff --git a/Engine/Source/Editor/UnrealEd/Private/EditorServer.cpp b/Engine/Source/Editor/UnrealEd/Private/EditorServer.cpp
--- a/Engine/Source/Editor/UnrealEd/Private/EditorServer.cpp
+++ b/Engine/Source/Editor/UnrealEd/Private/EditorServer.cpp
@@ -1075,6 +1075,7 @@ void UEditorEngine::HandleTransactorRedo( FUndoSessionContext SessionContext, bo
 	PostUndo(Succeeded);
 
 	BroadcastPostRedo(SessionContext.Context, SessionContext.PrimaryObject, Succeeded);
+	InvalidateAllViewportsAndHitProxies();
 	ShowUndoRedoNotification(FText::Format(NSLOCTEXT("UnrealEd", "RedoMessageFormat", "Redo: {0}"), SessionContext.Title), Succeeded);
 }
 
@@ -1084,6 +1085,7 @@ void UEditorEngine::HandleTransactorUndo( FUndoSessionContext SessionContext, bo
 	PostUndo(Succeeded);
 
 	BroadcastPostUndo(SessionContext.Context, SessionContext.PrimaryObject, Succeeded);
+	InvalidateAllViewportsAndHitProxies();
 	ShowUndoRedoNotification(FText::Format(NSLOCTEXT("UnrealEd", "UndoMessageFormat", "Undo: {0}"), SessionContext.Title), Succeeded);
 }
 
@@ -1125,7 +1127,7 @@ UTransactor* UEditorEngine::CreateTrans()
 	return TransBuffer;
 }
 
-void UEditorEngine::PostUndo (bool bSuccess)
+void UEditorEngine::PostUndo(bool bSuccess)
 {
 	//Update the actor selection followed by the component selection if needed (note: order is important)
 		
@@ -2022,7 +2024,7 @@ UWorld* UEditorEngine::NewMap()
 	GUnrealEd->ResetTransaction( CleanseText );
 
 	// Invalidate all the level viewport hit proxies
-	InvalidateAllViewportClientHitProxies();
+	InvalidateAllLevelEditorViewportClientHitProxies();
 
 	return NewWorld;
 }
@@ -2472,7 +2474,7 @@ bool UEditorEngine::Map_Load(const TCHAR* Str, FOutputDevice& Ar)
 					}
 
 					// Invalidate all the level viewport hit proxies
-					InvalidateAllViewportClientHitProxies();
+					InvalidateAllLevelEditorViewportClientHitProxies();
 				}
 			}
 			else
@@ -4856,6 +4858,9 @@ void UEditorEngine::BroadcastPostRedo(const FString& Context, UObject* PrimaryOb
 			Client->PostRedo( bRedoSuccess );
 		}
 	}
+
+	// Invalidate all viewports
+	InvalidateAllViewportsAndHitProxies();
 }
 
 bool UEditorEngine::Exec_Particle(const TCHAR* Str, FOutputDevice& Ar)
