Editor: Changed FDragTool to take in a reference to FEditorModeTools* instead of referencing the global one

diff --git a/Engine/Source/Editor/UnrealEd/Private/DragTool_BoxSelect.cpp b/Engine/Source/Editor/UnrealEd/Private/DragTool_BoxSelect.cpp
--- a/Engine/Source/Editor/UnrealEd/Private/DragTool_BoxSelect.cpp
+++ b/Engine/Source/Editor/UnrealEd/Private/DragTool_BoxSelect.cpp
@@ -126,8 +126,7 @@ void FDragTool_ActorBoxSelect::AddDelta( const FVector& InDelta )
 */
 void FDragTool_ActorBoxSelect::EndDrag()
 {
-	FEditorModeTools& EdModeTools = GLevelEditorModeTools();
-	const bool bGeometryMode = EdModeTools.IsModeActive( FBuiltinEditorModes::EM_Geometry );
+	const bool bGeometryMode = ModeTools->IsModeActive(FBuiltinEditorModes::EM_Geometry);
 	
 	FScopedTransaction Transaction( NSLOCTEXT("ActorFrustumSelect", "MarqueeSelectTransation", "Marquee Select" ) );
 
@@ -143,11 +142,11 @@ void FDragTool_ActorBoxSelect::EndDrag()
 	else if( !bShiftDown )
 	{
 		// If the user is selecting, but isn't hold down SHIFT, remove all current selections.
-		GLevelEditorModeTools().SelectNone();
+		ModeTools->SelectNone();
 	}
 
 	// Let the editor mode try to handle the box selection.
-	const bool bEditorModeHandledBoxSelection = GLevelEditorModeTools().BoxSelect( SelBBox, bLeftMouseButtonDown );
+	const bool bEditorModeHandledBoxSelection = ModeTools->BoxSelect(SelBBox, bLeftMouseButtonDown);
 
 	// If the edit mode didn't handle the selection, try normal actor box selection.
 	if ( !bEditorModeHandledBoxSelection )
@@ -287,8 +286,7 @@ bool FDragTool_ActorBoxSelect::IntersectsBox( AActor& InActor, const FBox& InBox
 {
 	bool bActorHitByBox = false;
 
-	FEditorModeTools& EdModeTools = GLevelEditorModeTools();
-	const bool bGeometryMode = EdModeTools.IsModeActive( FBuiltinEditorModes::EM_Geometry );
+	const bool bGeometryMode = ModeTools->IsModeActive(FBuiltinEditorModes::EM_Geometry);
 
 	// Check for special cases (like certain show flags that might hide an actor)
 	bool bActorIsHiddenByShowFlags = false;
diff --git a/Engine/Source/Editor/UnrealEd/Private/DragTool_FrustumSelect.cpp b/Engine/Source/Editor/UnrealEd/Private/DragTool_FrustumSelect.cpp
--- a/Engine/Source/Editor/UnrealEd/Private/DragTool_FrustumSelect.cpp
+++ b/Engine/Source/Editor/UnrealEd/Private/DragTool_FrustumSelect.cpp
@@ -20,12 +20,6 @@
 //
 ///////////////////////////////////////////////////////////////////////////////
 
-/** 
- * Updates the drag tool's end location with the specified delta.  The end location is
- * snapped to the editor constraints if bUseSnapping is true.
- *
- * @param	InDelta		A delta of mouse movement.
- */
 void FDragTool_ActorFrustumSelect::AddDelta( const FVector& InDelta )
 {
 	FIntPoint MousePos;
@@ -37,7 +31,6 @@ void FDragTool_ActorFrustumSelect::AddDelta( const FVector& InDelta )
 	const bool bUseHoverFeedback = GEditor != NULL && GetDefault<ULevelEditorViewportSettings>()->bEnableViewportHoverFeedback;
 }
 
-
 void FDragTool_ActorFrustumSelect::StartDrag(FEditorViewportClient* InViewportClient, const FVector& InStart, const FVector2D& InStartScreen)
 {
 	FDragTool::StartDrag(InViewportClient, InStart, InStartScreen);
@@ -55,13 +48,9 @@ void FDragTool_ActorFrustumSelect::StartDrag(FEditorViewportClient* InViewportCl
 
 }
 
-/**
-* Ends a mouse drag behavior (the user has let go of the mouse button).
-*/
 void FDragTool_ActorFrustumSelect::EndDrag()
 {
-	FEditorModeTools& EdModeTools = GLevelEditorModeTools();
-	const bool bGeometryMode = EdModeTools.IsModeActive( FBuiltinEditorModes::EM_Geometry );
+	const bool bGeometryMode = ModeTools->IsModeActive(FBuiltinEditorModes::EM_Geometry);
 
 	FSceneViewFamilyContext ViewFamily(FSceneViewFamily::ConstructionValues(LevelViewportClient->Viewport, LevelViewportClient->GetScene(), LevelViewportClient->EngineShowFlags ));
 	FSceneView* SceneView = LevelViewportClient->CalcSceneView(&ViewFamily);
@@ -78,7 +67,7 @@ void FDragTool_ActorFrustumSelect::EndDrag()
 	if( !bShiftDown )
 	{
 		// If the user is selecting, but isn't hold down SHIFT, remove all current selections.
-		GLevelEditorModeTools().SelectNone();
+		ModeTools->SelectNone();
 	}
 
 	// Does an actor have to be fully contained in the box to be selected
@@ -86,7 +75,7 @@ void FDragTool_ActorFrustumSelect::EndDrag()
 	bool bSelectionChanged = false;
 
 	// Let the editor mode try to handle the selection.
-	const bool bEditorModeHandledSelection = GLevelEditorModeTools().FrustumSelect( Frustum, bLeftMouseButtonDown );
+	const bool bEditorModeHandledSelection = ModeTools->FrustumSelect(Frustum, bLeftMouseButtonDown);
 
 	if( !bEditorModeHandledSelection )
 	{
@@ -238,8 +227,7 @@ bool FDragTool_ActorFrustumSelect::IntersectsFrustum( AActor& InActor, const FCo
 {	
 	bool bActorHitByBox = false;
 
-	FEditorModeTools& EdModeTools = GLevelEditorModeTools();
-	const bool bGeometryMode = EdModeTools.IsModeActive(FBuiltinEditorModes::EM_Geometry);
+	const bool bGeometryMode = ModeTools->IsModeActive(FBuiltinEditorModes::EM_Geometry);
 
 	// Check for special cases (like certain show flags that might hide an actor)
 	bool bActorIsHiddenByShowFlags = false;
@@ -275,14 +263,6 @@ bool FDragTool_ActorFrustumSelect::IntersectsFrustum( AActor& InActor, const FCo
 	return bActorHitByBox;
 }
 
-/** 
- * Returns true if the provided BSP node intersects with the provided frustum 
- *
- * @param InModel				The model containing BSP nodes to check
- * @param NodeIndex				The index to a BSP node in the model.  This node is used for the bounds check.
- * @param InFrustum				The frustum to check against.
- * @param bUseStrictSelection	true if the node must be entirely within the frustum
- */
 bool FDragTool_ActorFrustumSelect::IntersectsFrustum( const UModel& InModel, int32 NodeIndex, const FConvexVolume& InFrustum, bool bUseStrictSelection ) const
 {
 	FBox NodeBB;
@@ -297,12 +277,6 @@ bool FDragTool_ActorFrustumSelect::IntersectsFrustum( const UModel& InModel, int
 	return bIntersects && (!bUseStrictSelection || (bUseStrictSelection && bFullyContained));
 }
 
-/** 
- * Calculates a frustum to check actors against 
- * 
- * @param OutFrustum		The created frustum
- * @param bUseBoxFrustum	If true a frustum out of the current dragged box will be created.  false will use the view frustum.
- */
 void FDragTool_ActorFrustumSelect::CalculateFrustum( FSceneView* View, FConvexVolume& OutFrustum, bool bUseBoxFrustum )
 {
 	if( bUseBoxFrustum )
@@ -350,7 +324,6 @@ void FDragTool_ActorFrustumSelect::CalculateFrustum( FSceneView* View, FConvexVo
 	}
 }
 
-/** Adds a hover effect to the passed in actor */
 void FDragTool_ActorFrustumSelect::AddHoverEffect( AActor& InActor )
 {
 	FViewportHoverTarget HoverTarget( &InActor );
@@ -358,7 +331,6 @@ void FDragTool_ActorFrustumSelect::AddHoverEffect( AActor& InActor )
 	FLevelEditorViewportClient::HoveredObjects.Add( HoverTarget );
 }
 
-/** Removes a hover effect from the passed in actor */
 void FDragTool_ActorFrustumSelect::RemoveHoverEffect( AActor& InActor  )
 {
 	FViewportHoverTarget HoverTarget( &InActor );
@@ -370,7 +342,6 @@ void FDragTool_ActorFrustumSelect::RemoveHoverEffect( AActor& InActor  )
 	}
 }
 
-/** Adds a hover effect to the passed in bsp surface */
 void FDragTool_ActorFrustumSelect::AddHoverEffect( UModel& InModel, int32 SurfIndex )
 {
 	FViewportHoverTarget HoverTarget( &InModel, SurfIndex );
@@ -378,7 +349,6 @@ void FDragTool_ActorFrustumSelect::AddHoverEffect( UModel& InModel, int32 SurfIn
 	FLevelEditorViewportClient::HoveredObjects.Add( HoverTarget );
 }
 
-/** Removes a hover effect from the passed in bsp surface */
 void FDragTool_ActorFrustumSelect::RemoveHoverEffect( UModel& InModel, int32 SurfIndex )
 {
 	FViewportHoverTarget HoverTarget( &InModel, SurfIndex );
diff --git a/Engine/Source/Editor/UnrealEd/Private/DragTool_Measure.cpp b/Engine/Source/Editor/UnrealEd/Private/DragTool_Measure.cpp
--- a/Engine/Source/Editor/UnrealEd/Private/DragTool_Measure.cpp
+++ b/Engine/Source/Editor/UnrealEd/Private/DragTool_Measure.cpp
@@ -14,7 +14,8 @@
 ///////////////////////////////////////////////////////////////////////////////////////////////////////////////////////
 
 FDragTool_Measure::FDragTool_Measure(FEditorViewportClient* InViewportClient)
-	: ViewportClient(InViewportClient)
+	: FDragTool(InViewportClient->GetModeTools())
+	, ViewportClient(InViewportClient)
 {
 	bUseSnapping = true;
 	bConvertDelta = false;
diff --git a/Engine/Source/Editor/UnrealEd/Private/EditorDragTools.cpp b/Engine/Source/Editor/UnrealEd/Private/EditorDragTools.cpp
--- a/Engine/Source/Editor/UnrealEd/Private/EditorDragTools.cpp
+++ b/Engine/Source/Editor/UnrealEd/Private/EditorDragTools.cpp
@@ -1,6 +1,5 @@
 // Copyright 1998-2015 Epic Games, Inc. All Rights Reserved.
 
-
 #include "UnrealEd.h"
 #include "EditorDragTools.h"
 #include "SnappingUtils.h"
@@ -11,22 +10,17 @@
 //
 ///////////////////////////////////////////////////////////////////////////////////////////////////////////////////////
 
-FDragTool::FDragTool()
-	:	bConvertDelta( true )
-	,	Start( FVector::ZeroVector )
-	,	End( FVector::ZeroVector )
-	,	EndWk( FVector::ZeroVector )
-	,	bUseSnapping( false )
-	,   bIsDragging( false )
+FDragTool::FDragTool(FEditorModeTools* InModeTools)
+	: bConvertDelta( true )
+	, ModeTools( InModeTools )
+	, Start( FVector::ZeroVector )
+	, End( FVector::ZeroVector )
+	, EndWk( FVector::ZeroVector )
+	, bUseSnapping( false )
+	, bIsDragging( false )
 {
 }
 
-/**
- * Updates the drag tool's end location with the specified delta.  The end location is
- * snapped to the editor constraints if bUseSnapping is true.
- *
- * @param	InDelta		A delta of mouse movement.
- */
 void FDragTool::AddDelta( const FVector& InDelta )
 {
 	EndWk += InDelta;
@@ -42,12 +36,6 @@ void FDragTool::AddDelta( const FVector& InDelta )
 	}
 }
 
-/**
- * Starts a mouse drag behavior.  The start location is snapped to the editor constraints if bUseSnapping is true.
- *
- * @param	InViewportClient	The viewport client in which the drag event occurred.
- * @param	InStart				Where the mouse was when the drag started.
- */
 void FDragTool::StartDrag(FEditorViewportClient* InViewportClient, const FVector& InStart, const FVector2D& StartScreen)
 {
 	Start = InStart;
@@ -71,9 +59,6 @@ void FDragTool::StartDrag(FEditorViewportClient* InViewportClient, const FVector
 	bMiddleMouseButtonDown = InViewportClient->Viewport->KeyState(EKeys::MiddleMouseButton);
 }
 
-/**
- * Ends a mouse drag behavior (the user has let go of the mouse button).
- */
 void FDragTool::EndDrag()
 {
 	Start = End = EndWk = FVector::ZeroVector;
diff --git a/Engine/Source/Editor/UnrealEd/Private/EditorViewportClient.cpp b/Engine/Source/Editor/UnrealEd/Private/EditorViewportClient.cpp
--- a/Engine/Source/Editor/UnrealEd/Private/EditorViewportClient.cpp
+++ b/Engine/Source/Editor/UnrealEd/Private/EditorViewportClient.cpp
@@ -3646,12 +3646,12 @@ bool FEditorViewportClient::ShouldPanOrDollyCamera() const
 
 TSharedPtr<FDragTool> FEditorViewportClient::MakeDragTool(EDragTool::Type)
 {
-	return MakeShareable( new FDragTool );
+	return MakeShareable( new FDragTool(GetModeTools()) );
 }
 
 bool FEditorViewportClient::CanUseDragTool() const
 {
-	return !ShouldOrbitCamera() && GetCurrentWidgetAxis() == EAxisList::None && (!ModeTools || ModeTools->AllowsViewportDragTool());
+	return !ShouldOrbitCamera() && (GetCurrentWidgetAxis() == EAxisList::None) && ((ModeTools == nullptr) || ModeTools->AllowsViewportDragTool());
 }
 
 bool FEditorViewportClient::ShouldOrbitCamera() const
diff --git a/Engine/Source/Editor/UnrealEd/Public/DragTool_BoxSelect.h b/Engine/Source/Editor/UnrealEd/Public/DragTool_BoxSelect.h
--- a/Engine/Source/Editor/UnrealEd/Public/DragTool_BoxSelect.h
+++ b/Engine/Source/Editor/UnrealEd/Public/DragTool_BoxSelect.h
@@ -13,7 +13,7 @@ class FDragTool_ActorBoxSelect : public FDragTool
 {
 public:
 	explicit FDragTool_ActorBoxSelect( FLevelEditorViewportClient* InLevelViewportClient )
-		: FDragTool()
+		: FDragTool(LevelViewportClient->GetModeTools())
 		, LevelViewportClient( InLevelViewportClient )
 	{}
 
diff --git a/Engine/Source/Editor/UnrealEd/Public/DragTool_FrustumSelect.h b/Engine/Source/Editor/UnrealEd/Public/DragTool_FrustumSelect.h
--- a/Engine/Source/Editor/UnrealEd/Public/DragTool_FrustumSelect.h
+++ b/Engine/Source/Editor/UnrealEd/Public/DragTool_FrustumSelect.h
@@ -13,7 +13,7 @@ class UNREALED_API FDragTool_ActorFrustumSelect : public FDragTool
 {
 public:
 	explicit FDragTool_ActorFrustumSelect(FLevelEditorViewportClient* InLevelViewportClient)
-		: FDragTool()
+		: FDragTool(InLevelViewportClient->GetModeTools())
 		, LevelViewportClient( InLevelViewportClient )
 	{}
 
diff --git a/Engine/Source/Editor/UnrealEd/Public/EditorDragTools.h b/Engine/Source/Editor/UnrealEd/Public/EditorDragTools.h
--- a/Engine/Source/Editor/UnrealEd/Public/EditorDragTools.h
+++ b/Engine/Source/Editor/UnrealEd/Public/EditorDragTools.h
@@ -13,7 +13,7 @@ class FEditorViewportClient;
 class FDragTool
 {
 public:
-	FDragTool();
+	FDragTool(FEditorModeTools* InModeTools);
 	virtual ~FDragTool() {}
 
 	/**
@@ -51,6 +51,8 @@ public:
 	bool bConvertDelta;
 
 protected:
+	FEditorModeTools* ModeTools;
+
 	/** The start/end location of the current drag. */
 	FVector Start, End, EndWk;
 
diff --git a/Engine/Source/Editor/UnrealEd/Public/EditorViewportClient.h b/Engine/Source/Editor/UnrealEd/Public/EditorViewportClient.h
--- a/Engine/Source/Editor/UnrealEd/Public/EditorViewportClient.h
+++ b/Engine/Source/Editor/UnrealEd/Public/EditorViewportClient.h
@@ -634,14 +634,14 @@ public:
 	virtual TSharedPtr<class FDragTool> MakeDragTool( EDragTool::Type DragToolType );
 
 	/** @return true if a drag tool can be used */
-	bool CanUseDragTool() const;
+	virtual bool CanUseDragTool() const;
 
 	/** @return Whether or not to orbit the camera */
-	virtual bool ShouldOrbitCamera() const ;
+	virtual bool ShouldOrbitCamera() const;
 
 	bool IsMovingCamera() const;
 
-	virtual void UpdateLinkedOrthoViewports( bool bInvalidate = false ) {};
+	virtual void UpdateLinkedOrthoViewports( bool bInvalidate = false ) {}
 
 	/**
 	 * @return true to lock the pitch of the viewport camera
