Editor: Fixed the use of the wrong variable in FDragTool_ActorBoxSelect's initializer list

diff --git a/Engine/Source/Editor/UnrealEd/Public/DragTool_BoxSelect.h b/Engine/Source/Editor/UnrealEd/Public/DragTool_BoxSelect.h
--- a/Engine/Source/Editor/UnrealEd/Public/DragTool_BoxSelect.h
+++ b/Engine/Source/Editor/UnrealEd/Public/DragTool_BoxSelect.h
@@ -13,7 +13,7 @@ class FDragTool_ActorBoxSelect : public FDragTool
 {
 public:
 	explicit FDragTool_ActorBoxSelect( FLevelEditorViewportClient* InLevelViewportClient )
-		: FDragTool(LevelViewportClient->GetModeTools())
+		: FDragTool(InLevelViewportClient->GetModeTools())
 		, LevelViewportClient( InLevelViewportClient )
 	{}
 
