Fix crash due to reentrancy of the slate tick loop if a slow task window and a modal window are siblings.  The slow task window could re-tick  the modal window while it was being ticked

diff --git a/Engine/Source/Editor/UnrealEd/Private/FeedbackContextEditor.cpp b/Engine/Source/Editor/UnrealEd/Private/FeedbackContextEditor.cpp
--- a/Engine/Source/Editor/UnrealEd/Private/FeedbackContextEditor.cpp
+++ b/Engine/Source/Editor/UnrealEd/Private/FeedbackContextEditor.cpp
@@ -312,13 +312,17 @@ const int32 SSlowTaskWidget::FixedPaddingH;;
 const int32 SSlowTaskWidget::MainBarHeight;
 const int32 SSlowTaskWidget::SecondaryBarHeight;
 
-static void TickSlate()
+static void TickSlate(TSharedPtr<SWindow> SlowTaskWindow)
 {
-	// Tick Slate application
-	FSlateApplication::Get().Tick();
+	// Avoid re-entrancy by ticking the active modal window again. This can happen if thhe slow task window is open and a sibling modal window is open as well.  We only tick slate if we are the active modal window or a child of the active modal window
+	if( SlowTaskWindow.IsValid() && ( FSlateApplication::Get().GetActiveModalWindow() == SlowTaskWindow || SlowTaskWindow->IsDescendantOf( FSlateApplication::Get().GetActiveModalWindow() ) ) )
+	{
+		// Tick Slate application
+		FSlateApplication::Get().Tick();
 
-	// Sync the game thread and the render thread. This is needed if many StatusUpdate are called
-	FSlateApplication::Get().GetRenderer()->Sync();
+		// Sync the game thread and the render thread. This is needed if many StatusUpdate are called
+		FSlateApplication::Get().GetRenderer()->Sync();
+	}
 }
 
 FFeedbackContextEditor::FFeedbackContextEditor()
@@ -404,7 +408,7 @@ void FFeedbackContextEditor::StartSlowTask( const FText& Task, bool bShowCancelB
 
 			SlowTaskWindowRef->ShowWindow();
 
-			TickSlate();
+			TickSlate(SlowTaskWindow.Pin());
 		}
 
 		FPlatformSplash::SetSplashText( SplashTextType::StartupProgress, *Task.ToString() );
@@ -450,11 +454,11 @@ void FFeedbackContextEditor::ProgressReported( const float TotalProgressInterp,
 			}
 
 			BuildProgressWidget->SetBuildProgressPercent(TotalProgressInterp * 100, 100);
-			TickSlate();
+			TickSlate(SlowTaskWindow.Pin());
 		}
 		else if (SlowTaskWindow.IsValid())
 		{
-			TickSlate();
+			TickSlate(SlowTaskWindow.Pin());
 		}
 	}
 	else
@@ -552,7 +556,7 @@ TWeakPtr<class SBuildProgressWidget> FFeedbackContextEditor::ShowBuildProgressWi
 	
 	if (FSlateApplication::Get().CanDisplayWindows())
 	{
-		TickSlate();
+		TickSlate(BuildProgressWindow.Pin());
 	}
 
 	return BuildProgressWidget;
