Source control dialogs are no longer popped up once a file which cannot be checked out is made writable.

diff --git a/Engine/Source/Editor/UnrealEd/Private/FileHelpers.cpp b/Engine/Source/Editor/UnrealEd/Private/FileHelpers.cpp
--- a/Engine/Source/Editor/UnrealEd/Private/FileHelpers.cpp
+++ b/Engine/Source/Editor/UnrealEd/Private/FileHelpers.cpp
@@ -1050,7 +1050,7 @@ static bool AddCheckoutPackageItems(bool bCheckDirty, TArray<UPackage*> Packages
 		FSourceControlStatePtr SourceControlState = SourceControlProvider.GetState(CurPackage, EStateCacheUsage::Use);
 
 		// Package does not need to be checked out if its already checked out or we are ignoring it for source control
-		bool bSCCCanEdit = !SourceControlState.IsValid() || SourceControlState->CanCheckIn() || SourceControlState->IsIgnored() || SourceControlState->IsUnknown();
+		bool bSCCCanEdit = !SourceControlState.IsValid() || SourceControlState->CanCheckIn() || SourceControlState->IsIgnored() || SourceControlState->IsUnknown() || (bCareAboutReadOnly && !bPkgReadOnly);
 		bool bIsSourceControlled = SourceControlState.IsValid() && SourceControlState->IsSourceControlled();
 
 		if (!bSCCCanEdit && (bIsSourceControlled && (!bCheckDirty || (bCheckDirty && CurPackage->IsDirty()))) && !SourceControlState->IsCheckedOut())
