[UE-11380] Reroute nodes now function correctly in an animation graph after compiling an animation blueprint.

change summary:
- added new UEdGraphNode::IsCompilerRelevant() and UEdGraphNode::GetPassThroughPin() APIs
- added a UK2Node::GetPassThroughPin() override to return the opposing exec pin on any given exec pin, else NULL, as base node logic
- added a UK2Node_Knot::IsCompilerRelevant() override to always return false for reroute nodes (since they are always precompiled out)
- added a UK2Node_Knot::GetPassThroughPin() override to always return the opposing pin given either the input or output pin, else NULL
- added new FBlueprintEditorUtils::GetCompilerRelevantNodes() and FBlueprintEditorUtils::FindFirstCompilerRelevantNodes() utility methods
- modified UAnimGraphNode_Base::GetLinkIDLocation() to look for the first linked "compiler-relevant" node rather than the literal first linked node

diff --git a/Engine/Source/Editor/AnimGraph/Private/AnimGraphNode_Base.cpp b/Engine/Source/Editor/AnimGraph/Private/AnimGraphNode_Base.cpp
--- a/Engine/Source/Editor/AnimGraph/Private/AnimGraphNode_Base.cpp
+++ b/Engine/Source/Editor/AnimGraph/Private/AnimGraphNode_Base.cpp
@@ -299,7 +299,7 @@ FPoseLinkMappingRecord UAnimGraphNode_Base::GetLinkIDLocation(const UScriptStruc
 {
 	if (SourcePin->LinkedTo.Num() > 0)
 	{
-		if (UAnimGraphNode_Base* LinkedNode = Cast<UAnimGraphNode_Base>(SourcePin->LinkedTo[0]->GetOwningNode()))
+		if (UAnimGraphNode_Base* LinkedNode = Cast<UAnimGraphNode_Base>(FBlueprintEditorUtils::FindFirstCompilerRelevantNode(SourcePin->LinkedTo[0])))
 		{
 			//@TODO: Name-based hackery, avoid the roundtrip and better indicate when it's an array pose pin
 			int32 UnderscoreIndex = SourcePin->PinName.Find(TEXT("_"), ESearchCase::CaseSensitive);
diff --git a/Engine/Source/Editor/BlueprintGraph/Classes/K2Node.h b/Engine/Source/Editor/BlueprintGraph/Classes/K2Node.h
--- a/Engine/Source/Editor/BlueprintGraph/Classes/K2Node.h
+++ b/Engine/Source/Editor/BlueprintGraph/Classes/K2Node.h
@@ -103,6 +103,7 @@ class UK2Node : public UEdGraphNode
 	BLUEPRINTGRAPH_API virtual void GetPinHoverText(const UEdGraphPin& Pin, FString& HoverTextOut) const override;
 	BLUEPRINTGRAPH_API virtual bool ShowPaletteIconOnNode() const override { return true; }
 	BLUEPRINTGRAPH_API virtual bool AllowSplitPins() const override;
+	BLUEPRINTGRAPH_API virtual UEdGraphPin* GetPassThroughPin(const UEdGraphPin* FromPin) const override;
 	// End of UEdGraphNode interface
 
 	// K2Node interface
diff --git a/Engine/Source/Editor/BlueprintGraph/Classes/K2Node_Knot.h b/Engine/Source/Editor/BlueprintGraph/Classes/K2Node_Knot.h
--- a/Engine/Source/Editor/BlueprintGraph/Classes/K2Node_Knot.h
+++ b/Engine/Source/Editor/BlueprintGraph/Classes/K2Node_Knot.h
@@ -20,6 +20,8 @@ public:
 	virtual void OnRenameNode(const FString& NewName) override;
 	virtual TSharedPtr<class INameValidatorInterface> MakeNameValidator() const override;
 	virtual bool AllowSplitPins() const override;
+	virtual bool IsCompilerRelevant() const override { return false; }
+	virtual UEdGraphPin* GetPassThroughPin(const UEdGraphPin* FromPin) const override;
 	// End of UEdGraphNode interface
 
 	// UK2Node interface
diff --git a/Engine/Source/Editor/BlueprintGraph/Private/K2Node.cpp b/Engine/Source/Editor/BlueprintGraph/Private/K2Node.cpp
--- a/Engine/Source/Editor/BlueprintGraph/Private/K2Node.cpp
+++ b/Engine/Source/Editor/BlueprintGraph/Private/K2Node.cpp
@@ -944,6 +944,28 @@ UEdGraphPin* UK2Node::GetExecPin() const
 	return Pin;
 }
 
+UEdGraphPin* UK2Node::GetPassThroughPin(const UEdGraphPin* FromPin) const
+{
+	UEdGraphPin* MatchedPin = nullptr;
+	if(FromPin && Pins.Contains(FromPin))
+	{
+		const UEdGraphSchema_K2* K2Schema = GetDefault<UEdGraphSchema_K2>();
+		if(K2Schema->IsExecPin(*FromPin))
+		{
+			if(FromPin->Direction == EGPD_Input)
+			{
+				MatchedPin = FindPin(K2Schema->PN_Then);
+			}
+			else
+			{
+				MatchedPin = FindPin(K2Schema->PN_Execute);
+			}
+		}
+	}
+
+	return MatchedPin;
+}
+
 bool UK2Node::CanCreateUnderSpecifiedSchema(const UEdGraphSchema* DesiredSchema) const
 {
 	return DesiredSchema->GetClass()->IsChildOf(UEdGraphSchema_K2::StaticClass());
diff --git a/Engine/Source/Editor/BlueprintGraph/Private/K2Node_Knot.cpp b/Engine/Source/Editor/BlueprintGraph/Private/K2Node_Knot.cpp
--- a/Engine/Source/Editor/BlueprintGraph/Private/K2Node_Knot.cpp
+++ b/Engine/Source/Editor/BlueprintGraph/Private/K2Node_Knot.cpp
@@ -189,6 +189,16 @@ TSharedPtr<class INameValidatorInterface> UK2Node_Knot::MakeNameValidator() cons
 	return MakeShareable(new FDummyNameValidator(EValidatorResult::Ok));
 }
 
+UEdGraphPin* UK2Node_Knot::GetPassThroughPin(const UEdGraphPin* FromPin) const
+{
+	if(FromPin && Pins.Contains(FromPin))
+	{
+		return FromPin == Pins[0] ? Pins[1] : Pins[0];
+	}
+
+	return nullptr;
+}
+
 /////////////////////////////////////////////////////
 
 #undef LOCTEXT_NAMESPACE
diff --git a/Engine/Source/Editor/UnrealEd/Private/Kismet2/BlueprintEditorUtils.cpp b/Engine/Source/Editor/UnrealEd/Private/Kismet2/BlueprintEditorUtils.cpp
--- a/Engine/Source/Editor/UnrealEd/Private/Kismet2/BlueprintEditorUtils.cpp
+++ b/Engine/Source/Editor/UnrealEd/Private/Kismet2/BlueprintEditorUtils.cpp
@@ -672,6 +672,44 @@ void FBlueprintEditorUtils::GetAllGraphNames(const UBlueprint* Blueprint, TArray
 	}
 }
 
+void FBlueprintEditorUtils::GetCompilerRelevantNodes(const UEdGraphPin* FromPin, FCompilerRelevantNodesArray& OutNodes)
+{
+	if(FromPin)
+	{
+		// Start with the given pin's owning node
+		UK2Node* OwningNode = Cast<UK2Node>(FromPin->GetOwningNode());
+		if(OwningNode)
+		{
+			// If this node is not compiler relevant
+			if(!OwningNode->IsCompilerRelevant())
+			{
+				// And if this node has a matching "pass-through" pin
+				FromPin = OwningNode->GetPassThroughPin(FromPin);
+				if(FromPin)
+				{
+					// Recursively check each link for a compiler-relevant node that will "pass through" this node at compile time
+					for(auto LinkedPin : FromPin->LinkedTo)
+					{
+						GetCompilerRelevantNodes(LinkedPin, OutNodes);
+					}
+				}
+			}
+			else
+			{
+				OutNodes.Add(OwningNode);
+			}
+		}		
+	}
+}
+
+UK2Node* FBlueprintEditorUtils::FindFirstCompilerRelevantNode(const UEdGraphPin* FromPin)
+{
+	FCompilerRelevantNodesArray RelevantNodes;
+	GetCompilerRelevantNodes(FromPin, RelevantNodes);
+	
+	return RelevantNodes.Num() > 0 ? RelevantNodes[0] : nullptr;
+}
+
 /** 
  * Check FKismetCompilerContext::SetCanEverTickForActor
  */
diff --git a/Engine/Source/Editor/UnrealEd/Public/Kismet2/BlueprintEditorUtils.h b/Engine/Source/Editor/UnrealEd/Public/Kismet2/BlueprintEditorUtils.h
--- a/Engine/Source/Editor/UnrealEd/Public/Kismet2/BlueprintEditorUtils.h
+++ b/Engine/Source/Editor/UnrealEd/Public/Kismet2/BlueprintEditorUtils.h
@@ -83,6 +83,9 @@ public:
 
 DECLARE_CYCLE_STAT_EXTERN(TEXT("Notify Blueprint Changed"), EKismetCompilerStats_NotifyBlueprintChanged, STATGROUP_KismetCompiler, );
 
+/** Array type for GetCompilerRelevantNodes()                                                                                                                                                                                                                                                                                                                                                                                                                                                                                                                                                                                 */
+typedef TArray<UK2Node*, TInlineAllocator<4> > FCompilerRelevantNodesArray;
+
 class UNREALED_API FBlueprintEditorUtils
 {
 public:
@@ -496,6 +499,23 @@ public:
 	 */
 	static void GetAllGraphNames(const UBlueprint* Blueprint, TArray<FName>& GraphNames);
 
+	/**
+	 * Gets the compiler-relevant (i.e. non-ignorable) nodes from the given pin.
+	 *
+	 * @param			FromPin			The pin to start searching from.
+	 * @param			OutNodes		Will contain the given pin's owning node if compiler-relevant, or all nodes linked to the owning node at the matching "pass-through" pin that are compiler-relevant. Empty if no compiler-relevant nodes can be found from the given pin.
+	 */
+	static void GetCompilerRelevantNodes(const UEdGraphPin* FromPin, FCompilerRelevantNodesArray& OutNodes);
+
+	/**
+	 * Finds the first compiler-relevant (i.e. non-ignorable) node from the given pin.
+	 *
+	 * @param			FromPin			The pin to start searching from.
+	 *
+	 * @return			The given pin's owning node if compiler-relevant, or the first node linked to the owning node at the matching "pass-through" pin that is compiler-relevant. May be NULL if no compiler-relevant nodes can be found from the given pin.
+	 */
+	static UK2Node* FindFirstCompilerRelevantNode(const UEdGraphPin* FromPin);
+
 
 	//////////////////////////////////////////////////////////////////////////
 	// Functions
diff --git a/Engine/Source/Runtime/Engine/Classes/EdGraph/EdGraphNode.h b/Engine/Source/Runtime/Engine/Classes/EdGraph/EdGraphNode.h
--- a/Engine/Source/Runtime/Engine/Classes/EdGraph/EdGraphNode.h
+++ b/Engine/Source/Runtime/Engine/Classes/EdGraph/EdGraphNode.h
@@ -426,6 +426,12 @@ public:
 
 	/** Return the requested metadata for the pin if there is any */
 	virtual FString GetPinMetaData(FString InPinName, FName InKey) { return FString(); }
+
+	/** Return false if the node and any expansion will isolate itself during compile */
+	virtual bool IsCompilerRelevant() const { return true; }
+
+	/** Return the matching "pass-through" pin for the given pin (if supported by this node) */
+	virtual UEdGraphPin* GetPassThroughPin(const UEdGraphPin* FromPin) const { return nullptr; }
 #endif // WITH_EDITOR
 
 };
