Add an option to avoid generating DDC for engine packages.

diff --git a/Engine/Source/Editor/UnrealEd/Private/Commandlets/DerivedDataCacheCommandlet.cpp b/Engine/Source/Editor/UnrealEd/Private/Commandlets/DerivedDataCacheCommandlet.cpp
--- a/Engine/Source/Editor/UnrealEd/Private/Commandlets/DerivedDataCacheCommandlet.cpp
+++ b/Engine/Source/Editor/UnrealEd/Private/Commandlets/DerivedDataCacheCommandlet.cpp
@@ -81,6 +81,11 @@ int32 UDerivedDataCacheCommandlet::Main( const FString& Params )
 			PackageFilter |= NORMALIZE_ExcludeContentPackages;
 		}
 
+		if ( Switches.Contains(TEXT("PROJECTONLY")) )
+		{
+			PackageFilter |= NORMALIZE_ExcludeEnginePackages;
+		}
+
 		if ( !Switches.Contains(TEXT("DEV")) )
 		{
 			PackageFilter |= NORMALIZE_ExcludeDeveloperPackages;
