Fixed crash on unloding sub-level in world composition

diff --git a/Engine/Source/Editor/UnrealEd/Private/EditorLevelUtils.cpp b/Engine/Source/Editor/UnrealEd/Private/EditorLevelUtils.cpp
--- a/Engine/Source/Editor/UnrealEd/Private/EditorLevelUtils.cpp
+++ b/Engine/Source/Editor/UnrealEd/Private/EditorLevelUtils.cpp
@@ -420,9 +420,6 @@ namespace EditorLevelUtils
 
 			EditorDestroyLevel(InLevel);
 
-			// Scrub any levelscript actions so GC can take effect
-			FWorldDelegates::LevelRemovedFromWorld.Broadcast( InLevel, OwningWorld );
-
 			// Redraw the main editor viewports.
 			FEditorSupportDelegates::RedrawAllViewports.Broadcast();
 
