Paper2D: Moved 2d level editor settings -> project settings
- Added shortcut to edit layers settings
- Added button to snap selected actors to selected layer

diff --git a/Engine/Source/Editor/EditorSettingsViewer/Private/EditorSettingsViewerModule.cpp b/Engine/Source/Editor/EditorSettingsViewer/Private/EditorSettingsViewerModule.cpp
--- a/Engine/Source/Editor/EditorSettingsViewer/Private/EditorSettingsViewerModule.cpp
+++ b/Engine/Source/Editor/EditorSettingsViewer/Private/EditorSettingsViewerModule.cpp
@@ -174,14 +174,6 @@ protected:
 			GetMutableDefault<ULevelEditorViewportSettings>()
 		);
 
-		// view port settings
-		SettingsModule.RegisterSettings("Editor", "LevelEditor", "2D",
-			LOCTEXT("LevelEditor2DSettingsName", "2D"),
-			LOCTEXT("LevelEditor2DSettingsDescription", "Configure the settings for the 2D Level Editor."),
-			GetMutableDefault<ULevelEditor2DSettings>()
-		);
-
-
 		// miscellaneous settings
 		SettingsModule.RegisterSettings("Editor", "LevelEditor", "Misc",
 			LOCTEXT("LevelEditorMiscSettingsName", "Miscellaneous"),
diff --git a/Engine/Source/Editor/ProjectSettingsViewer/Private/ProjectSettingsViewerModule.cpp b/Engine/Source/Editor/ProjectSettingsViewer/Private/ProjectSettingsViewerModule.cpp
--- a/Engine/Source/Editor/ProjectSettingsViewer/Private/ProjectSettingsViewerModule.cpp
+++ b/Engine/Source/Editor/ProjectSettingsViewer/Private/ProjectSettingsViewerModule.cpp
@@ -290,6 +290,13 @@ protected:
 			LOCTEXT("AppearanceSettingsDescription", "Settings pertaining to the appearance of the editor"),
 			GetMutableDefault<UEditorProjectAppearanceSettings>()
 		);
+
+		// view port settings
+		SettingsModule.RegisterSettings("Project", "Editor", "2D",
+			LOCTEXT("Editor2DSettingsName", "2D"),
+			LOCTEXT("Editor2DSettingsDescription", "Configure the settings for the 2D Level Editor."),
+			GetMutableDefault<ULevelEditor2DSettings>()
+		);
 	}
 
 	/** Unregisters all settings. */
diff --git a/Engine/Source/Editor/UnrealEd/Classes/Settings/LevelEditorViewportSettings.h b/Engine/Source/Editor/UnrealEd/Classes/Settings/LevelEditorViewportSettings.h
--- a/Engine/Source/Editor/UnrealEd/Classes/Settings/LevelEditorViewportSettings.h
+++ b/Engine/Source/Editor/UnrealEd/Classes/Settings/LevelEditorViewportSettings.h
@@ -180,56 +180,6 @@ struct UNREALED_API FSnapToSurfaceSettings
 	bool bSnapRotation;
 };
 
-
-/**
-* 2D layer settings
-*/
-USTRUCT()
-struct UNREALED_API FMode2DLayer
-{
-	GENERATED_USTRUCT_BODY()
-
-	FMode2DLayer()
-		: Name(TEXT("Default"))
-		, LayerOffset(0)
-	{ }
-
-	/** Whether snapping to surfaces in the world is enabled */
-	UPROPERTY(EditAnywhere, config, Category=Layer)
-	FString Name;
-
-	/** The amount of offset to apply when snapping to surfaces */
-	UPROPERTY(EditAnywhere, config, Category=Layer)
-	float LayerOffset;
-};
-
-/**
-* Implements the Level Editor's 2D view port settings.
-*/
-UCLASS(config = EditorUserSettings)
-class UNREALED_API ULevelEditor2DSettings
-	: public UObject
-{
-	GENERATED_UCLASS_BODY()
-
-public:
-	/** If enabled will allow 2D mode */
-	UPROPERTY(EditAnywhere, config, Category = General, meta = (DisplayName = "Enable 2D Mode"))
-	bool bMode2DEnabled;
-
-	/** When enabled and active snap layer is valid, items dropped in the viewport will be dropped on the appropriate Y distance away */
-	UPROPERTY(EditAnywhere, config, Category = LayerSnapping, AdvancedDisplay)
-	bool bEnableLayerSnap;
-
-	/** Will try to use snap settings from SnapLayers */
-	UPROPERTY(EditAnywhere, config, Category = LayerSnapping, AdvancedDisplay)
-	int32 ActiveSnapLayerIndex;
-
-	/** Snap layers that are displayed in the viewport toolbar */
-	UPROPERTY(EditAnywhere, config, Category = LayerSnapping)
-	TArray<FMode2DLayer> SnapLayers;
-};
-
 /**
  * Implements the Level Editor's view port settings.
  */
diff --git a/Engine/Source/Editor/UnrealEd/Private/LevelEditorViewport.cpp b/Engine/Source/Editor/UnrealEd/Private/LevelEditorViewport.cpp
--- a/Engine/Source/Editor/UnrealEd/Private/LevelEditorViewport.cpp
+++ b/Engine/Source/Editor/UnrealEd/Private/LevelEditorViewport.cpp
@@ -49,6 +49,7 @@
 #include "NotificationManager.h"
 #include "SNotificationList.h"
 #include "ComponentEditorUtils.h"
+#include "Settings/EditorProjectSettings.h"
 
 DEFINE_LOG_CATEGORY(LogEditorViewport);
 
@@ -998,9 +999,16 @@ static FActorPositionTraceResult TraceForPositionOn2DLayer(const FViewportCursor
 	const ULevelEditor2DSettings* Settings2D = GetDefault<ULevelEditor2DSettings>();
 	check(Settings2D->SnapLayers.IsValidIndex(Settings2D->ActiveSnapLayerIndex));
 
-	float Offset = Settings2D->SnapLayers[Settings2D->ActiveSnapLayerIndex].LayerOffset;
-	FVector PlaneCenter(0, Offset, 0);
-	FVector PlaneNormal(0, -1, 0);
+	float Offset = Settings2D->SnapLayers[Settings2D->ActiveSnapLayerIndex].Depth;
+	FVector PlaneCenter(0, 0, 0);
+	FVector PlaneNormal(0, 0, 0);
+
+	switch (Settings2D->SnapAxis)
+	{
+	case ELevelEditor2DAxis::X: PlaneCenter.X = Offset; PlaneNormal.X = -1; break;
+	case ELevelEditor2DAxis::Y: PlaneCenter.Y = Offset; PlaneNormal.Y = -1; break;
+	case ELevelEditor2DAxis::Z: PlaneCenter.Z = Offset; PlaneNormal.Z = -1; break;
+	}
 
 	FActorPositionTraceResult Result;
 	float Numerator = FVector::DotProduct(PlaneCenter - Cursor.GetOrigin(), PlaneNormal);
diff --git a/Engine/Source/Editor/UnrealEd/Private/SEditorViewport.cpp b/Engine/Source/Editor/UnrealEd/Private/SEditorViewport.cpp
--- a/Engine/Source/Editor/UnrealEd/Private/SEditorViewport.cpp
+++ b/Engine/Source/Editor/UnrealEd/Private/SEditorViewport.cpp
@@ -9,6 +9,7 @@
 #include "TutorialMetaData.h"
 #include "SNotificationList.h"
 #include "NotificationManager.h"
+#include "Settings/EditorProjectSettings.h"
 
 #define LOCTEXT_NAMESPACE "EditorViewport"
 
diff --git a/Engine/Source/Editor/UnrealEd/Private/STransformViewportToolbar.cpp b/Engine/Source/Editor/UnrealEd/Private/STransformViewportToolbar.cpp
--- a/Engine/Source/Editor/UnrealEd/Private/STransformViewportToolbar.cpp
+++ b/Engine/Source/Editor/UnrealEd/Private/STransformViewportToolbar.cpp
@@ -7,7 +7,10 @@
 #include "EditorViewportCommands.h"
 #include "SViewportToolBarIconMenu.h"
 #include "SViewportToolBarComboMenu.h"
+#include "ISettingsModule.h"
 #include "SNumericEntryBox.h"
+#include "ScopedTransaction.h"
+#include "Settings/EditorProjectSettings.h"
 
 #define LOCTEXT_NAMESPACE "TransformToolBar"
 
@@ -640,6 +643,79 @@ TSharedRef<SWidget> STransformViewportToolBar::FillLayer2DSnapMenu()
 		ShowMenuBuilder.AddMenuEntry(FText::FromName(LayerName), FText::GetEmpty(), FSlateIcon(), Action, NAME_None, EUserInterfaceActionType::RadioButton);
 	}
 
+	struct FLocalFunctions
+	{
+		static void ShowSettingsViewer()
+		{
+			if (ISettingsModule* SettingsModule = FModuleManager::GetModulePtr<ISettingsModule>("Settings"))
+			{
+				SettingsModule->ShowViewer("Project", "Editor", "2D");
+			}
+		}
+
+		static FReply SnapSelectionToSelectedOffset(int32 SnapLayerIndex)
+		{
+			auto Settings = GetDefault<ULevelEditor2DSettings>();
+			if (Settings->SnapLayers.IsValidIndex(SnapLayerIndex))
+			{
+				const FScopedTransaction Transaction(NSLOCTEXT("UnrealEd", "SnapSelection2D", "Snap Selection to 2D Layer"));
+
+				float SnapDepth = Settings->SnapLayers[Settings->ActiveSnapLayerIndex].Depth;
+				USelection* SelectedActors = GEditor->GetSelectedActors();
+				for (FSelectionIterator Iter(*SelectedActors); Iter; ++Iter)
+				{
+					AActor* Actor = CastChecked<AActor>(*Iter);
+
+					// Only snap actors that are not attached to something else
+					if (Actor->GetAttachParentActor() == nullptr)
+					{
+						FTransform Transform = Actor->GetTransform();
+						FVector CurrentLocation = Transform.GetLocation();
+
+						switch (Settings->SnapAxis)
+						{
+						case ELevelEditor2DAxis::X: CurrentLocation.X = SnapDepth; break;
+						case ELevelEditor2DAxis::Y: CurrentLocation.Y = SnapDepth; break;
+						case ELevelEditor2DAxis::Z: CurrentLocation.Z = SnapDepth; break;
+						}
+						
+						Transform.SetLocation(CurrentLocation);
+						Actor->Modify();
+						Actor->SetActorTransform(Transform);
+					}
+				}
+
+				GEditor->RedrawViewports();
+			}
+
+			return FReply::Handled();
+		}
+
+		static bool IsSnapSelectedActorsEnabled()
+		{
+			auto Settings = GetDefault<ULevelEditor2DSettings>();
+			return Settings->SnapLayers.IsValidIndex(Settings->ActiveSnapLayerIndex) && GEditor->GetSelectedActorCount() > 0;
+		}
+	};
+
+	FUIAction ShowSettingsAction(FExecuteAction::CreateStatic(&FLocalFunctions::ShowSettingsViewer));
+	ShowMenuBuilder.AddMenuEntry(LOCTEXT("2DSnap_EditLayer", "Edit Layers..."), FText::GetEmpty(), FSlateIcon(), ShowSettingsAction, NAME_None, EUserInterfaceActionType::Button);
+
+	// -------------------------------------------------------
+	ShowMenuBuilder.AddMenuSeparator();
+
+	ShowMenuBuilder.AddWidget(
+		SNew(SButton)
+		.Text(LOCTEXT("2DSnap_SnapSelection", "Snap Selected Actors"))
+		.IsEnabled(TAttribute<bool>::Create(TAttribute<bool>::FGetter::CreateStatic(FLocalFunctions::IsSnapSelectedActorsEnabled)))
+		.OnClicked(FOnClicked::CreateStatic(FLocalFunctions::SnapSelectionToSelectedOffset, Settings2D->ActiveSnapLayerIndex))
+		.VAlign(VAlign_Center)
+		.HAlign(HAlign_Fill),
+
+		FText::GetEmpty(),
+		true
+		);
+
 	return ShowMenuBuilder.MakeWidget();
 }
 
diff --git a/Engine/Source/Editor/UnrealEd/Private/Settings/EditorProjectSettings.cpp b/Engine/Source/Editor/UnrealEd/Private/Settings/EditorProjectSettings.cpp
--- a/Engine/Source/Editor/UnrealEd/Private/Settings/EditorProjectSettings.cpp
+++ b/Engine/Source/Editor/UnrealEd/Private/Settings/EditorProjectSettings.cpp
@@ -65,4 +65,16 @@ void UEditorProjectAppearanceSettings::PostInitProperties()
 
 	FUnitConversion::Settings().SetGlobalUnitDisplay(ConvertGlobalUnitDisplay(UnitDisplay));
 	FUnitConversion::Settings().SetDefaultInputUnit(ConvertDefaultInputUnits(DefaultInputUnits));
-}
\ No newline at end of file
+}
+
+
+
+/* ULevelEditor2DSettings
+*****************************************************************************/
+
+ULevelEditor2DSettings::ULevelEditor2DSettings(const FObjectInitializer& ObjectInitializer)
+	: Super(ObjectInitializer),
+	  SnapAxis(ELevelEditor2DAxis::Y)
+{
+}
+
diff --git a/Engine/Source/Editor/UnrealEd/Private/Settings/SettingsClasses.cpp b/Engine/Source/Editor/UnrealEd/Private/Settings/SettingsClasses.cpp
--- a/Engine/Source/Editor/UnrealEd/Private/Settings/SettingsClasses.cpp
+++ b/Engine/Source/Editor/UnrealEd/Private/Settings/SettingsClasses.cpp
@@ -344,15 +344,6 @@ ULevelEditorPlaySettings::ULevelEditorPlaySettings( const FObjectInitializer& Ob
 	RouteGamepadToSecondWindow = false;
 }
 
-/* ULevelEditor2DSettings 
-*****************************************************************************/
-
-ULevelEditor2DSettings::ULevelEditor2DSettings(const FObjectInitializer& ObjectInitializer)
-	: Super(ObjectInitializer)
-{
-}
-
-
 /* ULevelEditorViewportSettings interface
  *****************************************************************************/
 
diff --git a/Engine/Source/Editor/UnrealEd/Public/Settings/EditorProjectSettings.h b/Engine/Source/Editor/UnrealEd/Public/Settings/EditorProjectSettings.h
--- a/Engine/Source/Editor/UnrealEd/Public/Settings/EditorProjectSettings.h
+++ b/Engine/Source/Editor/UnrealEd/Public/Settings/EditorProjectSettings.h
@@ -41,4 +41,65 @@ private:
 
 	UPROPERTY(EditAnywhere, config, Category=Units, meta=(Tooltip="The default units to use for location and distance measurements on editor user interfaces when not explicitly specified by the user."))
 	EDefaultLocationUnit DefaultInputUnits;
-};
\ No newline at end of file
+};
+
+/**
+* 2D layer settings
+*/
+USTRUCT()
+struct UNREALED_API FMode2DLayer
+{
+	GENERATED_USTRUCT_BODY()
+
+	FMode2DLayer()
+	: Name(TEXT("Default"))
+	, Depth(0)
+	{ }
+
+	/** Whether snapping to surfaces in the world is enabled */
+	UPROPERTY(EditAnywhere, config, Category = Layer)
+	FString Name;
+
+	/** The amount of depth to apply when snapping to surfaces */
+	UPROPERTY(EditAnywhere, config, Category = Layer)
+	float Depth;
+};
+
+UENUM()
+enum class ELevelEditor2DAxis : uint8
+{
+	X,
+	Y,
+	Z
+};
+
+/**
+* Implements the Project specific 2D viewport settings.
+*/
+UCLASS(config=Editor)
+class UNREALED_API ULevelEditor2DSettings : public UObject
+{
+	GENERATED_UCLASS_BODY()
+
+public:
+	/** If enabled will allow 2D mode */
+	UPROPERTY(EditAnywhere, config, Category = General, meta = (DisplayName = "Enable 2D Mode"))
+	bool bMode2DEnabled;
+
+	/** When enabled and active snap layer is valid, items dropped in the viewport will be dropped on the appropriate Y distance away */
+	UPROPERTY(EditAnywhere, config, Category = LayerSnapping, AdvancedDisplay)
+	bool bEnableLayerSnap;
+
+	/** Will try to use snap settings from SnapLayers */
+	UPROPERTY(EditAnywhere, config, Category = LayerSnapping, AdvancedDisplay)
+	int32 ActiveSnapLayerIndex;
+
+	/** Snap axis */
+	UPROPERTY(EditAnywhere, config, Category = LayerSnapping)
+	ELevelEditor2DAxis SnapAxis;
+
+	/** Snap layers that are displayed in the viewport toolbar */
+	UPROPERTY(EditAnywhere, config, Category = LayerSnapping)
+	TArray<FMode2DLayer> SnapLayers;
+};
+
