Fix for watch values being reported as out of scope in the blueprint debugger.
#UE-12022 - Variables incorrectly returning as Out of Scope when watched
#Branch UE4
#Proj UnrealEd

#info The previous fix was broken during cleanup for check-in.

diff --git a/Engine/Source/Editor/UnrealEd/Private/Kismet2/KismetDebugUtilities.cpp b/Engine/Source/Editor/UnrealEd/Private/Kismet2/KismetDebugUtilities.cpp
--- a/Engine/Source/Editor/UnrealEd/Private/Kismet2/KismetDebugUtilities.cpp
+++ b/Engine/Source/Editor/UnrealEd/Private/Kismet2/KismetDebugUtilities.cpp
@@ -897,18 +897,18 @@ FKismetDebugUtilities::EWatchTextResult FKismetDebugUtilities::GetWatchText(FStr
 
 			// Try at member scope if it wasn't part of a current function scope
 			UClass* PropertyClass = Cast<UClass>(Property->GetOuter());
-			if( !PropertyBase && PropertyClass )
+			if (!PropertyBase && PropertyClass)
 			{
-				if( ActiveObject->GetClass()->IsChildOf( PropertyClass ))
+				if (ActiveObject->GetClass()->IsChildOf(PropertyClass))
 				{
 					PropertyBase = ActiveObject;
 				}
-				else if( AActor* Actor = Cast<AActor>( ActiveObject ))
+				else if (AActor* Actor = Cast<AActor>(ActiveObject))
 				{
 					// Try and locate the propertybase in the actor components
-					for( auto ComponentIter : Actor->GetComponents() )
+					for (auto ComponentIter : Actor->GetComponents())
 					{
-						if( ComponentIter->GetClass()->IsChildOf(PropertyClass))
+						if (ComponentIter->GetClass()->IsChildOf(PropertyClass))
 						{
 							PropertyBase = ComponentIter;
 							break;
@@ -918,12 +918,12 @@ FKismetDebugUtilities::EWatchTextResult FKismetDebugUtilities::GetWatchText(FStr
 			}
 #if USE_UBER_GRAPH_PERSISTENT_FRAME
 			// Try find the propertybase in the persistent ubergraph frame
-			UFunction* OuterFunction = Cast<UFunction>( Property->GetOuter() );
-			if( !PropertyBase && OuterFunction )
+			UFunction* OuterFunction = Cast<UFunction>(Property->GetOuter());
+			if(!PropertyBase && OuterFunction)
 			{
-				if( UBlueprintGeneratedClass* BPGC = Cast<UBlueprintGeneratedClass>( Blueprint->GeneratedClass ))
+				if(UBlueprintGeneratedClass* BPGC = Cast<UBlueprintGeneratedClass>(Blueprint->GeneratedClass))
 				{
-					PropertyBase = BPGC->GetPersistentUberGraphFrame( ActiveObject, OuterFunction );
+					PropertyBase = BPGC->GetPersistentUberGraphFrame(ActiveObject, BPGC->UberGraphFunction);
 				}
 			}
 #endif // USE_UBER_GRAPH_PERSISTENT_FRAME
