Improved performance for startup cook time.
Improved caching routine for cook in editor modes.
Improved startup time for editor.

diff --git a/Engine/Source/Editor/UnrealEd/Private/CookOnTheFlyServer.cpp b/Engine/Source/Editor/UnrealEd/Private/CookOnTheFlyServer.cpp
--- a/Engine/Source/Editor/UnrealEd/Private/CookOnTheFlyServer.cpp
+++ b/Engine/Source/Editor/UnrealEd/Private/CookOnTheFlyServer.cpp
@@ -676,9 +676,6 @@ void UCookOnTheFlyServer::GetDependencies( const TSet<UPackage*>& Packages, TSet
 	TArray<UObject*> Exclude;
 	FArchiveFindReferences ArFindReferences( RootSet, Found, Exclude );
 
-
-
-
 	// Iterate through the object list
 	/*for( FObjectIterator It; It; ++It )
 	{
@@ -1208,7 +1205,7 @@ uint32 UCookOnTheFlyServer::TickCookOnTheSide( const float TimeSlice, uint32 &Co
 				// if this objects data is cached then we can call FinishedCookedPLatformDataCache
 				// we can only safely call this when we are finished caching this object completely.
 				// this doesn't ever happen for cook in editor or cook on the fly mode
-				if ( bIsAllDataCached && (CurrentCookMode == ECookMode::CookByTheBook))
+				if ( bIsAllDataCached && (CurrentCookMode == ECookMode::CookByTheBook) )
 				{
 					// this might be run multiple times for a single object
 					Obj->WillNeverCacheCookedPlatformDataAgain();
@@ -1327,6 +1324,8 @@ uint32 UCookOnTheFlyServer::TickCookOnTheSide( const float TimeSlice, uint32 &Co
 			SCOPE_TIMER(SavingPackages);
 			for ( int32 I = 0; I < PackagesToSave.Num(); ++I )
 			{
+				UPackage *Package = PackagesToSave[I];
+
 				// if we are processing unsolicited packages we can optionally not save these right now
 				// the unsolicited packages which we missed now will be picked up on next run
 				// we want to do this in cook on the fly also, if there is a new network package request instead of saving unsolicited packages we can process the requested package
@@ -1345,6 +1344,23 @@ uint32 UCookOnTheFlyServer::TickCookOnTheSide( const float TimeSlice, uint32 &Co
 						// our timeslice is up
 					}
 
+					bool bFinishedCachingCookedPlatformData = false;
+					// if we are in realtime mode then don't wait forever for the package to be ready
+					TArray<UObject*> ObjectsInPackage;
+					GetObjectsWithOuter( Package, ObjectsInPackage );
+					while ( (!Timer.IsTimeUp()) && IsRealtimeMode() && (bShouldFinishTick == false) )
+					{
+						if ( FinishPackageCacheForCookedPlatformData(ObjectsInPackage) == true )
+						{
+							bFinishedCachingCookedPlatformData = true;
+							break;
+						}
+
+						// sleep for a bit
+						FPlatformProcess::Sleep(0.0f);
+					}
+
+					bShouldFinishTick |= !bFinishedCachingCookedPlatformData;
 
 					if ( bShouldFinishTick )
 					{
@@ -1372,7 +1388,6 @@ uint32 UCookOnTheFlyServer::TickCookOnTheSide( const float TimeSlice, uint32 &Co
 					GenerateManifestInfo( PackagesToSave[I], AllTargetPlatformNames );
 				}
 
-				UPackage *Package = PackagesToSave[I];
 				FName PackageFName = GetCachedStandardPackageFileFName(Package);
 				TArray<FName> SaveTargetPlatformNames = AllTargetPlatformNames;
 				TArray<FName> CookedTargetPlatforms;
@@ -1558,14 +1573,16 @@ void UCookOnTheFlyServer::MarkPackageDirtyForCooker( UPackage *Package )
 		/*FString PackageFilename(GetPackageFilename(Package));
 		FPaths::MakeStandardFilename(PackageFilename);*/
 
-		FString PackageFilename = GetCachedStandardPackageFilename(Package);
+		FName PackageFFileName = GetCachedStandardPackageFileFName(Package);
 
-		ClearPackageFilenameCacheForPackage( Package );
+		if ( PackageFFileName == NAME_None )
+		{
+			ClearPackageFilenameCacheForPackage( Package );
+		}
 
 #if DEBUG_COOKONTHEFLY
 		UE_LOG(LogCook, Display, TEXT("Modification detected to package %s"), *PackageFilename);
 #endif
-		const FName PackageFFileName = FName(*PackageFilename);
 
 		if ( CurrentCookMode == ECookMode::CookByTheBookFromTheEditor )
 		{
@@ -1819,6 +1836,8 @@ bool UCookOnTheFlyServer::SaveCookedPackage( UPackage* Package, uint32 SaveFlags
 			bool bCookPackage = true;
 
 
+			// this code has been moved to the clean sandbox so that it's only processed in one location on startup
+			// will remove this code in future release
 			// If we are not iterative cooking, then cook the package
 			/*bool bCookPackage = (IsCookFlagSet(ECookInitializationFlags::Iterative) == false);
 
@@ -2428,6 +2447,7 @@ void UCookOnTheFlyServer::GenerateAssetRegistry(const TArray<ITargetPlatform*>&
 
 void UCookOnTheFlyServer::GenerateLongPackageNames(TArray<FString>& FilesInPath)
 {
+#if 0
 	TArray<FString> FilesInPathReverse;
 	FilesInPathReverse.Reserve(FilesInPath.Num());
 	for( int32 FileIndex = 0; FileIndex < FilesInPath.Num(); FileIndex++ )
@@ -2452,6 +2472,39 @@ void UCookOnTheFlyServer::GenerateLongPackageNames(TArray<FString>& FilesInPath)
 		}
 	}
 	Exchange(FilesInPathReverse, FilesInPath);
+#else
+	TArray<FName> FilesInPathReverse;
+	FilesInPathReverse.Reserve(FilesInPath.Num());
+
+	for( int32 FileIndex = 0; FileIndex < FilesInPath.Num(); FileIndex++ )
+	{
+		const FString& FileInPath = FilesInPath[FilesInPath.Num() - FileIndex - 1];
+		if (FPackageName::IsValidLongPackageName(FileInPath))
+		{
+			const FName FileInPathFName(*FileInPath);
+			FilesInPathReverse.AddUnique(FileInPathFName);
+		}
+		else
+		{
+			FString LongPackageName;
+			if (FPackageName::TryConvertFilenameToLongPackageName(FileInPath, LongPackageName))
+			{
+				const FName LongPackageFName(*LongPackageName);
+				FilesInPathReverse.AddUnique(LongPackageFName);
+			}
+			else
+			{
+				LogCookerMessage( FString::Printf(TEXT("Unable to generate long package name for %s"), *FileInPath), EMessageSeverity::Warning);
+				UE_LOG(LogCook, Warning, TEXT("Unable to generate long package name for %s"), *FileInPath);
+			}
+		}
+	}
+	// Exchange(FilesInPathReverse, FilesInPath);
+	for ( const auto& Files : FilesInPathReverse )
+	{
+		FilesInPath.Add(Files.ToString());
+	}
+#endif
 }
 
 void UCookOnTheFlyServer::AddFileToCook( TArray<FString>& InOutFilesToCook, const FString &InFilename ) const
@@ -2544,7 +2597,7 @@ void UCookOnTheFlyServer::CollectFilesToCook(TArray<FString>& FilesInPath, const
 	if ( IsCookingDLC() )
 	{
 		// get the dlc and make sure we cook that directory 
-		FString DLCPath = FPaths::GamePluginsDir() / CookByTheBookOptions->DlcName / FString(TEXT("Content"));;
+		FString DLCPath = FPaths::GamePluginsDir() / CookByTheBookOptions->DlcName / FString(TEXT("Content"));
 
 		TArray<FString> Files;
 		IFileManager::Get().FindFilesRecursive(Files, *DLCPath, *(FString(TEXT("*")) + FPackageName::GetAssetPackageExtension()), true, false);
@@ -3038,7 +3091,6 @@ void UCookOnTheFlyServer::StartCookByTheBook( const FCookByTheBookStartupOptions
 		}
 	}
 
-
 	if ( IsCookingDLC() )
 	{
 		// if we are cooking dlc we must be based of a release version cook
@@ -3063,8 +3115,6 @@ void UCookOnTheFlyServer::StartCookByTheBook( const FCookByTheBookStartupOptions
 		}
 	}
 	
-	
-	
 	// allow the game to fill out the asset registry, as well as get a list of objects to always cook
 	TArray<FString> FilesInPath;
 	FGameDelegates::Get().GetCookModificationDelegate().ExecuteIfBound(FilesInPath);
@@ -3083,9 +3133,12 @@ void UCookOnTheFlyServer::StartCookByTheBook( const FCookByTheBookStartupOptions
 		UE_LOG(LogCook, Warning, TEXT("No files found."));
 	}
 
-
-	GenerateLongPackageNames(FilesInPath);
-
+	{
+#if OUTPUT_TIMING
+		SCOPE_TIMER(GenerateLongPackageName);
+#endif
+		GenerateLongPackageNames(FilesInPath);
+	}
 	// add all the files for the requested platform to the cook list
 	for ( const auto& FileName : FilesInPath )
 	{
@@ -3110,6 +3163,7 @@ void UCookOnTheFlyServer::StartCookByTheBook( const FCookByTheBookStartupOptions
 		if ( !BasedOnReleaseVersion.IsEmpty() )
 		{
 			// if we are based of a release and we are not cooking dlc then we should always be creating a new one (note that we could be creating the same one we are based of).
+			// note that we might erroneously enter here if we are generating a patch instead and we accidentally passed in BasedOnReleaseVersion to the cooker instead of to unrealpak
 			check( !CreateReleaseVersion.IsEmpty() );
 
 			for ( const auto& PlatformName : TargetPlatformNames )
@@ -3335,4 +3389,4 @@ bool UCookOnTheFlyServer::GetAllPackagesFromAssetRegistry( const FString& AssetR
 		return true;
 	}
 	return false;
-}
+}
\ No newline at end of file
