Fix build error on Mac

diff --git a/Engine/Source/Editor/UnrealEd/UnrealEd.Build.cs b/Engine/Source/Editor/UnrealEd/UnrealEd.Build.cs
--- a/Engine/Source/Editor/UnrealEd/UnrealEd.Build.cs
+++ b/Engine/Source/Editor/UnrealEd/UnrealEd.Build.cs
@@ -79,6 +79,7 @@ public class UnrealEd : ModuleRules
                 "Http",
 				"UnrealAudio",
                 "Niagara",
+                "VectorVM",
 			}
 		);
 
@@ -125,6 +126,7 @@ public class UnrealEd : ModuleRules
 				"HeadMountedDisplay",
 				"MeshPaint",
                 "Foliage",
+                "VectorVM",
 			}
 		);
 
@@ -223,6 +225,7 @@ public class UnrealEd : ModuleRules
 				"BlueprintGraph",
 				"NiagaraEditor",
                 "Niagara",
+                "VectorVM",
 			}
 			);
 
