Implemented better filtering in the "Replace References" part of the Delete Asset dialog, so that Blueprint assets are filtered by their parent class as well.

diff --git a/Engine/Source/Editor/UnrealEd/Private/AssetDeleteModel.cpp b/Engine/Source/Editor/UnrealEd/Private/AssetDeleteModel.cpp
--- a/Engine/Source/Editor/UnrealEd/Private/AssetDeleteModel.cpp
+++ b/Engine/Source/Editor/UnrealEd/Private/AssetDeleteModel.cpp
@@ -290,7 +290,7 @@ bool FAssetDeleteModel::ComputeCanReplaceReferences()
 		PendingDeletedObjects.Add(PendingDelete->GetObject());
 	}
 
-	return ObjectTools::AreObjectsValidForReplace(PendingDeletedObjects);
+	return ObjectTools::AreObjectsOfEquivalantType(PendingDeletedObjects);
 }
 
 bool FAssetDeleteModel::CanReplaceReferences() const
@@ -316,6 +316,31 @@ bool FAssetDeleteModel::CanReplaceReferencesWith( const FAssetData& InAssetData
 		return true;
 	}
 
+	// Filter out blueprints of different types
+	if (FirstPendingDelete->IsChildOf(UBlueprint::StaticClass()) && AssetDataClass->IsChildOf(UBlueprint::StaticClass()))
+	{
+		// Get BP parent classes
+		static const FName ParentClassName("ParentClass");
+		const UClass* OriginalBPClass = CastChecked<UBlueprint>(PendingDeletes[0]->GetObject())->ParentClass;
+		const FString* BPClassNameToTest = InAssetData.TagsAndValues.Find(ParentClassName);
+
+		if (BPClassNameToTest && !BPClassNameToTest->IsEmpty())
+		{
+			UClass* ParentClass = FindObject<UClass>(NULL, **BPClassNameToTest);
+			if (!ParentClass)
+			{
+				ParentClass = LoadObject<UClass>(NULL, **BPClassNameToTest);
+			}
+
+			if (!ParentClass || !ParentClass->IsChildOf(OriginalBPClass))
+			{
+				// If we couldn't determine the asset parent class (e.g. because the ParentClass tag wasn't present in the FAssetData), or if
+				// the asset parent class wasn't equal to or derived from the pending delete BP class, filter i
+				return true;
+			}
+		}
+	}
+
 	// Only show objects that are of replaceable because their classes are compatible.
 	return !ObjectTools::AreClassesInterchangeable( FirstPendingDelete, AssetDataClass );
 }
diff --git a/Engine/Source/Editor/UnrealEd/Private/ObjectTools.cpp b/Engine/Source/Editor/UnrealEd/Private/ObjectTools.cpp
--- a/Engine/Source/Editor/UnrealEd/Private/ObjectTools.cpp
+++ b/Engine/Source/Editor/UnrealEd/Private/ObjectTools.cpp
@@ -3430,27 +3430,6 @@ namespace ObjectTools
 		return bIsPlaceable && !bIsAbstractOrDeprecated && !bIsSkeletonClass;
 	}
 
-	bool AreObjectsValidForReplace(const TArray<UObject*>& InProposedObjects)
-	{
-		if (InProposedObjects.Num() > 0)
-		{
-			// Make sure we don't have any blueprinted components
-			for (TArray<UObject*>::TConstIterator ProposedObjIter(InProposedObjects); ProposedObjIter; ++ProposedObjIter)
-			{
-				UObject* CurProposedObj = *ProposedObjIter;
-				check(CurProposedObj);
-				
-				const UBlueprint* BlueprintObject = Cast<UBlueprint>(CurProposedObj);
-				if (BlueprintObject && BlueprintObject->GeneratedClass->IsChildOf(UActorComponent::StaticClass()))
-				{
-					return false;
-				}
-			}
-		}
-		// If we haven't found any classes we deem illegal just make sure they are all the same type.
-		return AreObjectsOfEquivalantType(InProposedObjects);
-	}
-
 	bool AreObjectsOfEquivalantType( const TArray<UObject*>& InProposedObjects )
 	{
 		if ( InProposedObjects.Num() > 0 )
@@ -3471,6 +3450,14 @@ namespace ObjectTools
 
 				const UClass* CurProposedClass = CurProposedObj->GetClass();
 
+				if (ComparisonClass->IsChildOf(UBlueprint::StaticClass()) && CurProposedClass->IsChildOf(UBlueprint::StaticClass()))
+				{
+					if (*CastChecked<UBlueprint>(ComparisonObject)->ParentClass != *CastChecked<UBlueprint>(CurProposedObj)->ParentClass)
+					{
+						return false;
+					}
+				}
+
 				if ( !AreClassesInterchangeable( ComparisonClass, CurProposedClass ) )
 				{
 					return false;
diff --git a/Engine/Source/Editor/UnrealEd/Public/ObjectTools.h b/Engine/Source/Editor/UnrealEd/Public/ObjectTools.h
--- a/Engine/Source/Editor/UnrealEd/Public/ObjectTools.h
+++ b/Engine/Source/Editor/UnrealEd/Public/ObjectTools.h
@@ -504,15 +504,6 @@ namespace ObjectTools
 	 */
 	UNREALED_API bool IsClassRedirector( const UClass* Class );
 
-	/**
-	 * Determines if an array of objects are valid for replace
-	 *
-	 * @param InProposedObjects	The objects to check.
-	 * 
-	 * @return true if all objects are not a restricted type for replacement.
-	 */
-	UNREALED_API bool AreObjectsValidForReplace( const TArray<UObject*>& InProposedObjects );
-
 	/**
 	 * Determines if an array of objects are all of interchangeable types.
 	 *
