Fixed issue in which the root component of an actor can be lost after hot reload if its class has changed.

diff --git a/Engine/Source/Developer/HotReload/Private/HotReloadClassReinstancer.h b/Engine/Source/Developer/HotReload/Private/HotReloadClassReinstancer.h
--- a/Engine/Source/Developer/HotReload/Private/HotReloadClassReinstancer.h
+++ b/Engine/Source/Developer/HotReload/Private/HotReloadClassReinstancer.h
@@ -116,6 +116,12 @@ public:
 	// End of FSerializableObject interface
 
 	virtual bool IsClassObjectReplaced() const override  { return true; }
+
+protected:
+
+	// FBlueprintCompileReinstancer interface
+	virtual bool ShouldPreserveRootComponentOfReinstancedActor() const override { return false; }
+	// End of FBlueprintCompileReinstancer interface
 };
 
 #endif // WITH_ENGINE
\ No newline at end of file
diff --git a/Engine/Source/Editor/UnrealEd/Private/Kismet2/KismetReinstanceUtilities.cpp b/Engine/Source/Editor/UnrealEd/Private/Kismet2/KismetReinstanceUtilities.cpp
--- a/Engine/Source/Editor/UnrealEd/Private/Kismet2/KismetReinstanceUtilities.cpp
+++ b/Engine/Source/Editor/UnrealEd/Private/Kismet2/KismetReinstanceUtilities.cpp
@@ -391,7 +391,7 @@ void FBlueprintCompileReinstancer::ReinstanceInner(bool bForceAlwaysReinstance)
 		if (bShouldReinstance)
 		{
 			UE_LOG(LogBlueprint, Log, TEXT("BlueprintCompileReinstancer: Doing a full reinstance on class '%s'"), *GetPathNameSafe(ClassToReinstance));
-			ReplaceInstancesOfClass(DuplicatedClass, ClassToReinstance, OriginalCDO, &ObjectsThatShouldUseOldStuff, IsClassObjectReplaced());
+			ReplaceInstancesOfClass(DuplicatedClass, ClassToReinstance, OriginalCDO, &ObjectsThatShouldUseOldStuff, IsClassObjectReplaced(), ShouldPreserveRootComponentOfReinstancedActor());
 		}
 		else if (ClassToReinstance->IsChildOf<UActorComponent>())
 		{
@@ -779,7 +779,7 @@ void FActorReplacementHelper::AttachChildActors(USceneComponent* RootComponent,
 	}
 }
 
-void FBlueprintCompileReinstancer::ReplaceInstancesOfClass(UClass* OldClass, UClass* NewClass, UObject*	OriginalCDO, TSet<UObject*>* ObjectsThatShouldUseOldStuff, bool bClassObjectReplaced)
+void FBlueprintCompileReinstancer::ReplaceInstancesOfClass(UClass* OldClass, UClass* NewClass, UObject*	OriginalCDO, TSet<UObject*>* ObjectsThatShouldUseOldStuff, bool bClassObjectReplaced, bool bPreserveRootComponent)
 {
 	USelection* SelectedActors;
 	bool bSelectionChanged = false;
@@ -906,7 +906,10 @@ void FBlueprintCompileReinstancer::ReplaceInstancesOfClass(UClass* OldClass, UCl
 				// Unregister any native components, might have cached state based on properties we are going to overwrite
 				NewActor->UnregisterAllComponents(); 
 
-				UEditorEngine::CopyPropertiesForUnrelatedObjects(OldActor, NewActor);
+				UEngine::FCopyPropertiesForUnrelatedObjectsParams Params;
+				Params.bPreserveRootComponent = bPreserveRootComponent;
+				UEngine::CopyPropertiesForUnrelatedObjects(OldActor, NewActor, Params);
+
 				// reset properties/streams
 				NewActor->ResetPropertiesForConstruction(); 
 				// register native components
diff --git a/Engine/Source/Editor/UnrealEd/Public/Kismet2/KismetReinstanceUtilities.h b/Engine/Source/Editor/UnrealEd/Public/Kismet2/KismetReinstanceUtilities.h
--- a/Engine/Source/Editor/UnrealEd/Public/Kismet2/KismetReinstanceUtilities.h
+++ b/Engine/Source/Editor/UnrealEd/Public/Kismet2/KismetReinstanceUtilities.h
@@ -76,7 +76,7 @@ public:
 	void UpdateBytecodeReferences();
 
 	/** Worker function to replace all instances of OldClass with a new instance of NewClass */
-	static void ReplaceInstancesOfClass(UClass* OldClass, UClass* NewClass, UObject* OriginalCDO = NULL, TSet<UObject*>* ObjectsThatShouldUseOldStuff = NULL, bool bClassObjectReplaced = false);
+	static void ReplaceInstancesOfClass(UClass* OldClass, UClass* NewClass, UObject* OriginalCDO = NULL, TSet<UObject*>* ObjectsThatShouldUseOldStuff = NULL, bool bClassObjectReplaced = false, bool bPreserveRootComponent = true);
 
 	/**
 	 * When re-instancing a component, we have to make sure all instance owners' 
@@ -115,6 +115,9 @@ protected:
 	/** Reparents the specified blueprint or class to be the duplicated class in order to allow properties to be copied from the previous CDO to the new one */
 	void ReparentChild(UBlueprint* ChildBP);
 	void ReparentChild(UClass* ChildClass);
+
+	/** Determine whether reinstancing actors should preserve the root component of the new actor */
+	virtual bool ShouldPreserveRootComponentOfReinstancedActor() const { return true; }
 };
 
 
diff --git a/Engine/Source/Runtime/Engine/Classes/Engine/Engine.h b/Engine/Source/Runtime/Engine/Classes/Engine/Engine.h
--- a/Engine/Source/Runtime/Engine/Classes/Engine/Engine.h
+++ b/Engine/Source/Runtime/Engine/Classes/Engine/Engine.h
@@ -2170,12 +2170,14 @@ public:
 		bool bDoDelta;
 		bool bReplaceObjectClassReferences;
 		bool bCopyDeprecatedProperties;
+		bool bPreserveRootComponent;
 
 		FCopyPropertiesForUnrelatedObjectsParams()
 			: bAggressiveDefaultSubobjectReplacement(false)
 			, bDoDelta(true)
 			, bReplaceObjectClassReferences(true)
 			, bCopyDeprecatedProperties(false)
+			, bPreserveRootComponent(true)
 		{}
 	};
 	static void CopyPropertiesForUnrelatedObjects(UObject* OldObject, UObject* NewObject, FCopyPropertiesForUnrelatedObjectsParams Params = FCopyPropertiesForUnrelatedObjectsParams());//bool bAggressiveDefaultSubobjectReplacement = false, bool bDoDelta = true);
diff --git a/Engine/Source/Runtime/Engine/Private/UnrealEngine.cpp b/Engine/Source/Runtime/Engine/Private/UnrealEngine.cpp
--- a/Engine/Source/Runtime/Engine/Private/UnrealEngine.cpp
+++ b/Engine/Source/Runtime/Engine/Private/UnrealEngine.cpp
@@ -10519,7 +10519,7 @@ void UEngine::CopyPropertiesForUnrelatedObjects(UObject* OldObject, UObject* New
 	// If the new object is an Actor, save the root component reference, to be restored later
 	USceneComponent* SavedRootComponent = nullptr;
 	UObjectProperty* RootComponentProperty = nullptr;
-	if (NewActor != nullptr)
+	if (NewActor != nullptr && Params.bPreserveRootComponent)
 	{
 		RootComponentProperty = FindField<UObjectProperty>(NewActor->GetClass(), "RootComponent");
 		if (RootComponentProperty != nullptr)
@@ -10665,7 +10665,7 @@ void UEngine::CopyPropertiesForUnrelatedObjects(UObject* OldObject, UObject* New
 	}
 
 	// Restore the root component reference
-	if (NewActor != nullptr)
+	if (NewActor != nullptr && Params.bPreserveRootComponent)
 	{
 		if (RootComponentProperty != nullptr)
 		{
