Fix skel meshes not allowing editing of any materials if none were found during FBX import.  Now ensure the correct number of material slots are always created

diff --git a/Engine/Source/Editor/UnrealEd/Private/Factories/SkeletalMeshImport.cpp b/Engine/Source/Editor/UnrealEd/Private/Factories/SkeletalMeshImport.cpp
--- a/Engine/Source/Editor/UnrealEd/Private/Factories/SkeletalMeshImport.cpp
+++ b/Engine/Source/Editor/UnrealEd/Private/Factories/SkeletalMeshImport.cpp
@@ -176,8 +176,10 @@ void ProcessImportMeshMaterials(TArray<FSkeletalMaterial>& Materials, FSkeletalM
 		Materials.Add( FSkeletalMaterial( Material, bEnableShadowCasting ) );
 	}
 
-	// Pad the material pointers.
-	while( ImportedMaterials.Num() > Materials.Num() )
+	int32 NumMaterialsToAdd = FMath::Max<int32>( ImportedMaterials.Num(), ImportData.MaxMaterialIndex + 1 );
+
+	// Pad the material pointers
+	while( NumMaterialsToAdd > Materials.Num() )
 	{
 		Materials.Add( FSkeletalMaterial( NULL, true ) );
 	}
diff --git a/Engine/Source/Editor/UnrealEd/Private/Fbx/FbxSkeletalMeshImport.cpp b/Engine/Source/Editor/UnrealEd/Private/Fbx/FbxSkeletalMeshImport.cpp
--- a/Engine/Source/Editor/UnrealEd/Private/Fbx/FbxSkeletalMeshImport.cpp
+++ b/Engine/Source/Editor/UnrealEd/Private/Fbx/FbxSkeletalMeshImport.cpp
@@ -2151,7 +2151,8 @@ bool UnFbx::FFbxImporter::FillSkelMeshImporterFromFbx( FSkeletalMeshImportData&
 				Triangle.MatIndex = 0;
 			}
 		}
-	
+		ImportData.MaxMaterialIndex = FMath::Max<uint32>( ImportData.MaxMaterialIndex, Triangle.MatIndex );
+
 		Triangle.AuxMatIndex = 0;
 		for (int32 VertexIndex=0; VertexIndex<3; VertexIndex++)
 		{
diff --git a/Engine/Source/Editor/UnrealEd/Public/SkelImport.h b/Engine/Source/Editor/UnrealEd/Public/SkelImport.h
--- a/Engine/Source/Editor/UnrealEd/Public/SkelImport.h
+++ b/Engine/Source/Editor/UnrealEd/Public/SkelImport.h
@@ -154,12 +154,14 @@ public:
 	TArray <VRawBoneInfluence>	Influences;		// Influences
 	TArray <int32>				PointToRawMap;	// Mapping from current point index to the original import point index
 	uint32	NumTexCoords;						// The number of texture coordinate sets
+	uint32	MaxMaterialIndex;					// The max material index found on a triangle
 	bool 	bHasVertexColors; 					// If true there are vertex colors in the imported file
 	bool	bHasNormals;						// If true there are normals in the imported file
 	bool	bHasTangents;						// If true there are tangents in the imported file
 
 	FSkeletalMeshImportData()
 		: NumTexCoords(0)
+		, MaxMaterialIndex(0)
 		, bHasVertexColors(false)
 		, bHasNormals(false)
 		, bHasTangents(false)
