Fix BSP not rebuilding in some cases

https://jira.ol.epicgames.net/browse/UE-13720
https://jira.ol.epicgames.net/browse/UE-13849

diff --git a/Engine/Source/Editor/DetailCustomizations/Private/BrushDetails.cpp b/Engine/Source/Editor/DetailCustomizations/Private/BrushDetails.cpp
--- a/Engine/Source/Editor/DetailCustomizations/Private/BrushDetails.cpp
+++ b/Engine/Source/Editor/DetailCustomizations/Private/BrushDetails.cpp
@@ -112,6 +112,8 @@ void FBrushDetails::CustomizeDetails( IDetailLayoutBuilder& DetailLayout )
 
 		static void OnClassPicked(UClass* InChosenClass, TSharedRef<IPropertyHandle> BrushBuilderHandle, IDetailLayoutBuilder* InDetailLayout)
 		{
+			FSlateApplication::Get().DismissAllMenus();
+
 			TArray<UObject*> OuterObjects;
 			BrushBuilderHandle->GetOuterObjects(OuterObjects);
 
@@ -145,9 +147,11 @@ void FBrushDetails::CustomizeDetails( IDetailLayoutBuilder& DetailLayout )
 				{
 					NewObject.Builder->Build(NewObject.Brush->GetWorld(), NewObject.Brush);
 				}
+
+				GEditor->RebuildAlteredBSP();
 			}
 
-			FSlateApplication::Get().DismissAllMenus();
+
 
 			InDetailLayout->ForceRefreshDetails();
 		}
diff --git a/Engine/Source/Editor/GeometryMode/Private/GeometryModifiers.cpp b/Engine/Source/Editor/GeometryMode/Private/GeometryModifiers.cpp
--- a/Engine/Source/Editor/GeometryMode/Private/GeometryModifiers.cpp
+++ b/Engine/Source/Editor/GeometryMode/Private/GeometryModifiers.cpp
@@ -1377,8 +1377,14 @@ void UGeomModifier_Pen::Apply()
 		//tool->SetCurrentModifier( tool->GetModifier(0) );
 
 		//force a rebuild of the brush (otherwise the auto-update will do it and this will result in the undo buffer being incorrect)
+
+
+		ABrush::SetNeedRebuild(ResultingBrush->GetLevel());
+		
 		FBSPOps::RebuildBrush( ResultingBrush->Brush );
 
+		GEditor->RebuildAlteredBSP();
+
 		GEditor->RedrawLevelEditingViewports(true);
 	}
 }
@@ -1565,7 +1571,7 @@ bool UGeomModifier_Pen::InputKey(FEditorViewportClient* ViewportClient, FViewpor
 		}
 		else if( Key == EKeys::Enter )
 		{
-			if (!DoesFinalLineIntersectWithShape(ShapeVertices, ShapeVertices[0]))
+			if (ShapeVertices.Num() > 0 && !DoesFinalLineIntersectWithShape(ShapeVertices, ShapeVertices[0]))
 			{
 				Apply();
 				bResult = true;
diff --git a/Engine/Source/Editor/UnrealEd/Private/EditorServer.cpp b/Engine/Source/Editor/UnrealEd/Private/EditorServer.cpp
--- a/Engine/Source/Editor/UnrealEd/Private/EditorServer.cpp
+++ b/Engine/Source/Editor/UnrealEd/Private/EditorServer.cpp
@@ -1627,72 +1627,72 @@ void UEditorEngine::RebuildAlteredBSP()
 {
 	if( GUndo && !GIsTransacting )
 	{
-	// Early out if BSP auto-updating is disabled
-	if (!GetDefault<ULevelEditorMiscSettings>()->bBSPAutoUpdate)
-	{
-		return;
-	}
-
-	FlushRenderingCommands();
+		// Early out if BSP auto-updating is disabled
+		if (!GetDefault<ULevelEditorMiscSettings>()->bBSPAutoUpdate)
+		{
+			return;
+		}
 
-	// A list of all the levels that need to be rebuilt
-	TArray< TWeakObjectPtr< ULevel > > LevelsToRebuild;
-		ABrush::NeedsRebuild(&LevelsToRebuild);
+		FlushRenderingCommands();
 
-	// Determine which levels need to be rebuilt
-	for (FSelectionIterator It(GetSelectedActorIterator()); It; ++It)
-	{
-		AActor* Actor = static_cast<AActor*>(*It);
-			checkSlow(Actor->IsA(AActor::StaticClass()));
+		// A list of all the levels that need to be rebuilt
+		TArray< TWeakObjectPtr< ULevel > > LevelsToRebuild;
+			ABrush::NeedsRebuild(&LevelsToRebuild);
 
-			ABrush* SelectedBrush = Cast< ABrush >(Actor);
-		if (SelectedBrush && !FActorEditorUtils::IsABuilderBrush(Actor))
+		// Determine which levels need to be rebuilt
+		for (FSelectionIterator It(GetSelectedActorIterator()); It; ++It)
 		{
-			ULevel* Level = SelectedBrush->GetLevel();
-			if (Level)
+			AActor* Actor = static_cast<AActor*>(*It);
+				checkSlow(Actor->IsA(AActor::StaticClass()));
+
+				ABrush* SelectedBrush = Cast< ABrush >(Actor);
+			if (SelectedBrush && !FActorEditorUtils::IsABuilderBrush(Actor))
 			{
-				LevelsToRebuild.AddUnique(Level);
+				ULevel* Level = SelectedBrush->GetLevel();
+				if (Level)
+				{
+					LevelsToRebuild.AddUnique(Level);
+				}
 			}
-		}
-		else
-		{
-			// In addition to any selected brushes, any brushes attached to a selected actor should be rebuilt
-			TArray<AActor*> AttachedActors;
-				Actor->GetAttachedActors(AttachedActors);
-
-			const bool bExactClass = true;
-			TArray<AActor*> AttachedBrushes;
-			// Get any brush actors attached to the selected actor
-				if (ContainsObjectOfClass(AttachedActors, ABrush::StaticClass(), bExactClass, &AttachedBrushes))
+			else
 			{
-					for (int32 BrushIndex = 0; BrushIndex < AttachedBrushes.Num(); ++BrushIndex)
+				// In addition to any selected brushes, any brushes attached to a selected actor should be rebuilt
+				TArray<AActor*> AttachedActors;
+					Actor->GetAttachedActors(AttachedActors);
+
+				const bool bExactClass = true;
+				TArray<AActor*> AttachedBrushes;
+				// Get any brush actors attached to the selected actor
+					if (ContainsObjectOfClass(AttachedActors, ABrush::StaticClass(), bExactClass, &AttachedBrushes))
 				{
-						ULevel* Level = CastChecked<ABrush>(AttachedBrushes[BrushIndex])->GetLevel();
-					if (Level)
+						for (int32 BrushIndex = 0; BrushIndex < AttachedBrushes.Num(); ++BrushIndex)
 					{
-						LevelsToRebuild.AddUnique(Level);
+							ULevel* Level = CastChecked<ABrush>(AttachedBrushes[BrushIndex])->GetLevel();
+						if (Level)
+						{
+							LevelsToRebuild.AddUnique(Level);
+						}
 					}
 				}
+
 			}
 
 		}
 
-	}
-
-	// Rebuild the levels
-	for (int32 LevelIdx = 0; LevelIdx < LevelsToRebuild.Num(); ++LevelIdx)
-	{
-		TWeakObjectPtr< ULevel > LevelToRebuild = LevelsToRebuild[LevelIdx];
-			if (LevelToRebuild.IsValid())
+		// Rebuild the levels
+		for (int32 LevelIdx = 0; LevelIdx < LevelsToRebuild.Num(); ++LevelIdx)
 		{
-			RebuildLevel(*LevelToRebuild.Get());
+			TWeakObjectPtr< ULevel > LevelToRebuild = LevelsToRebuild[LevelIdx];
+				if (LevelToRebuild.IsValid())
+			{
+				RebuildLevel(*LevelToRebuild.Get());
+			}
 		}
-	}
 
-	RedrawLevelEditingViewports();
+		RedrawLevelEditingViewports();
 
-	ABrush::OnRebuildDone();
-}
+		ABrush::OnRebuildDone();
+	}
 	else
 	{
  		ensureMsgf(0, TEXT("Rebuild BSP ignored. Not in a transaction") );
