Fix level editor always being focused when stopping PIE

https://jira.ol.epicgames.net/browse/UE-13315

diff --git a/Engine/Source/Editor/UnrealEd/Private/Kismet2/DebuggerCommands.cpp b/Engine/Source/Editor/UnrealEd/Private/Kismet2/DebuggerCommands.cpp
--- a/Engine/Source/Editor/UnrealEd/Private/Kismet2/DebuggerCommands.cpp
+++ b/Engine/Source/Editor/UnrealEd/Private/Kismet2/DebuggerCommands.cpp
@@ -162,8 +162,11 @@ static void LeaveDebuggingMode()
 		GUnrealEd->PlayWorld->bDebugPauseExecution = false;
 	}
 
-	// Focus the game view port 
-	FModuleManager::GetModuleChecked<FLevelEditorModule>("LevelEditor").FocusPIEViewport();
+	if( FSlateApplication::Get().InKismetDebuggingMode() )
+	{
+		// Focus the game view port when resuming from debugging
+		FModuleManager::GetModuleChecked<FLevelEditorModule>("LevelEditor").FocusPIEViewport();
+	}
 
 	// Tell the application to stop ticking in this stack frame
 	FSlateApplication::Get().LeaveDebuggingMode( FKismetDebugUtilities::IsSingleStepping() );
