Fixed delta translate not being undoable.

https://jira.ol.epicgames.net/browse/UE-13847

diff --git a/Engine/Source/Editor/UnrealEd/Private/UnrealEdSrv.cpp b/Engine/Source/Editor/UnrealEd/Private/UnrealEdSrv.cpp
--- a/Engine/Source/Editor/UnrealEd/Private/UnrealEdSrv.cpp
+++ b/Engine/Source/Editor/UnrealEd/Private/UnrealEdSrv.cpp
@@ -2328,6 +2328,7 @@ bool UUnrealEdEngine::Exec_Actor( UWorld* InWorld, const TCHAR* Str, FOutputDevi
 	}
 	else if( FParse::Command(&Str,TEXT("DELTAMOVE")) )
 	{
+		const FScopedTransaction Transaction(NSLOCTEXT("UnrealEd", "DeltaMoveActors", "Move Actors by Delta"));
 		FVector DeltaMove = FVector::ZeroVector;
 		GetFVECTOR( Str, DeltaMove );
 
