Removed the Blueprint Editor's legacy node menu system.

diff --git a/Engine/Plugins/Experimental/LiveEditor/Source/LiveEditor/Private/LiveEditor.cpp b/Engine/Plugins/Experimental/LiveEditor/Source/LiveEditor/Private/LiveEditor.cpp
--- a/Engine/Plugins/Experimental/LiveEditor/Source/LiveEditor/Private/LiveEditor.cpp
+++ b/Engine/Plugins/Experimental/LiveEditor/Source/LiveEditor/Private/LiveEditor.cpp
@@ -9,7 +9,6 @@
 #include "BlueprintEditor.h"
 #include "BlueprintEditorModes.h"
 #include "BlueprintEditorTabs.h"
-#include "Editor/BlueprintGraph/Public/K2ActionMenuBuilder.h" // for FK2ActionMenuBuilder::AddNewNodeAction()
 #include "SDockTab.h"
 
 DEFINE_LOG_CATEGORY_STATIC(LiveEditor, Log, All);
@@ -85,26 +84,16 @@ class FLiveEditorNodeInjector : public TSharedFromThis<FLiveEditorNodeInjector>
 public:
 	FLiveEditorNodeInjector() {}
 	void InstallHooks();
-	void AddContextMenuNodes(FBlueprintGraphActionListBuilder &ContextMenuBuilder);
 	void OnBeginPIE(const bool bIsSimulating);
 	void OnEndPIE(const bool bIsSimulating);
 };
 
 void FLiveEditorNodeInjector::InstallHooks()
 {
-	FEditorDelegates::OnBlueprintContextMenuCreated.AddSP( this, &FLiveEditorNodeInjector::AddContextMenuNodes );
 	FEditorDelegates::BeginPIE.AddRaw(this, &FLiveEditorNodeInjector::OnBeginPIE);
 	FEditorDelegates::EndPIE.AddRaw(this, &FLiveEditorNodeInjector::OnEndPIE);
 }
 
-void FLiveEditorNodeInjector::AddContextMenuNodes(FBlueprintGraphActionListBuilder &ContextMenuBuilder)
-{
-	UK2Node* NodeTemplate = ContextMenuBuilder.CreateTemplateNode<UK2Node_LiveEditObject>();
-	FString FunctionCategory( TEXT("LiveEditor") );
-	TSharedPtr<FEdGraphSchemaAction_K2NewNode> Action = FK2ActionMenuBuilder::AddNewNodeAction(ContextMenuBuilder, FunctionCategory, NSLOCTEXT("LiveEditorPlugin", "LiveEditObject", "LiveEdit Object"), NodeTemplate->GetTooltipText().ToString(), 0, NodeTemplate->GetKeywords());
-	Action->NodeTemplate = NodeTemplate;
-}
-
 void FLiveEditorNodeInjector::OnBeginPIE(const bool bIsSimulating)
 {
 	FLiveEditorManager::Get().OnBeginPIE();
diff --git a/Engine/Source/Editor/AIGraph/Private/SGraphEditorActionMenuAI.cpp b/Engine/Source/Editor/AIGraph/Private/SGraphEditorActionMenuAI.cpp
--- a/Engine/Source/Editor/AIGraph/Private/SGraphEditorActionMenuAI.cpp
+++ b/Engine/Source/Editor/AIGraph/Private/SGraphEditorActionMenuAI.cpp
@@ -2,7 +2,6 @@
 
 #include "AIGraphPrivatePCH.h"
 #include "SGraphEditorActionMenuAI.h"
-#include "K2ActionMenuBuilder.h" // for FBlueprintGraphActionListBuilder
 #include "AIGraphSchema.h"
 
 SGraphEditorActionMenuAI::~SGraphEditorActionMenuAI()
@@ -41,7 +40,7 @@ void SGraphEditorActionMenuAI::Construct( const FArguments& InArgs )
 void SGraphEditorActionMenuAI::CollectAllActions(FGraphActionListBuilderBase& OutAllActions)
 {
 	// Build up the context object
-	FBlueprintGraphActionListBuilder ContextMenuBuilder(GraphObj);
+	FGraphContextMenuBuilder ContextMenuBuilder(GraphObj);
 	if (GraphNode != NULL)
 	{
 		ContextMenuBuilder.SelectedObjects.Add(GraphNode);
diff --git a/Engine/Source/Editor/AnimGraph/Classes/AnimGraphNode_Base.h b/Engine/Source/Editor/AnimGraph/Classes/AnimGraphNode_Base.h
--- a/Engine/Source/Editor/AnimGraph/Classes/AnimGraphNode_Base.h
+++ b/Engine/Source/Editor/AnimGraph/Classes/AnimGraphNode_Base.h
@@ -111,7 +111,6 @@ class UAnimGraphNode_Base : public UK2Node
 	virtual bool CanPlaceBreakpoints() const override { return false; }
 	ANIMGRAPH_API virtual void ReallocatePinsDuringReconstruction(TArray<UEdGraphPin*>& OldPins) override;
 	ANIMGRAPH_API virtual bool CanCreateUnderSpecifiedSchema(const UEdGraphSchema* DesiredSchema) const override;
-	ANIMGRAPH_API virtual void GetMenuEntries(FGraphContextMenuBuilder& ContextMenuBuilder) const override;
 	ANIMGRAPH_API virtual void GetNodeAttributes(TArray<TKeyValuePair<FString, FString>>& OutNodeAttributes) const override;
 	ANIMGRAPH_API virtual void GetMenuActions(FBlueprintActionDatabaseRegistrar& ActionRegistrar) const override;
 	ANIMGRAPH_API virtual FText GetMenuCategory() const override;
@@ -198,9 +197,6 @@ protected:
 	//
 	virtual ANIMGRAPH_API void GetPinAssociatedProperty(const UScriptStruct* NodeType, UEdGraphPin* InputPin, UProperty*& OutProperty, int32& OutIndex);
 
-	// Creates a default menu entry
-	ANIMGRAPH_API TSharedPtr<FEdGraphSchemaAction_K2NewNode> CreateDefaultMenuEntry(FGraphContextMenuBuilder& ContextMenuBuilder) const;
-
 	// Allocates or reallocates pins
 	ANIMGRAPH_API void InternalPinCreation(TArray<UEdGraphPin*>* OldPins);
 };
diff --git a/Engine/Source/Editor/AnimGraph/Classes/AnimGraphNode_BlendListByEnum.h b/Engine/Source/Editor/AnimGraph/Classes/AnimGraphNode_BlendListByEnum.h
--- a/Engine/Source/Editor/AnimGraph/Classes/AnimGraphNode_BlendListByEnum.h
+++ b/Engine/Source/Editor/AnimGraph/Classes/AnimGraphNode_BlendListByEnum.h
@@ -34,7 +34,6 @@ public:
 
 	// UAnimGraphNode_Base interface
 	virtual FString GetNodeCategory() const override;
-	virtual void GetMenuEntries(FGraphContextMenuBuilder& ContextMenuBuilder) const override;
 	virtual void GetMenuActions(FBlueprintActionDatabaseRegistrar& ActionRegistrar) const override;
 	virtual void CustomizePinData(UEdGraphPin* Pin, FName SourcePropertyName, int32 ArrayIndex) const override;
 	virtual void ValidateAnimNodeDuringCompilation(class USkeleton* ForSkeleton, class FCompilerResultsLog& MessageLog) override;
diff --git a/Engine/Source/Editor/AnimGraph/Classes/AnimGraphNode_BlendSpaceBase.h b/Engine/Source/Editor/AnimGraph/Classes/AnimGraphNode_BlendSpaceBase.h
--- a/Engine/Source/Editor/AnimGraph/Classes/AnimGraphNode_BlendSpaceBase.h
+++ b/Engine/Source/Editor/AnimGraph/Classes/AnimGraphNode_BlendSpaceBase.h
@@ -21,7 +21,5 @@ class UAnimGraphNode_BlendSpaceBase : public UAnimGraphNode_Base
 	// End of UAnimGraphNode_Base interface
 
 protected:
-	static void GetBlendSpaceEntries(bool bWantAimOffsets, FGraphContextMenuBuilder& ContextMenuBuilder);
-
 	UBlendSpaceBase* GetBlendSpace() const { return Cast<UBlendSpaceBase>(GetAnimationAsset()); }
 };
diff --git a/Engine/Source/Editor/AnimGraph/Classes/AnimGraphNode_BlendSpaceEvaluator.h b/Engine/Source/Editor/AnimGraph/Classes/AnimGraphNode_BlendSpaceEvaluator.h
--- a/Engine/Source/Editor/AnimGraph/Classes/AnimGraphNode_BlendSpaceEvaluator.h
+++ b/Engine/Source/Editor/AnimGraph/Classes/AnimGraphNode_BlendSpaceEvaluator.h
@@ -24,7 +24,6 @@ class UAnimGraphNode_BlendSpaceEvaluator : public UAnimGraphNode_BlendSpaceBase
 	// End of UEdGraphNode interface
 
 	// UAnimGraphNode_Base interface
-	virtual void GetMenuEntries(FGraphContextMenuBuilder& ContextMenuBuilder) const override;
 	virtual void ValidateAnimNodeDuringCompilation(class USkeleton* ForSkeleton, class FCompilerResultsLog& MessageLog) override;
 	virtual void BakeDataDuringCompilation(class FCompilerResultsLog& MessageLog) override;
 	// Interface to support transition getter
diff --git a/Engine/Source/Editor/AnimGraph/Classes/AnimGraphNode_BlendSpacePlayer.h b/Engine/Source/Editor/AnimGraph/Classes/AnimGraphNode_BlendSpacePlayer.h
--- a/Engine/Source/Editor/AnimGraph/Classes/AnimGraphNode_BlendSpacePlayer.h
+++ b/Engine/Source/Editor/AnimGraph/Classes/AnimGraphNode_BlendSpacePlayer.h
@@ -24,7 +24,6 @@ class UAnimGraphNode_BlendSpacePlayer : public UAnimGraphNode_BlendSpaceBase
 	// End of UEdGraphNode interface
 
 	// UAnimGraphNode_Base interface
-	virtual void GetMenuEntries(FGraphContextMenuBuilder& ContextMenuBuilder) const override;
 	virtual void ValidateAnimNodeDuringCompilation(class USkeleton* ForSkeleton, class FCompilerResultsLog& MessageLog) override;
 	virtual void BakeDataDuringCompilation(class FCompilerResultsLog& MessageLog) override;
 	virtual bool DoesSupportTimeForTransitionGetter() const override;
diff --git a/Engine/Source/Editor/AnimGraph/Classes/AnimGraphNode_Root.h b/Engine/Source/Editor/AnimGraph/Classes/AnimGraphNode_Root.h
--- a/Engine/Source/Editor/AnimGraph/Classes/AnimGraphNode_Root.h
+++ b/Engine/Source/Editor/AnimGraph/Classes/AnimGraphNode_Root.h
@@ -19,7 +19,6 @@ class UAnimGraphNode_Root : public UAnimGraphNode_Base
 	virtual FText GetNodeTitle(ENodeTitleType::Type TitleType) const override;
 	virtual bool CanUserDeleteNode() const override { return false; }
 	virtual bool CanDuplicateNode() const override { return false; }
-	virtual void GetMenuEntries(FGraphContextMenuBuilder& ContextMenuBuilder) const override;
 	virtual void GetMenuActions(FBlueprintActionDatabaseRegistrar& ActionRegistrar) const override;
 	// End UEdGraphNode interface.
 
diff --git a/Engine/Source/Editor/AnimGraph/Classes/AnimGraphNode_RotationOffsetBlendSpace.h b/Engine/Source/Editor/AnimGraph/Classes/AnimGraphNode_RotationOffsetBlendSpace.h
--- a/Engine/Source/Editor/AnimGraph/Classes/AnimGraphNode_RotationOffsetBlendSpace.h
+++ b/Engine/Source/Editor/AnimGraph/Classes/AnimGraphNode_RotationOffsetBlendSpace.h
@@ -20,7 +20,6 @@ class UAnimGraphNode_RotationOffsetBlendSpace: public UAnimGraphNode_BlendSpaceB
 	// End of UEdGraphNode interface
 
 	// UAnimGraphNode_Base interface
-	virtual void GetMenuEntries(FGraphContextMenuBuilder& ContextMenuBuilder) const override;
 	virtual void ValidateAnimNodeDuringCompilation(class USkeleton* ForSkeleton, class FCompilerResultsLog& MessageLog) override;
 	virtual UAnimationAsset* GetAnimationAsset() const override { return Node.BlendSpace; }
 	virtual void GetAllAnimationSequencesReferred(TArray<UAnimationAsset*>& ComplexAnims, TArray<UAnimSequence*>& AnimationSequences) const override;
diff --git a/Engine/Source/Editor/AnimGraph/Classes/AnimGraphNode_SaveCachedPose.h b/Engine/Source/Editor/AnimGraph/Classes/AnimGraphNode_SaveCachedPose.h
--- a/Engine/Source/Editor/AnimGraph/Classes/AnimGraphNode_SaveCachedPose.h
+++ b/Engine/Source/Editor/AnimGraph/Classes/AnimGraphNode_SaveCachedPose.h
@@ -31,7 +31,6 @@ class UAnimGraphNode_SaveCachedPose : public UAnimGraphNode_Base
 
 	// UAnimGraphNode_Base interface
 	virtual FString GetNodeCategory() const override;
-	virtual void GetMenuEntries(FGraphContextMenuBuilder& ContextMenuBuilder) const override;
 	virtual bool IsSinkNode() const override { return true; }
 	// End of UAnimGraphNode_Base interface
 
diff --git a/Engine/Source/Editor/AnimGraph/Classes/AnimGraphNode_SequenceEvaluator.h b/Engine/Source/Editor/AnimGraph/Classes/AnimGraphNode_SequenceEvaluator.h
--- a/Engine/Source/Editor/AnimGraph/Classes/AnimGraphNode_SequenceEvaluator.h
+++ b/Engine/Source/Editor/AnimGraph/Classes/AnimGraphNode_SequenceEvaluator.h
@@ -20,7 +20,6 @@ class UAnimGraphNode_SequenceEvaluator : public UAnimGraphNode_Base
 	// End of UEdGraphNode
 
 	// UAnimGraphNode_Base interface
-	virtual void GetMenuEntries(FGraphContextMenuBuilder& ContextMenuBuilder) const override;
 	virtual void ValidateAnimNodeDuringCompilation(class USkeleton* ForSkeleton, class FCompilerResultsLog& MessageLog) override;
 	virtual void PreloadRequiredAssets() override;	
 	// Interface to support transition getter
diff --git a/Engine/Source/Editor/AnimGraph/Classes/AnimGraphNode_SequencePlayer.h b/Engine/Source/Editor/AnimGraph/Classes/AnimGraphNode_SequencePlayer.h
--- a/Engine/Source/Editor/AnimGraph/Classes/AnimGraphNode_SequencePlayer.h
+++ b/Engine/Source/Editor/AnimGraph/Classes/AnimGraphNode_SequencePlayer.h
@@ -26,7 +26,6 @@ class UAnimGraphNode_SequencePlayer : public UAnimGraphNode_Base
 	// End of UEdGraphNode interface
 
 	// UAnimGraphNode_Base interface
-	virtual void GetMenuEntries(FGraphContextMenuBuilder& ContextMenuBuilder) const override;
 	virtual void ValidateAnimNodeDuringCompilation(class USkeleton* ForSkeleton, class FCompilerResultsLog& MessageLog) override;
 	virtual void PreloadRequiredAssets() override;		
 	virtual void BakeDataDuringCompilation(class FCompilerResultsLog& MessageLog) override;
diff --git a/Engine/Source/Editor/AnimGraph/Classes/AnimGraphNode_SkeletalControlBase.h b/Engine/Source/Editor/AnimGraph/Classes/AnimGraphNode_SkeletalControlBase.h
--- a/Engine/Source/Editor/AnimGraph/Classes/AnimGraphNode_SkeletalControlBase.h
+++ b/Engine/Source/Editor/AnimGraph/Classes/AnimGraphNode_SkeletalControlBase.h
@@ -21,7 +21,6 @@ public:
 	// UAnimGraphNode_Base interface
 	ANIMGRAPH_API virtual FString GetNodeCategory() const override;
 	ANIMGRAPH_API virtual void CreateOutputPins() override;
-	ANIMGRAPH_API virtual void GetMenuEntries(FGraphContextMenuBuilder& ContextMenuBuilder) const override;
 	// End of UAnimGraphNode_Base interface
 
 	// Draw function for supporting visualization
diff --git a/Engine/Source/Editor/AnimGraph/Classes/AnimGraphNode_StateMachineBase.h b/Engine/Source/Editor/AnimGraph/Classes/AnimGraphNode_StateMachineBase.h
--- a/Engine/Source/Editor/AnimGraph/Classes/AnimGraphNode_StateMachineBase.h
+++ b/Engine/Source/Editor/AnimGraph/Classes/AnimGraphNode_StateMachineBase.h
@@ -30,7 +30,6 @@ class ANIMGRAPH_API UAnimGraphNode_StateMachineBase : public UAnimGraphNode_Base
 
 	// UAnimGraphNode_Base interface
 	virtual FString GetNodeCategory() const override;
-	virtual void GetMenuEntries(FGraphContextMenuBuilder& ContextMenuBuilder) const override;
 	// End of UAnimGraphNode_Base interface
 
 	//  @return the name of this state machine
diff --git a/Engine/Source/Editor/AnimGraph/Classes/AnimGraphNode_TransitionPoseEvaluator.h b/Engine/Source/Editor/AnimGraph/Classes/AnimGraphNode_TransitionPoseEvaluator.h
--- a/Engine/Source/Editor/AnimGraph/Classes/AnimGraphNode_TransitionPoseEvaluator.h
+++ b/Engine/Source/Editor/AnimGraph/Classes/AnimGraphNode_TransitionPoseEvaluator.h
@@ -24,7 +24,6 @@ class UAnimGraphNode_TransitionPoseEvaluator : public UAnimGraphNode_Base
 
 	// UAnimGraphNode_Base interface
 	virtual FString GetNodeCategory() const override;
-	virtual void GetMenuEntries(FGraphContextMenuBuilder& ContextMenuBuilder) const override;
 	// End of UAnimGraphNode_Base interface
 
 	// UK2Node interface.
diff --git a/Engine/Source/Editor/AnimGraph/Classes/AnimGraphNode_TransitionResult.h b/Engine/Source/Editor/AnimGraph/Classes/AnimGraphNode_TransitionResult.h
--- a/Engine/Source/Editor/AnimGraph/Classes/AnimGraphNode_TransitionResult.h
+++ b/Engine/Source/Editor/AnimGraph/Classes/AnimGraphNode_TransitionResult.h
@@ -19,7 +19,6 @@ class UAnimGraphNode_TransitionResult : public UAnimGraphNode_Base
 	virtual FLinearColor GetNodeTitleColor() const override;
 	virtual FText GetTooltipText() const override;
 	virtual FText GetNodeTitle(ENodeTitleType::Type TitleType) const override;
-	virtual void GetMenuEntries(FGraphContextMenuBuilder& ContextMenuBuilder) const override;
 	// End of UEdGraphNode interface
 
 	// UK2Node interface.
diff --git a/Engine/Source/Editor/AnimGraph/Classes/AnimGraphNode_UseCachedPose.h b/Engine/Source/Editor/AnimGraph/Classes/AnimGraphNode_UseCachedPose.h
--- a/Engine/Source/Editor/AnimGraph/Classes/AnimGraphNode_UseCachedPose.h
+++ b/Engine/Source/Editor/AnimGraph/Classes/AnimGraphNode_UseCachedPose.h
@@ -31,7 +31,6 @@ public:
 
 	// UAnimGraphNode_Base interface
 	virtual FString GetNodeCategory() const override;
-	virtual void GetMenuEntries(FGraphContextMenuBuilder& ContextMenuBuilder) const override;
 	virtual void EarlyValidation(class FCompilerResultsLog& MessageLog) const override;
 	// End of UAnimGraphNode_Base interface
 
diff --git a/Engine/Source/Editor/AnimGraph/Classes/AnimGraphNode_WheelHandler.h b/Engine/Source/Editor/AnimGraph/Classes/AnimGraphNode_WheelHandler.h
--- a/Engine/Source/Editor/AnimGraph/Classes/AnimGraphNode_WheelHandler.h
+++ b/Engine/Source/Editor/AnimGraph/Classes/AnimGraphNode_WheelHandler.h
@@ -21,7 +21,6 @@ public:
 	virtual FString GetKeywords() const override;
 	// validate if this is within VehicleAnimInstance
 	virtual void ValidateAnimNodePostCompile(class FCompilerResultsLog& MessageLog, class UAnimBlueprintGeneratedClass* CompiledClass, int32 CompiledNodeIndex) override;
-	virtual void GetMenuEntries(FGraphContextMenuBuilder& ContextMenuBuilder) const override;
 	virtual bool IsCompatibleWithGraph(const UEdGraph* TargetGraph) const override;
 	// End of UEdGraphNode interface
 
diff --git a/Engine/Source/Editor/AnimGraph/Classes/AnimationGraphSchema.h b/Engine/Source/Editor/AnimGraph/Classes/AnimationGraphSchema.h
--- a/Engine/Source/Editor/AnimGraph/Classes/AnimationGraphSchema.h
+++ b/Engine/Source/Editor/AnimGraph/Classes/AnimationGraphSchema.h
@@ -36,7 +36,6 @@ class UAnimationGraphSchema : public UEdGraphSchema_K2
 	FName DefaultEvaluationHandlerName;
 
 	// Begin UEdGraphSchema interface.
-	virtual void GetGraphContextActions(FGraphContextMenuBuilder& ContextMenuBuilder) const override;
 	virtual FLinearColor GetPinTypeColor(const FEdGraphPinType& PinType) const override;
 	virtual EGraphType GetGraphType(const UEdGraph* TestEdGraph) const override;
 	virtual void CreateDefaultNodesForGraph(UEdGraph& Graph) const override;
@@ -68,9 +67,6 @@ class UAnimationGraphSchema : public UEdGraphSchema_K2
 	/** Update the specified node to a new asset */
 	static void UpdateNodeWithAsset(class UK2Node* K2Node, UAnimationAsset* Asset);
 
-	// @todo document
-	void GetStateMachineMenuItems(FGraphContextMenuBuilder& ContextMenuBuilder) const;
-
 	// @todo document
 	ANIMGRAPH_API static bool IsPosePin(const FEdGraphPinType& PinType);
 
diff --git a/Engine/Source/Editor/AnimGraph/Classes/AnimationTransitionSchema.h b/Engine/Source/Editor/AnimGraph/Classes/AnimationTransitionSchema.h
--- a/Engine/Source/Editor/AnimGraph/Classes/AnimationTransitionSchema.h
+++ b/Engine/Source/Editor/AnimGraph/Classes/AnimationTransitionSchema.h
@@ -10,7 +10,6 @@ class UAnimationTransitionSchema : public UEdGraphSchema_K2
 	GENERATED_UCLASS_BODY()
 
 	// Begin UEdGraphSchema interface.
-	virtual void GetGraphContextActions(FGraphContextMenuBuilder& ContextMenuBuilder) const override;
 	virtual void CreateDefaultNodesForGraph(UEdGraph& Graph) const override;
 	virtual bool CanDuplicateGraph(UEdGraph* InSourceGraph) const override { return false; }
 	virtual void GetGraphDisplayInformation(const UEdGraph& Graph, /*out*/ FGraphDisplayInfo& DisplayInfo) const override;
@@ -19,8 +18,6 @@ class UAnimationTransitionSchema : public UEdGraphSchema_K2
 	// End UEdGraphSchema interface.
 
 private:
-	void GetSourceStateActions(FGraphContextMenuBuilder& ContextMenuBuilder) const;
-
 	static UAnimStateTransitionNode* GetTransitionNodeFromGraph(const FAnimBlueprintDebugData& DebugData, const UEdGraph* Graph);
 
 	static UAnimStateNode* GetStateNodeFromGraph(const FAnimBlueprintDebugData& DebugData, const UEdGraph* Graph);
diff --git a/Engine/Source/Editor/AnimGraph/Private/AnimGraphNode_Base.cpp b/Engine/Source/Editor/AnimGraph/Private/AnimGraphNode_Base.cpp
--- a/Engine/Source/Editor/AnimGraph/Private/AnimGraphNode_Base.cpp
+++ b/Engine/Source/Editor/AnimGraph/Private/AnimGraphNode_Base.cpp
@@ -3,7 +3,6 @@
 
 #include "AnimGraphPrivatePCH.h"
 #include "AnimGraphNode_Base.h"
-#include "K2ActionMenuBuilder.h" // for FK2ActionMenuBuilder::AddNewNodeAction()
 #include "AnimationGraphSchema.h"
 #include "BlueprintNodeSpawner.h"
 #include "BlueprintActionDatabaseRegistrar.h"
@@ -216,11 +215,6 @@ FString UAnimGraphNode_Base::GetNodeCategory() const
 	return TEXT("Misc.");
 }
 
-void UAnimGraphNode_Base::GetMenuEntries(FGraphContextMenuBuilder& ContextMenuBuilder) const
-{
-	CreateDefaultMenuEntry(ContextMenuBuilder);
-}
-
 void UAnimGraphNode_Base::GetNodeAttributes( TArray<TKeyValuePair<FString, FString>>& OutNodeAttributes ) const
 {
 	OutNodeAttributes.Add( TKeyValuePair<FString, FString>( TEXT( "Type" ), TEXT( "AnimGraphNode" ) ));
@@ -252,21 +246,6 @@ FText UAnimGraphNode_Base::GetMenuCategory() const
 	return FText::FromString(GetNodeCategory());
 }
 
-TSharedPtr<FEdGraphSchemaAction_K2NewNode> UAnimGraphNode_Base::CreateDefaultMenuEntry(FGraphContextMenuBuilder& ContextMenuBuilder) const
-{
-	UAnimGraphNode_Base* TemplateNode = NewObject<UAnimGraphNode_Base>(GetTransientPackage(), GetClass());
-
-	FString Category = TemplateNode->GetNodeCategory();
-	FText MenuDesc = TemplateNode->GetNodeTitle(ENodeTitleType::ListView);
-	FString Tooltip = TemplateNode->GetTooltipText().ToString();
-	FString Keywords = TemplateNode->GetKeywords();
-
-	TSharedPtr<FEdGraphSchemaAction_K2NewNode> NodeAction = FK2ActionMenuBuilder::AddNewNodeAction(ContextMenuBuilder, Category, MenuDesc, Tooltip, 0, Keywords);
-	NodeAction->NodeTemplate = TemplateNode;
-
-	return NodeAction;
-}
-
 void UAnimGraphNode_Base::GetPinAssociatedProperty(const UScriptStruct* NodeType, UEdGraphPin* InputPin, UProperty*& OutProperty, int32& OutIndex)
 {
 	OutProperty = NULL;
diff --git a/Engine/Source/Editor/AnimGraph/Private/AnimGraphNode_BlendListByEnum.cpp b/Engine/Source/Editor/AnimGraph/Private/AnimGraphNode_BlendListByEnum.cpp
--- a/Engine/Source/Editor/AnimGraph/Private/AnimGraphNode_BlendListByEnum.cpp
+++ b/Engine/Source/Editor/AnimGraph/Private/AnimGraphNode_BlendListByEnum.cpp
@@ -5,7 +5,6 @@
 #include "GraphEditorActions.h"
 #include "ScopedTransaction.h"
 #include "CompilerResultsLog.h"
-#include "K2ActionMenuBuilder.h" // for FK2ActionMenuBuilder::AddNewNodeAction()
 #include "AnimGraphNode_BlendListByEnum.h"
 #include "AnimationGraphSchema.h"
 #include "BlueprintActionDatabaseRegistrar.h"
@@ -57,27 +56,6 @@ void UAnimGraphNode_BlendListByEnum::PostPlacedNewNode()
 	Node.AddPose();
 }
 
-void UAnimGraphNode_BlendListByEnum::GetMenuEntries(FGraphContextMenuBuilder& ContextMenuBuilder) const
-{
-	const UAnimationGraphSchema* Schema = GetDefault<UAnimationGraphSchema>();
-
-	// add all blendlist enum entries
-	for (TObjectIterator<UEnum> EnumIt; EnumIt; ++EnumIt)
-	{
-		UEnum* CurrentEnum = *EnumIt;
-
-		const bool bIsBlueprintType = UEdGraphSchema_K2::IsAllowableBlueprintVariableType(CurrentEnum);
-		if (bIsBlueprintType)
-		{
-			UAnimGraphNode_BlendListByEnum* EnumTemplate = NewObject<UAnimGraphNode_BlendListByEnum>();
-			EnumTemplate->BoundEnum = CurrentEnum;
-
-			TSharedPtr<FEdGraphSchemaAction_K2NewNode> Action = FK2ActionMenuBuilder::AddNewNodeAction(ContextMenuBuilder, EnumTemplate->GetNodeCategory(), EnumTemplate->GetNodeTitle(ENodeTitleType::ListView), EnumTemplate->GetTooltipText().ToString(), 0, EnumTemplate->GetKeywords());
-			Action->NodeTemplate = EnumTemplate;
-		}
-	}
-}
-
 void UAnimGraphNode_BlendListByEnum::GetMenuActions(FBlueprintActionDatabaseRegistrar& ActionRegistrar) const
 {
 	auto CustomizeBlendListEnumNodeLambda = [](UEdGraphNode* NewNode, bool bIsTemplateNode, UEnum* EnumValue)
diff --git a/Engine/Source/Editor/AnimGraph/Private/AnimGraphNode_BlendSpaceBase.cpp b/Engine/Source/Editor/AnimGraph/Private/AnimGraphNode_BlendSpaceBase.cpp
--- a/Engine/Source/Editor/AnimGraph/Private/AnimGraphNode_BlendSpaceBase.cpp
+++ b/Engine/Source/Editor/AnimGraph/Private/AnimGraphNode_BlendSpaceBase.cpp
@@ -54,33 +54,6 @@ FLinearColor UAnimGraphNode_BlendSpaceBase::GetNodeTitleColor() const
 	return FLinearColor(0.2f, 0.8f, 0.2f);
 }
 
-void UAnimGraphNode_BlendSpaceBase::GetBlendSpaceEntries(bool bWantsAimOffset, FGraphContextMenuBuilder& ContextMenuBuilder)
-{
-	if ((ContextMenuBuilder.FromPin == NULL) || (UAnimationGraphSchema::IsPosePin(ContextMenuBuilder.FromPin->PinType) && (ContextMenuBuilder.FromPin->Direction == EGPD_Input)))
-	{
-		UBlueprint* Blueprint = FBlueprintEditorUtils::FindBlueprintForGraphChecked(ContextMenuBuilder.CurrentGraph);
-
-		// Add an entry for each loaded animation sequence
-		if (UAnimBlueprint* AnimBlueprint = Cast<UAnimBlueprint>(Blueprint))
-		{
-			for (TObjectIterator<UBlendSpaceBase> BlendSpaceIt; BlendSpaceIt; ++BlendSpaceIt)
-			{
-				UBlendSpaceBase* BlendSpace = *BlendSpaceIt;
-				
-				const bool bIsAimOffset = BlendSpace->IsA(UAimOffsetBlendSpace::StaticClass()) || BlendSpace->IsA(UAimOffsetBlendSpace1D::StaticClass());
-				const bool bPassesAimOffsetFilter = !(bIsAimOffset ^ bWantsAimOffset);
-				const bool bPassesSkeletonFilter = (BlendSpace->GetSkeleton() == AnimBlueprint->TargetSkeleton);
-
-				if (bPassesAimOffsetFilter && bPassesSkeletonFilter)
-				{
-					TSharedPtr<FNewBlendSpacePlayerAction> NewAction(new FNewBlendSpacePlayerAction(BlendSpace));
-					ContextMenuBuilder.AddAction( NewAction );
-				}
-			}
-		}
-	}
-}
-
 void UAnimGraphNode_BlendSpaceBase::CustomizePinData(UEdGraphPin* Pin, FName SourcePropertyName, int32 ArrayIndex) const
 {
 	UBlendSpaceBase * BlendSpace = GetBlendSpace();
diff --git a/Engine/Source/Editor/AnimGraph/Private/AnimGraphNode_BlendSpaceEvaluator.cpp b/Engine/Source/Editor/AnimGraph/Private/AnimGraphNode_BlendSpaceEvaluator.cpp
--- a/Engine/Source/Editor/AnimGraph/Private/AnimGraphNode_BlendSpaceEvaluator.cpp
+++ b/Engine/Source/Editor/AnimGraph/Private/AnimGraphNode_BlendSpaceEvaluator.cpp
@@ -68,12 +68,6 @@ FText UAnimGraphNode_BlendSpaceEvaluator::GetNodeTitle(ENodeTitleType::Type Titl
 	return CachedNodeTitles[TitleType];
 }
 
-void UAnimGraphNode_BlendSpaceEvaluator::GetMenuEntries(FGraphContextMenuBuilder& ContextMenuBuilder) const
-{
-	// Intentionally empty so that we don't get duplicate blend space entries.
-	// You can convert a regular blend space player to an evaluator via the right click context menu
-}
-
 void UAnimGraphNode_BlendSpaceEvaluator::GetMenuActions(FBlueprintActionDatabaseRegistrar& ActionRegistrar) const
 {
 	// Intentionally empty so that we don't get duplicate blend space entries.
diff --git a/Engine/Source/Editor/AnimGraph/Private/AnimGraphNode_BlendSpacePlayer.cpp b/Engine/Source/Editor/AnimGraph/Private/AnimGraphNode_BlendSpacePlayer.cpp
--- a/Engine/Source/Editor/AnimGraph/Private/AnimGraphNode_BlendSpacePlayer.cpp
+++ b/Engine/Source/Editor/AnimGraph/Private/AnimGraphNode_BlendSpacePlayer.cpp
@@ -69,12 +69,6 @@ FText UAnimGraphNode_BlendSpacePlayer::GetNodeTitle(ENodeTitleType::Type TitleTy
 	return CachedNodeTitles[TitleType];
 }
 
-void UAnimGraphNode_BlendSpacePlayer::GetMenuEntries(FGraphContextMenuBuilder& ContextMenuBuilder) const
-{
-	const bool bWantAimOffsets = false;
-	GetBlendSpaceEntries(bWantAimOffsets, ContextMenuBuilder);
-}
-
 void UAnimGraphNode_BlendSpacePlayer::ValidateAnimNodeDuringCompilation(class USkeleton* ForSkeleton, class FCompilerResultsLog& MessageLog)
 {
 	if (Node.BlendSpace == NULL)
diff --git a/Engine/Source/Editor/AnimGraph/Private/AnimGraphNode_Root.cpp b/Engine/Source/Editor/AnimGraph/Private/AnimGraphNode_Root.cpp
--- a/Engine/Source/Editor/AnimGraph/Private/AnimGraphNode_Root.cpp
+++ b/Engine/Source/Editor/AnimGraph/Private/AnimGraphNode_Root.cpp
@@ -63,11 +63,6 @@ bool UAnimGraphNode_Root::IsSinkNode() const
 	return true;
 }
 
-void UAnimGraphNode_Root::GetMenuEntries(FGraphContextMenuBuilder& ContextMenuBuilder) const
-{
-	// Intentionally empty.  This node is autogenerated when a new graph is created.
-}
-
 void UAnimGraphNode_Root::GetMenuActions(FBlueprintActionDatabaseRegistrar& ActionRegistrar) const
 {
 	// Intentionally empty. This node is auto-generated when a new graph is created.
diff --git a/Engine/Source/Editor/AnimGraph/Private/AnimGraphNode_RotationOffsetBlendSpace.cpp b/Engine/Source/Editor/AnimGraph/Private/AnimGraphNode_RotationOffsetBlendSpace.cpp
--- a/Engine/Source/Editor/AnimGraph/Private/AnimGraphNode_RotationOffsetBlendSpace.cpp
+++ b/Engine/Source/Editor/AnimGraph/Private/AnimGraphNode_RotationOffsetBlendSpace.cpp
@@ -58,12 +58,6 @@ FText UAnimGraphNode_RotationOffsetBlendSpace::GetNodeTitle(ENodeTitleType::Type
 	return CachedNodeTitles[TitleType];
 }
 
-void UAnimGraphNode_RotationOffsetBlendSpace::GetMenuEntries(FGraphContextMenuBuilder& ContextMenuBuilder) const
-{
-	const bool bWantAimOffsets = true;
-	GetBlendSpaceEntries(bWantAimOffsets, ContextMenuBuilder);
-}
-
 void UAnimGraphNode_RotationOffsetBlendSpace::GetMenuActions(FBlueprintActionDatabaseRegistrar& ActionRegistrar) const
 {
 	auto PostSpawnSetupLambda = [](UEdGraphNode* NewNode, bool /*bIsTemplateNode*/, TWeakObjectPtr<UBlendSpaceBase> BlendSpace)
diff --git a/Engine/Source/Editor/AnimGraph/Private/AnimGraphNode_SaveCachedPose.cpp b/Engine/Source/Editor/AnimGraph/Private/AnimGraphNode_SaveCachedPose.cpp
--- a/Engine/Source/Editor/AnimGraph/Private/AnimGraphNode_SaveCachedPose.cpp
+++ b/Engine/Source/Editor/AnimGraph/Private/AnimGraphNode_SaveCachedPose.cpp
@@ -5,7 +5,6 @@
 #include "GraphEditorActions.h"
 #include "ScopedTransaction.h"
 #include "Kismet2NameValidators.h"
-#include "K2ActionMenuBuilder.h" // for FK2ActionMenuBuilder::AddNewNodeAction()
 #include "AnimGraphNode_SaveCachedPose.h"
 #include "BlueprintNodeSpawner.h"
 #include "BlueprintActionDatabaseRegistrar.h"
@@ -69,26 +68,6 @@ FString UAnimGraphNode_SaveCachedPose::GetNodeCategory() const
 	return TEXT("Cached Poses");
 }
 
-void UAnimGraphNode_SaveCachedPose::GetMenuEntries(FGraphContextMenuBuilder& ContextMenuBuilder) const
-{
-	//@TODO: Check the type of the from pin to make sure it's a pose
-	if ((ContextMenuBuilder.FromPin == NULL) || (ContextMenuBuilder.FromPin->Direction == EGPD_Output))
-	{
-		// Disallow the graph being cached from residing in a state for safety and sanity; the tree is in effect it's own graph entirely and should be treated that way
-		const bool bNotInStateMachine = ContextMenuBuilder.CurrentGraph->GetOuter()->IsA(UAnimBlueprint::StaticClass());
-
-		if (bNotInStateMachine)
-		{
-			// Offer save cached pose
-			UAnimGraphNode_SaveCachedPose* SaveCachedPose = NewObject<UAnimGraphNode_SaveCachedPose>();
-			SaveCachedPose->CacheName = TEXT("SavedPose") + FString::FromInt(FMath::Rand());
-
-			TSharedPtr<FEdGraphSchemaAction_K2NewNode> SaveCachedPoseAction = FK2ActionMenuBuilder::AddNewNodeAction(ContextMenuBuilder, GetNodeCategory(), LOCTEXT("NewSaveCachedPose", "New Save cached pose..."), SaveCachedPose->GetTooltipText().ToString(), 0, SaveCachedPose->GetKeywords());
-			SaveCachedPoseAction->NodeTemplate = SaveCachedPose;
-		}
-	}
-}
-
 void UAnimGraphNode_SaveCachedPose::GetMenuActions(FBlueprintActionDatabaseRegistrar& ActionRegistrar) const
 {
 	auto PostSpawnSetupLambda = [](UEdGraphNode* NewNode, bool bIsTemplateNode)
diff --git a/Engine/Source/Editor/AnimGraph/Private/AnimGraphNode_SequenceEvaluator.cpp b/Engine/Source/Editor/AnimGraph/Private/AnimGraphNode_SequenceEvaluator.cpp
--- a/Engine/Source/Editor/AnimGraph/Private/AnimGraphNode_SequenceEvaluator.cpp
+++ b/Engine/Source/Editor/AnimGraph/Private/AnimGraphNode_SequenceEvaluator.cpp
@@ -72,11 +72,6 @@ FText UAnimGraphNode_SequenceEvaluator::GetNodeTitle(ENodeTitleType::Type TitleT
 	return CachedNodeTitle;
 }
 
-void UAnimGraphNode_SequenceEvaluator::GetMenuEntries(FGraphContextMenuBuilder& ContextMenuBuilder) const
-{
-	// Intentionally empty; you can drop down a regular sequence player and convert into a sequence evaluator in the right-click menu.
-}
-
 void UAnimGraphNode_SequenceEvaluator::GetMenuActions(FBlueprintActionDatabaseRegistrar& ActionRegistrar) const
 {
 	// Intentionally empty; you can drop down a regular sequence player and convert into a sequence evaluator in the right-click menu.
diff --git a/Engine/Source/Editor/AnimGraph/Private/AnimGraphNode_SequencePlayer.cpp b/Engine/Source/Editor/AnimGraph/Private/AnimGraphNode_SequencePlayer.cpp
--- a/Engine/Source/Editor/AnimGraph/Private/AnimGraphNode_SequencePlayer.cpp
+++ b/Engine/Source/Editor/AnimGraph/Private/AnimGraphNode_SequencePlayer.cpp
@@ -130,52 +130,6 @@ FText UAnimGraphNode_SequencePlayer::GetTitleGivenAssetInfo(const FText& AssetNa
 	}
 }
 
-void UAnimGraphNode_SequencePlayer::GetMenuEntries(FGraphContextMenuBuilder& ContextMenuBuilder) const
-{
-	if ((ContextMenuBuilder.FromPin == NULL) || (UAnimationGraphSchema::IsPosePin(ContextMenuBuilder.FromPin->PinType) && (ContextMenuBuilder.FromPin->Direction == EGPD_Input)))
-	{
-		UBlueprint* Blueprint = FBlueprintEditorUtils::FindBlueprintForGraphChecked(ContextMenuBuilder.CurrentGraph);
-
-		if (UAnimBlueprint* AnimBlueprint = Cast<UAnimBlueprint>(Blueprint))
-		{
-			// Load the asset registry module
-			FAssetRegistryModule& AssetRegistryModule = FModuleManager::LoadModuleChecked<FAssetRegistryModule>(TEXT("AssetRegistry"));
-
-			FARFilter Filter;
-			Filter.ClassNames.Add(UAnimSequence::StaticClass()->GetFName());
-			Filter.bRecursiveClasses = true;
-
-			// Filter by skeleton
-			FAssetData SkeletonData(AnimBlueprint->TargetSkeleton);
-			Filter.TagsAndValues.Add(TEXT("Skeleton"), SkeletonData.GetExportTextName());
-
-			// Find matching assets and add an entry for each one
-			TArray<FAssetData> SequenceList;
-			AssetRegistryModule.Get().GetAssets(Filter, /*out*/ SequenceList);
-
-			for (auto AssetIt = SequenceList.CreateConstIterator(); AssetIt; ++AssetIt)
-			{
-				const FAssetData& Asset = *AssetIt;
-
-				// Try to determine if the asset is additive (can't do it right now if the asset is unloaded)
-				bool bAdditive = false;
-				if (Asset.IsAssetLoaded())
-				{
-					if (UAnimSequence* Sequence = Cast<UAnimSequence>(Asset.GetAsset()))
-					{
-						bAdditive = Sequence->IsValidAdditive();
-					}
-				}
-
-				// Create the menu item
-				const FText Title = UAnimGraphNode_SequencePlayer::GetTitleGivenAssetInfo(FText::FromName(Asset.AssetName), bAdditive);
-				TSharedPtr<FNewSequencePlayerAction> NewAction(new FNewSequencePlayerAction(Asset, Title));
-				ContextMenuBuilder.AddAction( NewAction );
-			}
-		}
-	}
-}
-
 FText UAnimGraphNode_SequencePlayer::GetMenuCategory() const
 {
 	return FEditorCategoryUtils::GetCommonCategory(FCommonEditorCategory::Animation);
diff --git a/Engine/Source/Editor/AnimGraph/Private/AnimGraphNode_SkeletalControlBase.cpp b/Engine/Source/Editor/AnimGraph/Private/AnimGraphNode_SkeletalControlBase.cpp
--- a/Engine/Source/Editor/AnimGraph/Private/AnimGraphNode_SkeletalControlBase.cpp
+++ b/Engine/Source/Editor/AnimGraph/Private/AnimGraphNode_SkeletalControlBase.cpp
@@ -3,7 +3,6 @@
 #include "AnimGraphPrivatePCH.h"
 #include "AnimGraphNode_SkeletalControlBase.h"
 #include "AnimationGraphSchema.h"
-#include "K2ActionMenuBuilder.h" // for FK2ActionMenuBuilder::AddNewNodeAction()
 
 /////////////////////////////////////////////////////
 // UAnimGraphNode_SkeletalControlBase
@@ -57,19 +56,6 @@ void UAnimGraphNode_SkeletalControlBase::CreateOutputPins()
 	CreatePin(EGPD_Output, Schema->PC_Struct, TEXT(""), FComponentSpacePoseLink::StaticStruct(), /*bIsArray=*/ false, /*bIsReference=*/ false, TEXT("Pose"));
 }
 
-void UAnimGraphNode_SkeletalControlBase::GetMenuEntries(FGraphContextMenuBuilder& ContextMenuBuilder) const
-{
-	UAnimGraphNode_Base* TemplateNode = NewObject<UAnimGraphNode_Base>(GetTransientPackage(), GetClass());
-
-	FString Category = TemplateNode->GetNodeCategory();
-	FText MenuDesc = GetControllerDescription();
-	FString Tooltip = GetControllerDescription().ToString();
-	FString Keywords = TemplateNode->GetKeywords();
-
-	TSharedPtr<FEdGraphSchemaAction_K2NewNode> NodeAction = FK2ActionMenuBuilder::AddNewNodeAction(ContextMenuBuilder, Category, MenuDesc, Tooltip, 0, Keywords);
-	NodeAction->NodeTemplate = TemplateNode;
-}
-
 void UAnimGraphNode_SkeletalControlBase::ConvertToComponentSpaceTransform(const USkeletalMeshComponent* SkelComp, const FTransform & InTransform, FTransform & OutCSTransform, int32 BoneIndex, EBoneControlSpace Space) const
 {
 	USkeleton * Skeleton = SkelComp->SkeletalMesh->Skeleton;
diff --git a/Engine/Source/Editor/AnimGraph/Private/AnimGraphNode_StateMachineBase.cpp b/Engine/Source/Editor/AnimGraph/Private/AnimGraphNode_StateMachineBase.cpp
--- a/Engine/Source/Editor/AnimGraph/Private/AnimGraphNode_StateMachineBase.cpp
+++ b/Engine/Source/Editor/AnimGraph/Private/AnimGraphNode_StateMachineBase.cpp
@@ -94,17 +94,6 @@ FString UAnimGraphNode_StateMachineBase::GetNodeCategory() const
 	return TEXT("State Machines");
 }
 
-void UAnimGraphNode_StateMachineBase::GetMenuEntries(FGraphContextMenuBuilder& ContextMenuBuilder) const
-{
-	if ((ContextMenuBuilder.FromPin == NULL) || ((ContextMenuBuilder.FromPin->Direction == EGPD_Input) && (ContextMenuBuilder.FromPin->PinType.PinSubCategoryObject == FPoseLink::StaticStruct())))
-	{
-		TSharedPtr<FEdGraphSchemaAction_K2NewNode> MenuEntry = CreateDefaultMenuEntry(ContextMenuBuilder);
-
-		MenuEntry->MenuDescription = LOCTEXT("AddNewStateMachine", "Add New State Machine...");
-		MenuEntry->TooltipDescription = LOCTEXT("AddNewStateMachine_Tooltip", "Create a new state machine").ToString();
-	}
-}
-
 void UAnimGraphNode_StateMachineBase::PostPlacedNewNode()
 {
 	// Create a new animation graph
diff --git a/Engine/Source/Editor/AnimGraph/Private/AnimGraphNode_TransitionPoseEvaluator.cpp b/Engine/Source/Editor/AnimGraph/Private/AnimGraphNode_TransitionPoseEvaluator.cpp
--- a/Engine/Source/Editor/AnimGraph/Private/AnimGraphNode_TransitionPoseEvaluator.cpp
+++ b/Engine/Source/Editor/AnimGraph/Private/AnimGraphNode_TransitionPoseEvaluator.cpp
@@ -58,11 +58,6 @@ bool UAnimGraphNode_TransitionPoseEvaluator::CanUserDeleteNode() const
 	return !(GetGraph()->IsA(UAnimationCustomTransitionGraph::StaticClass()));
 }
 
-void UAnimGraphNode_TransitionPoseEvaluator::GetMenuEntries(FGraphContextMenuBuilder& ContextMenuBuilder) const
-{
-	// Intentionally empty: Don't allow an option to create them, as they're auto-created when custom blend graphs are made
-}
-
 void UAnimGraphNode_TransitionPoseEvaluator::GetMenuActions(FBlueprintActionDatabaseRegistrar& ActionRegistrar) const
 {
 	// Intentionally empty: Don't allow an option to create them, as they're auto-created when custom blend graphs are made
diff --git a/Engine/Source/Editor/AnimGraph/Private/AnimGraphNode_TransitionResult.cpp b/Engine/Source/Editor/AnimGraph/Private/AnimGraphNode_TransitionResult.cpp
--- a/Engine/Source/Editor/AnimGraph/Private/AnimGraphNode_TransitionResult.cpp
+++ b/Engine/Source/Editor/AnimGraph/Private/AnimGraphNode_TransitionResult.cpp
@@ -32,12 +32,6 @@ FText UAnimGraphNode_TransitionResult::GetNodeTitle(ENodeTitleType::Type TitleTy
 	return LOCTEXT("Result", "Result");
 }
 
-void UAnimGraphNode_TransitionResult::GetMenuEntries(FGraphContextMenuBuilder& ContextMenuBuilder) const
-{
-	// Intentionally empty.  This node is autogenerated when a transition graph is created.
-}
-
-
 void UAnimGraphNode_TransitionResult::GetMenuActions(FBlueprintActionDatabaseRegistrar& ActionRegistrar) const
 {
 	// Intentionally empty. This node is auto-generated when a transition graph is created.
diff --git a/Engine/Source/Editor/AnimGraph/Private/AnimGraphNode_UseCachedPose.cpp b/Engine/Source/Editor/AnimGraph/Private/AnimGraphNode_UseCachedPose.cpp
--- a/Engine/Source/Editor/AnimGraph/Private/AnimGraphNode_UseCachedPose.cpp
+++ b/Engine/Source/Editor/AnimGraph/Private/AnimGraphNode_UseCachedPose.cpp
@@ -4,7 +4,6 @@
 
 #include "GraphEditorActions.h"
 #include "ScopedTransaction.h"
-#include "K2ActionMenuBuilder.h" // for FK2ActionMenuBuilder::AddNewNodeAction()
 #include "EdGraphUtilities.h"
 #include "AnimGraphNode_SaveCachedPose.h"
 #include "AnimGraphNode_UseCachedPose.h"
@@ -93,29 +92,6 @@ FString UAnimGraphNode_UseCachedPose::GetNodeCategory() const
 	return TEXT("Cached Poses");
 }
 
-void UAnimGraphNode_UseCachedPose::GetMenuEntries(FGraphContextMenuBuilder& ContextMenuBuilder) const
-{
-	//@TODO: Check the type of the from pin to make sure it's a pose
-	if ((ContextMenuBuilder.FromPin == NULL) || (ContextMenuBuilder.FromPin->Direction == EGPD_Input))
-	{
-		// Get a list of all save cached pose nodes
-		TArray<UAnimGraphNode_SaveCachedPose*> CachedPoseNodes;
-		FBlueprintEditorUtils::GetAllNodesOfClass<UAnimGraphNode_SaveCachedPose>(FBlueprintEditorUtils::FindBlueprintForGraphChecked(ContextMenuBuilder.CurrentGraph), /*out*/ CachedPoseNodes);
-
-		// Offer a use node for each of them
-		for (auto NodeIt = CachedPoseNodes.CreateIterator(); NodeIt; ++NodeIt)
-		{
-			UAnimGraphNode_UseCachedPose* UseCachedPose = NewObject<UAnimGraphNode_UseCachedPose>();
-			UseCachedPose->NameOfCache = (*NodeIt)->CacheName;
-			UseCachedPose->SaveCachedPoseNode = *NodeIt;
-
-			TSharedPtr<FEdGraphSchemaAction_K2NewNode> UseCachedPoseAction = FK2ActionMenuBuilder::AddNewNodeAction(ContextMenuBuilder, GetNodeCategory(), UseCachedPose->GetNodeTitle(ENodeTitleType::ListView), UseCachedPose->GetTooltipText().ToString(), 0, UseCachedPose->GetKeywords());
-			UseCachedPoseAction->NodeTemplate = UseCachedPose;
-		}
-	}
-}
-
-
 void UAnimGraphNode_UseCachedPose::GetMenuActions(FBlueprintActionDatabaseRegistrar& ActionRegistrar) const
 {
 	auto PostSpawnSetupLambda = [](UEdGraphNode* NewNode, bool bIsTemplateNode, FString CacheNodeName, UAnimGraphNode_SaveCachedPose* SaveCachePoseNode)
diff --git a/Engine/Source/Editor/AnimGraph/Private/AnimGraphNode_WheelHandler.cpp b/Engine/Source/Editor/AnimGraph/Private/AnimGraphNode_WheelHandler.cpp
--- a/Engine/Source/Editor/AnimGraph/Private/AnimGraphNode_WheelHandler.cpp
+++ b/Engine/Source/Editor/AnimGraph/Private/AnimGraphNode_WheelHandler.cpp
@@ -55,18 +55,6 @@ void UAnimGraphNode_WheelHandler::ValidateAnimNodePostCompile(class FCompilerRes
 	}
 }
 
-void UAnimGraphNode_WheelHandler::GetMenuEntries(FGraphContextMenuBuilder& ContextMenuBuilder) const
-{
-	UAnimBlueprint* Blueprint = Cast<UAnimBlueprint> (FBlueprintEditorUtils::FindBlueprintForGraph(ContextMenuBuilder.CurrentGraph));
-
-	if (Blueprint &&  Blueprint->ParentClass->IsChildOf(UVehicleAnimInstance::StaticClass()))
-	{
-		Super::GetMenuEntries( ContextMenuBuilder );
-	}
-
-	// else we don't show
-}
-
 bool UAnimGraphNode_WheelHandler::IsCompatibleWithGraph(const UEdGraph* TargetGraph) const
 {
 	UBlueprint* Blueprint = FBlueprintEditorUtils::FindBlueprintForGraph(TargetGraph);
diff --git a/Engine/Source/Editor/AnimGraph/Private/AnimationGraphSchema.cpp b/Engine/Source/Editor/AnimGraph/Private/AnimationGraphSchema.cpp
--- a/Engine/Source/Editor/AnimGraph/Private/AnimationGraphSchema.cpp
+++ b/Engine/Source/Editor/AnimGraph/Private/AnimationGraphSchema.cpp
@@ -9,7 +9,6 @@
 #include "GraphEditorActions.h"
 #include "ScopedTransaction.h"
 #include "AssetData.h"
-#include "K2ActionMenuBuilder.h" // for FK2ActionMenuBuilder::AddNewNodeAction()
 #include "AnimationGraphSchema.h"
 #include "K2Node_TransitionRuleGetter.h"
 #include "AnimStateNode.h"
@@ -39,44 +38,6 @@ UAnimationGraphSchema::UAnimationGraphSchema(const FObjectInitializer& ObjectIni
 	DefaultEvaluationHandlerName = TEXT("EvaluateGraphExposedInputs");
 }
 
-void UAnimationGraphSchema::GetStateMachineMenuItems(FGraphContextMenuBuilder& ContextMenuBuilder) const
-{
-	// Show state machine helpers
-	if ((ContextMenuBuilder.FromPin == NULL) || ((ContextMenuBuilder.FromPin->Direction == EGPD_Input) && (ContextMenuBuilder.FromPin->PinType.PinCategory == PC_Float)))
-	{
-		FString StateCategory(TEXT("State Machines"));
-
-		// Create the additional entries to get state information
-		TArray<UAnimStateNode*> States;
-		UBlueprint* CurrentBlueprint = FBlueprintEditorUtils::FindBlueprintForGraphChecked(ContextMenuBuilder.CurrentGraph);
-		FBlueprintEditorUtils::GetAllNodesOfClass(CurrentBlueprint, /*out*/ States);
-	
-		for (auto StateIt = States.CreateIterator(); StateIt; ++StateIt)
-		{
-			UAnimStateNode* StateNode = *StateIt;
-
-			UK2Node_TransitionRuleGetter* NodeTemplate = ContextMenuBuilder.CreateTemplateNode<UK2Node_TransitionRuleGetter>();
-			NodeTemplate->AssociatedStateNode = StateNode;
-			NodeTemplate->GetterType = ETransitionGetter::ArbitraryState_GetBlendWeight;
-
-			const FString OwnerName = StateNode->GetOuter()->GetName();
-			const FString Title = FString::Printf(TEXT("Current %s for state '%s.%s'"), *UK2Node_TransitionRuleGetter::GetFriendlyName(NodeTemplate->GetterType).ToString(), *OwnerName, *StateNode->GetStateName());
-
-			TSharedPtr<FEdGraphSchemaAction_K2NewNode> Action = FK2ActionMenuBuilder::AddNewNodeAction(ContextMenuBuilder, StateCategory, FText::FromString(Title), NodeTemplate->GetTooltipText().ToString(), 0, NodeTemplate->GetKeywords());
-			Action->NodeTemplate = NodeTemplate;
-		}
-	}
-}
-
-void UAnimationGraphSchema::GetGraphContextActions(FGraphContextMenuBuilder& ContextMenuBuilder) const
-{
-	// Add state machine info items
-	GetStateMachineMenuItems(ContextMenuBuilder);
-
-	// Inherited abilities
-	return Super::GetGraphContextActions(ContextMenuBuilder);
-}
-
 FLinearColor UAnimationGraphSchema::GetPinTypeColor(const FEdGraphPinType& PinType) const
 {
 	const bool bAdditive = PinType.PinSubCategory == TEXT("Additive");
diff --git a/Engine/Source/Editor/AnimGraph/Private/AnimationTransitionSchema.cpp b/Engine/Source/Editor/AnimGraph/Private/AnimationTransitionSchema.cpp
--- a/Engine/Source/Editor/AnimGraph/Private/AnimationTransitionSchema.cpp
+++ b/Engine/Source/Editor/AnimGraph/Private/AnimationTransitionSchema.cpp
@@ -17,7 +17,6 @@
 #include "GraphEditorActions.h"
 #include "ScopedTransaction.h"
 #include "Kismet2/KismetDebugUtilities.h"
-#include "K2ActionMenuBuilder.h" // for FK2ActionMenuBuilder::AddNewNodeAction()
 #include "K2Node_TransitionRuleGetter.h"
 
 /////////////////////////////////////////////////////
@@ -31,101 +30,6 @@ UAnimationTransitionSchema::UAnimationTransitionSchema(const FObjectInitializer&
 	// Initialize defaults
 }
 
-void UAnimationTransitionSchema::GetGraphContextActions(FGraphContextMenuBuilder& ContextMenuBuilder) const
-{
-	Super::GetGraphContextActions(ContextMenuBuilder);
-
-	GetSourceStateActions(ContextMenuBuilder);
-}
-
-void UAnimationTransitionSchema::GetSourceStateActions(FGraphContextMenuBuilder& ContextMenuBuilder) const
-{
-	if ((ContextMenuBuilder.FromPin == NULL) || ((ContextMenuBuilder.FromPin->Direction == EGPD_Input) && (ContextMenuBuilder.FromPin->PinType.PinCategory == PC_Float)))
-	{
-		// Find the source state associated with this transition
-		UAnimBlueprint* Blueprint = CastChecked<UAnimBlueprint>(FBlueprintEditorUtils::FindBlueprintForGraph(ContextMenuBuilder.CurrentGraph));
-
-		if (UAnimBlueprintGeneratedClass* AnimBlueprintClass = Blueprint->GetAnimBlueprintSkeletonClass())
-		{
-			if (UAnimStateTransitionNode* TransNode = GetTransitionNodeFromGraph(AnimBlueprintClass->GetAnimBlueprintDebugData(), ContextMenuBuilder.CurrentGraph))
-			{
-				if (UAnimStateNode* SourceStateNode = Cast<UAnimStateNode>(TransNode->GetPreviousState()))
-				{
-					// Offer options from the source state
-
-					// Sequence player positions
-					ETransitionGetter::Type SequenceSpecificGetters[] =
-					{
-						ETransitionGetter::AnimationAsset_GetCurrentTime,
-						ETransitionGetter::AnimationAsset_GetLength,
-						ETransitionGetter::AnimationAsset_GetCurrentTimeFraction,
-						ETransitionGetter::AnimationAsset_GetTimeFromEnd,
-						ETransitionGetter::AnimationAsset_GetTimeFromEndFraction
-					};
-
-					TArray<UK2Node*> AssetPlayers;
-					SourceStateNode->BoundGraph->GetNodesOfClassEx<UAnimGraphNode_Base, UK2Node>(/*out*/ AssetPlayers);
-
-					const FString Category_AssetPlayer(TEXT("Asset Player"));
-
-					for (int32 TypeIndex = 0; TypeIndex < ARRAY_COUNT(SequenceSpecificGetters); ++TypeIndex)
-					{
-						for (auto NodeIt = AssetPlayers.CreateConstIterator(); NodeIt; ++NodeIt)
-						{
-							UAnimGraphNode_Base* AnimNode = CastChecked<UAnimGraphNode_Base>(*NodeIt);
-							
-							if (AnimNode->DoesSupportTimeForTransitionGetter())
-							{
-								UK2Node_TransitionRuleGetter* NodeTemplate = ContextMenuBuilder.CreateTemplateNode<UK2Node_TransitionRuleGetter>();
-
-								FString AssetName;
-
-								UAnimationAsset * AnimAsset = AnimNode->GetAnimationAsset();
-								if (AnimAsset)
-								{
-									NodeTemplate->AssociatedAnimAssetPlayerNode = AnimNode;
-									AssetName = AnimAsset->GetName();
-								}
-
-								NodeTemplate->GetterType = SequenceSpecificGetters[TypeIndex];
-
-								FFormatNamedArguments Args;
-								Args.Add(TEXT("NodeName"), UK2Node_TransitionRuleGetter::GetFriendlyName(NodeTemplate->GetterType));
-								Args.Add(TEXT("AssetName"), FText::FromString(AssetName));
-								FText Title = FText::Format(LOCTEXT("TransitionFor", "{NodeName} for '{AssetName}'"), Args);
-
-								TSharedPtr<FEdGraphSchemaAction_K2NewNode> Action = FK2ActionMenuBuilder::AddNewNodeAction(ContextMenuBuilder, Category_AssetPlayer, Title, NodeTemplate->GetTooltipText().ToString(), 0, NodeTemplate->GetKeywords());
-								Action->NodeTemplate = NodeTemplate;
-							}
-						}
-					}
-
-					// Non-sequence specific ones
-					ETransitionGetter::Type NonSpecificGetters[] =
-					{
-						ETransitionGetter::CurrentTransitionDuration,
-						ETransitionGetter::CurrentState_ElapsedTime,
-						ETransitionGetter::CurrentState_GetBlendWeight
-					};
-
-					for (int32 TypeIndex = 0; TypeIndex < ARRAY_COUNT(NonSpecificGetters); ++TypeIndex)
-					{
-						FString Category_Transition(TEXT("Transition"));
-
-						UK2Node_TransitionRuleGetter* NodeTemplate = ContextMenuBuilder.CreateTemplateNode<UK2Node_TransitionRuleGetter>();
-						NodeTemplate->GetterType = NonSpecificGetters[TypeIndex];
-
-						FText Title = UK2Node_TransitionRuleGetter::GetFriendlyName(NodeTemplate->GetterType);
-
-						TSharedPtr<FEdGraphSchemaAction_K2NewNode> Action = FK2ActionMenuBuilder::AddNewNodeAction(ContextMenuBuilder, Category_Transition, Title, NodeTemplate->GetTooltipText().ToString(), 0, NodeTemplate->GetKeywords());
-						Action->NodeTemplate = NodeTemplate;
-					}
-				}
-			}
-		}
-	}
-}
-
 void UAnimationTransitionSchema::CreateDefaultNodesForGraph(UEdGraph& Graph) const
 {
 	// Create the entry/exit tunnels
diff --git a/Engine/Source/Editor/BlueprintGraph/Classes/EdGraphSchema_K2.h b/Engine/Source/Editor/BlueprintGraph/Classes/EdGraphSchema_K2.h
--- a/Engine/Source/Editor/BlueprintGraph/Classes/EdGraphSchema_K2.h
+++ b/Engine/Source/Editor/BlueprintGraph/Classes/EdGraphSchema_K2.h
@@ -165,50 +165,6 @@ private:
 	FBlueprintMetadata() {}
 };
 
-/** Information about what we want to call this function on */
-struct FFunctionTargetInfo
-{
-	enum EFunctionTarget
-	{
-		EFT_Default, // Just call function on target object
-		EFT_Actor, // Create an Actor node and wire to target
-		EFT_Component // Create a component variable ref node and wire to target
-	};
-
-	/** What kind of call function action are we creating */
-	EFunctionTarget FunctionTarget;
-
-	/** If CFO_Actor, call on these Actors. */
-	TArray< TWeakObjectPtr<AActor> > Actors;
-
-	/** If CFO_Component, call on this component variable of blueprint */
-	FName ComponentPropertyName;
-
-	// Constructor
-	FFunctionTargetInfo()
-		: FunctionTarget(EFT_Default)
-		, ComponentPropertyName(NAME_None)
-	{}
-
-	FFunctionTargetInfo(TArray<AActor*>& InActors)
-		: FunctionTarget(EFT_Actor)
-		, ComponentPropertyName(NAME_None)
-	{
-		for ( int32 ActorIndex = 0; ActorIndex < InActors.Num(); ActorIndex++ )
-		{
-			if ( InActors[ActorIndex] != NULL )
-			{
-				Actors.Add( TWeakObjectPtr<AActor>( InActors[ActorIndex] ) );
-			}
-		}
-	}
-
-	FFunctionTargetInfo( FName InComponentPropertyName )
-		: FunctionTarget(EFT_Component)
-		, ComponentPropertyName(InComponentPropertyName)
-	{}
-};
-
 USTRUCT()
 // Structure used to automatically convert blueprintcallable functions (that have blueprint parameter) calls (in bp graph) 
 // into their never versions (with class param instead of blueprint).
@@ -385,7 +341,6 @@ public:
 
 public:
 	// Begin EdGraphSchema Interface
-	virtual void GetGraphContextActions(FGraphContextMenuBuilder& ContextMenuBuilder) const override;
 	virtual void GetContextMenuActions(const UEdGraph* CurrentGraph, const UEdGraphNode* InGraphNode, const UEdGraphPin* InGraphPin, class FMenuBuilder* MenuBuilder, bool bIsDebugging) const override;
 	virtual const FPinConnectionResponse CanCreateConnection(const UEdGraphPin* A, const UEdGraphPin* B) const override;
 	virtual bool TryCreateConnection(UEdGraphPin* A, UEdGraphPin* B) const override;
@@ -953,10 +908,9 @@ public:
 	 * @param	InDestGraph			Graph we will be using action for (may be NULL)
 	 * @param	InFunctionTypes		Combination of EFunctionType to indicate types of functions accepted
 	 * @param	bInCalledForEach	Call for each element in an array (a node accepts array)
-	 * @param	InTargetInfo		Allows spawning nodes which also create a target variable as well
 	 * @param	OutReason			Allows callers to receive a localized string containing more detail when the function is determined to be invalid (optional)
 	 */
-	bool CanFunctionBeUsedInGraph(const UClass* InClass, const UFunction* InFunction, const UEdGraph* InDestGraph, uint32 InFunctionTypes, bool bInCalledForEach, const FFunctionTargetInfo& InTargetInfo, FText* OutReason = nullptr) const;
+	bool CanFunctionBeUsedInGraph(const UClass* InClass, const UFunction* InFunction, const UEdGraph* InDestGraph, uint32 InFunctionTypes, bool bInCalledForEach, FText* OutReason = nullptr) const;
 
 	/**
 	 * Makes connections into/or out of the gateway node, connect directly to the associated networks on the opposite side of the tunnel
@@ -983,12 +937,6 @@ public:
 	 */
 	void LinkDataPinFromOutputToInput(UEdGraphNode* InOutputNode, UEdGraphNode* InInputNode) const;
 
-	/** Function that returns _all_ nodes we could place */
-	static void GetAllActions(struct FBlueprintPaletteListBuilder& PaletteBuilder);
-
-	/** Helper method to add items valid to the palette list */
-	static void GetPaletteActions(struct FBlueprintPaletteListBuilder& ActionMenuBuilder, TWeakObjectPtr<UClass> FilterClass = NULL);
-
 	/** some inherited schemas don't want anim-notify actions listed, so this is an easy way to check that */
 	virtual bool DoesSupportAnimNotifyActions() const { return true; }
 
diff --git a/Engine/Source/Editor/BlueprintGraph/Classes/K2Node_BaseAsyncTask.h b/Engine/Source/Editor/BlueprintGraph/Classes/K2Node_BaseAsyncTask.h
--- a/Engine/Source/Editor/BlueprintGraph/Classes/K2Node_BaseAsyncTask.h
+++ b/Engine/Source/Editor/BlueprintGraph/Classes/K2Node_BaseAsyncTask.h
@@ -17,7 +17,6 @@ class BLUEPRINTGRAPH_API UK2Node_BaseAsyncTask : public UK2Node
 
 	// UEdGraphNode interface
 	virtual void AllocateDefaultPins() override;
-	virtual void GetMenuEntries(FGraphContextMenuBuilder& ContextMenuBuilder) const override;
 	virtual FText GetTooltipText() const override;
 	virtual FText GetNodeTitle(ENodeTitleType::Type TitleType) const override;
 	virtual bool IsCompatibleWithGraph(const UEdGraph* TargetGraph) const override;
@@ -33,9 +32,6 @@ class BLUEPRINTGRAPH_API UK2Node_BaseAsyncTask : public UK2Node
 	// End of UK2Node interface
 
 protected:
-	// Creates a default menu entry
-	TSharedPtr<FEdGraphSchemaAction_K2NewNode> CreateDefaultMenuEntry(UK2Node_BaseAsyncTask* NodeTemplate, FGraphContextMenuBuilder& ContextMenuBuilder) const;
-
 	// Returns the factory function (checked)
 	UFunction* GetFactoryFunction() const;
 
diff --git a/Engine/Source/Editor/BlueprintGraph/Classes/K2Node_DoOnceMultiInput.h b/Engine/Source/Editor/BlueprintGraph/Classes/K2Node_DoOnceMultiInput.h
--- a/Engine/Source/Editor/BlueprintGraph/Classes/K2Node_DoOnceMultiInput.h
+++ b/Engine/Source/Editor/BlueprintGraph/Classes/K2Node_DoOnceMultiInput.h
@@ -13,7 +13,6 @@ class UK2Node_DoOnceMultiInput : public UK2Node
 	UPROPERTY()
 	int32 NumAdditionalInputs;
 
-	virtual void GetMenuEntries(FGraphContextMenuBuilder& ContextMenuBuilder) const override;
 	virtual FText GetNodeTitle(ENodeTitleType::Type TitleType) const override;
 
 	/** Reference to the integer that contains */
diff --git a/Engine/Source/Editor/BlueprintGraph/Classes/K2Node_EnumEquality.h b/Engine/Source/Editor/BlueprintGraph/Classes/K2Node_EnumEquality.h
--- a/Engine/Source/Editor/BlueprintGraph/Classes/K2Node_EnumEquality.h
+++ b/Engine/Source/Editor/BlueprintGraph/Classes/K2Node_EnumEquality.h
@@ -23,7 +23,6 @@ class UK2Node_EnumEquality : public UK2Node
 	virtual bool IsNodePure() const override { return true; }
 	virtual bool ShouldDrawCompact() const override { return true; }
 	virtual FText GetCompactNodeTitle() const override { return NSLOCTEXT("K2Node", "EqualEqual", "=="); }
-	virtual void GetMenuEntries(FGraphContextMenuBuilder& ContextMenuBuilder) const override;
 	virtual void NotifyPinConnectionListChanged(UEdGraphPin* Pin) override;
 	virtual class FNodeHandlingFunctor* CreateNodeHandler(class FKismetCompilerContext& CompilerContext) const override;
 	virtual void ExpandNode(class FKismetCompilerContext& CompilerContext, UEdGraph* SourceGraph) override;
diff --git a/Engine/Source/Editor/BlueprintGraph/Classes/K2Node_Knot.h b/Engine/Source/Editor/BlueprintGraph/Classes/K2Node_Knot.h
--- a/Engine/Source/Editor/BlueprintGraph/Classes/K2Node_Knot.h
+++ b/Engine/Source/Editor/BlueprintGraph/Classes/K2Node_Knot.h
@@ -14,7 +14,6 @@ public:
 	virtual void AllocateDefaultPins() override;
 	virtual FText GetTooltipText() const override;
 	virtual FText GetNodeTitle(ENodeTitleType::Type TitleType) const override;
-	virtual void GetMenuEntries(struct FGraphContextMenuBuilder& ContextMenuBuilder) const override;
 	virtual bool ShouldOverridePinNames() const override;
 	virtual FText GetPinNameOverride(const UEdGraphPin& Pin) const override;
 	virtual void OnRenameNode(const FString& NewName) override;
diff --git a/Engine/Source/Editor/BlueprintGraph/Classes/K2Node_MathExpression.h b/Engine/Source/Editor/BlueprintGraph/Classes/K2Node_MathExpression.h
--- a/Engine/Source/Editor/BlueprintGraph/Classes/K2Node_MathExpression.h
+++ b/Engine/Source/Editor/BlueprintGraph/Classes/K2Node_MathExpression.h
@@ -38,7 +38,6 @@ public:
 	// End of UEdGraphNode interface
 
 	// UK2Node interface
-	virtual void GetMenuEntries(FGraphContextMenuBuilder& ContextMenuBuilder) const override;
 	virtual void ValidateNodeDuringCompilation(class FCompilerResultsLog& MessageLog) const override;
 	virtual void GetMenuActions(FBlueprintActionDatabaseRegistrar& ActionRegistrar) const override;
 	// End of UK2Node interface
diff --git a/Engine/Source/Editor/BlueprintGraph/Private/BlueprintEditorSettings.cpp b/Engine/Source/Editor/BlueprintGraph/Private/BlueprintEditorSettings.cpp
--- a/Engine/Source/Editor/BlueprintGraph/Private/BlueprintEditorSettings.cpp
+++ b/Engine/Source/Editor/BlueprintGraph/Private/BlueprintEditorSettings.cpp
@@ -18,7 +18,6 @@ UBlueprintEditorSettings::UBlueprintEditorSettings(const FObjectInitializer& Obj
 	, bFlattenFavoritesMenus(true)
 	, bFavorPureCastNodes(false)
 	, bAutoCastObjectConnections(false)
-	, bForceLegacyMenuingSystem(false)
 	, bShowViewportOnSimulate(false)
 	, bShowInheritedVariables(false)
 	, bShowEmptySections(true)
diff --git a/Engine/Source/Editor/BlueprintGraph/Private/EdGraphSchema_K2.cpp b/Engine/Source/Editor/BlueprintGraph/Private/EdGraphSchema_K2.cpp
--- a/Engine/Source/Editor/BlueprintGraph/Private/EdGraphSchema_K2.cpp
+++ b/Engine/Source/Editor/BlueprintGraph/Private/EdGraphSchema_K2.cpp
@@ -19,8 +19,8 @@
 #include "DefaultValueHelper.h"
 #include "ObjectEditorUtils.h"
 #include "ActorEditorUtils.h"
-
-#include "K2ActionMenuBuilder.h"
+#include "ComponentTypeRegistry.h"
+#include "BlueprintComponentNodeSpawner.h"
 #include "AssetRegistryModule.h"
 #include "Blueprint/AIBlueprintHelperLibrary.h"
 
@@ -515,7 +515,7 @@ bool UEdGraphSchema_K2::DoesFunctionHaveOutParameters( const UFunction* Function
 	return false;
 }
 
-bool UEdGraphSchema_K2::CanFunctionBeUsedInGraph(const UClass* InClass, const UFunction* InFunction, const UEdGraph* InDestGraph, uint32 InAllowedFunctionTypes, bool bInCalledForEach, const FFunctionTargetInfo& InTargetInfo, FText* OutReason) const
+bool UEdGraphSchema_K2::CanFunctionBeUsedInGraph(const UClass* InClass, const UFunction* InFunction, const UEdGraph* InDestGraph, uint32 InAllowedFunctionTypes, bool bInCalledForEach, FText* OutReason) const
 {
 	if (CanUserKismetCallFunction(InFunction))
 	{
@@ -691,21 +691,6 @@ bool UEdGraphSchema_K2::CanFunctionBeUsedInGraph(const UClass* InClass, const UF
 			return false;
 		}
 
-		const bool bClassIsAnActor = InClass->IsChildOf( AActor::StaticClass() );
-
-		// This will evaluate to false if there are multiple actors selected and the function has a return value or out parameters
-		const bool bFunctionHasReturnOrOutParameters = bHasReturnParams || DoesFunctionHaveOutParameters(InFunction);
-		const bool bAllowReturnValuesForNoneOrSingleActors = !bClassIsAnActor || InTargetInfo.Actors.Num() <= 1 || !bFunctionHasReturnOrOutParameters;
-		if (!bAllowReturnValuesForNoneOrSingleActors)
-		{
-			if(OutReason != nullptr)
-			{
-				*OutReason = LOCTEXT("FunctionNotAllowedWithMultipleTargets", "Functions that return a value cannot be used with multiple targets.");
-			}
-
-			return false;
-		}
-
 		return true;
 	}
 
@@ -1702,33 +1687,6 @@ void UEdGraphSchema_K2::GetJumpToConnectionSubMenuActions( class FMenuBuilder& M
 	}
 }
 
-void UEdGraphSchema_K2::GetGraphContextActions(FGraphContextMenuBuilder& ContextMenuBuilder) const
-{
-	FBlueprintGraphActionListBuilder BlueprintContextMenuBuilder(ContextMenuBuilder.CurrentGraph);
-	BlueprintContextMenuBuilder.FromPin = ContextMenuBuilder.FromPin;
-	BlueprintContextMenuBuilder.SelectedObjects.Append(ContextMenuBuilder.SelectedObjects);
-	check(BlueprintContextMenuBuilder.Blueprint != NULL);
-
-	// Run thru all nodes and add any menu items they want to add
-	Super::GetGraphContextActions(BlueprintContextMenuBuilder);
-
-	// Now do schema-specific stuff
-	FK2ActionMenuBuilder(this).GetGraphContextActions(BlueprintContextMenuBuilder);
-	ContextMenuBuilder.Append(BlueprintContextMenuBuilder);
-}
-
-void UEdGraphSchema_K2::GetAllActions(FBlueprintPaletteListBuilder& PaletteBuilder)
-{
-	const UEdGraphSchema_K2* K2SchemaInst = GetDefault<UEdGraphSchema_K2>();
-	FK2ActionMenuBuilder(K2SchemaInst).GetAllActions(PaletteBuilder);
-}
-
-void UEdGraphSchema_K2::GetPaletteActions(FBlueprintPaletteListBuilder& ActionMenuBuilder, TWeakObjectPtr<UClass> FilterClass/* = NULL*/)
-{
-	const UEdGraphSchema_K2* K2SchemaInst = GetDefault<UEdGraphSchema_K2>();
-	FK2ActionMenuBuilder(K2SchemaInst).GetPaletteActions(ActionMenuBuilder, FilterClass);
-}
-
 const FPinConnectionResponse UEdGraphSchema_K2::DetermineConnectionResponseOfCompatibleTypedPins(const UEdGraphPin* PinA, const UEdGraphPin* PinB, const UEdGraphPin* InputPin, const UEdGraphPin* OutputPin) const
 {
 	// Now check to see if there are already connections and this is an 'exclusive' connection
@@ -4419,8 +4377,11 @@ void UEdGraphSchema_K2::DroppedAssetsOnGraph(const TArray<FAssetData>& Assets, c
 				UEdGraph* TempOuter = NewObject<UEdGraph>((UObject*)Blueprint);
 				TempOuter->SetFlags(RF_Transient);
 
-				TSharedPtr<FEdGraphSchemaAction_K2AddComponent> Action = FK2ActionMenuBuilder::CreateAddComponentAction(TempOuter, Blueprint, DestinationComponentType, Asset);
-				Action->PerformAction(Graph, NULL, GraphPosition);
+				FComponentTypeEntry ComponentType = { FString(), FString(), DestinationComponentType };
+
+				IBlueprintNodeBinder::FBindingSet Bindings;
+				Bindings.Add(Asset);
+				UBlueprintComponentNodeSpawner::Create(ComponentType)->Invoke(Graph, Bindings, GraphPosition);
 			}
 		}
 	}
diff --git a/Engine/Source/Editor/BlueprintGraph/Private/K2ActionMenuBuilder.cpp b/Engine/Source/Editor/BlueprintGraph/Private/K2ActionMenuBuilder.cpp
deleted file mode 100644
diff --git a/Engine/Source/Editor/BlueprintGraph/Private/K2Node_BaseAsyncTask.cpp b/Engine/Source/Editor/BlueprintGraph/Private/K2Node_BaseAsyncTask.cpp
--- a/Engine/Source/Editor/BlueprintGraph/Private/K2Node_BaseAsyncTask.cpp
+++ b/Engine/Source/Editor/BlueprintGraph/Private/K2Node_BaseAsyncTask.cpp
@@ -5,7 +5,6 @@
 
 #include "CompilerResultsLog.h"
 #include "KismetCompiler.h"
-#include "K2ActionMenuBuilder.h"
 #include "BlueprintNodeSpawner.h"
 #include "BlueprintActionDatabaseRegistrar.h"
 
@@ -48,28 +47,6 @@ bool UK2Node_BaseAsyncTask::IsCompatibleWithGraph(const UEdGraph* TargetGraph) c
 	return bIsCompatible && Super::IsCompatibleWithGraph(TargetGraph);
 }
 
-void UK2Node_BaseAsyncTask::GetMenuEntries(FGraphContextMenuBuilder& ContextMenuBuilder) const
-{
-	const UEdGraphSchema_K2* K2Schema = GetDefault<UEdGraphSchema_K2>();
-	EGraphType GraphType = K2Schema->GetGraphType(ContextMenuBuilder.CurrentGraph);
-	const bool bAllowLatentFuncs = (GraphType == GT_Ubergraph || GraphType == GT_Macro);
-
-	if (bAllowLatentFuncs)
-{
-		UK2Node_BaseAsyncTask* NodeTemplate = NewObject<UK2Node_BaseAsyncTask>(ContextMenuBuilder.OwnerOfTemporaries, GetClass());
-		CreateDefaultMenuEntry(NodeTemplate, ContextMenuBuilder);
-	}
-}
-
-TSharedPtr<FEdGraphSchemaAction_K2NewNode> UK2Node_BaseAsyncTask::CreateDefaultMenuEntry(UK2Node_BaseAsyncTask* NodeTemplate, FGraphContextMenuBuilder& ContextMenuBuilder) const
-{
-	TSharedPtr<FEdGraphSchemaAction_K2NewNode> NodeAction = FK2ActionMenuBuilder::AddNewNodeAction(ContextMenuBuilder, NodeTemplate->GetMenuCategory().ToString(), NodeTemplate->GetNodeTitle(ENodeTitleType::ListView), NodeTemplate->GetTooltipText().ToString(), 0, NodeTemplate->GetKeywords());
-	
-	NodeAction->NodeTemplate = NodeTemplate;
-
-	return NodeAction;
-}
-
 void UK2Node_BaseAsyncTask::AllocateDefaultPins()
 {
 	const UEdGraphSchema_K2* K2Schema = GetDefault<UEdGraphSchema_K2>();
diff --git a/Engine/Source/Editor/BlueprintGraph/Private/K2Node_CallFunction.cpp b/Engine/Source/Editor/BlueprintGraph/Private/K2Node_CallFunction.cpp
--- a/Engine/Source/Editor/BlueprintGraph/Private/K2Node_CallFunction.cpp
+++ b/Engine/Source/Editor/BlueprintGraph/Private/K2Node_CallFunction.cpp
@@ -1055,7 +1055,7 @@ bool UK2Node_CallFunction::CanPasteHere(const UEdGraph* TargetGraph) const
 		{
 			TargetFunction = GetTargetFunctionFromSkeletonClass();
 		}
-		bCanPaste = K2Schema->CanFunctionBeUsedInGraph(FBlueprintEditorUtils::FindBlueprintForGraphChecked(TargetGraph)->GeneratedClass, TargetFunction, TargetGraph, AllowedFunctionTypes, false, FFunctionTargetInfo());
+		bCanPaste = K2Schema->CanFunctionBeUsedInGraph(FBlueprintEditorUtils::FindBlueprintForGraphChecked(TargetGraph)->GeneratedClass, TargetFunction, TargetGraph, AllowedFunctionTypes, false);
 	}
 	
 	return bCanPaste;
diff --git a/Engine/Source/Editor/BlueprintGraph/Private/K2Node_DoOnceMultiInput.cpp b/Engine/Source/Editor/BlueprintGraph/Private/K2Node_DoOnceMultiInput.cpp
--- a/Engine/Source/Editor/BlueprintGraph/Private/K2Node_DoOnceMultiInput.cpp
+++ b/Engine/Source/Editor/BlueprintGraph/Private/K2Node_DoOnceMultiInput.cpp
@@ -4,7 +4,6 @@
 
 #include "K2Node_CommutativeAssociativeBinaryOperator.h"
 #include "ScopedTransaction.h"
-#include "K2ActionMenuBuilder.h"
 #include "KismetCompiler.h"
 #include "BlueprintNodeSpawner.h"
 #include "EditorCategoryUtils.h"
@@ -12,19 +11,6 @@
 
 #define LOCTEXT_NAMESPACE "K2Node"
 
-void UK2Node_DoOnceMultiInput::GetMenuEntries(FGraphContextMenuBuilder& ContextMenuBuilder) const
-{
-	const UEdGraphSchema_K2* K2Schema = GetDefault<UEdGraphSchema_K2>();
-	EGraphType GraphType = K2Schema->GetGraphType(ContextMenuBuilder.CurrentGraph);
-
-	// Create a node template for this factory method
-	UK2Node_DoOnceMultiInput* NodeTemplate = NewObject<UK2Node_DoOnceMultiInput>(ContextMenuBuilder.OwnerOfTemporaries);
-	
-	TSharedPtr<FEdGraphSchemaAction_K2NewNode> NodeAction = FK2ActionMenuBuilder::AddNewNodeAction(ContextMenuBuilder, LOCTEXT("FlowCategory", "Utilities|Flow Control").ToString(), NodeTemplate->GetNodeTitle(ENodeTitleType::ListView), NodeTemplate->GetTooltipText().ToString(), 0, NodeTemplate->GetKeywords());
-
-	NodeAction->NodeTemplate = NodeTemplate;
-}
-
 void UK2Node_DoOnceMultiInput::ReallocatePinsDuringReconstruction(TArray<UEdGraphPin*>& OldPins)
 {
 	Super::ReallocatePinsDuringReconstruction(OldPins);
diff --git a/Engine/Source/Editor/BlueprintGraph/Private/K2Node_EnumEquality.cpp b/Engine/Source/Editor/BlueprintGraph/Private/K2Node_EnumEquality.cpp
--- a/Engine/Source/Editor/BlueprintGraph/Private/K2Node_EnumEquality.cpp
+++ b/Engine/Source/Editor/BlueprintGraph/Private/K2Node_EnumEquality.cpp
@@ -4,7 +4,6 @@
 #include "BlueprintGraphPrivatePCH.h"
 #include "KismetCompiler.h"
 #include "Kismet/KismetMathLibrary.h"
-#include "K2ActionMenuBuilder.h" // for FK2ActionMenuBuilder::AddNewNodeAction()
 #include "K2Node_EnumEquality.h"
 #include "BlueprintNodeSpawner.h"
 #include "EditorCategoryUtils.h"
@@ -178,44 +177,6 @@ void UK2Node_EnumEquality::GetConditionalFunction(FName& FunctionName, UClass**
 	*FunctionClass = UKismetMathLibrary::StaticClass();
 }
 
-void UK2Node_EnumEquality::GetMenuEntries(FGraphContextMenuBuilder& Context) const
-{
-	Super::GetMenuEntries(Context);
-
-	bool bShowEnumEquality = (Context.FromPin == NULL);
-
-	if (Context.FromPin != NULL)
-	{
-		// Show it on pin drags for output enums or input bools
-		const UEdGraphSchema_K2* K2Schema = GetDefault<UEdGraphSchema_K2>();
-
-		if ((Context.FromPin->Direction == EGPD_Output) && (Context.FromPin->PinType.PinCategory == K2Schema->PC_Byte))
-		{
-			if (Cast<UEnum>(Context.FromPin->PinType.PinSubCategoryObject.Get()))
-			{
-				bShowEnumEquality = true;
-			}
-		}
-		else if ((Context.FromPin->Direction == EGPD_Input) && (Context.FromPin->PinType.PinCategory == K2Schema->PC_Boolean))
-		{
-			bShowEnumEquality = true;
-		}
-	}
-
-	if (bShowEnumEquality)
-	{
-		UK2Node* EnumNodeTemplate = Context.CreateTemplateNode<UK2Node_EnumEquality>();
-
-		const FString Category = TEXT("Utilities| Enum");
-		const FText MenuDesc = EnumNodeTemplate->GetNodeTitle(ENodeTitleType::ListView);
-		const FString Tooltip = EnumNodeTemplate->GetTooltipText().ToString();
-		const FString Keywords = EnumNodeTemplate->GetKeywords();
-
-		TSharedPtr<FEdGraphSchemaAction_K2NewNode> NodeAction = FK2ActionMenuBuilder::AddNewNodeAction(Context, Category, MenuDesc, Tooltip, 0, Keywords);
-		NodeAction->NodeTemplate = EnumNodeTemplate;
-	}
-}
-
 FNodeHandlingFunctor* UK2Node_EnumEquality::CreateNodeHandler(FKismetCompilerContext& CompilerContext) const
 {
 	return new FNodeHandlingFunctor(CompilerContext);
diff --git a/Engine/Source/Editor/BlueprintGraph/Private/K2Node_EnumInequality.cpp b/Engine/Source/Editor/BlueprintGraph/Private/K2Node_EnumInequality.cpp
--- a/Engine/Source/Editor/BlueprintGraph/Private/K2Node_EnumInequality.cpp
+++ b/Engine/Source/Editor/BlueprintGraph/Private/K2Node_EnumInequality.cpp
@@ -3,7 +3,6 @@
 #include "BlueprintGraphPrivatePCH.h"
 //#include "KismetCompiler.h"
 #include "Kismet/KismetMathLibrary.h"
-//#include "K2ActionMenuBuilder.h" // for FK2ActionMenuBuilder::AddNewNodeAction()
 #include "K2Node_EnumInequality.h"
 //#include "BlueprintNodeSpawner.h"
 //#include "EditorCategoryUtils.h"
diff --git a/Engine/Source/Editor/BlueprintGraph/Private/K2Node_Knot.cpp b/Engine/Source/Editor/BlueprintGraph/Private/K2Node_Knot.cpp
--- a/Engine/Source/Editor/BlueprintGraph/Private/K2Node_Knot.cpp
+++ b/Engine/Source/Editor/BlueprintGraph/Private/K2Node_Knot.cpp
@@ -1,7 +1,6 @@
 // Copyright 1998-2015 Epic Games, Inc. All Rights Reserved.
 
 #include "BlueprintGraphPrivatePCH.h"
-#include "K2ActionMenuBuilder.h"
 #include "Kismet2NameValidators.h"
 #include "BlueprintNodeSpawner.h"
 #include "BlueprintActionDatabaseRegistrar.h"
@@ -53,22 +52,6 @@ FText UK2Node_Knot::GetNodeTitle(ENodeTitleType::Type TitleType) const
 	}
 }
 
-void UK2Node_Knot::GetMenuEntries(struct FGraphContextMenuBuilder& ContextMenuBuilder) const
-{
-	if (ContextMenuBuilder.FromPin != nullptr)
-	{
-		UK2Node_Knot* TemplateNode = NewObject<UK2Node_Knot>(GetTransientPackage(), GetClass());
-
-		FString EmptyCategory;
-		FText MenuDesc = LOCTEXT("KnotMenuDescription", "Add Reroute Node...");
-		FString Tooltip = TemplateNode->GetTooltipText().ToString();//@TODO: Make this work +LOCTEXT("KnotMenuExtraTooltip", "\nYou can also create one by Shift+Dragging off a pin").ToString();
-		FString Keywords = TemplateNode->GetKeywords();
-
-		TSharedPtr<FEdGraphSchemaAction_K2NewNode> NodeAction = FK2ActionMenuBuilder::AddNewNodeAction(ContextMenuBuilder, EmptyCategory, MenuDesc, Tooltip, 0, Keywords);
-		NodeAction->NodeTemplate = TemplateNode;
-	}
-}
-
 void UK2Node_Knot::ExpandNode(class FKismetCompilerContext& CompilerContext, UEdGraph* SourceGraph)
 {
 	Super::ExpandNode(CompilerContext, SourceGraph);
diff --git a/Engine/Source/Editor/BlueprintGraph/Private/K2Node_MathExpression.cpp b/Engine/Source/Editor/BlueprintGraph/Private/K2Node_MathExpression.cpp
--- a/Engine/Source/Editor/BlueprintGraph/Private/K2Node_MathExpression.cpp
+++ b/Engine/Source/Editor/BlueprintGraph/Private/K2Node_MathExpression.cpp
@@ -5,7 +5,6 @@
 #include "BlueprintGraphClasses.h"
 #include "Kismet2NameValidators.h"
 #include "EdGraphUtilities.h"
-#include "K2ActionMenuBuilder.h" // for FK2ActionMenuBuilder::AddNewNodeAction()
 #include "BasicTokenParser.h"
 #include "UnrealMathUtility.h"
 #include "BlueprintEditorUtils.h"
@@ -2405,19 +2404,6 @@ void UK2Node_MathExpression::PostEditChangeProperty(struct FPropertyChangedEvent
 	}
 }
 
-//------------------------------------------------------------------------------
-void UK2Node_MathExpression::GetMenuEntries(FGraphContextMenuBuilder& ContextMenuBuilder) const
-{
-	UK2Node_MathExpression* TemplateNode = NewObject<UK2Node_MathExpression>(GetTransientPackage(), GetClass());
-
-	const FString Category = TEXT("");
-	const FText   MenuDesc = LOCTEXT("AddMathExprMenuOption", "Add Math Expression...");
-	const FString Tooltip  = TEXT("Create a new mathematical expression");
-
-	TSharedPtr<FEdGraphSchemaAction_K2NewNode> NodeAction = FK2ActionMenuBuilder::AddNewNodeAction(ContextMenuBuilder, Category, MenuDesc, Tooltip);
-	NodeAction->NodeTemplate = TemplateNode;
-}
-
 //------------------------------------------------------------------------------
 void UK2Node_MathExpression::GetMenuActions(FBlueprintActionDatabaseRegistrar& ActionRegistrar) const
 {
diff --git a/Engine/Source/Editor/BlueprintGraph/Public/BlueprintEditorSettings.h b/Engine/Source/Editor/BlueprintGraph/Public/BlueprintEditorSettings.h
--- a/Engine/Source/Editor/BlueprintGraph/Public/BlueprintEditorSettings.h
+++ b/Engine/Source/Editor/BlueprintGraph/Public/BlueprintEditorSettings.h
@@ -59,10 +59,6 @@ public:
 	UPROPERTY(EditAnywhere, config, Category=Workflow)
 	bool bAutoCastObjectConnections;
 
-	/** If set, then the new refactored menu system will be replaced with the old (legacy) system (as a fallback, in case the new system has unforeseen problems)*/
-	UPROPERTY(EditAnywhere, AdvancedDisplay, config, Category=Workflow)
-	bool bForceLegacyMenuingSystem;
-
 	/** If true will show the viewport tab when simulate is clicked. */
 	UPROPERTY(EditAnywhere, config, Category=Workflow)
 	bool bShowViewportOnSimulate;
diff --git a/Engine/Source/Editor/BlueprintGraph/Public/K2ActionMenuBuilder.h b/Engine/Source/Editor/BlueprintGraph/Public/K2ActionMenuBuilder.h
deleted file mode 100644
diff --git a/Engine/Source/Editor/GameplayAbilitiesEditor/Classes/K2Node_GameplayCueEvent.h b/Engine/Source/Editor/GameplayAbilitiesEditor/Classes/K2Node_GameplayCueEvent.h
--- a/Engine/Source/Editor/GameplayAbilitiesEditor/Classes/K2Node_GameplayCueEvent.h
+++ b/Engine/Source/Editor/GameplayAbilitiesEditor/Classes/K2Node_GameplayCueEvent.h
@@ -16,7 +16,6 @@ class UK2Node_GameplayCueEvent : public UK2Node_Event
 	// Begin UEdGraphNode interface
 	virtual FText GetTooltipText() const override;
 	virtual FText GetNodeTitle(ENodeTitleType::Type TitleType) const override;
-	virtual void GetMenuEntries(struct FGraphContextMenuBuilder& ContextMenuBuilder) const override;
 	virtual bool IsCompatibleWithGraph(const UEdGraph* TargetGraph) const override;
 	// End UEdGraphNode interface
 	
diff --git a/Engine/Source/Editor/GameplayAbilitiesEditor/Classes/K2Node_LatentAbilityCall.h b/Engine/Source/Editor/GameplayAbilitiesEditor/Classes/K2Node_LatentAbilityCall.h
--- a/Engine/Source/Editor/GameplayAbilitiesEditor/Classes/K2Node_LatentAbilityCall.h
+++ b/Engine/Source/Editor/GameplayAbilitiesEditor/Classes/K2Node_LatentAbilityCall.h
@@ -15,7 +15,6 @@ class UK2Node_LatentAbilityCall : public UK2Node_BaseAsyncTask
 
 	// UEdGraphNode interface
 	virtual bool CanCreateUnderSpecifiedSchema(const UEdGraphSchema* DesiredSchema) const override;
-	virtual void GetMenuEntries(FGraphContextMenuBuilder& ContextMenuBuilder) const override;
 	virtual void GetMenuActions(FBlueprintActionDatabaseRegistrar& ActionRegistrar) const override;
 	virtual bool IsCompatibleWithGraph(UEdGraph const* TargetGraph) const override;
 	// End of UEdGraphNode interface
diff --git a/Engine/Source/Editor/GameplayAbilitiesEditor/Private/K2Node_GameplayCueEvent.cpp b/Engine/Source/Editor/GameplayAbilitiesEditor/Private/K2Node_GameplayCueEvent.cpp
--- a/Engine/Source/Editor/GameplayAbilitiesEditor/Private/K2Node_GameplayCueEvent.cpp
+++ b/Engine/Source/Editor/GameplayAbilitiesEditor/Private/K2Node_GameplayCueEvent.cpp
@@ -4,7 +4,6 @@
 #include "AbilitySystemEditorPrivatePCH.h"
 #include "K2Node_GameplayCueEvent.h"
 #include "CompilerResultsLog.h"
-#include "K2ActionMenuBuilder.h"
 #include "GameplayTagsModule.h"
 #include "GameplayCueInterface.h"
 #include "BlueprintEventNodeSpawner.h"
@@ -54,50 +53,6 @@ bool UK2Node_GameplayCueEvent::IsCompatibleWithGraph(UEdGraph const* TargetGraph
 	return bIsCompatible && Super::IsCompatibleWithGraph(TargetGraph);
 }
 
-void UK2Node_GameplayCueEvent::GetMenuEntries(FGraphContextMenuBuilder& Context) const
-{
-	Super::GetMenuEntries(Context);
-
-	if (!IsCompatibleWithGraph(Context.CurrentGraph))
-	{
-		return;
-	}
-
-	const FString FunctionCategory(TEXT("GameplayCue Event"));
-
-	IGameplayTagsModule& GameplayTagsModule = IGameplayTagsModule::Get();
-	FGameplayTag RootTag = GameplayTagsModule.GetGameplayTagsManager().RequestGameplayTag(FName(TEXT("GameplayCue")));
-
-	FGameplayTagContainer CueTags = GameplayTagsModule.GetGameplayTagsManager().RequestGameplayTagChildren(RootTag);
-	// Add a root GameplayCue function as a default
-	CueTags.AddTag(RootTag);
-
-	// Fixme: need to check if this function is already defined so that it can be reimplemented
-	//	-Checking MyBlueprint->GeneratedClass isn't enough since they may have added an event and not recompiled
-	//	-FEdGraphSchemaAction_K2AddCustomEvent does not check names/always ensures a valid name
-	//	-FEdGraphSchemaAction_K2AddEvent does check and recenters - but it looks at EventSignatureName/EventSignatureClass for equality and that
-	//		won't work here.
-	//	
-	//	Probably need a new EdGraphSchemaAction to do this properly. For now this is ok since they will get a compile error if they do drop in
-	//	two of the same GameplayCue even Nodes and it should be pretty clear that they can't do that.
-
-	for (auto It = CueTags.CreateConstIterator(); It; ++It)
-	{
-		FGameplayTag Tag = *It;
-		UK2Node_GameplayCueEvent* NodeTemplate = Context.CreateTemplateNode<UK2Node_GameplayCueEvent>();
-
-		NodeTemplate->CustomFunctionName = Tag.GetTagName();
-
-		const FString Category = FunctionCategory;
-		const FText MenuDesc = NodeTemplate->GetNodeTitle(ENodeTitleType::ListView);
-		const FString Tooltip = NodeTemplate->GetTooltipText().ToString();
-		const FString Keywords = NodeTemplate->GetKeywords();
-
-		TSharedPtr<FEdGraphSchemaAction_K2NewNode> NodeAction = FK2ActionMenuBuilder::AddNewNodeAction(Context, Category, MenuDesc, Tooltip, 0, Keywords);
-		NodeAction->NodeTemplate = NodeTemplate;
-	}	
-}
-
 void UK2Node_GameplayCueEvent::GetMenuActions(FBlueprintActionDatabaseRegistrar& ActionRegistrar) const
 {
 	// actions get registered under specific object-keys; the idea is that 
diff --git a/Engine/Source/Editor/GameplayAbilitiesEditor/Private/K2Node_LatentAbilityCall.cpp b/Engine/Source/Editor/GameplayAbilitiesEditor/Private/K2Node_LatentAbilityCall.cpp
--- a/Engine/Source/Editor/GameplayAbilitiesEditor/Private/K2Node_LatentAbilityCall.cpp
+++ b/Engine/Source/Editor/GameplayAbilitiesEditor/Private/K2Node_LatentAbilityCall.cpp
@@ -48,41 +48,6 @@ bool UK2Node_LatentAbilityCall::IsCompatibleWithGraph(UEdGraph const* TargetGrap
 	return bIsCompatible;
 }
 
-void UK2Node_LatentAbilityCall::GetMenuEntries(FGraphContextMenuBuilder& ContextMenuBuilder) const
-{
-	if (!IsCompatibleWithGraph(ContextMenuBuilder.CurrentGraph))
-	{
-		return;
-	}
-
-	for (TObjectIterator<UClass> ClassIt; ClassIt; ++ClassIt)
-	{
-		UClass* TestClass = *ClassIt;
-		if (TestClass->IsChildOf(UAbilityTask::StaticClass()) && !TestClass->HasAnyClassFlags(CLASS_Abstract | CLASS_Deprecated))
-		{
-			for (TFieldIterator<UFunction> FuncIt(TestClass, EFieldIteratorFlags::ExcludeSuper); FuncIt; ++FuncIt)
-			{
-				// See if the function is a static factory method for online proxies
-				UFunction* CandidateFunction = *FuncIt;
-
-				UObjectProperty* ReturnProperty = Cast<UObjectProperty>(CandidateFunction->GetReturnProperty());
-				const bool bValidReturnType = (ReturnProperty != nullptr) && (ReturnProperty->PropertyClass != nullptr) && (ReturnProperty->PropertyClass->IsChildOf(UAbilityTask::StaticClass()));
-
-				if (CandidateFunction->HasAllFunctionFlags(FUNC_Static) && bValidReturnType)
-				{
-					// Create a node template for this factory method
-					UK2Node_LatentAbilityCall* NodeTemplate = NewObject<UK2Node_LatentAbilityCall>(ContextMenuBuilder.OwnerOfTemporaries);
-					NodeTemplate->ProxyFactoryFunctionName = CandidateFunction->GetFName();
-					NodeTemplate->ProxyFactoryClass = TestClass;
-					NodeTemplate->ProxyClass = ReturnProperty->PropertyClass;
-
-					CreateDefaultMenuEntry(NodeTemplate, ContextMenuBuilder);
-				}
-			}
-		}
-	}
-}
-
 void UK2Node_LatentAbilityCall::GetMenuActions(FBlueprintActionDatabaseRegistrar& ActionRegistrar) const
 {
 	// these nested loops are combing over the same classes/functions the
diff --git a/Engine/Source/Editor/GameplayTagsEditor/Classes/GameplayTagsK2Node_LiteralGameplayTag.h b/Engine/Source/Editor/GameplayTagsEditor/Classes/GameplayTagsK2Node_LiteralGameplayTag.h
--- a/Engine/Source/Editor/GameplayTagsEditor/Classes/GameplayTagsK2Node_LiteralGameplayTag.h
+++ b/Engine/Source/Editor/GameplayTagsEditor/Classes/GameplayTagsK2Node_LiteralGameplayTag.h
@@ -18,7 +18,6 @@ class UGameplayTagsK2Node_LiteralGameplayTag : public UK2Node
 	virtual bool CanDuplicateNode() const override { return false; }
 	virtual bool NodeCausesStructuralBlueprintChange() const override { return true; }
 	virtual bool CanCreateUnderSpecifiedSchema(const UEdGraphSchema* Schema) const override;
-	virtual void GetMenuEntries(FGraphContextMenuBuilder& ContextMenuBuilder) const override;
 	// End UEdGraphNode interface
 
 	// Begin UK2Node interface
diff --git a/Engine/Source/Editor/GameplayTagsEditor/Private/GameplayTagsK2Node_LiteralGameplayTag.cpp b/Engine/Source/Editor/GameplayTagsEditor/Private/GameplayTagsK2Node_LiteralGameplayTag.cpp
--- a/Engine/Source/Editor/GameplayTagsEditor/Private/GameplayTagsK2Node_LiteralGameplayTag.cpp
+++ b/Engine/Source/Editor/GameplayTagsEditor/Private/GameplayTagsK2Node_LiteralGameplayTag.cpp
@@ -3,7 +3,6 @@
 #include "GameplayTagsEditorModulePrivatePCH.h"
 #include "KismetCompiler.h"
 #include "GameplayTags.h"
-#include "K2ActionMenuBuilder.h" // for FK2ActionMenuBuilder::AddNewNodeAction()
 #include "BlueprintNodeSpawner.h"
 #include "EditorCategoryUtils.h"
 #include "BlueprintActionDatabaseRegistrar.h"
@@ -39,21 +38,6 @@ bool UGameplayTagsK2Node_LiteralGameplayTag::CanCreateUnderSpecifiedSchema( cons
 	return Schema->IsA(UEdGraphSchema_K2::StaticClass());
 }
 
-void UGameplayTagsK2Node_LiteralGameplayTag::GetMenuEntries( FGraphContextMenuBuilder& ContextMenuBuilder ) const
-{
-	Super::GetMenuEntries(ContextMenuBuilder);
-
-	UK2Node* EnumNodeTemplate = ContextMenuBuilder.CreateTemplateNode<UGameplayTagsK2Node_LiteralGameplayTag>();
-
-	const FString Category = TEXT("Gameplay Tags|Tag Container");
-	const FText MenuDesc = EnumNodeTemplate->GetNodeTitle(ENodeTitleType::ListView);
-	const FString Tooltip = EnumNodeTemplate->GetTooltipText().ToString();
-	const FString Keywords = EnumNodeTemplate->GetKeywords();
-
-	TSharedPtr<FEdGraphSchemaAction_K2NewNode> NodeAction = FK2ActionMenuBuilder::AddNewNodeAction(ContextMenuBuilder, Category, MenuDesc, Tooltip, 0, Keywords);
-	NodeAction->NodeTemplate = EnumNodeTemplate;
-}
-
 void UGameplayTagsK2Node_LiteralGameplayTag::ExpandNode(class FKismetCompilerContext& CompilerContext, UEdGraph* SourceGraph)
 {
 	Super::ExpandNode(CompilerContext, SourceGraph);
diff --git a/Engine/Source/Editor/GraphEditor/Private/SGraphEditorActionMenu.cpp b/Engine/Source/Editor/GraphEditor/Private/SGraphEditorActionMenu.cpp
--- a/Engine/Source/Editor/GraphEditor/Private/SGraphEditorActionMenu.cpp
+++ b/Engine/Source/Editor/GraphEditor/Private/SGraphEditorActionMenu.cpp
@@ -3,8 +3,6 @@
 
 #include "GraphEditorCommon.h"
 #include "SGraphEditorActionMenu.h"
-#include "BlueprintEditorUtils.h"
-#include "K2ActionMenuBuilder.h" // for FBlueprintGraphActionListBuilder
 
 SGraphEditorActionMenu::~SGraphEditorActionMenu()
 {
@@ -41,7 +39,7 @@ void SGraphEditorActionMenu::Construct( const FArguments& InArgs )
 void SGraphEditorActionMenu::CollectAllActions(FGraphActionListBuilderBase& OutAllActions)
 {
 	// Build up the context object
-	FBlueprintGraphActionListBuilder ContextMenuBuilder(GraphObj);
+	FGraphContextMenuBuilder ContextMenuBuilder(GraphObj);
 	if (DraggedFromPins.Num() > 0)
 	{
 		ContextMenuBuilder.FromPin = DraggedFromPins[0];
diff --git a/Engine/Source/Editor/Kismet/Private/BPFunctionDragDropAction.cpp b/Engine/Source/Editor/Kismet/Private/BPFunctionDragDropAction.cpp
--- a/Engine/Source/Editor/Kismet/Private/BPFunctionDragDropAction.cpp
+++ b/Engine/Source/Editor/Kismet/Private/BPFunctionDragDropAction.cpp
@@ -6,7 +6,7 @@
 #include "BlueprintEditorUtils.h"
 #include "GraphEditorDragDropAction.h"
 #include "BPFunctionDragDropAction.h"
-#include "K2ActionMenuBuilder.h" // for FK2ActionMenuBuilder::AddSpawnInfoForFunction()
+#include "BlueprintFunctionNodeSpawner.h"
 
 #define LOCTEXT_NAMESPACE "FunctionDragDropAction"
 
@@ -208,24 +208,21 @@ FReply FKismetFunctionDragDropAction::DroppedOnPanel(TSharedRef<SWidget> const&
 	FReply Reply = FReply::Unhandled();
 
 	// The ActionNode set during construction points to the Graph, this is suitable for displaying the mouse decorator but needs to be more complete based on the current graph
-	FGraphActionListBuilderBase::ActionGroup DropActionSet(TSharedPtr<FEdGraphSchemaAction>(NULL));
-	GetDropAction(Graph, DropActionSet);
+	UBlueprintFunctionNodeSpawner* FunctionNodeSpawner = nullptr;
+	FunctionNodeSpawner = GetDropAction(Graph);
 
-	if (DropActionSet.Actions.Num() > 0)
+	if (FunctionNodeSpawner)
 	{
-		// we really only expect there to be one action
-		TSharedPtr<FEdGraphSchemaAction> FirstDropAction = DropActionSet.Actions[0];
-
 		FText CannotDropReason = FText::GetEmpty();
-		if (!CanBeDroppedDelegate.IsBound() || CanBeDroppedDelegate.Execute(FirstDropAction, GetHoveredGraph(), CannotDropReason))
+		if (!CanBeDroppedDelegate.IsBound() || CanBeDroppedDelegate.Execute(nullptr, GetHoveredGraph(), CannotDropReason))
 		{
 			UFunction const* Function = GetFunctionProperty();
 			if ((Function != NULL) && UEdGraphSchema_K2::CanUserKismetCallFunction(Function))
 			{
 				AnalyticCallback.ExecuteIfBound();
 
-				TArray<UEdGraphPin*> DummyPins;
-				DropActionSet.PerformAction(&Graph, DummyPins, GraphPosition);
+				IBlueprintNodeBinder::FBindingSet Bindings;
+				FunctionNodeSpawner->Invoke(&Graph, Bindings, GraphPosition);
 
 				Reply = FReply::Handled();
 			}
@@ -246,7 +243,7 @@ UFunction const* FKismetFunctionDragDropAction::GetFunctionProperty() const
 }
 
 //------------------------------------------------------------------------------
-void FKismetFunctionDragDropAction::GetDropAction(UEdGraph& Graph, FGraphActionListBuilderBase::ActionGroup& DropActionOut) const
+UBlueprintFunctionNodeSpawner* FKismetFunctionDragDropAction::GetDropAction(UEdGraph& Graph) const
 {
 	if (UEdGraph const* const HoveredGraph = &Graph)
 	{
@@ -262,16 +259,11 @@ void FKismetFunctionDragDropAction::GetDropAction(UEdGraph& Graph, FGraphActionL
 
 			if (UFunction const* Function = GetFunctionProperty())
 			{
-				// Use schema function to make 'spawn action'
-				FK2ActionMenuBuilder::AddSpawnInfoForFunction(Function, false, FFunctionTargetInfo(), CallOnMember, TEXT(""), K2Schema->AG_LevelReference, TempListBuilder);
-				// we expect a single action
-				if (ensure(TempListBuilder.GetNumActions() == 1))
-				{
-					DropActionOut = TempListBuilder.GetAction(0);
-				}
+				return UBlueprintFunctionNodeSpawner::Create(Function);
 			}
 		}
 	}
+	return nullptr;
 }
 
 /*******************************************************************************
diff --git a/Engine/Source/Editor/Kismet/Private/BPFunctionDragDropAction.h b/Engine/Source/Editor/Kismet/Private/BPFunctionDragDropAction.h
--- a/Engine/Source/Editor/Kismet/Private/BPFunctionDragDropAction.h
+++ b/Engine/Source/Editor/Kismet/Private/BPFunctionDragDropAction.h
@@ -71,7 +71,7 @@ protected:
 	UFunction const* GetFunctionProperty() const;
 
 	/** Constructs an action to execute, placing a function call node for the associated function */
-	void GetDropAction(UEdGraph& Graph, FGraphActionListBuilderBase::ActionGroup& DropActionOut) const;
+	class UBlueprintFunctionNodeSpawner* GetDropAction(UEdGraph& Graph) const;
 };
 
 /*******************************************************************************
diff --git a/Engine/Source/Editor/Kismet/Private/BlueprintActionMenuBuilder.cpp b/Engine/Source/Editor/Kismet/Private/BlueprintActionMenuBuilder.cpp
--- a/Engine/Source/Editor/Kismet/Private/BlueprintActionMenuBuilder.cpp
+++ b/Engine/Source/Editor/Kismet/Private/BlueprintActionMenuBuilder.cpp
@@ -10,13 +10,10 @@
 #include "BlueprintFunctionNodeSpawner.h"
 #include "BlueprintDelegateNodeSpawner.h"
 #include "BlueprintVariableNodeSpawner.h"
-#include "K2ActionMenuBuilder.h"		// for FBlueprintPaletteListBuilder/FBlueprintGraphActionListBuilder
-#include "EdGraphSchema_K2.h"			// for StaticClass(), bUseLegacyActionMenus, etc.
 #include "BlueprintEditor.h"			// for GetMyBlueprintWidget(), and GetIsContextSensitive()
 #include "SMyBlueprint.h"				// for SelectionAsVar()
 #include "BlueprintEditorUtils.h"		// for FindBlueprintForGraphChecked()
 #include "BlueprintEditor.h"			// for GetFocusedGraph()
-#include "BlueprintEditorSettings.h"	// for bForceLegacyMenuingSystem
 
 #define LOCTEXT_NAMESPACE "BlueprintActionMenuBuilder"
 
@@ -251,17 +248,6 @@ namespace FBlueprintActionMenuBuilderImpl
 		TMap<UProperty const*, TSharedPtr<FBlueprintDragDropMenuItem>> ConsolidatedProperties;
 	};
 	
-	/**
-	 * To offer a fallback in case this menu system is unstable on release, this
-	 * method implements the old way we used collect blueprint menu actions (for
-	 * both the palette and context menu).
-	 *
-	 * @param  MenuSection		The primary section for the FBlueprintActionMenuBuilder.
-	 * @param  BlueprintEditor	
-	 * @param  MenuOut  		The menu builder we want all the legacy actions appended to.
-	 */
-	static void AppendLegacyItems(FMenuSectionDefinition const& MenuDef, TWeakPtr<FBlueprintEditor> BlueprintEditor, FBlueprintActionMenuBuilder& MenuOut);
-
 	/**
 	 * 
 	 * 
@@ -430,69 +416,6 @@ void FBlueprintActionMenuBuilderImpl::FMenuSectionDefinition::Empty()
 	ConsolidatedProperties.Empty();
 }
 
-//------------------------------------------------------------------------------
-void FBlueprintActionMenuBuilderImpl::AppendLegacyItems(FMenuSectionDefinition const& MenuDef, TWeakPtr<FBlueprintEditor> BlueprintEditor, FBlueprintActionMenuBuilder& MenuOut)
-{
-	FBlueprintActionFilter const&  MenuFilter  = MenuDef.Filter;
-	FBlueprintActionContext const& MenuContext = MenuFilter.Context;
-	
-	// if this is for the context menu
-	if (MenuContext.Graphs.Num() > 0)
-	{
-		UEdGraph* Graph = MenuContext.Graphs[0];
-		UEdGraphSchema const* GraphSchema = GetDefault<UEdGraphSchema>(Graph->Schema);
-		
-		FBlueprintGraphActionListBuilder LegacyBuilder(Graph);
-		if (MenuContext.Pins.Num() > 0)
-		{
-			LegacyBuilder.FromPin = MenuContext.Pins[0];
-		}
-		
-		bool bIsContextSensitive = true;
-		if (BlueprintEditor.IsValid())
-		{
-			bIsContextSensitive = BlueprintEditor.Pin()->GetIsContextSensitive();
-			if (bIsContextSensitive)
-			{
-				FEdGraphSchemaAction_K2Var* SelectedVar = BlueprintEditor.Pin()->GetMyBlueprintWidget()->SelectionAsVar();
-				if ((SelectedVar != nullptr) && (SelectedVar->GetProperty() != nullptr))
-				{
-					LegacyBuilder.SelectedObjects.Add(SelectedVar->GetProperty());
-				}
-			}
-		}
-		
-		if (bIsContextSensitive)
-		{
-			GraphSchema->GetGraphContextActions(LegacyBuilder);
-			MenuOut.Append(LegacyBuilder);
-		}
-		else
-		{
-			UBlueprint* Blueprint = FBlueprintEditorUtils::FindBlueprintForGraphChecked(Graph);
-			FBlueprintPaletteListBuilder ContextlessLegacyBuilder(Blueprint);
-			UEdGraphSchema_K2::GetAllActions(ContextlessLegacyBuilder);
-			MenuOut.Append(ContextlessLegacyBuilder);
-		}
-	}
-	else if (MenuContext.Blueprints.Num() > 0)
-	{
-		UBlueprint* Blueprint = MenuContext.Blueprints[0];
-		FBlueprintPaletteListBuilder LegacyBuilder(Blueprint, MenuDef.GetSectionHeading().ToString());
-		
-		UClass* ClassFilter = nullptr;
-		if (MenuFilter.TargetClasses.Num() > 0)
-		{
-			ClassFilter = MenuFilter.TargetClasses[0];
-		}
-		
-		UEdGraphSchema_K2 const* K2Schema = GetDefault<UEdGraphSchema_K2>();
-		FK2ActionMenuBuilder(K2Schema).GetPaletteActions(LegacyBuilder, ClassFilter);
-		
-		MenuOut.Append(LegacyBuilder);
-	}
-}
-
 //------------------------------------------------------------------------------
 static TArray<UObject*> FBlueprintActionMenuBuilderImpl::GetBindingCandidates(FBlueprintActionContext const& Context)
 {
@@ -541,32 +464,24 @@ void FBlueprintActionMenuBuilder::RebuildActionList()
 		MenuSection->Empty();
 	}
 	
-	const UBlueprintEditorSettings* BlueprintSettings = GetDefault<UBlueprintEditorSettings>();
-	if (!BlueprintSettings->bForceLegacyMenuingSystem)
+	FBlueprintActionDatabase::FActionRegistry const& ActionDatabase = FBlueprintActionDatabase::Get().GetAllActions();
+	for (auto const& ActionEntry : ActionDatabase)
 	{
-		FBlueprintActionDatabase::FActionRegistry const& ActionDatabase = FBlueprintActionDatabase::Get().GetAllActions();
-		for (auto const& ActionEntry : ActionDatabase)
+		for (UBlueprintNodeSpawner const* NodeSpawner : ActionEntry.Value)
 		{
-			for (UBlueprintNodeSpawner const* NodeSpawner : ActionEntry.Value)
-			{
-				FBlueprintActionInfo BlueprintAction(ActionEntry.Key, NodeSpawner);
+			FBlueprintActionInfo BlueprintAction(ActionEntry.Key, NodeSpawner);
 
-				// @TODO: could probably have a super filter that spreads across 
-				//        all MenuSctions (to pair down on performance?)
-				for (TSharedRef<FMenuSectionDefinition> const& MenuSection : MenuSections)
+			// @TODO: could probably have a super filter that spreads across 
+			//        all MenuSctions (to pair down on performance?)
+			for (TSharedRef<FMenuSectionDefinition> const& MenuSection : MenuSections)
+			{
+				for (TSharedPtr<FEdGraphSchemaAction> MenuEntry : MenuSection->MakeMenuItems(BlueprintEditorPtr, BlueprintAction))
 				{
-					for (TSharedPtr<FEdGraphSchemaAction> MenuEntry : MenuSection->MakeMenuItems(BlueprintEditorPtr, BlueprintAction))
-					{
-						AddAction(MenuEntry);
-					}
+					AddAction(MenuEntry);
 				}
 			}
-		}	
-	}
-	else if (MenuSections.Num() > 0)
-	{
-		AppendLegacyItems(*MenuSections[0], BlueprintEditorPtr, *this);
-	}
+		}
+	}	
 }
 
 #undef LOCTEXT_NAMESPACE
diff --git a/Engine/Source/Editor/Kismet/Private/BlueprintActionMenuUtils.cpp b/Engine/Source/Editor/Kismet/Private/BlueprintActionMenuUtils.cpp
--- a/Engine/Source/Editor/Kismet/Private/BlueprintActionMenuUtils.cpp
+++ b/Engine/Source/Editor/Kismet/Private/BlueprintActionMenuUtils.cpp
@@ -16,7 +16,6 @@
 #include "K2Node_VariableGet.h"
 #include "K2Node_VariableSet.h"
 #include "KismetEditorUtilities.h"	// for CanPasteNodes()
-#include "K2ActionMenuBuilder.h"
 #include "BlueprintEditorSettings.h"
 #include "Engine/Selection.h"
 #include "Engine/LevelScriptActor.h"
@@ -668,34 +667,31 @@ void FBlueprintActionMenuUtils::MakeContextMenu(FBlueprintActionContext const& C
 
 	MenuOut.RebuildActionList();
 
-	if (!BlueprintSettings->bForceLegacyMenuingSystem)
+	for (UEdGraph const* Graph : Context.Graphs)
 	{
-		for (UEdGraph const* Graph : Context.Graphs)
+		if (FKismetEditorUtilities::CanPasteNodes(Graph))
 		{
-			if (FKismetEditorUtilities::CanPasteNodes(Graph))
-			{
-				// @TODO: Grey out menu option with tooltip if one of the nodes cannot paste into this graph
-				TSharedPtr<FEdGraphSchemaAction> PasteHereAction(new FEdGraphSchemaAction_K2PasteHere(TEXT(""), LOCTEXT("PasteHereMenuName", "Paste here"), TEXT(""), MainMenuSectionGroup));
-				MenuOut.AddAction(PasteHereAction);
-				break;
-			}
+			// @TODO: Grey out menu option with tooltip if one of the nodes cannot paste into this graph
+			TSharedPtr<FEdGraphSchemaAction> PasteHereAction(new FEdGraphSchemaAction_K2PasteHere(TEXT(""), LOCTEXT("PasteHereMenuName", "Paste here"), TEXT(""), MainMenuSectionGroup));
+			MenuOut.AddAction(PasteHereAction);
+			break;
 		}
+	}
 
-		if (bIsContextSensitive && bCanHaveActorComponents && !bAddComponentsSection)
-		{
-			FText SelectComponentMsg = LOCTEXT("SelectComponentForEvents", "Select a Component to see available Events & Functions");
-			FText SelectComponentToolTip = LOCTEXT("SelectComponentForEventsTooltip", "Select a Component in the MyBlueprint tab to see available Events and Functions in this menu.");
-			TSharedPtr<FEdGraphSchemaAction> MsgAction = TSharedPtr<FEdGraphSchemaAction>(new FEdGraphSchemaAction_Dummy(TEXT(""), SelectComponentMsg, SelectComponentToolTip.ToString(), ComponentsSectionGroup));
-			MenuOut.AddAction(MsgAction);
-		}
+	if (bIsContextSensitive && bCanHaveActorComponents && !bAddComponentsSection)
+	{
+		FText SelectComponentMsg = LOCTEXT("SelectComponentForEvents", "Select a Component to see available Events & Functions");
+		FText SelectComponentToolTip = LOCTEXT("SelectComponentForEventsTooltip", "Select a Component in the MyBlueprint tab to see available Events and Functions in this menu.");
+		TSharedPtr<FEdGraphSchemaAction> MsgAction = TSharedPtr<FEdGraphSchemaAction>(new FEdGraphSchemaAction_Dummy(TEXT(""), SelectComponentMsg, SelectComponentToolTip.ToString(), ComponentsSectionGroup));
+		MenuOut.AddAction(MsgAction);
+	}
 
-		if (bIsContextSensitive && bCanOperateOnLevelActors && !bAddLevelActorsSection)
-		{
-			FText SelectActorsMsg = LOCTEXT("SelectActorForEvents", "Select Actor(s) to see available Events & Functions");
-			FText SelectActorsToolTip = LOCTEXT("SelectActorForEventsTooltip", "Select Actor(s) in the level to see available Events and Functions in this menu.");
-			TSharedPtr<FEdGraphSchemaAction> MsgAction = TSharedPtr<FEdGraphSchemaAction>(new FEdGraphSchemaAction_Dummy(TEXT(""), SelectActorsMsg, SelectActorsToolTip.ToString(), LevelActorSectionGroup));
-			MenuOut.AddAction(MsgAction);
-		}
+	if (bIsContextSensitive && bCanOperateOnLevelActors && !bAddLevelActorsSection)
+	{
+		FText SelectActorsMsg = LOCTEXT("SelectActorForEvents", "Select Actor(s) to see available Events & Functions");
+		FText SelectActorsToolTip = LOCTEXT("SelectActorForEventsTooltip", "Select Actor(s) in the level to see available Events and Functions in this menu.");
+		TSharedPtr<FEdGraphSchemaAction> MsgAction = TSharedPtr<FEdGraphSchemaAction>(new FEdGraphSchemaAction_Dummy(TEXT(""), SelectActorsMsg, SelectActorsToolTip.ToString(), LevelActorSectionGroup));
+		MenuOut.AddAction(MsgAction);
 	}
 }
 
@@ -704,54 +700,19 @@ void FBlueprintActionMenuUtils::MakeFavoritesMenu(FBlueprintActionContext const&
 {
 	MenuOut.Empty();
 
+	FBlueprintActionFilter MenuFilter;
+	MenuFilter.Context = Context;
+	MenuFilter.AddRejectionTest(FBlueprintActionFilter::FRejectionTestDelegate::CreateStatic(BlueprintActionMenuUtilsImpl::IsNonFavoritedAction));
+
+	uint32 SectionFlags = 0x00;
 	const UBlueprintEditorSettings* BlueprintSettings = GetDefault<UBlueprintEditorSettings>();
-	if (!BlueprintSettings->bForceLegacyMenuingSystem)
+	if (BlueprintSettings->bFlattenFavoritesMenus)
 	{
-		FBlueprintActionFilter MenuFilter;
-		MenuFilter.Context = Context;
-		MenuFilter.AddRejectionTest(FBlueprintActionFilter::FRejectionTestDelegate::CreateStatic(BlueprintActionMenuUtilsImpl::IsNonFavoritedAction));
-
-		uint32 SectionFlags = 0x00;
-		if (BlueprintSettings->bFlattenFavoritesMenus)
-		{
-			SectionFlags = FBlueprintActionMenuBuilder::FlattenCategoryHierarcy;
-		}
-
-		MenuOut.AddMenuSection(MenuFilter, FText::GetEmpty(), BlueprintActionMenuUtilsImpl::MainMenuSectionGroup, SectionFlags);
-		MenuOut.RebuildActionList();
+		SectionFlags = FBlueprintActionMenuBuilder::FlattenCategoryHierarcy;
 	}
-	else
-	{
-		check(Context.Blueprints.Num() > 0);
-		FBlueprintPaletteListBuilder LegacyMenuBuilder(Context.Blueprints[0]);
-		UEdGraphSchema_K2 const* K2Schema = GetDefault<UEdGraphSchema_K2>();
-		FK2ActionMenuBuilder(K2Schema).GetPaletteActions(LegacyMenuBuilder, /*ClassFilter =*/nullptr);
-
-		UEditorUserSettings& EditorUserSettings = GEditor->AccessEditorUserSettings();
-		// grab the user's favorites
-		UBlueprintPaletteFavorites const* BlueprintFavorites = EditorUserSettings.BlueprintFavorites;
-		check(BlueprintFavorites != nullptr);
-
-		for (int32 ActionIndex = 0; ActionIndex < LegacyMenuBuilder.GetNumActions(); ++ActionIndex)
-		{
-			FGraphActionListBuilderBase::ActionGroup& ActionGroup = LegacyMenuBuilder.GetAction(ActionIndex);
-			if (!ensure(ActionGroup.Actions.Num() == 1))
-			{
-				continue;
-			}
 
-			TSharedPtr<FEdGraphSchemaAction> Action = ActionGroup.Actions[0];
-			if (!Action.IsValid())
-			{
-				continue;
-			}
-
-			if (BlueprintFavorites->IsFavorited(Action))
-			{
-				MenuOut.AddAction(Action);
-			}
-		}
-	} 	
+	MenuOut.AddMenuSection(MenuFilter, FText::GetEmpty(), BlueprintActionMenuUtilsImpl::MainMenuSectionGroup, SectionFlags);
+	MenuOut.RebuildActionList();
 }
 
 //------------------------------------------------------------------------------
diff --git a/Engine/Source/Editor/Kismet/Private/BlueprintEditor.cpp b/Engine/Source/Editor/Kismet/Private/BlueprintEditor.cpp
--- a/Engine/Source/Editor/Kismet/Private/BlueprintEditor.cpp
+++ b/Engine/Source/Editor/Kismet/Private/BlueprintEditor.cpp
@@ -20,7 +20,6 @@
 #include "UObjectToken.h"
 #include "Kismet/GameplayStatics.h"
 #include "Editor/Kismet/Public/FindInBlueprintManager.h"
-#include "Editor/BlueprintGraph/Public/K2ActionMenuBuilder.h"
 #include "BlueprintEditorCommands.h"
 #include "BlueprintEditor.h"
 #include "BlueprintEditorModule.h"
@@ -29,6 +28,7 @@
 #include "SDockTab.h"
 #include "EditorClassUtils.h"
 #include "IDocumentation.h"
+#include "BlueprintFunctionNodeSpawner.h"
 
 #include "SBlueprintEditorToolbar.h"
 #include "FindInBlueprints.h"
@@ -1263,12 +1263,6 @@ void FBlueprintEditor::OnChangeBreadCrumbGraph(class UEdGraph* InGraph)
 	}
 }
 
-static auto CVarBlueprintExpertMode	= IConsoleManager::Get().RegisterConsoleVariableRef(
-	TEXT("Kismet.ExpertMode"),
-	FK2ActionMenuBuilder::bAllowUnsafeCommands,
-	*LOCTEXT("EnableExportMode_Desc", "Enable expert mode options, which will cause context menus to show additional nodes that are not verified to be safe/legal to use in the current context.").ToString()
-	);
-
 FBlueprintEditor::FBlueprintEditor()
 	: bSaveIntermediateBuildProducts(false)
 	, bPendingDeferredClose(false)
@@ -6116,15 +6110,9 @@ UEdGraph* FBlueprintEditor::CollapseSelectionToFunction(TSharedPtr<SGraphEditor>
 	TempListBuilder.OwnerOfTemporaries = NewObject<UEdGraph>(GetBlueprintObj());
 	TempListBuilder.OwnerOfTemporaries->SetFlags(RF_Transient);
 
-	// Use schema function to make 'spawn action'
-	FK2ActionMenuBuilder::AddSpawnInfoForFunction(FindField<UFunction>(GetBlueprintObj()->SkeletonGeneratedClass, DocumentName), false, FFunctionTargetInfo(), FMemberReference(), TEXT(""), K2Schema->AG_LevelReference, TempListBuilder);
-	// and execute it
-	if(TempListBuilder.GetNumActions() == 1)
-	{
-		TArray<UEdGraphPin*> DummyPins;
-		FGraphActionListBuilderBase::ActionGroup& Action = TempListBuilder.GetAction(0);
-		OutFunctionNode = Action.Actions[0]->PerformAction(FocusedGraphEd->GetCurrentGraph(), DummyPins, FVector2D::ZeroVector, false);
-	}
+	IBlueprintNodeBinder::FBindingSet Bindings;
+	OutFunctionNode = UBlueprintFunctionNodeSpawner::Create(FindField<UFunction>(GetBlueprintObj()->SkeletonGeneratedClass, DocumentName))->Invoke(SourceGraph, Bindings, FVector2D::ZeroVector);
+
 	check(OutFunctionNode);
 
 	CollapseNodesIntoGraph(OutFunctionNode, EntryNode, ResultNode, SourceGraph, NewGraph, InCollapsableNodes);
@@ -7404,37 +7392,19 @@ FReply FBlueprintEditor::OnSpawnGraphNodeByShortcut(FInputChord InChord, const F
 		return FReply::Handled();
 	}
 
-	FBlueprintPaletteListBuilder PaletteBuilder(GetBlueprintObj());
-	FBlueprintSpawnNodeCommands::Get().GetGraphActionByChord(InChord, PaletteBuilder, InGraph);
-
-	TSet<const UEdGraphNode*> NodesToSelect;
+	TArray<UEdGraphNode*> OutNodes;
 	FVector2D NodeSpawnPos = InPosition;
+	FBlueprintSpawnNodeCommands::Get().GetGraphActionByChord(InChord, InGraph, NodeSpawnPos, OutNodes);
 
-	bool bPerformedAction = false;
+	TSet<const UEdGraphNode*> NodesToSelect;
 
-	for (int32 ActionIndex = 0; ActionIndex < PaletteBuilder.GetNumActions(); ++ActionIndex)
+	for (UEdGraphNode* CurrentNode : OutNodes)
 	{
-		FGraphActionListBuilderBase::ActionGroup& ActionSet = PaletteBuilder.GetAction(ActionIndex);
-		if ((ActionSet.Actions.Num() > 0) && ActionSet.Actions[0].IsValid())
-		{
-			int32 const OldNodeCount = Graph->Nodes.Num();
-
-			TArray<UEdGraphPin*> DummyPins;
-			ActionSet.PerformAction(Graph, DummyPins, NodeSpawnPos);
-
-			for (int32 NodeIndex = OldNodeCount; NodeIndex < Graph->Nodes.Num(); ++NodeIndex)
-			{
-				UEdGraphNode* SpawnedNode = Graph->Nodes[NodeIndex];
-				NodeSpawnPos.Y = FMath::Max(NodeSpawnPos.Y, SpawnedNode->NodePosY + UEdGraphSchema_K2::EstimateNodeHeight(SpawnedNode));
-				NodesToSelect.Add(SpawnedNode);
-			}
-
-			bPerformedAction = true;
-		}	
+		NodesToSelect.Add(CurrentNode);
 	}
 
 	// Do not change node selection if no actions were performed
-	if(bPerformedAction)
+	if(OutNodes.Num() > 0)
 	{
 		Graph->SelectNodeSet(NodesToSelect, /*bFromUI =*/true);
 	}
diff --git a/Engine/Source/Editor/Kismet/Private/BlueprintEditorCommands.cpp b/Engine/Source/Editor/Kismet/Private/BlueprintEditorCommands.cpp
--- a/Engine/Source/Editor/Kismet/Private/BlueprintEditorCommands.cpp
+++ b/Engine/Source/Editor/Kismet/Private/BlueprintEditorCommands.cpp
@@ -5,11 +5,10 @@
 #include "SlateBasics.h"
 #include "BlueprintEditorCommands.h"
 #include "Editor/UnrealEd/Public/Kismet2/BlueprintEditorUtils.h"
-#include "Editor/BlueprintGraph/Public/K2ActionMenuBuilder.h" // for AddEventForFunction(), ...
 #include "Engine/Selection.h"
 #include "Engine/LevelScriptActor.h"
-
-struct FBlueprintPaletteListBuilder;
+#include "BlueprintFunctionNodeSpawner.h" 
+#include "BlueprintEventNodeSpawner.h" 
 
 #define LOCTEXT_NAMESPACE ""
 
@@ -157,43 +156,36 @@ public:
 	/**
 	 * Creates an action to be used for placing a node into the graph
 	 *
-	 * @param	InPaletteBuilder	The Blueprint palette the action should be created for
-	 * @param	InDestGraph			The graph the action should be created for
+	 * @param	InDestGraph					The graph the action should be created for
+	 * @param	InOutDestPosition			Position to start placing nodes, will be updated to be at the next safe position for node placement
+	 * @param	OutNodes					All nodes spawned by this operation
 	 *
-	 * @return						A fully prepared action containing the information to spawn the node
+	 * @return								A fully prepared action containing the information to spawn the node
 	 */
-	virtual void GetActions(FBlueprintPaletteListBuilder& InPaletteBuilder, UEdGraph* InDestGraph) = 0;
+	virtual void GetActions(UEdGraph* InDestGraph, FVector2D& InOutDestPosition, TArray<UEdGraphNode*>& OutNodes) = 0;
 };
 
 class FEdGraphNodeSpawnInfo : public FNodeSpawnInfo
 {
 public:
-	FEdGraphNodeSpawnInfo(UClass* InClass) : NodeClass(InClass), GraphNode(NULL) {}
+	FEdGraphNodeSpawnInfo(UClass* InClass) : NodeClass(InClass) {}
 
 	// FNodeSpawnInfo interface
-	virtual void GetActions(FBlueprintPaletteListBuilder& InPaletteBuilder, UEdGraph* InDestGraph) override
+	virtual void GetActions(UEdGraph* InDestGraph, FVector2D& InOutDestPosition, TArray<UEdGraphNode*>& OutNodes) override
 	{
 		TSharedPtr<FEdGraphSchemaAction_NewNode> NewActionNode = TSharedPtr<FEdGraphSchemaAction_NewNode>(new FEdGraphSchemaAction_NewNode);
 
-		if (!GraphNode.IsValid())
-		{
-			GraphNode = InPaletteBuilder.CreateTemplateNode<UEdGraphNode>(NodeClass);
-		}
-		ensureMsgf( GraphNode->GetClass() == NodeClass, TEXT("FEdGraphNodeSpawnInfo::GetAction() class mismatch (GraphNode class %s, NodeClass %s"), *GraphNode->GetClass()->GetName(), *NodeClass->GetName() );
-
-		// the NodeTemplate UPROPERTY takes ownership of GraphNode's lifetime (hence it being a weak-pointer here)
-		NewActionNode->NodeTemplate = GraphNode.Get();
+		IBlueprintNodeBinder::FBindingSet Bindings;
+		UEdGraphNode* GraphNode = UBlueprintNodeSpawner::Create(NodeClass)->Invoke(InDestGraph, Bindings, InOutDestPosition);
 
-		InPaletteBuilder.AddAction(NewActionNode);
+		InOutDestPosition.Y = FMath::Max(InOutDestPosition.Y, GraphNode->NodePosY + UEdGraphSchema_K2::EstimateNodeHeight(GraphNode));
+		OutNodes.Add(GraphNode);
 	}
 	// End of FNodeSpawnInfo interface
 
 private:
 	/** The class type the node should be */
 	UClass* NodeClass;
-
-	/** The graph node to place into the action to spawn it in the graph */
-	TWeakObjectPtr<UEdGraphNode> GraphNode;
 };
 
 class FFunctionNodeSpawnInfo : public FNodeSpawnInfo
@@ -202,9 +194,9 @@ public:
 	FFunctionNodeSpawnInfo(UFunction* InFunctionPtr) : FunctionPtr(InFunctionPtr) {}
 
 	// FNodeSpawnInfo interface
-	virtual void GetActions(FBlueprintPaletteListBuilder& InPaletteBuilder, UEdGraph* InDestGraph) override
+	virtual void GetActions(UEdGraph* InDestGraph, FVector2D& InOutDestPosition, TArray<UEdGraphNode*>& OutNodes) override
 	{
-		const UBlueprint* Blueprint = InPaletteBuilder.Blueprint;
+		const UBlueprint* Blueprint = FBlueprintEditorUtils::FindBlueprintForGraph(InDestGraph);
 
 		const UEdGraphSchema_K2* K2Schema = GetDefault<UEdGraphSchema_K2>();
 		if(NodeSpawnInfoHelpers::IsFunctionAvailableAsEvent(Blueprint, FunctionPtr))
@@ -212,7 +204,11 @@ public:
 			UFunction* FunctionEvent = NodeSpawnInfoHelpers::FindEventFunctionForClass(Blueprint, FunctionPtr);
 			if(FunctionEvent)
 			{
-				FK2ActionMenuBuilder::AddEventForFunction(InPaletteBuilder, FunctionEvent);
+				IBlueprintNodeBinder::FBindingSet Bindings;
+				UEdGraphNode* GraphNode = UBlueprintEventNodeSpawner::Create(FunctionEvent)->Invoke(InDestGraph, Bindings, InOutDestPosition);
+
+				InOutDestPosition.Y = FMath::Max(InOutDestPosition.Y, GraphNode->NodePosY + UEdGraphSchema_K2::EstimateNodeHeight(GraphNode));
+				OutNodes.Add(GraphNode);
 			}
 		}
 		else
@@ -225,9 +221,13 @@ public:
 				FunctionTypes |= UEdGraphSchema_K2::EFunctionType::FT_Imperative;
 			}
 
-			if(K2Schema->CanFunctionBeUsedInGraph(Blueprint->GeneratedClass, FunctionPtr, InDestGraph, FunctionTypes, false, FFunctionTargetInfo()))
+			if(K2Schema->CanFunctionBeUsedInGraph(Blueprint->GeneratedClass, FunctionPtr, InDestGraph, FunctionTypes, false))
 			{
-				FK2ActionMenuBuilder::AddSpawnInfoForFunction(FunctionPtr, false, FFunctionTargetInfo(), FMemberReference(), TEXT(""), K2Schema->AG_LevelReference, InPaletteBuilder);
+				IBlueprintNodeBinder::FBindingSet Bindings;
+				UEdGraphNode* GraphNode = UBlueprintFunctionNodeSpawner::Create(FunctionPtr)->Invoke(InDestGraph, Bindings, InOutDestPosition);
+
+				InOutDestPosition.Y = FMath::Max(InOutDestPosition.Y, GraphNode->NodePosY + UEdGraphSchema_K2::EstimateNodeHeight(GraphNode));
+				OutNodes.Add(GraphNode);
 			}
 		}
 	}
@@ -244,9 +244,12 @@ public:
 	FMacroNodeSpawnInfo(UEdGraph* InMacroGraph) : MacroGraph(InMacroGraph) {}
 
 	// FNodeSpawnInfo interface
-	virtual void GetActions(FBlueprintPaletteListBuilder& InPaletteBuilder, UEdGraph* InDestGraph) override
+	virtual void GetActions(UEdGraph* InDestGraph, FVector2D& InOutDestPosition, TArray<UEdGraphNode*>& OutNodes) override
 	{
-		FK2ActionMenuBuilder::AttachMacroGraphAction(InPaletteBuilder, MacroGraph);
+		IBlueprintNodeBinder::FBindingSet Bindings;
+		UK2Node_MacroInstance* MacroInstanceNode = Cast<UK2Node_MacroInstance>(UBlueprintNodeSpawner::Create(UK2Node_MacroInstance::StaticClass())->Invoke(InDestGraph, Bindings, InOutDestPosition));
+		MacroInstanceNode->SetMacroGraph(MacroGraph);
+		MacroInstanceNode->ReconstructNode();
 	}
 	// End of FNodeSpawnInfo interface
 
@@ -259,22 +262,20 @@ class FActorRefSpawnInfo : public FNodeSpawnInfo
 {
 public:
 	// FNodeSpawnInfo interface
-	virtual void GetActions(FBlueprintPaletteListBuilder& InPaletteBuilder, UEdGraph* InDestGraph) override
+	virtual void GetActions(UEdGraph* InDestGraph, FVector2D& InOutDestPosition, TArray<UEdGraphNode*>& OutNodes) override
 	{
-		const UBlueprint* Blueprint = InPaletteBuilder.Blueprint;
+		const UBlueprint* Blueprint = FBlueprintEditorUtils::FindBlueprintForGraph(InDestGraph);
 		USelection* SelectedLvlActors = GEditor->GetSelectedActors();
 
 		if ((Blueprint != nullptr) && Blueprint->ParentClass->IsChildOf<ALevelScriptActor>() && (SelectedLvlActors->Num() > 0))
 		{
 			for (FSelectionIterator LvlActorIt(*SelectedLvlActors); LvlActorIt; ++LvlActorIt)
 			{
-				UK2Node_Literal* TemplateRefNode = InPaletteBuilder.CreateTemplateNode<UK2Node_Literal>();
+				IBlueprintNodeBinder::FBindingSet Bindings;
+				UK2Node_Literal* TemplateRefNode = Cast<UK2Node_Literal>(UBlueprintNodeSpawner::Create(UK2Node_Literal::StaticClass())->Invoke(InDestGraph, Bindings, InOutDestPosition));
 				TemplateRefNode->SetObjectRef(*LvlActorIt);
 
-				TSharedPtr<FEdGraphSchemaAction_NewNode> NewActionNode = TSharedPtr<FEdGraphSchemaAction_NewNode>(new FEdGraphSchemaAction_NewNode);
-				NewActionNode->NodeTemplate = TemplateRefNode;
-
-				InPaletteBuilder.AddAction(NewActionNode);
+				OutNodes.Add(TemplateRefNode);
 			}
 		}
 	}
@@ -395,7 +396,7 @@ void FBlueprintSpawnNodeCommands::RegisterCommands()
 }
 
 
-void FBlueprintSpawnNodeCommands::GetGraphActionByChord(FInputChord& InChord, FBlueprintPaletteListBuilder& InPaletteBuilder, UEdGraph* InDestGraph) const
+void FBlueprintSpawnNodeCommands::GetGraphActionByChord(FInputChord& InChord, UEdGraph* InDestGraph, FVector2D& InOutDestPosition, TArray<UEdGraphNode*>& OutNodes) const
 {
 	if(InChord.IsValidChord())
 	{
@@ -403,7 +404,7 @@ void FBlueprintSpawnNodeCommands::GetGraphActionByChord(FInputChord& InChord, FB
 		{
 			if(NodeCommands[x]->CommandInfo->GetActiveChord().Get() == InChord)
 			{
-				NodeCommands[x]->GetActions(InPaletteBuilder, InDestGraph);
+				NodeCommands[x]->GetActions(InDestGraph, InOutDestPosition, OutNodes);
 			}
 		}
 	}
diff --git a/Engine/Source/Editor/Kismet/Private/BlueprintEditorCommands.h b/Engine/Source/Editor/Kismet/Private/BlueprintEditorCommands.h
--- a/Engine/Source/Editor/Kismet/Private/BlueprintEditorCommands.h
+++ b/Engine/Source/Editor/Kismet/Private/BlueprintEditorCommands.h
@@ -90,11 +90,12 @@ public:
 	/**
 	 * Returns a graph action assigned to the passed in chord
 	 *
-	 * @param InChord			The chord to use for lookup
-	 * @param InPaletteBuilder	The Blueprint palette to create the graph action for, used for validation purposes and to link any important node data to the blueprint
-	 * @param InDestGraph		The graph to create the graph action for, used for validation purposes and to link any important node data to the graph
+	 * @param InChord					The chord to use for lookup
+	 * @param InDestGraph				The graph to create the graph action for, used for validation purposes and to link any important node data to the graph
+	 * @param InOutDestPosition			Position to start placing nodes, will be updated to be at the next safe position for node placement
+	 * @param OutNodes					All nodes spawned by this operation
 	 */
-	void GetGraphActionByChord(FInputChord& InChord, struct FBlueprintPaletteListBuilder& InPaletteBuilder, UEdGraph* InDestGraph) const;
+	void GetGraphActionByChord(FInputChord& InChord, UEdGraph* InDestGraph, FVector2D& InOutDestPosition, TArray<UEdGraphNode*>& OutNodes) const;
 
 private:
 	/** An array of all the possible commands for spawning nodes */
diff --git a/Engine/Source/Editor/Kismet/Private/Nodes/K2Node_AsyncAction.cpp b/Engine/Source/Editor/Kismet/Private/Nodes/K2Node_AsyncAction.cpp
--- a/Engine/Source/Editor/Kismet/Private/Nodes/K2Node_AsyncAction.cpp
+++ b/Engine/Source/Editor/Kismet/Private/Nodes/K2Node_AsyncAction.cpp
@@ -16,44 +16,6 @@ UK2Node_AsyncAction::UK2Node_AsyncAction(const FObjectInitializer& ObjectInitial
 	ProxyActivateFunctionName = GET_FUNCTION_NAME_CHECKED(UBlueprintAsyncActionBase, Activate);
 }
 
-void UK2Node_AsyncAction::GetMenuEntries(FGraphContextMenuBuilder& ContextMenuBuilder) const
-{
-	const UEdGraphSchema_K2* K2Schema = GetDefault<UEdGraphSchema_K2>();
-	EGraphType GraphType = K2Schema->GetGraphType(ContextMenuBuilder.CurrentGraph);
-	const bool bAllowLatentFuncs = (GraphType == GT_Ubergraph || GraphType == GT_Macro);
-
-
-	if (bAllowLatentFuncs)
-	{
-		for (TObjectIterator<UClass> ClassIt; ClassIt; ++ClassIt)
-		{
-			UClass* TestClass = *ClassIt;
-			if ( TestClass->IsChildOf(UBlueprintAsyncActionBase::StaticClass()) && !TestClass->HasAnyClassFlags(CLASS_Abstract | CLASS_Deprecated) )
-			{
-				for (TFieldIterator<UFunction> FuncIt(TestClass, EFieldIteratorFlags::ExcludeSuper); FuncIt; ++FuncIt)
-				{
-					// See if the function is a static factory method for online proxies
-					UFunction* CandidateFunction = *FuncIt;
-
-					UObjectProperty* ReturnProperty = Cast<UObjectProperty>(CandidateFunction->GetReturnProperty());
-					const bool bValidReturnType = ( ReturnProperty != nullptr ) && ( ReturnProperty->PropertyClass != nullptr ) && ( ReturnProperty->PropertyClass->IsChildOf(UBlueprintAsyncActionBase::StaticClass()) );
-
-					if (CandidateFunction->HasAllFunctionFlags(FUNC_Static) && bValidReturnType)
-					{
-						// Create a node template for this factory method
-						UK2Node_AsyncAction* NodeTemplate = NewObject<UK2Node_AsyncAction>(ContextMenuBuilder.OwnerOfTemporaries);
-						NodeTemplate->ProxyFactoryFunctionName = CandidateFunction->GetFName();
-						NodeTemplate->ProxyFactoryClass = TestClass;
-						NodeTemplate->ProxyClass = ReturnProperty->PropertyClass;
-
-						CreateDefaultMenuEntry(NodeTemplate, ContextMenuBuilder);
-					}
-				}
-			}
-		}
-	}
-}
-
 void UK2Node_AsyncAction::GetMenuActions(FBlueprintActionDatabaseRegistrar& ActionRegistrar) const
 {
 	// these nested loops are combing over the same classes/functions the
diff --git a/Engine/Source/Editor/Kismet/Private/SBlueprintActionMenu.cpp b/Engine/Source/Editor/Kismet/Private/SBlueprintActionMenu.cpp
--- a/Engine/Source/Editor/Kismet/Private/SBlueprintActionMenu.cpp
+++ b/Engine/Source/Editor/Kismet/Private/SBlueprintActionMenu.cpp
@@ -9,7 +9,6 @@
 #include "BlueprintEditor.h"
 #include "SMyBlueprint.h"
 #include "BlueprintEditorUtils.h"
-#include "K2ActionMenuBuilder.h" // for FBlueprintPaletteListBuilder
 #include "BlueprintActionMenuBuilder.h"
 #include "BlueprintActionFilter.h"
 #include "BlueprintActionMenuUtils.h"
diff --git a/Engine/Source/Editor/Kismet/Private/SBlueprintFavoritesPalette.cpp b/Engine/Source/Editor/Kismet/Private/SBlueprintFavoritesPalette.cpp
--- a/Engine/Source/Editor/Kismet/Private/SBlueprintFavoritesPalette.cpp
+++ b/Engine/Source/Editor/Kismet/Private/SBlueprintFavoritesPalette.cpp
@@ -4,7 +4,6 @@
 #include "SBlueprintFavoritesPalette.h"
 #include "Engine.h"
 #include "KismetClasses.h"
-#include "K2ActionMenuBuilder.h"
 #include "BlueprintActionMenuUtils.h"
 #include "BlueprintActionMenuBuilder.h"
 #include "BlueprintActionFilter.h" // for FBlueprintActionContext
diff --git a/Engine/Source/Editor/Kismet/Private/SBlueprintFavoritesPalette.h b/Engine/Source/Editor/Kismet/Private/SBlueprintFavoritesPalette.h
--- a/Engine/Source/Editor/Kismet/Private/SBlueprintFavoritesPalette.h
+++ b/Engine/Source/Editor/Kismet/Private/SBlueprintFavoritesPalette.h
@@ -3,7 +3,6 @@
 #pragma once
 
 #include "SBlueprintSubPalette.h"
-#include "K2ActionMenuBuilder.h"
 
 /*******************************************************************************
 * SBlueprintFavoritesPalette
diff --git a/Engine/Source/Editor/Kismet/Private/SBlueprintLibraryPalette.cpp b/Engine/Source/Editor/Kismet/Private/SBlueprintLibraryPalette.cpp
--- a/Engine/Source/Editor/Kismet/Private/SBlueprintLibraryPalette.cpp
+++ b/Engine/Source/Editor/Kismet/Private/SBlueprintLibraryPalette.cpp
@@ -4,7 +4,6 @@
 #include "SBlueprintLibraryPalette.h"
 #include "ClassViewerModule.h" // for FClassViewerInitializationOptions
 #include "ClassViewerFilter.h" // for FClassViewerFilterFuncs
-#include "K2ActionMenuBuilder.h"
 #include "BlueprintPaletteFavorites.h"
 #include "BlueprintActionFilter.h"
 #include "BlueprintActionMenuBuilder.h"
diff --git a/Engine/Source/Editor/Kismet/Private/SBlueprintSubPalette.cpp b/Engine/Source/Editor/Kismet/Private/SBlueprintSubPalette.cpp
--- a/Engine/Source/Editor/Kismet/Private/SBlueprintSubPalette.cpp
+++ b/Engine/Source/Editor/Kismet/Private/SBlueprintSubPalette.cpp
@@ -100,7 +100,7 @@ static bool CanPaletteItemBePlaced(TSharedPtr<FEdGraphSchemaAction> DropActionIn
 					}
 
 					const UClass* GeneratedClass = FBlueprintEditorUtils::FindBlueprintForGraphChecked(HoveredGraphIn)->GeneratedClass;
-					bCanBePlaced = K2Schema->CanFunctionBeUsedInGraph(GeneratedClass, Function, HoveredGraphIn, AllowedFunctionTypes, false, FFunctionTargetInfo(), &ImpededReasonOut);
+					bCanBePlaced = K2Schema->CanFunctionBeUsedInGraph(GeneratedClass, Function, HoveredGraphIn, AllowedFunctionTypes, false, &ImpededReasonOut);
 				}
 			}
 		}
diff --git a/Engine/Source/Editor/Kismet/Public/Nodes/K2Node_AsyncAction.h b/Engine/Source/Editor/Kismet/Public/Nodes/K2Node_AsyncAction.h
--- a/Engine/Source/Editor/Kismet/Public/Nodes/K2Node_AsyncAction.h
+++ b/Engine/Source/Editor/Kismet/Public/Nodes/K2Node_AsyncAction.h
@@ -9,8 +9,7 @@ class KISMET_API UK2Node_AsyncAction : public UK2Node_BaseAsyncTask
 {
 	GENERATED_UCLASS_BODY()
 	
-	// UEdGraphNode interface
-	virtual void GetMenuEntries(FGraphContextMenuBuilder& ContextMenuBuilder) const override;
+	// UK2Node interface
 	virtual void GetMenuActions(FBlueprintActionDatabaseRegistrar& ActionRegistrar) const override;
-	// End of UEdGraphNode interface
+	// End of UK2Node interface
 };
diff --git a/Engine/Source/Editor/OnlineBlueprintSupport/Classes/K2Node_LatentOnlineCall.h b/Engine/Source/Editor/OnlineBlueprintSupport/Classes/K2Node_LatentOnlineCall.h
--- a/Engine/Source/Editor/OnlineBlueprintSupport/Classes/K2Node_LatentOnlineCall.h
+++ b/Engine/Source/Editor/OnlineBlueprintSupport/Classes/K2Node_LatentOnlineCall.h
@@ -10,8 +10,7 @@ class ONLINEBLUEPRINTSUPPORT_API UK2Node_LatentOnlineCall : public UK2Node_BaseA
 {
 	GENERATED_UCLASS_BODY()
 	
-	// UEdGraphNode interface
-	virtual void GetMenuEntries(FGraphContextMenuBuilder& ContextMenuBuilder) const override;
+	// UK2Node interface
 	virtual void GetMenuActions(FBlueprintActionDatabaseRegistrar& ActionRegistrar) const override;
-	// End of UEdGraphNode interface
+	// End of UK2Node interface
 };
diff --git a/Engine/Source/Editor/OnlineBlueprintSupport/Private/K2Node_LatentOnlineCall.cpp b/Engine/Source/Editor/OnlineBlueprintSupport/Private/K2Node_LatentOnlineCall.cpp
--- a/Engine/Source/Editor/OnlineBlueprintSupport/Private/K2Node_LatentOnlineCall.cpp
+++ b/Engine/Source/Editor/OnlineBlueprintSupport/Private/K2Node_LatentOnlineCall.cpp
@@ -14,44 +14,6 @@ UK2Node_LatentOnlineCall::UK2Node_LatentOnlineCall(const FObjectInitializer& Obj
 	ProxyActivateFunctionName = GET_FUNCTION_NAME_CHECKED(UOnlineBlueprintCallProxyBase, Activate);
 }
 
-void UK2Node_LatentOnlineCall::GetMenuEntries(FGraphContextMenuBuilder& ContextMenuBuilder) const
-{
-	const UEdGraphSchema_K2* K2Schema = GetDefault<UEdGraphSchema_K2>();
-	EGraphType GraphType = K2Schema->GetGraphType(ContextMenuBuilder.CurrentGraph);
-	const bool bAllowLatentFuncs = (GraphType == GT_Ubergraph || GraphType == GT_Macro);
-
-
-	if (bAllowLatentFuncs)
-	{
-		for (TObjectIterator<UClass> ClassIt; ClassIt; ++ClassIt)
-		{
-			UClass* TestClass = *ClassIt;
-			if (TestClass->IsChildOf(UOnlineBlueprintCallProxyBase::StaticClass()) && !TestClass->HasAnyClassFlags(CLASS_Abstract | CLASS_Deprecated))
-			{
-				for (TFieldIterator<UFunction> FuncIt(TestClass, EFieldIteratorFlags::ExcludeSuper); FuncIt; ++FuncIt)
-				{
-					// See if the function is a static factory method for online proxies
-					UFunction* CandidateFunction = *FuncIt;
-
-					UObjectProperty* ReturnProperty = Cast<UObjectProperty>(CandidateFunction->GetReturnProperty());
-					const bool bValidReturnType = (ReturnProperty != nullptr) && (ReturnProperty->PropertyClass != nullptr) && (ReturnProperty->PropertyClass->IsChildOf(UOnlineBlueprintCallProxyBase::StaticClass()));
-
-					if (CandidateFunction->HasAllFunctionFlags(FUNC_Static) && bValidReturnType)
-					{
-						// Create a node template for this factory method
-						UK2Node_LatentOnlineCall* NodeTemplate = NewObject<UK2Node_LatentOnlineCall>(ContextMenuBuilder.OwnerOfTemporaries);
-						NodeTemplate->ProxyFactoryFunctionName = CandidateFunction->GetFName();
-						NodeTemplate->ProxyFactoryClass = TestClass;
-						NodeTemplate->ProxyClass = ReturnProperty->PropertyClass;
-
-						CreateDefaultMenuEntry(NodeTemplate, ContextMenuBuilder);
-					}
-				}
-			}
-		}
-	}
-}
-
 void UK2Node_LatentOnlineCall::GetMenuActions(FBlueprintActionDatabaseRegistrar& ActionRegistrar) const
 {
 	// these nested loops are combing over the same classes/functions the
diff --git a/Engine/Source/Editor/UMGEditor/Private/Nodes/K2Node_CreateDragDropOperation.cpp b/Engine/Source/Editor/UMGEditor/Private/Nodes/K2Node_CreateDragDropOperation.cpp
--- a/Engine/Source/Editor/UMGEditor/Private/Nodes/K2Node_CreateDragDropOperation.cpp
+++ b/Engine/Source/Editor/UMGEditor/Private/Nodes/K2Node_CreateDragDropOperation.cpp
@@ -6,7 +6,6 @@
 #include "BlueprintNodeSpawner.h"
 #include "EditorCategoryUtils.h"
 #include "KismetCompiler.h"
-#include "K2ActionMenuBuilder.h"
 #include "K2Node_CreateDragDropOperation.h"
 
 #define LOCTEXT_NAMESPACE "UMG"
@@ -37,18 +36,6 @@ FText UK2Node_CreateDragDropOperation::GetMenuCategory() const
 	return FEditorCategoryUtils::GetCommonCategory(FCommonEditorCategory::UserInterface);
 }
 
-void UK2Node_CreateDragDropOperation::GetMenuEntries(FGraphContextMenuBuilder& ContextMenuBuilder) const
-{
-	UK2Node_CreateDragDropOperation* TemplateNode = NewObject<UK2Node_CreateDragDropOperation>(GetTransientPackage(), GetClass());
-
-	const FString Category = TEXT("User Interface");
-	const FText   MenuDesc = LOCTEXT("CreateDragDropOperationMenuOption", "Create Drag & Drop Operation...");
-	const FString Tooltip  = TEXT("Create a new UI drag and drop operation.  Use inside of OnDragDetected.");
-
-	TSharedPtr<FEdGraphSchemaAction_K2NewNode> NodeAction = FK2ActionMenuBuilder::AddNewNodeAction(ContextMenuBuilder, Category, MenuDesc, Tooltip);
-	NodeAction->NodeTemplate = TemplateNode;
-}
-
 FName UK2Node_CreateDragDropOperation::GetCornerIcon() const
 {
 	return TEXT("Graph.Replication.ClientEvent");
diff --git a/Engine/Source/Editor/UMGEditor/Private/Nodes/K2Node_CreateDragDropOperation.h b/Engine/Source/Editor/UMGEditor/Private/Nodes/K2Node_CreateDragDropOperation.h
--- a/Engine/Source/Editor/UMGEditor/Private/Nodes/K2Node_CreateDragDropOperation.h
+++ b/Engine/Source/Editor/UMGEditor/Private/Nodes/K2Node_CreateDragDropOperation.h
@@ -17,7 +17,6 @@ class UMGEDITOR_API UK2Node_CreateDragDropOperation : public UK2Node_ConstructOb
 	// End UEdGraphNode interface.
 
 	// Begin UK2Node interface
-	virtual void GetMenuEntries(FGraphContextMenuBuilder& ContextMenuBuilder) const override;
 	virtual FText GetMenuCategory() const override;
 	virtual FName GetCornerIcon() const override;
 	// End UK2Node interface.
diff --git a/Engine/Source/Editor/UMGEditor/Private/Nodes/K2Node_CreateWidget.cpp b/Engine/Source/Editor/UMGEditor/Private/Nodes/K2Node_CreateWidget.cpp
--- a/Engine/Source/Editor/UMGEditor/Private/Nodes/K2Node_CreateWidget.cpp
+++ b/Engine/Source/Editor/UMGEditor/Private/Nodes/K2Node_CreateWidget.cpp
@@ -6,7 +6,6 @@
 #include "KismetCompiler.h"
 #include "BlueprintNodeSpawner.h"
 #include "EditorCategoryUtils.h"
-#include "K2ActionMenuBuilder.h"
 
 #include "K2Node_CreateWidget.h"
 
@@ -61,18 +60,6 @@ FText UK2Node_CreateWidget::GetMenuCategory() const
 	return FEditorCategoryUtils::GetCommonCategory(FCommonEditorCategory::UserInterface);
 }
 
-void UK2Node_CreateWidget::GetMenuEntries(FGraphContextMenuBuilder& ContextMenuBuilder) const
-{
-	UK2Node_CreateWidget* TemplateNode = NewObject<UK2Node_CreateWidget>(GetTransientPackage(), GetClass());
-
-	const FString Category = TEXT("User Interface");
-	const FText   MenuDesc = LOCTEXT("CreateWidgetMenuOption", "Create Widget...");
-	const FString Tooltip  = TEXT("Create a new UI widget");
-
-	TSharedPtr<FEdGraphSchemaAction_K2NewNode> NodeAction = FK2ActionMenuBuilder::AddNewNodeAction(ContextMenuBuilder, Category, MenuDesc, Tooltip);
-	NodeAction->NodeTemplate = TemplateNode;
-}
-
 FName UK2Node_CreateWidget::GetCornerIcon() const
 {
 	return TEXT("Graph.Replication.ClientEvent");
diff --git a/Engine/Source/Editor/UMGEditor/Private/Nodes/K2Node_CreateWidget.h b/Engine/Source/Editor/UMGEditor/Private/Nodes/K2Node_CreateWidget.h
--- a/Engine/Source/Editor/UMGEditor/Private/Nodes/K2Node_CreateWidget.h
+++ b/Engine/Source/Editor/UMGEditor/Private/Nodes/K2Node_CreateWidget.h
@@ -18,7 +18,6 @@ class UMGEDITOR_API UK2Node_CreateWidget : public UK2Node_ConstructObjectFromCla
 	// End UEdGraphNode interface.
 
 	// Begin UK2Node interface
-	virtual void GetMenuEntries(FGraphContextMenuBuilder& ContextMenuBuilder) const override;
 	virtual FText GetMenuCategory() const override;
 	virtual FName GetCornerIcon() const override;
 	// End UK2Node interface.
diff --git a/Engine/Source/Editor/UnrealEd/Private/Commandlets/DumpBlueprintsInfoCommandlet.cpp b/Engine/Source/Editor/UnrealEd/Private/Commandlets/DumpBlueprintsInfoCommandlet.cpp
--- a/Engine/Source/Editor/UnrealEd/Private/Commandlets/DumpBlueprintsInfoCommandlet.cpp
+++ b/Engine/Source/Editor/UnrealEd/Private/Commandlets/DumpBlueprintsInfoCommandlet.cpp
@@ -3,7 +3,6 @@
 #include "UnrealEd.h"
 #include "ObjectEditorUtils.h"
 #include "Kismet2/KismetEditorUtilities.h"
-#include "K2ActionMenuBuilder.h"
 #include "EdGraphSchema_K2_Actions.h"
 #include "K2Node.h"
 #include "AnimationGraph.h"
@@ -107,9 +106,6 @@ DumpBlueprintsInfo commandlet params: \n\
                         you wish to use your own tool then specify it here. Use\n\
                         '{1}' and '{2}' as placeholders for filenames, like so:\n\
                         -diffcmd=\"AraxisP4Diff.exe {2} {1}\".                 \n\
-\n\
-    -legacy             Uses the old way of constructing Blueprint action      \n\
-                        menus.                                                 \n\
 \n\
     -name=<Filename>    Overrides the default filename. Leave off the extention\n\
                         (this will add .json to the end). When -multifile is   \n\
@@ -144,7 +140,6 @@ DumpBlueprintsInfo commandlet params: \n\
 		BPDUMP_DoNotDumpActionInfo  = (1<<7),
 		BPDUMP_RecordTiming			= (1<<8),
 		BPDUMP_SelectAllObjTypes    = (1<<9), 
-		BPDUMP_UseLegacyMenuBuilder = (1<<10),
 
 		BPDUMP_ActionDatabaseInfo   = (1<<11),
 
@@ -299,7 +294,7 @@ DumpBlueprintsInfo commandlet params: \n\
 	 *
 	 * @return The amount of time (in seconds) that the menu building took.
 	 */
-	static double GetPaletteMenuActions(FBlueprintPaletteListBuilder& ActionListBuilder, UClass* PaletteFilter);
+	static double GetPaletteMenuActions(FCategorizedGraphActionListBuilder& PaletteBuilder, UBlueprint* Blueprint, UClass* PaletteFilter);
 
 	/**
 	 * Dumps all palette actions listed for the specified blueprint. Determines
@@ -368,7 +363,7 @@ DumpBlueprintsInfo commandlet params: \n\
 	 * into the DumpActionList(), while dumping context information to go along 
 	 * with it (to give the reader context).
 	 */
-	static void DumpContextActionList(uint32 Indent, FBlueprintGraphActionListBuilder ActionBuilder, FArchive* FileOutWriter);
+	static void DumpContextActionList(uint32 Indent, FGraphContextMenuBuilder ActionBuilder, UBlueprint* Blueprint, FArchive* FileOutWriter);
 
 	/**
 	 * Assumes that the specified ActionListBuilder is configured with all the 
@@ -378,13 +373,13 @@ DumpBlueprintsInfo commandlet params: \n\
 	 *
 	 * @return The amount of time (in seconds) that the menu building took.
 	 */
-	static double GetContextMenuActions(FBlueprintGraphActionListBuilder& ActionListBuilder);
+	static double GetContextMenuActions(FGraphContextMenuBuilder ActionBuilder, UBlueprint* Blueprint);
 
 	/**
 	 * Looks at the filter/context items set on the specified GraphActionListBuilder
 	 * and writes them out, to provide context with any dumped actions.
 	 */
-	static void DumpContextInfo(uint32 Indent, FBlueprintGraphActionListBuilder const& ActionBuilder, FArchive* FileOutWriter);
+	static void DumpContextInfo(uint32 Indent, FGraphContextMenuBuilder ActionBuilder, UBlueprint* Blueprint, FArchive* FileOutWriter);
 
 	/**
 	 * A utility function aimed at determining whether we should open a diff for
@@ -593,10 +588,6 @@ DumpBlueprintInfoUtils::CommandletOptions::CommandletOptions(TArray<FString> con
 		{
 			NewDumpFlags |= BPDUMP_RecordTiming;
 		}
-		else if (!Switch.Compare("legacy", ESearchCase::IgnoreCase))
-		{
-			NewDumpFlags |= BPDUMP_UseLegacyMenuBuilder;
-		}
 		else if (Switch.StartsWith("name="))
 		{
 			FString NameSwitch;
@@ -1107,7 +1098,7 @@ static bool DumpBlueprintInfoUtils::DumpActionDatabaseInfo(uint32 Indent, FArchi
 {
 	bool bWroteToFile = false;
 
-	uint32 const DbInfoMask = (BPDUMP_UseLegacyMenuBuilder | BPDUMP_ActionDatabaseInfo);
+	uint32 const DbInfoMask = BPDUMP_ActionDatabaseInfo;
 	if ((CommandOptions.DumpFlags & DbInfoMask) == BPDUMP_ActionDatabaseInfo)
 	{
 		UE_LOG(LogBlueprintInfoDump, Display, TEXT("%sDumping Database info..."), *BuildIndentString(Indent, true));
@@ -1489,28 +1480,20 @@ static void DumpBlueprintInfoUtils::DumpInfoForClass(uint32 Indent, UClass* Blue
 }
 
 //------------------------------------------------------------------------------
-static double DumpBlueprintInfoUtils::GetPaletteMenuActions(FBlueprintPaletteListBuilder& PaletteBuilder, UClass* PaletteFilter)
+static double DumpBlueprintInfoUtils::GetPaletteMenuActions(FCategorizedGraphActionListBuilder& PaletteBuilder, UBlueprint* Blueprint, UClass* PaletteFilter)
 {
 	PaletteBuilder.Empty();
 	UEdGraphSchema_K2 const* K2Schema = GetDefault<UEdGraphSchema_K2>();
 
 	double MenuBuildDuration = 0.0;
 	
-	if (CommandOptions.DumpFlags & BPDUMP_UseLegacyMenuBuilder)
-	{
+	FBlueprintActionContext FilterContext;
+	FilterContext.Blueprints.Add(const_cast<UBlueprint*>(Blueprint));
+
+	FBlueprintActionMenuBuilder MenuBuilder(nullptr);
+	{			
 		FScopedDurationTimer DurationTimer(MenuBuildDuration);
-		FK2ActionMenuBuilder(K2Schema).GetPaletteActions(PaletteBuilder, PaletteFilter);
-	}
-	else
-	{
-		FBlueprintActionContext FilterContext;
-		FilterContext.Blueprints.Add(const_cast<UBlueprint*>(PaletteBuilder.Blueprint));
-		
-		FBlueprintActionMenuBuilder MenuBuilder(nullptr);
-		{			
-			FScopedDurationTimer DurationTimer(MenuBuildDuration);
-			FBlueprintActionMenuUtils::MakePaletteMenu(FilterContext, PaletteFilter, MenuBuilder);
-		}
+		FBlueprintActionMenuUtils::MakePaletteMenu(FilterContext, PaletteFilter, MenuBuilder);
 		PaletteBuilder.Append(MenuBuilder);
 	}
 
@@ -1569,8 +1552,8 @@ static void DumpBlueprintInfoUtils::DumpPalette(uint32 Indent, UBlueprint* Bluep
 	}
 	else
 	{
-		FBlueprintPaletteListBuilder PaletteBuilder(Blueprint, TEXT("Library"));
-		double MenuBuildDuration = GetPaletteMenuActions(PaletteBuilder, ClassFilter);
+		FCategorizedGraphActionListBuilder PaletteBuilder;
+		double MenuBuildDuration = GetPaletteMenuActions(PaletteBuilder, Blueprint, ClassFilter);
 
 		BeginPaletteEntry += NestedIndent + "\"FilterClass\" : \"" + FilterClassName + "\",\n";
 		if (CommandOptions.DumpFlags & BPDUMP_RecordTiming)
@@ -1782,9 +1765,18 @@ static void DumpBlueprintInfoUtils::DumpGraphContextActions(uint32 Indent, UEdGr
 
 	FileOutWriter->Serialize(TCHAR_TO_ANSI(*BeginGraphEntry), BeginGraphEntry.Len());
 
-	FBlueprintGraphActionListBuilder ActionBuilder(Graph);
+	UObject* GraphOuter = Graph->GetOuter();
+	UBlueprint* Blueprint = Cast<UBlueprint>(GraphOuter);
+	while ((Blueprint == nullptr) && (GraphOuter != nullptr))
+	{
+		GraphOuter = GraphOuter->GetOuter();
+		Blueprint = Cast<UBlueprint>(GraphOuter);
+	}
+	check(Blueprint != nullptr);
+
+	FGraphContextMenuBuilder ActionBuilder(Graph);
 	UE_LOG(LogBlueprintInfoDump, Display, TEXT("%sDumping graph context actions..."), *BuildIndentString(Indent, true));
-	DumpContextActionList(Indent, ActionBuilder, FileOutWriter);
+	DumpContextActionList(Indent, ActionBuilder, Blueprint, FileOutWriter);
 
 	if (CommandOptions.DumpFlags & BPDUMP_SelectAllObjTypes)
 	{
@@ -1808,21 +1800,12 @@ static void DumpBlueprintInfoUtils::DumpGraphContextActions(uint32 Indent, UEdGr
 					FString ActorSelectionEntry = "," + IndentedNewline + "\"LevelActorMenu-" + Class->GetName() + "\" : \n";
 					FileOutWriter->Serialize(TCHAR_TO_ANSI(*ActorSelectionEntry), ActorSelectionEntry.Len());
 
-					DumpContextActionList(Indent, ActionBuilder, FileOutWriter);
+					DumpContextActionList(Indent, ActionBuilder, Blueprint, FileOutWriter);
 				}
 			}
 		}
 	}
 
-	UObject* GraphOuter = Graph->GetOuter();
-	UBlueprint* Blueprint = Cast<UBlueprint>(GraphOuter);
-	while ((Blueprint == nullptr) && (GraphOuter != nullptr))
-	{
-		GraphOuter = GraphOuter->GetOuter();
-		Blueprint = Cast<UBlueprint>(GraphOuter);
-	}
-	check(Blueprint != nullptr);
-
 	TArray<UObjectProperty*> ComponentProperties;
 	GetComponentProperties(Blueprint, ComponentProperties);
 
@@ -1841,7 +1824,7 @@ static void DumpBlueprintInfoUtils::DumpGraphContextActions(uint32 Indent, UEdGr
 
 		ActionBuilder.SelectedObjects.Empty();
 		ActionBuilder.SelectedObjects.Add(Component);
-		DumpContextActionList(Indent, ActionBuilder, FileOutWriter);
+		DumpContextActionList(Indent, ActionBuilder, Blueprint, FileOutWriter);
 	}
 
 	FString EndGraphEntry;
@@ -1867,7 +1850,7 @@ static void DumpBlueprintInfoUtils::DumpGraphContextActions(uint32 Indent, UEdGr
 //------------------------------------------------------------------------------
 static bool DumpBlueprintInfoUtils::DumpPinContextActions(uint32 Indent, UEdGraph* Graph, FArchive* FileOutWriter)
 {
-	FBlueprintGraphActionListBuilder ContextMenuBuilder(Graph);
+	FGraphContextMenuBuilder ContextMenuBuilder(Graph);
 
 	bool bWroteToFile = false;
 	if (!CommandOptions.PinType.IsEmpty())
@@ -2015,7 +1998,7 @@ static bool DumpBlueprintInfoUtils::DumpTypeTreeActions(uint32 Indent, UEdGraph*
 //------------------------------------------------------------------------------
 static void DumpBlueprintInfoUtils::DumpContextualPinTypeActions(uint32 Indent, UEdGraph* Graph, FEdGraphPinType const& PinType, FArchive* FileOutWriter)
 {
-	FBlueprintGraphActionListBuilder ContextMenuBuilder(Graph);
+	FGraphContextMenuBuilder ContextMenuBuilder(Graph);
 
 	UK2Node_Composite* DummyNode = NewObject<UK2Node_Composite>(Graph);
 	UEdGraphPin* DummyPin = DummyNode->CreatePin(
@@ -2029,14 +2012,15 @@ static void DumpBlueprintInfoUtils::DumpContextualPinTypeActions(uint32 Indent,
 	);
 	ContextMenuBuilder.FromPin = DummyPin;
 
-	DumpContextActionList(Indent, ContextMenuBuilder, FileOutWriter);
+	UBlueprint* Blueprint = FBlueprintEditorUtils::FindBlueprintForGraph(Graph);
+	DumpContextActionList(Indent, ContextMenuBuilder, Blueprint, FileOutWriter);
 	FileOutWriter->Serialize(TCHAR_TO_ANSI(TEXT(",\n")), 2);
 	DummyPin->Direction = EGPD_Output;
-	DumpContextActionList(Indent, ContextMenuBuilder, FileOutWriter);
+	DumpContextActionList(Indent, ContextMenuBuilder, Blueprint, FileOutWriter);
 }
 
 //------------------------------------------------------------------------------
-static void DumpBlueprintInfoUtils::DumpContextActionList(uint32 Indent, FBlueprintGraphActionListBuilder ActionBuilder, FArchive* FileOutWriter)
+static void DumpBlueprintInfoUtils::DumpContextActionList(uint32 Indent, FGraphContextMenuBuilder ActionBuilder, UBlueprint* Blueprint, FArchive* FileOutWriter)
 {
 	ActionBuilder.Empty();
 
@@ -2056,7 +2040,7 @@ static void DumpBlueprintInfoUtils::DumpContextActionList(uint32 Indent, FBluepr
 		UE_LOG(LogBlueprintInfoDump, Display, TEXT("%sDumping pin actions: %s"), *BuildIndentString(Indent, true), *PinTypeLog);
 	}
 	
-	double MenuBuildDuration = GetContextMenuActions(ActionBuilder);
+	double MenuBuildDuration = GetContextMenuActions(ActionBuilder, Blueprint);
 
 	FString const ContextEntryIndent = BuildIndentString(Indent);
 	FString BeginContextEntry = FString::Printf(TEXT("%s{\n"), *ContextEntryIndent);
@@ -2067,7 +2051,7 @@ static void DumpBlueprintInfoUtils::DumpContextActionList(uint32 Indent, FBluepr
 	}
 	FileOutWriter->Serialize(TCHAR_TO_ANSI(*BeginContextEntry), BeginContextEntry.Len());
 
-	DumpContextInfo(Indent + 1, ActionBuilder, FileOutWriter);
+	DumpContextInfo(Indent + 1, ActionBuilder, Blueprint, FileOutWriter);
 	FileOutWriter->Serialize(TCHAR_TO_ANSI(TEXT(",\n")), 2);
 	DumpActionList(Indent + 1, ActionBuilder, FileOutWriter);
 
@@ -2076,53 +2060,44 @@ static void DumpBlueprintInfoUtils::DumpContextActionList(uint32 Indent, FBluepr
 }
 
 //------------------------------------------------------------------------------
-static double DumpBlueprintInfoUtils::GetContextMenuActions(FBlueprintGraphActionListBuilder& ActionBuilder)
+static double DumpBlueprintInfoUtils::GetContextMenuActions(FGraphContextMenuBuilder ActionBuilder, UBlueprint* Blueprint)
 {
 	ActionBuilder.Empty();
 	check(ActionBuilder.CurrentGraph != nullptr);
 
 	double MenuBuildDuration = 0.0;
+	
+	FBlueprintActionContext FilterContext;
+	FilterContext.Blueprints.Add(const_cast<UBlueprint*>(Blueprint));
+	FilterContext.Graphs.Add(const_cast<UEdGraph*>(ActionBuilder.CurrentGraph));
 
-	UEdGraphSchema const* GraphSchema = GetDefault<UEdGraphSchema>(ActionBuilder.CurrentGraph->Schema);
-	if (CommandOptions.DumpFlags & BPDUMP_UseLegacyMenuBuilder)
+	if (ActionBuilder.FromPin != nullptr)
 	{
-		FScopedDurationTimer DurationTimer(MenuBuildDuration);
-		GraphSchema->GetGraphContextActions(ActionBuilder);
+		FilterContext.Pins.Add(const_cast<UEdGraphPin*>(ActionBuilder.FromPin));
 	}
-	else
+
+	TArray<UProperty*> SelectedProperties;
+	for (UObject* SelectedObj : ActionBuilder.SelectedObjects)
 	{
-		FBlueprintActionContext FilterContext;
-		FilterContext.Blueprints.Add(const_cast<UBlueprint*>(ActionBuilder.Blueprint));
-		FilterContext.Graphs.Add(const_cast<UEdGraph*>(ActionBuilder.CurrentGraph));
-		
-		if (ActionBuilder.FromPin != nullptr)
+		if (UProperty* SelectedProperty = Cast<UObjectProperty>(SelectedObj))
 		{
-			FilterContext.Pins.Add(const_cast<UEdGraphPin*>(ActionBuilder.FromPin));
+			FilterContext.SelectedObjects.Add(SelectedProperty);
 		}
-		
-		TArray<UProperty*> SelectedProperties;
-		for (UObject* SelectedObj : ActionBuilder.SelectedObjects)
-		{
-			if (UProperty* SelectedProperty = Cast<UObjectProperty>(SelectedObj))
-			{
-				FilterContext.SelectedObjects.Add(SelectedProperty);
-			}
-		}
-		
-		FBlueprintActionMenuBuilder MenuBuilder(nullptr);
-		{
-			FScopedDurationTimer DurationTimer(MenuBuildDuration);
-			FBlueprintActionMenuUtils::MakeContextMenu(FilterContext, /*bIsContextSensitive =*/true, 
-				(EContextTargetFlags::TARGET_Blueprint | EContextTargetFlags::TARGET_NodeTarget | EContextTargetFlags::TARGET_PinObject | EContextTargetFlags::TARGET_SiblingPinObjects), MenuBuilder);
-		}
-		ActionBuilder.Append(MenuBuilder);
 	}
 
+	FBlueprintActionMenuBuilder MenuBuilder(nullptr);
+	{
+		FScopedDurationTimer DurationTimer(MenuBuildDuration);
+		FBlueprintActionMenuUtils::MakeContextMenu(FilterContext, /*bIsContextSensitive =*/true, 
+			(EContextTargetFlags::TARGET_Blueprint | EContextTargetFlags::TARGET_NodeTarget | EContextTargetFlags::TARGET_PinObject | EContextTargetFlags::TARGET_SiblingPinObjects), MenuBuilder);
+	}
+	ActionBuilder.Append(MenuBuilder);
+
 	return MenuBuildDuration;
 }
 
 //------------------------------------------------------------------------------
-static void DumpBlueprintInfoUtils::DumpContextInfo(uint32 Indent, FBlueprintGraphActionListBuilder const& ActionBuilder, FArchive* FileOutWriter)
+static void DumpBlueprintInfoUtils::DumpContextInfo(uint32 Indent, FGraphContextMenuBuilder ActionBuilder, UBlueprint* Blueprint, FArchive* FileOutWriter)
 {
 	FString const ContextEntryIndent = BuildIndentString(Indent);
 	++Indent;
diff --git a/Engine/Source/Editor/UnrealEd/Private/Commandlets/GenerateBlueprintAPICommandlet.cpp b/Engine/Source/Editor/UnrealEd/Private/Commandlets/GenerateBlueprintAPICommandlet.cpp
--- a/Engine/Source/Editor/UnrealEd/Private/Commandlets/GenerateBlueprintAPICommandlet.cpp
+++ b/Engine/Source/Editor/UnrealEd/Private/Commandlets/GenerateBlueprintAPICommandlet.cpp
@@ -3,7 +3,6 @@
 #include "UnrealEd.h"
 #include "ObjectEditorUtils.h"
 #include "Kismet2/KismetEditorUtilities.h"
-#include "K2ActionMenuBuilder.h"
 #include "EdGraphSchema_K2_Actions.h"
 #include "K2Node.h"
 #include "K2Node_CommutativeAssociativeBinaryOperator.h"
@@ -183,7 +182,7 @@ GenerateBlueprintAPI commandlet params: \n\
 	 *
 	 * @return The amount of time (in seconds) that the menu building took.
 	 */
-	static double GetPaletteMenuActions(FBlueprintPaletteListBuilder& ActionListBuilder, UClass* PaletteFilter);
+	static double GetPaletteMenuActions(FCategorizedGraphActionListBuilder& PaletteBuilder, UBlueprint* Blueprint, UClass* PaletteFilter);
 
 	/**
 	 * Dumps all palette actions listed for the specified blueprint. Determines
@@ -613,31 +612,23 @@ static void GenerateBlueprintAPIUtils::DumpCategoryInfo(uint32 Indent, FArchive*
 }
 
 //------------------------------------------------------------------------------
-static double GenerateBlueprintAPIUtils::GetPaletteMenuActions(FBlueprintPaletteListBuilder& PaletteBuilder, UClass* PaletteFilter)
+static double GenerateBlueprintAPIUtils::GetPaletteMenuActions(FCategorizedGraphActionListBuilder& PaletteBuilder, UBlueprint* Blueprint, UClass* PaletteFilter)
 {
 	PaletteBuilder.Empty();
 	UEdGraphSchema_K2 const* K2Schema = GetDefault<UEdGraphSchema_K2>();
 
 	double MenuBuildDuration = 0.0;
 	
-	if (CommandOptions.DumpFlags & BPDUMP_UseLegacyMenuBuilder)
+	FBlueprintActionContext FilterContext;
+	FilterContext.Blueprints.Add(const_cast<UBlueprint*>(Blueprint));
+
+	FBlueprintActionMenuBuilder MenuBuilder(nullptr);
 	{
+		// prime the database so it's not recorded in our timing capture
+		FBlueprintActionDatabase::Get();
+
 		FScopedDurationTimer DurationTimer(MenuBuildDuration);
-		FK2ActionMenuBuilder(K2Schema).GetPaletteActions(PaletteBuilder, PaletteFilter);
-	}
-	else
-	{
-		FBlueprintActionContext FilterContext;
-		FilterContext.Blueprints.Add(const_cast<UBlueprint*>(PaletteBuilder.Blueprint));
-		
-		FBlueprintActionMenuBuilder MenuBuilder(nullptr);
-		{
-			// prime the database so it's not recorded in our timing capture
-			FBlueprintActionDatabase::Get();
-			
-			FScopedDurationTimer DurationTimer(MenuBuildDuration);
-			FBlueprintActionMenuUtils::MakePaletteMenu(FilterContext, PaletteFilter, MenuBuilder);
-		}
+		FBlueprintActionMenuUtils::MakePaletteMenu(FilterContext, PaletteFilter, MenuBuilder);
 		PaletteBuilder.Append(MenuBuilder);
 	}
 
@@ -696,8 +687,8 @@ static void GenerateBlueprintAPIUtils::DumpPalette(uint32 Indent, UBlueprint* Bl
 	}
 	else
 	{
-		FBlueprintPaletteListBuilder PaletteBuilder(Blueprint, TEXT("Library"));
-		double MenuBuildDuration = GetPaletteMenuActions(PaletteBuilder, ClassFilter);
+		FCategorizedGraphActionListBuilder PaletteBuilder;
+		double MenuBuildDuration = GetPaletteMenuActions(PaletteBuilder, Blueprint, ClassFilter);
 
 		BeginPaletteEntry += NestedIndent + "\"FilterClass\" : \"" + FilterClassName + "\",\n";
 		if (CommandOptions.DumpFlags & BPDUMP_RecordTiming)
diff --git a/Engine/Source/Editor/UnrealEd/Private/Editor.cpp b/Engine/Source/Editor/UnrealEd/Private/Editor.cpp
--- a/Engine/Source/Editor/UnrealEd/Private/Editor.cpp
+++ b/Engine/Source/Editor/UnrealEd/Private/Editor.cpp
@@ -157,7 +157,6 @@ FEditorDelegates::FOnFocusViewportOnActors				FEditorDelegates::OnFocusViewportO
 FEditorDelegates::FOnMapOpened							FEditorDelegates::OnMapOpened;
 FEditorDelegates::FOnEditorCameraMoved					FEditorDelegates::OnEditorCameraMoved;
 FEditorDelegates::FOnDollyPerspectiveCamera				FEditorDelegates::OnDollyPerspectiveCamera;
-FEditorDelegates::FOnBlueprintContextMenuCreated		FEditorDelegates::OnBlueprintContextMenuCreated;
 FSimpleMulticastDelegate								FEditorDelegates::OnShutdownPostPackagesSaved;
 FEditorDelegates::FOnAssetsPreDelete					FEditorDelegates::OnAssetsPreDelete;
 FEditorDelegates::FOnAssetsDeleted						FEditorDelegates::OnAssetsDeleted;
diff --git a/Engine/Source/Editor/UnrealEd/Public/Editor.h b/Engine/Source/Editor/UnrealEd/Public/Editor.h
--- a/Engine/Source/Editor/UnrealEd/Public/Editor.h
+++ b/Engine/Source/Editor/UnrealEd/Public/Editor.h
@@ -23,7 +23,6 @@ class FGeomPoly;
 class FGeomObject;
 class FScopedTransaction;
 class UFactory;
-struct FBlueprintGraphActionListBuilder;
 
 
 /** The shorthand identifier used for editor modes */
@@ -86,8 +85,6 @@ struct UNREALED_API FEditorDelegates
 	DECLARE_MULTICAST_DELEGATE_OneParam(FOnFocusViewportOnActors, const TArray<AActor*>&);
 	/** delegate type for triggering when a map is opened */
 	DECLARE_MULTICAST_DELEGATE_TwoParams(FOnMapOpened, const FString& /* Filename */, bool /*bAsTemplate*/);
-	/** delegate type for triggering when the Blueprint ContextMenu is created */
-	DECLARE_MULTICAST_DELEGATE_OneParam(FOnBlueprintContextMenuCreated, FBlueprintGraphActionListBuilder &/*ContextMenuBuilder*/);
 	/** Delegate used for entering or exiting an editor mode */
 	DECLARE_MULTICAST_DELEGATE_OneParam(FOnEditorModeTransitioned, FEdMode* /*Mode*/);
 	/** delegate type for when a user requests to delete certain assets... DOES NOT mean the asset(s) will be deleted (the user could cancel) */
@@ -185,8 +182,6 @@ struct UNREALED_API FEditorDelegates
 	static FOnEditorCameraMoved OnEditorCameraMoved;
 	/** Called when the editor camera is moved */
 	static FOnDollyPerspectiveCamera OnDollyPerspectiveCamera;
-	/** Called when the Blueprint ContextMenu is created */
-	static FOnBlueprintContextMenuCreated OnBlueprintContextMenuCreated;
 	/** Called on editor shutdown after packages have been successfully saved */
 	static FSimpleMulticastDelegate OnShutdownPostPackagesSaved;
 	/** Called when the user requests certain assets be deleted (DOES NOT imply that the asset will be deleted... the user could cancel) */
diff --git a/Engine/Source/Runtime/Engine/Classes/EdGraph/EdGraphSchema.h b/Engine/Source/Runtime/Engine/Classes/EdGraph/EdGraphSchema.h
--- a/Engine/Source/Runtime/Engine/Classes/EdGraph/EdGraphSchema.h
+++ b/Engine/Source/Runtime/Engine/Classes/EdGraph/EdGraphSchema.h
@@ -365,10 +365,7 @@ public:
 	// The selected objects
 	TArray<UObject*> SelectedObjects;
 public:
-	ENGINE_API FGraphContextMenuBuilder(const UEdGraph* InGraph) 
-		: CurrentGraph(InGraph)
-	{
-	}
+	ENGINE_API FGraphContextMenuBuilder(const UEdGraph* InGraph);
 };
 
 /** This is a response from GetGraphDisplayInformation */
diff --git a/Engine/Source/Runtime/Engine/Private/EdGraph/EdGraphSchema.cpp b/Engine/Source/Runtime/Engine/Private/EdGraph/EdGraphSchema.cpp
--- a/Engine/Source/Runtime/Engine/Private/EdGraph/EdGraphSchema.cpp
+++ b/Engine/Source/Runtime/Engine/Private/EdGraph/EdGraphSchema.cpp
@@ -128,6 +128,15 @@ void FCategorizedGraphActionListBuilder::AddActionList(TArray<TSharedPtr<FEdGrap
 	FGraphActionListBuilderBase::AddActionList(NewActions, ConcatCategories(Category, CategoryIn));
 }
 
+/////////////////////////////////////////////////////
+// FGraphContextMenuBuilder
+
+FGraphContextMenuBuilder::FGraphContextMenuBuilder(const UEdGraph* InGraph) 
+	: CurrentGraph(InGraph)
+{
+	OwnerOfTemporaries =  NewObject<UEdGraph>((UObject*)GetTransientPackage());
+}
+
 /////////////////////////////////////////////////////
 // FEdGraphSchemaAction_NewNode
 
