Fixed active timer registration issue when simulating if the viewport wasn't realtime

diff --git a/Engine/Source/Editor/UnrealEd/Private/EditorViewportClient.cpp b/Engine/Source/Editor/UnrealEd/Private/EditorViewportClient.cpp
--- a/Engine/Source/Editor/UnrealEd/Private/EditorViewportClient.cpp
+++ b/Engine/Source/Editor/UnrealEd/Private/EditorViewportClient.cpp
@@ -276,6 +276,7 @@ FEditorViewportClient::FEditorViewportClient(FEditorModeTools* InModeTools, FPre
 	, NearPlane(-1.0f)
 	, FarPlane(0.0f)
 	, bInGameViewMode(false)
+	, bShouldInvalidateViewportWidget(false)
 {
 	if (ModeTools == nullptr)
 	{
@@ -367,6 +368,33 @@ FEditorViewportClient::~FEditorViewportClient()
 	}
 }
 
+bool FEditorViewportClient::ToggleRealtime()
+{
+	SetRealtime(!bIsRealtime);
+	return bIsRealtime;
+}
+
+void FEditorViewportClient::SetRealtime(bool bInRealtime, bool bStoreCurrentValue)
+{
+	if (bStoreCurrentValue)
+	{
+		//Cache the Realtime and ShowStats flags
+		bStoredRealtime = bIsRealtime;
+		bStoredShowStats = bShowStats;
+	}
+
+	bIsRealtime = bInRealtime;
+
+	if (!bIsRealtime)
+	{
+		SetShowStats(false);
+	}
+	else
+	{
+		bShouldInvalidateViewportWidget = true;
+	}
+}
+
 void FEditorViewportClient::RestoreRealtime(const bool bAllowDisable)
 {
 	if (bAllowDisable)
@@ -378,13 +406,27 @@ void FEditorViewportClient::RestoreRealtime(const bool bAllowDisable)
 	{
 		bIsRealtime |= bStoredRealtime;
 		bShowStats |= bStoredShowStats;
+	}	
+
+	if (bIsRealtime)
+	{
+		bShouldInvalidateViewportWidget = true;
 	}
+}
 
-	if (EditorViewportWidget.IsValid() && ( bIsRealtime || bShowStats ))
+void FEditorViewportClient::SetShowStats(bool bWantStats)
+{
+	bShowStats = bWantStats;
+}
+
+void FEditorViewportClient::InvalidateViewportWidget()
+{
+	if (EditorViewportWidget.IsValid())
 	{
-		// Invalidate the viewport widget to re-register its active timer
+		// Invalidate the viewport widget to register its active timer
 		EditorViewportWidget.Pin()->Invalidate();
 	}
+	bShouldInvalidateViewportWidget = false;
 }
 
 void FEditorViewportClient::RedrawRequested(FViewport* InViewport)
@@ -937,6 +979,12 @@ void FEditorViewportClient::Tick(float DeltaTime)
 		}
 	}
 
+	// Invalidate the viewport widget if pending
+	if (bShouldInvalidateViewportWidget)
+	{
+		InvalidateViewportWidget();
+	}
+
 	// Tick the editor modes
 	ModeTools->Tick(this, DeltaTime);
 }
diff --git a/Engine/Source/Editor/UnrealEd/Public/EditorViewportClient.h b/Engine/Source/Editor/UnrealEd/Public/EditorViewportClient.h
--- a/Engine/Source/Editor/UnrealEd/Public/EditorViewportClient.h
+++ b/Engine/Source/Editor/UnrealEd/Public/EditorViewportClient.h
@@ -256,29 +256,10 @@ public:
 	 *
 	 * @return		The current state of the realtime flag.
 	 */
-	bool ToggleRealtime()
-	{ 
-		SetRealtime(!bIsRealtime);
-		return bIsRealtime;
-	}
+	bool ToggleRealtime();
 
 	/** Sets whether or not the viewport updates in realtime. */
-	void SetRealtime( bool bInRealtime, bool bStoreCurrentValue = false )
-	{ 
-		if( bStoreCurrentValue )
-		{
-			//Cache the Realtime and ShowStats flags
-			bStoredRealtime = bIsRealtime;
-			bStoredShowStats = bShowStats;
-		}
-		
-		bIsRealtime	= bInRealtime;
-
-		if( !bIsRealtime )
-		{
-			SetShowStats(false);
-		}
-	}
+	void SetRealtime(bool bInRealtime, bool bStoreCurrentValue = false);
 
 	/** @return		True if viewport is in realtime mode, false otherwise. */
 	bool IsRealtime() const				
@@ -736,10 +717,7 @@ public:
 	 *
 	 * @param	bWantStats	true if stats should be displayed
 	 */
-	void SetShowStats( bool bWantStats )
-	{
-		bShowStats = bWantStats;
-	}
+	void SetShowStats( bool bWantStats );
 	
 	/**
 	 * Sets how the viewport is displayed (lit, wireframe, etc) for the current viewport type
@@ -1007,6 +985,9 @@ public:
 	void OverrideFarClipPlane(const float InFarPlane);
 
 protected:
+	/** Invalidates the viewport widget (if valid) to register its active timer */
+	void InvalidateViewportWidget();
+
 	/** Subclasses may override the near clipping plane. Set to a negative value to disable the override. */
 	void OverrideNearClipPlane(float InNearPlane);
 
@@ -1448,6 +1429,9 @@ private:
 
 	/** If true, we are in Game View mode*/
 	bool bInGameViewMode;
+
+	/** If true, the viewport widget should be invalidated on the next tick (needed to ensure thread safety) */
+	bool bShouldInvalidateViewportWidget;
 };
 
 
