ANIM: disable creating anim BP from Persona until we support multi asset opening.

diff --git a/Engine/Source/Editor/UnrealEd/Private/AnimationEditorUtils.cpp b/Engine/Source/Editor/UnrealEd/Private/AnimationEditorUtils.cpp
--- a/Engine/Source/Editor/UnrealEd/Private/AnimationEditorUtils.cpp
+++ b/Engine/Source/Editor/UnrealEd/Private/AnimationEditorUtils.cpp
@@ -329,6 +329,8 @@ namespace AnimationEditorUtils
 	{
 		MenuBuilder.BeginSection("CreateAnimAssets", LOCTEXT("CreateAnimAssetsMenuHeading", "Anim Assets"));
 		{
+// commented out for UE-13743
+/*
 			MenuBuilder.AddMenuEntry(
 				LOCTEXT("Skeleton_NewAnimBlueprint", "Anim Blueprint"),
 				LOCTEXT("Skeleton_NewAnimBlueprintTooltip", "Creates an Anim Blueprint using the selected skeleton."),
@@ -338,6 +340,7 @@ namespace AnimationEditorUtils
 					FCanExecuteAction()
 					)
 				);
+*/
 
 			MenuBuilder.AddMenuEntry(
 				LOCTEXT("Skeleton_NewAnimComposite", "Anim Composite"),
