#ANIM: Disabled creating animation sequence from menu

https://jira.ol.epicgames.net/browse/UE-13852

diff --git a/Engine/Source/Editor/UnrealEd/Classes/Factories/AnimSequenceFactory.h b/Engine/Source/Editor/UnrealEd/Classes/Factories/AnimSequenceFactory.h
--- a/Engine/Source/Editor/UnrealEd/Classes/Factories/AnimSequenceFactory.h
+++ b/Engine/Source/Editor/UnrealEd/Classes/Factories/AnimSequenceFactory.h
@@ -18,6 +18,7 @@ class UAnimSequenceFactory : public UFactory
 	// Begin UFactory Interface
 	virtual bool ConfigureProperties() override;
 	virtual UObject* FactoryCreateNew(UClass* Class,UObject* InParent,FName Name,EObjectFlags Flags,UObject* Context,FFeedbackContext* Warn) override;
+	virtual bool ShouldShowInNewMenu() const override { return false; } 
 	// Begin UFactory Interface	
 
 private:
