#HLOD: make it silence for mege actor

diff --git a/Engine/Source/Developer/MeshUtilities/Private/MeshUtilities.cpp b/Engine/Source/Developer/MeshUtilities/Private/MeshUtilities.cpp
--- a/Engine/Source/Developer/MeshUtilities/Private/MeshUtilities.cpp
+++ b/Engine/Source/Developer/MeshUtilities/Private/MeshUtilities.cpp
@@ -134,7 +134,8 @@ private:
 		const FString& PackageName,
 		TArray<UObject*>& 
 		OutAssetsToSync, 
-		FVector& OutMergedActorLocation) const override;
+		FVector& OutMergedActorLocation, 
+		bool bSilent=false) const override;
 	
 	virtual void CreateProxyMesh( 
 		const TArray<AActor*>& Actors, 
@@ -4120,7 +4121,8 @@ void FMeshUtilities::MergeActors(
 	const FMeshMergingSettings& InSettings, 
 	const FString& InPackageName, 
 	TArray<UObject*>& OutAssetsToSync, 
-	FVector& OutMergedActorLocation) const
+	FVector& OutMergedActorLocation, 
+	bool bSilent) const
 {
 	TArray<UStaticMeshComponent*> ComponentsToMerge;
 	ComponentsToMerge.Reserve(SourceActors.Num());
@@ -4381,7 +4383,7 @@ void FMeshUtilities::MergeActors(
 			StaticMesh->Materials.Add(Material);
 		}
 		
-		StaticMesh->Build();
+		StaticMesh->Build(bSilent);
 		StaticMesh->PostEditChange();
 
 		OutAssetsToSync.Add(StaticMesh);
diff --git a/Engine/Source/Developer/MeshUtilities/Public/MeshUtilities.h b/Engine/Source/Developer/MeshUtilities/Public/MeshUtilities.h
--- a/Engine/Source/Developer/MeshUtilities/Public/MeshUtilities.h
+++ b/Engine/Source/Developer/MeshUtilities/Public/MeshUtilities.h
@@ -221,7 +221,8 @@ public:
 		const FMeshMergingSettings& InSettings,
 		const FString& PackageName, 
 		TArray<UObject*>& OutAssetsToSync, 
-		FVector& OutMergedActorLocation) const = 0;
+		FVector& OutMergedActorLocation, 
+		bool bSilent=false) const = 0;
 	
 	/**
 	 *	Merges list of actors into single proxy mesh
diff --git a/Engine/Source/Editor/UnrealEd/Private/HierarchicalLOD.cpp b/Engine/Source/Editor/UnrealEd/Private/HierarchicalLOD.cpp
--- a/Engine/Source/Editor/UnrealEd/Private/HierarchicalLOD.cpp
+++ b/Engine/Source/Editor/UnrealEd/Private/HierarchicalLOD.cpp
@@ -117,7 +117,11 @@ void FHierarchicalLODBuilder::BuildClusters(class ULevel* InLevel)
 		{
 			for (auto& Asset: LodActor->SubObjects)
 			{
-				AssetsToDelete.Add(Asset);
+				// @TOOD: This is not permanent fix
+				if (Asset)
+				{
+					AssetsToDelete.Add(Asset);
+				}
 			}
 			World->DestroyActor(LodActor);
 		}
@@ -265,7 +269,7 @@ void FHierarchicalLODBuilder::MergeClusters(class ULevel* InLevel, const int32 L
 		{
 			if(Cluster.IsValid())
 			{
-				GWarn->StatusUpdate(Index, TotalValidCluster, LOCTEXT("HierarchicalLOD_BuildActor", "Building Actor"));
+				GWarn->StatusForceUpdate(Index, TotalValidCluster, LOCTEXT("HierarchicalLOD_BuildActor", "Building Actor"));
 				UE_LOG(LogLODGenerator, Log, TEXT("%d. %0.2f (%s)"), ++Index, Cluster.GetCost(), *Cluster.ToString());
 				if (bBuildActor)
 				{
@@ -577,7 +581,7 @@ void FLODCluster::BuildActor(class UWorld* InWorld, class ULevel* InLevel, const
 				else
 				{
 					FMeshMergingSettings MergeSetting;
-					MeshUtilities.MergeActors(Actors, MergeSetting, AssetsPath, OutAssets, OutProxyLocation );
+					MeshUtilities.MergeActors(Actors, MergeSetting, AssetsPath, OutAssets, OutProxyLocation, true );
 				}
 
 				UStaticMesh* MainMesh=NULL;
