Editor: Disabled some timing log spam in normal map identification code (seen for every imported texture)

diff --git a/Engine/Source/Editor/UnrealEd/Private/NormalMapIdentification.cpp b/Engine/Source/Editor/UnrealEd/Private/NormalMapIdentification.cpp
--- a/Engine/Source/Editor/UnrealEd/Private/NormalMapIdentification.cpp
+++ b/Engine/Source/Editor/UnrealEd/Private/NormalMapIdentification.cpp
@@ -6,7 +6,7 @@
 #include "SNotificationList.h"
 #include "NotificationManager.h"
 
-#define NORMALMAP_IDENTIFICATION_TIMING	(1)
+#define NORMALMAP_IDENTIFICATION_TIMING	(0)
 
 #define LOCTEXT_NAMESPACE "NormalMapIdentification"
 
@@ -377,7 +377,7 @@ static bool IsTextureANormalMap( UTexture* Texture )
 #if NORMALMAP_IDENTIFICATION_TIMING
 	double EndSeconds = FPlatformTime::Seconds();
 
-	FString Msg = FString::Printf( TEXT("%f seconds to analyze %s\n"), (EndSeconds-StartSeconds), *Texture->GetFullName() ); 
+	FString Msg = FString::Printf( TEXT("NormalMapIdentification took %f seconds to analyze %s"), (EndSeconds-StartSeconds), *Texture->GetFullName() ); 
 
 	GLog->Log(Msg);
 #endif
