Separate automated system tests from project-specific tests

diff --git a/Engine/Plugins/Messaging/UdpMessaging/Source/UdpMessaging/Private/Transport/Tests/UdpMessageSegmenterTest.cpp b/Engine/Plugins/Messaging/UdpMessaging/Source/UdpMessaging/Private/Transport/Tests/UdpMessageSegmenterTest.cpp
--- a/Engine/Plugins/Messaging/UdpMessaging/Source/UdpMessaging/Private/Transport/Tests/UdpMessageSegmenterTest.cpp
+++ b/Engine/Plugins/Messaging/UdpMessaging/Source/UdpMessaging/Private/Transport/Tests/UdpMessageSegmenterTest.cpp
@@ -4,7 +4,7 @@
 #include "AutomationTest.h"
 
 
-IMPLEMENT_SIMPLE_AUTOMATION_TEST(FUdpMessageSegmenterTest, "Core.Messaging.Transports.Udp.UdpMessageSegmenter", EAutomationTestFlags::ATF_Editor)
+IMPLEMENT_SIMPLE_AUTOMATION_TEST(FUdpMessageSegmenterTest, "System.Core.Messaging.Transports.Udp.UdpMessageSegmenter", EAutomationTestFlags::ATF_Editor)
 
 
 bool FUdpMessageSegmenterTest::RunTest( const FString& Parameters )
diff --git a/Engine/Source/Developer/AITestSuite/Private/Tests/BTTest.cpp b/Engine/Source/Developer/AITestSuite/Private/Tests/BTTest.cpp
--- a/Engine/Source/Developer/AITestSuite/Private/Tests/BTTest.cpp
+++ b/Engine/Source/Developer/AITestSuite/Private/Tests/BTTest.cpp
@@ -28,7 +28,7 @@ struct FAITest_BTBasicSelector : public FAITest_SimpleBT
 		ExpectedResult.Add(2);
 	}
 };
-IMPLEMENT_AI_LATENT_TEST(FAITest_BTBasicSelector, "Engine.AI.Behavior Trees.Composite node: selector")
+IMPLEMENT_AI_LATENT_TEST(FAITest_BTBasicSelector, "System.Engine.AI.Behavior Trees.Composite node: selector")
 
 struct FAITest_BTBasicSequence : public FAITest_SimpleBT
 {
@@ -53,7 +53,7 @@ struct FAITest_BTBasicSequence : public FAITest_SimpleBT
 		ExpectedResult.Add(2);
 	}
 };
-IMPLEMENT_AI_LATENT_TEST(FAITest_BTBasicSequence, "Engine.AI.Behavior Trees.Composite node: sequence")
+IMPLEMENT_AI_LATENT_TEST(FAITest_BTBasicSequence, "System.Engine.AI.Behavior Trees.Composite node: sequence")
 
 struct FAITest_BTBasicParallelWait : public FAITest_SimpleBT
 {
@@ -87,7 +87,7 @@ struct FAITest_BTBasicParallelWait : public FAITest_SimpleBT
 		ExpectedResult.Add(4);
 	}
 };
-IMPLEMENT_AI_LATENT_TEST(FAITest_BTBasicParallelWait, "Engine.AI.Behavior Trees.Composite node: parallel (wait)")
+IMPLEMENT_AI_LATENT_TEST(FAITest_BTBasicParallelWait, "System.Engine.AI.Behavior Trees.Composite node: parallel (wait)")
 
 struct FAITest_BTBasicParallelAbort : public FAITest_SimpleBT
 {
@@ -118,7 +118,7 @@ struct FAITest_BTBasicParallelAbort : public FAITest_SimpleBT
 		ExpectedResult.Add(4);
 	}
 };
-IMPLEMENT_AI_LATENT_TEST(FAITest_BTBasicParallelAbort, "Engine.AI.Behavior Trees.Composite node: parallel (abort)")
+IMPLEMENT_AI_LATENT_TEST(FAITest_BTBasicParallelAbort, "System.Engine.AI.Behavior Trees.Composite node: parallel (abort)")
 
 struct FAITest_BTCompositeDecorator : public FAITest_SimpleBT
 {
@@ -153,7 +153,7 @@ struct FAITest_BTCompositeDecorator : public FAITest_SimpleBT
 		ExpectedResult.Add(2);
 	}
 };
-IMPLEMENT_AI_LATENT_TEST(FAITest_BTCompositeDecorator, "Engine.AI.Behavior Trees.Composite decorator")
+IMPLEMENT_AI_LATENT_TEST(FAITest_BTCompositeDecorator, "System.Engine.AI.Behavior Trees.Composite decorator")
 
 struct FAITest_BTAbortSelfFail : public FAITest_SimpleBT
 {
@@ -179,7 +179,7 @@ struct FAITest_BTAbortSelfFail : public FAITest_SimpleBT
 		ExpectedResult.Add(1);
 	}
 };
-IMPLEMENT_AI_LATENT_TEST(FAITest_BTAbortSelfFail, "Engine.AI.Behavior Trees.Abort: self failure")
+IMPLEMENT_AI_LATENT_TEST(FAITest_BTAbortSelfFail, "System.Engine.AI.Behavior Trees.Abort: self failure")
 
 struct FAITest_BTAbortSelfSuccess : public FAITest_SimpleBT
 {
@@ -207,7 +207,7 @@ struct FAITest_BTAbortSelfSuccess : public FAITest_SimpleBT
 		ExpectedResult.Add(3);
 	}
 };
-IMPLEMENT_AI_LATENT_TEST(FAITest_BTAbortSelfSuccess, "Engine.AI.Behavior Trees.Abort: self success")
+IMPLEMENT_AI_LATENT_TEST(FAITest_BTAbortSelfSuccess, "System.Engine.AI.Behavior Trees.Abort: self success")
 
 struct FAITest_BTAbortLowerPri : public FAITest_SimpleBT
 {
@@ -234,7 +234,7 @@ struct FAITest_BTAbortLowerPri : public FAITest_SimpleBT
 		ExpectedResult.Add(0);
 	}
 };
-IMPLEMENT_AI_LATENT_TEST(FAITest_BTAbortLowerPri, "Engine.AI.Behavior Trees.Abort: lower priority")
+IMPLEMENT_AI_LATENT_TEST(FAITest_BTAbortLowerPri, "System.Engine.AI.Behavior Trees.Abort: lower priority")
 
 struct FAITest_BTAbortMerge1 : public FAITest_SimpleBT
 {
@@ -263,7 +263,7 @@ struct FAITest_BTAbortMerge1 : public FAITest_SimpleBT
 		ExpectedResult.Add(2);
 	}
 };
-IMPLEMENT_AI_LATENT_TEST(FAITest_BTAbortMerge1, "Engine.AI.Behavior Trees.Abort: merge ranges 1")
+IMPLEMENT_AI_LATENT_TEST(FAITest_BTAbortMerge1, "System.Engine.AI.Behavior Trees.Abort: merge ranges 1")
 
 struct FAITest_BTAbortMerge2 : public FAITest_SimpleBT
 {
@@ -293,7 +293,7 @@ struct FAITest_BTAbortMerge2 : public FAITest_SimpleBT
 		ExpectedResult.Add(2);
 	}
 };
-IMPLEMENT_AI_LATENT_TEST(FAITest_BTAbortMerge2, "Engine.AI.Behavior Trees.Abort: merge ranges 2")
+IMPLEMENT_AI_LATENT_TEST(FAITest_BTAbortMerge2, "System.Engine.AI.Behavior Trees.Abort: merge ranges 2")
 
 struct FAITest_BTAbortMerge3 : public FAITest_SimpleBT
 {
@@ -331,7 +331,7 @@ struct FAITest_BTAbortMerge3 : public FAITest_SimpleBT
 		ExpectedResult.Add(2);
 	}
 };
-IMPLEMENT_AI_LATENT_TEST(FAITest_BTAbortMerge3, "Engine.AI.Behavior Trees.Abort: merge ranges 3")
+IMPLEMENT_AI_LATENT_TEST(FAITest_BTAbortMerge3, "System.Engine.AI.Behavior Trees.Abort: merge ranges 3")
 
 struct FAITest_BTAbortParallelInternal : public FAITest_SimpleBT
 {
@@ -371,7 +371,7 @@ struct FAITest_BTAbortParallelInternal : public FAITest_SimpleBT
 		ExpectedResult.Add(4);
 	}
 };
-IMPLEMENT_AI_LATENT_TEST(FAITest_BTAbortParallelInternal, "Engine.AI.Behavior Trees.Abort: parallel internal")
+IMPLEMENT_AI_LATENT_TEST(FAITest_BTAbortParallelInternal, "System.Engine.AI.Behavior Trees.Abort: parallel internal")
 
 struct FAITest_BTAbortParallelOut : public FAITest_SimpleBT
 {
@@ -404,7 +404,7 @@ struct FAITest_BTAbortParallelOut : public FAITest_SimpleBT
 		ExpectedResult.Add(0);
 	}
 };
-IMPLEMENT_AI_LATENT_TEST(FAITest_BTAbortParallelOut, "Engine.AI.Behavior Trees.Abort: parallel out")
+IMPLEMENT_AI_LATENT_TEST(FAITest_BTAbortParallelOut, "System.Engine.AI.Behavior Trees.Abort: parallel out")
 
 struct FAITest_BTAbortParallelOutAndBack : public FAITest_SimpleBT
 {
@@ -439,7 +439,7 @@ struct FAITest_BTAbortParallelOutAndBack : public FAITest_SimpleBT
 		ExpectedResult.Add(4);
 	}
 };
-IMPLEMENT_AI_LATENT_TEST(FAITest_BTAbortParallelOutAndBack, "Engine.AI.Behavior Trees.Abort: parallel out & back")
+IMPLEMENT_AI_LATENT_TEST(FAITest_BTAbortParallelOutAndBack, "System.Engine.AI.Behavior Trees.Abort: parallel out & back")
 
 struct FAITest_BTAbortMultipleDelayed : public FAITest_SimpleBT
 {
@@ -465,7 +465,7 @@ struct FAITest_BTAbortMultipleDelayed : public FAITest_SimpleBT
 		ExpectedResult.Add(4);
 	}
 };
-IMPLEMENT_AI_LATENT_TEST(FAITest_BTAbortMultipleDelayed, "Engine.AI.Behavior Trees.Abort: multiple delayed requests")
+IMPLEMENT_AI_LATENT_TEST(FAITest_BTAbortMultipleDelayed, "System.Engine.AI.Behavior Trees.Abort: multiple delayed requests")
 
 struct FAITest_BTAbortToInactiveParallel : public FAITest_SimpleBT
 {
@@ -504,7 +504,7 @@ struct FAITest_BTAbortToInactiveParallel : public FAITest_SimpleBT
 		ExpectedResult.Add(5);
 	}
 };
-IMPLEMENT_AI_LATENT_TEST(FAITest_BTAbortToInactiveParallel, "Engine.AI.Behavior Trees.Abort: observer in inactive parallel")
+IMPLEMENT_AI_LATENT_TEST(FAITest_BTAbortToInactiveParallel, "System.Engine.AI.Behavior Trees.Abort: observer in inactive parallel")
 
 struct FAITest_BTAbortDuringLatentAbort : public FAITest_SimpleBT
 {
@@ -533,7 +533,7 @@ struct FAITest_BTAbortDuringLatentAbort : public FAITest_SimpleBT
 		ExpectedResult.Add(5);
 	}
 };
-IMPLEMENT_AI_LATENT_TEST(FAITest_BTAbortDuringLatentAbort, "Engine.AI.Behavior Trees.Abort: during latent task abort")
+IMPLEMENT_AI_LATENT_TEST(FAITest_BTAbortDuringLatentAbort, "System.Engine.AI.Behavior Trees.Abort: during latent task abort")
 
 struct FAITest_BTSubtreeSimple : public FAITest_SimpleBT
 {
@@ -583,7 +583,7 @@ struct FAITest_BTSubtreeSimple : public FAITest_SimpleBT
 		ExpectedResult.Add(3);
 	}
 };
-IMPLEMENT_AI_LATENT_TEST(FAITest_BTSubtreeSimple, "Engine.AI.Behavior Trees.Subtree: simple")
+IMPLEMENT_AI_LATENT_TEST(FAITest_BTSubtreeSimple, "System.Engine.AI.Behavior Trees.Subtree: simple")
 
 struct FAITest_BTSubtreeAbortOut : public FAITest_SimpleBT
 {
@@ -613,6 +613,6 @@ struct FAITest_BTSubtreeAbortOut : public FAITest_SimpleBT
 		ExpectedResult.Add(0);
 	}
 };
-IMPLEMENT_AI_LATENT_TEST(FAITest_BTSubtreeAbortOut, "Engine.AI.Behavior Trees.Subtree: abort out")
+IMPLEMENT_AI_LATENT_TEST(FAITest_BTSubtreeAbortOut, "System.Engine.AI.Behavior Trees.Subtree: abort out")
 
 #undef LOCTEXT_NAMESPACE
diff --git a/Engine/Source/Developer/AITestSuite/Private/Tests/PawnActionsTest.cpp b/Engine/Source/Developer/AITestSuite/Private/Tests/PawnActionsTest.cpp
--- a/Engine/Source/Developer/AITestSuite/Private/Tests/PawnActionsTest.cpp
+++ b/Engine/Source/Developer/AITestSuite/Private/Tests/PawnActionsTest.cpp
@@ -33,7 +33,7 @@ struct FAITest_PawnActions_Push : public FAITest_SimpleActionsTest
 		return true;
 	}
 };
-IMPLEMENT_AI_LATENT_TEST(FAITest_PawnActions_Push, "Engine.AI.Pawn Actions.Pushing Single Action")
+IMPLEMENT_AI_LATENT_TEST(FAITest_PawnActions_Push, "System.Engine.AI.Pawn Actions.Pushing Single Action")
 
 //----------------------------------------------------------------------//
 // 
@@ -61,7 +61,7 @@ struct FAITest_PawnActions_PushingSameActionWithDelay : public FAITest_SimpleAct
 		return true;
 	}
 };
-IMPLEMENT_AI_LATENT_TEST(FAITest_PawnActions_PushingSameActionWithDelay, "Engine.AI.Pawn Actions.Pusihng action that has already been pushed should fail")
+IMPLEMENT_AI_LATENT_TEST(FAITest_PawnActions_PushingSameActionWithDelay, "System.Engine.AI.Pawn Actions.Pusihng action that has already been pushed should fail")
 
 
 //----------------------------------------------------------------------//
@@ -100,7 +100,7 @@ struct FAITest_PawnActions_Pause : public FAITest_SimpleActionsTest
 		return true; 
 	}
 };
-IMPLEMENT_AI_LATENT_TEST(FAITest_PawnActions_Pause, "Engine.AI.Pawn Actions.Pausing Action by younger Action of same priority")
+IMPLEMENT_AI_LATENT_TEST(FAITest_PawnActions_Pause, "System.Engine.AI.Pawn Actions.Pausing Action by younger Action of same priority")
 
 //----------------------------------------------------------------------//
 // 
@@ -131,7 +131,7 @@ struct FAITest_PawnActions_SamePriorityOrder : public FAITest_SimpleActionsTest
 		return true;
 	}
 };
-IMPLEMENT_AI_LATENT_TEST(FAITest_PawnActions_SamePriorityOrder, "Engine.AI.Pawn Actions.Respecting push order")
+IMPLEMENT_AI_LATENT_TEST(FAITest_PawnActions_SamePriorityOrder, "System.Engine.AI.Pawn Actions.Respecting push order")
 
 //----------------------------------------------------------------------//
 // 
@@ -157,7 +157,7 @@ struct FAITest_PawnActions_DoublePushingAction : public FAITest_SimpleActionsTes
 		return true;
 	}
 };
-IMPLEMENT_AI_LATENT_TEST(FAITest_PawnActions_DoublePushingAction, "Engine.AI.Pawn Actions.Pushing same action twice")
+IMPLEMENT_AI_LATENT_TEST(FAITest_PawnActions_DoublePushingAction, "System.Engine.AI.Pawn Actions.Pushing same action twice")
 
 //----------------------------------------------------------------------//
 // 
@@ -187,7 +187,7 @@ struct FAITest_PawnActions_SimplePriority : public FAITest_SimpleActionsTest
 		return true;
 	}
 };
-IMPLEMENT_AI_LATENT_TEST(FAITest_PawnActions_SimplePriority, "Engine.AI.Pawn Actions.Pushing different priority actions")
+IMPLEMENT_AI_LATENT_TEST(FAITest_PawnActions_SimplePriority, "System.Engine.AI.Pawn Actions.Pushing different priority actions")
 
 //----------------------------------------------------------------------//
 // 
@@ -216,7 +216,7 @@ struct FAITest_PawnActions_HighPriorityKeepRunning : public FAITest_SimpleAction
 		return true;
 	}
 };
-IMPLEMENT_AI_LATENT_TEST(FAITest_PawnActions_HighPriorityKeepRunning, "Engine.AI.Pawn Actions.High priority action still running after pushing lower priority action")
+IMPLEMENT_AI_LATENT_TEST(FAITest_PawnActions_HighPriorityKeepRunning, "System.Engine.AI.Pawn Actions.High priority action still running after pushing lower priority action")
 
 //----------------------------------------------------------------------//
 // 
@@ -243,7 +243,7 @@ struct FAITest_PawnActions_SamePriorityActionsPushing : public FAITest_SimpleAct
 		return true;
 	}
 };
-IMPLEMENT_AI_LATENT_TEST(FAITest_PawnActions_SamePriorityActionsPushing, "Engine.AI.Pawn Actions.Pushing multiple actions of same priority")
+IMPLEMENT_AI_LATENT_TEST(FAITest_PawnActions_SamePriorityActionsPushing, "System.Engine.AI.Pawn Actions.Pushing multiple actions of same priority")
 
 //----------------------------------------------------------------------//
 // 
@@ -269,7 +269,7 @@ struct FAITest_PawnActions_Aborting : public FAITest_SimpleActionsTest
 		return true;
 	}
 };
-IMPLEMENT_AI_LATENT_TEST(FAITest_PawnActions_Aborting, "Engine.AI.Pawn Actions.Basic aborting mechanics")
+IMPLEMENT_AI_LATENT_TEST(FAITest_PawnActions_Aborting, "System.Engine.AI.Pawn Actions.Basic aborting mechanics")
 
 //----------------------------------------------------------------------//
 // 
@@ -292,7 +292,7 @@ struct FAITest_PawnActions_PushAndAbort : public FAITest_SimpleActionsTest
 		return true;
 	}
 };
-IMPLEMENT_AI_LATENT_TEST(FAITest_PawnActions_PushAndAbort, "Engine.AI.Pawn Actions.Push and Abort same frame")
+IMPLEMENT_AI_LATENT_TEST(FAITest_PawnActions_PushAndAbort, "System.Engine.AI.Pawn Actions.Push and Abort same frame")
 
 //----------------------------------------------------------------------//
 // 
@@ -322,7 +322,7 @@ struct FAITest_PawnActions_AbortAfterPushingNewAction : public FAITest_SimpleAct
 		return true;
 	}
 };
-IMPLEMENT_AI_LATENT_TEST(FAITest_PawnActions_AbortAfterPushingNewAction, "Engine.AI.Pawn Actions.Abort action after a newer action has been pushed")
+IMPLEMENT_AI_LATENT_TEST(FAITest_PawnActions_AbortAfterPushingNewAction, "System.Engine.AI.Pawn Actions.Abort action after a newer action has been pushed")
 
 //----------------------------------------------------------------------//
 // 
@@ -363,7 +363,7 @@ struct FAITest_PawnActions_ActionPushingActions : public FAITest_SimpleActionsTe
 		return true;
 	}
 };
-IMPLEMENT_AI_LATENT_TEST(FAITest_PawnActions_ActionPushingActions, "Engine.AI.Pawn Actions.Action pushing Actions")
+IMPLEMENT_AI_LATENT_TEST(FAITest_PawnActions_ActionPushingActions, "System.Engine.AI.Pawn Actions.Action pushing Actions")
 
 // scenarios to test
 // * Push, wait for Start and Abort instantly. Implement a special action that will finish in X tick with result Y
\ No newline at end of file
diff --git a/Engine/Source/Developer/AITestSuite/Private/Tests/ResourceIdTest.cpp b/Engine/Source/Developer/AITestSuite/Private/Tests/ResourceIdTest.cpp
--- a/Engine/Source/Developer/AITestSuite/Private/Tests/ResourceIdTest.cpp
+++ b/Engine/Source/Developer/AITestSuite/Private/Tests/ResourceIdTest.cpp
@@ -15,7 +15,7 @@ struct FAITest_ResourceIDBasic : public FAITestBase
 		Test(TEXT("Resource ID's indexes are broken!"), FAIResources::Movement == MovementID);
 	}
 };
-IMPLEMENT_AI_INSTANT_TEST(FAITest_ResourceIDBasic, "Engine.AI.Resource ID.Basic operations")
+IMPLEMENT_AI_INSTANT_TEST(FAITest_ResourceIDBasic, "System.Engine.AI.Resource ID.Basic operations")
 
 //----------------------------------------------------------------------//
 // 
@@ -63,7 +63,7 @@ struct FAITest_ResourceLock : public FAITestBase
 		Test(TEXT("At this point both locks should be identical"), MockLock == MockLock2);
 	}
 };
-IMPLEMENT_AI_INSTANT_TEST(FAITest_ResourceLock, "Engine.AI.Resource ID.Resource locking")
+IMPLEMENT_AI_INSTANT_TEST(FAITest_ResourceLock, "System.Engine.AI.Resource ID.Resource locking")
 
 //----------------------------------------------------------------------//
 // 
@@ -120,7 +120,7 @@ struct FAITest_ResourceSet : public FAITestBase
 		}
 	}
 };
-IMPLEMENT_AI_INSTANT_TEST(FAITest_ResourceSet, "Engine.AI.Resource ID.Resource locking")
+IMPLEMENT_AI_INSTANT_TEST(FAITest_ResourceSet, "System.Engine.AI.Resource ID.Resource locking")
 
 //----------------------------------------------------------------------//
 // 
@@ -149,7 +149,7 @@ struct FAITest_PawnActions_PausingActionsOfSameResource : public FAITest_SimpleA
 		Test(TEXT("First MoveAction should get paused"), MoveAction->IsPaused() == true);
 	}
 };
-IMPLEMENT_AI_INSTANT_TEST(FAITest_PawnActions_PausingActionsOfSameResource, "Engine.AI.Pawn Actions.Pausing actions of same resource")
+IMPLEMENT_AI_INSTANT_TEST(FAITest_PawnActions_PausingActionsOfSameResource, "System.Engine.AI.Pawn Actions.Pausing actions of same resource")
 
 //----------------------------------------------------------------------//
 // 
@@ -179,4 +179,4 @@ struct FAITest_PawnActions_NotPausingActionsOfDifferentResources : public FAITes
 		//Test(TEXT("First MoveAction should get paused"), MoveAction->IsPaused() == false && PerceptionAction->IsPaused() == false);
 	}
 };
-IMPLEMENT_AI_INSTANT_TEST(FAITest_PawnActions_NotPausingActionsOfDifferentResources, "Engine.AI.Pawn Actions.Not pausing actions of different resources")
\ No newline at end of file
+IMPLEMENT_AI_INSTANT_TEST(FAITest_PawnActions_NotPausingActionsOfDifferentResources, "System.Engine.AI.Pawn Actions.Not pausing actions of different resources")
\ No newline at end of file
diff --git a/Engine/Source/Developer/DirectoryWatcher/Private/DirectoryWatcherTests.cpp b/Engine/Source/Developer/DirectoryWatcher/Private/DirectoryWatcherTests.cpp
--- a/Engine/Source/Developer/DirectoryWatcher/Private/DirectoryWatcherTests.cpp
+++ b/Engine/Source/Developer/DirectoryWatcher/Private/DirectoryWatcherTests.cpp
@@ -98,7 +98,7 @@ FString GetWorkingDir()
 }
 
 static const float TestTickDelay = 0.1f;
-IMPLEMENT_SIMPLE_AUTOMATION_TEST(FDirectoryWatcherSimpleCreateTest, "Plugins.Directory Watcher.Simple Create", EAutomationTestFlags::ATF_Editor)
+IMPLEMENT_SIMPLE_AUTOMATION_TEST(FDirectoryWatcherSimpleCreateTest, "System.Plugins.Directory Watcher.Simple Create", EAutomationTestFlags::ATF_Editor)
 bool FDirectoryWatcherSimpleCreateTest::RunTest(const FString& Parameters)
 {
 	const FString WorkingDir = GetWorkingDir();
@@ -130,7 +130,7 @@ bool FDirectoryWatcherSimpleCreateTest::RunTest(const FString& Parameters)
 	return true;
 }
 
-IMPLEMENT_SIMPLE_AUTOMATION_TEST(FDirectoryWatcherSimpleModifyTest, "Plugins.Directory Watcher.Simple Modify", EAutomationTestFlags::ATF_Editor)
+IMPLEMENT_SIMPLE_AUTOMATION_TEST(FDirectoryWatcherSimpleModifyTest, "System.Plugins.Directory Watcher.Simple Modify", EAutomationTestFlags::ATF_Editor)
 bool FDirectoryWatcherSimpleModifyTest::RunTest(const FString& Parameters)
 {
 	const FString WorkingDir = GetWorkingDir();
@@ -168,7 +168,7 @@ bool FDirectoryWatcherSimpleModifyTest::RunTest(const FString& Parameters)
 	return true;
 }
 
-IMPLEMENT_SIMPLE_AUTOMATION_TEST(FDirectoryWatcherSimpleDeleteTest, "Plugins.Directory Watcher.Simple Delete", EAutomationTestFlags::ATF_Editor)
+IMPLEMENT_SIMPLE_AUTOMATION_TEST(FDirectoryWatcherSimpleDeleteTest, "System.Plugins.Directory Watcher.Simple Delete", EAutomationTestFlags::ATF_Editor)
 bool FDirectoryWatcherSimpleDeleteTest::RunTest(const FString& Parameters)
 {
 	const FString WorkingDir = GetWorkingDir();
@@ -203,7 +203,7 @@ bool FDirectoryWatcherSimpleDeleteTest::RunTest(const FString& Parameters)
 	return true;
 }
 
-IMPLEMENT_SIMPLE_AUTOMATION_TEST(FDirectoryWatcherSubFolderTest, "Plugins.Directory Watcher.Sub Folder", EAutomationTestFlags::ATF_Editor)
+IMPLEMENT_SIMPLE_AUTOMATION_TEST(FDirectoryWatcherSubFolderTest, "System.Plugins.Directory Watcher.Sub Folder", EAutomationTestFlags::ATF_Editor)
 bool FDirectoryWatcherSubFolderTest::RunTest(const FString& Parameters)
 {
 	const FString WorkingDir = GetWorkingDir();
diff --git a/Engine/Source/Developer/FunctionalTesting/Private/Tests/FuncTestPerforming.cpp b/Engine/Source/Developer/FunctionalTesting/Private/Tests/FuncTestPerforming.cpp
--- a/Engine/Source/Developer/FunctionalTesting/Private/Tests/FuncTestPerforming.cpp
+++ b/Engine/Source/Developer/FunctionalTesting/Private/Tests/FuncTestPerforming.cpp
@@ -108,7 +108,7 @@ bool FStartFTestsOnMap::Update()
  */
 
 // create test base class
-IMPLEMENT_COMPLEX_AUTOMATION_TEST_PRIVATE(FFunctionalTestingMapsBase, "Maps.Functional Testing", (EAutomationTestFlags::ATF_Game | EAutomationTestFlags::ATF_Editor))
+IMPLEMENT_COMPLEX_AUTOMATION_TEST_PRIVATE(FFunctionalTestingMapsBase, "Project.Maps.Functional Testing", (EAutomationTestFlags::ATF_Game | EAutomationTestFlags::ATF_Editor))
 // implement specific class with non-standard overrides
 class FFunctionalTestingMaps : public FFunctionalTestingMapsBase
 {
diff --git a/Engine/Source/Developer/SourceControl/Private/Tests/SourceControlTests.cpp b/Engine/Source/Developer/SourceControl/Private/Tests/SourceControlTests.cpp
--- a/Engine/Source/Developer/SourceControl/Private/Tests/SourceControlTests.cpp
+++ b/Engine/Source/Developer/SourceControl/Private/Tests/SourceControlTests.cpp
@@ -85,7 +85,7 @@ bool FSetProviderLatentCommand::Update()
 	return true;
 }
 
-IMPLEMENT_COMPLEX_AUTOMATION_TEST(FSetProviderTest, "Editor.Source Control.Set Provider", (EAutomationTestFlags::ATF_Editor | EAutomationTestFlags::ATF_RequiresUser))
+IMPLEMENT_COMPLEX_AUTOMATION_TEST(FSetProviderTest, "Project.Editor.Source Control.Set Provider", (EAutomationTestFlags::ATF_Editor | EAutomationTestFlags::ATF_RequiresUser))
 
 void FSetProviderTest::GetTests(TArray<FString>& OutBeautifiedNames, TArray<FString>& OutTestCommands) const
 {
@@ -116,7 +116,7 @@ bool FConnectLatentCommand::Update()
 	return AsyncHelper.IsDone();
 }
 
-IMPLEMENT_COMPLEX_AUTOMATION_TEST(FConnectTest, "Editor.Source Control.Connect", (EAutomationTestFlags::ATF_Editor | EAutomationTestFlags::ATF_RequiresUser))
+IMPLEMENT_COMPLEX_AUTOMATION_TEST(FConnectTest, "Project.Editor.Source Control.Connect", (EAutomationTestFlags::ATF_Editor | EAutomationTestFlags::ATF_RequiresUser))
 
 void FConnectTest::GetTests(TArray<FString>& OutBeautifiedNames, TArray<FString>& OutTestCommands) const
 {
@@ -213,7 +213,7 @@ bool FCheckOutLatentCommand::Update()
 	return AsyncHelper.IsDone();
 }
 
-IMPLEMENT_COMPLEX_AUTOMATION_TEST(FCheckOutTest, "Editor.Source Control.Check Out", (EAutomationTestFlags::ATF_Editor | EAutomationTestFlags::ATF_RequiresUser))
+IMPLEMENT_COMPLEX_AUTOMATION_TEST(FCheckOutTest, "Project.Editor.Source Control.Check Out", (EAutomationTestFlags::ATF_Editor | EAutomationTestFlags::ATF_RequiresUser))
 
 void FCheckOutTest::GetTests(TArray<FString>& OutBeautifiedNames, TArray<FString>& OutTestCommands) const
 {
@@ -364,7 +364,7 @@ bool FMarkForAddLatentCommand::Update()
 	return AsyncHelper.IsDone();
 }
 
-IMPLEMENT_COMPLEX_AUTOMATION_TEST(FMarkForAddTest, "Editor.Source Control.Mark For Add", (EAutomationTestFlags::ATF_Editor | EAutomationTestFlags::ATF_RequiresUser))
+IMPLEMENT_COMPLEX_AUTOMATION_TEST(FMarkForAddTest, "Project.Editor.Source Control.Mark For Add", (EAutomationTestFlags::ATF_Editor | EAutomationTestFlags::ATF_RequiresUser))
 
 void FMarkForAddTest::GetTests(TArray<FString>& OutBeautifiedNames, TArray<FString>& OutTestCommands) const
 {
@@ -437,7 +437,7 @@ bool FDeleteLatentCommand::Update()
 	return AsyncHelper.IsDone();
 }
 
-IMPLEMENT_COMPLEX_AUTOMATION_TEST(FDeleteTest, "Editor.Source Control.Delete", (EAutomationTestFlags::ATF_Editor | EAutomationTestFlags::ATF_RequiresUser))
+IMPLEMENT_COMPLEX_AUTOMATION_TEST(FDeleteTest, "Project.Editor.Source Control.Delete", (EAutomationTestFlags::ATF_Editor | EAutomationTestFlags::ATF_RequiresUser))
 
 void FDeleteTest::GetTests(TArray<FString>& OutBeautifiedNames, TArray<FString>& OutTestCommands) const
 {
@@ -538,7 +538,7 @@ bool FEditTextureLatentCommand::Update()
 	return true;
 }
 
-IMPLEMENT_COMPLEX_AUTOMATION_TEST(FCheckInTest, "Editor.Source Control.Check In", (EAutomationTestFlags::ATF_Editor | EAutomationTestFlags::ATF_RequiresUser))
+IMPLEMENT_COMPLEX_AUTOMATION_TEST(FCheckInTest, "Project.Editor.Source Control.Check In", (EAutomationTestFlags::ATF_Editor | EAutomationTestFlags::ATF_RequiresUser))
 
 void FCheckInTest::GetTests(TArray<FString>& OutBeautifiedNames, TArray<FString>& OutTestCommands) const
 {
@@ -608,7 +608,7 @@ bool FSyncLatentCommand::Update()
 	return AsyncHelper.IsDone();
 }
 
-IMPLEMENT_COMPLEX_AUTOMATION_TEST(FSyncTest, "Editor.Source Control.Sync", (EAutomationTestFlags::ATF_Editor | EAutomationTestFlags::ATF_RequiresUser))
+IMPLEMENT_COMPLEX_AUTOMATION_TEST(FSyncTest, "Project.Editor.Source Control.Sync", (EAutomationTestFlags::ATF_Editor | EAutomationTestFlags::ATF_RequiresUser))
 
 void FSyncTest::GetTests(TArray<FString>& OutBeautifiedNames, TArray<FString>& OutTestCommands) const
 {
@@ -631,7 +631,7 @@ bool FSyncTest::RunTest(const FString& Parameters)
 	return true;
 }
 
-IMPLEMENT_COMPLEX_AUTOMATION_TEST(FRevertTest, "Editor.Source Control.Revert", (EAutomationTestFlags::ATF_Editor | EAutomationTestFlags::ATF_RequiresUser))
+IMPLEMENT_COMPLEX_AUTOMATION_TEST(FRevertTest, "Project.Editor.Source Control.Revert", (EAutomationTestFlags::ATF_Editor | EAutomationTestFlags::ATF_RequiresUser))
 
 void FRevertTest::GetTests(TArray<FString>& OutBeautifiedNames, TArray<FString>& OutTestCommands) const
 {
@@ -726,7 +726,7 @@ bool FGetStateLatentCommand::Update()
 	return true;
 }
 
-IMPLEMENT_COMPLEX_AUTOMATION_TEST(FUpdateStatusTest, "Editor.Source Control.Update Status", (EAutomationTestFlags::ATF_Editor | EAutomationTestFlags::ATF_RequiresUser))
+IMPLEMENT_COMPLEX_AUTOMATION_TEST(FUpdateStatusTest, "Project.Editor.Source Control.Update Status", (EAutomationTestFlags::ATF_Editor | EAutomationTestFlags::ATF_RequiresUser))
 
 void FUpdateStatusTest::GetTests(TArray<FString>& OutBeautifiedNames, TArray<FString>& OutTestCommands) const
 {
@@ -810,7 +810,7 @@ bool FGetLabelLatentCommand::Update()
 	return true;
 }
 
-IMPLEMENT_COMPLEX_AUTOMATION_TEST(FGetLabelTest, "Editor.Source Control.Get Label", (EAutomationTestFlags::ATF_Editor | EAutomationTestFlags::ATF_RequiresUser))
+IMPLEMENT_COMPLEX_AUTOMATION_TEST(FGetLabelTest, "Project.Editor.Source Control.Get Label", (EAutomationTestFlags::ATF_Editor | EAutomationTestFlags::ATF_RequiresUser))
 
 void FGetLabelTest::GetTests(TArray<FString>& OutBeautifiedNames, TArray<FString>& OutTestCommands) const
 {
@@ -855,7 +855,7 @@ bool FSyncLabelLatentCommand::Update()
 	return true;
 }
 
-IMPLEMENT_COMPLEX_AUTOMATION_TEST(FSyncLabelTest, "Editor.Source Control.Sync Label", (EAutomationTestFlags::ATF_Editor | EAutomationTestFlags::ATF_RequiresUser))
+IMPLEMENT_COMPLEX_AUTOMATION_TEST(FSyncLabelTest, "Project.Editor.Source Control.Sync Label", (EAutomationTestFlags::ATF_Editor | EAutomationTestFlags::ATF_RequiresUser))
 
 void FSyncLabelTest::GetTests(TArray<FString>& OutBeautifiedNames, TArray<FString>& OutTestCommands) const
 {
@@ -921,7 +921,7 @@ bool FGetRevisionLatentCommand::Update()
 	return true;
 }
 
-IMPLEMENT_COMPLEX_AUTOMATION_TEST(FGetRevisionTest, "Editor.Source Control.Get Revision", (EAutomationTestFlags::ATF_Editor | EAutomationTestFlags::ATF_RequiresUser))
+IMPLEMENT_COMPLEX_AUTOMATION_TEST(FGetRevisionTest, "Project.Editor.Source Control.Get Revision", (EAutomationTestFlags::ATF_Editor | EAutomationTestFlags::ATF_RequiresUser))
 
 void FGetRevisionTest::GetTests(TArray<FString>& OutBeautifiedNames, TArray<FString>& OutTestCommands) const
 {
diff --git a/Engine/Source/Editor/GameProjectGeneration/Private/Tests/GameProjectAutomationTests.cpp b/Engine/Source/Editor/GameProjectGeneration/Private/Tests/GameProjectAutomationTests.cpp
--- a/Engine/Source/Editor/GameProjectGeneration/Private/Tests/GameProjectAutomationTests.cpp
+++ b/Engine/Source/Editor/GameProjectGeneration/Private/Tests/GameProjectAutomationTests.cpp
@@ -55,7 +55,7 @@ namespace GameProjectAutomationUtils
 /**
 * Automation test to clean up old test project files
 */
-IMPLEMENT_SIMPLE_AUTOMATION_TEST(FBuildPromotionNewProjectCleanupTest, "Promotion.Project Promotion Pass.Step 1 Blank Project Creation.Cleanup Potential Project Location", EAutomationTestFlags::ATF_Editor);
+IMPLEMENT_SIMPLE_AUTOMATION_TEST(FBuildPromotionNewProjectCleanupTest, "System.Promotion.Project Promotion Pass.Step 1 Blank Project Creation.Cleanup Potential Project Location", EAutomationTestFlags::ATF_Editor);
 bool FBuildPromotionNewProjectCleanupTest::RunTest(const FString& Parameters)
 {
 	FString DesiredProjectFilename = GameProjectAutomationUtils::GetDesiredProjectFilename();
@@ -90,7 +90,7 @@ bool FBuildPromotionNewProjectCleanupTest::RunTest(const FString& Parameters)
 /**
 * Automation test to create a new project
 */
-IMPLEMENT_SIMPLE_AUTOMATION_TEST(FBuildPromotionNewProjectCreateTest, "Promotion.Project Promotion Pass.Step 1 Blank Project Creation.Create Project", EAutomationTestFlags::ATF_Editor);
+IMPLEMENT_SIMPLE_AUTOMATION_TEST(FBuildPromotionNewProjectCreateTest, "System.Promotion.Project Promotion Pass.Step 1 Blank Project Creation.Create Project", EAutomationTestFlags::ATF_Editor);
 bool FBuildPromotionNewProjectCreateTest::RunTest(const FString& Parameters)
 {
 	FString DesiredProjectFilename = GameProjectAutomationUtils::GetDesiredProjectFilename();
@@ -135,7 +135,7 @@ bool FBuildPromotionNewProjectCreateTest::RunTest(const FString& Parameters)
 /**
 * Automation test to create a simple level and save it
 */
-IMPLEMENT_SIMPLE_AUTOMATION_TEST(FBuildPromotionNewProjectMapTest, "Promotion.Project Promotion Pass.Step 2 Basic Level Creation.Create Basic Level", EAutomationTestFlags::ATF_Editor);
+IMPLEMENT_SIMPLE_AUTOMATION_TEST(FBuildPromotionNewProjectMapTest, "System.Promotion.Project Promotion Pass.Step 2 Basic Level Creation.Create Basic Level", EAutomationTestFlags::ATF_Editor);
 bool FBuildPromotionNewProjectMapTest::RunTest(const FString& Parameters)
 {
 	//New level
diff --git a/Engine/Source/Editor/LandscapeEditor/Private/Tests/LandscapeAutomationTests.cpp b/Engine/Source/Editor/LandscapeEditor/Private/Tests/LandscapeAutomationTests.cpp
--- a/Engine/Source/Editor/LandscapeEditor/Private/Tests/LandscapeAutomationTests.cpp
+++ b/Engine/Source/Editor/LandscapeEditor/Private/Tests/LandscapeAutomationTests.cpp
@@ -117,7 +117,7 @@ bool FEndModifyLandscapeCommand::Update()
 /**
 * Landscape creation / edit test
 */
-IMPLEMENT_SIMPLE_AUTOMATION_TEST(FLandscapeEditorTest, "Promotion.Editor Promotion Pass.Step 1 Main Editor Test.New Landscape", EAutomationTestFlags::ATF_Editor | EAutomationTestFlags::ATF_NonNullRHI );
+IMPLEMENT_SIMPLE_AUTOMATION_TEST(FLandscapeEditorTest, "System.Promotion.Editor Promotion Pass.Step 1 Main Editor Test.New Landscape", EAutomationTestFlags::ATF_Editor | EAutomationTestFlags::ATF_NonNullRHI );
 bool FLandscapeEditorTest::RunTest(const FString& Parameters)
 {
 	//New level
diff --git a/Engine/Source/Editor/SceneOutliner/Private/SceneOutlinerSortingTests.cpp b/Engine/Source/Editor/SceneOutliner/Private/SceneOutlinerSortingTests.cpp
--- a/Engine/Source/Editor/SceneOutliner/Private/SceneOutlinerSortingTests.cpp
+++ b/Engine/Source/Editor/SceneOutliner/Private/SceneOutlinerSortingTests.cpp
@@ -3,7 +3,7 @@
 #include "SceneOutlinerPrivatePCH.h"
 #include "SortHelper.h"
 
-IMPLEMENT_SIMPLE_AUTOMATION_TEST(FSortBasicTest, "Editor.Scene Outliner.SortBasic", 0)
+IMPLEMENT_SIMPLE_AUTOMATION_TEST(FSortBasicTest, "System.Editor.Scene Outliner.SortBasic", 0)
 bool FSortBasicTest::RunTest(const FString& Parameters)
 {
 	TArray<SceneOutliner::FNumericStringWrapper> Muddled;
diff --git a/Engine/Source/Editor/StaticMeshEditor/Private/StaticMeshAutomationTests.cpp b/Engine/Source/Editor/StaticMeshEditor/Private/StaticMeshAutomationTests.cpp
--- a/Engine/Source/Editor/StaticMeshEditor/Private/StaticMeshAutomationTests.cpp
+++ b/Engine/Source/Editor/StaticMeshEditor/Private/StaticMeshAutomationTests.cpp
@@ -156,7 +156,7 @@ namespace EditorViewButtonHelper
 /**
  * Static mesh editor test
  */
-IMPLEMENT_SIMPLE_AUTOMATION_TEST( FStaticMeshEditorTest, "Editor.Content.Static Mesh Editor Test", EAutomationTestFlags::ATF_Editor )
+IMPLEMENT_SIMPLE_AUTOMATION_TEST( FStaticMeshEditorTest, "System.Editor.Content.Static Mesh Editor Test", EAutomationTestFlags::ATF_Editor )
 
 /**
  * Take screenshots of the SME window with each of the toolbar buttons toggled separately
diff --git a/Engine/Source/Editor/UnrealEd/Private/AutoReimport/AutoReimportTests.cpp b/Engine/Source/Editor/UnrealEd/Private/AutoReimport/AutoReimportTests.cpp
--- a/Engine/Source/Editor/UnrealEd/Private/AutoReimport/AutoReimportTests.cpp
+++ b/Engine/Source/Editor/UnrealEd/Private/AutoReimport/AutoReimportTests.cpp
@@ -130,7 +130,7 @@ bool CopyTestFiles(const FAutoReimportTestPayload& Test, const TArray<FSrcDstFil
 }
 
 /** Test that creating a new file gets reported correctly */
-IMPLEMENT_SIMPLE_AUTOMATION_TEST(FAutoReimportSimpleCreateTest, "Editor.Auto Reimport.Simple Create", EAutomationTestFlags::ATF_Editor)
+IMPLEMENT_SIMPLE_AUTOMATION_TEST(FAutoReimportSimpleCreateTest, "System.Editor.Auto Reimport.Simple Create", EAutomationTestFlags::ATF_Editor)
 bool FAutoReimportSimpleCreateTest::RunTest(const FString& Parameters)
 {
 	const FString WorkingDir = GetWorkingDir();
@@ -186,7 +186,7 @@ bool FAutoReimportSimpleCreateTest::RunTest(const FString& Parameters)
 }
 
 /** Test that modifying an existing file gets reported correctly */
-IMPLEMENT_SIMPLE_AUTOMATION_TEST(FAutoReimportSimpleModifyTest, "Editor.Auto Reimport.Simple Modify", EAutomationTestFlags::ATF_Editor)
+IMPLEMENT_SIMPLE_AUTOMATION_TEST(FAutoReimportSimpleModifyTest, "System.Editor.Auto Reimport.Simple Modify", EAutomationTestFlags::ATF_Editor)
 bool FAutoReimportSimpleModifyTest::RunTest(const FString& Parameters)
 {
 	const FString WorkingDir = GetWorkingDir();
@@ -258,7 +258,7 @@ bool FAutoReimportSimpleModifyTest::RunTest(const FString& Parameters)
 }
 
 /** Test that deleting an existing file gets reported correctly */
-IMPLEMENT_SIMPLE_AUTOMATION_TEST(FAutoReimportSimpleDeleteTest, "Editor.Auto Reimport.Simple Delete", EAutomationTestFlags::ATF_Editor)
+IMPLEMENT_SIMPLE_AUTOMATION_TEST(FAutoReimportSimpleDeleteTest, "System.Editor.Auto Reimport.Simple Delete", EAutomationTestFlags::ATF_Editor)
 bool FAutoReimportSimpleDeleteTest::RunTest(const FString& Parameters)
 {
 	const FString WorkingDir = GetWorkingDir();
@@ -323,7 +323,7 @@ bool FAutoReimportSimpleDeleteTest::RunTest(const FString& Parameters)
 }
 
 /** Test that renaming an existing file gets reported correctly */
-IMPLEMENT_SIMPLE_AUTOMATION_TEST(FAutoReimportSimpleRenameTest, "Editor.Auto Reimport.Simple Rename", EAutomationTestFlags::ATF_Editor)
+IMPLEMENT_SIMPLE_AUTOMATION_TEST(FAutoReimportSimpleRenameTest, "System.Editor.Auto Reimport.Simple Rename", EAutomationTestFlags::ATF_Editor)
 bool FAutoReimportSimpleRenameTest::RunTest(const FString& Parameters)
 {
 	const FString WorkingDir = GetWorkingDir();
@@ -388,7 +388,7 @@ bool FAutoReimportSimpleRenameTest::RunTest(const FString& Parameters)
 }
 
 /** Test that moving a file outside of the monitored directory gets reported correctly (should be reported as a delete) */
-IMPLEMENT_SIMPLE_AUTOMATION_TEST(FAutoReimportSimpleMoveExternallyTest, "Editor.Auto Reimport.Move Externally", EAutomationTestFlags::ATF_Editor)
+IMPLEMENT_SIMPLE_AUTOMATION_TEST(FAutoReimportSimpleMoveExternallyTest, "System.Editor.Auto Reimport.Move Externally", EAutomationTestFlags::ATF_Editor)
 bool FAutoReimportSimpleMoveExternallyTest::RunTest(const FString& Parameters)
 {
 	const FString WorkingDir = GetWorkingDir();
@@ -454,7 +454,7 @@ bool FAutoReimportSimpleMoveExternallyTest::RunTest(const FString& Parameters)
 }
 
 /** Test that bDetectChangesSinceLastRun works correctly when true and when false */
-IMPLEMENT_SIMPLE_AUTOMATION_TEST(FAutoReimportRestartDetectionTest, "Editor.Auto Reimport.Restart Detection", EAutomationTestFlags::ATF_Editor)
+IMPLEMENT_SIMPLE_AUTOMATION_TEST(FAutoReimportRestartDetectionTest, "System.Editor.Auto Reimport.Restart Detection", EAutomationTestFlags::ATF_Editor)
 bool FAutoReimportRestartDetectionTest::RunTest(const FString& Parameters)
 {
 	const FString WorkingDir = GetWorkingDir();
@@ -553,7 +553,7 @@ bool FAutoReimportRestartDetectionTest::RunTest(const FString& Parameters)
 }
 
 /** Test that making multiple changes to the same file gets picked up correctly */
-IMPLEMENT_SIMPLE_AUTOMATION_TEST(FAutoReimportMultipleChangesTest, "Editor.Auto Reimport.Multiple Changes", EAutomationTestFlags::ATF_Editor)
+IMPLEMENT_SIMPLE_AUTOMATION_TEST(FAutoReimportMultipleChangesTest, "System.Editor.Auto Reimport.Multiple Changes", EAutomationTestFlags::ATF_Editor)
 bool FAutoReimportMultipleChangesTest::RunTest(const FString& Parameters)
 {
 	const FString WorkingDir = GetWorkingDir();
@@ -609,7 +609,7 @@ bool FAutoReimportMultipleChangesTest::RunTest(const FString& Parameters)
 }
 
 /** Test that starting up a cache file with a different set of applicable extensions correctly ignores/updates the extensions, whilst reporting changes only for applicable extensions */
-IMPLEMENT_SIMPLE_AUTOMATION_TEST(FAutoReimportChangeExtensionsTest, "Editor.Auto Reimport.Change Extensions", EAutomationTestFlags::ATF_Editor)
+IMPLEMENT_SIMPLE_AUTOMATION_TEST(FAutoReimportChangeExtensionsTest, "System.Editor.Auto Reimport.Change Extensions", EAutomationTestFlags::ATF_Editor)
 bool FAutoReimportChangeExtensionsTest::RunTest(const FString& Parameters)
 {
 	const FString WorkingDir = GetWorkingDir();
@@ -701,7 +701,7 @@ bool FAutoReimportChangeExtensionsTest::RunTest(const FString& Parameters)
 }
 
 /** Test that starting up a cache file with a different set of applicable extensions correctly ignores/updates the extensions, whilst reporting changes only for applicable extensions */
-IMPLEMENT_SIMPLE_AUTOMATION_TEST(FAutoReimportWildcardFiltersTest, "Editor.Auto Reimport.Wildcard Filters", EAutomationTestFlags::ATF_Editor)
+IMPLEMENT_SIMPLE_AUTOMATION_TEST(FAutoReimportWildcardFiltersTest, "System.Editor.Auto Reimport.Wildcard Filters", EAutomationTestFlags::ATF_Editor)
 bool FAutoReimportWildcardFiltersTest::RunTest(const FString& Parameters)
 {
 	const FString WorkingDir = GetWorkingDir();
diff --git a/Engine/Source/Editor/UnrealEd/Private/AutoReimport/AutoReimportUtilities.cpp b/Engine/Source/Editor/UnrealEd/Private/AutoReimport/AutoReimportUtilities.cpp
--- a/Engine/Source/Editor/UnrealEd/Private/AutoReimport/AutoReimportUtilities.cpp
+++ b/Engine/Source/Editor/UnrealEd/Private/AutoReimport/AutoReimportUtilities.cpp
@@ -36,7 +36,7 @@ bool MatchExtensionString(const TCHAR* Filename, const TCHAR* Extensions)
 }
 
 /** Test MatchExtensionString */
-IMPLEMENT_SIMPLE_AUTOMATION_TEST(FMatchExtensionStringTest, "Editor.Auto Reimport.Extension Matching", EAutomationTestFlags::ATF_Editor)
+IMPLEMENT_SIMPLE_AUTOMATION_TEST(FMatchExtensionStringTest, "System.Editor.Auto Reimport.Extension Matching", EAutomationTestFlags::ATF_Editor)
 bool FMatchExtensionStringTest::RunTest(const FString& Parameters)
 {
 	auto Test = [](const TCHAR* Needle, const TCHAR* Haystack, bool bExpected) -> bool{
diff --git a/Engine/Source/Editor/UnrealEd/Private/Tests/BlueprintAutomationTests.cpp b/Engine/Source/Editor/UnrealEd/Private/Tests/BlueprintAutomationTests.cpp
--- a/Engine/Source/Editor/UnrealEd/Private/Tests/BlueprintAutomationTests.cpp
+++ b/Engine/Source/Editor/UnrealEd/Private/Tests/BlueprintAutomationTests.cpp
@@ -18,12 +18,12 @@
 
 DEFINE_LOG_CATEGORY_STATIC(LogBlueprintAutomationTests, Log, All);
 
-IMPLEMENT_COMPLEX_AUTOMATION_TEST( FBlueprintCompileOnLoadTest, "Blueprints.Compile-On-Load", EAutomationTestFlags::ATF_Editor )
-IMPLEMENT_COMPLEX_AUTOMATION_TEST( FBlueprintInstancesTest, "Blueprints.Instance Test", EAutomationTestFlags::ATF_Editor )
-IMPLEMENT_COMPLEX_AUTOMATION_TEST( FBlueprintReparentTest, "Blueprints.Reparent", EAutomationTestFlags::ATF_Editor )
-IMPLEMENT_COMPLEX_AUTOMATION_TEST( FBlueprintRenameAndCloneTest, "Blueprints.Rename And Clone", (EAutomationTestFlags::ATF_Editor | EAutomationTestFlags::ATF_RequiresUser) )
-IMPLEMENT_COMPLEX_AUTOMATION_TEST( FCompileBlueprintsTest, "Blueprints.Compile Blueprints", EAutomationTestFlags::ATF_Editor )
-IMPLEMENT_COMPLEX_AUTOMATION_TEST( FCompileAnimBlueprintsTest, "Blueprints.Compile Anims", EAutomationTestFlags::ATF_Editor )
+IMPLEMENT_COMPLEX_AUTOMATION_TEST( FBlueprintCompileOnLoadTest, "Project.Blueprints.Compile-On-Load", EAutomationTestFlags::ATF_Editor )
+IMPLEMENT_COMPLEX_AUTOMATION_TEST( FBlueprintInstancesTest, "Project.Blueprints.Instance Test", EAutomationTestFlags::ATF_Editor )
+IMPLEMENT_COMPLEX_AUTOMATION_TEST( FBlueprintReparentTest, "System.Blueprints.Reparent", EAutomationTestFlags::ATF_Editor )
+IMPLEMENT_COMPLEX_AUTOMATION_TEST( FBlueprintRenameAndCloneTest, "Project.Blueprints.Rename And Clone", (EAutomationTestFlags::ATF_Editor | EAutomationTestFlags::ATF_RequiresUser) )
+IMPLEMENT_COMPLEX_AUTOMATION_TEST( FCompileBlueprintsTest, "Project.Blueprints.Compile Blueprints", EAutomationTestFlags::ATF_Editor )
+IMPLEMENT_COMPLEX_AUTOMATION_TEST( FCompileAnimBlueprintsTest, "Project.Blueprints.Compile Anims", EAutomationTestFlags::ATF_Editor )
 
 class FBlueprintAutomationTestUtilities
 {
diff --git a/Engine/Source/Editor/UnrealEd/Private/Tests/EditorAssetAutomationTests.cpp b/Engine/Source/Editor/UnrealEd/Private/Tests/EditorAssetAutomationTests.cpp
--- a/Engine/Source/Editor/UnrealEd/Private/Tests/EditorAssetAutomationTests.cpp
+++ b/Engine/Source/Editor/UnrealEd/Private/Tests/EditorAssetAutomationTests.cpp
@@ -457,7 +457,7 @@ bool FLogAssetCreationStatsCommand::Update()
 /**
 * Automation test for creating, saving, and duplicating assets
 */
-IMPLEMENT_SIMPLE_AUTOMATION_TEST(FAssetEditorTest, "Editor.Content.Asset Creation and Duplication", EAutomationTestFlags::ATF_Editor);
+IMPLEMENT_SIMPLE_AUTOMATION_TEST(FAssetEditorTest, "System.Editor.Content.Asset Creation and Duplication", EAutomationTestFlags::ATF_Editor);
 
 bool FAssetEditorTest::RunTest(const FString& Parameters)
 {
@@ -1063,7 +1063,7 @@ bool FLogImportExportTestResultsCommand::Update()
 /**
 * Automation test to import, open, screenshot, and export assets
 */
-IMPLEMENT_SIMPLE_AUTOMATION_TEST(FAssetImportEditorTest, "Editor.Content.Asset Import and Export", EAutomationTestFlags::ATF_Editor);
+IMPLEMENT_SIMPLE_AUTOMATION_TEST(FAssetImportEditorTest, "Project.Editor.Content.Asset Import and Export", EAutomationTestFlags::ATF_Editor);
 bool FAssetImportEditorTest::RunTest(const FString& Parameters)
 {
 	UAutomationTestSettings const* AutomationTestSettings = GetDefault<UAutomationTestSettings>();
diff --git a/Engine/Source/Editor/UnrealEd/Private/Tests/EditorAutomationTests.cpp b/Engine/Source/Editor/UnrealEd/Private/Tests/EditorAutomationTests.cpp
--- a/Engine/Source/Editor/UnrealEd/Private/Tests/EditorAutomationTests.cpp
+++ b/Engine/Source/Editor/UnrealEd/Private/Tests/EditorAutomationTests.cpp
@@ -115,7 +115,7 @@ void TakeLatentAutomationScreenshot(struct WindowScreenshotParameters Screenshot
  * in the Content Browser (does not allow for specific settings to be made per import factory). Cannot be run in a commandlet
  * as it executes code that routes through Slate UI.
  */
-IMPLEMENT_COMPLEX_AUTOMATION_TEST( FGenericImportAssetsAutomationTest, "Editor.Import", (EAutomationTestFlags::ATF_Editor | EAutomationTestFlags::ATF_NonNullRHI) )
+IMPLEMENT_COMPLEX_AUTOMATION_TEST( FGenericImportAssetsAutomationTest, "System.Editor.Import", (EAutomationTestFlags::ATF_Editor | EAutomationTestFlags::ATF_NonNullRHI) )
 
 /** 
  * Requests a enumeration of all sample assets to import
@@ -194,7 +194,7 @@ bool FGenericImportAssetsAutomationTest::RunTest(const FString& Parameters)
  * FJsonAutomationTest
  * Simple unit test that runs Json's in-built test cases
  */
-IMPLEMENT_SIMPLE_AUTOMATION_TEST( FJsonAutomationTest, "Engine.FileSystem.JSON", EAutomationTestFlags::ATF_SmokeTest)
+IMPLEMENT_SIMPLE_AUTOMATION_TEST( FJsonAutomationTest, "System.Engine.FileSystem.JSON", EAutomationTestFlags::ATF_SmokeTest)
 
 typedef TJsonWriterFactory< TCHAR, TCondensedJsonPrintPolicy<TCHAR> > FCondensedJsonStringWriterFactory;
 typedef TJsonWriter< TCHAR, TCondensedJsonPrintPolicy<TCHAR> > FCondensedJsonStringWriter;
@@ -866,7 +866,7 @@ bool FJsonAutomationTest::RunTest(const FString& Parameters)
  * Pie Test
  * Verification PIE works
  */
-IMPLEMENT_SIMPLE_AUTOMATION_TEST( FPIETest, "Maps.PIE", EAutomationTestFlags::ATF_Editor )
+IMPLEMENT_SIMPLE_AUTOMATION_TEST( FPIETest, "System.Maps.PIE", EAutomationTestFlags::ATF_Editor )
 
 /** 
  * Execute the loading of one map to verify PIE works
@@ -897,7 +897,7 @@ bool FPIETest::RunTest(const FString& Parameters)
  * LoadAllMaps
  * Verification automation test to make sure loading all maps succeed without crashing
  */
-IMPLEMENT_COMPLEX_AUTOMATION_TEST( FLoadAllMapsInEditorTest, "Maps.Load All In Editor", EAutomationTestFlags::ATF_Editor )
+IMPLEMENT_COMPLEX_AUTOMATION_TEST( FLoadAllMapsInEditorTest, "Project.Maps.Load All In Editor", EAutomationTestFlags::ATF_Editor )
 
 /** 
  * Requests a enumeration of all maps to be loaded
@@ -1004,7 +1004,7 @@ bool FLoadAllMapsInEditorTest::RunTest(const FString& Parameters)
 /**
  * Reinitialize all RHI resources
  */
-IMPLEMENT_SIMPLE_AUTOMATION_TEST(FReinitializeRHIResources, "Engine.Rendering.Reinit Resources", EAutomationTestFlags::ATF_Editor);
+IMPLEMENT_SIMPLE_AUTOMATION_TEST(FReinitializeRHIResources, "System.Engine.Rendering.Reinit Resources", EAutomationTestFlags::ATF_Editor);
 
 bool FReinitializeRHIResources::RunTest(const FString& Parameters)
 {
@@ -1017,7 +1017,7 @@ bool FReinitializeRHIResources::RunTest(const FString& Parameters)
 /**
  * QA BSP Regression Testing
  */
-IMPLEMENT_SIMPLE_AUTOMATION_TEST(FBSPValidation, "QA.BSP Validation", EAutomationTestFlags::ATF_Editor);
+IMPLEMENT_SIMPLE_AUTOMATION_TEST(FBSPValidation, "System.QA.BSP Validation", EAutomationTestFlags::ATF_Editor);
 
 bool FBSPValidation::RunTest(const FString& Parameters)
 {
@@ -1166,7 +1166,7 @@ bool FBSPValidation::RunTest(const FString& Parameters)
 /**
  * QA Static Mesh Regression Testing
  */
-IMPLEMENT_SIMPLE_AUTOMATION_TEST(FStaticMeshValidation, "QA.Mesh Factory Validation", EAutomationTestFlags::ATF_Editor);
+IMPLEMENT_SIMPLE_AUTOMATION_TEST(FStaticMeshValidation, "System.QA.Mesh Factory Validation", EAutomationTestFlags::ATF_Editor);
 
 bool FStaticMeshValidation::RunTest(const FString& Parameters)
 {
@@ -1242,7 +1242,7 @@ bool FStaticMeshValidation::RunTest(const FString& Parameters)
 /**
  * QA Convert Meshes Regression Testing
  */
-IMPLEMENT_SIMPLE_AUTOMATION_TEST(FConvertToValidation, "QA.Convert Meshes", EAutomationTestFlags::ATF_Editor);
+IMPLEMENT_SIMPLE_AUTOMATION_TEST(FConvertToValidation, "System.QA.Convert Meshes", EAutomationTestFlags::ATF_Editor);
 
 //gather all brushes that exist now
 void ConvertTestFindAllBrushes(TArray<ABrush*> &PreviousBrushes)
@@ -1466,7 +1466,7 @@ bool FCleanupConvertToValidation::Update()
 /**
  * QA Static Mesh Regression Testing
  */
-IMPLEMENT_SIMPLE_AUTOMATION_TEST(FStaticMeshPlacement, "QA.Static Mesh Placement", EAutomationTestFlags::ATF_Editor);
+IMPLEMENT_SIMPLE_AUTOMATION_TEST(FStaticMeshPlacement, "System.QA.Static Mesh Placement", EAutomationTestFlags::ATF_Editor);
 
 bool FStaticMeshPlacement::RunTest(const FString& Parameters)
 {
@@ -1520,7 +1520,7 @@ bool FStaticMeshPlacement::RunTest(const FString& Parameters)
 /**
  * QA Light Placement Regression Testing
  */
-IMPLEMENT_SIMPLE_AUTOMATION_TEST(FLightPlacement, "QA.Point Light Placement", EAutomationTestFlags::ATF_Editor);
+IMPLEMENT_SIMPLE_AUTOMATION_TEST(FLightPlacement, "System.QA.Point Light Placement", EAutomationTestFlags::ATF_Editor);
 
 bool FLightPlacement::RunTest(const FString& Parameters)
 {
@@ -1619,7 +1619,7 @@ bool FLightPlacement::RunTest(const FString& Parameters)
  * Unit test to find all timelines in blueprints and list the events that can trigger them.
  * Timelines implicitly tick and are usually used for cosmetic events, so they can cause performance problems on dedicated servers.
  */
-IMPLEMENT_COMPLEX_AUTOMATION_TEST( FTraceAllTimelinesAutomationTest, "Performance Audits.Find Timelines On Server", (EAutomationTestFlags::ATF_Editor | EAutomationTestFlags::ATF_RequiresUser))
+IMPLEMENT_COMPLEX_AUTOMATION_TEST( FTraceAllTimelinesAutomationTest, "Project.Performance Audits.Find Timelines On Server", (EAutomationTestFlags::ATF_Editor | EAutomationTestFlags::ATF_RequiresUser))
 
 /** 
  * Requests an enumeration of all blueprints to be loaded
@@ -2014,7 +2014,7 @@ bool FTraceAllTimelinesAutomationTest::RunTest(const FString& BlueprintName)
 /**
 * Tool to look for overlapping UV's in static meshes.
 */
-IMPLEMENT_COMPLEX_AUTOMATION_TEST(FStaticMeshUVCheck, "Tools.Static Mesh.Static Mesh UVs Check", (EAutomationTestFlags::ATF_Editor | EAutomationTestFlags::ATF_RequiresUser));
+IMPLEMENT_COMPLEX_AUTOMATION_TEST(FStaticMeshUVCheck, "Project.Tools.Static Mesh.Static Mesh UVs Check", (EAutomationTestFlags::ATF_Editor | EAutomationTestFlags::ATF_RequiresUser));
 
 void FStaticMeshUVCheck::GetTests(TArray<FString>& OutBeautifiedNames, TArray<FString>& OutTestCommands) const
 {
@@ -2044,7 +2044,7 @@ bool FStaticMeshUVCheck::RunTest(const FString& Parameters)
 /**
 * Launches a map onto a specified device after making a change to it.
 */
-IMPLEMENT_COMPLEX_AUTOMATION_TEST(FLaunchOnTest, "Editor.Launch On Test", (EAutomationTestFlags::ATF_Editor | EAutomationTestFlags::ATF_RequiresUser))
+IMPLEMENT_COMPLEX_AUTOMATION_TEST(FLaunchOnTest, "Project.Editor.Launch On Test", (EAutomationTestFlags::ATF_Editor | EAutomationTestFlags::ATF_RequiresUser))
 
 void FLaunchOnTest::GetTests(TArray<FString>& OutBeautifiedNames, TArray <FString>& OutTestCommands) const
 {
diff --git a/Engine/Source/Editor/UnrealEd/Private/Tests/EditorBuildPromotionTests.cpp b/Engine/Source/Editor/UnrealEd/Private/Tests/EditorBuildPromotionTests.cpp
--- a/Engine/Source/Editor/UnrealEd/Private/Tests/EditorBuildPromotionTests.cpp
+++ b/Engine/Source/Editor/UnrealEd/Private/Tests/EditorBuildPromotionTests.cpp
@@ -4292,7 +4292,7 @@ namespace BuildPromotionTestHelper
 /**
 * Automation test that handles cleanup of the build promotion test
 */
-IMPLEMENT_SIMPLE_AUTOMATION_TEST(FBuildPromotionInitialCleanupTest, "Promotion.Editor Promotion Pass.Step 1 Main Editor Test.Cleanup old files", EAutomationTestFlags::ATF_Editor);
+IMPLEMENT_SIMPLE_AUTOMATION_TEST(FBuildPromotionInitialCleanupTest, "System.Promotion.Editor Promotion Pass.Step 1 Main Editor Test.Cleanup old files", EAutomationTestFlags::ATF_Editor);
 bool FBuildPromotionInitialCleanupTest::RunTest(const FString& Parameters)
 {
 	EditorBuildPromotionTestUtils::PerformCleanup();
@@ -4338,7 +4338,7 @@ bool FLogTestResultCommand::Update()
 /**
 * Automation test that handles setting keybindings and editor preferences
 */
-IMPLEMENT_SIMPLE_AUTOMATION_TEST(FBuildPromotionSettingsTest, "Promotion.Editor Promotion Pass.Step 1 Main Editor Test.Settings", EAutomationTestFlags::ATF_Editor);
+IMPLEMENT_SIMPLE_AUTOMATION_TEST(FBuildPromotionSettingsTest, "System.Promotion.Editor Promotion Pass.Step 1 Main Editor Test.Settings", EAutomationTestFlags::ATF_Editor);
 
 bool FBuildPromotionSettingsTest::RunTest(const FString& Parameters)
 {
@@ -4447,7 +4447,7 @@ bool FRunBuildPromotionTestCommand::Update()
 /**
 * Automation test that handles the build promotion process
 */
-IMPLEMENT_SIMPLE_AUTOMATION_TEST(FBuildPromotionTest, "Promotion.Editor Promotion Pass.Step 1 Main Editor Test.General Editor Test", EAutomationTestFlags::ATF_Editor);
+IMPLEMENT_SIMPLE_AUTOMATION_TEST(FBuildPromotionTest, "System.Promotion.Editor Promotion Pass.Step 1 Main Editor Test.General Editor Test", EAutomationTestFlags::ATF_Editor);
 
 bool FBuildPromotionTest::RunTest(const FString& Parameters)
 {
@@ -4487,7 +4487,7 @@ bool FPIEExecCommand::Update()
 /**
 * Execute the loading of one map to verify PIE works
 */
-IMPLEMENT_SIMPLE_AUTOMATION_TEST(FBuildPromotionPIETest, "Promotion.Editor Promotion Pass.Step 2 Run Map After Re-launch.Run Map", EAutomationTestFlags::ATF_Editor)
+IMPLEMENT_SIMPLE_AUTOMATION_TEST(FBuildPromotionPIETest, "System.Promotion.Editor Promotion Pass.Step 2 Run Map After Re-launch.Run Map", EAutomationTestFlags::ATF_Editor)
 bool FBuildPromotionPIETest::RunTest(const FString& Parameters)
 {
 	UAutomationTestSettings const* AutomationTestSettings = GetDefault<UAutomationTestSettings>();
@@ -4534,7 +4534,7 @@ bool FBuildPromotionPIETest::RunTest(const FString& Parameters)
 /**
 * Automation test that handles cleanup of the build promotion test
 */
-IMPLEMENT_SIMPLE_AUTOMATION_TEST(FBuildPromotionCleanupTest, "Promotion.Editor Promotion Pass.Step 3 Test Cleanup.Cleanup", EAutomationTestFlags::ATF_Editor);
+IMPLEMENT_SIMPLE_AUTOMATION_TEST(FBuildPromotionCleanupTest, "System.Promotion.Editor Promotion Pass.Step 3 Test Cleanup.Cleanup", EAutomationTestFlags::ATF_Editor);
 bool FBuildPromotionCleanupTest::RunTest(const FString& Parameters)
 {
 	EditorBuildPromotionTestUtils::PerformCleanup();
diff --git a/Engine/Source/Editor/UnrealEd/Private/Tests/EditorOpenAssetAutomationTests.cpp b/Engine/Source/Editor/UnrealEd/Private/Tests/EditorOpenAssetAutomationTests.cpp
--- a/Engine/Source/Editor/UnrealEd/Private/Tests/EditorOpenAssetAutomationTests.cpp
+++ b/Engine/Source/Editor/UnrealEd/Private/Tests/EditorOpenAssetAutomationTests.cpp
@@ -61,7 +61,7 @@ bool FOpenActualAssetEditors(const FString& Parameters)
 * Test to open the sub editor windows for a specified list of assets.
 * This list can be setup in the Editor Preferences window within the editor or the DefaultEngine.ini file for that particular project.
 */
-IMPLEMENT_COMPLEX_AUTOMATION_TEST(FOpenAssetEditors, "QA.Open Asset Editors", EAutomationTestFlags::ATF_Editor);
+IMPLEMENT_COMPLEX_AUTOMATION_TEST(FOpenAssetEditors, "System.QA.Open Asset Editors", EAutomationTestFlags::ATF_Editor);
 
 void FOpenAssetEditors::GetTests(TArray<FString>& OutBeautifiedNames, TArray <FString>& OutTestCommands) const
 {
@@ -242,7 +242,7 @@ void FOpenAllAssetTests(OpenAssetParameters& AssetList)
 /**
 * This test opens each AimOffsetBlendSpace into its sub-editor.
 */
-IMPLEMENT_COMPLEX_AUTOMATION_TEST(FOpenAimOffsetBlendSpaceAssetEditors, "Tools.Skeletal Mesh.Animation.BlendSpaces.Open Aim Offset BlendSpace Assets", (EAutomationTestFlags::ATF_Editor | EAutomationTestFlags::ATF_RequiresUser));
+IMPLEMENT_COMPLEX_AUTOMATION_TEST(FOpenAimOffsetBlendSpaceAssetEditors, "Project.Tools.Skeletal Mesh.Animation.BlendSpaces.Open Aim Offset BlendSpace Assets", (EAutomationTestFlags::ATF_Editor | EAutomationTestFlags::ATF_RequiresUser));
 
 void FOpenAimOffsetBlendSpaceAssetEditors::GetTests(TArray<FString>& OutBeautifiedNames, TArray<FString>& OutTestCommands) const
 {
@@ -261,7 +261,7 @@ bool FOpenAimOffsetBlendSpaceAssetEditors::RunTest(const FString& Parameters)
 /**
 * This test opens each AimOffsetBlendSpace1D into its sub-editor.
 */
-IMPLEMENT_COMPLEX_AUTOMATION_TEST(FOpenAimOffsetBlendSpace1DAssetEditors, "Tools.Skeletal Mesh.Animation.BlendSpaces.Open Aim Offset BlendSpace1D Assets", (EAutomationTestFlags::ATF_Editor | EAutomationTestFlags::ATF_RequiresUser));
+IMPLEMENT_COMPLEX_AUTOMATION_TEST(FOpenAimOffsetBlendSpace1DAssetEditors, "Project.Tools.Skeletal Mesh.Animation.BlendSpaces.Open Aim Offset BlendSpace1D Assets", (EAutomationTestFlags::ATF_Editor | EAutomationTestFlags::ATF_RequiresUser));
 
 void FOpenAimOffsetBlendSpace1DAssetEditors::GetTests(TArray<FString>& OutBeautifiedNames, TArray<FString>& OutTestCommands) const
 {
@@ -280,7 +280,7 @@ bool FOpenAimOffsetBlendSpace1DAssetEditors::RunTest(const FString& Parameters)
 /**
 * This test opens each AnimMontage into its sub-editor.
 */
-IMPLEMENT_COMPLEX_AUTOMATION_TEST(FOpenAnimMontageAssetEditors, "Tools.Skeletal Mesh.Animation.Open AnimMontage Assets", (EAutomationTestFlags::ATF_Editor | EAutomationTestFlags::ATF_RequiresUser));
+IMPLEMENT_COMPLEX_AUTOMATION_TEST(FOpenAnimMontageAssetEditors, "Project.Tools.Skeletal Mesh.Animation.Open AnimMontage Assets", (EAutomationTestFlags::ATF_Editor | EAutomationTestFlags::ATF_RequiresUser));
 
 void FOpenAnimMontageAssetEditors::GetTests(TArray<FString>& OutBeautifiedNames, TArray<FString>& OutTestCommands) const
 {
@@ -299,7 +299,7 @@ bool FOpenAnimMontageAssetEditors::RunTest(const FString& Parameters)
 /**
 * This test opens each AnimSequence into its sub-editor.
 */
-IMPLEMENT_COMPLEX_AUTOMATION_TEST(FOpenAnimSequenceAssetEditors, "Tools.Skeletal Mesh.Animation.Open AnimSequence Assets", (EAutomationTestFlags::ATF_Editor | EAutomationTestFlags::ATF_RequiresUser));
+IMPLEMENT_COMPLEX_AUTOMATION_TEST(FOpenAnimSequenceAssetEditors, "Project.Tools.Skeletal Mesh.Animation.Open AnimSequence Assets", (EAutomationTestFlags::ATF_Editor | EAutomationTestFlags::ATF_RequiresUser));
 
 void FOpenAnimSequenceAssetEditors::GetTests(TArray<FString>& OutBeautifiedNames, TArray<FString>& OutTestCommands) const
 {
@@ -318,7 +318,7 @@ bool FOpenAnimSequenceAssetEditors::RunTest(const FString& Parameters)
 /**
 * This test opens each BEHAVIOR TREE into its sub-editor.
 */
-IMPLEMENT_COMPLEX_AUTOMATION_TEST(FOpenBehaviorTreeAssetEditors, "Tools.AI.Open Behavior Tree Assets", (EAutomationTestFlags::ATF_Editor | EAutomationTestFlags::ATF_RequiresUser));
+IMPLEMENT_COMPLEX_AUTOMATION_TEST(FOpenBehaviorTreeAssetEditors, "Project.Tools.AI.Open Behavior Tree Assets", (EAutomationTestFlags::ATF_Editor | EAutomationTestFlags::ATF_RequiresUser));
 
 void FOpenBehaviorTreeAssetEditors::GetTests(TArray<FString>& OutBeautifiedNames, TArray<FString>& OutTestCommands) const
 {
@@ -337,7 +337,7 @@ bool FOpenBehaviorTreeAssetEditors::RunTest(const FString& Parameters)
 /**
 * This test opens each BlendSpace into its sub-editor.
 */
-IMPLEMENT_COMPLEX_AUTOMATION_TEST(FOpenBlendSpaceAssetEditors, "Tools.Skeletal Mesh.Animation.BlendSpaces.Open BlendSpace Assets", (EAutomationTestFlags::ATF_Editor | EAutomationTestFlags::ATF_RequiresUser));
+IMPLEMENT_COMPLEX_AUTOMATION_TEST(FOpenBlendSpaceAssetEditors, "Project.Tools.Skeletal Mesh.Animation.BlendSpaces.Open BlendSpace Assets", (EAutomationTestFlags::ATF_Editor | EAutomationTestFlags::ATF_RequiresUser));
 
 void FOpenBlendSpaceAssetEditors::GetTests(TArray<FString>& OutBeautifiedNames, TArray<FString>& OutTestCommands) const
 {
@@ -356,7 +356,7 @@ bool FOpenBlendSpaceAssetEditors::RunTest(const FString& Parameters)
 /**
 * This test opens each BlendSpace1D into its sub-editor.
 */
-IMPLEMENT_COMPLEX_AUTOMATION_TEST(FOpenBlendSpace1DAssetEditors, "Tools.Skeletal Mesh.Animation.BlendSpaces.Open BlendSpace1D Assets", (EAutomationTestFlags::ATF_Editor | EAutomationTestFlags::ATF_RequiresUser));
+IMPLEMENT_COMPLEX_AUTOMATION_TEST(FOpenBlendSpace1DAssetEditors, "Project.Tools.Skeletal Mesh.Animation.BlendSpaces.Open BlendSpace1D Assets", (EAutomationTestFlags::ATF_Editor | EAutomationTestFlags::ATF_RequiresUser));
 
 void FOpenBlendSpace1DAssetEditors::GetTests(TArray<FString>& OutBeautifiedNames, TArray<FString>& OutTestCommands) const
 {
@@ -375,7 +375,7 @@ bool FOpenBlendSpace1DAssetEditors::RunTest(const FString& Parameters)
 /**
 * This test opens each BLUEPRINT into its sub-editor.
 */
-IMPLEMENT_COMPLEX_AUTOMATION_TEST(FOpenBlueprintAssetEditors, "Tools.Blueprint.Open Blueprint Assets", (EAutomationTestFlags::ATF_Editor | EAutomationTestFlags::ATF_RequiresUser));
+IMPLEMENT_COMPLEX_AUTOMATION_TEST(FOpenBlueprintAssetEditors, "Project.Tools.Blueprint.Open Blueprint Assets", (EAutomationTestFlags::ATF_Editor | EAutomationTestFlags::ATF_RequiresUser));
 
 void FOpenBlueprintAssetEditors::GetTests(TArray<FString>& OutBeautifiedNames, TArray<FString>& OutTestCommands) const
 {
@@ -394,7 +394,7 @@ bool FOpenBlueprintAssetEditors::RunTest(const FString& Parameters)
 /**
 * This test opens each UserDefinedEnum into its sub-editor.
 */
-IMPLEMENT_COMPLEX_AUTOMATION_TEST(FOpenUserDefinedEnumAssetEditors, "Tools.Blueprint.Open User Defined Enum Assets", (EAutomationTestFlags::ATF_Editor | EAutomationTestFlags::ATF_RequiresUser));
+IMPLEMENT_COMPLEX_AUTOMATION_TEST(FOpenUserDefinedEnumAssetEditors, "Project.Tools.Blueprint.Open User Defined Enum Assets", (EAutomationTestFlags::ATF_Editor | EAutomationTestFlags::ATF_RequiresUser));
 
 void FOpenUserDefinedEnumAssetEditors::GetTests(TArray<FString>& OutBeautifiedNames, TArray<FString>& OutTestCommands) const
 {
@@ -413,7 +413,7 @@ bool FOpenUserDefinedEnumAssetEditors::RunTest(const FString& Parameters)
 /**
 * This test opens each DESTRUCTIBLE MESH into its sub-editor.
 */
-IMPLEMENT_COMPLEX_AUTOMATION_TEST(FOpenDestructibleMeshAssetEditors, "Tools.Skeletal Mesh.Open Destructible Mesh Assets", (EAutomationTestFlags::ATF_Editor | EAutomationTestFlags::ATF_RequiresUser));
+IMPLEMENT_COMPLEX_AUTOMATION_TEST(FOpenDestructibleMeshAssetEditors, "Project.Tools.Skeletal Mesh.Open Destructible Mesh Assets", (EAutomationTestFlags::ATF_Editor | EAutomationTestFlags::ATF_RequiresUser));
 
 void FOpenDestructibleMeshAssetEditors::GetTests(TArray<FString>& OutBeautifiedNames, TArray<FString>& OutTestCommands) const
 {
@@ -432,7 +432,7 @@ bool FOpenDestructibleMeshAssetEditors::RunTest(const FString& Parameters)
 /**
 * This test opens each DialogueVoice into its sub-editor.
 */
-IMPLEMENT_COMPLEX_AUTOMATION_TEST(FOpenDialogueVoiceAssetEditors, "Tools.Sounds.Open Dialogue Voice Assets", (EAutomationTestFlags::ATF_Editor | EAutomationTestFlags::ATF_RequiresUser));
+IMPLEMENT_COMPLEX_AUTOMATION_TEST(FOpenDialogueVoiceAssetEditors, "Project.Tools.Sounds.Open Dialogue Voice Assets", (EAutomationTestFlags::ATF_Editor | EAutomationTestFlags::ATF_RequiresUser));
 
 void FOpenDialogueVoiceAssetEditors::GetTests(TArray<FString>& OutBeautifiedNames, TArray<FString>& OutTestCommands) const
 {
@@ -451,7 +451,7 @@ bool FOpenDialogueVoiceAssetEditors::RunTest(const FString& Parameters)
 /**
 * This test opens each DialogueWave into its sub-editor.
 */
-IMPLEMENT_COMPLEX_AUTOMATION_TEST(FOpenDialogueWaveAssetEditors, "Tools.Sounds.Open Dialogue Wave Assets", (EAutomationTestFlags::ATF_Editor | EAutomationTestFlags::ATF_RequiresUser));
+IMPLEMENT_COMPLEX_AUTOMATION_TEST(FOpenDialogueWaveAssetEditors, "Project.Tools.Sounds.Open Dialogue Wave Assets", (EAutomationTestFlags::ATF_Editor | EAutomationTestFlags::ATF_RequiresUser));
 
 void FOpenDialogueWaveAssetEditors::GetTests(TArray<FString>& OutBeautifiedNames, TArray<FString>& OutTestCommands) const
 {
@@ -470,7 +470,7 @@ bool FOpenDialogueWaveAssetEditors::RunTest(const FString& Parameters)
 /**
 * This test opens each Material into its sub-editor.
 */
-IMPLEMENT_COMPLEX_AUTOMATION_TEST(FOpenMaterialAssetEditors, "Tools.Materials.Open Material Assets", (EAutomationTestFlags::ATF_Editor | EAutomationTestFlags::ATF_RequiresUser));
+IMPLEMENT_COMPLEX_AUTOMATION_TEST(FOpenMaterialAssetEditors, "Project.Tools.Materials.Open Material Assets", (EAutomationTestFlags::ATF_Editor | EAutomationTestFlags::ATF_RequiresUser));
 
 void FOpenMaterialAssetEditors::GetTests(TArray<FString>& OutBeautifiedNames, TArray<FString>& OutTestCommands) const
 {
@@ -489,7 +489,7 @@ bool FOpenMaterialAssetEditors::RunTest(const FString& Parameters)
 ///**
 //* This test opens each MaterialFunction into its sub-editor.
 //*/
-IMPLEMENT_COMPLEX_AUTOMATION_TEST(FOpenMaterialFunctionAssetEditors, "Tools.Materials.Open Material Function Assets", (EAutomationTestFlags::ATF_Editor | EAutomationTestFlags::ATF_RequiresUser));
+IMPLEMENT_COMPLEX_AUTOMATION_TEST(FOpenMaterialFunctionAssetEditors, "Project.Tools.Materials.Open Material Function Assets", (EAutomationTestFlags::ATF_Editor | EAutomationTestFlags::ATF_RequiresUser));
 
 void FOpenMaterialFunctionAssetEditors::GetTests(TArray<FString>& OutBeautifiedNames, TArray<FString>& OutTestCommands) const
 {
@@ -508,7 +508,7 @@ bool FOpenMaterialFunctionAssetEditors::RunTest(const FString& Parameters)
 /**
 * This test opens each MaterialInstanceConstant into its sub-editor.
 */
-IMPLEMENT_COMPLEX_AUTOMATION_TEST(FOpenMaterialInstanceConstantAssetEditors, "Tools.Materials.Open Material Instance Constant Assets", (EAutomationTestFlags::ATF_Editor | EAutomationTestFlags::ATF_RequiresUser));
+IMPLEMENT_COMPLEX_AUTOMATION_TEST(FOpenMaterialInstanceConstantAssetEditors, "Project.Tools.Materials.Open Material Instance Constant Assets", (EAutomationTestFlags::ATF_Editor | EAutomationTestFlags::ATF_RequiresUser));
 
 void FOpenMaterialInstanceConstantAssetEditors::GetTests(TArray<FString>& OutBeautifiedNames, TArray<FString>& OutTestCommands) const
 {
@@ -527,7 +527,7 @@ bool FOpenMaterialInstanceConstantAssetEditors::RunTest(const FString& Parameter
 /**
 * This test opens each MaterialParameterCollection into its sub-editor.
 */
-IMPLEMENT_COMPLEX_AUTOMATION_TEST(FOpenMaterialParameterCollectionAssetEditors, "Tools.Materials.Open Material Parameter Collection Assets", (EAutomationTestFlags::ATF_Editor | EAutomationTestFlags::ATF_RequiresUser));
+IMPLEMENT_COMPLEX_AUTOMATION_TEST(FOpenMaterialParameterCollectionAssetEditors, "Project.Tools.Materials.Open Material Parameter Collection Assets", (EAutomationTestFlags::ATF_Editor | EAutomationTestFlags::ATF_RequiresUser));
 
 void FOpenMaterialParameterCollectionAssetEditors::GetTests(TArray<FString>& OutBeautifiedNames, TArray<FString>& OutTestCommands) const
 {
@@ -546,7 +546,7 @@ bool FOpenMaterialParameterCollectionAssetEditors::RunTest(const FString& Parame
 /**
 * This test opens each PARTICLE SYSTEM into its sub-editor.
 */
-IMPLEMENT_COMPLEX_AUTOMATION_TEST(FOpenParticleSystemAssetEditors, "Tools.Particle System.Open Particle System Assets", (EAutomationTestFlags::ATF_Editor | EAutomationTestFlags::ATF_RequiresUser));
+IMPLEMENT_COMPLEX_AUTOMATION_TEST(FOpenParticleSystemAssetEditors, "Project.Tools.Particle System.Open Particle System Assets", (EAutomationTestFlags::ATF_Editor | EAutomationTestFlags::ATF_RequiresUser));
 
 void FOpenParticleSystemAssetEditors::GetTests(TArray<FString>& OutBeautifiedNames, TArray<FString>& OutTestCommands) const
 {
@@ -565,7 +565,7 @@ bool FOpenParticleSystemAssetEditors::RunTest(const FString& Parameters)
 /**
 * This test opens each PhysicalMaterial into its sub-editor.
 */
-IMPLEMENT_COMPLEX_AUTOMATION_TEST(FOpenPhysicalMaterialAssetEditors, "Tools.Physics.Open Physical Material Assets", (EAutomationTestFlags::ATF_Editor | EAutomationTestFlags::ATF_RequiresUser));
+IMPLEMENT_COMPLEX_AUTOMATION_TEST(FOpenPhysicalMaterialAssetEditors, "Project.Tools.Physics.Open Physical Material Assets", (EAutomationTestFlags::ATF_Editor | EAutomationTestFlags::ATF_RequiresUser));
 
 void FOpenPhysicalMaterialAssetEditors::GetTests(TArray<FString>& OutBeautifiedNames, TArray<FString>& OutTestCommands) const
 {
@@ -585,7 +585,7 @@ bool FOpenPhysicalMaterialAssetEditors::RunTest(const FString& Parameters)
 /**
 * This test opens each PhysicsAsset into its sub-editor.
 */
-IMPLEMENT_COMPLEX_AUTOMATION_TEST(FOpenPhysicsAssetAssetEditors, "Tools.Physics.Open Physics Assets", (EAutomationTestFlags::ATF_Editor | EAutomationTestFlags::ATF_RequiresUser));
+IMPLEMENT_COMPLEX_AUTOMATION_TEST(FOpenPhysicsAssetAssetEditors, "Project.Tools.Physics.Open Physics Assets", (EAutomationTestFlags::ATF_Editor | EAutomationTestFlags::ATF_RequiresUser));
 
 void FOpenPhysicsAssetAssetEditors::GetTests(TArray<FString>& OutBeautifiedNames, TArray<FString>& OutTestCommands) const
 {
@@ -605,7 +605,7 @@ bool FOpenPhysicsAssetAssetEditors::RunTest(const FString& Parameters)
 /**
 * This test opens each ReverbEffect into its sub-editor.
 */
-IMPLEMENT_COMPLEX_AUTOMATION_TEST(FOpenReverbEffectAssetEditors, "Tools.Sounds.Open Reverb Effect Assets", (EAutomationTestFlags::ATF_Editor | EAutomationTestFlags::ATF_RequiresUser));
+IMPLEMENT_COMPLEX_AUTOMATION_TEST(FOpenReverbEffectAssetEditors, "Project.Tools.Sounds.Open Reverb Effect Assets", (EAutomationTestFlags::ATF_Editor | EAutomationTestFlags::ATF_RequiresUser));
 
 void FOpenReverbEffectAssetEditors::GetTests(TArray<FString>& OutBeautifiedNames, TArray<FString>& OutTestCommands) const
 {
@@ -626,7 +626,7 @@ bool FOpenReverbEffectAssetEditors::RunTest(const FString& Parameters)
 /**
 * This test opens each SlateWidgetStyleAsset into its sub-editor.
 */
-IMPLEMENT_COMPLEX_AUTOMATION_TEST(FOpenSlateWidgetStyleAssetAssetEditors, "Tools.Slate.Open Slate Widget Style Assets", (EAutomationTestFlags::ATF_Editor | EAutomationTestFlags::ATF_RequiresUser));
+IMPLEMENT_COMPLEX_AUTOMATION_TEST(FOpenSlateWidgetStyleAssetAssetEditors, "Project.Tools.Slate.Open Slate Widget Style Assets", (EAutomationTestFlags::ATF_Editor | EAutomationTestFlags::ATF_RequiresUser));
 
 void FOpenSlateWidgetStyleAssetAssetEditors::GetTests(TArray<FString>& OutBeautifiedNames, TArray<FString>& OutTestCommands) const
 {
@@ -645,7 +645,7 @@ bool FOpenSlateWidgetStyleAssetAssetEditors::RunTest(const FString& Parameters)
 /**
 * This test opens each Slate Brush Asset into its sub-editor.
 */
-IMPLEMENT_COMPLEX_AUTOMATION_TEST(FOpenSlateBrushAssetAssetEditors, "Tools.Slate.Open Slate Brush Assets", (EAutomationTestFlags::ATF_Editor | EAutomationTestFlags::ATF_RequiresUser));
+IMPLEMENT_COMPLEX_AUTOMATION_TEST(FOpenSlateBrushAssetAssetEditors, "Project.Tools.Slate.Open Slate Brush Assets", (EAutomationTestFlags::ATF_Editor | EAutomationTestFlags::ATF_RequiresUser));
 
 void FOpenSlateBrushAssetAssetEditors::GetTests(TArray<FString>& OutBeautifiedNames, TArray<FString>& OutTestCommands) const
 {
@@ -664,7 +664,7 @@ bool FOpenSlateBrushAssetAssetEditors::RunTest(const FString& Parameters)
 /**
 * This test opens each SoundAttenuation into its sub-editor.
 */
-IMPLEMENT_COMPLEX_AUTOMATION_TEST(FOpenSoundAttenuationAssetEditors, "Tools.Sounds.Open Sound Attenuation Assets", (EAutomationTestFlags::ATF_Editor | EAutomationTestFlags::ATF_RequiresUser));
+IMPLEMENT_COMPLEX_AUTOMATION_TEST(FOpenSoundAttenuationAssetEditors, "Project.Tools.Sounds.Open Sound Attenuation Assets", (EAutomationTestFlags::ATF_Editor | EAutomationTestFlags::ATF_RequiresUser));
 
 void FOpenSoundAttenuationAssetEditors::GetTests(TArray<FString>& OutBeautifiedNames, TArray<FString>& OutTestCommands) const
 {
@@ -683,7 +683,7 @@ bool FOpenSoundAttenuationAssetEditors::RunTest(const FString& Parameters)
 /**
 * This test opens each Static Mesh into their sub-editor.
 */
-IMPLEMENT_COMPLEX_AUTOMATION_TEST(FOpenStaticMeshAssetEditors, "Tools.Static Mesh.Open Static Mesh Assets", (EAutomationTestFlags::ATF_Editor | EAutomationTestFlags::ATF_RequiresUser));
+IMPLEMENT_COMPLEX_AUTOMATION_TEST(FOpenStaticMeshAssetEditors, "Project.Tools.Static Mesh.Open Static Mesh Assets", (EAutomationTestFlags::ATF_Editor | EAutomationTestFlags::ATF_RequiresUser));
 
 void FOpenStaticMeshAssetEditors::GetTests(TArray<FString>& OutBeautifiedNames, TArray<FString>& OutTestCommands) const
 {
@@ -702,7 +702,7 @@ bool FOpenStaticMeshAssetEditors::RunTest(const FString& Parameters)
 /**
 * This test opens each Skeletal Mesh into their sub-editor.
 */
-IMPLEMENT_COMPLEX_AUTOMATION_TEST(FOpenSkeletalMeshAssetEditors, "Tools.Skeletal Mesh.Open Skeletal Mesh Assets", (EAutomationTestFlags::ATF_Editor | EAutomationTestFlags::ATF_RequiresUser));
+IMPLEMENT_COMPLEX_AUTOMATION_TEST(FOpenSkeletalMeshAssetEditors, "Project.Tools.Skeletal Mesh.Open Skeletal Mesh Assets", (EAutomationTestFlags::ATF_Editor | EAutomationTestFlags::ATF_RequiresUser));
 
 void FOpenSkeletalMeshAssetEditors::GetTests(TArray<FString>& OutBeautifiedNames, TArray<FString>& OutTestCommands) const
 {
@@ -721,7 +721,7 @@ bool FOpenSkeletalMeshAssetEditors::RunTest(const FString& Parameters)
 /**
 * This test opens each Skeleton into its sub-editor.
 */
-IMPLEMENT_COMPLEX_AUTOMATION_TEST(FOpenSkeletonAssetEditors, "Tools.Skeletal Mesh.Open Skeleton Assets", (EAutomationTestFlags::ATF_Editor | EAutomationTestFlags::ATF_RequiresUser));
+IMPLEMENT_COMPLEX_AUTOMATION_TEST(FOpenSkeletonAssetEditors, "Project.Tools.Skeletal Mesh.Open Skeleton Assets", (EAutomationTestFlags::ATF_Editor | EAutomationTestFlags::ATF_RequiresUser));
 
 void FOpenSkeletonAssetEditors::GetTests(TArray<FString>& OutBeautifiedNames, TArray<FString>& OutTestCommands) const
 {
@@ -740,7 +740,7 @@ bool FOpenSkeletonAssetEditors::RunTest(const FString& Parameters)
 /**
 * This test opens each SoundClass into its sub-editor.
 */
-IMPLEMENT_COMPLEX_AUTOMATION_TEST(FOpenSoundClassAssetEditors, "Tools.Sounds.Open SoundClass Assets", (EAutomationTestFlags::ATF_Editor | EAutomationTestFlags::ATF_RequiresUser));
+IMPLEMENT_COMPLEX_AUTOMATION_TEST(FOpenSoundClassAssetEditors, "Project.Tools.Sounds.Open SoundClass Assets", (EAutomationTestFlags::ATF_Editor | EAutomationTestFlags::ATF_RequiresUser));
 
 void FOpenSoundClassAssetEditors::GetTests(TArray<FString>& OutBeautifiedNames, TArray<FString>& OutTestCommands) const
 {
@@ -760,7 +760,7 @@ bool FOpenSoundClassAssetEditors::RunTest(const FString& Parameters)
 /**
 * This test opens each SoundCue into its sub-editor.
 */
-IMPLEMENT_COMPLEX_AUTOMATION_TEST(FOpenSoundCueAssetEditors, "Tools.Sounds.Open Sound Cue Assets", (EAutomationTestFlags::ATF_Editor | EAutomationTestFlags::ATF_RequiresUser));
+IMPLEMENT_COMPLEX_AUTOMATION_TEST(FOpenSoundCueAssetEditors, "Project.Tools.Sounds.Open Sound Cue Assets", (EAutomationTestFlags::ATF_Editor | EAutomationTestFlags::ATF_RequiresUser));
 
 void FOpenSoundCueAssetEditors::GetTests(TArray<FString>& OutBeautifiedNames, TArray<FString>& OutTestCommands) const
 {
@@ -780,7 +780,7 @@ bool FOpenSoundCueAssetEditors::RunTest(const FString& Parameters)
 /**
 * This test opens each SoundMix into its sub-editor.
 */
-IMPLEMENT_COMPLEX_AUTOMATION_TEST(FOpenSoundMixAssetEditors, "Tools.Sounds.Open Sound Mix Assets", (EAutomationTestFlags::ATF_Editor | EAutomationTestFlags::ATF_RequiresUser));
+IMPLEMENT_COMPLEX_AUTOMATION_TEST(FOpenSoundMixAssetEditors, "Project.Tools.Sounds.Open Sound Mix Assets", (EAutomationTestFlags::ATF_Editor | EAutomationTestFlags::ATF_RequiresUser));
 
 void FOpenSoundMixAssetEditors::GetTests(TArray<FString>& OutBeautifiedNames, TArray<FString>& OutTestCommands) const
 {
@@ -800,7 +800,7 @@ bool FOpenSoundMixAssetEditors::RunTest(const FString& Parameters)
 /**
 * This test opens each SoundWave into its sub-editor.
 */
-IMPLEMENT_COMPLEX_AUTOMATION_TEST(FOpenSoundWaveAssetEditors, "Tools.Sounds.Open Sound Wave Assets", (EAutomationTestFlags::ATF_Editor | EAutomationTestFlags::ATF_RequiresUser));
+IMPLEMENT_COMPLEX_AUTOMATION_TEST(FOpenSoundWaveAssetEditors, "Project.Tools.Sounds.Open Sound Wave Assets", (EAutomationTestFlags::ATF_Editor | EAutomationTestFlags::ATF_RequiresUser));
 
 void FOpenSoundWaveAssetEditors::GetTests(TArray<FString>& OutBeautifiedNames, TArray<FString>& OutTestCommands) const
 {
@@ -819,7 +819,7 @@ bool FOpenSoundWaveAssetEditors::RunTest(const FString& Parameters)
 /**
 * This test opens each SubsurfaceProfile into its sub-editor.
 */
-IMPLEMENT_COMPLEX_AUTOMATION_TEST(FOpenSubsurfaceProfileAssetEditors, "Tools.Materials.Open Subsurface Profile Assets", (EAutomationTestFlags::ATF_Editor | EAutomationTestFlags::ATF_RequiresUser));
+IMPLEMENT_COMPLEX_AUTOMATION_TEST(FOpenSubsurfaceProfileAssetEditors, "Project.Tools.Materials.Open Subsurface Profile Assets", (EAutomationTestFlags::ATF_Editor | EAutomationTestFlags::ATF_RequiresUser));
 
 void FOpenSubsurfaceProfileAssetEditors::GetTests(TArray<FString>& OutBeautifiedNames, TArray<FString>& OutTestCommands) const
 {
@@ -840,7 +840,7 @@ bool FOpenSubsurfaceProfileAssetEditors::RunTest(const FString& Parameters)
 /**
 * This test opens each TEXTURE asset into its sub-editor.
 */
-IMPLEMENT_COMPLEX_AUTOMATION_TEST(FOpenTextureAssetEditors, "Tools.Textures.Open Texture Assets", (EAutomationTestFlags::ATF_Editor | EAutomationTestFlags::ATF_RequiresUser));
+IMPLEMENT_COMPLEX_AUTOMATION_TEST(FOpenTextureAssetEditors, "Project.Tools.Textures.Open Texture Assets", (EAutomationTestFlags::ATF_Editor | EAutomationTestFlags::ATF_RequiresUser));
 
 void FOpenTextureAssetEditors::GetTests(TArray<FString>& OutBeautifiedNames, TArray<FString>& OutTestCommands) const
 {
@@ -860,7 +860,7 @@ bool FOpenTextureAssetEditors::RunTest(const FString& Parameters)
 /**
 * This test opens each Font into its sub-editor.
 */
-IMPLEMENT_COMPLEX_AUTOMATION_TEST(FOpenFontAssetEditors, "Tools.Textures.Open Font Assets", (EAutomationTestFlags::ATF_Editor | EAutomationTestFlags::ATF_RequiresUser));
+IMPLEMENT_COMPLEX_AUTOMATION_TEST(FOpenFontAssetEditors, "Project.Tools.Textures.Open Font Assets", (EAutomationTestFlags::ATF_Editor | EAutomationTestFlags::ATF_RequiresUser));
 
 void FOpenFontAssetEditors::GetTests(TArray<FString>& OutBeautifiedNames, TArray<FString>& OutTestCommands) const
 {
@@ -880,7 +880,7 @@ bool FOpenFontAssetEditors::RunTest(const FString& Parameters)
 ///**
 //* This test opens each misc asset into its sub-editor.
 //*/
-IMPLEMENT_COMPLEX_AUTOMATION_TEST(FOpenMiscAssetEditors, "Tools.Misc.Open Misc Assets", (EAutomationTestFlags::ATF_Editor | EAutomationTestFlags::ATF_RequiresUser));
+IMPLEMENT_COMPLEX_AUTOMATION_TEST(FOpenMiscAssetEditors, "Project.Tools.Misc.Open Misc Assets", (EAutomationTestFlags::ATF_Editor | EAutomationTestFlags::ATF_RequiresUser));
 
 void FOpenMiscAssetEditors::GetTests(TArray<FString>& OutBeautifiedNames, TArray<FString>& OutTestCommands) const
 {
diff --git a/Engine/Source/Editor/UnrealEd/Private/Tests/EditorPerformanceTests.cpp b/Engine/Source/Editor/UnrealEd/Private/Tests/EditorPerformanceTests.cpp
--- a/Engine/Source/Editor/UnrealEd/Private/Tests/EditorPerformanceTests.cpp
+++ b/Engine/Source/Editor/UnrealEd/Private/Tests/EditorPerformanceTests.cpp
@@ -219,7 +219,7 @@ bool FEditorPerfCaptureCommand::Update()
 * Map Performance in Editor tests
 * Grabs certain performance numbers and saves it to a file.
 */
-IMPLEMENT_COMPLEX_AUTOMATION_TEST(FMapPerformanceInEditor, "Performance.Map Performance in Editor", EAutomationTestFlags::ATF_Editor | EAutomationTestFlags::ATF_RequiresUser)
+IMPLEMENT_COMPLEX_AUTOMATION_TEST(FMapPerformanceInEditor, "Project.Performance.Map Performance in Editor", EAutomationTestFlags::ATF_Editor | EAutomationTestFlags::ATF_RequiresUser)
 
 /**
 * Requests a enumeration of all maps to be loaded
diff --git a/Engine/Source/Runtime/Core/Private/Containers/ContainersTest.cpp b/Engine/Source/Runtime/Core/Private/Containers/ContainersTest.cpp
--- a/Engine/Source/Runtime/Core/Private/Containers/ContainersTest.cpp
+++ b/Engine/Source/Runtime/Core/Private/Containers/ContainersTest.cpp
@@ -4,7 +4,7 @@
 #include "AutomationTest.h"
 
 
-IMPLEMENT_SIMPLE_AUTOMATION_TEST(FContainersTest, "Core.Misc.Containers", EAutomationTestFlags::ATF_SmokeTest)
+IMPLEMENT_SIMPLE_AUTOMATION_TEST(FContainersTest, "System.Core.Misc.Containers", EAutomationTestFlags::ATF_SmokeTest)
 
 #define MAX_TEST_OBJECTS      65
 #define MAX_TEST_OBJECTS_STEP 1
diff --git a/Engine/Source/Runtime/Core/Private/Serialization/Csv/CsvParserTests.cpp b/Engine/Source/Runtime/Core/Private/Serialization/Csv/CsvParserTests.cpp
--- a/Engine/Source/Runtime/Core/Private/Serialization/Csv/CsvParserTests.cpp
+++ b/Engine/Source/Runtime/Core/Private/Serialization/Csv/CsvParserTests.cpp
@@ -49,7 +49,7 @@ namespace CsvParser_Tests
 	}
 }
 
-IMPLEMENT_SIMPLE_AUTOMATION_TEST(FBasicTest, "Core.CSV Parser.Basic", 0)
+IMPLEMENT_SIMPLE_AUTOMATION_TEST(FBasicTest, "System.Core.CSV Parser.Basic", 0)
 
 bool FBasicTest::RunTest(const FString& Parameters)
 {
@@ -64,7 +64,7 @@ bool FBasicTest::RunTest(const FString& Parameters)
 	return CsvParser_Tests::CheckRows(Parser.GetRows(), Expected, this);
 }
 
-IMPLEMENT_SIMPLE_AUTOMATION_TEST(FQuotedStringTest, "Core.CSV Parser.Quoted String", 0)
+IMPLEMENT_SIMPLE_AUTOMATION_TEST(FQuotedStringTest, "System.Core.CSV Parser.Quoted String", 0)
 
 bool FQuotedStringTest::RunTest(const FString& Parameters)
 {
@@ -82,7 +82,7 @@ bool FQuotedStringTest::RunTest(const FString& Parameters)
 	return CsvParser_Tests::CheckRows(Parser.GetRows(), Expected, this);
 }
 
-IMPLEMENT_SIMPLE_AUTOMATION_TEST(FEmptyCellsTest, "Core.CSV Parser.Empty Cells", 0)
+IMPLEMENT_SIMPLE_AUTOMATION_TEST(FEmptyCellsTest, "System.Core.CSV Parser.Empty Cells", 0)
 
 bool FEmptyCellsTest::RunTest(const FString& Parameters)
 {
@@ -100,7 +100,7 @@ bool FEmptyCellsTest::RunTest(const FString& Parameters)
 	return CsvParser_Tests::CheckRows(Parser.GetRows(), Expected, this);
 }
 
-IMPLEMENT_SIMPLE_AUTOMATION_TEST(FMixedNewLinesTest, "Core.CSV Parser.Mixed Newlines", 0)
+IMPLEMENT_SIMPLE_AUTOMATION_TEST(FMixedNewLinesTest, "System.Core.CSV Parser.Mixed Newlines", 0)
 
 bool FMixedNewLinesTest::RunTest(const FString& Parameters)
 {
@@ -113,7 +113,7 @@ bool FMixedNewLinesTest::RunTest(const FString& Parameters)
 	return CsvParser_Tests::CheckRows(Parser.GetRows(), Expected, this);
 }
 
-IMPLEMENT_SIMPLE_AUTOMATION_TEST(FEmptyStringTest, "Core.CSV Parser.Empty String", 0)
+IMPLEMENT_SIMPLE_AUTOMATION_TEST(FEmptyStringTest, "System.Core.CSV Parser.Empty String", 0)
 
 bool FEmptyStringTest::RunTest(const FString& Parameters)
 {
@@ -121,7 +121,7 @@ bool FEmptyStringTest::RunTest(const FString& Parameters)
 			FCsvParser(FString("")).GetRows().Num() == 0;
 }
 
-IMPLEMENT_SIMPLE_AUTOMATION_TEST(FOnlyNewLinesTest, "Core.CSV Parser.Only Newlines", 0)
+IMPLEMENT_SIMPLE_AUTOMATION_TEST(FOnlyNewLinesTest, "System.Core.CSV Parser.Only Newlines", 0)
 
 bool FOnlyNewLinesTest::RunTest(const FString& Parameters)
 {
diff --git a/Engine/Source/Runtime/Core/Private/Tests/Async/AsyncTest.cpp b/Engine/Source/Runtime/Core/Private/Tests/Async/AsyncTest.cpp
--- a/Engine/Source/Runtime/Core/Private/Tests/Async/AsyncTest.cpp
+++ b/Engine/Source/Runtime/Core/Private/Tests/Async/AsyncTest.cpp
@@ -7,7 +7,7 @@
 #include "Future.h"
 
 
-IMPLEMENT_SIMPLE_AUTOMATION_TEST(FAsyncTest, "Core.Async.Async", EAutomationTestFlags::ATF_Editor)
+IMPLEMENT_SIMPLE_AUTOMATION_TEST(FAsyncTest, "System.Core.Async.Async", EAutomationTestFlags::ATF_Editor)
 
 
 bool FAsyncTest::RunTest(const FString& Parameters)
diff --git a/Engine/Source/Runtime/Core/Private/Tests/HAL/PlatformTest.cpp b/Engine/Source/Runtime/Core/Private/Tests/HAL/PlatformTest.cpp
--- a/Engine/Source/Runtime/Core/Private/Tests/HAL/PlatformTest.cpp
+++ b/Engine/Source/Runtime/Core/Private/Tests/HAL/PlatformTest.cpp
@@ -34,7 +34,7 @@ struct TestC : public TestA, TestB
 };
 
 
-IMPLEMENT_SIMPLE_AUTOMATION_TEST(FPlatformVerificationTest, "Core.HAL.Platform Verification", EAutomationTestFlags::ATF_SmokeTest)
+IMPLEMENT_SIMPLE_AUTOMATION_TEST(FPlatformVerificationTest, "System.Core.HAL.Platform Verification", EAutomationTestFlags::ATF_SmokeTest)
 
 bool FPlatformVerificationTest::RunTest (const FString& Parameters)
 {
diff --git a/Engine/Source/Runtime/Core/Private/Tests/Internationalization/ArchiveTest.cpp b/Engine/Source/Runtime/Core/Private/Tests/Internationalization/ArchiveTest.cpp
--- a/Engine/Source/Runtime/Core/Private/Tests/Internationalization/ArchiveTest.cpp
+++ b/Engine/Source/Runtime/Core/Private/Tests/Internationalization/ArchiveTest.cpp
@@ -6,7 +6,7 @@
 #include "Internationalization/InternationalizationMetadata.h"
 
 
-IMPLEMENT_SIMPLE_AUTOMATION_TEST(FArchiveTest, "Core.Misc.Internationalization Archive", EAutomationTestFlags::ATF_SmokeTest)
+IMPLEMENT_SIMPLE_AUTOMATION_TEST(FArchiveTest, "System.Core.Misc.Internationalization Archive", EAutomationTestFlags::ATF_SmokeTest)
 
 
 bool FArchiveTest::RunTest( const FString& Parameters )
diff --git a/Engine/Source/Runtime/Core/Private/Tests/Internationalization/DateTimeFormattingRulesTest.cpp b/Engine/Source/Runtime/Core/Private/Tests/Internationalization/DateTimeFormattingRulesTest.cpp
--- a/Engine/Source/Runtime/Core/Private/Tests/Internationalization/DateTimeFormattingRulesTest.cpp
+++ b/Engine/Source/Runtime/Core/Private/Tests/Internationalization/DateTimeFormattingRulesTest.cpp
@@ -7,7 +7,7 @@
 // disable optimization for DateTimeFormattingRulesTest as it compiles very slowly in development builds
 PRAGMA_DISABLE_OPTIMIZATION
 
-	IMPLEMENT_SIMPLE_AUTOMATION_TEST(FDateTimeFormattingRulesTest, "Core.Misc.DateTime Formatting Rules", EAutomationTestFlags::ATF_None)
+	IMPLEMENT_SIMPLE_AUTOMATION_TEST(FDateTimeFormattingRulesTest, "System.Core.Misc.DateTime Formatting Rules", EAutomationTestFlags::ATF_None)
 
 	namespace
 {
diff --git a/Engine/Source/Runtime/Core/Private/Tests/Internationalization/ManifestTest.cpp b/Engine/Source/Runtime/Core/Private/Tests/Internationalization/ManifestTest.cpp
--- a/Engine/Source/Runtime/Core/Private/Tests/Internationalization/ManifestTest.cpp
+++ b/Engine/Source/Runtime/Core/Private/Tests/Internationalization/ManifestTest.cpp
@@ -15,7 +15,7 @@ int32 CountManifestEntries( const FInternationalizationManifest& Manifest )
 }
 
 
-IMPLEMENT_SIMPLE_AUTOMATION_TEST(FLocContextTest, "Core.Misc.Internationalization Context", EAutomationTestFlags::ATF_SmokeTest)
+IMPLEMENT_SIMPLE_AUTOMATION_TEST(FLocContextTest, "System.Core.Misc.Internationalization Context", EAutomationTestFlags::ATF_SmokeTest)
 
 bool FLocContextTest::RunTest( const FString& Parameters )
 {
@@ -200,7 +200,7 @@ bool FLocContextTest::RunTest( const FString& Parameters )
 }
 
 
-IMPLEMENT_SIMPLE_AUTOMATION_TEST(FLocItemTest, "Core.Misc.Internationalization LocItem", EAutomationTestFlags::ATF_SmokeTest)
+IMPLEMENT_SIMPLE_AUTOMATION_TEST(FLocItemTest, "System.Core.Misc.Internationalization LocItem", EAutomationTestFlags::ATF_SmokeTest)
 
 bool FLocItemTest::RunTest( const FString& Parameters )
 {
@@ -354,7 +354,7 @@ bool FLocItemTest::RunTest( const FString& Parameters )
 }
 
 
-IMPLEMENT_SIMPLE_AUTOMATION_TEST(FManifestTest, "Core.Misc.Internationalization Manifest", EAutomationTestFlags::ATF_SmokeTest)
+IMPLEMENT_SIMPLE_AUTOMATION_TEST(FManifestTest, "System.Core.Misc.Internationalization Manifest", EAutomationTestFlags::ATF_SmokeTest)
 
 bool FManifestTest::RunTest( const FString& Parameters )
 {
diff --git a/Engine/Source/Runtime/Core/Private/Tests/Internationalization/MetadataTest.cpp b/Engine/Source/Runtime/Core/Private/Tests/Internationalization/MetadataTest.cpp
--- a/Engine/Source/Runtime/Core/Private/Tests/Internationalization/MetadataTest.cpp
+++ b/Engine/Source/Runtime/Core/Private/Tests/Internationalization/MetadataTest.cpp
@@ -5,7 +5,7 @@
 #include "Internationalization/InternationalizationMetadata.h"
 
 
-IMPLEMENT_SIMPLE_AUTOMATION_TEST(FMetadataTest, "Core.Misc.Internationalization Metadata", EAutomationTestFlags::ATF_SmokeTest)
+IMPLEMENT_SIMPLE_AUTOMATION_TEST(FMetadataTest, "System.Core.Misc.Internationalization Metadata", EAutomationTestFlags::ATF_SmokeTest)
 
 bool FMetadataTest::RunTest( const FString& Parameters )
 {
diff --git a/Engine/Source/Runtime/Core/Private/Tests/Internationalization/NumberFormatingRulesTest.cpp b/Engine/Source/Runtime/Core/Private/Tests/Internationalization/NumberFormatingRulesTest.cpp
--- a/Engine/Source/Runtime/Core/Private/Tests/Internationalization/NumberFormatingRulesTest.cpp
+++ b/Engine/Source/Runtime/Core/Private/Tests/Internationalization/NumberFormatingRulesTest.cpp
@@ -7,7 +7,7 @@
 // Disable optimization for NumberFormatingRulesTest as it compiles very slowly in development builds
 PRAGMA_DISABLE_OPTIMIZATION
 
-	IMPLEMENT_SIMPLE_AUTOMATION_TEST(FNumberFormattingRulesTest, "Core.Misc.Number Formatting Rules", EAutomationTestFlags::ATF_None)
+	IMPLEMENT_SIMPLE_AUTOMATION_TEST(FNumberFormattingRulesTest, "System.Core.Misc.Number Formatting Rules", EAutomationTestFlags::ATF_None)
 
 	namespace
 {
diff --git a/Engine/Source/Runtime/Core/Private/Tests/Internationalization/TextTest.cpp b/Engine/Source/Runtime/Core/Private/Tests/Internationalization/TextTest.cpp
--- a/Engine/Source/Runtime/Core/Private/Tests/Internationalization/TextTest.cpp
+++ b/Engine/Source/Runtime/Core/Private/Tests/Internationalization/TextTest.cpp
@@ -11,7 +11,7 @@ PRAGMA_DISABLE_OPTIMIZATION
 #define LOCTEXT_NAMESPACE "Core.Tests.TextFormatTest"
 
 
-IMPLEMENT_SIMPLE_AUTOMATION_TEST(FTextTest, "Core.Misc.Text", EAutomationTestFlags::ATF_None)
+IMPLEMENT_SIMPLE_AUTOMATION_TEST(FTextTest, "System.Core.Misc.Text", EAutomationTestFlags::ATF_None)
 
 namespace
 {
@@ -638,7 +638,7 @@ bool FTextTest::RunTest (const FString& Parameters)
 
 #if UE_ENABLE_ICU
 
-IMPLEMENT_SIMPLE_AUTOMATION_TEST(FICUTextTest, "Core.Misc.ICUText", EAutomationTestFlags::ATF_SmokeTest)
+IMPLEMENT_SIMPLE_AUTOMATION_TEST(FICUTextTest, "System.Core.Misc.ICUText", EAutomationTestFlags::ATF_SmokeTest)
 
 bool FICUTextTest::RunTest (const FString& Parameters)
 {
diff --git a/Engine/Source/Runtime/Core/Private/Tests/Math/RangeBoundTest.cpp b/Engine/Source/Runtime/Core/Private/Tests/Math/RangeBoundTest.cpp
--- a/Engine/Source/Runtime/Core/Private/Tests/Math/RangeBoundTest.cpp
+++ b/Engine/Source/Runtime/Core/Private/Tests/Math/RangeBoundTest.cpp
@@ -4,7 +4,7 @@
 #include "AutomationTest.h"
 
 
-IMPLEMENT_SIMPLE_AUTOMATION_TEST(FRangeBoundTest, "Core.Math.RangeBound", EAutomationTestFlags::ATF_SmokeTest)
+IMPLEMENT_SIMPLE_AUTOMATION_TEST(FRangeBoundTest, "System.Core.Math.RangeBound", EAutomationTestFlags::ATF_SmokeTest)
 
 bool FRangeBoundTest::RunTest( const FString& Parameters )
 {
diff --git a/Engine/Source/Runtime/Core/Private/Tests/Math/RangeSetTest.cpp b/Engine/Source/Runtime/Core/Private/Tests/Math/RangeSetTest.cpp
--- a/Engine/Source/Runtime/Core/Private/Tests/Math/RangeSetTest.cpp
+++ b/Engine/Source/Runtime/Core/Private/Tests/Math/RangeSetTest.cpp
@@ -4,7 +4,7 @@
 #include "AutomationTest.h"
 
 
-IMPLEMENT_SIMPLE_AUTOMATION_TEST(FRangeSetTest, "Core.Math.RangeSet", EAutomationTestFlags::ATF_SmokeTest)
+IMPLEMENT_SIMPLE_AUTOMATION_TEST(FRangeSetTest, "System.Core.Math.RangeSet", EAutomationTestFlags::ATF_SmokeTest)
 
 bool FRangeSetTest::RunTest( const FString& Parameters )
 {
diff --git a/Engine/Source/Runtime/Core/Private/Tests/Math/RangeTest.cpp b/Engine/Source/Runtime/Core/Private/Tests/Math/RangeTest.cpp
--- a/Engine/Source/Runtime/Core/Private/Tests/Math/RangeTest.cpp
+++ b/Engine/Source/Runtime/Core/Private/Tests/Math/RangeTest.cpp
@@ -8,7 +8,7 @@
 PRAGMA_DISABLE_OPTIMIZATION
 
 
-IMPLEMENT_SIMPLE_AUTOMATION_TEST(FRangeTest, "Core.Math.Range", EAutomationTestFlags::ATF_SmokeTest)
+IMPLEMENT_SIMPLE_AUTOMATION_TEST(FRangeTest, "System.Core.Math.Range", EAutomationTestFlags::ATF_SmokeTest)
 
 bool FRangeTest::RunTest( const FString& Parameters )
 {
diff --git a/Engine/Source/Runtime/Core/Private/Tests/Math/UnitConversionTests.cpp b/Engine/Source/Runtime/Core/Private/Tests/Math/UnitConversionTests.cpp
--- a/Engine/Source/Runtime/Core/Private/Tests/Math/UnitConversionTests.cpp
+++ b/Engine/Source/Runtime/Core/Private/Tests/Math/UnitConversionTests.cpp
@@ -4,7 +4,7 @@
 #include "UnitConversion.h"
 
 
-IMPLEMENT_SIMPLE_AUTOMATION_TEST(FUnitUnitTests, "Core.Math.Unit Conversion", EAutomationTestFlags::ATF_SmokeTest)
+IMPLEMENT_SIMPLE_AUTOMATION_TEST(FUnitUnitTests, "System.Core.Math.Unit Conversion", EAutomationTestFlags::ATF_SmokeTest)
 
 bool IsRoughlyEqual(double One, double Two, float Epsilon)
 {
diff --git a/Engine/Source/Runtime/Core/Private/Tests/Math/UnrealMathTest.cpp b/Engine/Source/Runtime/Core/Private/Tests/Math/UnrealMathTest.cpp
--- a/Engine/Source/Runtime/Core/Private/Tests/Math/UnrealMathTest.cpp
+++ b/Engine/Source/Runtime/Core/Private/Tests/Math/UnrealMathTest.cpp
@@ -432,7 +432,7 @@ void SetScratch( float X, float Y, float Z, float W, float U = 0.0f )
 }
 
 
-IMPLEMENT_SIMPLE_AUTOMATION_TEST(FVectorRegisterAbstractionTest, "Core.Math.Vector Register Abstraction Test", EAutomationTestFlags::ATF_SmokeTest)
+IMPLEMENT_SIMPLE_AUTOMATION_TEST(FVectorRegisterAbstractionTest, "System.Core.Math.Vector Register Abstraction Test", EAutomationTestFlags::ATF_SmokeTest)
 
 /**
  * Run a suite of vector operations to validate vector intrinsics are working on the platform
diff --git a/Engine/Source/Runtime/Core/Private/Tests/Misc/CircularBufferTest.cpp b/Engine/Source/Runtime/Core/Private/Tests/Misc/CircularBufferTest.cpp
--- a/Engine/Source/Runtime/Core/Private/Tests/Misc/CircularBufferTest.cpp
+++ b/Engine/Source/Runtime/Core/Private/Tests/Misc/CircularBufferTest.cpp
@@ -4,7 +4,7 @@
 #include "AutomationTest.h"
 
 
-IMPLEMENT_SIMPLE_AUTOMATION_TEST(FCircularBufferTest, "Core.Misc.CircularBuffer", EAutomationTestFlags::ATF_SmokeTest)
+IMPLEMENT_SIMPLE_AUTOMATION_TEST(FCircularBufferTest, "System.Core.Misc.CircularBuffer", EAutomationTestFlags::ATF_SmokeTest)
 
 bool FCircularBufferTest::RunTest( const FString& Parameters )
 {
diff --git a/Engine/Source/Runtime/Core/Private/Tests/Misc/CircularQueueTest.cpp b/Engine/Source/Runtime/Core/Private/Tests/Misc/CircularQueueTest.cpp
--- a/Engine/Source/Runtime/Core/Private/Tests/Misc/CircularQueueTest.cpp
+++ b/Engine/Source/Runtime/Core/Private/Tests/Misc/CircularQueueTest.cpp
@@ -5,7 +5,7 @@
 #include "Misc/AutomationTest.h"
 
 
-IMPLEMENT_SIMPLE_AUTOMATION_TEST(FCircularQueueTest, "Core.Misc.CircularQueue", EAutomationTestFlags::ATF_SmokeTest)
+IMPLEMENT_SIMPLE_AUTOMATION_TEST(FCircularQueueTest, "System.Core.Misc.CircularQueue", EAutomationTestFlags::ATF_SmokeTest)
 
 bool FCircularQueueTest::RunTest( const FString& Parameters )
 {
diff --git a/Engine/Source/Runtime/Core/Private/Tests/Misc/DateTimeTest.cpp b/Engine/Source/Runtime/Core/Private/Tests/Misc/DateTimeTest.cpp
--- a/Engine/Source/Runtime/Core/Private/Tests/Misc/DateTimeTest.cpp
+++ b/Engine/Source/Runtime/Core/Private/Tests/Misc/DateTimeTest.cpp
@@ -4,7 +4,7 @@
 #include "AutomationTest.h"
 
 
-IMPLEMENT_SIMPLE_AUTOMATION_TEST(FDateTimeTest, "Core.Misc.DateTime", EAutomationTestFlags::ATF_SmokeTest)
+IMPLEMENT_SIMPLE_AUTOMATION_TEST(FDateTimeTest, "System.Core.Misc.DateTime", EAutomationTestFlags::ATF_SmokeTest)
 
 bool FDateTimeTest::RunTest( const FString& Parameters )
 {
diff --git a/Engine/Source/Runtime/Core/Private/Tests/Misc/GuidTest.cpp b/Engine/Source/Runtime/Core/Private/Tests/Misc/GuidTest.cpp
--- a/Engine/Source/Runtime/Core/Private/Tests/Misc/GuidTest.cpp
+++ b/Engine/Source/Runtime/Core/Private/Tests/Misc/GuidTest.cpp
@@ -4,7 +4,7 @@
 #include "AutomationTest.h"
 
 
-IMPLEMENT_SIMPLE_AUTOMATION_TEST(FGuidTest, "Core.Misc.Guid", EAutomationTestFlags::ATF_SmokeTest)
+IMPLEMENT_SIMPLE_AUTOMATION_TEST(FGuidTest, "System.Core.Misc.Guid", EAutomationTestFlags::ATF_SmokeTest)
 
 bool FGuidTest::RunTest( const FString& Parameters )
 {
diff --git a/Engine/Source/Runtime/Core/Private/Tests/Misc/PathsTest.cpp b/Engine/Source/Runtime/Core/Private/Tests/Misc/PathsTest.cpp
--- a/Engine/Source/Runtime/Core/Private/Tests/Misc/PathsTest.cpp
+++ b/Engine/Source/Runtime/Core/Private/Tests/Misc/PathsTest.cpp
@@ -7,7 +7,7 @@
 FString FPaths::GameProjectFilePath;
 
 
-IMPLEMENT_SIMPLE_AUTOMATION_TEST(FPathTests, "Core.Misc.Paths", EAutomationTestFlags::ATF_SmokeTest)
+IMPLEMENT_SIMPLE_AUTOMATION_TEST(FPathTests, "System.Core.Misc.Paths", EAutomationTestFlags::ATF_SmokeTest)
 
 bool FPathTests::RunTest( const FString& Parameters )
 {
diff --git a/Engine/Source/Runtime/Core/Private/Tests/Misc/QueueTest.cpp b/Engine/Source/Runtime/Core/Private/Tests/Misc/QueueTest.cpp
--- a/Engine/Source/Runtime/Core/Private/Tests/Misc/QueueTest.cpp
+++ b/Engine/Source/Runtime/Core/Private/Tests/Misc/QueueTest.cpp
@@ -4,7 +4,7 @@
 #include "AutomationTest.h"
 
 
-IMPLEMENT_SIMPLE_AUTOMATION_TEST(FQueueTest, "Core.Misc.Queue", EAutomationTestFlags::ATF_SmokeTest)
+IMPLEMENT_SIMPLE_AUTOMATION_TEST(FQueueTest, "System.Core.Misc.Queue", EAutomationTestFlags::ATF_SmokeTest)
 
 bool FQueueTest::RunTest( const FString& Parameters )
 {
diff --git a/Engine/Source/Runtime/Core/Private/Tests/Misc/TimespanTest.cpp b/Engine/Source/Runtime/Core/Private/Tests/Misc/TimespanTest.cpp
--- a/Engine/Source/Runtime/Core/Private/Tests/Misc/TimespanTest.cpp
+++ b/Engine/Source/Runtime/Core/Private/Tests/Misc/TimespanTest.cpp
@@ -4,7 +4,7 @@
 #include "AutomationTest.h"
 
 
-IMPLEMENT_SIMPLE_AUTOMATION_TEST(FTimespanTest, "Core.Misc.Timespan", EAutomationTestFlags::ATF_SmokeTest)
+IMPLEMENT_SIMPLE_AUTOMATION_TEST(FTimespanTest, "System.Core.Misc.Timespan", EAutomationTestFlags::ATF_SmokeTest)
 
 bool FTimespanTest::RunTest( const FString& Parameters )
 {
diff --git a/Engine/Source/Runtime/Engine/Private/Commandlets/SmokeTestCommandlet.cpp b/Engine/Source/Runtime/Engine/Private/Commandlets/SmokeTestCommandlet.cpp
--- a/Engine/Source/Runtime/Engine/Private/Commandlets/SmokeTestCommandlet.cpp
+++ b/Engine/Source/Runtime/Engine/Private/Commandlets/SmokeTestCommandlet.cpp
@@ -29,7 +29,7 @@ int32 USmokeTestCommandlet::Main( const FString& Params )
 }
 
 
-IMPLEMENT_SIMPLE_AUTOMATION_TEST(FRunServer, "Core.Misc.Run Server", EAutomationTestFlags::ATF_Commandlet | EAutomationTestFlags::ATF_SmokeTest)
+IMPLEMENT_SIMPLE_AUTOMATION_TEST(FRunServer, "System.Core.Misc.Run Server", EAutomationTestFlags::ATF_Commandlet | EAutomationTestFlags::ATF_SmokeTest)
 
 bool FRunServer::RunTest(const FString& Parameters)
 {
diff --git a/Engine/Source/Runtime/Engine/Private/Tests/AnalyticsTest.cpp b/Engine/Source/Runtime/Engine/Private/Tests/AnalyticsTest.cpp
--- a/Engine/Source/Runtime/Engine/Private/Tests/AnalyticsTest.cpp
+++ b/Engine/Source/Runtime/Engine/Private/Tests/AnalyticsTest.cpp
@@ -10,7 +10,7 @@
 /**
 * Artificial Record Event for analytics - Simulates the engine startup simulation.  
 */
-IMPLEMENT_SIMPLE_AUTOMATION_TEST(FAnalyticStartUpSimTest, "Engine.Analytic.Record Event - Simulate Program Start", EAutomationTestFlags::ATF_FeatureMask)
+IMPLEMENT_SIMPLE_AUTOMATION_TEST(FAnalyticStartUpSimTest, "System.Engine.Analytic.Record Event - Simulate Program Start", EAutomationTestFlags::ATF_FeatureMask)
 
 bool FAnalyticStartUpSimTest::RunTest(const FString& Parameters)
 {
@@ -81,7 +81,7 @@ bool FAnalyticStartUpSimTest::RunTest(const FString& Parameters)
 /**
 * FAnalyticsEventAttribute Unit Test.
 */
-IMPLEMENT_SIMPLE_AUTOMATION_TEST(FAnalyticsEventAttributeUnitTest, "Engine.Analytic.EventAttribute Struct Unit Test", EAutomationTestFlags::ATF_FeatureMask)
+IMPLEMENT_SIMPLE_AUTOMATION_TEST(FAnalyticsEventAttributeUnitTest, "System.Engine.Analytic.EventAttribute Struct Unit Test", EAutomationTestFlags::ATF_FeatureMask)
 
 bool FAnalyticsEventAttributeUnitTest::RunTest(const FString& Parameters)
 {
diff --git a/Engine/Source/Runtime/Engine/Private/Tests/EngineAutomationTests.cpp b/Engine/Source/Runtime/Engine/Private/Tests/EngineAutomationTests.cpp
--- a/Engine/Source/Runtime/Engine/Private/Tests/EngineAutomationTests.cpp
+++ b/Engine/Source/Runtime/Engine/Private/Tests/EngineAutomationTests.cpp
@@ -63,7 +63,7 @@ namespace
 /**
  * SetRes Verification - Verify changing resolution works
  */
-IMPLEMENT_SIMPLE_AUTOMATION_TEST( FSetResTest, "Windows.Set Resolution", EAutomationTestFlags::ATF_Game )
+IMPLEMENT_SIMPLE_AUTOMATION_TEST( FSetResTest, "System.Windows.Set Resolution", EAutomationTestFlags::ATF_Game )
 
 /** 
  * Change resolutions, wait, and change back
@@ -97,7 +97,7 @@ bool FSetResTest::RunTest(const FString& Parameters)
 /**
  * Stats verification - Toggle various "stats" commands
  */
-IMPLEMENT_SIMPLE_AUTOMATION_TEST( FStatsVerificationMapTest, "Maps.Stats Verification", EAutomationTestFlags::ATF_Game )
+IMPLEMENT_SIMPLE_AUTOMATION_TEST( FStatsVerificationMapTest, "System.Maps.Stats Verification", EAutomationTestFlags::ATF_Game )
 
 /** 
  * Execute the loading of one map to verify screen captures and performance captures work
@@ -192,7 +192,7 @@ bool FStatsVerificationMapTest::RunTest(const FString& Parameters)
  * LoadAutomationMap
  * Verification automation test to make sure features of map loading work (load, screen capture, performance capture)
  */
-IMPLEMENT_SIMPLE_AUTOMATION_TEST( FPerformanceCaptureTest, "Maps.Performance Capture", EAutomationTestFlags::ATF_Game )
+IMPLEMENT_SIMPLE_AUTOMATION_TEST( FPerformanceCaptureTest, "System.Maps.Performance Capture", EAutomationTestFlags::ATF_Game )
 
 /** 
  * Execute the loading of one map to verify screen captures and performance captures work
@@ -228,7 +228,7 @@ bool FTakeViewportScreenshotCommand::Update()
  * LoadAllMapsInGame
  * Verification automation test to make sure loading all maps succeed without crashing AND does performance captures
  */
-IMPLEMENT_COMPLEX_AUTOMATION_TEST( FLoadAllMapsInGameTest, "Maps.Load All In Game", EAutomationTestFlags::ATF_Game )
+IMPLEMENT_COMPLEX_AUTOMATION_TEST( FLoadAllMapsInGameTest, "Project.Maps.Load All In Game", EAutomationTestFlags::ATF_Game )
 
 /** 
  * Requests a enumeration of all maps to be loaded
@@ -276,7 +276,7 @@ bool FLoadAllMapsInGameTest::RunTest(const FString& Parameters)
  * SaveGameTest
  * Test makes sure a save game (without UI) saves and loads correctly
  */
-IMPLEMENT_SIMPLE_AUTOMATION_TEST( FSaveGameTest, "Engine.Game.Noninteractive Save", EAutomationTestFlags::ATF_Game )
+IMPLEMENT_SIMPLE_AUTOMATION_TEST( FSaveGameTest, "System.Engine.Game.Noninteractive Save", EAutomationTestFlags::ATF_Game )
 
 /** 
  * Saves and loads a savegame file
@@ -338,7 +338,7 @@ bool FSaveGameTest::RunTest(const FString& Parameters)
 /**
  * Automation test to load a map and capture FPS performance charts
  */
-IMPLEMENT_COMPLEX_AUTOMATION_TEST(FCinematicFPSPerfTest, "Maps.Cinematic FPS Perf Capture", (EAutomationTestFlags::ATF_Game | EAutomationTestFlags::ATF_NonNullRHI));
+IMPLEMENT_COMPLEX_AUTOMATION_TEST(FCinematicFPSPerfTest, "Project.Maps.Cinematic FPS Perf Capture", (EAutomationTestFlags::ATF_Game | EAutomationTestFlags::ATF_NonNullRHI));
 
 void FCinematicFPSPerfTest::GetTests(TArray<FString>& OutBeautifiedNames, TArray <FString>& OutTestCommands) const
 {
@@ -427,7 +427,7 @@ bool FCinematicFPSPerfTest::RunTest(const FString& Parameters)
 }
 
 
-IMPLEMENT_SIMPLE_AUTOMATION_TEST(FLogTypesTest, "Automation Framework.Logging Test", EAutomationTestFlags::ATF_None)
+IMPLEMENT_SIMPLE_AUTOMATION_TEST(FLogTypesTest, "System.Automation Framework.Logging Test", EAutomationTestFlags::ATF_None)
 
 bool FLogTypesTest::RunTest(const FString& Parameters)
 {
diff --git a/Engine/Source/Runtime/Engine/Private/Tests/IntSerializationTest.cpp b/Engine/Source/Runtime/Engine/Private/Tests/IntSerializationTest.cpp
--- a/Engine/Source/Runtime/Engine/Private/Tests/IntSerializationTest.cpp
+++ b/Engine/Source/Runtime/Engine/Private/Tests/IntSerializationTest.cpp
@@ -4,7 +4,7 @@
 #include "AutomationTest.h"
 #include "Runtime/Engine/Classes/Engine/IntSerialization.h"
 
-IMPLEMENT_SIMPLE_AUTOMATION_TEST(FIntSerializationTest, "Engine.Int Serialization", EAutomationTestFlags::ATF_SmokeTest)
+IMPLEMENT_SIMPLE_AUTOMATION_TEST(FIntSerializationTest, "System.Engine.Int Serialization", EAutomationTestFlags::ATF_SmokeTest)
 
 bool FIntSerializationTest::RunTest (const FString& Parameters)
 {
diff --git a/Engine/Source/Runtime/Engine/Private/Tests/Internationalization/TextPropertyTest.cpp b/Engine/Source/Runtime/Engine/Private/Tests/Internationalization/TextPropertyTest.cpp
--- a/Engine/Source/Runtime/Engine/Private/Tests/Internationalization/TextPropertyTest.cpp
+++ b/Engine/Source/Runtime/Engine/Private/Tests/Internationalization/TextPropertyTest.cpp
@@ -11,7 +11,7 @@ UTextPropertyTestObject::UTextPropertyTestObject(const FObjectInitializer& Objec
 {
 }
 
-IMPLEMENT_SIMPLE_AUTOMATION_TEST(FTextPropertyTest, "Engine.Internationalization.Text Property Test", EAutomationTestFlags::ATF_SmokeTest)
+IMPLEMENT_SIMPLE_AUTOMATION_TEST(FTextPropertyTest, "System.Engine.Internationalization.Text Property Test", EAutomationTestFlags::ATF_SmokeTest)
 
 bool FTextPropertyTest::RunTest (const FString& Parameters)
 {
diff --git a/Engine/Source/Runtime/Engine/Private/Tests/NetworkAutomationTests.cpp b/Engine/Source/Runtime/Engine/Private/Tests/NetworkAutomationTests.cpp
--- a/Engine/Source/Runtime/Engine/Private/Tests/NetworkAutomationTests.cpp
+++ b/Engine/Source/Runtime/Engine/Private/Tests/NetworkAutomationTests.cpp
@@ -24,7 +24,7 @@ namespace EMultiplayerAutomationRoles
  * 2-Multiplayer session automation test
  * Verification for 2 player multiplayer session start up and tear down
  */
-IMPLEMENT_NETWORKED_AUTOMATION_TEST( FMultiplayer4PlayerTest, "Networking.Multiplayer.TwoPlayerSessionStartupShutdown", EAutomationTestFlags::ATF_Game, EMultiplayerAutomationRoles::MaxNumParticipants )
+IMPLEMENT_NETWORKED_AUTOMATION_TEST( FMultiplayer4PlayerTest, "System.Networking.Multiplayer.TwoPlayerSessionStartupShutdown", EAutomationTestFlags::ATF_Game, EMultiplayerAutomationRoles::MaxNumParticipants )
 
 
 /** 
diff --git a/Engine/Source/Runtime/Engine/Private/Tests/OSSTests.cpp b/Engine/Source/Runtime/Engine/Private/Tests/OSSTests.cpp
--- a/Engine/Source/Runtime/Engine/Private/Tests/OSSTests.cpp
+++ b/Engine/Source/Runtime/Engine/Private/Tests/OSSTests.cpp
@@ -18,7 +18,7 @@ DEFINE_LOG_CATEGORY_STATIC(LogHackAutomationTests, Log, All);
 //////////////////////////////////////////////////////////////////////////
 
 
-IMPLEMENT_SIMPLE_AUTOMATION_TEST(FHttpTest,"OSS.Test Http", EAutomationTestFlags::ATF_Editor )
+IMPLEMENT_SIMPLE_AUTOMATION_TEST(FHttpTest,"System.OSS.Test Http", EAutomationTestFlags::ATF_Editor )
 
 bool FHttpTest::RunTest(const FString& Parameters)
 {
@@ -28,7 +28,7 @@ bool FHttpTest::RunTest(const FString& Parameters)
 	return true;
 }
 
-IMPLEMENT_SIMPLE_AUTOMATION_TEST(FMCPFriendsTest,"OSS.MCP.Test Friends", EAutomationTestFlags::ATF_Editor )
+IMPLEMENT_SIMPLE_AUTOMATION_TEST(FMCPFriendsTest,"System.OSS.MCP.Test Friends", EAutomationTestFlags::ATF_Editor )
 
 bool FMCPFriendsTest::RunTest(const FString& Parameters)
 {
@@ -39,7 +39,7 @@ bool FMCPFriendsTest::RunTest(const FString& Parameters)
 	return true;
 }
 
-IMPLEMENT_SIMPLE_AUTOMATION_TEST(FAmazonFriendsTest,"OSS.Amazon.Test Friends", EAutomationTestFlags::ATF_Editor )
+IMPLEMENT_SIMPLE_AUTOMATION_TEST(FAmazonFriendsTest,"System.OSS.Amazon.Test Friends", EAutomationTestFlags::ATF_Editor )
 
 bool FAmazonFriendsTest::RunTest(const FString& Parameters)
 {
@@ -49,7 +49,7 @@ bool FAmazonFriendsTest::RunTest(const FString& Parameters)
 	return true;
 }
 
-IMPLEMENT_SIMPLE_AUTOMATION_TEST(FFacebookFriendsTest,"OSS.Facebook.Test Friends", EAutomationTestFlags::ATF_Editor )
+IMPLEMENT_SIMPLE_AUTOMATION_TEST(FFacebookFriendsTest,"System.OSS.Facebook.Test Friends", EAutomationTestFlags::ATF_Editor )
 
 bool FFacebookFriendsTest::RunTest(const FString& Parameters)
 {
@@ -59,7 +59,7 @@ bool FFacebookFriendsTest::RunTest(const FString& Parameters)
 	return true;
 }
 
-IMPLEMENT_SIMPLE_AUTOMATION_TEST(FSteamFriendsTest,"OSS.Steam.Test Friends", EAutomationTestFlags::ATF_Editor )
+IMPLEMENT_SIMPLE_AUTOMATION_TEST(FSteamFriendsTest,"System.OSS.Steam.Test Friends", EAutomationTestFlags::ATF_Editor )
 
 bool FSteamFriendsTest::RunTest(const FString& Parameters)
 {
@@ -69,7 +69,7 @@ bool FSteamFriendsTest::RunTest(const FString& Parameters)
 	return true;
 }
 
-IMPLEMENT_SIMPLE_AUTOMATION_TEST(FIOSFriendsTest,"OSS.IOS.Test Friends", EAutomationTestFlags::ATF_Editor )
+IMPLEMENT_SIMPLE_AUTOMATION_TEST(FIOSFriendsTest,"System.OSS.IOS.Test Friends", EAutomationTestFlags::ATF_Editor )
 
 bool FIOSFriendsTest::RunTest(const FString& Parameters)
 {
@@ -83,7 +83,7 @@ bool FIOSFriendsTest::RunTest(const FString& Parameters)
 
 /////////////Currently crashing//////////////
 
-IMPLEMENT_SIMPLE_AUTOMATION_TEST(FNullFriendsTest,"OSS.NULL.Test Friends", EAutomationTestFlags::ATF_Editor )
+IMPLEMENT_SIMPLE_AUTOMATION_TEST(FNullFriendsTest,"System.OSS.NULL.Test Friends", EAutomationTestFlags::ATF_Editor )
 
 bool FNullFriendsTest::RunTest(const FString& Parameters)
 {
@@ -95,7 +95,7 @@ bool FNullFriendsTest::RunTest(const FString& Parameters)
 }
 */
 
-IMPLEMENT_SIMPLE_AUTOMATION_TEST(FMCPCloudTest,"OSS.MCP.Test Cloud", EAutomationTestFlags::ATF_Editor )
+IMPLEMENT_SIMPLE_AUTOMATION_TEST(FMCPCloudTest,"System.OSS.MCP.Test Cloud", EAutomationTestFlags::ATF_Editor )
 
 bool FMCPCloudTest::RunTest(const FString& Parameters)
 {
@@ -105,7 +105,7 @@ bool FMCPCloudTest::RunTest(const FString& Parameters)
 	return true;
 }
 
-IMPLEMENT_SIMPLE_AUTOMATION_TEST(FSteamCloudTest,"OSS.Steam.Test Cloud", EAutomationTestFlags::ATF_Editor )
+IMPLEMENT_SIMPLE_AUTOMATION_TEST(FSteamCloudTest,"System.OSS.Steam.Test Cloud", EAutomationTestFlags::ATF_Editor )
 
 bool FSteamCloudTest::RunTest(const FString& Parameters)
 {
@@ -119,7 +119,7 @@ bool FSteamCloudTest::RunTest(const FString& Parameters)
 
 /////////////Currently crashing//////////////
 
-IMPLEMENT_SIMPLE_AUTOMATION_TEST(FNullCloudTest,"OSS.NULL.Test Cloud", EAutomationTestFlags::ATF_Editor )
+IMPLEMENT_SIMPLE_AUTOMATION_TEST(FNullCloudTest,"System.OSS.NULL.Test Cloud", EAutomationTestFlags::ATF_Editor )
 
 bool FNullCloudTest::RunTest(const FString& Parameters)
 {
@@ -134,7 +134,7 @@ bool FNullCloudTest::RunTest(const FString& Parameters)
 
 /////////////Currently crashing//////////////
 
-IMPLEMENT_SIMPLE_AUTOMATION_TEST(FNullLeaderBoardTest,"OSS.NULL.Test Leaderboards", EAutomationTestFlags::ATF_Editor )
+IMPLEMENT_SIMPLE_AUTOMATION_TEST(FNullLeaderBoardTest,"System.OSS.NULL.Test Leaderboards", EAutomationTestFlags::ATF_Editor )
 
 bool FNullLeaderBoardTest::RunTest(const FString& Parameters)
 {
@@ -145,7 +145,7 @@ bool FNullLeaderBoardTest::RunTest(const FString& Parameters)
 }
 */
 
-IMPLEMENT_SIMPLE_AUTOMATION_TEST(FIosLeaderBoardTest,"OSS.IOS.Test Leaderboards", EAutomationTestFlags::ATF_Editor )
+IMPLEMENT_SIMPLE_AUTOMATION_TEST(FIosLeaderBoardTest,"System.OSS.IOS.Test Leaderboards", EAutomationTestFlags::ATF_Editor )
 
 bool FIosLeaderBoardTest::RunTest(const FString& Parameters)
 {
@@ -156,7 +156,7 @@ bool FIosLeaderBoardTest::RunTest(const FString& Parameters)
 }
 
 
-IMPLEMENT_SIMPLE_AUTOMATION_TEST(FSteamLeaderBoardTest,"OSS.Steam.Test Leaderboards", EAutomationTestFlags::ATF_Editor )
+IMPLEMENT_SIMPLE_AUTOMATION_TEST(FSteamLeaderBoardTest,"System.OSS.Steam.Test Leaderboards", EAutomationTestFlags::ATF_Editor )
 
 bool FSteamLeaderBoardTest::RunTest(const FString& Parameters)
 {
@@ -166,7 +166,7 @@ bool FSteamLeaderBoardTest::RunTest(const FString& Parameters)
 	return true;
 }
 
-IMPLEMENT_SIMPLE_AUTOMATION_TEST(FMCPTimeTest,"OSS.MCP.Test Time", EAutomationTestFlags::ATF_Editor )
+IMPLEMENT_SIMPLE_AUTOMATION_TEST(FMCPTimeTest,"System.OSS.MCP.Test Time", EAutomationTestFlags::ATF_Editor )
 
 bool FMCPTimeTest::RunTest(const FString& Parameters)
 {
@@ -176,7 +176,7 @@ bool FMCPTimeTest::RunTest(const FString& Parameters)
 	return true;
 }
 
-IMPLEMENT_SIMPLE_AUTOMATION_TEST(FMCPIdentityTest,"OSS.MCP.Test Identity", EAutomationTestFlags::ATF_Editor )
+IMPLEMENT_SIMPLE_AUTOMATION_TEST(FMCPIdentityTest,"System.OSS.MCP.Test Identity", EAutomationTestFlags::ATF_Editor )
 
 bool FMCPIdentityTest::RunTest(const FString& Parameters)
 {
@@ -186,7 +186,7 @@ bool FMCPIdentityTest::RunTest(const FString& Parameters)
 	return true;
 }
 
-IMPLEMENT_SIMPLE_AUTOMATION_TEST(FSteamIdentityTest,"OSS.Steam.Test Identity", EAutomationTestFlags::ATF_Editor )
+IMPLEMENT_SIMPLE_AUTOMATION_TEST(FSteamIdentityTest,"System.OSS.Steam.Test Identity", EAutomationTestFlags::ATF_Editor )
 
 bool FSteamIdentityTest::RunTest(const FString& Parameters)
 {
@@ -196,7 +196,7 @@ bool FSteamIdentityTest::RunTest(const FString& Parameters)
 	return true;
 }
 
-IMPLEMENT_SIMPLE_AUTOMATION_TEST(FAmazonIdentityTest,"OSS.Amazon.Test Identity", EAutomationTestFlags::ATF_Editor )
+IMPLEMENT_SIMPLE_AUTOMATION_TEST(FAmazonIdentityTest,"System.OSS.Amazon.Test Identity", EAutomationTestFlags::ATF_Editor )
 
 bool FAmazonIdentityTest::RunTest(const FString& Parameters)
 {
@@ -206,7 +206,7 @@ bool FAmazonIdentityTest::RunTest(const FString& Parameters)
 	return true;
 }
 
-IMPLEMENT_SIMPLE_AUTOMATION_TEST(FFacebookIdentityTest,"OSS.Facebook.Test Identity", EAutomationTestFlags::ATF_Editor )
+IMPLEMENT_SIMPLE_AUTOMATION_TEST(FFacebookIdentityTest,"System.OSS.Facebook.Test Identity", EAutomationTestFlags::ATF_Editor )
 
 bool FFacebookIdentityTest::RunTest(const FString& Parameters)
 {
@@ -216,7 +216,7 @@ bool FFacebookIdentityTest::RunTest(const FString& Parameters)
 	return true;
 }
 
-IMPLEMENT_SIMPLE_AUTOMATION_TEST(FIOSIdentityTest,"OSS.IOS.Test Identity", EAutomationTestFlags::ATF_Editor )
+IMPLEMENT_SIMPLE_AUTOMATION_TEST(FIOSIdentityTest,"System.OSS.IOS.Test Identity", EAutomationTestFlags::ATF_Editor )
 
 bool FIOSIdentityTest::RunTest(const FString& Parameters)
 {
@@ -230,7 +230,7 @@ bool FIOSIdentityTest::RunTest(const FString& Parameters)
 
 /////////////Currently crashing//////////////
 
-IMPLEMENT_SIMPLE_AUTOMATION_TEST(FNullIdentityTest,"OSS.NULL.Test Identity", EAutomationTestFlags::ATF_Editor )
+IMPLEMENT_SIMPLE_AUTOMATION_TEST(FNullIdentityTest,"System.OSS.NULL.Test Identity", EAutomationTestFlags::ATF_Editor )
 
 bool FNullIdentityTest::RunTest(const FString& Parameters)
 {
@@ -241,7 +241,7 @@ bool FNullIdentityTest::RunTest(const FString& Parameters)
 }
 */
 
-IMPLEMENT_SIMPLE_AUTOMATION_TEST(FMCPEntitlementsTest,"OSS.MCP.Test Entitlements", EAutomationTestFlags::ATF_Editor )
+IMPLEMENT_SIMPLE_AUTOMATION_TEST(FMCPEntitlementsTest,"System.OSS.MCP.Test Entitlements", EAutomationTestFlags::ATF_Editor )
 
 bool FMCPEntitlementsTest::RunTest(const FString& Parameters)
 {
@@ -251,7 +251,7 @@ bool FMCPEntitlementsTest::RunTest(const FString& Parameters)
 	return true;
 }
 
-IMPLEMENT_SIMPLE_AUTOMATION_TEST(FMCPTitleFileTest,"OSS.MCP.Test Title File", EAutomationTestFlags::ATF_Editor )
+IMPLEMENT_SIMPLE_AUTOMATION_TEST(FMCPTitleFileTest,"System.OSS.MCP.Test Title File", EAutomationTestFlags::ATF_Editor )
 
 bool FMCPTitleFileTest::RunTest(const FString& Parameters)
 {
@@ -261,7 +261,7 @@ bool FMCPTitleFileTest::RunTest(const FString& Parameters)
 	return true;
 }
 
-IMPLEMENT_SIMPLE_AUTOMATION_TEST(FSteamAchievementsTest,"OSS.Steam.Test Achievements", EAutomationTestFlags::ATF_Editor )
+IMPLEMENT_SIMPLE_AUTOMATION_TEST(FSteamAchievementsTest,"System.OSS.Steam.Test Achievements", EAutomationTestFlags::ATF_Editor )
 
 bool FSteamAchievementsTest::RunTest(const FString& Parameters)
 {
@@ -275,7 +275,7 @@ bool FSteamAchievementsTest::RunTest(const FString& Parameters)
 
 /////////////Currently crashing//////////////
 
-IMPLEMENT_SIMPLE_AUTOMATION_TEST(FNullAchievementsTest,"OSS.NULL.Test Achievements", EAutomationTestFlags::ATF_Editor )
+IMPLEMENT_SIMPLE_AUTOMATION_TEST(FNullAchievementsTest,"System.OSS.NULL.Test Achievements", EAutomationTestFlags::ATF_Editor )
 
 bool FNullAchievementsTest::RunTest(const FString& Parameters)
 {
diff --git a/Engine/Source/Runtime/Engine/Private/TimerManagerTests.cpp b/Engine/Source/Runtime/Engine/Private/TimerManagerTests.cpp
--- a/Engine/Source/Runtime/Engine/Private/TimerManagerTests.cpp
+++ b/Engine/Source/Runtime/Engine/Private/TimerManagerTests.cpp
@@ -3,7 +3,7 @@
 #include "EnginePrivate.h"
 #include "TimerManager.h"
 
-IMPLEMENT_SIMPLE_AUTOMATION_TEST(FTimerManagerTest, "Engine.TimerManager", EAutomationTestFlags::ATF_Editor)
+IMPLEMENT_SIMPLE_AUTOMATION_TEST(FTimerManagerTest, "System.Engine.TimerManager", EAutomationTestFlags::ATF_Editor)
 
 #define TIMER_TEST_TEXT( Format, ... ) FString::Printf(TEXT("%s - %d: %s"), TEXT(__FILE__) , __LINE__ , *FString::Printf(TEXT(Format), ##__VA_ARGS__) )
 
diff --git a/Engine/Source/Runtime/Engine/Private/VisualLogger/VisualLoggerAutomationTests.cpp b/Engine/Source/Runtime/Engine/Private/VisualLogger/VisualLoggerAutomationTests.cpp
--- a/Engine/Source/Runtime/Engine/Private/VisualLogger/VisualLoggerAutomationTests.cpp
+++ b/Engine/Source/Runtime/Engine/Private/VisualLogger/VisualLoggerAutomationTests.cpp
@@ -105,7 +105,7 @@ struct FTestDeviceContext
 	uint32 EngineDisableAILoggingFlag : 1;
 };
 
-IMPLEMENT_SIMPLE_AUTOMATION_TEST(FVisualLogTest, "Engine.VisualLogger.Logging simple text", EAutomationTestFlags::ATF_Game | EAutomationTestFlags::ATF_Editor)
+IMPLEMENT_SIMPLE_AUTOMATION_TEST(FVisualLogTest, "System.Engine.VisualLogger.Logging simple text", EAutomationTestFlags::ATF_Game | EAutomationTestFlags::ATF_Editor)
 /**
 *
 *
@@ -166,7 +166,7 @@ bool FVisualLogTest::RunTest(const FString& Parameters)
 	return true;
 }
 
-IMPLEMENT_SIMPLE_AUTOMATION_TEST(FVisualLogSegmentsTest, "Engine.VisualLogger.Logging segment shape", EAutomationTestFlags::ATF_Game | EAutomationTestFlags::ATF_Editor)
+IMPLEMENT_SIMPLE_AUTOMATION_TEST(FVisualLogSegmentsTest, "System.Engine.VisualLogger.Logging segment shape", EAutomationTestFlags::ATF_Game | EAutomationTestFlags::ATF_Editor)
 bool FVisualLogSegmentsTest::RunTest(const FString& Parameters)
 {
 	UWorld* World = GetSimpleEngineAutomationTestWorld(GetTestFlags());
@@ -226,7 +226,7 @@ DEFINE_VLOG_EVENT(EventTest2, All, "Second simple event for vlog tests")
 
 DEFINE_VLOG_EVENT(EventTest3, All, "Third simple event for vlog tests")
 
-IMPLEMENT_SIMPLE_AUTOMATION_TEST(FVisualLogEventsTest, "Engine.VisualLogger.Logging events", EAutomationTestFlags::ATF_Game | EAutomationTestFlags::ATF_Editor)
+IMPLEMENT_SIMPLE_AUTOMATION_TEST(FVisualLogEventsTest, "System.Engine.VisualLogger.Logging events", EAutomationTestFlags::ATF_Game | EAutomationTestFlags::ATF_Editor)
 bool FVisualLogEventsTest::RunTest(const FString& Parameters)
 {
 	UWorld* World = GetSimpleEngineAutomationTestWorld(GetTestFlags());
diff --git a/Engine/Source/Runtime/GameplayAbilities/Private/GameplayEffectTests.cpp b/Engine/Source/Runtime/GameplayAbilities/Private/GameplayEffectTests.cpp
--- a/Engine/Source/Runtime/GameplayAbilities/Private/GameplayEffectTests.cpp
+++ b/Engine/Source/Runtime/GameplayAbilities/Private/GameplayEffectTests.cpp
@@ -295,7 +295,7 @@ public:
 	virtual uint32 GetRequiredDeviceNum() const override { return 1; }
 
 protected:
-	virtual FString GetBeautifiedTestName() const override { return "AbilitySystem.GameplayEffects"; }
+	virtual FString GetBeautifiedTestName() const override { return "System.AbilitySystem.GameplayEffects"; }
 	virtual void GetTests(TArray<FString>& OutBeautifiedNames, TArray <FString>& OutTestCommands) const override
 	{
 		for (const FString& TestName : TestFunctionNames)
diff --git a/Engine/Source/Runtime/GameplayTags/Private/Tests/GameplayTagTests.cpp b/Engine/Source/Runtime/GameplayTags/Private/Tests/GameplayTagTests.cpp
--- a/Engine/Source/Runtime/GameplayTags/Private/Tests/GameplayTagTests.cpp
+++ b/Engine/Source/Runtime/GameplayTags/Private/Tests/GameplayTagTests.cpp
@@ -2,7 +2,7 @@
 
 #include "GameplayTagsModulePrivatePCH.h"
 
-IMPLEMENT_SIMPLE_AUTOMATION_TEST(FGameplayTagTest, "GameplayTags.GameplayTag", EAutomationTestFlags::ATF_Editor)
+IMPLEMENT_SIMPLE_AUTOMATION_TEST(FGameplayTagTest, "System.GameplayTags.GameplayTag", EAutomationTestFlags::ATF_Editor)
 
 #if WITH_EDITOR
 static UDataTable* CreateGameplayDataTable()
diff --git a/Engine/Source/Runtime/Networking/Private/Tests/IPv4AddressTest.cpp b/Engine/Source/Runtime/Networking/Private/Tests/IPv4AddressTest.cpp
--- a/Engine/Source/Runtime/Networking/Private/Tests/IPv4AddressTest.cpp
+++ b/Engine/Source/Runtime/Networking/Private/Tests/IPv4AddressTest.cpp
@@ -4,7 +4,7 @@
 #include "AutomationTest.h"
 
 
-IMPLEMENT_SIMPLE_AUTOMATION_TEST(FIPv4AddressTest, "Engine.Networking.IPv4.IPv4Address", EAutomationTestFlags::ATF_SmokeTest)
+IMPLEMENT_SIMPLE_AUTOMATION_TEST(FIPv4AddressTest, "System.Engine.Networking.IPv4.IPv4Address", EAutomationTestFlags::ATF_SmokeTest)
 
 
 bool FIPv4AddressTest::RunTest( const FString& Parameters )
diff --git a/Engine/Source/Runtime/Serialization/Private/Tests/StructSerializerTest.cpp b/Engine/Source/Runtime/Serialization/Private/Tests/StructSerializerTest.cpp
--- a/Engine/Source/Runtime/Serialization/Private/Tests/StructSerializerTest.cpp
+++ b/Engine/Source/Runtime/Serialization/Private/Tests/StructSerializerTest.cpp
@@ -72,7 +72,7 @@ namespace StructSerializerTest
 /* Tests
  *****************************************************************************/
 
-IMPLEMENT_SIMPLE_AUTOMATION_TEST(FJsonStructSerializerTest, "Core.Serialization.JsonStructSerializer", EAutomationTestFlags::ATF_Editor)
+IMPLEMENT_SIMPLE_AUTOMATION_TEST(FJsonStructSerializerTest, "System.Core.Serialization.JsonStructSerializer", EAutomationTestFlags::ATF_Editor)
 
 
 bool FJsonStructSerializerTest::RunTest( const FString& Parameters )
diff --git a/Engine/Source/Runtime/Slate/Private/Framework/Text/Tests/RichTextMarkupProcessingTest.cpp b/Engine/Source/Runtime/Slate/Private/Framework/Text/Tests/RichTextMarkupProcessingTest.cpp
--- a/Engine/Source/Runtime/Slate/Private/Framework/Text/Tests/RichTextMarkupProcessingTest.cpp
+++ b/Engine/Source/Runtime/Slate/Private/Framework/Text/Tests/RichTextMarkupProcessingTest.cpp
@@ -13,7 +13,7 @@ PRAGMA_DISABLE_OPTIMIZATION
 
 #define LOCTEXT_NAMESPACE "Slate.Tests.RickTextMarkupProcessingTest"
 
-IMPLEMENT_SIMPLE_AUTOMATION_TEST(FRTFProcessingTest, "Slate.Misc.RickTextMarkupProcessing", EAutomationTestFlags::ATF_SmokeTest)
+IMPLEMENT_SIMPLE_AUTOMATION_TEST(FRTFProcessingTest, "System.Slate.Misc.RickTextMarkupProcessing", EAutomationTestFlags::ATF_SmokeTest)
 
 namespace
 {
diff --git a/Engine/Source/Runtime/VectorVM/Private/VectorVM.cpp b/Engine/Source/Runtime/VectorVM/Private/VectorVM.cpp
--- a/Engine/Source/Runtime/VectorVM/Private/VectorVM.cpp
+++ b/Engine/Source/Runtime/VectorVM/Private/VectorVM.cpp
@@ -988,7 +988,7 @@ uint8 VectorVM::CreateSrcOperandMask(VectorVM::EOperandType Type1, VectorVM::EOp
 	Automation test for the VM.
 ------------------------------------------------------------------------------*/
 
-IMPLEMENT_SIMPLE_AUTOMATION_TEST(FVectorVMTest, "Core.Math.Vector VM", EAutomationTestFlags::ATF_SmokeTest)
+IMPLEMENT_SIMPLE_AUTOMATION_TEST(FVectorVMTest, "System.Core.Math.Vector VM", EAutomationTestFlags::ATF_SmokeTest)
 
 bool FVectorVMTest::RunTest(const FString& Parameters)
 {
