Creating a Blueprint based off another will no longer cause a transaction for an "Add Node".

Changed the methodology of creating the call to parent for the construction graph.

diff --git a/Engine/Source/Editor/UnrealEd/Private/Kismet2/Kismet2.cpp b/Engine/Source/Editor/UnrealEd/Private/Kismet2/Kismet2.cpp
--- a/Engine/Source/Editor/UnrealEd/Private/Kismet2/Kismet2.cpp
+++ b/Engine/Source/Editor/UnrealEd/Private/Kismet2/Kismet2.cpp
@@ -415,13 +415,17 @@ UBlueprint* FKismetEditorUtilities::CreateBlueprint(UClass* ParentClass, UObject
 		{
 			check( UCSGraph->Nodes.Num() > 0 );
 			UK2Node_FunctionEntry* UCSEntry = CastChecked<UK2Node_FunctionEntry>(UCSGraph->Nodes[0]);
-			UK2Node_CallParentFunction* ParentCallNodeTemplate = NewObject<UK2Node_CallParentFunction>();
-			ParentCallNodeTemplate->FunctionReference.SetExternalMember(K2Schema->FN_UserConstructionScript, NewBP->ParentClass);
-			UK2Node_CallParentFunction* ParentCallNode = FEdGraphSchemaAction_K2NewNode::SpawnNodeFromTemplate<UK2Node_CallParentFunction>(UCSGraph, ParentCallNodeTemplate, FVector2D(200, 0));
+			FGraphNodeCreator<UK2Node_CallParentFunction> FunctionNodeCreator(*UCSGraph);
+			UK2Node_CallParentFunction* ParentFunctionNode = FunctionNodeCreator.CreateNode();
+			ParentFunctionNode->FunctionReference.SetExternalMember(K2Schema->FN_UserConstructionScript, NewBP->ParentClass);
+			ParentFunctionNode->NodePosX = 200;
+			ParentFunctionNode->NodePosY = 0;
+			ParentFunctionNode->AllocateDefaultPins();
+			FunctionNodeCreator.Finalize();
 
 			// Wire up the new node
 			UEdGraphPin* ExecPin = UCSEntry->FindPin(K2Schema->PN_Then);
-			UEdGraphPin* SuperPin = ParentCallNode->FindPin(K2Schema->PN_Execute);
+			UEdGraphPin* SuperPin = ParentFunctionNode->FindPin(K2Schema->PN_Execute);
 			ExecPin->MakeLinkTo(SuperPin);
 		}
 
