BP: Crash fixed, when mouse pointer hovers over a pin in Construction Script.

I couldn't reproduce UE-12022 after the change.

diff --git a/Engine/Source/Editor/UnrealEd/Private/Kismet2/KismetDebugUtilities.cpp b/Engine/Source/Editor/UnrealEd/Private/Kismet2/KismetDebugUtilities.cpp
--- a/Engine/Source/Editor/UnrealEd/Private/Kismet2/KismetDebugUtilities.cpp
+++ b/Engine/Source/Editor/UnrealEd/Private/Kismet2/KismetDebugUtilities.cpp
@@ -923,7 +923,7 @@ FKismetDebugUtilities::EWatchTextResult FKismetDebugUtilities::GetWatchText(FStr
 			{
 				if(UBlueprintGeneratedClass* BPGC = Cast<UBlueprintGeneratedClass>(Blueprint->GeneratedClass))
 				{
-					PropertyBase = BPGC->GetPersistentUberGraphFrame(ActiveObject, BPGC->UberGraphFunction);
+					PropertyBase = BPGC->GetPersistentUberGraphFrame(ActiveObject, OuterFunction);
 				}
 			}
 #endif // USE_UBER_GRAPH_PERSISTENT_FRAME
