[UE-13528] RenameTimeline() shouldn't be using the old name ref to compare names against (in case it alters the ref's value as it iterates over nodes that need to be renamed).

diff --git a/Engine/Source/Editor/UnrealEd/Private/Kismet2/BlueprintEditorUtils.cpp b/Engine/Source/Editor/UnrealEd/Private/Kismet2/BlueprintEditorUtils.cpp
--- a/Engine/Source/Editor/UnrealEd/Private/Kismet2/BlueprintEditorUtils.cpp
+++ b/Engine/Source/Editor/UnrealEd/Private/Kismet2/BlueprintEditorUtils.cpp
@@ -6018,12 +6018,15 @@ UK2Node_Timeline* FBlueprintEditorUtils::FindNodeForTimeline(UBlueprint* Bluepri
 	return NULL; // no node found
 }
 
-bool FBlueprintEditorUtils::RenameTimeline(UBlueprint* Blueprint, const FName& OldName, const FName& NewName)
+bool FBlueprintEditorUtils::RenameTimeline(UBlueprint* Blueprint, const FName& OldNameRef, const FName& NewName)
 {
 	check(Blueprint);
 
 	bool bRenamed = false;
 
+	// make a copy, in case we alter the value of what is referenced by 
+	// OldNameRef through the course of this function
+	FName OldName = OldNameRef;
 
 	TSharedPtr<INameValidatorInterface> NameValidator = MakeShareable(new FKismetNameValidator(Blueprint));
 	const FString NewTemplateName = UTimelineTemplate::TimelineVariableNameToTemplateName(NewName);
