UE-12886 Crashes on Structure including Same class blueprint within Construction Script

flag bFirstCompileChildrenThenReinstance was removed.

Problem was: actor cannot be reconstructed before all objects have new/proper class set.
Fast reinstancing was split into 2 phases in the first phase the new class is set. In the 2nd actors are reconstructed, etc..

diff --git a/Engine/Config/BaseEngine.ini b/Engine/Config/BaseEngine.ini
--- a/Engine/Config/BaseEngine.ini
+++ b/Engine/Config/BaseEngine.ini
@@ -1545,7 +1545,6 @@ bOptimizeAdjacentStates=true
 bEnableInheritableComponents=true
 bDeferDependencyLoads=true
 bForceDisableCookedDependencyDeferring=false
-bFirstCompileChildrenThenReinstance=true
 bExecutionAfterReturn=false
 
 [/Script/Engine.Blueprint]
diff --git a/Engine/Source/Editor/UnrealEd/Private/Kismet2/Kismet2.cpp b/Engine/Source/Editor/UnrealEd/Private/Kismet2/Kismet2.cpp
--- a/Engine/Source/Editor/UnrealEd/Private/Kismet2/Kismet2.cpp
+++ b/Engine/Source/Editor/UnrealEd/Private/Kismet2/Kismet2.cpp
@@ -752,9 +752,7 @@ void FKismetEditorUtilities::CompileBlueprint(UBlueprint* BlueprintObj, bool bIs
 	ReinstanceHelper->UpdateBytecodeReferences();
 
 	const bool bIsInterface = FBlueprintEditorUtils::IsInterfaceBlueprint(BlueprintObj);
-	static const FBoolConfigValueHelper FirstCompileChildrenThenReinstance(TEXT("Kismet"), TEXT("bFirstCompileChildrenThenReinstance"), GEngineIni);
-	const bool bLetReinstancerRefreshDependBP = !bIsRegeneratingOnLoad && (OldClass != NULL)
-		&& FirstCompileChildrenThenReinstance && !bIsInterface;
+	const bool bLetReinstancerRefreshDependBP = !bIsRegeneratingOnLoad && (OldClass != NULL) && !bIsInterface;
 	if (bLetReinstancerRefreshDependBP)
 	{
 		BP_SCOPED_COMPILER_EVENT_STAT(EKismetCompilerStats_RefreshDependentBlueprints);
diff --git a/Engine/Source/Editor/UnrealEd/Private/Kismet2/KismetReinstanceUtilities.cpp b/Engine/Source/Editor/UnrealEd/Private/Kismet2/KismetReinstanceUtilities.cpp
--- a/Engine/Source/Editor/UnrealEd/Private/Kismet2/KismetReinstanceUtilities.cpp
+++ b/Engine/Source/Editor/UnrealEd/Private/Kismet2/KismetReinstanceUtilities.cpp
@@ -289,39 +289,48 @@ FBlueprintCompileReinstancer::~FBlueprintCompileReinstancer()
 {
 }
 
-void FBlueprintCompileReinstancer::ReinstanceFast()
+class FReinstanceFinalizer : public TSharedFromThis<FReinstanceFinalizer>
 {
-	UE_LOG(LogBlueprint, Log, TEXT("BlueprintCompileReinstancer: Doing a fast path refresh on class '%s'."), *GetPathNameSafe(ClassToReinstance));
-
+public:
+	TSharedPtr<FBlueprintCompileReinstancer> Reinstancer;
 	TArray<UObject*> ObjectsToReplace;
-	GetObjectsOfClass(DuplicatedClass, ObjectsToReplace, /*bIncludeDerivedClasses=*/ false);
+	TArray<UObject*> ObjectsToFinalize;
+	TSet<UObject*> SelectedObjecs;
+	UClass* ClassToReinstance;
 
-	const bool bIsActor = ClassToReinstance->IsChildOf<AActor>();
-	const bool bIsAnimInstance = ClassToReinstance->IsChildOf<UAnimInstance>();
-	const bool bIsComponent = ClassToReinstance->IsChildOf<UActorComponent>();
-	for (auto Obj : ObjectsToReplace)
+	FReinstanceFinalizer(UClass* InClassToReinstance) : ClassToReinstance(InClassToReinstance)
 	{
-		UE_LOG(LogBlueprint, Log, TEXT("  Fast path is refreshing (not replacing) %s"), *Obj->GetFullName());
+		check(ClassToReinstance);
+	}
 
-		if ((!Obj->IsTemplate() || bIsComponent) && !Obj->IsPendingKill())
+	void Finalize()
+	{
+		if (!ensure(Reinstancer.IsValid()))
 		{
-			const bool bIsSelected = bIsActor ? Obj->IsSelected() : false;
-
-			Obj->SetClass(ClassToReinstance);
+			return;
+		}
+		check(ClassToReinstance);
 
-			if (bIsActor)
+		const bool bIsActor = ClassToReinstance->IsChildOf<AActor>();
+		if (bIsActor)
+		{
+			for (auto Obj : ObjectsToFinalize)
 			{
 				auto Actor = CastChecked<AActor>(Obj);
 				Actor->ReregisterAllComponents();
 				Actor->RerunConstructionScripts();
 
-				if (bIsSelected)
+				if (SelectedObjecs.Contains(Obj))
 				{
 					GEditor->SelectActor(Actor, /*bInSelected =*/true, /*bNotify =*/true, false, true);
 				}
 			}
+		}
 
-			if (bIsAnimInstance)
+		const bool bIsAnimInstance = ClassToReinstance->IsChildOf<UAnimInstance>();
+		if (bIsAnimInstance)
+		{
+			for (auto Obj : ObjectsToFinalize)
 			{
 				// Initialising the anim instance isn't enough to correctly set up the skeletal mesh again in a
 				// paused world, need to initialise the skeletal mesh component that contains the anim instance.
@@ -331,8 +340,45 @@ void FBlueprintCompileReinstancer::ReinstanceFast()
 				}
 			}
 		}
+
+		Reinstancer->FinalizeFastReinstancing(ObjectsToReplace);
 	}
+};
+
+TSharedPtr<FReinstanceFinalizer> FBlueprintCompileReinstancer::ReinstanceFast()
+{
+	UE_LOG(LogBlueprint, Log, TEXT("BlueprintCompileReinstancer: Doing a fast path refresh on class '%s'."), *GetPathNameSafe(ClassToReinstance));
+
+	TSharedPtr<FReinstanceFinalizer> Finalizer = MakeShareable(new FReinstanceFinalizer(ClassToReinstance));
+	Finalizer->Reinstancer = SharedThis(this);
+
+	GetObjectsOfClass(DuplicatedClass, Finalizer->ObjectsToReplace, /*bIncludeDerivedClasses=*/ false);
+
+	const bool bIsActor = ClassToReinstance->IsChildOf<AActor>();
+	const bool bIsAnimInstance = ClassToReinstance->IsChildOf<UAnimInstance>();
+	const bool bIsComponent = ClassToReinstance->IsChildOf<UActorComponent>();
+	for (auto Obj : Finalizer->ObjectsToReplace)
+	{
+		UE_LOG(LogBlueprint, Log, TEXT("  Fast path is refreshing (not replacing) %s"), *Obj->GetFullName());
+
+		if ((!Obj->IsTemplate() || bIsComponent) && !Obj->IsPendingKill())
+		{
+			if (bIsActor && Obj->IsSelected())
+			{
+				Finalizer->SelectedObjecs.Add(Obj);
+			}
+
+			Obj->SetClass(ClassToReinstance);
+
+			Finalizer->ObjectsToFinalize.Push(Obj);
+		}
+	}
+
+	return Finalizer;
+}
 
+void FBlueprintCompileReinstancer::FinalizeFastReinstancing(TArray<UObject*>& ObjectsToReplace)
+{
 	TArray<UObject*> SourceObjects;
 	TMap<UObject*, UObject*> OldToNewInstanceMap;
 	TMap<FStringAssetReference, UObject*> ReinstancedObjectsWeakReferenceMap;
@@ -344,6 +390,12 @@ void FBlueprintCompileReinstancer::ReinstanceFast()
 	}
 
 	FReplaceReferenceHelper::FindAndReplaceReferences(SourceObjects, &ObjectsThatShouldUseOldStuff, ObjectsToReplace, OldToNewInstanceMap, ReinstancedObjectsWeakReferenceMap);
+
+	if (ClassToReinstance->IsChildOf<UActorComponent>())
+	{
+		// ReplaceInstancesOfClass() handles this itself, if we had to re-instance
+		ReconstructOwnerInstances(ClassToReinstance);
+	}
 }
 
 void FBlueprintCompileReinstancer::CompileChildren()
@@ -363,8 +415,9 @@ void FBlueprintCompileReinstancer::CompileChildren()
 	}
 }
 
-void FBlueprintCompileReinstancer::ReinstanceInner(bool bForceAlwaysReinstance)
+TSharedPtr<FReinstanceFinalizer> FBlueprintCompileReinstancer::ReinstanceInner(bool bForceAlwaysReinstance)
 {
+	TSharedPtr<FReinstanceFinalizer> Finalizer;
 	if (ClassToReinstance && DuplicatedClass)
 	{
 		static const FBoolConfigValueHelper ReinstanceOnlyWhenNecessary(TEXT("Kismet"), TEXT("bReinstanceOnlyWhenNecessary"), GEngineIni);
@@ -383,7 +436,7 @@ void FBlueprintCompileReinstancer::ReinstanceInner(bool bForceAlwaysReinstance)
 			const bool bAllowedToDoFastPath = bTheSameDefaultValues && bTheSameLayout;
 			if (bAllowedToDoFastPath)
 			{
-				ReinstanceFast();
+				Finalizer = ReinstanceFast();
 				bShouldReinstance = false;
 			}
 		}
@@ -393,18 +446,12 @@ void FBlueprintCompileReinstancer::ReinstanceInner(bool bForceAlwaysReinstance)
 			UE_LOG(LogBlueprint, Log, TEXT("BlueprintCompileReinstancer: Doing a full reinstance on class '%s'"), *GetPathNameSafe(ClassToReinstance));
 			ReplaceInstancesOfClass(DuplicatedClass, ClassToReinstance, OriginalCDO, &ObjectsThatShouldUseOldStuff, IsClassObjectReplaced(), ShouldPreserveRootComponentOfReinstancedActor());
 		}
-		else if (ClassToReinstance->IsChildOf<UActorComponent>())
-		{
-			// ReplaceInstancesOfClass() handles this itself, if we had to re-instance
-			ReconstructOwnerInstances(ClassToReinstance);
-		}
 	}
+	return Finalizer;
 }
 
 void FBlueprintCompileReinstancer::ListDependentBlueprintsToRefresh(const TArray<UBlueprint*>& DependentBPs)
 {
-	static const FBoolConfigValueHelper FirstCompileChildrenThenReinstance(TEXT("Kismet"), TEXT("bFirstCompileChildrenThenReinstance"), GEngineIni);
-	check(FirstCompileChildrenThenReinstance);
 	for (auto& Element : DependentBPs)
 	{
 		DependentBlueprintsToRefresh.Add(Element);
@@ -415,70 +462,63 @@ void FBlueprintCompileReinstancer::ReinstanceObjects(bool bForceAlwaysReinstance
 {
 	BP_SCOPED_COMPILER_EVENT_STAT(EKismetCompilerStats_ReinstanceObjects);
 	
-	static const FBoolConfigValueHelper FirstCompileChildrenThenReinstance(TEXT("Kismet"), TEXT("bFirstCompileChildrenThenReinstance"), GEngineIni);
-
-	if (FirstCompileChildrenThenReinstance)
+	// Make sure we only reinstance classes once!
+	static TArray<TSharedRef<FBlueprintCompileReinstancer>> QueueToReinstance;
+	TSharedRef<FBlueprintCompileReinstancer> SharedThis = AsShared();
+	const bool bAlreadyQueued = QueueToReinstance.Contains(SharedThis);
+	if (!bAlreadyQueued && !bHasReinstanced)
 	{
-		// Make sure we only reinstance classes once!
-		static TArray<TSharedRef<FBlueprintCompileReinstancer>> QueueToReinstance;
-		TSharedRef<FBlueprintCompileReinstancer> SharedThis = AsShared();
-		const bool bAlreadyQueued = QueueToReinstance.Contains(SharedThis);
-		if (!bAlreadyQueued && !bHasReinstanced)
+		QueueToReinstance.Push(SharedThis);
+
+		if (ClassToReinstance && DuplicatedClass)
 		{
-			QueueToReinstance.Push(SharedThis);
+			CompileChildren();
+		}
 
-			if (ClassToReinstance && DuplicatedClass)
+		if (QueueToReinstance.Num() && (QueueToReinstance[0] == SharedThis))
+		{
+			TArray<TSharedPtr<FReinstanceFinalizer>> Finalizers;
+
+			// All children were recompiled. It's safe to reinstance.
+			for (int32 Idx = 0; Idx < QueueToReinstance.Num(); ++Idx)
 			{
-				CompileChildren();
+				auto Finalizer = QueueToReinstance[Idx]->ReinstanceInner(bForceAlwaysReinstance);
+				if (Finalizer.IsValid())
+				{
+					Finalizers.Push(Finalizer);
+				}
+				QueueToReinstance[Idx]->bHasReinstanced = true;
+
+				auto LocalClassToReinstance = QueueToReinstance[Idx]->ClassToReinstance;
+				UBlueprint* BP = LocalClassToReinstance ? Cast<UBlueprint>(LocalClassToReinstance->ClassGeneratedBy) : nullptr;
+				if (BP)
+				{
+					DependentBlueprintsToRefresh.Remove(BP);
+				}
 			}
+			QueueToReinstance.Empty();
 
-			if (QueueToReinstance.Num() && (QueueToReinstance[0] == SharedThis))
+			for (auto Finalizer : Finalizers)
 			{
+				if (Finalizer.IsValid())
 				{
-					BP_SCOPED_COMPILER_EVENT_STAT(EKismetCompilerStats_RefreshDependentBlueprintsInReinstancer);
-					for (auto Reinst : QueueToReinstance)
-					{
-						UBlueprint* BP = Reinst->ClassToReinstance ? Cast<UBlueprint>(Reinst->ClassToReinstance->ClassGeneratedBy) : nullptr;
-						if (BP)
-						{
-							DependentBlueprintsToRefresh.Remove(BP);
-						}
-					}
-					for (auto BPPtr : DependentBlueprintsToRefresh)
-					{
-						if (BPPtr.IsValid())
-						{
-							BPPtr->BroadcastChanged();
-						}
-					}
-					DependentBlueprintsToRefresh.Empty();
+					Finalizer->Finalize();
 				}
+			}
 
-				// All children were recompiled. It's safe to reinstance.
-				for (int32 Idx = 0; Idx < QueueToReinstance.Num(); ++Idx)
+			{
+				BP_SCOPED_COMPILER_EVENT_STAT(EKismetCompilerStats_RefreshDependentBlueprintsInReinstancer);
+				for (auto BPPtr : DependentBlueprintsToRefresh)
 				{
-					QueueToReinstance[Idx]->ReinstanceInner(bForceAlwaysReinstance);
-					QueueToReinstance[Idx]->bHasReinstanced = true;
+					if (BPPtr.IsValid())
+					{
+						BPPtr->BroadcastChanged();
+					}
 				}
-				QueueToReinstance.Empty();
+				DependentBlueprintsToRefresh.Empty();
 			}
 		}
 	}
-	else
-	{
-		//THE OLD WAY
-		if (bHasReinstanced)
-		{
-			return;
-		}
-		bHasReinstanced = true;
-
-		if (ClassToReinstance && DuplicatedClass)
-		{
-			ReinstanceInner(bForceAlwaysReinstance);
-			CompileChildren();
-		}
-	}
 }
 
 
diff --git a/Engine/Source/Editor/UnrealEd/Public/Kismet2/KismetReinstanceUtilities.h b/Engine/Source/Editor/UnrealEd/Public/Kismet2/KismetReinstanceUtilities.h
--- a/Engine/Source/Editor/UnrealEd/Public/Kismet2/KismetReinstanceUtilities.h
+++ b/Engine/Source/Editor/UnrealEd/Public/Kismet2/KismetReinstanceUtilities.h
@@ -10,6 +10,8 @@
 
 DECLARE_STATS_GROUP(TEXT("Kismet Reinstancer"), STATGROUP_KismetReinstancer, STATCAT_Advanced);
 
+class FReinstanceFinalizer;
+
 class UNREALED_API FBlueprintCompileReinstancer : public TSharedFromThis<FBlueprintCompileReinstancer>, public FGCObject
 {
 protected:
@@ -92,10 +94,12 @@ public:
 
 	virtual bool IsClassObjectReplaced() const { return false; }
 
+	void FinalizeFastReinstancing(TArray<UObject*>& ObjectsToReplace);
 protected:
-	void ReinstanceInner(bool bForceAlwaysReinstance);
 
-	void ReinstanceFast();
+	TSharedPtr<FReinstanceFinalizer> ReinstanceInner(bool bForceAlwaysReinstance);
+
+	TSharedPtr<FReinstanceFinalizer> ReinstanceFast();
 
 	void CompileChildren();
 
