Editor Modes: Added support for displaying properties as widgets for selected components
- Changed FEdMode widget property display code to prefer a selected component
- Removed some duplicated code in various places in FEdMode
- Changed the widget diamond for transform properties to draw in the correct orientation
- Cleaned up various GEditor references that had crept back into FEdMode
- Added FEditorModeTools::GetSelectedComponents
- Initialized some unused variables in FEditorModeTools
- Some minor coding standards tweaks, implicit int to bools, etc...
[UE-12942] [UE-12943]

diff --git a/Engine/Source/Editor/UnrealEd/Private/AssetEditorModeManager.cpp b/Engine/Source/Editor/UnrealEd/Private/AssetEditorModeManager.cpp
--- a/Engine/Source/Editor/UnrealEd/Private/AssetEditorModeManager.cpp
+++ b/Engine/Source/Editor/UnrealEd/Private/AssetEditorModeManager.cpp
@@ -18,6 +18,10 @@ FAssetEditorModeManager::FAssetEditorModeManager()
 	ObjectSet = NewObject<USelection>();
 	ObjectSet->SetFlags(RF_Transactional);
 	ObjectSet->AddToRoot();
+
+	ComponentSet = NewObject<USelection>();
+	ComponentSet->SetFlags(RF_Transactional);
+	ComponentSet->AddToRoot();
 }
 
 FAssetEditorModeManager::~FAssetEditorModeManager()
@@ -26,6 +30,8 @@ FAssetEditorModeManager::~FAssetEditorModeManager()
 	ActorSet = nullptr;
 	ObjectSet->RemoveFromRoot();
 	ObjectSet = nullptr;
+	ComponentSet->RemoveFromRoot();
+	ComponentSet = nullptr;
 }
 
 USelection* FAssetEditorModeManager::GetSelectedActors() const
@@ -38,6 +44,11 @@ USelection* FAssetEditorModeManager::GetSelectedObjects() const
 	return ObjectSet;
 }
 
+USelection* FAssetEditorModeManager::GetSelectedComponents() const
+{
+	return ComponentSet;
+}
+
 UWorld* FAssetEditorModeManager::GetWorld() const
 {
 	return (PreviewScene != nullptr) ? PreviewScene->GetWorld() : GEditor->GetEditorWorldContext().World();
diff --git a/Engine/Source/Editor/UnrealEd/Private/EditorModeManager.cpp b/Engine/Source/Editor/UnrealEd/Private/EditorModeManager.cpp
--- a/Engine/Source/Editor/UnrealEd/Private/EditorModeManager.cpp
+++ b/Engine/Source/Editor/UnrealEd/Private/EditorModeManager.cpp
@@ -33,18 +33,22 @@
 ------------------------------------------------------------------------------*/
 
 FEditorModeTools::FEditorModeTools()
-	:	PivotShown( 0 )
-	,	Snapping( 0 )
-	,	SnappedActor( 0 )
-	,	TranslateRotateXAxisAngle(0)
-	,	TranslateRotate2DAngle(0)
-	,	DefaultID(FBuiltinEditorModes::EM_Default)
-	,	WidgetMode( FWidget::WM_None )
-	,	OverrideWidgetMode( FWidget::WM_None )
-	,	bShowWidget( 1 )
-	,	bHideViewportUI(false)
-	,	CoordSystem(COORD_World)
-	,	bIsTracking(false)
+	: PivotShown(false)
+	, Snapping(false)
+	, SnappedActor(false)
+	, CachedLocation(ForceInitToZero)
+	, PivotLocation(ForceInitToZero)
+	, SnappedLocation(ForceInitToZero)
+	, GridBase(ForceInitToZero)
+	, TranslateRotateXAxisAngle(0.0f)
+	, TranslateRotate2DAngle(0.0f)
+	, DefaultID(FBuiltinEditorModes::EM_Default)
+	, WidgetMode(FWidget::WM_None)
+	, OverrideWidgetMode(FWidget::WM_None)
+	, bShowWidget(true)
+	, bHideViewportUI(false)
+	, CoordSystem(COORD_World)
+	, bIsTracking(false)
 {
 	// Load the last used settings
 	LoadConfig();
@@ -71,9 +75,6 @@ FEditorModeTools::~FEditorModeTools()
 	GEditor->UnregisterForUndo(this);
 }
 
-/**
- * Loads the state that was saved in the INI file
- */
 void FEditorModeTools::LoadConfig(void)
 {
 	GConfig->GetBool(TEXT("FEditorModeTools"),TEXT("ShowWidget"),bShowWidget,
@@ -89,17 +90,12 @@ void FEditorModeTools::LoadConfig(void)
 	LoadWidgetSettings();
 }
 
-/**
- * Saves the current state to the INI file
- */
 void FEditorModeTools::SaveConfig(void)
 {
-	GConfig->SetBool(TEXT("FEditorModeTools"),TEXT("ShowWidget"),bShowWidget,
-		GEditorUserSettingsIni);
+	GConfig->SetBool(TEXT("FEditorModeTools"), TEXT("ShowWidget"), bShowWidget, GEditorUserSettingsIni);
 
 	const bool bGetRawValue = true;
-	GConfig->SetInt(TEXT("FEditorModeTools"),TEXT("CoordSystem"),(int32)GetCoordSystem(bGetRawValue),
-		GEditorUserSettingsIni);
+	GConfig->SetInt(TEXT("FEditorModeTools"), TEXT("CoordSystem"), (int32)GetCoordSystem(bGetRawValue), GEditorUserSettingsIni);
 
 	SaveWidgetSettings();
 }
@@ -117,16 +113,21 @@ void FEditorModeTools::SetToolkitHost(TSharedRef<class IToolkitHost> InHost)
 	ToolkitHost = InHost;
 }
 
-class USelection* FEditorModeTools::GetSelectedActors() const
+USelection* FEditorModeTools::GetSelectedActors() const
 {
 	return GEditor->GetSelectedActors();
 }
 
-class USelection* FEditorModeTools::GetSelectedObjects() const
+USelection* FEditorModeTools::GetSelectedObjects() const
 {
 	return GEditor->GetSelectedObjects();
 }
 
+USelection* FEditorModeTools::GetSelectedComponents() const
+{
+	return GEditor->GetSelectedComponents();
+}
+
 UWorld* FEditorModeTools::GetWorld() const
 {
 	return GEditor->GetEditorWorldContext().World();
@@ -136,7 +137,7 @@ void FEditorModeTools::OnEditorSelectionChanged(UObject* NewSelection)
 {
 	// If selecting an actor, move the pivot location.
 	AActor* Actor = Cast<AActor>(NewSelection);
-	if (Actor != NULL)
+	if (Actor != nullptr)
 	{
 		//@fixme - why isn't this using UObject::IsSelected()?
 		if ( GEditor->GetSelectedActors()->IsSelected( Actor ) )
@@ -172,12 +173,6 @@ void FEditorModeTools::OnEditorSelectNone()
 	GEditor->ActorsThatWereSelected.Empty();
 }
 
-/** 
- * Sets the pivot locations
- * 
- * @param Location 		The location to set
- * @param bIncGridBase	Whether or not to also set the GridBase
- */
 void FEditorModeTools::SetPivotLocation( const FVector& Location, const bool bIncGridBase )
 {
 	CachedLocation = PivotLocation = SnappedLocation = Location;
@@ -189,7 +184,7 @@ void FEditorModeTools::SetPivotLocation( const FVector& Location, const bool bIn
 
 ECoordSystem FEditorModeTools::GetCoordSystem(bool bGetRawValue)
 {
-	if (!bGetRawValue && GetWidgetMode() == FWidget::WM_Scale )
+	if (!bGetRawValue && (GetWidgetMode() == FWidget::WM_Scale))
 	{
 		return COORD_Local;
 	}
@@ -227,11 +222,6 @@ void FEditorModeTools::DeactivateModeAtIndex(int32 InIndex)
 	Modes.RemoveAt( InIndex );
 }
 
-/**
- * Deactivates an editor mode. 
- * 
- * @param InID		The ID of the editor mode to deactivate.
- */
 void FEditorModeTools::DeactivateMode( FEditorModeID InID )
 {
 	// Find the mode from the ID and exit it.
@@ -277,13 +267,7 @@ void FEditorModeTools::DestroyMode( FEditorModeID InID )
 	RecycledModes.Remove(InID);
 }
 
-/**
- * Activates an editor mode. Shuts down all other active modes which cannot run with the passed in mode.
- * 
- * @param InID		The ID of the editor mode to activate.
- * @param bToggle	true if the passed in editor mode should be toggled off if it is already active.
- */
-void FEditorModeTools::ActivateMode( FEditorModeID InID, bool bToggle )
+void FEditorModeTools::ActivateMode(FEditorModeID InID, bool bToggle)
 {
 	if (InID == FBuiltinEditorModes::EM_Default)
 	{
@@ -291,51 +275,51 @@ void FEditorModeTools::ActivateMode( FEditorModeID InID, bool bToggle )
 	}
 
 	// Check to see if the mode is already active
-	if( IsModeActive(InID) )
+	if (IsModeActive(InID))
 	{
 		// The mode is already active toggle it off if we should toggle off already active modes.
-		if( bToggle )
+		if (bToggle)
 		{
-			DeactivateMode( InID );
+			DeactivateMode(InID);
 		}
 		// Nothing more to do
 		return;
 	}
 
 	// Recycle a mode or factory a new one
-	TSharedPtr<FEdMode> Mode = RecycledModes.FindRef( InID );
+	TSharedPtr<FEdMode> Mode = RecycledModes.FindRef(InID);
 
-	if ( Mode.IsValid() )
+	if (Mode.IsValid())
 	{
-		RecycledModes.Remove( InID );
+		RecycledModes.Remove(InID);
 	}
 	else
 	{
-		Mode = FEditorModeRegistry::Get().CreateMode( InID, *this );
+		Mode = FEditorModeRegistry::Get().CreateMode(InID, *this);
 	}
 
-	if( !Mode.IsValid() )
+	if (!Mode.IsValid())
 	{
-		UE_LOG(LogEditorModes, Log, TEXT("FEditorModeTools::ActivateMode : Couldn't find mode '%s'."), *InID.ToString() );
+		UE_LOG(LogEditorModes, Log, TEXT("FEditorModeTools::ActivateMode : Couldn't find mode '%s'."), *InID.ToString());
 		// Just return and leave the mode list unmodified
 		return;
 	}
 
 	// Remove anything that isn't compatible with this mode
-	for( int32 ModeIndex = Modes.Num() - 1; ModeIndex >= 0; --ModeIndex )
+	for (int32 ModeIndex = Modes.Num() - 1; ModeIndex >= 0; --ModeIndex)
 	{
-		const bool bModesAreCompatible = Mode->IsCompatibleWith( Modes[ModeIndex]->GetID() ) || Modes[ModeIndex]->IsCompatibleWith( Mode->GetID() );
-		if ( !bModesAreCompatible )
+		const bool bModesAreCompatible = Mode->IsCompatibleWith(Modes[ModeIndex]->GetID()) || Modes[ModeIndex]->IsCompatibleWith(Mode->GetID());
+		if (!bModesAreCompatible)
 		{
 			DeactivateModeAtIndex(ModeIndex);
 		}
 	}
 
-	Modes.Add( Mode );
+	Modes.Add(Mode);
 
 	// Enter the new mode
 	Mode->Enter();
-	
+
 	// Update the editor UI
 	FEditorSupportDelegates::UpdateUI.Broadcast();
 }
@@ -370,7 +354,7 @@ FEdMode* FEditorModeTools::FindMode( FEditorModeID InID )
 		}
 	}
 
-	return NULL;
+	return nullptr;
 }
 
 /**
@@ -1014,7 +998,7 @@ void FEditorModeTools::ClearBookmark( uint32 InIndex, FEditorViewportClient* InV
 		// Verify the index is valid for the bookmark
 		if ( pWorldSettings && InIndex < AWorldSettings::MAX_BOOKMARK_NUMBER )
 		{
-			pWorldSettings->BookMarks[ InIndex ] = NULL;
+			pWorldSettings->BookMarks[ InIndex ] = nullptr;
 		}
 	}
 }
@@ -1030,10 +1014,6 @@ void FEditorModeTools::ClearAllBookmarks( FEditorViewportClient* InViewportClien
 	}
 }
 
-/**
- * Serializes the components for all modes.
- */
-
 void FEditorModeTools::AddReferencedObjects( FReferenceCollector& Collector )
 {
 	for( int32 x = 0 ; x < Modes.Num() ; ++x )
@@ -1042,10 +1022,6 @@ void FEditorModeTools::AddReferencedObjects( FReferenceCollector& Collector )
 	}
 }
 
-/**
- * Returns a pointer to an active mode specified by the passed in ID
- * If the editor mode is not active, NULL is returned
- */
 FEdMode* FEditorModeTools::GetActiveMode( FEditorModeID InID )
 {
 	for( auto& Mode : Modes )
@@ -1058,10 +1034,6 @@ FEdMode* FEditorModeTools::GetActiveMode( FEditorModeID InID )
 	return nullptr;
 }
 
-/**
- * Returns a pointer to an active mode specified by the passed in ID
- * If the editor mode is not active, NULL is returned
- */
 const FEdMode* FEditorModeTools::GetActiveMode( FEditorModeID InID ) const
 {
 	for (const auto& Mode : Modes)
@@ -1075,14 +1047,10 @@ const FEdMode* FEditorModeTools::GetActiveMode( FEditorModeID InID ) const
 	return nullptr;
 }
 
-/**
- * Returns the active tool of the passed in editor mode.
- * If the passed in editor mode is not active or the mode has no active tool, NULL is returned
- */
 const FModeTool* FEditorModeTools::GetActiveTool( FEditorModeID InID ) const
 {
 	const FEdMode* ActiveMode = GetActiveMode( InID );
-	const FModeTool* Tool = NULL;
+	const FModeTool* Tool = nullptr;
 	if( ActiveMode )
 	{
 		Tool = ActiveMode->GetCurrentTool();
@@ -1090,25 +1058,16 @@ const FModeTool* FEditorModeTools::GetActiveTool( FEditorModeID InID ) const
 	return Tool;
 }
 
-/** 
- * Returns true if the passed in editor mode is active 
- */
 bool FEditorModeTools::IsModeActive( FEditorModeID InID ) const
 {
-	return GetActiveMode( InID ) != NULL;
+	return GetActiveMode( InID ) != nullptr;
 }
 
-/** 
- * Returns true if the default editor mode is active 
- */
 bool FEditorModeTools::IsDefaultModeActive() const
 {
 	return IsModeActive(DefaultID);
 }
 
-/** 
- * Returns an array of all active modes
- */
 void FEditorModeTools::GetActiveModes( TArray<FEdMode*>& OutActiveModes )
 {
 	OutActiveModes.Empty();
diff --git a/Engine/Source/Editor/UnrealEd/Public/AssetEditorModeManager.h b/Engine/Source/Editor/UnrealEd/Public/AssetEditorModeManager.h
--- a/Engine/Source/Editor/UnrealEd/Public/AssetEditorModeManager.h
+++ b/Engine/Source/Editor/UnrealEd/Public/AssetEditorModeManager.h
@@ -16,6 +16,7 @@ public:
 	// FEditorModeTools interface
 	virtual class USelection* GetSelectedActors() const override;
 	virtual class USelection* GetSelectedObjects() const override;
+	virtual class USelection* GetSelectedComponents() const override;
 	virtual UWorld* GetWorld() const override;
 	// End of FEditorModeTools interface
 
@@ -23,5 +24,6 @@ public:
 protected:
 	class USelection* ActorSet;
 	class USelection* ObjectSet;
+	class USelection* ComponentSet;
 	class FPreviewScene* PreviewScene;
 };
diff --git a/Engine/Source/Editor/UnrealEd/Public/EdMode.h b/Engine/Source/Editor/UnrealEd/Public/EdMode.h
--- a/Engine/Source/Editor/UnrealEd/Public/EdMode.h
+++ b/Engine/Source/Editor/UnrealEd/Public/EdMode.h
@@ -289,6 +289,8 @@ public:
 			, bIsTransform(false)
 		{
 		}
+		
+		void GetTransformAndColor(UObject* BestSelectedItem, bool bIsSelected, FTransform& OutLocalTransform, FString& OutValidationMessage, FColor& OutDrawColor) const;
 	};
 
 protected:
@@ -306,6 +308,9 @@ protected:
 	 */
 	void GetPropertyWidgetInfos(const UStruct* InStruct, const void* InContainer, TArray<FPropertyWidgetInfo>& OutInfos) const;
 
+	/** Finds the best item to display widgets for (preferring selected components over selected actors) */
+	virtual UObject* GetItemToTryDisplayingWidgetsFor(FTransform& OutWidgetToWorld) const;
+
 	/** Name of the property currently being edited */
 	FString EditedPropertyName;
 
diff --git a/Engine/Source/Editor/UnrealEd/Public/EditorModeManager.h b/Engine/Source/Editor/UnrealEd/Public/EditorModeManager.h
--- a/Engine/Source/Editor/UnrealEd/Public/EditorModeManager.h
+++ b/Engine/Source/Editor/UnrealEd/Public/EditorModeManager.h
@@ -2,6 +2,9 @@
 
 #pragma once
 
+class USelection;
+class IToolkitHost;
+
 /**
  * A helper class to store the state of the various editor modes.
  */
@@ -337,8 +340,14 @@ public:
 	/** Is the viewport UI hidden? */
 	bool IsViewportUIHidden() const { return bHideViewportUI; }
 
-	bool PivotShown, Snapping, SnappedActor;
-	FVector CachedLocation, PivotLocation, SnappedLocation, GridBase;
+	bool PivotShown;
+	bool Snapping;
+	bool SnappedActor;
+
+	FVector CachedLocation;
+	FVector PivotLocation;
+	FVector SnappedLocation;
+	FVector GridBase;
 
 	/** The angle for the translate rotate widget */
 	float TranslateRotateXAxisAngle;
@@ -350,20 +359,25 @@ public:
 	FString InfoString;
 
 	/** Sets the host for toolkits created via modes from this mode manager (can only be called once) */
-	void SetToolkitHost(TSharedRef<class IToolkitHost> Host);
+	void SetToolkitHost(TSharedRef<IToolkitHost> Host);
 
 	/** Returns the host for toolkits created via modes from this mode manager */
-	TSharedPtr<class IToolkitHost> GetToolkitHost() const;
+	TSharedPtr<IToolkitHost> GetToolkitHost() const;
 
 	/**
 	 * Returns the set of selected actors.
 	 */
-	virtual class USelection* GetSelectedActors() const;
+	virtual USelection* GetSelectedActors() const;
 
 	/**
 	 * @return the set of selected non-actor objects.
 	 */
-	virtual class USelection* GetSelectedObjects() const;
+	virtual USelection* GetSelectedObjects() const;
+
+	/**
+	 * Returns the set of selected components.
+	 */
+	virtual USelection* GetSelectedComponents() const;
 
 	/**
 	 * Returns the world that is being edited by this mode manager
