Split the UEditorEngine implementation out from Editor.cpp into EditorEngine.cpp, removed some GWorld references, and purged some dead code

UEditorEngine cleanup:
- Split all UEditorEngine methods from Editor.cpp into EditorEngine.cpp (though more are still scattered around in other files like EditorServer.cpp)
- Removed GetClassFromPairMap, ClassMapPair, and associated keyboard shortcuts to place nav pylons, since they don't exist anymore
- Removed dead variable StreamingBoundsTexture and associated method SetStreamingBoundsTexture
- Split out the implementations of SplitActorLabel and SetActorLabelUnique into statics on FActorLabelUtilities, the non-static versions will be deprecated in the future
- Moved SoundWaveQualityPreview to be a static in FSoundPreviewThread
- Moved FPreviewInfo into DlgSoundWaveOptions.h
- Removed empty method HandleCreateMeshFromBSPCommand and associated exec command

GWorld cleanup:
- Removed EditorClearComponents and updated the few existing uses to call ClearWorldComponents directly on their world reference
- Removed UpdateSkyCaptures and updated the only existing use to call UpdateAllSkyCaptures directly on a world reference

diff --git a/Engine/Source/Editor/DetailCustomizations/Private/SkyLightComponentDetails.cpp b/Engine/Source/Editor/DetailCustomizations/Private/SkyLightComponentDetails.cpp
--- a/Engine/Source/Editor/DetailCustomizations/Private/SkyLightComponentDetails.cpp
+++ b/Engine/Source/Editor/DetailCustomizations/Private/SkyLightComponentDetails.cpp
@@ -69,7 +69,10 @@ FReply FSkyLightComponentDetails::OnUpdateSkyCapture()
 {
 	if( SkyLight.IsValid() )
 	{
-		GEditor->UpdateSkyCaptures();
+		if (UWorld* SkyLightWorld = SkyLight->GetWorld())
+		{
+			SkyLightWorld->UpdateAllSkyCaptures();
+		}
 	}
 
 	return FReply::Handled();
diff --git a/Engine/Source/Editor/UnrealEd/Classes/Editor/EditorEngine.h b/Engine/Source/Editor/UnrealEd/Classes/Editor/EditorEngine.h
--- a/Engine/Source/Editor/UnrealEd/Classes/Editor/EditorEngine.h
+++ b/Engine/Source/Editor/UnrealEd/Classes/Editor/EditorEngine.h
@@ -397,17 +397,10 @@ public:
 	UPROPERTY()
 	TArray<class UActorFactory*> ActorFactories;
 
-	/** Actors that are being deleted and should processed in the global re-attach*/
+	/** Actors that are being deleted and should processed in the global re-attach */
 	UPROPERTY()
 	TArray<class AActor*> ActorsForGlobalReregister;
 
-	/** String that maps one class name to another, used to create hook for game-specific actors created through shortcuts etc 
-	 *  Pairing is "ORIGINALCLASS;DESIREDCLASS
-	 *  (ie APylon;AMyGamePylon)
-	 */
-	UPROPERTY(config)
-	TArray<FString> ClassMapPair;
-
 	/** The name of the file currently being opened in the editor. "" if no file is being opened. */
 	UPROPERTY()
 	FString UserOpenedFile;
@@ -542,12 +535,6 @@ public:
 	UPROPERTY(transient)
 	class UTextureRenderTarget2D* ScratchRenderTarget256;
 
-	/**
-	 *	Display StreamingBounds for textures
-	 */
-	UPROPERTY(transient)
-	class UTexture2D* StreamingBoundsTexture;
-
 	/** Global instance of the editor user settings */
 	UPROPERTY()
 	class UEditorUserSettings* EditorUserSettings;
@@ -689,14 +676,40 @@ public:
 	void BroadcastObjectReimported(UObject* InObject);
 
 	// Begin UObject interface.
-	virtual void FinishDestroy() override;	
+	virtual void FinishDestroy() override;
 	virtual void PostEditChangeProperty(FPropertyChangedEvent& PropertyChangedEvent) override;
 	static void AddReferencedObjects(UObject* InThis, FReferenceCollector& Collector);
 	// End UObject interface.
 
 	// Begin UEngine interface.
+public:
 	virtual void Init(IEngineLoop* InEngineLoop) override;
+	virtual float GetMaxTickRate(float DeltaTime, bool bAllowFrameRateSmoothing = true) const override;
+	virtual void Tick(float DeltaSeconds, bool bIdleMode) override;
+	virtual bool ShouldDrawBrushWireframe(AActor* InActor) override;
+	virtual void NotifyToolsOfObjectReplacement(const TMap<UObject*, UObject*>& OldToNewInstanceMap) override;
+	virtual bool ShouldThrottleCPUUsage() const override;
+	virtual bool GetPropertyColorationColor(class UObject* Object, FColor& OutColor) override;
+	virtual bool WorldIsPIEInNewViewport(UWorld* InWorld) override;
+	virtual void FocusNextPIEWorld(UWorld* CurrentPieWorld, bool previous = false) override;
+	virtual class UGameViewportClient* GetNextPIEViewport(UGameViewportClient* CurrentViewport) override;
+	virtual UWorld* CreatePIEWorldByDuplication(FWorldContext &WorldContext, UWorld* InWorld, FString &PlayWorldMapName) override;
+	virtual bool GetMapBuildCancelled() const override { return false; }
+	virtual void SetMapBuildCancelled(bool InCancelled) override { /* Intentionally empty. */ }
+protected:
 	virtual void InitializeObjectReferences() override;
+	virtual void ProcessToggleFreezeCommand(UWorld* InWorld) override;
+	virtual void ProcessToggleFreezeStreamingCommand(UWorld* InWorld) override;
+private:
+	virtual void RemapGamepadControllerIdForPIE(class UGameViewportClient* GameViewport, int32 &ControllerId) override;
+	virtual TSharedPtr<SViewport> GetGameViewportWidget() const override;
+	virtual void TriggerStreamingDataRebuild() override;
+	virtual bool NetworkRemapPath(UWorld* InWorld, FString &Str, bool reading = true) override;
+	virtual bool NetworkRemapPath(UPendingNetGame* PendingNetGame, FString& Str, bool reading = true) override;
+	virtual bool AreEditorAnalyticsEnabled() const override;
+	virtual void CreateStartupAnalyticsAttributes(TArray<FAnalyticsEventAttribute>& StartSessionAttributes) const override;
+	virtual void VerifyLoadMapWorldCleanup() override;
+public:
 	// End UEngine interface.
 	
 	// Begin FExec Interface
@@ -735,17 +748,12 @@ public:
 	bool	HandleSetDetailModeCommand( const TCHAR* Str, FOutputDevice& Ar );
 	bool	HandleSetDetailModeViewCommand( const TCHAR* Str, FOutputDevice& Ar, UWorld* InWorld );
 	bool	HandleCleanBSPMaterialCommand( const TCHAR* Str, FOutputDevice& Ar, UWorld* InWorld  );
-	bool	HandleCreateMeshFromBSPCommand( const TCHAR* Str, FOutputDevice& Ar );
 	bool	HandleAutoMergeStaticMeshCommand( const TCHAR* Str, FOutputDevice& Ar );
 	bool	HandleAddSelectedCommand( const TCHAR* Str, FOutputDevice& Ar );
 	bool	HandleToggleSocketGModeCommand( const TCHAR* Str, FOutputDevice& Ar );
 	bool	HandleListMapPackageDependenciesCommand( const TCHAR* Str, FOutputDevice& Ar );
 	bool	HandleRebuildVolumesCommand( const TCHAR* Str, FOutputDevice& Ar, UWorld* InWorld );
 	bool	HandleRemoveArchtypeFlagCommand( const TCHAR* Str, FOutputDevice& Ar );
-	
-
-	/** Get tick rate limitor. */
-	virtual float GetMaxTickRate( float DeltaTime, bool bAllowFrameRateSmoothing = true ) const override;
 
 	/**
 	 * Initializes the Editor.
@@ -759,8 +767,6 @@ public:
 	 */
 	void InitBuilderBrush( UWorld* InWorld );
 
-	virtual void Tick( float DeltaSeconds, bool bIdleMode ) override;
-
 	/** Returns the global instance of the editor user settings class. */
 	const UEditorUserSettings& GetEditorUserSettings() const;
 	UEditorUserSettings& AccessEditorUserSettings();
@@ -808,30 +814,6 @@ public:
 		MP_NoChange
 	};
 
-	/**
-	 * Returns whether or not the map build in progressed was cancelled by the user.
-	 */
-	virtual bool GetMapBuildCancelled() const override
-	{
-		return false;
-	}
-
-	/**
-	 * Sets the flag that states whether or not the map build was cancelled.
-	 *
-	 * @param InCancelled	New state for the cancelled flag.
-	 */
-	virtual void SetMapBuildCancelled( bool InCancelled ) override
-	{
-		// Intentionally empty.
-	}
-
-	/**
-	 * Returns whether or not the actor passed in should draw as wireframe.
-	 *
-	 * @param InActor	Actor that is being drawn.
-	 */
-	virtual bool ShouldDrawBrushWireframe( AActor* InActor ) override;
 
 	// Execute a command that is safe for rebuilds.
 	virtual bool SafeExec( UWorld* InWorld, const TCHAR* Cmd, FOutputDevice& Ar=*GLog );
@@ -981,12 +963,6 @@ public:
 	 */
 	void RemovePerspectiveViewRotation(bool Roll, bool Pitch, bool Yaw);
 
-	/**
-	 * Notifies tools that a set of objects is being replaced with a new set of objects, so if any of the old
-	 * objects are being edited, they should be discarded and replaced with the new instance.
-	 */
-	virtual void NotifyToolsOfObjectReplacement(const TMap<UObject*, UObject*>& OldToNewInstanceMap) override;
-
 	//
 	// Pivot handling.
 	//
@@ -1511,11 +1487,6 @@ public:
 	// Level management.
 	void AnalyzeLevel(ULevel* Level,FOutputDevice& Ar);
 
-	/**
-	 * Removes all components from the current level's scene.
-	 */
-	void EditorClearComponents();
-
 	/**
 	 * Updates all components in the current level's scene.
 	 */
@@ -1702,7 +1673,6 @@ public:
 	 */
 	bool IsAnyViewportRealtime();
 
-	virtual bool ShouldThrottleCPUUsage() const override;
 
 	/**
 	 * @return true if all windows are hidden (including minimized)                                                         
@@ -1720,16 +1690,6 @@ public:
 	 */
 	TSharedRef<class FTimerManager> GetTimerManager() { return TimerManager.ToSharedRef(); }
 
-	/**
-	 * Handles freezing/unfreezing of rendering
-	 */
-	virtual void ProcessToggleFreezeCommand( UWorld* InWorld ) override;
-
-	/**
-	 * Handles frezing/unfreezing of streaming
-	 */
-	virtual void ProcessToggleFreezeStreamingCommand( UWorld* InWorld ) override;
-
 	// Editor specific
 
 	/**
@@ -1859,15 +1819,6 @@ public:
 	 */
 	void PasteSelectedActorsFromClipboard( UWorld* InWorld, const FText& TransDescription, const EPasteTo PasteTo );
 
-	/**
-	 * Computes a color to use for property coloration for the given object.
-	 *
-	 * @param	Object		The object for which to compute a property color.
-	 * @param	OutColor	[out] The returned color.
-	 * @return				true if a color was successfully set on OutColor, false otherwise.
-	 */
-	virtual bool GetPropertyColorationColor(class UObject* Object, FColor& OutColor) override;
-
 	/**
 	 * Sets property value and property chain to be used for property-based coloration.
 	 *
@@ -1893,14 +1844,6 @@ public:
 	 */
 	void SelectByPropertyColoration(UWorld* InWorld);
 
-	/**
-	 *	Sets the texture to use for displaying StreamingBounds.
-	 *
-	 *	@param	InTexture	The source texture for displaying StreamingBounds.
-	 *						Pass in NULL to disable displaying them.
-	 */
-	void SetStreamingBoundsTexture(UTexture2D* InTexture);
-
 	/**
 	 * Warns the user of any hidden levels, and prompts them with a Yes/No dialog
 	 * for whether they wish to continue with the operation.  No dialog is presented if all
@@ -1942,12 +1885,9 @@ public:
 	 */
 	virtual bool Game_Map_Check_Actor(const TCHAR* Str, FOutputDevice& Ar, bool bCheckDeprecatedOnly, AActor* InActor) { return true; }
 
-	/** Map given class name key to value in ClassMapPair array - if no key match just return the class of given name */
-	class UClass* GetClassFromPairMap( FString ClassName );
-
 	/**
 	 * Auto merge all staticmeshes that are able to be merged
-	*/
+	 */
 	void AutoMergeStaticMeshes();
 
 	/**
@@ -2200,8 +2140,6 @@ public:
 
 	void UpdateReflectionCaptures();
 
-	void UpdateSkyCaptures();
-
 	/**
 	 * Convenience method for adding a Slate modal window that is parented to the main frame (if it exists)
 	 * This function does not return until the modal window is closed.
@@ -2323,15 +2261,6 @@ public:
 	 */
 	bool ShouldAbortBecauseOfUnsavedWorld() const;
 
-	/**
-	 * Assigns a new label to an actor. If the name exists it will be appended with a number to make it unique. Actor labels are only available in development builds.
-	 *
-	 * @param	Actor					The actor to change the label of
-	 * @param	NewActorLabel			The new label string to assign to the actor.  If empty, the actor will have a default label.
-	 * @param	InExistingActorLabels	(optional) Pointer to a set of actor labels that are currently in use
-	 */
-	void SetActorLabelUnique( AActor* Actor, const FString& NewActorLabel, const FCachedActorLabels* InExistingActorLabels = nullptr ) const;
-
 	/**
 	 * Gets the user-friendly, localized (if exists) name of a property
 	 *
@@ -2347,13 +2276,13 @@ public:
 	 * Register a client tool to receive undo events 
 	 * @param UndoClient	An object wanting to receive PostUndo/PostRedo events
 	 */
-	void RegisterForUndo (class FEditorUndoClient* UndoClient );
+	void RegisterForUndo(class FEditorUndoClient* UndoClient );
 
 	/**
 	 * Unregister a client from receiving undo events 
 	 * @param UndoClient	An object wanting to unsubscribe from PostUndo/PostRedo events
 	 */
-	void UnregisterForUndo ( class FEditorUndoClient* UndoEditor );
+	void UnregisterForUndo( class FEditorUndoClient* UndoEditor );
 
 	/** 
 	 * Are we playing on a local PC session?
@@ -2383,19 +2312,11 @@ public:
 	 */
 	bool UsePercentageBasedScaling() const;
 
-	/**
-	 * Returns if this world is a PIE world that has its own viewport (is not using the editor viewport)
-	 */
-	virtual bool WorldIsPIEInNewViewport(UWorld *InWorld) override;
-
-	virtual void FocusNextPIEWorld(UWorld *CurrentPieWorld, bool previous=false) override;
-
 	DECLARE_DELEGATE(FPIEInstanceWindowSwitch);
 
 	/** Sets the delegate for when the focused PIE window is changed */
 	void SetPIEInstanceWindowSwitchDelegate(FPIEInstanceWindowSwitch PIEInstanceWindowSwitchDelegate);
 
-	virtual class UGameViewportClient *	GetNextPIEViewport(UGameViewportClient * CurrentViewport) override;
 
 private:
 	//
@@ -2441,13 +2362,6 @@ private:
 	 * @param	InLevel		Destination level.
 	 */
 	void DoMoveSelectedActorsToLevel( ULevel* InLevel );
-public:
-	/**
-	 * Creates a PIE world by duplicating the editor world	 
-	 */
-	virtual UWorld* CreatePIEWorldByDuplication(FWorldContext &WorldContext, UWorld* InWorld, FString &PlayWorldMapName) override;
-private:
-	virtual void RemapGamepadControllerIdForPIE(class UGameViewportClient* GameViewport, int32 &ControllerId) override;
 
 public:
 	/** Creates a PIE world by saving to a temp file and then reloading it */
@@ -2577,19 +2491,6 @@ private:
 	 */
 	void EditorDestroyWorld( FWorldContext & Context, const FText& CleanseText, UWorld* NewWorld = nullptr );
 
-	/** Returns the GameViewport widget */
-	virtual TSharedPtr<SViewport> GetGameViewportWidget() const override;
-
-	/**
-	 * Given a label, attempts to split this into its alpha/numeric parts.
-	 *
-	 * @param	InOutLabel	The label to start with, this will only be modified if it ends in a number.
-	 * @param	OutIdx		The number which the string ends with, if any.
-	 *
-	 * @return	true if the label ends with a number.
-	 */
-	bool SplitActorLabel( FString& InOutLabel, int32& OutIdx ) const;
-
 	ULevel* CreateTransLevelMoveBuffer( UWorld* InWorld );
 
 	/**	Broadcasts that an undo has just occurred. */
@@ -2625,14 +2526,6 @@ private:
 	/** Updates the project file to auto load and initializes the bLoadTheMostRecentlyLoadedProjectAtStartup flag */
 	void UpdateAutoLoadProject();
 
-	virtual void TriggerStreamingDataRebuild() override;
-
-	/** Remaps a network path for PIE Networking under multiple worlds */
-	virtual bool NetworkRemapPath( UWorld *InWorld, FString &Str, bool reading=true) override;
-
-	/** Remaps a network path for PIE Networking under multiple worlds */
-	virtual bool NetworkRemapPath( UPendingNetGame *PendingNetGame, FString &Str, bool reading=true) override;
-
 	/** Handles user setting changes. */
 	void HandleSettingChanged( FName Name );
 
@@ -2645,9 +2538,6 @@ private:
 	/** Callback for finished undo transactions. */
 	void HandleTransactorUndo( FUndoSessionContext SessionContext, bool Succeeded );
 
-	/** UEngine interface */
-	virtual bool AreEditorAnalyticsEnabled() const override;
-	virtual void CreateStartupAnalyticsAttributes( TArray<FAnalyticsEventAttribute>& StartSessionAttributes ) const override;
 private:
 
 	/** Delegate broadcast just before a blueprint is compiled */
@@ -2722,8 +2612,6 @@ private:
 	/** List of files we are deferring adding to source control */
 	TArray<FString> DeferredFilesToAddToSourceControl;
 
-	virtual void VerifyLoadMapWorldCleanup() override;
-
 	FPIEInstanceWindowSwitch PIEInstanceWindowSwitchDelegate;
 
 	/** Number of currently running instances logged into an online platform */
@@ -2805,6 +2693,58 @@ public:
 	/** Function to run the Play On command for automation testing. */
 	void AutomationPlayUsingLauncher(const FString& InLauncherDeviceId);	
 
+public:
+	/**
+	 * Given a label, attempts to split this into its alpha/numeric parts.
+	 *
+	 * @param	InOutLabel	The label to start with, this will only be modified if it ends in a number.
+	 * @param	OutIdx		The number which the string ends with, if any.
+	 *
+	 * @return	true if the label ends with a number.
+	 */
+	//DEPRECATED(4.8, "This function is deprecated.  Please use FActorLabelUtilities::SplitActorLabel instead")
+	bool SplitActorLabel( FString& InOutLabel, int32& OutIdx ) const;
+
+	/**
+	 * Assigns a new label to an actor. If the name exists it will be appended with a number to make it unique. Actor labels are only available in development builds.
+	 *
+	 * @param	Actor					The actor to change the label of
+	 * @param	NewActorLabel			The new label string to assign to the actor.  If empty, the actor will have a default label.
+	 * @param	InExistingActorLabels	(optional) Pointer to a set of actor labels that are currently in use
+	 */
+	//DEPRECATED(4.8, "This function is deprecated.  Please use FActorLabelUtilities::SetActorLabelUnique instead")
+	void SetActorLabelUnique( AActor* Actor, const FString& NewActorLabel, const FCachedActorLabels* InExistingActorLabels = nullptr ) const;
+
 private:
 	FTimerHandle CleanupPIEOnlineSessionsTimerHandle;
 };
+
+//////////////////////////////////////////////////////////////////////////
+// UEditorEngine
+
+struct UNREALED_API FActorLabelUtilities
+{
+public:
+	/**
+	 * Given a label, attempts to split this into its alpha/numeric parts.
+	 *
+	 * @param	InOutLabel	The label to start with, this will only be modified if it ends in a number.
+	 * @param	OutIdx		The number which the string ends with, if any.
+	 *
+	 * @return	true if the label ends with a number.
+	 */
+	static bool SplitActorLabel(FString& InOutLabel, int32& OutIdx);
+
+	/**
+	 * Assigns a new label to an actor. If the name exists it will be appended with a number to make it unique. Actor labels are only available in development builds.
+	 *
+	 * @param	Actor					The actor to change the label of
+	 * @param	NewActorLabel			The new label string to assign to the actor.  If empty, the actor will have a default label.
+	 * @param	InExistingActorLabels	(optional) Pointer to a set of actor labels that are currently in use
+	 */
+	//DEPRECATED(4.8, "This function is deprecated.  Please use FActorLabelUtilities::SetActorLabelUnique instead")
+	static void SetActorLabelUnique(AActor* Actor, const FString& NewActorLabel, const FCachedActorLabels* InExistingActorLabels = nullptr);
+
+private:
+	FActorLabelUtilities() {}
+};
