Split the UEditorEngine implementation out from Editor.cpp into EditorEngine.cpp, removed some GWorld references, and purged some dead code

diff --git a/Engine/Source/Editor/UnrealEd/Private/Editor.cpp b/Engine/Source/Editor/UnrealEd/Private/Editor.cpp
--- a/Engine/Source/Editor/UnrealEd/Private/Editor.cpp
+++ b/Engine/Source/Editor/UnrealEd/Private/Editor.cpp
@@ -110,8 +110,6 @@
 #include "EngineStats.h"
 #include "Engine/SimpleConstructionScript.h"
 
-DEFINE_LOG_CATEGORY_STATIC(LogEditor, Log, All);
-
 #define LOCTEXT_NAMESPACE "UnrealEd.Editor"
 
 TArray<FTickableEditorObject*> FTickableEditorObject::TickableObjects;
@@ -169,6804 +167,432 @@ FEditorDelegates::FOnAddLevelToWorld					FEditorDelegates::OnAddLevelToWorld;
 
 IMPLEMENT_STRUCT(SlatePlayInEditorInfo);
 
-UEditorEngine*	GEditor;
-
-static inline USelection*& PrivateGetSelectedActors()
-{
-	static USelection* SSelectedActors = NULL;
-	return SSelectedActors;
-};
-
-static inline USelection*& PrivateGetSelectedComponents()
-{
-	static USelection* SSelectedComponents = NULL;
-	return SSelectedComponents;
-}
-
-static inline USelection*& PrivateGetSelectedObjects()
-{
-	static USelection* SSelectedObjects = NULL;
-	return SSelectedObjects;
-};
-
-static void OnObjectSelected(UObject* Object)
-{
-	// Whenever an actor is unselected we must remove its components from the components selection
-	if (!Object->IsSelected())
-	{
-		TArray<UActorComponent*> ComponentsToDeselect;
-		for (FSelectionIterator It(*PrivateGetSelectedComponents()); It; ++It)
-		{
-			UActorComponent* Component = CastChecked<UActorComponent>(*It);
-			if (Component->GetOwner() == Object)
-			{
-				ComponentsToDeselect.Add(Component);
-			}
-		}
-		if (ComponentsToDeselect.Num() > 0)
-		{
-			PrivateGetSelectedComponents()->Modify();
-			PrivateGetSelectedComponents()->BeginBatchSelectOperation();
-			for (UActorComponent* Component : ComponentsToDeselect)
-			{
-				PrivateGetSelectedComponents()->Deselect(Component);
-			}
-			PrivateGetSelectedComponents()->EndBatchSelectOperation();
-		}
-	}
-}
-
-static void PrivateInitSelectedSets()
-{
-	PrivateGetSelectedActors() = NewObject<USelection>(GetTransientPackage(), TEXT("SelectedActors"), RF_Transactional);
-	PrivateGetSelectedActors()->AddToRoot();
-
-	PrivateGetSelectedActors()->SelectObjectEvent.AddStatic(&OnObjectSelected);
+//////////////////////////////////////////////////////////////////////////
+// FReimportManager
 
-	PrivateGetSelectedComponents() = NewObject<USelection>(GetTransientPackage(), TEXT("SelectedComponents"), RF_Transactional);
-	PrivateGetSelectedComponents()->AddToRoot();
-
-	PrivateGetSelectedObjects() = NewObject<USelection>(GetTransientPackage(), TEXT("SelectedObjects"), RF_Transactional);
-	PrivateGetSelectedObjects()->AddToRoot();
-}
-
-static void PrivateDestroySelectedSets()
-{
-#if 0
-	PrivateGetSelectedActors()->RemoveFromRoot();
-	PrivateGetSelectedActors() = NULL;
-	PrivateGetSelectedComponents()->RemoveFromRoot();
-	PrivateGetSelectedComponents() = NULL;
-	PrivateGetSelectedObjects()->RemoveFromRoot();
-	PrivateGetSelectedObjects() = NULL;
-#endif
-}
-
-UEditorEngine::UEditorEngine(const FObjectInitializer& ObjectInitializer)
-: Super(ObjectInitializer)
+FReimportManager* FReimportManager::Instance()
 {
-	if (!IsRunningCommandlet())
-	{
-		// Structure to hold one-time initialization
-		struct FConstructorStatics
-		{
-			ConstructorHelpers::FObjectFinder<UTexture2D> BadTexture;
-			ConstructorHelpers::FObjectFinder<UTexture2D> BackgroundTexture;
-			ConstructorHelpers::FObjectFinder<UTexture2D> BackgroundHiTexture;
-			ConstructorHelpers::FObjectFinder<UStaticMesh> EditorCubeMesh;
-			ConstructorHelpers::FObjectFinder<UStaticMesh> EditorSphereMesh;
-			ConstructorHelpers::FObjectFinder<UStaticMesh> EditorPlaneMesh;
-			ConstructorHelpers::FObjectFinder<UStaticMesh> EditorCylinderMesh;
-			ConstructorHelpers::FObjectFinder<UFont> SmallFont;
-			FConstructorStatics()
-				: BadTexture(TEXT("/Engine/EditorResources/Bad"))
-				, BackgroundTexture(TEXT("/Engine/EditorResources/Bkgnd"))
-				, BackgroundHiTexture(TEXT("/Engine/EditorResources/BkgndHi"))
-				, EditorCubeMesh(TEXT("/Engine/EditorMeshes/EditorCube"))
-				, EditorSphereMesh(TEXT("/Engine/EditorMeshes/EditorSphere"))
-				, EditorPlaneMesh(TEXT("/Engine/EditorMeshes/EditorPlane"))
-				, EditorCylinderMesh(TEXT("/Engine/EditorMeshes/EditorCylinder"))
-				, SmallFont(TEXT("/Engine/EngineFonts/Roboto"))
-			{
-			}
-		};
-		static FConstructorStatics ConstructorStatics;
-
-		Bad = ConstructorStatics.BadTexture.Object;
-		Bkgnd = ConstructorStatics.BackgroundTexture.Object;
-		BkgndHi = ConstructorStatics.BackgroundHiTexture.Object;
-		EditorCube = ConstructorStatics.EditorCubeMesh.Object;
-		EditorSphere = ConstructorStatics.EditorSphereMesh.Object;
-		EditorPlane = ConstructorStatics.EditorPlaneMesh.Object;
-		EditorCylinder = ConstructorStatics.EditorCylinderMesh.Object;
-		EditorFont = ConstructorStatics.SmallFont.Object;
-	}
-
-	DetailMode = DM_MAX;
-	PlayInEditorViewportIndex = -1;
-	CurrentPlayWorldDestination = -1;
-	bDisableDeltaModification = false;
-	bPlayOnLocalPcSession = false;
-	bAllowMultiplePIEWorlds = true;
-	NumOnlinePIEInstances = 0;
-	DefaultWorldFeatureLevel = GMaxRHIFeatureLevel;
+	static FReimportManager Inst;
+	return &Inst;
 }
 
-
-int32 UEditorEngine::GetSelectedActorCount() const
+void FReimportManager::RegisterHandler( FReimportHandler& InHandler )
 {
-	int32 NumSelectedActors = 0;
-	for(FSelectionIterator It(GetSelectedActorIterator()); It; ++It)
-	{
-		++NumSelectedActors;
-	}
-
-	return NumSelectedActors;
+	Handlers.AddUnique( &InHandler );
+	bHandlersNeedSorting = true;
 }
 
-
-USelection* UEditorEngine::GetSelectedActors() const
+void FReimportManager::UnregisterHandler( FReimportHandler& InHandler )
 {
-	return PrivateGetSelectedActors();
+	Handlers.Remove( &InHandler );
 }
 
-bool UEditorEngine::IsWorldSettingsSelected() 
+bool FReimportManager::CanReimport( UObject* Obj ) const
 {
-	if( bCheckForWorldSettingsActors )
+	if ( Obj )
 	{
-		bIsWorldSettingsSelected = false;
-		for ( FSelectionIterator It( GetSelectedActorIterator() ) ; It ; ++It )
+		TArray<FString> SourceFilenames;
+		for( int32 HandlerIndex = 0; HandlerIndex < Handlers.Num(); ++HandlerIndex )
 		{
-			if ( Cast<AWorldSettings>( *It ) )
+			SourceFilenames.Empty();
+			if ( Handlers[ HandlerIndex ]->CanReimport(Obj, SourceFilenames) )
 			{
-				bIsWorldSettingsSelected = true;
-				break;
+				return true;
 			}
 		}
-		bCheckForWorldSettingsActors = false;
-	}
-
-	return bIsWorldSettingsSelected;
-}
-
-FSelectionIterator UEditorEngine::GetSelectedActorIterator() const
-{
-	return FSelectionIterator( *GetSelectedActors() );
-};
-
-int32 UEditorEngine::GetSelectedComponentCount() const
-{
-	int32 NumSelectedComponents = 0;
-	for (FSelectionIterator It(GetSelectedComponentIterator()); It; ++It)
-	{
-		++NumSelectedComponents;
 	}
-
-	return NumSelectedComponents;
-}
-
-FSelectionIterator UEditorEngine::GetSelectedComponentIterator() const
-{
-	return FSelectionIterator(*GetSelectedComponents());
-};
-
-FSelectedEditableComponentIterator UEditorEngine::GetSelectedEditableComponentIterator() const
-{
-	return FSelectedEditableComponentIterator(*GetSelectedComponents());
-}
-
-USelection* UEditorEngine::GetSelectedComponents() const
-{
-	return PrivateGetSelectedComponents();
-}
-
-USelection* UEditorEngine::GetSelectedObjects() const
-{
-	return PrivateGetSelectedObjects();
+	return false;
 }
 
-void UEditorEngine::GetContentBrowserSelectionClasses(TArray<UClass*>& Selection) const
+void FReimportManager::UpdateReimportPaths( UObject* Obj, const TArray<FString>& InFilenames )
 {
-	FContentBrowserModule& ContentBrowserModule = FModuleManager::LoadModuleChecked<FContentBrowserModule>("ContentBrowser");
-	TArray<FAssetData> SelectedAssets;
-	ContentBrowserModule.Get().GetSelectedAssets(SelectedAssets);
-
-	for ( auto AssetIt = SelectedAssets.CreateConstIterator(); AssetIt; ++AssetIt )
+	if (Obj)
 	{
-		UClass* AssetClass = FindObject<UClass>(ANY_PACKAGE, *(*AssetIt).AssetClass.ToString());
-
-		if ( AssetClass != NULL )
+		TArray<FString> UnusedExistingFilenames;
+		auto* Handler = Handlers.FindByPredicate([&](FReimportHandler* InHandler){ return InHandler->CanReimport(Obj, UnusedExistingFilenames); });
+		if (Handler)
 		{
-			Selection.AddUnique(AssetClass);
+			(*Handler)->SetReimportPaths(Obj, InFilenames);
+			Obj->MarkPackageDirty();
 		}
 	}
 }
 
-void UEditorEngine::GetContentBrowserSelections(TArray<FAssetData>& Selection) const
+bool FReimportManager::Reimport( UObject* Obj, bool bAskForNewFileIfMissing, bool bShowNotification )
 {
-	FContentBrowserModule& ContentBrowserModule = FModuleManager::LoadModuleChecked<FContentBrowserModule>("ContentBrowser");
-	ContentBrowserModule.Get().GetSelectedAssets(Selection);
-}
+	// Warn that were about to reimport, so prep for it
+	PreReimport.Broadcast( Obj );
 
-USelection* UEditorEngine::GetSelectedSet( const UClass* Class ) const
-{
-	USelection* SelectedSet = GetSelectedActors();
-	if ( Class->IsChildOf( AActor::StaticClass() ) )
-	{
-		return SelectedSet;
-	}
-	else
+	bool bSuccess = false;
+	if ( Obj )
 	{
-		//make sure this actor isn't derived off of an interface class
-		for ( FSelectionIterator It( GetSelectedActorIterator() ) ; It ; ++It )
+		if (bHandlersNeedSorting)
 		{
-			AActor* TestActor = static_cast<AActor*>( *It );
-			if (TestActor->GetClass()->ImplementsInterface(Class))
-			{
-				return SelectedSet;
-			}
+			// Use > operator because we want higher priorities earlier in the list
+			Handlers.Sort([](const FReimportHandler& A, const FReimportHandler& B) { return A.GetPriority() > B.GetPriority(); });
+			bHandlersNeedSorting = false;
 		}
+		
+		bool bValidSourceFilename = false;
+		TArray<FString> SourceFilenames;
 
-		//no actor matched the interface class
-		return GetSelectedObjects();
-	}
-}
-
-const UClass* UEditorEngine::GetFirstSelectedClass( const UClass* const RequiredParentClass ) const
-{
-	const USelection* const SelectedObjects = GetSelectedObjects();
-
-	for(int32 i = 0; i < SelectedObjects->Num(); ++i)
-	{
-		const UObject* const SelectedObject = SelectedObjects->GetSelectedObject(i);
-
-		if(SelectedObject)
+		for( int32 HandlerIndex = 0; HandlerIndex < Handlers.Num(); ++HandlerIndex )
 		{
-			const UClass* SelectedClass = nullptr;
-
-			if(SelectedObject->IsA(UBlueprint::StaticClass()))
+			SourceFilenames.Empty();
+			if ( Handlers[ HandlerIndex ]->CanReimport(Obj, SourceFilenames) )
 			{
-				// Handle selecting a blueprint
-				const UBlueprint* const SelectedBlueprint = StaticCast<const UBlueprint*>(SelectedObject);
-				if(SelectedBlueprint->GeneratedClass)
+				// Check all filenames for missing files
+				bool bMissingFiles = false;
+				if (SourceFilenames.Num() > 0)
+				{
+					for (int32 FileIndex = 0; FileIndex < SourceFilenames.Num(); ++FileIndex)
+					{
+						if (SourceFilenames[FileIndex].IsEmpty() || IFileManager::Get().FileSize(*SourceFilenames[FileIndex]) == INDEX_NONE)
+						{
+							bMissingFiles = true;
+							break;
+						}
+					}
+				}
+				else
 				{
-					SelectedClass = SelectedBlueprint->GeneratedClass;
+					bMissingFiles = true;
 				}
-			}
-			else if(SelectedObject->IsA(UClass::StaticClass()))
-			{
-				// Handle selecting a class
-				SelectedClass = StaticCast<const UClass*>(SelectedObject);
-			}
-
-			if(SelectedClass && (!RequiredParentClass || SelectedClass->IsChildOf(RequiredParentClass)))
-			{
-				return SelectedClass;
-			}
-		}
-	}
-
-	return nullptr;
-}
-
-static bool GetSmallToolBarIcons()
-{
-	return GetDefault<UEditorStyleSettings>()->bUseSmallToolBarIcons;
-}
-
-static bool GetDisplayMultiboxHooks()
-{
-	return GEditor->AccessEditorUserSettings().bDisplayUIExtensionPoints;
-}
-
-void UEditorEngine::InitEditor(IEngineLoop* InEngineLoop)
-{
-	// Call base.
-	UEngine::Init(InEngineLoop);
-
-	// Specify "-ForceLauncher" on the command-line to always open the launcher, even in unusual cases.  This is useful for debugging the Launcher startup.
-	const bool bForceLauncherToOpen = FParse::Param(FCommandLine::Get(), TEXT("ForceLauncher"));
-
-	if ( bForceLauncherToOpen ||
-		( !FEngineBuildSettings::IsInternalBuild() &&
-		!FEngineBuildSettings::IsPerforceBuild() &&
-		!FPlatformMisc::IsDebuggerPresent() &&	// Don't spawn launcher while running in the Visual Studio debugger by default
-		!FApp::IsBenchmarking() &&
-		!GIsDemoMode &&
-		!IsRunningCommandlet() &&
-		!FPlatformProcess::IsApplicationRunning(TEXT("UnrealEngineLauncher")) &&
-		!FPlatformProcess::IsApplicationRunning(TEXT("Unreal Engine Launcher")) &&
-		!FPlatformProcess::IsApplicationRunning(TEXT("Epic Launcher")) ) )
-	{
-		IDesktopPlatform* DesktopPlatform = FDesktopPlatformModule::Get();
-		if ( DesktopPlatform != NULL )
-		{
-			DesktopPlatform->OpenLauncher(false, TEXT(""));
-		}
-	}
-
-	// Create selection sets.
-	PrivateInitSelectedSets();
-
-	// Set slate options
-	FMultiBoxSettings::UseSmallToolBarIcons = TAttribute<bool>::Create(TAttribute<bool>::FGetter::CreateStatic(&GetSmallToolBarIcons));
-	FMultiBoxSettings::DisplayMultiboxHooks = TAttribute<bool>::Create(TAttribute<bool>::FGetter::CreateStatic(&GetDisplayMultiboxHooks));
-
-	if ( FSlateApplication::IsInitialized() )
-	{
-		FSlateApplication::Get().GetRenderer()->SetColorVisionDeficiencyType((uint32)( GetDefault<UEditorStyleSettings>()->ColorVisionDeficiencyPreviewType.GetValue() ));
-		FSlateApplication::Get().EnableMenuAnimations(GetDefault<UEditorStyleSettings>()->bEnableWindowAnimations);
-	}
-
-	const UEditorStyleSettings* StyleSettings = GetDefault<UEditorStyleSettings>();
-	const ULevelEditorViewportSettings* ViewportSettings = GetDefault<ULevelEditorViewportSettings>();
-
-	// Needs to be set early as materials can be cached with selected material color baked in
-	GEngine->SetSelectedMaterialColor(ViewportSettings->bHighlightWithBrackets ? FLinearColor::Black : StyleSettings->SelectionColor);
-	GEngine->SetSelectionOutlineColor(StyleSettings->SelectionColor);
-	GEngine->SetSubduedSelectionOutlineColor(StyleSettings->GetSubduedSelectionColor());
-	GEngine->SelectionHighlightIntensity = ViewportSettings->SelectionHighlightIntensity;
-	GEngine->BSPSelectionHighlightIntensity = ViewportSettings->BSPSelectionHighlightIntensity;
-	GEngine->HoverHighlightIntensity = ViewportSettings->HoverHighlightIntensity;
-
-	// Set navigation system property indicating whether navigation is supposed to rebuild automatically 
-	FWorldContext &EditorContext = GEditor->GetEditorWorldContext();
-	UNavigationSystem::SetNavigationAutoUpdateEnabled(GetDefault<ULevelEditorMiscSettings>()->bNavigationAutoUpdate, EditorContext.World()->GetNavigationSystem());
-
-	// Allocate temporary model.
-	TempModel = NewObject<UModel>();
-	TempModel->Initialize(nullptr, 1);
-	ConversionTempModel = NewObject<UModel>();
-	ConversionTempModel->Initialize(nullptr, 1);
-
-	// create the timer manager
-	TimerManager = MakeShareable(new FTimerManager());
-
-	// Settings.
-	FBSPOps::GFastRebuild = 0;
-
-	// Setup delegate callbacks for SavePackage()
-	FCoreUObjectDelegates::IsPackageOKToSaveDelegate.BindUObject(this, &UEditorEngine::IsPackageOKToSave);
-	FCoreUObjectDelegates::AutoPackageBackupDelegate.BindStatic(&FAutoPackageBackup::BackupPackage);
-
-	extern void SetupDistanceFieldBuildNotification();
-	SetupDistanceFieldBuildNotification();
-
-	// Update recents
-	UpdateRecentlyLoadedProjectFiles();
-
-	// Update the auto-load project
-	UpdateAutoLoadProject();
-
-	// Load any modules that might be required by commandlets
-	FModuleManager::Get().LoadModule(TEXT("OnlineBlueprintSupport"));
-
-	if ( FSlateApplication::IsInitialized() )
-	{
-		// Setup a delegate to handle requests for opening assets
-		FSlateApplication::Get().SetWidgetReflectorAssetAccessDelegate(FAccessAsset::CreateUObject(this, &UEditorEngine::HandleOpenAsset));
-	}
-}
-
-bool UEditorEngine::HandleOpenAsset(UObject* Asset)
-{
-	return FAssetEditorManager::Get().OpenEditorForAsset(Asset);
-}
-
-void UEditorEngine::HandleSettingChanged( FName Name )
-{
-	// When settings are reset to default, the property name will be "None" so make sure that case is handled.
-	if (Name == FName(TEXT("ColorVisionDeficiencyPreviewType")) || Name == NAME_None)
-	{
-		uint32 DeficiencyType = (uint32)GetDefault<UEditorStyleSettings>()->ColorVisionDeficiencyPreviewType.GetValue();
-		FSlateApplication::Get().GetRenderer()->SetColorVisionDeficiencyType(DeficiencyType);
-
-		GEngine->Exec(NULL, TEXT("RecompileShaders SlateElementPixelShader"));
-	}
-	if (Name == FName("SelectionColor") || Name == NAME_None)
-	{
-		// Selection outline color and material color use the same color but sometimes the selected material color can be overidden so these need to be set independently
-		GEngine->SetSelectedMaterialColor(GetDefault<UEditorStyleSettings>()->SelectionColor);
-		GEngine->SetSelectionOutlineColor(GetDefault<UEditorStyleSettings>()->SelectionColor);
-		GEngine->SetSubduedSelectionOutlineColor(GetDefault<UEditorStyleSettings>()->GetSubduedSelectionColor());
-	}
-}
-
-void UEditorEngine::InitializeObjectReferences()
-{
-	Super::InitializeObjectReferences();
-
-	if ( PlayFromHerePlayerStartClass == NULL )
-	{
-		PlayFromHerePlayerStartClass = LoadClass<ANavigationObjectBase>(NULL, *GetDefault<ULevelEditorPlaySettings>()->PlayFromHerePlayerStartClassName, NULL, LOAD_None, NULL);
-	}
-}
-
-bool UEditorEngine::ShouldDrawBrushWireframe( AActor* InActor )
-{
-	bool bResult = true;
-
-	bResult = GLevelEditorModeTools().ShouldDrawBrushWireframe( InActor );
-	
-	return bResult;
-}
-
-//
-// Init the editor.
-//
-
-extern void StripUnusedPackagesFromList(TArray<FString>& PackageList, const FString& ScriptSourcePath);
-
-void UEditorEngine::Init(IEngineLoop* InEngineLoop)
-{
-	FScopedSlowTask SlowTask(100);
-
-	DECLARE_SCOPE_CYCLE_COUNTER(TEXT("Editor Engine Initialized"), STAT_EditorEngineStartup, STATGROUP_LoadTime);
-
-	check(!HasAnyFlags(RF_ClassDefaultObject));
-
-	FSlateApplication::Get().SetAppIcon(FEditorStyle::GetBrush(TEXT("Editor.AppIcon")));
-
-	FCoreDelegates::ModalErrorMessage.BindUObject(this, &UEditorEngine::OnModalMessageDialog);
-	FCoreUObjectDelegates::ShouldLoadOnTop.BindUObject(this, &UEditorEngine::OnShouldLoadOnTop);
-	FCoreDelegates::PreWorldOriginOffset.AddUObject(this, &UEditorEngine::PreWorldOriginOffset);
-	FCoreUObjectDelegates::OnAssetLoaded.AddUObject(this, &UEditorEngine::OnAssetLoaded);
-	FWorldDelegates::LevelAddedToWorld.AddUObject(this, &UEditorEngine::OnLevelAddedToWorld);
-	FWorldDelegates::LevelRemovedFromWorld.AddUObject(this, &UEditorEngine::OnLevelRemovedFromWorld);
-	FLevelStreamingGCHelper::OnGCStreamedOutLevels.AddUObject(this, &UEditorEngine::OnGCStreamedOutLevels);
-	
-	// Init editor.
-	SlowTask.EnterProgressFrame(40);
-	GEditor = this;
-	InitEditor(InEngineLoop);
-
-	Layers = FLayers::Create( TWeakObjectPtr< UEditorEngine >( this ) );
-
-	// Init transactioning.
-	Trans = CreateTrans();
-
-	// Load all of the runtime modules that the game needs.  The game is part of the editor, so we'll need these loaded.
-	UGameEngine::LoadRuntimeEngineStartupModules();
-
-	SlowTask.EnterProgressFrame(50);
-
-	// Load all editor modules here
-	{
-		static const TCHAR* ModuleNames[] =
-		{
-			TEXT("Documentation"),
-			TEXT("WorkspaceMenuStructure"),
-			TEXT("MainFrame"),
-			TEXT("GammaUI"),
-			TEXT("OutputLog"),
-			TEXT("SourceControl"),
-			TEXT("TextureCompressor"),
-			TEXT("MeshUtilities"),
-			TEXT("MovieSceneTools"),
-			TEXT("ModuleUI"),
-			TEXT("Toolbox"),
-			TEXT("ClassViewer"),
-			TEXT("ContentBrowser"),
-			TEXT("AssetTools"),
-			TEXT("GraphEditor"),
-			TEXT("KismetCompiler"),
-			TEXT("Kismet"),
-			TEXT("Persona"),
-			TEXT("LevelEditor"),
-			TEXT("MainFrame"),
-			TEXT("PropertyEditor"),
-			TEXT("EditorStyle"),
-			TEXT("PackagesDialog"),
-			TEXT("AssetRegistry"),
-			TEXT("DetailCustomizations"),
-			TEXT("ComponentVisualizers"),
-			TEXT("Layers"),
-			TEXT("AutomationWindow"),
-			TEXT("AutomationController"),
-			TEXT("DeviceManager"),
-			TEXT("ProfilerClient"),
-//			TEXT("Search"),
-			TEXT("SessionFrontend"),
-			TEXT("ProjectLauncher"),
-			TEXT("SettingsEditor"),
-			TEXT("EditorSettingsViewer"),
-			TEXT("ProjectSettingsViewer"),
-			TEXT("AndroidRuntimeSettings"),
-			TEXT("AndroidPlatformEditor"),
-			TEXT("HTML5PlatformEditor"),
-			TEXT("IOSRuntimeSettings"),
-			TEXT("IOSPlatformEditor"),
-			TEXT("Blutility"),
-			TEXT("OnlineBlueprintSupport"),
-			TEXT("XmlParser"),
-			TEXT("UserFeedback"),
-			TEXT("GameplayTagsEditor"),
-			TEXT("UndoHistory"),
-			TEXT("DeviceProfileEditor"),
-			TEXT("SourceCodeAccess"),
-			TEXT("BehaviorTreeEditor"),
-			TEXT("HardwareTargeting"),
-			TEXT("LocalizationDashboard")
-		};
-
-		FScopedSlowTask ModuleSlowTask(ARRAY_COUNT(ModuleNames));
-		for (const TCHAR* ModuleName : ModuleNames)
-		{
-			ModuleSlowTask.EnterProgressFrame(1);
-			FModuleManager::Get().LoadModule(ModuleName);
-		}
-
-		if (!IsRunningCommandlet())
-		{
-			FModuleManager::Get().LoadModule(TEXT("EditorLiveStreaming"));
-			FModuleManager::Get().LoadModule(TEXT("IntroTutorials"));
-		}
-
-		if( FParse::Param( FCommandLine::Get(),TEXT( "PListEditor" ) ) )
-		{
-			FModuleManager::Get().LoadModule(TEXT("PListEditor"));
-		}
-
-		bool bEnvironmentQueryEditor = false;
-		GConfig->GetBool(TEXT("EnvironmentQueryEd"), TEXT("EnableEnvironmentQueryEd"), bEnvironmentQueryEditor, GEngineIni);
-		if (bEnvironmentQueryEditor || GetDefault<UEditorExperimentalSettings>()->bEQSEditor)
-		{
-			FModuleManager::Get().LoadModule(TEXT("EnvironmentQueryEditor"));
-		}
-
-		bool bGameplayAbilitiesEnabled = false;
-		GConfig->GetBool(TEXT("GameplayAbilities"), TEXT("GameplayAbilitiesEditorEnabled"), bGameplayAbilitiesEnabled, GEngineIni);
-		if (bGameplayAbilitiesEnabled)
-		{
-			FModuleManager::Get().LoadModule(TEXT("GameplayAbilitiesEditor"));
-		}
-
-
-		FModuleManager::Get().LoadModule(TEXT("LogVisualizer"));
-		FModuleManager::Get().LoadModule(TEXT("HotReload"));
-	}
-
-	SlowTask.EnterProgressFrame(10);
-
-	float BSPTexelScale = 100.0f;
-	if( GetDefault<ULevelEditorViewportSettings>()->bUsePowerOf2SnapSize )
-	{
-		BSPTexelScale=128.0f;
-	}
-	UModel::SetGlobalBSPTexelScale(BSPTexelScale);
-
-	GLog->EnableBacklog( false );
-
-	{
-		// avoid doing this every time, create a list of classes that derive from AVolume
-		TArray<UClass*> VolumeClasses;
-		for (TObjectIterator<UClass> ObjectIt; ObjectIt; ++ObjectIt)
-		{
-			UClass* TestClass = *ObjectIt;
-			// we want classes derived from AVolume, but not AVolume itself
-			if ( TestClass->IsChildOf(AVolume::StaticClass()) && TestClass != AVolume::StaticClass() )
-			{
-				VolumeClasses.Add( TestClass );
-			}
-		}
 
-		FAssetData NoAssetData;
-
-		// Create array of ActorFactory instances.
-		for (TObjectIterator<UClass> ObjectIt; ObjectIt; ++ObjectIt)
-		{
-			UClass* TestClass = *ObjectIt;
-			if (TestClass->IsChildOf(UActorFactory::StaticClass()))
-			{
-				if (!TestClass->HasAnyClassFlags(CLASS_Abstract))
+				bValidSourceFilename = true;
+				if ( bAskForNewFileIfMissing && bMissingFiles )
 				{
-					// if the factory is a volume shape factory we create an instance for all volume types
-					if ( TestClass->IsChildOf(UActorFactoryBoxVolume::StaticClass()) ||
-						 TestClass->IsChildOf(UActorFactorySphereVolume::StaticClass()) ||
-						 TestClass->IsChildOf(UActorFactoryCylinderVolume::StaticClass()) )
+					GetNewReimportPath(Obj, SourceFilenames);
+					if ( SourceFilenames.Num() == 0 )
 					{
-						for ( int32 i=0; i < VolumeClasses.Num(); i++ )
-						{
-							UActorFactory* NewFactory = NewObject<UActorFactory>(GetTransientPackage(), TestClass);
-							check(NewFactory);
-							NewFactory->NewActorClass = VolumeClasses[i];
-							ActorFactories.Add(NewFactory);
-						}
+						// Failed to specify a new filename. Don't show a notification of the failure since the user exited on his own
+						bValidSourceFilename = false;
+						bShowNotification = false;
 					}
 					else
 					{
-						UActorFactory* NewFactory = NewObject<UActorFactory>(GetTransientPackage(), TestClass);
-						check(NewFactory);
-						ActorFactories.Add(NewFactory);
+						// A new filename was supplied, update the path
+						Handlers[ HandlerIndex ]->SetReimportPaths(Obj, SourceFilenames);
 					}
 				}
-			}
-		}
-	}
-
-	// Used for sorting ActorFactory classes.
-	struct FCompareUActorFactoryByMenuPriority
-	{
-		FORCEINLINE bool operator()(const UActorFactory& A, const UActorFactory& B) const
-		{
-			if (B.MenuPriority == A.MenuPriority)
-			{
-				if ( A.GetClass() != UActorFactory::StaticClass() && B.IsA(A.GetClass()) )
-				{
-					return false;
-				}
-				else if ( B.GetClass() != UActorFactory::StaticClass() && A.IsA(B.GetClass()) )
-				{
-					return true;
-				}
-				else
-				{
-					return A.GetClass()->GetName() < B.GetClass()->GetName();
-				}
-			}
-			else 
-			{
-				return B.MenuPriority < A.MenuPriority;
-			}
-		}
-	};
-	// Sort by menu priority.
-	ActorFactories.Sort( FCompareUActorFactoryByMenuPriority() );
-
-	// Load game user settings and apply
-	auto GameUserSettings = GetGameUserSettings();
-	if (GameUserSettings)
-	{
-		GameUserSettings->LoadSettings();
-		GameUserSettings->ApplySettings(true);
-	}
-
-	UEditorStyleSettings* Settings = GetMutableDefault<UEditorStyleSettings>();
-	Settings->OnSettingChanged().AddUObject(this, &UEditorEngine::HandleSettingChanged);
-
-	// Purge garbage.
-	Cleanse( false, 0, NSLOCTEXT("UnrealEd", "Startup", "Startup") );
-
-	FEditorCommandLineUtils::ProcessEditorCommands(FCommandLine::Get());
-
-	// for IsInitialized()
-	bIsInitialized = true;
-};
-
-
-void UEditorEngine::InitBuilderBrush( UWorld* InWorld )
-{
-	check( InWorld );
-	const bool bOldDirtyState = InWorld->GetCurrentLevel()->GetOutermost()->IsDirty();
-
-	// For additive geometry mode, make the builder brush a small 256x256x256 cube so its visible.
-	const int32 CubeSize = 256;
-	UCubeBuilder* CubeBuilder = NewObject<UCubeBuilder>();
-	CubeBuilder->X = CubeSize;
-	CubeBuilder->Y = CubeSize;
-	CubeBuilder->Z = CubeSize;
-	CubeBuilder->Build( InWorld );
-
-	// Restore the level's dirty state, so that setting a builder brush won't mark the map as dirty.
-	if (!bOldDirtyState)
-	{
-		InWorld->GetCurrentLevel()->GetOutermost()->SetDirtyFlag( bOldDirtyState );
-	}
-}
-
-void UEditorEngine::BroadcastObjectReimported(UObject* InObject)
-{
-	ObjectReimportedEvent.Broadcast(InObject);
-	FEditorDelegates::OnAssetReimport.Broadcast(InObject);
-}
-
-void UEditorEngine::FinishDestroy()
-{
-	if ( !HasAnyFlags(RF_ClassDefaultObject) )
-	{
-		// this needs to be already cleaned up
-		check(PlayWorld == NULL);
-
-		// Unregister events
-		FEditorDelegates::MapChange.RemoveAll(this);
-		FCoreDelegates::ModalErrorMessage.Unbind();
-		FCoreUObjectDelegates::ShouldLoadOnTop.Unbind();
-		FCoreDelegates::PreWorldOriginOffset.RemoveAll(this);
-		FCoreUObjectDelegates::OnAssetLoaded.RemoveAll(this);
-		FWorldDelegates::LevelAddedToWorld.RemoveAll(this);
-		FWorldDelegates::LevelRemovedFromWorld.RemoveAll(this);
-		FLevelStreamingGCHelper::OnGCStreamedOutLevels.RemoveAll(this);
-		GetMutableDefault<UEditorStyleSettings>()->OnSettingChanged().RemoveAll(this);
-
-		EditorClearComponents();
-		UWorld* World = GWorld;
-		if( World != NULL )
-		{
-			World->CleanupWorld();
-		}
-	
-		// Shut down transaction tracking system.
-		if( Trans )
-		{
-			if( GUndo )
-			{
-				UE_LOG(LogEditor, Warning, TEXT("Warning: A transaction is active") );
-			}
-			ResetTransaction( NSLOCTEXT("UnrealEd", "Shutdown", "Shutdown") );
-		}
-
-		// Destroy selection sets.
-		PrivateDestroySelectedSets();
-
-		extern void TearDownDistanceFieldBuildNotification();
-		TearDownDistanceFieldBuildNotification();
-
-		// Remove editor array from root.
-		UE_LOG(LogExit, Log, TEXT("Editor shut down") );
-	}
-
-	Super::FinishDestroy();
-}
-
-void UEditorEngine::AddReferencedObjects(UObject* InThis, FReferenceCollector& Collector)
-{
-	UEditorEngine* This = CastChecked<UEditorEngine>(InThis);
-	// Serialize viewport clients.
-	for(uint32 ViewportIndex = 0;ViewportIndex < (uint32)This->AllViewportClients.Num(); ViewportIndex++)
-	{
-		This->AllViewportClients[ViewportIndex]->AddReferencedObjects( Collector );
-	}
-
-	// Serialize ActorFactories
-	for( int32 Index = 0; Index < This->ActorFactories.Num(); Index++ )
-	{
-		Collector.AddReferencedObject( This->ActorFactories[ Index ], This );
-	}
-
-	Super::AddReferencedObjects( This, Collector );
-}
-
-void UEditorEngine::Tick( float DeltaSeconds, bool bIdleMode )
-{
-	NETWORK_PROFILER(GNetworkProfiler.TrackFrameBegin());
-
-	UWorld* CurrentGWorld = GWorld;
-	check( CurrentGWorld );
-	check( CurrentGWorld != PlayWorld || bIsSimulatingInEditor );
-
-	// Clear out the list of objects modified this frame, used for OnObjectModified notification.
-	FCoreUObjectDelegates::ObjectsModifiedThisFrame.Empty();
-
-	// Always ensure we've got adequate slack for any worlds that are going to get created in this frame so that
-	// our EditorContext reference doesn't get invalidated
-	WorldList.Reserve(WorldList.Num() + 10);
-
-	FWorldContext& EditorContext = GetEditorWorldContext();
-	check( CurrentGWorld == EditorContext.World() );
-
-	// was there a reregister requested last frame?
-	if (bHasPendingGlobalReregister)
-	{
-		// make sure outstanding deletion has completed before the reregister
-		CollectGarbage(GARBAGE_COLLECTION_KEEPFLAGS);
-
-		//Only reregister actors whose replacement primitive is a child of the global reregister list
-		FGlobalComponentReregisterContext Reregister(ActorsForGlobalReregister);
-		ActorsForGlobalReregister.Empty();
-
-		bHasPendingGlobalReregister = false;
-	}
-
-	// early in the Tick() to get the callbacks for cvar changes called
-	IConsoleManager::Get().CallAllConsoleVariableSinks();
-
-	// Tick the hot reload interface
-	IHotReloadInterface* HotReload = IHotReloadInterface::GetPtr();
-	if(HotReload != nullptr)
-	{
-		HotReload->Tick();
-	}
-
-	// Tick the remote config IO manager
-	FRemoteConfigAsyncTaskManager::Get()->Tick();
-
-	// Clean up the game viewports that have been closed.
-	CleanupGameViewport();
-
-	// If all viewports closed, close the current play level.
-	if( PlayWorld && !bIsSimulatingInEditor )
-	{
-		for (auto It=WorldList.CreateIterator(); It; ++It)
-		{
-			// For now, kill PIE session if any of the viewports are closed
-			if (It->WorldType == EWorldType::PIE && It->GameViewport == NULL && !It->RunAsDedicated && !It->bWaitingOnOnlineSubsystem)
-			{
-				EndPlayMap();
-				break;
-			}
-		}
-	}
-
-
-	// Potentially rebuilds the streaming data.
-	EditorContext.World()->ConditionallyBuildStreamingData();
-
-	// Update the timer manager
-	TimerManager->Tick(DeltaSeconds);
-
-	// Update subsystems.
-	{
-		// This assumes that UObject::StaticTick only calls ProcessAsyncLoading.	
-		StaticTick(DeltaSeconds, bAsyncLoadingUseFullTimeLimit, AsyncLoadingTimeLimit / 1000.f);
-	}
-
-	// Look for realtime flags.
-	bool IsRealtime = false;
-
-	// True if a viewport has realtime audio	// If any realtime audio is enabled in the editor
-	bool bAudioIsRealtime = GetDefault<ULevelEditorMiscSettings>()->bEnableRealTimeAudio;
-
-	// By default we tick the editor world.  
-	// When in PIE if we are in immersive we do not tick the editor world unless there is a visible editor viewport.
-	bool bShouldTickEditorWorld = true;
-
-	//@todo Multiple Worlds: Do we need to consider what world we are in here?
-
-	// Find which viewport has audio focus, i.e. gets to set the listener location
-	// Priorities are:
-	//  Active perspective realtime view
-	//	> Any realtime perspective view (first encountered)
-	//	> Active perspective view
-	//	> Any perspective view (first encountered)
-	FEditorViewportClient* AudioFocusViewportClient = NULL;
-	{
-		FEditorViewportClient* BestRealtimePerspViewport = NULL;
-		FEditorViewportClient* BestPerspViewport = NULL;
-
-		for( int32 ViewportIndex = 0; ViewportIndex < AllViewportClients.Num(); ViewportIndex++ )
-		{
-			FEditorViewportClient* const ViewportClient = AllViewportClients[ ViewportIndex ];
-
-			// clear any previous audio focus flags
-			ViewportClient->ClearAudioFocus();
-
-			if (ViewportClient->IsPerspective())
-			{
-				if (ViewportClient->IsRealtime())
-				{
-					if (ViewportClient->Viewport && ViewportClient->Viewport->HasFocus())
-					{
-						// active realtime perspective -- use this and be finished
-						BestRealtimePerspViewport = ViewportClient;
-						break;
-					}
-					else if (BestRealtimePerspViewport == NULL)
-					{
-						// save this
-						BestRealtimePerspViewport = ViewportClient;
-					}
-				}
-				else 
-				{
-					if (ViewportClient->Viewport && ViewportClient->Viewport->HasFocus())
-					{
-						// active non-realtime perspective -- use this
-						BestPerspViewport = ViewportClient;
-					}
-					else if (BestPerspViewport == NULL)
-					{
-						// save this
-						BestPerspViewport = ViewportClient;
-					}
-
-				}
-			}
-		}
-
-		// choose realtime if set.  note this could still be null.
-		AudioFocusViewportClient = BestRealtimePerspViewport ? BestRealtimePerspViewport : BestPerspViewport;
-	}
-	// tell viewportclient it has audio focus
-	if (AudioFocusViewportClient)
-	{
-		AudioFocusViewportClient->SetAudioFocus();
-
-		// override realtime setting if viewport chooses (i.e. for matinee preview)
-		if (AudioFocusViewportClient->IsForcedRealtimeAudio())
-		{
-			bAudioIsRealtime = true;
-		}
-	}
-
-	// Find realtime and visibility settings on all viewport clients
-	for( int32 ViewportIndex = 0; ViewportIndex < AllViewportClients.Num(); ViewportIndex++ )
-	{
-		FEditorViewportClient* const ViewportClient = AllViewportClients[ ViewportIndex ];
-
-		if( PlayWorld && ViewportClient->IsVisible() )
-		{
-			if( ViewportClient->IsInImmersiveViewport() )
-			{
-				// if a viewport client is immersive then by default we do not tick the editor world during PIE unless there is a visible editor world viewport
-				bShouldTickEditorWorld = false;
-			}
-			else
-			{
-				// If the viewport is not immersive but still visible while we have a play world then we need to tick the editor world
-				bShouldTickEditorWorld = true;
-			}
-		}
-
-		if( ViewportClient->GetScene() == EditorContext.World()->Scene )
-		{
-			if( ViewportClient->IsRealtime() )
-			{
-				IsRealtime = true;
-			}
-		}
-	}
-
-	// Find out if the editor has focus. Audio should only play if the editor has focus.
-	const bool bHasFocus = FPlatformProcess::IsThisApplicationForeground();
-
-	if (bHasFocus || GetDefault<ULevelEditorMiscSettings>()->bAllowBackgroundAudio)
-	{
-		if (!PlayWorld)
-		{
-			// Adjust the global volume multiplier if the window has focus and there is no pie world or no viewport overriding audio.
-			FApp::SetVolumeMultiplier( GetDefault<ULevelEditorMiscSettings>()->EditorVolumeLevel );
-		}
-		else
-		{
-			// If there is currently a pie world a viewport is overriding audio settings do not adjust the volume.
-			FApp::SetVolumeMultiplier( 1.0f );
-		}
-	}
-
-	// Tick any editor FTickableEditorObject dervived classes
-	FTickableEditorObject::TickObjects( DeltaSeconds );
-
-	// Tick the asset registry
-	FAssetRegistryModule::TickAssetRegistry(DeltaSeconds);
-
-	static FName SourceCodeAccessName("SourceCodeAccess");
-	ISourceCodeAccessModule& SourceCodeAccessModule = FModuleManager::LoadModuleChecked<ISourceCodeAccessModule>(SourceCodeAccessName);
-	SourceCodeAccessModule.GetAccessor().Tick(DeltaSeconds);
-
-	// tick the directory watcher
-	// @todo: Put me into an FTicker that is created when the DW module is loaded
-	static FName DirectoryWatcherName("DirectoryWatcher");
-	FDirectoryWatcherModule& DirectoryWatcherModule = FModuleManager::Get().LoadModuleChecked<FDirectoryWatcherModule>(DirectoryWatcherName);
-	DirectoryWatcherModule.Get()->Tick(DeltaSeconds);
-
-	if( bShouldTickEditorWorld )
-	{ 
-		// Tick level.
-		ELevelTick TickType = IsRealtime ? LEVELTICK_ViewportsOnly : LEVELTICK_TimeOnly;
-
-		//EditorContext.World()->FXSystem->Resume();
-		// Note: Still allowing the FX system to tick so particle systems dont restart after entering/leaving responsive mode
-		if( FSlateThrottleManager::Get().IsAllowingExpensiveTasks() )
-		{
-			FKismetDebugUtilities::NotifyDebuggerOfStartOfGameFrame(EditorContext.World());
-			EditorContext.World()->Tick(TickType, DeltaSeconds);
-			FKismetDebugUtilities::NotifyDebuggerOfEndOfGameFrame(EditorContext.World());
-		}
-	}
-	else 
-	{
-		//EditorContext.World()->FXSystem->Suspend();
-	}
-
-
-	// Perform editor level streaming previs if no PIE session is currently in progress.
-	if( !PlayWorld )
-	{
-		for ( int32 ClientIndex = 0 ; ClientIndex < LevelViewportClients.Num() ; ++ClientIndex )
-		{
-			FLevelEditorViewportClient* ViewportClient = LevelViewportClients[ClientIndex];
 
-			// Previs level streaming volumes in the Editor.
-			if ( ViewportClient->IsPerspective() && GetDefault<ULevelEditorViewportSettings>()->bLevelStreamingVolumePrevis )
-			{
-				bool bProcessViewer = false;
-				const FVector& ViewLocation = ViewportClient->GetViewLocation();
-
-				// Iterate over streaming levels and compute whether the ViewLocation is in their associated volumes.
-				TMap<ALevelStreamingVolume*, bool> VolumeMap;
-
-				for( int32 LevelIndex = 0 ; LevelIndex < EditorContext.World()->StreamingLevels.Num() ; ++LevelIndex )
+				if ( bValidSourceFilename )
 				{
-					ULevelStreaming* StreamingLevel = EditorContext.World()->StreamingLevels[LevelIndex];
-					if( StreamingLevel )
+					// Do the reimport
+					EReimportResult::Type Result = Handlers[ HandlerIndex ]->Reimport( Obj );
+					if( Result == EReimportResult::Succeeded )
 					{
-						// Assume the streaming level is invisible until we find otherwise.
-						bool bStreamingLevelShouldBeVisible = false;
-
-						// We're not going to change level visibility unless we encounter at least one
-						// volume associated with the level.
-						bool bFoundValidVolume = false;
-
-						// For each streaming volume associated with this level . . .
-						for ( int32 VolumeIndex = 0 ; VolumeIndex < StreamingLevel->EditorStreamingVolumes.Num() ; ++VolumeIndex )
-						{
-							ALevelStreamingVolume* StreamingVolume = StreamingLevel->EditorStreamingVolumes[VolumeIndex];
-							if ( StreamingVolume && !StreamingVolume->bDisabled )
-							{
-								bFoundValidVolume = true;
-
-								bool bViewpointInVolume;
-								bool* bResult = VolumeMap.Find(StreamingVolume);
-								if ( bResult )
-								{
-									// This volume has already been considered for another level.
-									bViewpointInVolume = *bResult;
-								}
-								else
-								{
-									// Compute whether the viewpoint is inside the volume and cache the result.
-									bViewpointInVolume = StreamingVolume->EncompassesPoint( ViewLocation );							
-
-								
-									VolumeMap.Add( StreamingVolume, bViewpointInVolume );
-								}
-
-								// Halt when we find a volume associated with the level that the viewpoint is in.
-								if ( bViewpointInVolume )
-								{
-									bStreamingLevelShouldBeVisible = true;
-									break;
-								}
-							}
-						}
-
-						// Set the streaming level visibility status if we encountered at least one volume.
-						if ( bFoundValidVolume && StreamingLevel->bShouldBeVisibleInEditor != bStreamingLevelShouldBeVisible )
+						Obj->PostEditChange();
+						GEditor->BroadcastObjectReimported(Obj);
+						if (FEngineAnalytics::IsAvailable())
 						{
-							StreamingLevel->bShouldBeVisibleInEditor = bStreamingLevelShouldBeVisible;
-							bProcessViewer = true;
+							TArray<FAnalyticsEventAttribute> Attributes;
+							Attributes.Add( FAnalyticsEventAttribute( TEXT( "ObjectType" ), Obj->GetClass()->GetName() ) );
+							FEngineAnalytics::GetProvider().RecordEvent(TEXT("Editor.Usage.AssetReimported"), Attributes);
 						}
+						bSuccess = true;
 					}
-				}
-				
-				// Simulate world composition streaming while in editor world
-				if (EditorContext.World()->WorldComposition)
-				{
-					if (EditorContext.World()->WorldComposition->UpdateEditorStreamingState(ViewLocation))
+					else if( Result == EReimportResult::Cancelled )
 					{
-						bProcessViewer = true;
+						bShowNotification = false;
 					}
 				}
 				
-				// Call UpdateLevelStreaming if the visibility of any streaming levels was modified.
-				if ( bProcessViewer )
-				{
-					EditorContext.World()->UpdateLevelStreaming();
-					FEditorDelegates::RefreshPrimitiveStatsBrowser.Broadcast();
-				}
 				break;
 			}
 		}
-	}
-
-	// kick off a "Play From Here" if we got one
-	if (bIsPlayWorldQueued)
-	{
-		StartQueuedPlayMapRequest();
-	}
-	else if( bIsToggleBetweenPIEandSIEQueued )
-	{
-		ToggleBetweenPIEandSIE();
-	}
-
-	static bool bFirstTick = true;
-
-	// Skip updating reflection captures on the first update as the level will not be ready to display
-	if (!bFirstTick)
-	{
-		// Update sky light first because sky diffuse will be visible in reflection capture indirect specular
-		USkyLightComponent::UpdateSkyCaptureContents(EditorContext.World());
-		UReflectionCaptureComponent::UpdateReflectionCaptureContents(EditorContext.World());
-	}
 
-	// if we have the side-by-side world for "Play From Here", tick it unless we are ensuring slate is responsive
-	if( FSlateThrottleManager::Get().IsAllowingExpensiveTasks() )
-	{
-		for (auto ContextIt = WorldList.CreateIterator(); ContextIt; ++ContextIt)
+		if ( bShowNotification )
 		{
-			FWorldContext &PieContext = *ContextIt;
-			if (PieContext.WorldType != EWorldType::PIE || PieContext.World() == NULL)
-			{
-				continue;
-			}
-
-			GPlayInEditorID = PieContext.PIEInstance;
-
-			PlayWorld = PieContext.World();
-			GameViewport = PieContext.GameViewport;
-
-			UWorld* OldGWorld = NULL;
-			// Use the PlayWorld as the GWorld, because who knows what will happen in the Tick.
-			OldGWorld = SetPlayInEditorWorld( PlayWorld );
-
-			// Tick all travel and Pending NetGames (Seamless, server, client)
-			TickWorldTravel(PieContext, DeltaSeconds);
-
-			// Updates 'connecting' message in PIE network games
-			UpdateTransitionType(PlayWorld);
-
-			// Update streaming for dedicated servers in PIE
-			if (PieContext.RunAsDedicated)
-			{
-				SCOPE_CYCLE_COUNTER(STAT_UpdateLevelStreaming);
-				PlayWorld->UpdateLevelStreaming();
-			}
-
-			// Release mouse if the game is paused. The low level input code might ignore the request when e.g. in fullscreen mode.
-			if ( GameViewport != NULL && GameViewport->Viewport != NULL )
-			{
-				// Decide whether to drop high detail because of frame rate
-				GameViewport->SetDropDetail(DeltaSeconds);
-			}
-
-			if (!bFirstTick)
-			{
-				// Update sky light first because sky diffuse will be visible in reflection capture indirect specular
-				USkyLightComponent::UpdateSkyCaptureContents(PlayWorld);
-				UReflectionCaptureComponent::UpdateReflectionCaptureContents(PlayWorld);
-			}
-
-			// Update the level.
-			GameCycles=0;
-			CLOCK_CYCLES(GameCycles);
-
+			// Send a notification of the results
+			FText NotificationText;
+			if ( bSuccess )
 			{
-				// So that hierarchical stats work in PIE
-				SCOPE_CYCLE_COUNTER(STAT_FrameTime);
-
-				FKismetDebugUtilities::NotifyDebuggerOfStartOfGameFrame(PieContext.World());
-
-				static TArray< TWeakObjectPtr<AActor> > RecordedActors;
-				RecordedActors.Reset();
-
-				// Check to see if we want to use sequencer's live recording feature
-				bool bIsRecordingActive = false;
-				GetActorRecordingStateEvent.Broadcast( bIsRecordingActive );
-				if( bIsRecordingActive )
+				if ( bValidSourceFilename )
 				{
-					// @todo sequencer livecapture: How do we capture the destruction of actors? (needs to hide the spawned puppet actor, or destroy it)
-					// @todo sequencer livecapture: Actor parenting state is not captured or retained on puppets
-					// @todo sequencer livecapture: Needs to capture state besides transforms (animation, audio, property changes, etc.)
+					const FString FirstLeafFilename = FPaths::GetCleanFilename(SourceFilenames[0]);
 
-					// @todo sequencer livecapture: Hacky test code for sequencer live capture feature
-					for( FActorIterator ActorIter( PlayWorld ); ActorIter; ++ActorIter )
+					if ( SourceFilenames.Num() == 1 )
 					{
-						AActor* Actor = *ActorIter;
-
-						// @todo sequencer livecapture: Restrict to certain actor types for now, just for testing
-						if( Actor->IsA(ASkeletalMeshActor::StaticClass()) || (Actor->IsA(AStaticMeshActor::StaticClass()) && Actor->IsRootComponentMovable()) )
-						{
-							GEditor->BroadcastBeginObjectMovement( *Actor );
-							RecordedActors.Add( Actor );
-						}
-					}				
-				}
-
-				// tick the level
-				PieContext.World()->Tick( LEVELTICK_All, DeltaSeconds );
-
-				if( bIsRecordingActive )
-				{
-					for( auto RecordedActorIter( RecordedActors.CreateIterator() ); RecordedActorIter; ++RecordedActorIter )
+						FFormatNamedArguments Args;
+						Args.Add( TEXT("ObjectName"), FText::FromString( Obj->GetName() ) );
+						Args.Add( TEXT("ObjectType"), FText::FromString( Obj->GetClass()->GetName() ) );
+						Args.Add( TEXT("SourceFile"), FText::FromString( FirstLeafFilename ) );
+						NotificationText = FText::Format( LOCTEXT("ReimportSuccessfulFrom", "Successfully Reimported: {ObjectName} ({ObjectType}) from file ({SourceFile})"), Args );
+					}
+					else
 					{
-						AActor* Actor = RecordedActorIter->Get();
-						if( Actor != NULL )
-						{
-							GEditor->BroadcastEndObjectMovement( *Actor );
-						}
-					}				
-				}
-
-				FKismetDebugUtilities::NotifyDebuggerOfEndOfGameFrame(PieContext.World());
-			}
-
-			UNCLOCK_CYCLES(GameCycles);
-
-			// Tick the viewports.
-			if ( GameViewport != NULL )
-			{
-				GameViewport->Tick(DeltaSeconds);
-			}
-
-			// Pop the world
-			RestoreEditorWorld( OldGWorld );
-		}
-	}
-
-	if (bFirstTick)
-	{
-		bFirstTick = false;
-	}
-
-	GPlayInEditorID = -1;
-
-	// Clean up any game viewports that may have been closed during the level tick (eg by Kismet).
-	CleanupGameViewport();
-
-	// If all viewports closed, close the current play level.
-	if( GameViewport == NULL && PlayWorld && !bIsSimulatingInEditor )
-	{
-		FWorldContext& PieWorldContext = GetWorldContextFromWorldChecked(PlayWorld);
-		if (!PieWorldContext.RunAsDedicated && !PieWorldContext.bWaitingOnOnlineSubsystem)
-		{
-			EndPlayMap();
-		}
-	}
-
-	// Update viewports.
-
-	for(int32 ViewportIndex = 0;ViewportIndex < AllViewportClients.Num();ViewportIndex++)
-	{
-		FEditorViewportClient* ViewportClient = AllViewportClients[ ViewportIndex ];
-
-		// When throttling tick only viewports which need to be redrawn (they have been manually invalidated)
-		if( ( FSlateThrottleManager::Get().IsAllowingExpensiveTasks() || ViewportClient->bNeedsRedraw ) && ViewportClient->IsVisible() )
-		{
-			// Switch to the correct world for the client before it ticks
-			FScopedConditionalWorldSwitcher WorldSwitcher( ViewportClient );
-
-			ViewportClient->Tick(DeltaSeconds);
-		}
-	}
-	
-	bool bIsMouseOverAnyLevelViewport = false;
-
-	//Do this check separate to the above loop as the ViewportClient may no longer be valid after we have ticked it
-	for(int32 ViewportIndex = 0;ViewportIndex < LevelViewportClients.Num();ViewportIndex++)
-	{
-		FLevelEditorViewportClient* ViewportClient = LevelViewportClients[ ViewportIndex ];
-		FViewport* Viewport = ViewportClient->Viewport;
-
-		// Keep track of whether the mouse cursor is over any level viewports
-		if( Viewport != NULL )
-		{
-			const int32 MouseX = Viewport->GetMouseX();
-			const int32 MouseY = Viewport->GetMouseY();
-			if( MouseX >= 0 && MouseY >= 0 && MouseX < (int32)Viewport->GetSizeXY().X && MouseY < (int32)Viewport->GetSizeXY().Y )
-			{
-				bIsMouseOverAnyLevelViewport = true;
-				break;
-			}
-		}
-	}
-
-	// If the cursor is outside all level viewports, then clear the hover effect
-	if( !bIsMouseOverAnyLevelViewport )
-	{
-		FLevelEditorViewportClient::ClearHoverFromObjects();
-	}
-
-
-	// Commit changes to the BSP model.
-	EditorContext.World()->CommitModelSurfaces();
-	EditorContext.World()->SendAllEndOfFrameUpdates();
-
-	bool bUpdateLinkedOrthoViewports = false;
-	/////////////////////////////
-	// Redraw viewports.
-
-	// Do not redraw if the application is hidden
-	bool bAllWindowsHidden = !bHasFocus && GEditor->AreAllWindowsHidden();
-	if( !bAllWindowsHidden )
-	{
-		// Render view parents, then view children.
-		bool bEditorFrameNonRealtimeViewportDrawn = false;
-		if (GCurrentLevelEditingViewportClient && GCurrentLevelEditingViewportClient->IsVisible())
-		{
-			bool bAllowNonRealtimeViewports = true;
-			bool bWasNonRealtimeViewportDraw = UpdateSingleViewportClient(GCurrentLevelEditingViewportClient, bAllowNonRealtimeViewports, bUpdateLinkedOrthoViewports);
-			if (GCurrentLevelEditingViewportClient->IsLevelEditorClient())
-			{
-				bEditorFrameNonRealtimeViewportDrawn |= bWasNonRealtimeViewportDraw;
-			}
-		}
-		for (int32 bRenderingChildren = 0; bRenderingChildren < 2; bRenderingChildren++)
-		{
-			for (int32 ViewportIndex = 0; ViewportIndex < AllViewportClients.Num(); ViewportIndex++)
-			{
-				FEditorViewportClient* ViewportClient = AllViewportClients[ViewportIndex];
-				if (ViewportClient == GCurrentLevelEditingViewportClient)
-				{
-					//already given this window a chance to update
-					continue;
-				}
-
-				if ( ViewportClient->IsVisible() )
-				{
-					// Only update ortho viewports if that mode is turned on, the viewport client we are about to update is orthographic and the current editing viewport is orthographic and tracking mouse movement.
-					bUpdateLinkedOrthoViewports = GetDefault<ULevelEditorViewportSettings>()->bUseLinkedOrthographicViewports && ViewportClient->IsOrtho() && GCurrentLevelEditingViewportClient && GCurrentLevelEditingViewportClient->IsOrtho() && GCurrentLevelEditingViewportClient->IsTracking();
-
-					const bool bIsViewParent = ViewportClient->ViewState.GetReference()->IsViewParent();
-					if ((bRenderingChildren && !bIsViewParent) ||
-						(!bRenderingChildren && bIsViewParent) || bUpdateLinkedOrthoViewports)
-					{
-						//if we haven't drawn a non-realtime viewport OR not one of the main viewports
-						bool bAllowNonRealtimeViewports = (!bEditorFrameNonRealtimeViewportDrawn) || !(ViewportClient->IsLevelEditorClient());
-						bool bWasNonRealtimeViewportDrawn = UpdateSingleViewportClient(ViewportClient, bAllowNonRealtimeViewports, bUpdateLinkedOrthoViewports);
-						if (ViewportClient->IsLevelEditorClient())
-						{
-							bEditorFrameNonRealtimeViewportDrawn |= bWasNonRealtimeViewportDrawn;
-						}
-					}
-				}
-			}
-		}
-	}
-
-	ISourceControlModule::Get().Tick();
-
-	if( FSlateThrottleManager::Get().IsAllowingExpensiveTasks() )
-	{
-		for (auto ContextIt = WorldList.CreateIterator(); ContextIt; ++ContextIt)
-		{
-			FWorldContext &PieContext = *ContextIt;
-			if (PieContext.WorldType != EWorldType::PIE)
-			{
-				continue;
-			}
-
-			PlayWorld = PieContext.World();
-			GameViewport = PieContext.GameViewport;
-
-			// Render playworld. This needs to happen after the other viewports for screenshots to work correctly in PIE.
-			if(PlayWorld && GameViewport && !bIsSimulatingInEditor)
-			{
-				// Use the PlayWorld as the GWorld, because who knows what will happen in the Tick.
-				UWorld* OldGWorld = SetPlayInEditorWorld( PlayWorld );
-				GPlayInEditorID = PieContext.PIEInstance;
-
-				// Render everything.
-				GameViewport->LayoutPlayers();
-				check(GameViewport->Viewport);
-				GameViewport->Viewport->Draw();
-
-				// Pop the world
-				RestoreEditorWorld( OldGWorld );
-				GPlayInEditorID = -1;
-			}
-		}
-	}
-
-	// Update resource streaming after both regular Editor viewports and PIE had a chance to add viewers.
-	IStreamingManager::Get().Tick(DeltaSeconds);
-
-	// Update Audio. This needs to occur after rendering as the rendering code updates the listener position.
-	FAudioDeviceManager* AudioDeviceManager = GEngine->GetAudioDeviceManager();
-	if (AudioDeviceManager)
-	{
-		UWorld* OldGWorld = NULL;
-		if (PlayWorld)
-		{
-			// Use the PlayWorld as the GWorld if we're using PIE.
-			OldGWorld = SetPlayInEditorWorld(PlayWorld);
-		}
-
-		// Update audio device.
-		AudioDeviceManager->UpdateActiveAudioDevices((!PlayWorld && bAudioIsRealtime) || (PlayWorld && !PlayWorld->IsPaused()));
-
-		if (PlayWorld)
-		{
-			// Pop the world.
-			RestoreEditorWorld(OldGWorld);
-		}
-	}
-
-	// Update constraints if dirtied.
-	EditorContext.World()->UpdateConstraintActors();
-
-	{
-		// rendering thread commands
-
-		ENQUEUE_UNIQUE_RENDER_COMMAND_TWOPARAMETER(
-			TickRenderingTimer,
-			bool, bPauseRenderingRealtimeClock, GPauseRenderingRealtimeClock,
-			float, DeltaTime, DeltaSeconds,
-		{
-			if(!bPauseRenderingRealtimeClock)
-			{
-				// Tick the GRenderingRealtimeClock, unless it's paused
-				GRenderingRealtimeClock.Tick(DeltaTime);
-			}
-			GetRendererModule().TickRenderTargetPool();
-		});
-	}
-
-	// After the play world has ticked, see if a request was made to end pie
-	if( bRequestEndPlayMapQueued )
-	{
-		EndPlayMap();
-	}
-
-	FUnrealEdMisc::Get().TickAssetAnalytics();
-
-	FUnrealEdMisc::Get().TickPerformanceAnalytics();
-
-	// If the fadeout animation has completed for the undo/redo notification item, allow it to be deleted
-	if(UndoRedoNotificationItem.IsValid() && UndoRedoNotificationItem->GetCompletionState() == SNotificationItem::CS_None)
-	{
-		UndoRedoNotificationItem.Reset();
-	}
-}
-
-float UEditorEngine::GetMaxTickRate( float DeltaTime, bool bAllowFrameRateSmoothing ) const
-{
-	float MaxTickRate = 0.0f;
-	if( !ShouldThrottleCPUUsage() )
-	{
-		// do not limit fps in VR Preview mode
-		if (bUseVRPreviewForPlayWorld)
-		{
-			return 0.0f;
-		}
-		const float SuperMaxTickRate = Super::GetMaxTickRate( DeltaTime, bAllowFrameRateSmoothing );
-		if( SuperMaxTickRate != 0.0f )
-		{
-			return SuperMaxTickRate;
-		}
-
-		// Clamp editor frame rate, even if smoothing is disabled
-		if( !bSmoothFrameRate && GIsEditor && !GIsPlayInEditorWorld )
-		{
-			MaxTickRate = 1.0f / DeltaTime;
-			if (SmoothedFrameRateRange.HasLowerBound())
-			{
-				MaxTickRate = FMath::Max(MaxTickRate, SmoothedFrameRateRange.GetLowerBoundValue());
-			}
-			if (SmoothedFrameRateRange.HasUpperBound())
-			{
-				MaxTickRate = FMath::Min(MaxTickRate, SmoothedFrameRateRange.GetUpperBoundValue());
-			}
-		}
-
-		// Laptops should throttle to 60 hz in editor to reduce battery drain
-		static const auto CVarDontLimitOnBattery = IConsoleManager::Get().FindTConsoleVariableDataInt(TEXT("r.DontLimitOnBattery"));
-		const bool bLimitOnBattery = (FPlatformMisc::IsRunningOnBattery() && CVarDontLimitOnBattery->GetValueOnGameThread() == 0);
-		if( bLimitOnBattery )
-		{
-			MaxTickRate = 60.0f;
-		}
-	}
-	else
-	{
-		MaxTickRate = 3.0f;
-	}
-
-	return MaxTickRate;
-}
-
-const UEditorUserSettings& UEditorEngine::GetEditorUserSettings() const
-{
-	if (EditorUserSettings == NULL)
-	{
-		auto ConstThis = const_cast< UEditorEngine* >( this );	// Hack because Header Generator doesn't yet support mutable keyword
-		ConstThis->EditorUserSettings = NewObject<UEditorUserSettings>();
-	}
-	return *EditorUserSettings;
-}
-
-UEditorUserSettings& UEditorEngine::AccessEditorUserSettings()
-{
-	if (EditorUserSettings == NULL)
-	{
-		EditorUserSettings = NewObject<UEditorUserSettings>();
-	}
-	return *EditorUserSettings;
-}
-
-void UEditorEngine::SaveEditorUserSettings()
-{
-	if (!FUnrealEdMisc::Get().IsDeletePreferences())
-	{
-		AccessEditorUserSettings().SaveConfig();
-	}
-}
-
-const UEditorGameAgnosticSettings& UEditorEngine::GetGameAgnosticSettings() const
-{
-	if (GameAgnosticSettings == NULL)
-	{
-		auto ConstThis = const_cast< UEditorEngine* >( this );	// Hack because Header Generator doesn't yet support mutable keyword
-		ConstThis->GameAgnosticSettings = NewObject<UEditorGameAgnosticSettings>();
-		
-		// Load config from file, but only if we are not the build machine since game agnostic settings may put the builder in an unclean state
-		if ( !GIsBuildMachine )
-		{
-			ConstThis->GameAgnosticSettings->LoadConfig(UEditorGameAgnosticSettings::StaticClass(), *GEditorGameAgnosticIni);
-		}
-	}
-	return *GameAgnosticSettings;
-}
-
-UEditorGameAgnosticSettings& UEditorEngine::AccessGameAgnosticSettings()
-{
-	if (GameAgnosticSettings == NULL)
-	{
-		GameAgnosticSettings = NewObject<UEditorGameAgnosticSettings>();
-		
-		// Load config from file, but only if we are not the build machine since game agnostic settings may put the builder in an unclean state
-		if ( !GIsBuildMachine )
-		{
-			GameAgnosticSettings->LoadConfig(UEditorGameAgnosticSettings::StaticClass(), *GEditorGameAgnosticIni);
-		}
-	}
-	return *GameAgnosticSettings;
-}
-
-void UEditorEngine::SaveGameAgnosticSettings()
-{
-	// Save config to file, but only if we are not the build machine since game agnostic settings may put the builder in an unclean state
-	if ( !GIsBuildMachine )
-	{
-		AccessGameAgnosticSettings().SaveConfig(CPF_Config, *GEditorGameAgnosticIni);
-	}
-}
-
-bool UEditorEngine::IsRealTimeAudioMuted() const
-{
-	if (EditorUserSettings == NULL)
-	{
-		return true;
-	}
-	return GetDefault<ULevelEditorMiscSettings>()->bEnableRealTimeAudio ? false : true;
-}
-
-void UEditorEngine::MuteRealTimeAudio(bool bMute)
-{
-	ULevelEditorMiscSettings* LevelEditorMiscSettings = GetMutableDefault<ULevelEditorMiscSettings>();
-
-	LevelEditorMiscSettings->bEnableRealTimeAudio = bMute ? false : true;
-	LevelEditorMiscSettings->PostEditChange();
-}
-
-float UEditorEngine::GetRealTimeAudioVolume() const
-{
-	return GetDefault<ULevelEditorMiscSettings>()->EditorVolumeLevel;
-}
-
-void UEditorEngine::SetRealTimeAudioVolume(float VolumeLevel)
-{
-	ULevelEditorMiscSettings* LevelEditorMiscSettings = GetMutableDefault<ULevelEditorMiscSettings>();
-
-	LevelEditorMiscSettings->EditorVolumeLevel = VolumeLevel;
-	LevelEditorMiscSettings->PostEditChange();
-}
-
-bool UEditorEngine::UpdateSingleViewportClient(FEditorViewportClient* InViewportClient, const bool bInAllowNonRealtimeViewportToDraw, bool bLinkedOrthoMovement )
-{
-	bool bUpdatedNonRealtimeViewport = false;
-
-	// Always submit view information for content streaming 
-	// otherwise content for editor view can be streamed out if there are other views (ex: thumbnails)
-	if (InViewportClient->IsPerspective())
-	{
-		IStreamingManager::Get().AddViewInformation( InViewportClient->GetViewLocation(), InViewportClient->Viewport->GetSizeXY().X, InViewportClient->Viewport->GetSizeXY().X / FMath::Tan(InViewportClient->ViewFOV) );
-	}
-	
-	// Only allow viewports to be drawn if we are not throttling for slate UI responsiveness or if the viewport client requested a redraw
-	// Note about bNeedsRedraw: Redraws can happen during some Slate events like checking a checkbox in a menu to toggle a view mode in the viewport.  In those cases we need to show the user the results immediately
-	if( FSlateThrottleManager::Get().IsAllowingExpensiveTasks() || InViewportClient->bNeedsRedraw )
-	{
-		// Switch to the world used by the viewport before its drawn
-		FScopedConditionalWorldSwitcher WorldSwitcher( InViewportClient );
-	
-		// Add view information for perspective viewports.
-		if( InViewportClient->IsPerspective() )
-		{
-			InViewportClient->GetWorld()->ViewLocationsRenderedLastFrame.Add(InViewportClient->GetViewLocation());
-	
-			// If we're currently simulating in editor, then we'll need to make sure that sub-levels are streamed in.
-			// When using PIE, this normally happens by UGameViewportClient::Draw().  But for SIE, we need to do
-			// this ourselves!
-			if( PlayWorld != NULL && bIsSimulatingInEditor && InViewportClient->IsSimulateInEditorViewport() )
-			{
-				// Update level streaming.
-				InViewportClient->GetWorld()->UpdateLevelStreaming();
-
-				// Also make sure hit proxies are refreshed for SIE viewports, as the user may be trying to grab an object or widget manipulator that's moving!
-				if( InViewportClient->IsRealtime() )
-				{
-					// @todo simulate: This may cause simulate performance to be worse in cases where you aren't needing to interact with gizmos.  Consider making this optional.
-					InViewportClient->RequestInvalidateHitProxy( InViewportClient->Viewport );
-				}
-			}
-		}
-	
-		// Redraw the viewport if it's realtime.
-		if( InViewportClient->IsRealtime() )
-		{
-			InViewportClient->Viewport->Draw();
-			InViewportClient->bNeedsRedraw = false;
-			InViewportClient->bNeedsLinkedRedraw = false;
-		}
-		// Redraw any linked ortho viewports that need to be updated this frame.
-		else if( InViewportClient->IsOrtho() && bLinkedOrthoMovement && InViewportClient->IsVisible() )
-		{
-			if( InViewportClient->bNeedsLinkedRedraw || InViewportClient->bNeedsRedraw )
-			{
-				// Redraw this viewport
-				InViewportClient->Viewport->Draw();
-				InViewportClient->bNeedsLinkedRedraw = false;
-				InViewportClient->bNeedsRedraw = false;
-			}
-			else
-			{
-				// This viewport doesn't need to be redrawn.  Skip this frame and increment the number of frames we skipped.
-				InViewportClient->FramesSinceLastDraw++;
-			}
-		}
-		// Redraw the viewport if there are pending redraw, and we haven't already drawn one viewport this frame.
-		else if (InViewportClient->bNeedsRedraw && bInAllowNonRealtimeViewportToDraw)
-		{
-			InViewportClient->Viewport->Draw();
-			InViewportClient->bNeedsRedraw = false;
-			bUpdatedNonRealtimeViewport = true;
-		}
-
-		if (InViewportClient->bNeedsInvalidateHitProxy)
-		{
-			InViewportClient->Viewport->InvalidateHitProxy();
-			InViewportClient->bNeedsInvalidateHitProxy = false;
-		}
-	}
-
-	return bUpdatedNonRealtimeViewport;
-}
-
-void UEditorEngine::InvalidateAllViewportsAndHitProxies()
-{
-	for (FEditorViewportClient* ViewportClient : AllViewportClients)
-	{
-		ViewportClient->Invalidate();
-	}
-}
-
-void UEditorEngine::InvalidateAllLevelEditorViewportClientHitProxies()
-{
-	for (const auto* LevelViewportClient : LevelViewportClients)
-	{
-		if (LevelViewportClient->Viewport != nullptr)
-		{
-			LevelViewportClient->Viewport->InvalidateHitProxy();
-		}
-	}
-}
-
-void UEditorEngine::PostEditChangeProperty(FPropertyChangedEvent& PropertyChangedEvent)
-{
-	// Propagate the callback up to the superclass.
-	Super::PostEditChangeProperty(PropertyChangedEvent);
-
-	const FName PropertyName = PropertyChangedEvent.Property ? PropertyChangedEvent.Property->GetFName() : NAME_None;
-
-	if (PropertyName == GET_MEMBER_NAME_CHECKED(UEngine, MaximumLoopIterationCount))
-	{
-		// Clamp to a reasonable range and feed the new value to the script core
-		MaximumLoopIterationCount = FMath::Clamp( MaximumLoopIterationCount, 100, 10000000 );
-		FBlueprintCoreDelegates::SetScriptMaximumLoopIterations( MaximumLoopIterationCount );
-	}
-	else if (PropertyName == GET_MEMBER_NAME_CHECKED(UEngine, bCanBlueprintsTickByDefault))
-	{
-		FScopedSlowTask SlowTask(100, LOCTEXT("DirtyingBlueprintsDueToTickChange", "InvalidatingAllBlueprints"));
-
-		// Flag all Blueprints as out of date (this doesn't dirty the package as needs saving but will force a recompile during PIE)
-		for (TObjectIterator<UBlueprint> BlueprintIt; BlueprintIt; ++BlueprintIt)
-		{
-			UBlueprint* Blueprint = *BlueprintIt;
-			Blueprint->Status = BS_Dirty;
-		}
-	}
-}
-
-void UEditorEngine::Cleanse( bool ClearSelection, bool Redraw, const FText& TransReset )
-{
-	check( !TransReset.IsEmpty() );
-
-	if (GIsRunning)
-	{
-		if( ClearSelection )
-		{
-			// Clear selection sets.
-			GetSelectedActors()->DeselectAll();
-			GetSelectedObjects()->DeselectAll();
-		}
-
-		// Reset the transaction tracking system.
-		ResetTransaction( TransReset );
-
-		// Invalidate hit proxies as they can retain references to objects over a few frames
-		FEditorSupportDelegates::CleanseEditor.Broadcast();
-
-		// Redraw the levels.
-		if( Redraw )
-		{
-			RedrawLevelEditingViewports();
-		}
-
-		// Attempt to unload any loaded redirectors. Redirectors should not be referenced in memory and are only used to forward references at load time
-		for (TObjectIterator<UObjectRedirector> RedirIt; RedirIt; ++RedirIt)
-		{
-			RedirIt->ClearFlags(RF_Standalone | RF_RootSet | RF_Transactional);
-		}
-
-		// Collect garbage.
-		CollectGarbage( GARBAGE_COLLECTION_KEEPFLAGS );
-
-		// Remaining redirectors are probably referenced by editor tools. Keep them in memory for now.
-		for (TObjectIterator<UObjectRedirector> RedirIt; RedirIt; ++RedirIt)
-		{
-			if ( RedirIt->IsAsset() )
-			{
-				RedirIt->SetFlags(RF_Standalone);
-			}
-		}
-	}
-}
-
-void UEditorEngine::EditorClearComponents()
-{
-	UWorld* World = GWorld;
-	if (World != NULL)
-	{
-		World->ClearWorldComponents();
-	}
-}
-
-void UEditorEngine::EditorUpdateComponents()
-{
-	GWorld->UpdateWorldComponents( true, false );
-}
-
-UAudioComponent* UEditorEngine::GetPreviewAudioComponent()
-{
-	return PreviewAudioComponent;
-}
-
-UAudioComponent* UEditorEngine::ResetPreviewAudioComponent( USoundBase* Sound, USoundNode* SoundNode )
-{
-	if (FAudioDevice* AudioDevice = GetMainAudioDevice())
-	{
-		if (PreviewAudioComponent)
-		{
-			PreviewAudioComponent->Stop();
-		}
-		else
-		{
-			PreviewSoundCue = NewObject<USoundCue>();
-			// Set world to NULL as it will most likely become invalid in the next PIE/Simulate session and the
-			// component will be left with invalid pointer.
-			PreviewAudioComponent = FAudioDevice::CreateComponent(PreviewSoundCue, NULL, NULL, false);
-		}
-
-		check(PreviewAudioComponent);
-		// Mark as a preview component so the distance calculations can be ignored
-		PreviewAudioComponent->bPreviewComponent = true;
-
-		if (Sound)
-		{
-			PreviewAudioComponent->Sound = Sound;
-		}
-		else if (SoundNode)
-		{
-			PreviewSoundCue->FirstNode = SoundNode;
-			PreviewAudioComponent->Sound = PreviewSoundCue;
-		}
-	}
-
-	return PreviewAudioComponent;
-}
-
-void UEditorEngine::PlayPreviewSound( USoundBase* Sound,  USoundNode* SoundNode )
-{
-	UAudioComponent* AudioComponent = ResetPreviewAudioComponent(Sound, SoundNode);
-	if(AudioComponent)
-	{
-		AudioComponent->bAutoDestroy = false;
-		AudioComponent->bIsUISound = true;
-		AudioComponent->bAllowSpatialization = false;
-		AudioComponent->bReverb = false;
-		AudioComponent->bCenterChannelOnly = false;
-
-		AudioComponent->Play();	
-	}
-}
-
-void UEditorEngine::PlayEditorSound( const FString& SoundAssetName )
-{
-	// Only play sounds if the user has that feature enabled
-	if( GetDefault<ULevelEditorMiscSettings>()->bEnableEditorSounds && !GIsSavingPackage )
-	{
-		USoundBase* Sound = Cast<USoundBase>( StaticFindObject( USoundBase::StaticClass(), NULL, *SoundAssetName ) );
-		if( Sound == NULL )
-		{
-			Sound = Cast<USoundBase>( StaticLoadObject( USoundBase::StaticClass(), NULL, *SoundAssetName ) );
-		}
-
-		if( Sound != NULL )
-		{
-			GEditor->PlayPreviewSound( Sound );
-		}
-	}
-}
-
-
-void UEditorEngine::ClearPreviewComponents()
-{
-	if( PreviewAudioComponent )
-	{
-		PreviewAudioComponent->Stop();
-
-		// Just null out so they get GC'd
-		PreviewSoundCue->FirstNode = NULL;
-		PreviewSoundCue = NULL;
-		PreviewAudioComponent->Sound = NULL;
-		PreviewAudioComponent = NULL;
-	}
-
-	if (PreviewMeshComp)
-	{
-		PreviewMeshComp->UnregisterComponent();
-		PreviewMeshComp = NULL;
-	}
-}
-
-void UEditorEngine::CloseEditedWorldAssets(UWorld* InWorld)
-{
-	if (!InWorld)
-	{
-		return;
-	}
-
-	// Find all assets being edited
-	FAssetEditorManager& EditorManager = FAssetEditorManager::Get();
-	TArray<UObject*> AllAssets = EditorManager.GetAllEditedAssets();
-
-	TSet<UWorld*> ClosingWorlds;
-
-	ClosingWorlds.Add(InWorld);
-
-	for (int32 Index = 0; Index < InWorld->StreamingLevels.Num(); ++Index)
-	{
-		ULevelStreaming* LevelStreaming = InWorld->StreamingLevels[Index];
-		if (LevelStreaming && LevelStreaming->LoadedLevel)
-		{
-			ClosingWorlds.Add(CastChecked<UWorld>(LevelStreaming->LoadedLevel->GetOuter()));
-		}
-	}
-
-	for(int32 i=0; i<AllAssets.Num(); i++)
-	{
-		UObject* Asset = AllAssets[i];
-		UWorld* AssetWorld = Asset->GetTypedOuter<UWorld>();
-
-		if ( !AssetWorld )
-		{
-			// This might be a world, itself
-			AssetWorld = Cast<UWorld>(Asset);
-		}
-
-		if (AssetWorld && ClosingWorlds.Contains(AssetWorld))
-		{
-			const TArray<IAssetEditorInstance*> AssetEditors = EditorManager.FindEditorsForAsset(Asset);
-			for (IAssetEditorInstance* EditorInstance : AssetEditors )
-			{
-				if (EditorInstance != NULL)
-				{
-					EditorInstance->CloseWindow();
-				}
-			}
-		}
-	}
-}
-
-UTextureRenderTarget2D* UEditorEngine::GetScratchRenderTarget( uint32 MinSize )
-{
-	auto NewFactory = NewObject<UTextureRenderTargetFactoryNew>();
-	UTextureRenderTarget2D* ScratchRenderTarget = NULL;
-
-	// We never allow render targets greater than 2048
-	check( MinSize <= 2048 );
-
-	// 256x256
-	if( MinSize <= 256 )
-	{
-		if( ScratchRenderTarget256 == NULL )
-		{
-			NewFactory->Width = 256;
-			NewFactory->Height = 256;
-			UObject* NewObj = NewFactory->FactoryCreateNew( UTextureRenderTarget2D::StaticClass(), GetTransientPackage(), NAME_None, RF_Transient, NULL, GWarn );
-			ScratchRenderTarget256 = CastChecked<UTextureRenderTarget2D>(NewObj);
-		}
-		ScratchRenderTarget = ScratchRenderTarget256;
-	}
-	// 512x512
-	else if( MinSize <= 512 )
-	{
-		if( ScratchRenderTarget512 == NULL )
-		{
-			NewFactory->Width = 512;
-			NewFactory->Height = 512;
-			UObject* NewObj = NewFactory->FactoryCreateNew( UTextureRenderTarget2D::StaticClass(), GetTransientPackage(), NAME_None, RF_Transient, NULL, GWarn );
-			ScratchRenderTarget512 = CastChecked<UTextureRenderTarget2D>(NewObj);
-		}
-		ScratchRenderTarget = ScratchRenderTarget512;
-	}
-	// 1024x1024
-	else if( MinSize <= 1024 )
-	{
-		if( ScratchRenderTarget1024 == NULL )
-		{
-			NewFactory->Width = 1024;
-			NewFactory->Height = 1024;
-			UObject* NewObj = NewFactory->FactoryCreateNew( UTextureRenderTarget2D::StaticClass(), GetTransientPackage(), NAME_None, RF_Transient, NULL, GWarn );
-			ScratchRenderTarget1024 = CastChecked<UTextureRenderTarget2D>(NewObj);
-		}
-		ScratchRenderTarget = ScratchRenderTarget1024;
-	}
-	// 2048x2048
-	else if( MinSize <= 2048 )
-	{
-		if( ScratchRenderTarget2048 == NULL )
-		{
-			NewFactory->Width = 2048;
-			NewFactory->Height = 2048;
-			UObject* NewObj = NewFactory->FactoryCreateNew( UTextureRenderTarget2D::StaticClass(), GetTransientPackage(), NAME_None, RF_Transient, NULL, GWarn );
-			ScratchRenderTarget2048 = CastChecked<UTextureRenderTarget2D>(NewObj);
-		}
-		ScratchRenderTarget = ScratchRenderTarget2048;
-	}
-
-	check( ScratchRenderTarget != NULL );
-	return ScratchRenderTarget;
-}
-
-
-bool UEditorEngine::WarnAboutHiddenLevels( UWorld* InWorld, bool bIncludePersistentLvl) const
-{
-	bool bResult = true;
-
-	const bool bPersistentLvlHidden = !FLevelUtils::IsLevelVisible( InWorld->PersistentLevel );
-
-	// Make a list of all hidden streaming levels.
-	TArray< ULevelStreaming* > HiddenLevels;
-	for( int32 LevelIndex = 0 ; LevelIndex< InWorld->StreamingLevels.Num() ; ++LevelIndex )
-	{
-		ULevelStreaming* StreamingLevel = InWorld->StreamingLevels[ LevelIndex ];
-		if( StreamingLevel && !FLevelUtils::IsLevelVisible( StreamingLevel ) )
-		{
-			HiddenLevels.Add( StreamingLevel );
-		}
-	}
-
-	// Warn the user that some levels are hidden and prompt for continue.
-	if ( ( bIncludePersistentLvl && bPersistentLvlHidden ) || HiddenLevels.Num() > 0 )
-	{
-		FText Message;
-		if ( !bIncludePersistentLvl )
-		{
-			Message = LOCTEXT("TheFollowingStreamingLevelsAreHidden_Additional", "The following streaming levels are hidden:\n{HiddenLevelNameList}\n\n{ContinueMessage}");
-		}
-		else if ( bPersistentLvlHidden )
-		{
-			Message = LOCTEXT("TheFollowingLevelsAreHidden_Persistent", "The following levels are hidden:\n\n    Persistent Level{HiddenLevelNameList}\n\n{ContinueMessage}");
-		}
-		else
-		{
-			Message = LOCTEXT("TheFollowingLevelsAreHidden_Additional", "The following levels are hidden:\n{HiddenLevelNameList}\n\n{ContinueMessage}");
-		}
-
-		FString HiddenLevelNames;
-		for ( int32 LevelIndex = 0 ; LevelIndex < HiddenLevels.Num() ; ++LevelIndex )
-		{
-			HiddenLevelNames += FString::Printf( TEXT("\n    %s"), *HiddenLevels[LevelIndex]->GetWorldAssetPackageName() );
-		}
-
-		FFormatNamedArguments Args;
-		Args.Add( TEXT("HiddenLevelNameList"), FText::FromString( HiddenLevelNames ) );
-		Args.Add( TEXT("ContinueMessage"), LOCTEXT("HiddenLevelsContinueWithBuildQ", "These levels will not be rebuilt. Leaving them hidden may invalidate what is built in other levels.\n\nContinue with build?\n(Yes All will show all hidden levels and continue with the build)") );
-
-		const FText MessageBoxText = FText::Format( Message, Args );
-
-		// Create and show the user the dialog.
-		const EAppReturnType::Type Choice = FMessageDialog::Open(EAppMsgType::YesNoYesAll, MessageBoxText);
-
-		if( Choice == EAppReturnType::YesAll )
-		{
-			if ( bIncludePersistentLvl && bPersistentLvlHidden )
-			{
-				EditorLevelUtils::SetLevelVisibility( InWorld->PersistentLevel, true, false );
-			}
-
-			// The code below should technically also make use of FLevelUtils::SetLevelVisibility, but doing
-			// so would be much more inefficient, resulting in several calls to UpdateLevelStreaming
-			for( int32 HiddenLevelIdx = 0; HiddenLevelIdx < HiddenLevels.Num(); ++HiddenLevelIdx )
-			{
-				HiddenLevels[ HiddenLevelIdx ]->bShouldBeVisibleInEditor = true;
-			}
-
-			InWorld->FlushLevelStreaming();
-
-			// follow up using SetLevelVisibility - streaming should now be completed so we can show actors, layers, 
-			// BSPs etc. without too big a performance hit.
-			for( int32 HiddenLevelIdx = 0; HiddenLevelIdx < HiddenLevels.Num(); ++HiddenLevelIdx )
-			{
-				check(HiddenLevels[ HiddenLevelIdx ]->GetLoadedLevel());
-				ULevel* LoadedLevel = HiddenLevels[ HiddenLevelIdx ]->GetLoadedLevel();
-				EditorLevelUtils::SetLevelVisibility( LoadedLevel, true, false );
-			}
-
-			FEditorSupportDelegates::RedrawAllViewports.Broadcast();
-		}
-
-		// return true if the user pressed make all visible or yes.
-		bResult = (Choice != EAppReturnType::No);
-	}
-
-	return bResult;
-}
-
-
-void UEditorEngine::SetStreamingBoundsTexture(UTexture2D* InTexture)
-{
-	if (StreamingBoundsTexture != InTexture)
-	{
-		// Clear the currently stored streaming bounds information
-
-		// Set the new texture
-		StreamingBoundsTexture = InTexture;
-		if (StreamingBoundsTexture != NULL)
-		{
-			// Fill in the new streaming bounds info
-			ULevel::BuildStreamingData(NULL, NULL, InTexture);
-
-			// Turn on the StreamingBounds show flag
-			for(uint32 ViewportIndex = 0; ViewportIndex < (uint32)LevelViewportClients.Num(); ViewportIndex++)
-			{
-				FLevelEditorViewportClient* ViewportClient = LevelViewportClients[ViewportIndex];
-				if (ViewportClient)
-				{
-					ViewportClient->EngineShowFlags.StreamingBounds = 1;
-					ViewportClient->Invalidate(false,true);
-				}
-			}
-		}
-		else
-		{
-			// Clear the streaming bounds info
-			// Turn off the StreamingBounds show flag
-			for(uint32 ViewportIndex = 0; ViewportIndex < (uint32)LevelViewportClients.Num(); ViewportIndex++)
-			{
-				FLevelEditorViewportClient* ViewportClient = LevelViewportClients[ViewportIndex];
-				if (ViewportClient)
-				{
-					ViewportClient->EngineShowFlags.StreamingBounds = 0;
-					ViewportClient->Invalidate(false,true);
-				}
-			}
-		}
-	}
-}
-
-void UEditorEngine::ApplyDeltaToActor(AActor* InActor,
-									  bool bDelta,
-									  const FVector* InTrans,
-									  const FRotator* InRot,
-									  const FVector* InScale,
-									  bool bAltDown,
-									  bool bShiftDown,
-									  bool bControlDown) const
-{
-	if(!bDisableDeltaModification)
-	{
-		InActor->Modify();
-	}
-	
-	ABrush* Brush = Cast< ABrush >( InActor );
-	if( Brush )
-	{
-		if( Brush->GetBrushComponent() && Brush->GetBrushComponent()->Brush )
-		{
-			Brush->GetBrushComponent()->Brush->Polys->Element.ModifyAllItems();
-		}
-	}
-
-	FNavigationLockContext LockNavigationUpdates(InActor->GetWorld(), ENavigationLockReason::ContinuousEditorMove);
-
-	bool bTranslationOnly = true;
-
-	///////////////////
-	// Rotation
-
-	// Unfortunately this can't be moved into ABrush::EditorApplyRotation, as that would
-	// create a dependence in Engine on Editor.
-	if ( InRot )
-	{
-		const FRotator& InDeltaRot = *InRot;
-		const bool bRotatingActor = !bDelta || !InDeltaRot.IsZero();
-		if( bRotatingActor )
-		{
-			bTranslationOnly = false;
-
-			if ( bDelta )
-			{
-				if( InActor->GetRootComponent() != NULL )
-				{
-					const FRotator OriginalRotation = InActor->GetRootComponent()->GetComponentRotation();
-
-					InActor->EditorApplyRotation( InDeltaRot, bAltDown, bShiftDown, bControlDown );
-
-					// Check to see if we should transform the rigid body
-					UPrimitiveComponent* RootPrimitiveComponent = Cast< UPrimitiveComponent >( InActor->GetRootComponent() );
-					if( bIsSimulatingInEditor && GIsPlayInEditorWorld && RootPrimitiveComponent != NULL )
-					{
-						FRotator ActorRotWind, ActorRotRem;
-						OriginalRotation.GetWindingAndRemainder(ActorRotWind, ActorRotRem);
-
-						const FQuat ActorQ = ActorRotRem.Quaternion();
-						const FQuat DeltaQ = InDeltaRot.Quaternion();
-						const FQuat ResultQ = DeltaQ * ActorQ;
-
-						const FRotator NewActorRotRem = FRotator( ResultQ );
-						FRotator DeltaRot = NewActorRotRem - ActorRotRem;
-						DeltaRot.Normalize();
-
-						// @todo SIE: Not taking into account possible offset between root component and actor
-						RootPrimitiveComponent->SetWorldRotation( OriginalRotation + DeltaRot );
-					}
-				}
-
-				FVector NewActorLocation = InActor->GetActorLocation();
-				NewActorLocation -= GLevelEditorModeTools().PivotLocation;
-				NewActorLocation = FRotationMatrix(InDeltaRot).TransformPosition(NewActorLocation);
-				NewActorLocation += GLevelEditorModeTools().PivotLocation;
-				NewActorLocation -= InActor->GetActorLocation();
-				InActor->EditorApplyTranslation(NewActorLocation, bAltDown, bShiftDown, bControlDown);
-			}
-			else
-			{
-				InActor->SetActorRotation( InDeltaRot );
-			}
-		}
-	}
-
-	///////////////////
-	// Translation
-	if ( InTrans )
-	{
-		if ( bDelta )
-		{
-			if( InActor->GetRootComponent() != NULL )
-			{
-				const FVector OriginalLocation = InActor->GetRootComponent()->GetComponentLocation();
-
-				InActor->EditorApplyTranslation( *InTrans, bAltDown, bShiftDown, bControlDown );
-
-				// Check to see if we should transform the rigid body
-				UPrimitiveComponent* RootPrimitiveComponent = Cast< UPrimitiveComponent >( InActor->GetRootComponent() );
-				if( bIsSimulatingInEditor && GIsPlayInEditorWorld && RootPrimitiveComponent != NULL )
-				{
-					// @todo SIE: Not taking into account possible offset between root component and actor
-					RootPrimitiveComponent->SetWorldLocation( OriginalLocation + *InTrans );
-				}
-			}
-		}
-		else
-		{
-			InActor->SetActorLocation( *InTrans, false );
-		}
-	}
-
-	///////////////////
-	// Scaling
-	if ( InScale )
-	{
-		const FVector& InDeltaScale = *InScale;
-		const bool bScalingActor = !bDelta || !InDeltaScale.IsNearlyZero(0.000001f);
-		if( bScalingActor )
-		{
-			bTranslationOnly = false;
-
-			FVector ModifiedScale = InDeltaScale;
-
-			// Note: With the new additive scaling method, this is handled in FLevelEditorViewportClient::ModifyScale
-			if( GEditor->UsePercentageBasedScaling() )
-			{
-				// Get actor box extents
-				const FBox BoundingBox = InActor->GetComponentsBoundingBox( true );
-				const FVector BoundsExtents = BoundingBox.GetExtent();
-
-				// Make sure scale on actors is clamped to a minimum and maximum size.
-				const float MinThreshold = 1.0f;
-
-				for (int32 Idx=0; Idx<3; Idx++)
-				{
-					if ( ( FMath::Pow(BoundsExtents[Idx], 2) ) > BIG_NUMBER)
-					{
-						ModifiedScale[Idx] = 0.0f;
-					}
-					else if (SMALL_NUMBER < BoundsExtents[Idx])
-					{
-						const bool bBelowAllowableScaleThreshold = ((InDeltaScale[Idx] + 1.0f) * BoundsExtents[Idx]) < MinThreshold;
-
-						if(bBelowAllowableScaleThreshold)
-						{
-							ModifiedScale[Idx] = (MinThreshold / BoundsExtents[Idx]) - 1.0f;
-						}
-					}
-				}
-			}
-
-			if ( bDelta )
-			{
-				// Flag actors to use old-style scaling or not
-				// @todo: Remove this hack once we have decided on the scaling method to use.
-				AActor::bUsePercentageBasedScaling = GEditor->UsePercentageBasedScaling();
-
-				InActor->EditorApplyScale( 
-					ModifiedScale,
-					&GLevelEditorModeTools().PivotLocation,
-					bAltDown,
-					bShiftDown,
-					bControlDown
-					);
-
-			}
-			else if( InActor->GetRootComponent() != NULL )
-			{
-				InActor->GetRootComponent()->SetRelativeScale3D( InDeltaScale );
-			}
-		}
-	}
-
-	// Update the actor before leaving.
-	InActor->MarkPackageDirty();
-	InActor->InvalidateLightingCacheDetailed(bTranslationOnly);
-	InActor->PostEditMove( false );
-}
-
-void UEditorEngine::ApplyDeltaToComponent(USceneComponent* InComponent,
-	bool bDelta,
-	const FVector* InTrans,
-	const FRotator* InRot,
-	const FVector* InScale,
-	const FVector& PivotLocation ) const
-{
-	if(!bDisableDeltaModification)
-	{
-		InComponent->Modify();
-	}
-
-	///////////////////
-	// Rotation
-	if ( InRot )
-	{
-		const FRotator& InDeltaRot = *InRot;
-		const bool bRotatingComp = !bDelta || !InDeltaRot.IsZero();
-		if( bRotatingComp )
-		{
-			if ( bDelta )
-			{
-				const FQuat ActorQ = InComponent->RelativeRotation.Quaternion();
-				const FQuat DeltaQ = InDeltaRot.Quaternion();
-				const FQuat ResultQ = DeltaQ * ActorQ;
-
-				const FRotator NewActorRot = FRotator( ResultQ );
-
-				InComponent->SetRelativeRotation(NewActorRot);
-			}
-			else
-			{
-				InComponent->SetRelativeRotation( InDeltaRot );
-			}
-
-			if ( bDelta )
-			{
-				FVector NewCompLocation = InComponent->RelativeLocation;
-				NewCompLocation -= PivotLocation;
-				NewCompLocation = FRotationMatrix( InDeltaRot ).TransformPosition( NewCompLocation );
-				NewCompLocation += PivotLocation;
-				InComponent->SetRelativeLocation(NewCompLocation);
-			}
-		}
-	}
-
-	///////////////////
-	// Translation
-	if ( InTrans )
-	{
-		if ( bDelta )
-		{
-			InComponent->SetRelativeLocation(InComponent->RelativeLocation + *InTrans);
-		}
-		else
-		{
-			InComponent->SetRelativeLocation( *InTrans );
-		}
-	}
-
-	///////////////////
-	// Scaling
-	if ( InScale )
-	{
-		const FVector& InDeltaScale = *InScale;
-		const bool bScalingComp = !bDelta || !InDeltaScale.IsNearlyZero(0.000001f);
-		if( bScalingComp )
-		{
-			if ( bDelta )
-			{
-				const FScaleMatrix ScaleMatrix( FVector( InDeltaScale.X , InDeltaScale.Y, InDeltaScale.Z ) );
-
-				FVector DeltaScale3D = ScaleMatrix.TransformPosition( InComponent->RelativeScale3D );
-				InComponent->SetRelativeScale3D(InComponent->RelativeScale3D + DeltaScale3D);
-
-
-				FVector NewCompLocation = InComponent->RelativeLocation;
-				NewCompLocation -= PivotLocation;
-				NewCompLocation += ScaleMatrix.TransformPosition( NewCompLocation );
-				NewCompLocation += PivotLocation;
-				InComponent->SetRelativeLocation(NewCompLocation);
-			}
-			else
-			{
-				InComponent->SetRelativeScale3D(InDeltaScale);
-			}
-		}
-	}
-
-	// Update the actor before leaving.
-	InComponent->MarkPackageDirty();
-
-	// Fire callbacks
-	FEditorSupportDelegates::RefreshPropertyWindows.Broadcast();
-	FEditorSupportDelegates::UpdateUI.Broadcast();
-}
-
-
-void UEditorEngine::ProcessToggleFreezeCommand( UWorld* InWorld )
-{
-	if (InWorld->IsPlayInEditor())
-	{
-		ULocalPlayer* Player = PlayWorld->GetFirstLocalPlayerFromController();
-		if( Player )
-		{
-			Player->ViewportClient->Viewport->ProcessToggleFreezeCommand();
-		}
-	}
-	else
-	{
-		// pass along the freeze command to all perspective viewports
-		for(int32 ViewportIndex = 0; ViewportIndex < LevelViewportClients.Num(); ViewportIndex++)
-		{
-			if (LevelViewportClients[ViewportIndex]->IsPerspective())
-			{
-				LevelViewportClients[ViewportIndex]->Viewport->ProcessToggleFreezeCommand();
-			}
-		}
-	}
-
-	// tell editor to update views
-	RedrawAllViewports();
-}
-
-
-void UEditorEngine::ProcessToggleFreezeStreamingCommand(UWorld* InWorld)
-{
-	// freeze vis in PIE
-	if (InWorld && InWorld->WorldType == EWorldType::PIE)
-	{
-		InWorld->bIsLevelStreamingFrozen = !InWorld->bIsLevelStreamingFrozen;
-	}
-}
-
-/*-----------------------------------------------------------------------------
-	Reimporting.
------------------------------------------------------------------------------*/
-
-/** 
-* Singleton function
-* @return Singleton instance of this manager
-*/
-FReimportManager* FReimportManager::Instance()
-{
-	static FReimportManager Inst;
-	return &Inst;
-}
-
-/**
- * Register a reimport handler with the manager
- *
- * @param	InHandler	Handler to register with the manager
- */
-void FReimportManager::RegisterHandler( FReimportHandler& InHandler )
-{
-	Handlers.AddUnique( &InHandler );
-	bHandlersNeedSorting = true;
-}
-
-/**
- * Unregister a reimport handler from the manager
- *
- * @param	InHandler	Handler to unregister from the manager
- */
-void FReimportManager::UnregisterHandler( FReimportHandler& InHandler )
-{
-	Handlers.Remove( &InHandler );
-}
-
-bool FReimportManager::CanReimport( UObject* Obj ) const
-{
-	if ( Obj )
-	{
-		TArray<FString> SourceFilenames;
-		for( int32 HandlerIndex = 0; HandlerIndex < Handlers.Num(); ++HandlerIndex )
-		{
-			SourceFilenames.Empty();
-			if ( Handlers[ HandlerIndex ]->CanReimport(Obj, SourceFilenames) )
-			{
-				return true;
-			}
-		}
-	}
-	return false;
-}
-
-void FReimportManager::UpdateReimportPaths( UObject* Obj, const TArray<FString>& InFilenames )
-{
-	if (Obj)
-	{
-		TArray<FString> UnusedExistingFilenames;
-		auto* Handler = Handlers.FindByPredicate([&](FReimportHandler* InHandler){ return InHandler->CanReimport(Obj, UnusedExistingFilenames); });
-		if (Handler)
-		{
-			(*Handler)->SetReimportPaths(Obj, InFilenames);
-			Obj->MarkPackageDirty();
-		}
-	}
-}
-
-bool FReimportManager::Reimport( UObject* Obj, bool bAskForNewFileIfMissing, bool bShowNotification )
-{
-	// Warn that were about to reimport, so prep for it
-	PreReimport.Broadcast( Obj );
-
-	bool bSuccess = false;
-	if ( Obj )
-	{
-		if (bHandlersNeedSorting)
-		{
-			// Use > operator because we want higher priorities earlier in the list
-			Handlers.Sort([](const FReimportHandler& A, const FReimportHandler& B) { return A.GetPriority() > B.GetPriority(); });
-			bHandlersNeedSorting = false;
-		}
-		
-		bool bValidSourceFilename = false;
-		TArray<FString> SourceFilenames;
-
-		for( int32 HandlerIndex = 0; HandlerIndex < Handlers.Num(); ++HandlerIndex )
-		{
-			SourceFilenames.Empty();
-			if ( Handlers[ HandlerIndex ]->CanReimport(Obj, SourceFilenames) )
-			{
-				// Check all filenames for missing files
-				bool bMissingFiles = false;
-				if (SourceFilenames.Num() > 0)
-				{
-					for (int32 FileIndex = 0; FileIndex < SourceFilenames.Num(); ++FileIndex)
-					{
-						if (SourceFilenames[FileIndex].IsEmpty() || IFileManager::Get().FileSize(*SourceFilenames[FileIndex]) == INDEX_NONE)
-						{
-							bMissingFiles = true;
-							break;
-						}
-					}
-				}
-				else
-				{
-					bMissingFiles = true;
-				}
-
-				bValidSourceFilename = true;
-				if ( bAskForNewFileIfMissing && bMissingFiles )
-				{
-					GetNewReimportPath(Obj, SourceFilenames);
-					if ( SourceFilenames.Num() == 0 )
-					{
-						// Failed to specify a new filename. Don't show a notification of the failure since the user exited on his own
-						bValidSourceFilename = false;
-						bShowNotification = false;
-					}
-					else
-					{
-						// A new filename was supplied, update the path
-						Handlers[ HandlerIndex ]->SetReimportPaths(Obj, SourceFilenames);
-					}
-				}
-
-				if ( bValidSourceFilename )
-				{
-					// Do the reimport
-					EReimportResult::Type Result = Handlers[ HandlerIndex ]->Reimport( Obj );
-					if( Result == EReimportResult::Succeeded )
-					{
-						Obj->PostEditChange();
-						GEditor->BroadcastObjectReimported(Obj);
-						if (FEngineAnalytics::IsAvailable())
-						{
-							TArray<FAnalyticsEventAttribute> Attributes;
-							Attributes.Add( FAnalyticsEventAttribute( TEXT( "ObjectType" ), Obj->GetClass()->GetName() ) );
-							FEngineAnalytics::GetProvider().RecordEvent(TEXT("Editor.Usage.AssetReimported"), Attributes);
-						}
-						bSuccess = true;
-					}
-					else if( Result == EReimportResult::Cancelled )
-					{
-						bShowNotification = false;
-					}
-				}
-				
-				break;
-			}
-		}
-
-		if ( bShowNotification )
-		{
-			// Send a notification of the results
-			FText NotificationText;
-			if ( bSuccess )
-			{
-				if ( bValidSourceFilename )
-				{
-					const FString FirstLeafFilename = FPaths::GetCleanFilename(SourceFilenames[0]);
-
-					if ( SourceFilenames.Num() == 1 )
-					{
-						FFormatNamedArguments Args;
-						Args.Add( TEXT("ObjectName"), FText::FromString( Obj->GetName() ) );
-						Args.Add( TEXT("ObjectType"), FText::FromString( Obj->GetClass()->GetName() ) );
-						Args.Add( TEXT("SourceFile"), FText::FromString( FirstLeafFilename ) );
-						NotificationText = FText::Format( LOCTEXT("ReimportSuccessfulFrom", "Successfully Reimported: {ObjectName} ({ObjectType}) from file ({SourceFile})"), Args );
-					}
-					else
-					{
-						FFormatNamedArguments Args;
-						Args.Add( TEXT("ObjectName"), FText::FromString( Obj->GetName() ) );
-						Args.Add( TEXT("ObjectType"), FText::FromString( Obj->GetClass()->GetName() ) );
-						Args.Add( TEXT("SourceFile"), FText::FromString( FirstLeafFilename ) );
-						Args.Add( TEXT("Number"), SourceFilenames.Num() - 1 );
-						NotificationText = FText::Format( LOCTEXT("ReimportSuccessfulMultiple", "Successfuly Reimported: {ObjectName} ({ObjectType}) from file ({SourceFile}) and {Number} more"), Args );
-					}
-				}
-				else
-				{
-					FFormatNamedArguments Args;
-					Args.Add( TEXT("ObjectName"), FText::FromString( Obj->GetName() ) );
-					Args.Add( TEXT("ObjectType"), FText::FromString( Obj->GetClass()->GetName() ) );
-					NotificationText = FText::Format( LOCTEXT("ReimportSuccessful", "Successfully Reimported: {ObjectName} ({ObjectType})"), Args );
-				}
-			}
-			else
-			{
-				FFormatNamedArguments Args;
-				Args.Add( TEXT("ObjectName"), FText::FromString( Obj->GetName() ) );
-				Args.Add( TEXT("ObjectType"), FText::FromString( Obj->GetClass()->GetName() ) );
-				NotificationText = FText::Format( LOCTEXT("ReimportFailed", "Failed to Reimport: {ObjectName} ({ObjectType})"), Args );
-			}
-
-			FNotificationInfo Info(NotificationText);
-			Info.ExpireDuration = 3.0f;
-			Info.bUseLargeFont = false;
-			TSharedPtr<SNotificationItem> Notification = FSlateNotificationManager::Get().AddNotification(Info);
-			if ( Notification.IsValid() )
-			{
-				Notification->SetCompletionState( bSuccess ? SNotificationItem::CS_Success : SNotificationItem::CS_Fail );
-			}
-		}
-	}
-
-	// Let listeners know whether the reimport was successful or not
-	PostReimport.Broadcast( Obj, bSuccess );
-
-	GEditor->RedrawAllViewports();
-
-	return bSuccess;
-}
-
-void FReimportManager::GetNewReimportPath(UObject* Obj, TArray<FString>& InOutFilenames)
-{
-	TArray<UObject*> ReturnObjects;
-	FString FileTypes;
-	FString AllExtensions;
-	TArray<UFactory*> Factories;
-
-	// Determine whether we will allow multi select and clear old filenames
-	bool bAllowMultiSelect = InOutFilenames.Num() > 1;
-	InOutFilenames.Empty();
-
-	// Get the list of valid factories
-	for( TObjectIterator<UClass> It ; It ; ++It )
-	{
-		UClass* CurrentClass = (*It);
-
-		if( CurrentClass->IsChildOf(UFactory::StaticClass()) && !(CurrentClass->HasAnyClassFlags(CLASS_Abstract)) )
-		{
-			UFactory* Factory = Cast<UFactory>( CurrentClass->GetDefaultObject() );
-			if( Factory->bEditorImport && Obj->GetClass()->IsChildOf(Factory->GetSupportedClass()) )
-			{
-				Factories.Add( Factory );
-			}
-		}
-	}
-
-	if ( Factories.Num() <= 0 )
-	{
-		// No matching factories for this asset, fail
-		return;
-	}
-
-	TMultiMap<uint32, UFactory*> DummyFilterIndexToFactory;
-
-	// Generate the file types and extensions represented by the selected factories
-	ObjectTools::GenerateFactoryFileExtensions( Factories, FileTypes, AllExtensions, DummyFilterIndexToFactory );
-
-	FileTypes = FString::Printf(TEXT("All Files (%s)|%s|%s"),*AllExtensions,*AllExtensions,*FileTypes);
-
-	// Prompt the user for the filenames
-	TArray<FString> OpenFilenames;
-	IDesktopPlatform* DesktopPlatform = FDesktopPlatformModule::Get();
-	bool bOpened = false;
-	if ( DesktopPlatform )
-	{
-		void* ParentWindowWindowHandle = NULL;
-
-		IMainFrameModule& MainFrameModule = FModuleManager::LoadModuleChecked<IMainFrameModule>(TEXT("MainFrame"));
-		const TSharedPtr<SWindow>& MainFrameParentWindow = MainFrameModule.GetParentWindow();
-		if ( MainFrameParentWindow.IsValid() && MainFrameParentWindow->GetNativeWindow().IsValid() )
-		{
-			ParentWindowWindowHandle = MainFrameParentWindow->GetNativeWindow()->GetOSWindowHandle();
-		}
-
-		const FString Title = FString::Printf(TEXT("%s: %s"), *NSLOCTEXT("ReimportManager", "ImportDialogTitle", "Import For").ToString(), *Obj->GetName());
-		bOpened = DesktopPlatform->OpenFileDialog(
-			ParentWindowWindowHandle,
-			Title,
-			TEXT(""),
-			TEXT(""),
-			FileTypes,
-			bAllowMultiSelect ? EFileDialogFlags::Multiple : EFileDialogFlags::None,
-			OpenFilenames
-			);
-	}
-
-	if ( bOpened )
-	{
-		for (int32 FileIndex = 0; FileIndex < OpenFilenames.Num(); ++FileIndex)
-		{
-			InOutFilenames.Add(OpenFilenames[FileIndex]);
-		}
-	}
-}
-
-FString FReimportManager::SanitizeImportFilename(const FString& InPath, const UObject* Obj)
-{
-	const UPackage* Package = Obj->GetOutermost();
-	if (Package)
-	{
-		const bool		bIncludeDot = true;
-		const FString	PackagePath	= Package->GetPathName();
-		const FName		MountPoint	= FPackageName::GetPackageMountPoint(PackagePath);
-		const FString	PackageFilename = FPackageName::LongPackageNameToFilename(PackagePath, FPaths::GetExtension(InPath, bIncludeDot));
-		const FString	AbsolutePath = FPaths::ConvertRelativePathToFull(InPath);
-
-		if ((MountPoint == FName("Engine") && AbsolutePath.StartsWith(FPaths::ConvertRelativePathToFull(FPaths::EngineContentDir()))) ||
-			(MountPoint == FName("Game") &&	AbsolutePath.StartsWith(FPaths::ConvertRelativePathToFull(FPaths::GameDir()))))
-		{
-			FString RelativePath = InPath;
-			FPaths::MakePathRelativeTo(RelativePath, *PackageFilename);
-			return RelativePath;
-		}
-	}
-
-	return IFileManager::Get().ConvertToRelativePath(*InPath);
-}
-
-
-FString FReimportManager::ResolveImportFilename(const FString& InRelativePath, const UObject* Obj)
-{
-	FString RelativePath = InRelativePath;
-
-	const UPackage* Package = Obj->GetOutermost();
-	if (Package)
-	{
-		// Relative to the package filename?
-		const FString PathRelativeToPackage = FPaths::GetPath(FPackageName::LongPackageNameToFilename(Package->GetPathName())) / InRelativePath;
-		if (FPaths::FileExists(PathRelativeToPackage))
-		{
-			RelativePath = PathRelativeToPackage;
-		}
-	}
-
-	// Convert relative paths
-	return FPaths::ConvertRelativePathToFull(RelativePath);
-}
-
-
-/** Constructor */
-FReimportManager::FReimportManager()
-{
-	// Create reimport handler for textures
-	// NOTE: New factories can be created anywhere, inside or outside of editor
-	// This is just here for convenience
-	UReimportTextureFactory::StaticClass();
-
-	// Create reimport handler for FBX static meshes
-	UReimportFbxStaticMeshFactory::StaticClass();
-
-	// Create reimport handler for FBX skeletal meshes
-	UReimportFbxSkeletalMeshFactory::StaticClass();
-
-	// Create reimport handler for APEX destructible meshes
-	UReimportDestructibleMeshFactory::StaticClass();
-
-	// Create reimport handler for sound node waves
-	UReimportSoundFactory::StaticClass();
-
-	// Create reimport handler for surround sound waves
-	UReimportSoundSurroundFactory::StaticClass();
-}
-
-/** Destructor */
-FReimportManager::~FReimportManager()
-{
-	Handlers.Empty();
-}
-
-int32 FReimportHandler::GetPriority() const
-{
-	return UFactory::DefaultImportPriority;
-}
-
-/*-----------------------------------------------------------------------------
-	PIE helpers.
------------------------------------------------------------------------------*/
-
-/**
- * Sets GWorld to the passed in PlayWorld and sets a global flag indicating that we are playing
- * in the Editor.
- *
- * @param	PlayInEditorWorld		PlayWorld
- * @return	the original GWorld
- */
-UWorld* SetPlayInEditorWorld( UWorld* PlayInEditorWorld )
-{
-	check(!GIsPlayInEditorWorld);
-	UWorld* SavedWorld = GWorld;
-	GIsPlayInEditorWorld = true;
-	GWorld = PlayInEditorWorld;
-
-	return SavedWorld;
-}
-
-/**
- * Restores GWorld to the passed in one and reset the global flag indicating whether we are a PIE
- * world or not.
- *
- * @param EditorWorld	original world being edited
- */
-void RestoreEditorWorld( UWorld* EditorWorld )
-{
-	check(GIsPlayInEditorWorld);
-	GIsPlayInEditorWorld = false;
-	GWorld = EditorWorld;
-}
-
-/**
- * Compresses SoundWave for all available platforms, and then decompresses to PCM 
- *
- * @param	SoundWave				Wave file to compress
- * @param	PreviewInfo				Compressed stats and decompressed data
- */
-void SoundWaveQualityPreview( USoundWave* SoundWave, FPreviewInfo* PreviewInfo )
-{
-	FWaveModInfo WaveInfo;
-	FSoundQualityInfo QualityInfo = { 0 };
-
-	// Extract the info from the wave header
-	if( !WaveInfo.ReadWaveHeader( ( uint8* )SoundWave->ResourceData, SoundWave->ResourceSize, 0 ) )
-	{
-		return;
-	}
-
-	SoundWave->NumChannels = *WaveInfo.pChannels;
-	SoundWave->RawPCMDataSize = WaveInfo.SampleDataSize;
-
-	// Extract the stats
-	PreviewInfo->OriginalSize = SoundWave->ResourceSize;
-	PreviewInfo->OggVorbisSize = 0;
-	PreviewInfo->PS3Size = 0;
-	PreviewInfo->XMASize = 0;
-
-	QualityInfo.Quality = PreviewInfo->QualitySetting;
-	QualityInfo.NumChannels = *WaveInfo.pChannels;
-	QualityInfo.SampleRate = SoundWave->SampleRate;
-	QualityInfo.SampleDataSize = WaveInfo.SampleDataSize;
-	QualityInfo.DebugName = SoundWave->GetFullName();
-
-	// PCM -> Vorbis -> PCM 
-	ITargetPlatformManagerModule& TPM = GetTargetPlatformManagerRef();
-	static FName NAME_OGG(TEXT("OGG"));
-	const IAudioFormat* Compressor = TPM.FindAudioFormat(NAME_OGG); // why is this ogg specific?
-	PreviewInfo->OggVorbisSize = 0;
-
-	if (Compressor)
-	{
-		TArray<uint8> Input;
-		Input.AddUninitialized(WaveInfo.SampleDataSize);
-		FMemory::Memcpy(Input.GetData(), WaveInfo.SampleDataStart, Input.Num());
-		TArray<uint8> Output;
-		Compressor->Recompress(NAME_OGG, Input, QualityInfo, Output );
-		if (Output.Num())
-		{
-			PreviewInfo->OggVorbisSize = Output.Num();
-			PreviewInfo->DecompressedOggVorbis = ( uint8* )FMemory::Malloc(Output.Num());
-			FMemory::Memcpy(PreviewInfo->DecompressedOggVorbis, Output.GetData(), Output.Num());
-		}
-	}
-
-	if( PreviewInfo->OggVorbisSize <= 0 )
-	{
-		UE_LOG(LogAudio, Log, TEXT( "PC decompression failed" ) );
-	}
-}
-
-/**
- * Takes an FName and checks to see that it is unique among all loaded objects.
- *
- * @param	InName		The name to check
- * @param	Outer		The context for validating this object name. Should be a group/package, but could be ANY_PACKAGE if you want to check across the whole system (not recommended)
- * @param	InReason	If the check fails, this string is filled in with the reason why.
- *
- * @return	1 if the name is valid, 0 if it is not
- */
-
-bool IsUniqueObjectName( const FName& InName, UObject* Outer, FText* InReason )
-{
-	// See if the name is already in use.
-	if( StaticFindObject( UObject::StaticClass(), Outer, *InName.ToString() ) != NULL )
-	{
-		if ( InReason != NULL )
-		{
-			*InReason = NSLOCTEXT("UnrealEd", "NameAlreadyInUse", "Name is already in use by another object.");
-		}
-		return false;
-	}
-
-	return true;
-}
-
-/**
- * Takes an FName and checks to see that it is unique among all loaded objects.
- *
- * @param	InName		The name to check
- * @param	Outer		The context for validating this object name. Should be a group/package, but could be ANY_PACKAGE if you want to check across the whole system (not recommended)
- * @param	InReason	If the check fails, this string is filled in with the reason why.
- *
- * @return	1 if the name is valid, 0 if it is not
- */
-
-bool IsUniqueObjectName( const FName& InName, UObject* Outer, FText& InReason )
-{
-	return IsUniqueObjectName(InName,Outer,&InReason);
-}
-
-
-
-
-void UEditorEngine::ParseMapSectionIni(const TCHAR* InCmdParams, TArray<FString>& OutMapList)
-{
-	FString SectionStr;
-	if (FParse::Value(InCmdParams, TEXT("MAPINISECTION="), SectionStr))
-	{
-		if (SectionStr.Contains(TEXT("+")))
-		{
-			TArray<FString> Sections;
-			SectionStr.ParseIntoArray(Sections,TEXT("+"),true);
-			for (int32 Index = 0; Index < Sections.Num(); Index++)
-			{
-				LoadMapListFromIni(Sections[Index], OutMapList);
-			}
-		}
-		else
-		{
-			LoadMapListFromIni(SectionStr, OutMapList);
-		}
-	}
-}
-
-
-void UEditorEngine::LoadMapListFromIni(const FString& InSectionName, TArray<FString>& OutMapList)
-{
-	// 
-	FConfigSection* MapListList = GConfig->GetSectionPrivate(*InSectionName, false, true, GEditorIni);
-	if (MapListList)
-	{
-		for (FConfigSectionMap::TConstIterator It(*MapListList) ; It ; ++It)
-		{
-			FName EntryType = It.Key();
-			const FString& EntryValue = It.Value();
-
-			if (EntryType == NAME_Map)
-			{
-				// Add it to the list
-				OutMapList.AddUnique(EntryValue);
-			}
-			else if (EntryType == FName(TEXT("Section")))
-			{
-				// Recurse...
-				LoadMapListFromIni(EntryValue, OutMapList);
-			}
-			else
-			{
-				UE_LOG(LogEditor, Warning, TEXT("Invalid entry in map ini list: %s, %s=%s"),
-					*InSectionName, *(EntryType.ToString()), *EntryValue);
-			}
-		}
-	}
-}
-
-void UEditorEngine::SyncBrowserToObjects( TArray<UObject*>& InObjectsToSync )
-{
-	FContentBrowserModule& ContentBrowserModule = FModuleManager::Get().LoadModuleChecked<FContentBrowserModule>("ContentBrowser");
-	ContentBrowserModule.Get().SyncBrowserToAssets( InObjectsToSync );
-
-}
-
-void UEditorEngine::SyncBrowserToObjects( TArray<class FAssetData>& InAssetsToSync )
-{
-	FContentBrowserModule& ContentBrowserModule = FModuleManager::Get().LoadModuleChecked<FContentBrowserModule>("ContentBrowser");
-	ContentBrowserModule.Get().SyncBrowserToAssets( InAssetsToSync );
-}
-
-
-bool UEditorEngine::CanSyncToContentBrowser()
-{
-	TArray< UObject*> Objects;
-	GetObjectsToSyncToContentBrowser( Objects );
-	return Objects.Num() > 0;
-}
-
-
-void UEditorEngine::GetObjectsToSyncToContentBrowser( TArray<UObject*>& Objects )
-{
-	// If the user has any BSP surfaces selected, sync to the materials on them.
-	bool bFoundSurfaceMaterial = false;
-
-	for ( TSelectedSurfaceIterator<> It(GWorld) ; It ; ++It )
-	{
-		FBspSurf* Surf = *It;
-		UMaterialInterface* Material = Surf->Material;
-		if( Material )
-		{
-			Objects.AddUnique( Material );
-			bFoundSurfaceMaterial = true;
-		}
-	}
-
-	// Otherwise, assemble a list of resources from selected actors.
-	if( !bFoundSurfaceMaterial )
-	{
-		for ( FSelectionIterator It( GEditor->GetSelectedActorIterator() ) ; It ; ++It )
-		{
-			AActor* Actor = static_cast<AActor*>( *It );
-			checkSlow( Actor->IsA(AActor::StaticClass()) );
-
-			// If the actor is an instance of a blueprint, just add the blueprint.
-			UBlueprint* GeneratingBP = Cast<UBlueprint>(It->GetClass()->ClassGeneratedBy);
-			if ( GeneratingBP != NULL )
-			{
-				Objects.Add(GeneratingBP);
-			}
-			// Otherwise, add the results of the GetReferencedContentObjects call
-			else
-			{
-				Actor->GetReferencedContentObjects(Objects);
-			}
-		}
-	}
-}
-
-void UEditorEngine::SyncToContentBrowser()
-{
-	TArray<UObject*> Objects;
-
-	GetObjectsToSyncToContentBrowser( Objects );
-
-	// Sync the content browser to the object list.
-	SyncBrowserToObjects(Objects);
-}
-
-
-void UEditorEngine::GetReferencedAssetsForEditorSelection(TArray<UObject*>& Objects, const bool bIgnoreOtherAssetsIfBPReferenced)
-{
-	for ( TSelectedSurfaceIterator<> It(GWorld) ; It ; ++It )
-	{
-		FBspSurf* Surf = *It;
-		UMaterialInterface* Material = Surf->Material;
-		if( Material )
-		{
-			Objects.AddUnique( Material );
-		}
-	}
-
-	for ( FSelectionIterator It( GEditor->GetSelectedActorIterator() ) ; It ; ++It )
-	{
-		AActor* Actor = static_cast<AActor*>( *It );
-		checkSlow( Actor->IsA(AActor::StaticClass()) );
-
-		TArray<UObject*> ActorObjects;
-		Actor->GetReferencedContentObjects(ActorObjects);
-
-		// If Blueprint assets should take precedence over any other referenced asset, check if there are any blueprints in this actor's list
-		// and if so, add only those.
-		if (bIgnoreOtherAssetsIfBPReferenced && ActorObjects.ContainsByPredicate([](UObject* Obj) { return Obj->IsA(UBlueprint::StaticClass()); }))
-		{
-			for (UObject* Object : ActorObjects)
-			{
-				if (Object->IsA(UBlueprint::StaticClass()))
-				{
-					Objects.Add(Object);
-				}
-			}
-		}
-		else
-		{
-			Objects.Append(ActorObjects);
-		}
-	}
-}
-
-
-void UEditorEngine::ToggleSelectedActorMovementLock()
-{
-	// First figure out if any selected actor is already locked.
-	const bool bFoundLockedActor = HasLockedActors();
-
-	// Fires ULevel::LevelDirtiedEvent when falling out of scope.
-	FScopedLevelDirtied		LevelDirtyCallback;
-
-	for ( FSelectionIterator It( GetSelectedActorIterator() ) ; It ; ++It )
-	{
-		AActor* Actor = Cast<AActor>( *It );
-		checkSlow( Actor );
-
-		Actor->Modify();
-
-		// If nothing is locked then we'll turn on locked for all selected actors
-		// Otherwise, we'll turn off locking for any actors that are locked
-		Actor->bLockLocation = !bFoundLockedActor;
-
-		LevelDirtyCallback.Request();
-	}
-
-	bCheckForLockActors = true;
-}
-
-bool UEditorEngine::HasLockedActors()
-{
-	if( bCheckForLockActors )
-	{
-		bHasLockedActors = false;
-		for ( FSelectionIterator It( GetSelectedActorIterator() ) ; It ; ++It )
-		{
-			AActor* Actor = Cast<AActor>( *It );
-			checkSlow( Actor );
-
-			if( Actor->bLockLocation )
-			{
-				bHasLockedActors = true;
-				break;
-			}
-		}
-		bCheckForLockActors = false;
-	}
-
-	return bHasLockedActors;
-}
-
-void UEditorEngine::EditObject( UObject* ObjectToEdit )
-{
-	// @todo toolkit minor: Needs world-centric support?
-	FAssetEditorManager::Get().OpenEditorForAsset(ObjectToEdit);
-}
-
-void UEditorEngine::SelectLevelInLevelBrowser( bool bDeselectOthers )
-{
-	if( bDeselectOthers )
-	{
-		AActor* Actor = Cast<AActor>( *FSelectionIterator(*GEditor->GetSelectedActors()) );
-		if(Actor)
-		{
-			TArray<class ULevel*> EmptyLevelsList;
-			Actor->GetWorld()->SetSelectedLevels(EmptyLevelsList);
-		}
-	}
-
-	for ( FSelectionIterator Itor(*GEditor->GetSelectedActors()) ; Itor ; ++Itor )
-	{
-		AActor* Actor = Cast<AActor>( *Itor);
-		if ( Actor )
-		{
-			Actor->GetWorld()->SelectLevel( Actor->GetLevel() );
-		}
-	}
-
-	FLevelEditorModule& LevelEditorModule = FModuleManager::LoadModuleChecked<FLevelEditorModule>("LevelEditor");
-	if (FParse::Param(FCommandLine::Get(), TEXT("oldlevels")))
-	{
-		LevelEditorModule.SummonLevelBrowser();
-	}
-	else
-	{
-		LevelEditorModule.SummonWorldBrowserHierarchy();
-	}
-}
-
-void UEditorEngine::DeselectLevelInLevelBrowser()
-{
-	for ( FSelectionIterator Itor(*GEditor->GetSelectedActors()) ; Itor ; ++Itor )
-	{
-		AActor* Actor = Cast<AActor>( *Itor);
-		if ( Actor )
-		{
-			Actor->GetWorld()->DeSelectLevel( Actor->GetLevel() );
-		}
-	}
-	
-	FLevelEditorModule& LevelEditorModule = FModuleManager::LoadModuleChecked<FLevelEditorModule>("LevelEditor");
-	if (FParse::Param(FCommandLine::Get(), TEXT("oldlevels")))
-	{
-		LevelEditorModule.SummonLevelBrowser();
-	}
-	else
-	{
-		LevelEditorModule.SummonWorldBrowserHierarchy();
-	}
-}
-
-void UEditorEngine::SelectAllActorsControlledByMatinee()
-{
-	TArray<AActor *> AllActors;
-	UWorld* IteratorWorld = GWorld;
-	for( FSelectedActorIterator Iter(IteratorWorld); Iter; ++Iter)
-	{
-		AMatineeActor * CurActor = Cast<AMatineeActor>(*Iter);
-		if ( CurActor )
-		{
-			TArray<AActor*> Actors;			
-			CurActor->GetControlledActors(Actors);
-			AllActors.Append(Actors);
-		}
-	}
-
-	GUnrealEd->SelectNone(false, true, false);
-	for(int32 i=0; i<AllActors.Num(); i++)
-	{
-		GUnrealEd->SelectActor( AllActors[i], true, false, true );
-	}
-	GUnrealEd->NoteSelectionChange();
-}
-
-void UEditorEngine::SelectAllActorsWithClass( bool bArchetype )
-{
-	if( !bArchetype )
-	{
-		FSelectedActorInfo SelectionInfo = AssetSelectionUtils::GetSelectedActorInfo();
-
-		if( SelectionInfo.NumSelected )
-		{
-			GUnrealEd->Exec( SelectionInfo.SharedWorld, *FString::Printf( TEXT("ACTOR SELECT OFCLASS CLASS=%s"), *SelectionInfo.SelectionStr ) );
-		}
-	}
-	else
-	{
-		// For this function to have been called in the first place, all of the selected actors should be of the same type
-		// and with the same archetype; however, it's safest to confirm the assumption first
-		bool bAllSameClassAndArchetype = false;
-		TSubclassOf<AActor> FirstClass;
-		UObject* FirstArchetype = NULL;
-
-		// Find the class and archetype of the first selected actor; they will be used to check that all selected actors
-		// share the same class and archetype
-		UWorld* IteratorWorld = GWorld;
-		FSelectedActorIterator SelectedActorIter(IteratorWorld);
-		if ( SelectedActorIter )
-		{
-			AActor* FirstActor = *SelectedActorIter;
-			check( FirstActor );
-			FirstClass = FirstActor->GetClass();
-			FirstArchetype = FirstActor->GetArchetype();
-
-			// If the archetype of the first actor is NULL, then do not allow the selection to proceed
-			bAllSameClassAndArchetype = FirstArchetype ? true : false;
-
-			// Increment the iterator so the search begins on the second selected actor
-			++SelectedActorIter;
-		}
-		// Check all the other selected actors
-		for ( ; SelectedActorIter && bAllSameClassAndArchetype; ++SelectedActorIter )
-		{
-			AActor* CurActor = *SelectedActorIter;
-			if ( CurActor->GetClass() != FirstClass || CurActor->GetArchetype() != FirstArchetype )
-			{
-				bAllSameClassAndArchetype = false;
-				break;
-			}
-		}
-
-		// If all the selected actors have the same class and archetype, then go ahead and select all other actors
-		// matching the same class and archetype
-		if ( bAllSameClassAndArchetype )
-		{
-			FScopedTransaction Transaction( NSLOCTEXT("UnrealEd", "SelectOfClassAndArchetype", "Select of Class and Archetype") );
-			GUnrealEd->edactSelectOfClassAndArchetype( IteratorWorld, FirstClass, FirstArchetype );
-		}
-	}
-}
-
-
-void UEditorEngine::FindSelectedActorsInLevelScript()
-{
-	AActor* Actor = GEditor->GetSelectedActors()->GetTop<AActor>();
-	if(Actor != NULL)
-	{
-		FKismetEditorUtilities::ShowActorReferencesInLevelScript(Actor);
-	}
-}
-
-bool UEditorEngine::AreAnySelectedActorsInLevelScript()
-{
-	AActor* Actor = GEditor->GetSelectedActors()->GetTop<AActor>();
-	if(Actor != NULL)
-	{
-		ULevelScriptBlueprint* LSB = Actor->GetLevel()->GetLevelScriptBlueprint(true);
-		if( LSB != NULL )
-		{
-			int32 NumRefs = FBlueprintEditorUtils::FindNumReferencesToActorFromLevelScript(LSB, Actor);
-			if(NumRefs > 0)
-			{
-				return true;
-			}
-		}
-	}
-
-	return false;
-}
-
-void UEditorEngine::ConvertSelectedBrushesToVolumes( UClass* VolumeClass )
-{
-	TArray<ABrush*> BrushesToConvert;
-	for ( FSelectionIterator SelectedActorIter( GEditor->GetSelectedActorIterator() ); SelectedActorIter; ++SelectedActorIter )
-	{
-		AActor* CurSelectedActor = Cast<AActor>( *SelectedActorIter );
-		check( CurSelectedActor );
-		ABrush* Brush = Cast< ABrush >( CurSelectedActor );
-		if ( Brush && !FActorEditorUtils::IsABuilderBrush(CurSelectedActor) )
-		{
-			ABrush* CurBrushActor = CastChecked<ABrush>( CurSelectedActor );
-
-			BrushesToConvert.Add(CurBrushActor);
-		}
-	}
-
-	if (BrushesToConvert.Num())
-	{
-		GEditor->GetSelectedActors()->BeginBatchSelectOperation();
-
-		const FScopedTransaction Transaction( FText::Format( NSLOCTEXT("UnrealEd", "Transaction_ConvertToVolume", "Convert to Volume: {0}"), FText::FromString( VolumeClass->GetName() ) ) );
-		checkSlow( VolumeClass && VolumeClass->IsChildOf( AVolume::StaticClass() ) );
-
-		TArray< UWorld* > WorldsAffected;
-		// Iterate over all selected actors, converting the brushes to volumes of the provided class
-		for ( int32 BrushIdx = 0; BrushIdx < BrushesToConvert.Num(); BrushIdx++ )
-		{
-			ABrush* CurBrushActor = BrushesToConvert[BrushIdx];
-			check( CurBrushActor );
-			
-			ULevel* CurActorLevel = CurBrushActor->GetLevel();
-			check( CurActorLevel );
-
-			// Cache the world and store in a list.
-			UWorld* World = CurBrushActor->GetWorld();
-			check( World );
-			WorldsAffected.AddUnique( World );
-
-			FActorSpawnParameters SpawnInfo;
-			SpawnInfo.OverrideLevel = CurActorLevel;
-			ABrush* NewVolume = World->SpawnActor<ABrush>( VolumeClass, CurBrushActor->GetActorLocation(), FRotator::ZeroRotator );
-			if ( NewVolume )
-			{
-				NewVolume->PreEditChange( NULL );
-
-				FBSPOps::csgCopyBrush( NewVolume, CurBrushActor, 0, RF_Transactional, true, true );
-
-				// Set the texture on all polys to NULL.  This stops invisible texture
-				// dependencies from being formed on volumes.
-				if( NewVolume->Brush )
-				{
-					for ( TArray<FPoly>::TIterator PolyIter( NewVolume->Brush->Polys->Element ); PolyIter; ++PolyIter )
-					{
-						FPoly& CurPoly = *PolyIter;
-						CurPoly.Material = NULL;
-					}
-				}
-
-				// Select the new actor
-				GEditor->SelectActor( CurBrushActor, false, true );
-				GEditor->SelectActor( NewVolume, true, true );
-
-				NewVolume->PostEditChange();
-				NewVolume->PostEditMove( true );
-				NewVolume->Modify();
-
-				// Destroy the old actor.
-				GEditor->Layers->DisassociateActorFromLayers( CurBrushActor );
-				World->EditorDestroyActor( CurBrushActor, true );
-			}
-		}
-
-		GEditor->GetSelectedActors()->EndBatchSelectOperation();
-		GEditor->RedrawLevelEditingViewports();
-
-		// Broadcast a message that the levels in these worlds have changed
-		for (int32 iWorld = 0; iWorld < WorldsAffected.Num() ; iWorld++)
-		{
-			WorldsAffected[ iWorld ]->BroadcastLevelsChanged();
-		}
-
-		CollectGarbage( GARBAGE_COLLECTION_KEEPFLAGS );
-	}
-}
-
-/** Utility for copying properties that differ from defaults between mesh types. */
-struct FConvertStaticMeshActorInfo
-{
-	/** The level the source actor belonged to, and into which the new actor is created. */
-	ULevel*						SourceLevel;
-
-	// Actor properties.
-	FVector						Location;
-	FRotator					Rotation;
-	FVector						DrawScale3D;
-	bool						bHidden;
-	AActor*						Base;
-	UPrimitiveComponent*		BaseComponent;
-	// End actor properties.
-
-	/**
-	 * Used to indicate if any of the above properties differ from defaults; if so, they're copied over.
-	 * We don't want to simply copy all properties, because classes with different defaults will have
-	 * their defaults hosed by other types.
-	 */
-	bool bActorPropsDifferFromDefaults[14];
-
-	// Component properties.
-	UStaticMesh*						StaticMesh;
-	USkeletalMesh*						SkeletalMesh;
-	TArray<UMaterialInterface*>			OverrideMaterials;
-	TArray<FGuid>						IrrelevantLights;
-	float								CachedMaxDrawDistance;
-	bool								CastShadow;
-
-	FBodyInstance						BodyInstance;
-	TArray< TArray<FColor> >			OverrideVertexColors;
-
-
-	// for skeletalmeshcomponent animation conversion
-	// this is temporary until we have SkeletalMeshComponent.Animations
-	UAnimationAsset*					AnimAsset;
-	UVertexAnimation*					VertexAnimation;
-	bool								bLooping;
-	bool								bPlaying;
-	float								Rate;
-	float								CurrentPos;
-
-	// End component properties.
-
-	/**
-	 * Used to indicate if any of the above properties differ from defaults; if so, they're copied over.
-	 * We don't want to simply copy all properties, because classes with different defaults will have
-	 * their defaults hosed by other types.
-	 */
-	bool bComponentPropsDifferFromDefaults[7];
-
-	AGroupActor* ActorGroup;
-
-	bool PropsDiffer(const TCHAR* PropertyPath, UObject* Obj)
-	{
-		const UProperty* PartsProp = FindObjectChecked<UProperty>( ANY_PACKAGE, PropertyPath );
-
-		uint8* ClassDefaults = (uint8*)Obj->GetClass()->GetDefaultObject();
-		check( ClassDefaults );
-
-		for (int32 Index = 0; Index < PartsProp->ArrayDim; Index++)
-		{
-			const bool bMatches = PartsProp->Identical_InContainer(Obj, ClassDefaults, Index);
-			if (!bMatches)
-			{
-				return true;
-			}
-		}
-		return false;
-	}
-
-	void GetFromActor(AActor* Actor, UStaticMeshComponent* MeshComp)
-	{
-		InternalGetFromActor(Actor);
-
-		// Copy over component properties.
-		StaticMesh				= MeshComp->StaticMesh;
-		OverrideMaterials		= MeshComp->OverrideMaterials;
-		IrrelevantLights		= MeshComp->IrrelevantLights;
-		CachedMaxDrawDistance	= MeshComp->CachedMaxDrawDistance;
-		CastShadow				= MeshComp->CastShadow;
-
-		BodyInstance.CopyBodyInstancePropertiesFrom(&MeshComp->BodyInstance);
-
-		// Loop over each LODInfo in the static mesh component, storing the override vertex colors
-		// in each, if any
-		bool bHasAnyVertexOverrideColors = false;
-		for ( int32 LODIndex = 0; LODIndex < MeshComp->LODData.Num(); ++LODIndex )
-		{
-			const FStaticMeshComponentLODInfo& CurLODInfo = MeshComp->LODData[LODIndex];
-			const FColorVertexBuffer* CurVertexBuffer = CurLODInfo.OverrideVertexColors;
-
-			OverrideVertexColors.Add( TArray<FColor>() );
-			
-			// If the LODInfo has override vertex colors, store off each one
-			if ( CurVertexBuffer && CurVertexBuffer->GetNumVertices() > 0 )
-			{
-				for ( uint32 VertexIndex = 0; VertexIndex < CurVertexBuffer->GetNumVertices(); ++VertexIndex )
-				{
-					OverrideVertexColors[LODIndex].Add( CurVertexBuffer->VertexColor(VertexIndex) );
-				}
-				bHasAnyVertexOverrideColors = true;
-			}
-		}
-
-		// Record which component properties differ from their defaults.
-		bComponentPropsDifferFromDefaults[0] = PropsDiffer( TEXT("Engine.StaticMeshComponent:StaticMesh"), MeshComp );
-		bComponentPropsDifferFromDefaults[1] = true; // Assume the materials array always differs.
-		bComponentPropsDifferFromDefaults[2] = true; // Assume the set of irrelevant lights always differs.
-		bComponentPropsDifferFromDefaults[3] = PropsDiffer( TEXT("Engine.PrimitiveComponent:CachedMaxDrawDistance"), MeshComp );
-		bComponentPropsDifferFromDefaults[4] = PropsDiffer( TEXT("Engine.PrimitiveComponent:CastShadow"), MeshComp );
-		bComponentPropsDifferFromDefaults[5] = PropsDiffer( TEXT("Engine.PrimitiveComponent:BodyInstance"), MeshComp );
-		bComponentPropsDifferFromDefaults[6] = bHasAnyVertexOverrideColors;	// Differs from default if there are any vertex override colors
-	}
-
-	void SetToActor(AActor* Actor, UStaticMeshComponent* MeshComp)
-	{
-		InternalSetToActor(Actor);
-
-		// Set component properties.
-		if ( bComponentPropsDifferFromDefaults[0] ) MeshComp->StaticMesh			= StaticMesh;
-		if ( bComponentPropsDifferFromDefaults[1] ) MeshComp->OverrideMaterials		= OverrideMaterials;
-		if ( bComponentPropsDifferFromDefaults[2] ) MeshComp->IrrelevantLights		= IrrelevantLights;
-		if ( bComponentPropsDifferFromDefaults[3] ) MeshComp->CachedMaxDrawDistance	= CachedMaxDrawDistance;
-		if ( bComponentPropsDifferFromDefaults[4] ) MeshComp->CastShadow			= CastShadow;
-		if ( bComponentPropsDifferFromDefaults[5] ) 
-		{
-			MeshComp->BodyInstance.CopyBodyInstancePropertiesFrom(&BodyInstance);
-		}
-		if ( bComponentPropsDifferFromDefaults[6] )
-		{
-			// Ensure the LODInfo has the right number of entries
-			MeshComp->SetLODDataCount( OverrideVertexColors.Num(), MeshComp->StaticMesh->GetNumLODs() );
-			
-			// Loop over each LODInfo to see if there are any vertex override colors to restore
-			for ( int32 LODIndex = 0; LODIndex < MeshComp->LODData.Num(); ++LODIndex )
-			{
-				FStaticMeshComponentLODInfo& CurLODInfo = MeshComp->LODData[LODIndex];
-
-				// If there are override vertex colors specified for a particular LOD, set them in the LODInfo
-				if ( OverrideVertexColors.IsValidIndex( LODIndex ) && OverrideVertexColors[LODIndex].Num() > 0 )
-				{
-					const TArray<FColor>& OverrideColors = OverrideVertexColors[LODIndex];
-					
-					// Destroy the pre-existing override vertex buffer if it's not the same size as the override colors to be restored
-					if ( CurLODInfo.OverrideVertexColors && CurLODInfo.OverrideVertexColors->GetNumVertices() != OverrideColors.Num() )
-					{
-						CurLODInfo.ReleaseOverrideVertexColorsAndBlock();
-					}
-
-					// If there is a pre-existing color vertex buffer that is valid, release the render thread's hold on it and modify
-					// it with the saved off colors
-					if ( CurLODInfo.OverrideVertexColors )
-					{								
-						CurLODInfo.BeginReleaseOverrideVertexColors();
-						FlushRenderingCommands();
-						for ( int32 VertexIndex = 0; VertexIndex < OverrideColors.Num(); ++VertexIndex )
-						{
-							CurLODInfo.OverrideVertexColors->VertexColor(VertexIndex) = OverrideColors[VertexIndex];
-						}
-					}
-
-					// If there isn't a pre-existing color vertex buffer, create one and initialize it with the saved off colors 
-					else
-					{
-						CurLODInfo.OverrideVertexColors = new FColorVertexBuffer();
-						CurLODInfo.OverrideVertexColors->InitFromColorArray( OverrideColors );
-					}
-					BeginInitResource(CurLODInfo.OverrideVertexColors);
-				}
-			}
-		}
-	}
-
-	void GetFromActor(AActor* Actor, USkeletalMeshComponent* MeshComp)
-	{
-		InternalGetFromActor(Actor);
-
-		// Copy over component properties.
-		SkeletalMesh			= MeshComp->SkeletalMesh;
-		OverrideMaterials		= MeshComp->OverrideMaterials;
-		CachedMaxDrawDistance	= MeshComp->CachedMaxDrawDistance;
-		CastShadow				= MeshComp->CastShadow;
-
-		BodyInstance.CopyBodyInstancePropertiesFrom(&MeshComp->BodyInstance);
-
-		// Record which component properties differ from their defaults.
-		bComponentPropsDifferFromDefaults[0] = PropsDiffer( TEXT("Engine.SkinnedMeshComponent:SkeletalMesh"), MeshComp );
-		bComponentPropsDifferFromDefaults[1] = true; // Assume the materials array always differs.
-		bComponentPropsDifferFromDefaults[2] = true; // Assume the set of irrelevant lights always differs.
-		bComponentPropsDifferFromDefaults[3] = PropsDiffer( TEXT("Engine.PrimitiveComponent:CachedMaxDrawDistance"), MeshComp );
-		bComponentPropsDifferFromDefaults[4] = PropsDiffer( TEXT("Engine.PrimitiveComponent:CastShadow"), MeshComp );
-		bComponentPropsDifferFromDefaults[5] = PropsDiffer( TEXT("Engine.PrimitiveComponent:BodyInstance"), MeshComp );
-		bComponentPropsDifferFromDefaults[6] = false;	// Differs from default if there are any vertex override colors
-
-		InternalGetAnimationData(MeshComp);
-	}
-
-	void SetToActor(AActor* Actor, USkeletalMeshComponent* MeshComp)
-	{
-		InternalSetToActor(Actor);
-
-		// Set component properties.
-		if ( bComponentPropsDifferFromDefaults[0] ) MeshComp->SkeletalMesh			= SkeletalMesh;
-		if ( bComponentPropsDifferFromDefaults[1] ) MeshComp->OverrideMaterials		= OverrideMaterials;
-		if ( bComponentPropsDifferFromDefaults[3] ) MeshComp->CachedMaxDrawDistance	= CachedMaxDrawDistance;
-		if ( bComponentPropsDifferFromDefaults[4] ) MeshComp->CastShadow			= CastShadow;
-		if ( bComponentPropsDifferFromDefaults[5] ) MeshComp->BodyInstance.CopyBodyInstancePropertiesFrom(&BodyInstance);
-
-		InternalSetAnimationData(MeshComp);
-	}
-private:
-	void InternalGetFromActor(AActor* Actor)
-	{
-		SourceLevel				= Actor->GetLevel();
-
-		// Copy over actor properties.
-		Location				= Actor->GetActorLocation();
-		Rotation				= Actor->GetActorRotation();
-		DrawScale3D				= Actor->GetRootComponent() ? Actor->GetRootComponent()->RelativeScale3D : FVector(1.f,1.f,1.f);
-		bHidden					= Actor->bHidden;
-
-		// Record which actor properties differ from their defaults.
-		// we don't have properties for location, rotation, scale3D, so copy all the time. 
-		bActorPropsDifferFromDefaults[0] = true; 
-		bActorPropsDifferFromDefaults[1] = true; 
-		bActorPropsDifferFromDefaults[2] = false;
-		bActorPropsDifferFromDefaults[4] = true; 
-		bActorPropsDifferFromDefaults[5] = PropsDiffer( TEXT("Engine.Actor:bHidden"), Actor );
-		bActorPropsDifferFromDefaults[7] = false;
-		// used to point to Engine.Actor.bPathColliding
-		bActorPropsDifferFromDefaults[9] = false;
-	}
-
-	void InternalSetToActor(AActor* Actor)
-	{
-		if ( Actor->GetLevel() != SourceLevel )
-		{
-			UE_LOG(LogEditor, Fatal, TEXT("Actor was converted into a different level."));
-		}
-
-		// Set actor properties.
-		if ( bActorPropsDifferFromDefaults[0] ) Actor->SetActorLocation(Location, false);
-		if ( bActorPropsDifferFromDefaults[1] ) Actor->SetActorRotation(Rotation);
-		if ( bActorPropsDifferFromDefaults[4] )
-		{
-			if( Actor->GetRootComponent() != NULL )
-			{
-				Actor->GetRootComponent()->SetRelativeScale3D( DrawScale3D );
-			}
-		}
-		if ( bActorPropsDifferFromDefaults[5] ) Actor->bHidden				= bHidden;
-	}
-
-
-	void InternalGetAnimationData(USkeletalMeshComponent * SkeletalComp)
-	{
-		AnimAsset = SkeletalComp->AnimationData.AnimToPlay;
-		VertexAnimation = SkeletalComp->AnimationData.VertexAnimToPlay;
-		bLooping = SkeletalComp->AnimationData.bSavedLooping;
-		bPlaying = SkeletalComp->AnimationData.bSavedPlaying;
-		Rate = SkeletalComp->AnimationData.SavedPlayRate;
-		CurrentPos = SkeletalComp->AnimationData.SavedPosition;
-	}
-
-	void InternalSetAnimationData(USkeletalMeshComponent * SkeletalComp)
-	{
-		if (!AnimAsset && !VertexAnimation)
-		{
-			return;
-		}
-
-		UE_LOG(LogAnimation, Log, TEXT("Converting animation data for (%s) : %s(%s), bLooping(%d), bPlaying(%d), Rate(%0.2f), CurrentPos(%0.2f)"), 
-			AnimAsset? TEXT("AnimAsset") : TEXT("VertexAnim"),
-			AnimAsset? *AnimAsset->GetName() : *VertexAnimation->GetName(), bLooping, bPlaying, Rate, CurrentPos);
-
-		SkeletalComp->AnimationData.AnimToPlay = AnimAsset;
-		SkeletalComp->AnimationData.VertexAnimToPlay = VertexAnimation;
-		SkeletalComp->AnimationData.bSavedLooping = bLooping;
-		SkeletalComp->AnimationData.bSavedPlaying = bPlaying;
-		SkeletalComp->AnimationData.SavedPlayRate = Rate;
-		SkeletalComp->AnimationData.SavedPosition = CurrentPos;
-		// we don't convert back to SkeletalMeshComponent.Animations - that will be gone soon
-	}
-};
-
-void UEditorEngine::ConvertActorsFromClass( UClass* FromClass, UClass* ToClass )
-{
-	const bool bFromInteractiveFoliage = FromClass == AInteractiveFoliageActor::StaticClass();
-	// InteractiveFoliageActor derives from StaticMeshActor.  bFromStaticMesh should only convert static mesh actors that arent supported by some other conversion
-	const bool bFromStaticMesh = !bFromInteractiveFoliage && FromClass->IsChildOf( AStaticMeshActor::StaticClass() );
-	const bool bFromSkeletalMesh = FromClass->IsChildOf(ASkeletalMeshActor::StaticClass());
-
-	const bool bToInteractiveFoliage = ToClass == AInteractiveFoliageActor::StaticClass();
-	const bool bToStaticMesh = ToClass->IsChildOf( AStaticMeshActor::StaticClass() );
-	const bool bToSkeletalMesh = ToClass->IsChildOf(ASkeletalMeshActor::StaticClass());
-
-	const bool bFoundTarget = bToInteractiveFoliage || bToStaticMesh || bToSkeletalMesh;
-
-	TArray<AActor*>				SourceActors;
-	TArray<FConvertStaticMeshActorInfo>	ConvertInfo;
-
-	// Provide the option to abort up-front.
-	if ( !bFoundTarget || GUnrealEd->ShouldAbortActorDeletion() )
-	{
-		return;
-	}
-
-	const FScopedTransaction Transaction( NSLOCTEXT("UnrealEd", "ConvertMeshes", "Convert Meshes") );
-	// Iterate over selected Actors.
-	for ( FSelectionIterator It( GEditor->GetSelectedActorIterator() ) ; It ; ++It )
-	{
-		AActor* Actor				= static_cast<AActor*>( *It );
-		checkSlow( Actor->IsA(AActor::StaticClass()) );
-
-		AStaticMeshActor* SMActor				= bFromStaticMesh ? Cast<AStaticMeshActor>(Actor) : NULL;
-		AInteractiveFoliageActor* FoliageActor	= bFromInteractiveFoliage ? Cast<AInteractiveFoliageActor>(Actor) : NULL;
-		ASkeletalMeshActor* SKMActor			= bFromSkeletalMesh? Cast<ASkeletalMeshActor>(Actor) : NULL;
-
-		const bool bFoundActorToConvert = SMActor || FoliageActor || SKMActor;
-		if ( bFoundActorToConvert )
-		{
-			// clear all transient properties before copying from
-			Actor->UnregisterAllComponents();
-
-			// If its the type we are converting 'from' copy its properties and remember it.
-			FConvertStaticMeshActorInfo Info;
-			FMemory::Memzero(&Info, sizeof(FConvertStaticMeshActorInfo));
-
-			if( SMActor )
-			{
-				SourceActors.Add(Actor);
-				Info.GetFromActor(SMActor, SMActor->GetStaticMeshComponent());
-			}
-			else if( FoliageActor )
-			{
-				SourceActors.Add(Actor);
-				Info.GetFromActor(FoliageActor, FoliageActor->GetStaticMeshComponent());
-			}
-			else if ( bFromSkeletalMesh )
-			{
-				SourceActors.Add(Actor);
-				Info.GetFromActor(SKMActor, SKMActor->GetSkeletalMeshComponent());
-			}
-
-			// Get the actor group if any
-			Info.ActorGroup = AGroupActor::GetParentForActor(Actor);
-
-			ConvertInfo.Add(MoveTemp(Info));
-		}
-	}
-
-	if (SourceActors.Num())
-	{
-		GEditor->GetSelectedActors()->BeginBatchSelectOperation();
-
-		// Then clear selection, select and delete the source actors.
-		GEditor->SelectNone( false, false );
-		UWorld* World = NULL;
-		for( int32 ActorIndex = 0 ; ActorIndex < SourceActors.Num() ; ++ActorIndex )
-		{
-			AActor* SourceActor = SourceActors[ActorIndex];
-			GEditor->SelectActor( SourceActor, true, false );
-			World = SourceActor->GetWorld();
-		}
-		
-		if ( World && GUnrealEd->edactDeleteSelected( World, false ) )
-		{
-			// Now we need to spawn some new actors at the desired locations.
-			for( int32 i = 0 ; i < ConvertInfo.Num() ; ++i )
-			{
-				FConvertStaticMeshActorInfo& Info = ConvertInfo[i];
-
-				// Spawn correct type, and copy properties from intermediate struct.
-				AActor* Actor = NULL;
-				
-				// Cache the world pointer
-				check( World == Info.SourceLevel->OwningWorld );
-				
-				FActorSpawnParameters SpawnInfo;
-				SpawnInfo.OverrideLevel = Info.SourceLevel;
-				SpawnInfo.bNoCollisionFail = true;
-				if( bToStaticMesh )
-				{					
-					AStaticMeshActor* SMActor = CastChecked<AStaticMeshActor>( World->SpawnActor( ToClass, &Info.Location, &Info.Rotation, SpawnInfo ) );
-					SMActor->UnregisterAllComponents();
-					Info.SetToActor(SMActor, SMActor->GetStaticMeshComponent());
-					SMActor->RegisterAllComponents();
-					GEditor->SelectActor( SMActor, true, false );
-					Actor = SMActor;
-				}
-				else if(bToInteractiveFoliage)
-				{					
-					AInteractiveFoliageActor* FoliageActor = World->SpawnActor<AInteractiveFoliageActor>( Info.Location, Info.Rotation, SpawnInfo );
-					FoliageActor->UnregisterAllComponents();
-					Info.SetToActor(FoliageActor, FoliageActor->GetStaticMeshComponent());
-					FoliageActor->RegisterAllComponents();
-					GEditor->SelectActor( FoliageActor, true, false );
-					Actor = FoliageActor;
-				}
-				else if (bToSkeletalMesh)
-				{
-					check(ToClass->IsChildOf(ASkeletalMeshActor::StaticClass()));
-					// checked
-					ASkeletalMeshActor* SkeletalMeshActor = CastChecked<ASkeletalMeshActor>( World->SpawnActor( ToClass, &Info.Location, &Info.Rotation, SpawnInfo ));
-					SkeletalMeshActor->UnregisterAllComponents();
-					Info.SetToActor(SkeletalMeshActor, SkeletalMeshActor->GetSkeletalMeshComponent());
-					SkeletalMeshActor->RegisterAllComponents();
-					GEditor->SelectActor( SkeletalMeshActor, true, false );
-					Actor = SkeletalMeshActor;
-				}
-
-				// Fix up the actor group.
-				if( Actor )
-				{
-					if( Info.ActorGroup )
-					{
-						Info.ActorGroup->Add(*Actor);
-						Info.ActorGroup->Add(*Actor);
-					}
-				}
-			}
-		}
-
-		GEditor->GetSelectedActors()->EndBatchSelectOperation();
-	}
-}
-
-bool UEditorEngine::ShouldOpenMatinee(AMatineeActor* MatineeActor) const
-{
-	if ( MatineeActor && !MatineeActor->MatineeData )
-	{
-		FMessageDialog::Open( EAppMsgType::Ok, NSLOCTEXT("UnrealEd", "Error_MatineeActionMustHaveData", "UnrealMatinee must have valid InterpData assigned before being edited.") );
-		return false;
-	}
-
-	// Make sure we can't open the same action twice in Matinee.
-	if( GLevelEditorModeTools().IsModeActive(FBuiltinEditorModes::EM_InterpEdit) )
-	{
-		FMessageDialog::Open( EAppMsgType::Ok, NSLOCTEXT("UnrealEd", "MatineeActionAlreadyOpen", "An UnrealMatinee sequence is currently open in an editor.  Please close it before proceeding.") );
-		return false;
-	}
-
-	// Don't let you open Matinee if a transaction is currently active.
-	if( GEditor->IsTransactionActive() )
-	{
-		FMessageDialog::Open( EAppMsgType::Ok, NSLOCTEXT("UnrealEd", "TransactionIsActive", "Undo Transaction Is Active - Cannot Open UnrealMatinee.") );
-		return false;
-	}
-
-	return true;
-}
-
-void UEditorEngine::OpenMatinee(AMatineeActor* MatineeActor, bool bWarnUser)
-{
-	// Drop out if the user doesn't want to proceed to matinee atm
-	if( bWarnUser && !ShouldOpenMatinee( MatineeActor ) )
-	{
-		return;
-	}
-
-	// If already in Matinee mode, exit out before going back in with new Interpolation.
-	if( GLevelEditorModeTools().IsModeActive( FBuiltinEditorModes::EM_InterpEdit ) )
-	{
-		GLevelEditorModeTools().DeactivateMode( FBuiltinEditorModes::EM_InterpEdit );
-	}
-
-	GLevelEditorModeTools().ActivateMode( FBuiltinEditorModes::EM_InterpEdit );
-
-	FEdModeInterpEdit* InterpEditMode = (FEdModeInterpEdit*)GLevelEditorModeTools().GetActiveMode( FBuiltinEditorModes::EM_InterpEdit );
-
-	InterpEditMode->InitInterpMode( MatineeActor );
-
-	OnOpenMatinee();
-}
-
-void UEditorEngine::UpdateReflectionCaptures()
-{
-	// Update sky light first because it's considered direct lighting, sky diffuse will be visible in reflection capture indirect specular
-	UWorld* World = GWorld;
-	World->UpdateAllSkyCaptures();
-	World->UpdateAllReflectionCaptures();
-}
-
-void UEditorEngine::UpdateSkyCaptures()
-{
-	GWorld->UpdateAllSkyCaptures();
-}
-
-void UEditorEngine::EditorAddModalWindow( TSharedRef<SWindow> InModalWindow ) const
-{
-	// If there is already a modal window active, parent this new modal window to the existing window so that it doesnt fall behind
-	TSharedPtr<SWindow> ParentWindow = FSlateApplication::Get().GetActiveModalWindow();
-
-	if( !ParentWindow.IsValid() )
-	{
-		// Parent to the main frame window
-		if(FModuleManager::Get().IsModuleLoaded("MainFrame"))
-		{
-			IMainFrameModule& MainFrame = FModuleManager::LoadModuleChecked<IMainFrameModule>("MainFrame");
-			ParentWindow = MainFrame.GetParentWindow();
-		}
-	}
-
-	FSlateApplication::Get().AddModalWindow( InModalWindow, ParentWindow );
-}
-
-UBrushBuilder* UEditorEngine::FindBrushBuilder( UClass* BrushBuilderClass )
-{
-	TArray< UBrushBuilder* > FoundBuilders;
-	UBrushBuilder* Builder = NULL;
-	// Search for the existing brush builder
-	if( ContainsObjectOfClass<UBrushBuilder>( BrushBuilders, BrushBuilderClass, true, &FoundBuilders ) )
-	{
-		// Should not be more than one of the same type
-		check( FoundBuilders.Num() == 1 );
-		Builder = FoundBuilders[0];
-	}
-	else
-	{
-		// An existing builder does not exist so create one now
-		Builder = NewObject<UBrushBuilder>(GetTransientPackage(), BrushBuilderClass);
-		BrushBuilders.Add( Builder );
-	}
-
-	return Builder;
-}
-
-void UEditorEngine::ParentActors( AActor* ParentActor, AActor* ChildActor, const FName SocketName, USceneComponent* Component)
-{
-	if (CanParentActors(ParentActor, ChildActor))
-	{
-		USceneComponent* ChildRoot = ChildActor->GetRootComponent();
-		USceneComponent* ParentRoot = ParentActor->GetRootComponent();
-
-		check(ChildRoot);	// CanParentActors() call should ensure this
-		check(ParentRoot);	// CanParentActors() call should ensure this
-
-		// modify parent and child
-		const FScopedTransaction Transaction( NSLOCTEXT("Editor", "UndoAction_PerformAttachment", "Attach actors") );
-		ChildActor->Modify();
-		ParentActor->Modify();
-
-		// If child is already attached to something, modify the old parent and detach
-		if(ChildRoot->AttachParent != NULL)
-		{
-			AActor* OldParentActor = ChildRoot->AttachParent->GetOwner();
-			OldParentActor->Modify();
-			ChildRoot->DetachFromParent(true);
-
-			GEngine->BroadcastLevelActorDetached(ChildActor, OldParentActor);
-		}
-
-		// If the parent is already attached to this child, modify its parent and detach so we can allow the attachment
-		if(ParentRoot->IsAttachedTo(ChildRoot))
-		{
-			ParentRoot->AttachParent->GetOwner()->Modify();
-			ParentRoot->DetachFromParent(true);
-		}
-
-		// Snap to socket if a valid socket name was provided, otherwise attach without changing the relative transform
-		const bool bValidSocketName = !SocketName.IsNone() && ParentRoot->DoesSocketExist(SocketName);
-		ChildRoot->AttachTo(ParentRoot, SocketName, bValidSocketName ? EAttachLocation::SnapToTarget : EAttachLocation::KeepWorldPosition);
-
-		// Refresh editor in case child was translated after snapping to socket
-		RedrawLevelEditingViewports();
-	}
-}
-
-bool UEditorEngine::DetachSelectedActors()
-{
-	FScopedTransaction Transaction( NSLOCTEXT("Editor", "UndoAction_PerformDetach", "Detach actors") );
-
-	bool bDetachOccurred = false;
-	for ( FSelectionIterator It( GEditor->GetSelectedActorIterator() ) ; It ; ++It )
-	{
-		AActor* Actor = Cast<AActor>( *It );
-		checkSlow( Actor );
-
-		USceneComponent* RootComp = Actor->GetRootComponent();
-		if( RootComp != NULL && RootComp->AttachParent != NULL)
-		{
-			AActor* OldParentActor = RootComp->AttachParent->GetOwner();
-			OldParentActor->Modify();
-			RootComp->DetachFromParent(true);
-			bDetachOccurred = true;
-			Actor->SetFolderPath(OldParentActor->GetFolderPath());
-		}
-	}
-	return bDetachOccurred;
-}
-
-bool UEditorEngine::CanParentActors( const AActor* ParentActor, const AActor* ChildActor, FText* ReasonText)
-{
-	if(ChildActor == NULL || ParentActor == NULL)
-	{
-		if (ReasonText)
-		{
-			*ReasonText = NSLOCTEXT("ActorAttachmentError", "Null_ActorAttachmentError", "Cannot attach NULL actors.");
-		}
-		return false;
-	}
-
-	if(ChildActor == ParentActor)
-	{
-		if (ReasonText)
-		{
-			*ReasonText = NSLOCTEXT("ActorAttachmentError", "Self_ActorAttachmentError", "Cannot attach actor to self.");
-		}
-		return false;
-	}
-
-	USceneComponent* ChildRoot = ChildActor->GetRootComponent();
-	USceneComponent* ParentRoot = ParentActor->GetRootComponent();
-	if(ChildRoot == NULL || ParentRoot == NULL)
-	{
-		if (ReasonText)
-		{
-			*ReasonText = NSLOCTEXT("ActorAttachmentError", "MissingComponent_ActorAttachmentError", "Cannot attach actors without root components.");
-		}
-		return false;
-	}
-
-	const ABrush* ParentBrush = Cast<const  ABrush >( ParentActor );
-	const ABrush* ChildBrush = Cast<const  ABrush >( ChildActor );
-	if( (ParentBrush && !ParentBrush->IsVolumeBrush() ) || ( ChildBrush && !ChildBrush->IsVolumeBrush() ) )
-	{
-		if (ReasonText)
-		{
-			*ReasonText = NSLOCTEXT("ActorAttachmentError", "Brush_ActorAttachmentError", "BSP Brushes cannot be attached");
-		}
-		return false;
-	}
-
-	{
-		FText Reason;
-		if (!ChildActor->EditorCanAttachTo(ParentActor, Reason))
-		{
-			if (ReasonText)
-			{
-				if (Reason.IsEmpty())
-				{
-					*ReasonText = FText::Format(NSLOCTEXT("ActorAttachmentError", "CannotBeAttached_ActorAttachmentError", "{0} cannot be attached to {1}"), FText::FromString(ChildActor->GetActorLabel()), FText::FromString(ParentActor->GetActorLabel()));
-				}
-				else
-				{
-					*ReasonText = MoveTemp(Reason);
-				}
-			}
-			return false;
-		}
-	}
-
-	if (ChildRoot->Mobility == EComponentMobility::Static && ParentRoot->Mobility != EComponentMobility::Static)
-	{
-		if (ReasonText)
-		{
-			FFormatNamedArguments Arguments;
-			Arguments.Add(TEXT("StaticActor"), FText::FromString(ChildActor->GetActorLabel()));
-			Arguments.Add(TEXT("DynamicActor"), FText::FromString(ParentActor->GetActorLabel()));
-			*ReasonText = FText::Format( NSLOCTEXT("ActorAttachmentError", "StaticDynamic_ActorAttachmentError", "Cannot attach static actor {StaticActor} to dynamic actor {DynamicActor}."), Arguments);
-		}
-		return false;
-	}
-
-	if(ChildActor->GetLevel() != ParentActor->GetLevel())
-	{
-		if (ReasonText)
-		{
-			*ReasonText = NSLOCTEXT("ActorAttachmentError", "WrongLevel_AttachmentError", "Actors need to be in the same level!");
-		}
-		return false;
-	}
-
-	if(ParentRoot->IsAttachedTo( ChildRoot ))
-	{
-		if (ReasonText)
-		{
-			*ReasonText = NSLOCTEXT("ActorAttachmentError", "CircularInsest_ActorAttachmentError", "Parent cannot become the child of their descendant");
-		}
-		return false;
-	}
-
-	return true;
-}
-
-
-bool UEditorEngine::IsPackageValidForAutoAdding(UPackage* InPackage, const FString& InFilename)
-{
-	bool bPackageIsValid = false;
-
-	// Ensure the package exists, the user is running the editor (and not a commandlet or cooking), and that source control
-	// is enabled and expecting new files to be auto-added before attempting to test the validity of the package
-	if ( InPackage && GIsEditor && !IsRunningCommandlet() && ISourceControlModule::Get().IsEnabled() && GetDefault<UEditorLoadingSavingSettings>()->bSCCAutoAddNewFiles )
-	{
-		const FString CleanFilename = FPaths::GetCleanFilename(InFilename);
-
-		// Determine if the package has been saved before or not; if it has, it's not valid for auto-adding
-		bPackageIsValid = !FPaths::FileExists(InFilename);
-
-		// If the package is still considered valid up to this point, ensure that it is not a script or PIE package
-		// and that the editor is not auto-saving.
-		if ( bPackageIsValid )
-		{
-			const bool bIsPlayOnConsolePackage = CleanFilename.StartsWith( PLAYWORLD_CONSOLE_BASE_PACKAGE_PREFIX );
-			const bool bIsPIEOrScriptPackage = InPackage->RootPackageHasAnyFlags( PKG_ContainsScript | PKG_PlayInEditor );
-			const bool bIsAutosave = GUnrealEd->GetPackageAutoSaver().IsAutoSaving();
-
-			if ( bIsPlayOnConsolePackage || bIsPIEOrScriptPackage || bIsAutosave || GIsAutomationTesting )
-			{
-				bPackageIsValid = false;
-			}
-		}
-	}
-	return bPackageIsValid;
-}
-
-/** 
- * A mapping of all startup packages to whether or not we have warned the user about editing them
- */
-static TMap<UPackage*, bool>		StartupPackageToWarnState;
-
-bool UEditorEngine::IsPackageOKToSave(UPackage* InPackage, const FString& InFilename, FOutputDevice* Error)
-{
-	TArray<FString>	AllStartupPackageNames;
-	appGetAllPotentialStartupPackageNames(AllStartupPackageNames, GEngineIni, false);
-
-	FString ConvertedPackageName, ConversionError;
-	if (!FPackageName::TryConvertFilenameToLongPackageName(InFilename, ConvertedPackageName, &ConversionError))
-	{
-		Error->Logf(ELogVerbosity::Error, *FText::Format(NSLOCTEXT("UnrealEd", "CannotConvertPackageName", "Cannot save asset '{0}' as conversion of long package name failed. Reason: '{1}'."), FText::FromString(InFilename), FText::FromString(ConversionError)).ToString());
-		return false;
-	}
-
-	bool bIsStartupPackage = AllStartupPackageNames.Contains(ConvertedPackageName);
-
-	// Make sure that if the package is a startup package, the user indeed wants to save changes
-	if( !IsRunningCommandlet()																		&& // Don't prompt about saving startup packages when running UCC
-		InFilename.EndsWith(FPackageName::GetAssetPackageExtension())				&& // Maps, even startup maps, are ok
-		bIsStartupPackage && 
-		(!StartupPackageToWarnState.Find(InPackage) || (*(StartupPackageToWarnState.Find(InPackage)) == false))
-		)
-	{		
-		// Prompt to save startup packages
-		if( EAppReturnType::Yes == FMessageDialog::Open( EAppMsgType::YesNo, FText::Format(
-				NSLOCTEXT("UnrealEd", "Prompt_AboutToEditStartupPackage", "{0} is a startup package.  Startup packages are fully cooked and loaded when on consoles. ALL CONTENT IN THIS PACKAGE WILL ALWAYS USE MEMORY. Are you sure you want to save it?"),
-				FText::FromString(InPackage->GetName()))) )
-		{
-			StartupPackageToWarnState.Add( InPackage, true );
-		}
-		else
-		{
-			StartupPackageToWarnState.Add( InPackage, false );
-			Error->Logf(ELogVerbosity::Warning, *FText::Format( NSLOCTEXT( "UnrealEd", "CannotSaveStartupPackage", "Cannot save asset '{0}' as user opted not to save this startup asset" ), FText::FromString( InFilename ) ).ToString() );
-			return false;
-		}
-	}
-
-	return true;
-}
-
-void UEditorEngine::OnSourceControlDialogClosed(bool bEnabled)
-{
-	if(ISourceControlModule::Get().IsEnabled())
-	{
-		ISourceControlProvider& SourceControlProvider = ISourceControlModule::Get().GetProvider();
-		if(SourceControlProvider.IsAvailable())
-		{
-			if(DeferredFilesToAddToSourceControl.Num() > 0)
-			{
-				SourceControlProvider.Execute(ISourceControlOperation::Create<FMarkForAdd>(), SourceControlHelpers::PackageFilenames(DeferredFilesToAddToSourceControl));
-			}
-	
-			DeferredFilesToAddToSourceControl.Empty();
-		}
-	}
-	else
-	{
-		// the user decided to disable source control, so clear the deferred list so we dont try to add them again at a later time
-		DeferredFilesToAddToSourceControl.Empty();
-	}
-}
-
-bool UEditorEngine::SavePackage( UPackage* InOuter, UObject* InBase, EObjectFlags TopLevelFlags, const TCHAR* Filename, 
-				 FOutputDevice* Error, FLinkerLoad* Conform, bool bForceByteSwapping, bool bWarnOfLongFilename, 
-				 uint32 SaveFlags, const class ITargetPlatform* TargetPlatform, const FDateTime& FinalTimeStamp, bool bSlowTask )
-{
-	FScopedSlowTask SlowTask(100, FText(), bSlowTask);
-
-	UObject* Base = InBase;
-	if ( !Base && InOuter && InOuter->PackageFlags & PKG_ContainsMap )
-	{
-		Base = UWorld::FindWorldInPackage(InOuter);
-	}
-
-	// Record the package flags before OnPreSaveWorld. They will be used in OnPostSaveWorld.
-	uint32 OriginalPackageFlags = 0;
-	if ( InOuter )
-	{
-		OriginalPackageFlags = InOuter->PackageFlags;
-	}
-
-	SlowTask.EnterProgressFrame(10);
-
-	UWorld* World = Cast<UWorld>(Base);
-	if ( World )
-	{
-		OnPreSaveWorld(SaveFlags, World);
-	}
-
-	// See if the package is a valid candidate for being auto-added to the default changelist.
-	// Only allows the addition of newly created packages while in the editor and then only if the user has the option enabled.
-	bool bAutoAddPkgToSCC = false;
-	if( !TargetPlatform )
-	{
-		bAutoAddPkgToSCC = IsPackageValidForAutoAdding( InOuter, Filename );
-	}
-
-	SlowTask.EnterProgressFrame(70);
-
-	bool bSuccess = UPackage::SavePackage(InOuter, Base, TopLevelFlags, Filename, Error, Conform, bForceByteSwapping, bWarnOfLongFilename, SaveFlags, TargetPlatform, FinalTimeStamp, bSlowTask);
-
-	SlowTask.EnterProgressFrame(10);
-
-	// If the package is a valid candidate for being automatically-added to source control, go ahead and add it
-	// to the default changelist
-	if( bSuccess && bAutoAddPkgToSCC )
-	{
-		// IsPackageValidForAutoAdding should not return true if SCC is disabled
-		check(ISourceControlModule::Get().IsEnabled());
-
-		if(!ISourceControlModule::Get().GetProvider().IsAvailable())
-		{
-			// Show the login window here & store the file we are trying to add.
-			// We defer the add operation until we have a valid source control connection.
-			ISourceControlModule::Get().ShowLoginDialog(FSourceControlLoginClosed::CreateUObject(this, &UEditorEngine::OnSourceControlDialogClosed), ELoginWindowMode::Modeless);
-			DeferredFilesToAddToSourceControl.Add( Filename );
-		}
-		else
-		{
-			ISourceControlModule::Get().GetProvider().Execute(ISourceControlOperation::Create<FMarkForAdd>(), SourceControlHelpers::PackageFilename(Filename));
-		}
-	}
-
-	SlowTask.EnterProgressFrame(10);
-
-	if ( World )
-	{
-		OnPostSaveWorld(SaveFlags, World, OriginalPackageFlags, bSuccess);
-	}
-
-	return bSuccess;
-}
-
-void UEditorEngine::OnPreSaveWorld(uint32 SaveFlags, UWorld* World)
-{
-	if ( !ensure(World) )
-	{
-		return;
-	}
-
-	check(World->PersistentLevel);
-
-	// Pre save world event
-	FEditorDelegates::PreSaveWorld.Broadcast(SaveFlags, World);
-
-	// Update cull distance volumes (and associated primitives).
-	World->UpdateCullDistanceVolumes();
-
-	if ( !IsRunningCommandlet() )
-	{
-		const bool bAutosaveOrPIE = (SaveFlags & SAVE_FromAutosave) != 0;
-		if ( bAutosaveOrPIE )
-		{
-			// Temporarily flag packages saved under a PIE filename as PKG_PlayInEditor for serialization so loading
-			// them will have the flag set. We need to undo this as the object flagged isn't actually the PIE package, 
-			// but rather only the loaded one will be.
-			// PIE prefix detected, mark package.
-			if( World->GetName().StartsWith( PLAYWORLD_PACKAGE_PREFIX ) )
-			{
-				World->GetOutermost()->PackageFlags |= PKG_PlayInEditor;
-			}
-		}
-		else
-		{
-			// Normal non-pie and non-autosave codepath
-			FWorldContext &EditorContext = GEditor->GetEditorWorldContext();
-
-			// Check that this world is GWorld to avoid stomping on the saved views of sub-levels.
-			if ( World == EditorContext.World() )
-			{
-				if( FModuleManager::Get().IsModuleLoaded("LevelEditor") )
-				{
-					FLevelEditorModule& LevelEditor = FModuleManager::GetModuleChecked<FLevelEditorModule>("LevelEditor");
-
-					// Notify slate level editors of the map change
-					LevelEditor.BroadcastMapChanged( World, EMapChangeType::SaveMap );
-				}
-			}
-
-			// Shrink model and clean up deleted actors.
-			// Don't do this when autosaving or PIE saving so that actor adds can still undo.
-			World->ShrinkLevel();
-
-			{
-				FScopedSlowTask SlowTask(0, FText::Format(NSLOCTEXT("UnrealEd", "SavingMapStatus_CollectingGarbage", "Saving map: {0}... (Collecting garbage)"), FText::FromString(World->GetName())));
-				// NULL empty or "invalid" entries (e.g. IsPendingKill()) in actors array.
-				CollectGarbage( GARBAGE_COLLECTION_KEEPFLAGS );
-			}
-			
-			// Compact and sort actors array to remove empty entries.
-			// Don't do this when autosaving or PIE saving so that actor adds can still undo.
-			World->PersistentLevel->SortActorList();
-		}
-	}
-
-	// Move level position closer to world origin
-	UWorld* OwningWorld = World->PersistentLevel->OwningWorld;
-	if (OwningWorld->WorldComposition)
-	{
-		OwningWorld->WorldComposition->OnLevelPreSave(World->PersistentLevel);
-	}
-
-	// If we can get the streaming level, we should remove the editor transform before saving
-	ULevelStreaming* LevelStreaming = FLevelUtils::FindStreamingLevel( World->PersistentLevel );
-	if ( LevelStreaming )
-	{
-		FLevelUtils::RemoveEditorTransform(LevelStreaming);
-	}
-
-	// Make sure the public and standalone flags are set on this world to allow it to work properly with the editor
-	World->SetFlags(RF_Public | RF_Standalone);
-}
-
-void UEditorEngine::OnPostSaveWorld(uint32 SaveFlags, UWorld* World, uint32 OriginalPackageFlags, bool bSuccess)
-{
-	if ( !ensure(World) )
-	{
-		return;
-	}
-
-	if ( !IsRunningCommandlet() )
-	{
-		UPackage* WorldPackage = World->GetOutermost();
-		const bool bAutosaveOrPIE = (SaveFlags & SAVE_FromAutosave) != 0;
-		if ( bAutosaveOrPIE )
-		{
-			// Restore original value of PKG_PlayInEditor if we changed it during PIE saving
-			const bool bOriginallyPIE = (OriginalPackageFlags & PKG_PlayInEditor) != 0;
-			const bool bCurrentlyPIE = (WorldPackage->PackageFlags & PKG_PlayInEditor) != 0;
-			if ( !bOriginallyPIE && bCurrentlyPIE )
-			{
-				WorldPackage->PackageFlags &= ~PKG_PlayInEditor;
-			}
-		}
-		else
-		{
-			// Normal non-pie and non-autosave codepath
-			FWorldContext &EditorContext = GEditor->GetEditorWorldContext();
-
-			const bool bIsPersistentLevel = (World == EditorContext.World());
-			if ( bSuccess )
-			{
-				// Put the map into the MRU and mark it as not dirty.
-
-				if ( bIsPersistentLevel )
-				{
-					// Set the map filename.
-					const FString Filename = FPackageName::LongPackageNameToFilename(WorldPackage->GetName(), FPackageName::GetMapPackageExtension());
-					FEditorFileUtils::RegisterLevelFilename( World, Filename );
-
-					WorldPackage->SetDirtyFlag(false);
-
-					// Update the editor's MRU level list if we were asked to do that for this level
-					IMainFrameModule& MainFrameModule = FModuleManager::LoadModuleChecked<IMainFrameModule>( "MainFrame" );
-
-					if ( MainFrameModule.GetMRUFavoritesList() )
-					{
-						MainFrameModule.GetMRUFavoritesList()->AddMRUItem( Filename );
-					}
-
-					FEditorDirectories::Get().SetLastDirectory(ELastDirectory::UNR, FPaths::GetPath(Filename)); // Save path as default for next time.
-				}
-
-				// We saved the map, so unless there are any other dirty levels, go ahead and reset the autosave timer
-				if( GUnrealEd && !GUnrealEd->AnyWorldsAreDirty( World ) )
-				{
-					GUnrealEd->GetPackageAutoSaver().ResetAutoSaveTimer();
-				}
-			}
-
-			if ( bIsPersistentLevel )
-			{
-				if ( GUnrealEd )
-				{
-					GUnrealEd->ResetTransaction( NSLOCTEXT("UnrealEd", "MapSaved", "Map Saved") );
-				}
-
-				FPlatformProcess::SetCurrentWorkingDirectoryToBaseDir();
-			}
-		}
-	}
-
-	// Restore level original position
-	UWorld* OwningWorld = World->PersistentLevel->OwningWorld;
-	if (OwningWorld->WorldComposition)
-	{
-		OwningWorld->WorldComposition->OnLevelPostSave(World->PersistentLevel);
-	}
-
-	// If got the streaming level, we should re-apply the editor transform after saving
-	ULevelStreaming* LevelStreaming = FLevelUtils::FindStreamingLevel( World->PersistentLevel );
-	if ( LevelStreaming )
-	{
-		FLevelUtils::ApplyEditorTransform(LevelStreaming);
-	}
-
-	// Post save world event
-	FEditorDelegates::PostSaveWorld.Broadcast(SaveFlags, World, bSuccess);
-}
-
-APlayerStart* UEditorEngine::CheckForPlayerStart()
-{
-	UWorld* IteratorWorld = GWorld;
-	for( TActorIterator<APlayerStart> It(IteratorWorld); It; ++It )
-	{
-		// Return the found start.
-		return *It;
-	}
-
-	// No player start was found, return NULL.
-	return NULL;
-}
-
-void UEditorEngine::CloseEntryPopupWindow()
-{
-	if (PopupWindow.IsValid())
-	{
-		PopupWindow.Pin()->RequestDestroyWindow();
-	}
-}
-
-EAppReturnType::Type UEditorEngine::OnModalMessageDialog(EAppMsgType::Type InMessage, const FText& InText, const FText& InTitle)
-{
-	if( FSlateApplication::IsInitialized() && FSlateApplication::Get().CanAddModalWindow() )
-	{
-		return OpenMsgDlgInt(InMessage, InText, InTitle);
-	}
-	else
-	{
-		return FPlatformMisc::MessageBoxExt(InMessage, *InText.ToString(), *InTitle.ToString());
-	}
-}
-
-bool UEditorEngine::OnShouldLoadOnTop( const FString& Filename )
-{
-	 return GEditor && (FPaths::GetBaseFilename(Filename) == FPaths::GetBaseFilename(GEditor->UserOpenedFile));
-}
-
-TSharedPtr<SViewport> UEditorEngine::GetGameViewportWidget() const
-{
-	for (auto It = SlatePlayInEditorMap.CreateConstIterator(); It; ++It)
-	{
-		if (It.Value().SlatePlayInEditorWindowViewport.IsValid())
-		{
-			return It.Value().SlatePlayInEditorWindowViewport->GetViewportWidget().Pin();
-		}
-
-		TSharedPtr<ILevelViewport> DestinationLevelViewport = It.Value().DestinationSlateViewport.Pin();
-		if (DestinationLevelViewport.IsValid())
-		{
-			return DestinationLevelViewport->GetViewportWidget().Pin();
-		}
-	}
-
-	/*
-	if(SlatePlayInEditorSession.SlatePlayInEditorWindowViewport.IsValid())
-	{
-		return SlatePlayInEditorSession.SlatePlayInEditorWindowViewport->GetViewportWidget().Pin();
-	}
-	*/
-
-	return NULL;
-}
-
-bool UEditorEngine::SplitActorLabel( FString& InOutLabel, int32& OutIdx ) const
-{
-	// Look at the label and see if it ends in a number and separate them
-	const TArray<TCHAR>& LabelCharArray = InOutLabel.GetCharArray();
-	for(int32 CharIdx = LabelCharArray.Num()-1; CharIdx>=0; CharIdx--)
-	{
-		if ( CharIdx == 0 || !FChar::IsDigit( LabelCharArray[CharIdx-1] ) )
-		{
-			FString Idx = InOutLabel.RightChop(CharIdx);
-			if ( Idx.Len() > 0 )
-			{
-				InOutLabel = InOutLabel.Left(CharIdx);
-				OutIdx = FCString::Atoi(*Idx);
-				return true;
-			}
-			break;
-		}			
-	}
-	return false;
-}
-
-void UEditorEngine::SetActorLabelUnique(AActor* Actor, const FString& NewActorLabel, const FCachedActorLabels* InExistingActorLabels) const
-{
-	check( Actor );
-
-	FString Prefix             = NewActorLabel;
-	FString ModifiedActorLabel = NewActorLabel;
-	int32   LabelIdx           = 0;
-
-	FCachedActorLabels ActorLabels;
-	if (!InExistingActorLabels)
-	{
-		InExistingActorLabels = &ActorLabels;
-
-		TSet<AActor*> IgnoreActors;
-		IgnoreActors.Add(Actor);
-		ActorLabels.Populate(Actor->GetWorld(), IgnoreActors);
-	}
-
-
-	if (InExistingActorLabels->Contains(ModifiedActorLabel))
-	{
-		// See if the current label ends in a number, and try to create a new label based on that
-		if (!SplitActorLabel( Prefix, LabelIdx ))
-		{
-			// If there wasn't a number on there, append a number, starting from 2 (1 before incrementing below)
-			LabelIdx = 1;
-		}
-
-		// Update the actor label until we find one that doesn't already exist
-		while (InExistingActorLabels->Contains(ModifiedActorLabel))
-		{
-			++LabelIdx;
-			ModifiedActorLabel = FString::Printf(TEXT("%s%d"), *Prefix, LabelIdx);
-		}
-	}
-
-	Actor->SetActorLabel( ModifiedActorLabel );
-}
-
-
-
-FString UEditorEngine::GetFriendlyName( const UProperty* Property, UStruct* OwnerClass/* = NULL*/ )
-{
-	// first, try to pull the friendly name from the loc file
-	check( Property );
-	UClass* RealOwnerClass = Property->GetOwnerClass();
-	if ( OwnerClass == NULL)
-	{
-		OwnerClass = RealOwnerClass;
-	}
-	checkSlow(OwnerClass);
-
-	FText FoundText;
-	bool DidFindText = false;
-	UStruct* CurrentClass = OwnerClass;
-	do 
-	{
-		FString PropertyPathName = Property->GetPathName(CurrentClass);
-
-		DidFindText = FText::FindText(*CurrentClass->GetName(), *(PropertyPathName + TEXT(".FriendlyName")), /*OUT*/FoundText );
-		CurrentClass = CurrentClass->GetSuperStruct();
-	} while( CurrentClass != NULL && CurrentClass->IsChildOf(RealOwnerClass) && !DidFindText );
-
-	if ( !DidFindText )
-	{
-		FString DefaultFriendlyName = Property->GetMetaData(TEXT("DisplayName"));
-		if ( DefaultFriendlyName.IsEmpty() )
-		{
-			const bool bIsBool = Cast<const UBoolProperty>(Property) != NULL;
-			return FName::NameToDisplayString( Property->GetName(), bIsBool );
-		}
-		return DefaultFriendlyName;
-	}
-
-	return FoundText.ToString();
-}
-
-AActor* UEditorEngine::UseActorFactoryOnCurrentSelection( UActorFactory* Factory, const FTransform* InActorTransform, EObjectFlags ObjectFlags )
-{
-	// ensure that all selected assets are loaded
-	FEditorDelegates::LoadSelectedAssetsIfNeeded.Broadcast();
-	return UseActorFactory(Factory, FAssetData( GetSelectedObjects()->GetTop<UObject>() ), InActorTransform, ObjectFlags );
-}
-
-AActor* UEditorEngine::UseActorFactory( UActorFactory* Factory, const FAssetData& AssetData, const FTransform* InActorTransform, EObjectFlags ObjectFlags )
-{
-	check( Factory );
-
-	bool bIsAllowedToCreateActor = true;
-
-	FText ActorErrorMsg;
-	if( !Factory->CanCreateActorFrom( AssetData, ActorErrorMsg ) )
-	{
-		bIsAllowedToCreateActor = false;
-		FMessageLog EditorErrors("EditorErrors");
-		EditorErrors.Warning(ActorErrorMsg);
-		EditorErrors.Notify();
-	}
-
-	//Load Asset
-	UObject* Asset = AssetData.GetAsset();
-
-	AActor* Actor = NULL;
-	if ( bIsAllowedToCreateActor )
-	{
-		AActor* NewActorTemplate = Factory->GetDefaultActor( AssetData );
-
-		if ( !NewActorTemplate )
-		{
-			return NULL;
-		}
-
-		const FTransform ActorTransform = InActorTransform ? *InActorTransform : FActorPositioning::GetCurrentViewportPlacementTransform(*NewActorTemplate);
-
-		ULevel* DesiredLevel = GWorld->GetCurrentLevel();
-
-		// Don't spawn the actor if the current level is locked.
-		if( !FLevelUtils::IsLevelLocked( DesiredLevel ) )
-		{
-			// Check to see if the level it's being added to is hidden and ask the user if they want to proceed
-			const bool bLevelVisible = FLevelUtils::IsLevelVisible( DesiredLevel );
-			if ( bLevelVisible || EAppReturnType::Ok == FMessageDialog::Open( EAppMsgType::OkCancel, FText::Format( LOCTEXT("CurrentLevelHiddenActorWillAlsoBeHidden", "Current level [{0}] is hidden, actor will also be hidden until level is visible"), FText::FromString( DesiredLevel->GetOutermost()->GetName() ) ) ) )
-			{
-				const FScopedTransaction Transaction( NSLOCTEXT("UnrealEd", "CreateActor", "Create Actor") );
-
-				// Create the actor.
-				Actor = Factory->CreateActor( Asset, DesiredLevel, ActorTransform, ObjectFlags );
-				if(Actor != NULL)
-				{
-					SelectNone( false, true );
-					SelectActor( Actor, true, true );
-					Actor->InvalidateLightingCache();
-					Actor->PostEditMove( true );
-
-					// Make sure the actors visibility reflects that of the level it's in
-					if ( !bLevelVisible )
-					{
-						Actor->bHiddenEdLevel = true;
-						// We update components, so things like draw scale take effect.
-						Actor->ReregisterAllComponents(); // @todo UE4 insist on a property update callback
-					}
-				}
-
-				RedrawLevelEditingViewports();
-
-
-				if ( Actor )
-				{
-					Actor->MarkPackageDirty();
-					ULevel::LevelDirtiedEvent.Broadcast();
-				}
-			}
-			else
-			{
-				return NULL;
-			}
-		}
-		else
-		{
-			FNotificationInfo Info( NSLOCTEXT("UnrealEd", "Error_OperationDisallowedOnLockedLevel", "The requested operation could not be completed because the level is locked.") );
-			Info.ExpireDuration = 3.0f;
-			FSlateNotificationManager::Get().AddNotification(Info);
-			return NULL;
-		}
-	}
-
-	return Actor;
-}
-
-namespace ReattachActorsHelper
-{
-	/** Holds the actor and socket name for attaching. */
-	struct FActorAttachmentInfo
-	{
-		AActor* Actor;
-
-		FName SocketName;
-	};
-
-	/** Used to cache the attachment info for an actor. */
-	struct FActorAttachmentCache
-	{
-	public:
-		/** The post-conversion actor. */
-		AActor* NewActor;
-
-		/** The parent actor and socket. */
-		FActorAttachmentInfo ParentActor;
-
-		/** Children actors and the sockets they were attached to. */
-		TArray<FActorAttachmentInfo> AttachedActors;
-	};
-
-	/** 
-	 * Caches the attachment info for the actors being converted.
-	 *
-	 * @param InActorsToReattach			List of actors to reattach.
-	 * @param InOutAttachmentInfo			List of attachment info for the list of actors.
-	 */
-	void CacheAttachments(const TArray<AActor*>& InActorsToReattach, TArray<FActorAttachmentCache>& InOutAttachmentInfo)
-	{
-		for( int32 ActorIdx = 0; ActorIdx < InActorsToReattach.Num(); ++ActorIdx )
-		{
-			AActor* ActorToReattach = InActorsToReattach[ ActorIdx ];
-
-			InOutAttachmentInfo.AddZeroed();
-
-			FActorAttachmentCache& CurrentAttachmentInfo = InOutAttachmentInfo[ActorIdx];
-
-			// Retrieve the list of attached actors.
-			TArray<AActor*> AttachedActors;
-			ActorToReattach->GetAttachedActors(AttachedActors);
-
-			// Cache the parent actor and socket name.
-			CurrentAttachmentInfo.ParentActor.Actor = ActorToReattach->GetAttachParentActor();
-			CurrentAttachmentInfo.ParentActor.SocketName = ActorToReattach->GetAttachParentSocketName();
-
-			// Required to restore attachments properly.
-			for( int32 AttachedActorIdx = 0; AttachedActorIdx < AttachedActors.Num(); ++AttachedActorIdx )
-			{
-				// Store the attached actor and socket name in the cache.
-				CurrentAttachmentInfo.AttachedActors.AddZeroed();
-				CurrentAttachmentInfo.AttachedActors[AttachedActorIdx].Actor = AttachedActors[AttachedActorIdx];
-				CurrentAttachmentInfo.AttachedActors[AttachedActorIdx].SocketName = AttachedActors[AttachedActorIdx]->GetAttachParentSocketName();
-
-				AActor* ChildActor = CurrentAttachmentInfo.AttachedActors[AttachedActorIdx].Actor;
-				ChildActor->Modify();
-				ChildActor->DetachRootComponentFromParent(true);
-			}
-
-			// Modify the actor so undo will reattach it.
-			ActorToReattach->Modify();
-			ActorToReattach->DetachRootComponentFromParent(true);
-		}
-	}
-
-	/** 
-	 * Caches the actor old/new information, mapping the old actor to the new version for easy look-up and matching.
-	 *
-	 * @param InOldActor			The old version of the actor.
-	 * @param InNewActor			The new version of the actor.
-	 * @param InOutReattachmentMap	Map object for placing these in.
-	 * @param InOutAttachmentInfo	Update the required attachment info to hold the Converted Actor.
-	 */
-	void CacheActorConvert(AActor* InOldActor, AActor* InNewActor, TMap<AActor*, AActor*>& InOutReattachmentMap, FActorAttachmentCache& InOutAttachmentInfo)
-	{
-		// Add mapping data for the old actor to the new actor.
-		InOutReattachmentMap.Add(InOldActor, InNewActor);
-
-		// Set the converted actor so re-attachment can occur.
-		InOutAttachmentInfo.NewActor = InNewActor;
-	}
-
-	/** 
-	 * Checks if two actors can be attached, creates Message Log messages if there are issues.
-	 *
-	 * @param InParentActor			The parent actor.
-	 * @param InChildActor			The child actor.
-	 * @param InOutErrorMessages	Errors with attaching the two actors are stored in this array.
-	 *
-	 * @return Returns true if the actors can be attached, false if they cannot.
-	 */
-	bool CanParentActors(AActor* InParentActor, AActor* InChildActor)
-	{
-		FText ReasonText;
-		if (GEditor->CanParentActors(InParentActor, InChildActor, &ReasonText))
-		{
-			return true;
-		}
-		else
-		{
-			FMessageLog("EditorErrors").Error(ReasonText);
-			return false;
-		}
-	}
-
-	/** 
-	 * Reattaches actors to maintain the hierarchy they had previously using a conversion map and an array of attachment info. All errors displayed in Message Log along with notifications.
-	 *
-	 * @param InReattachmentMap			Used to find the corresponding new versions of actors using an old actor pointer.
-	 * @param InAttachmentInfo			Holds parent and child attachment data.
-	 */
-	void ReattachActors(TMap<AActor*, AActor*>& InReattachmentMap, TArray<FActorAttachmentCache>& InAttachmentInfo)
-	{
-		// Holds the errors for the message log.
-		FMessageLog EditorErrors("EditorErrors");
-		EditorErrors.NewPage(LOCTEXT("AttachmentLogPage", "Actor Reattachment"));
-
-		for( int32 ActorIdx = 0; ActorIdx < InAttachmentInfo.Num(); ++ActorIdx )
-		{
-			FActorAttachmentCache& CurrentAttachment = InAttachmentInfo[ActorIdx];
-
-			// Need to reattach all of the actors that were previously attached.
-			for( int32 AttachedIdx = 0; AttachedIdx < CurrentAttachment.AttachedActors.Num(); ++AttachedIdx )
-			{
-				// Check if the attached actor was converted. If it was it will be in the TMap.
-				AActor** CheckIfConverted = InReattachmentMap.Find(CurrentAttachment.AttachedActors[AttachedIdx].Actor);
-				if(CheckIfConverted)
-				{
-					// This should always be valid.
-					if(*CheckIfConverted)
-					{
-						AActor* ParentActor = CurrentAttachment.NewActor;
-						AActor* ChildActor = *CheckIfConverted;
-
-						if (CanParentActors(ParentActor, ChildActor))
-						{
-							// Attach the previously attached and newly converted actor to the current converted actor.
-							ChildActor->AttachRootComponentToActor(ParentActor, CurrentAttachment.AttachedActors[AttachedIdx].SocketName, EAttachLocation::KeepWorldPosition);
-						}
-					}
-
-				}
-				else
-				{
-					AActor* ParentActor = CurrentAttachment.NewActor;
-					AActor* ChildActor = CurrentAttachment.AttachedActors[AttachedIdx].Actor;
-
-					if (CanParentActors(ParentActor, ChildActor))
-					{
-						// Since the actor was not converted, reattach the unconverted actor.
-						ChildActor->AttachRootComponentToActor(ParentActor, CurrentAttachment.AttachedActors[AttachedIdx].SocketName, EAttachLocation::KeepWorldPosition);
-					}
-				}
-
-			}
-
-			// Check if the parent was converted.
-			AActor** CheckIfNewActor = InReattachmentMap.Find(CurrentAttachment.ParentActor.Actor);
-			if(CheckIfNewActor)
-			{
-				// Since the actor was converted, attach the current actor to it.
-				if(*CheckIfNewActor)
-				{
-					AActor* ParentActor = *CheckIfNewActor;
-					AActor* ChildActor = CurrentAttachment.NewActor;
-
-					if (CanParentActors(ParentActor, ChildActor))
-					{
-						ChildActor->AttachRootComponentToActor(ParentActor, CurrentAttachment.ParentActor.SocketName, EAttachLocation::KeepWorldPosition);
-					}
-				}
-
-			}
-			else
-			{
-				AActor* ParentActor = CurrentAttachment.ParentActor.Actor;
-				AActor* ChildActor = CurrentAttachment.NewActor;
-
-				// Verify the parent is valid, the actor may not have actually been attached before.
-				if (ParentActor && CanParentActors(ParentActor, ChildActor))
-				{
-					// The parent was not converted, attach to the unconverted parent.
-					ChildActor->AttachRootComponentToActor(ParentActor, CurrentAttachment.ParentActor.SocketName, EAttachLocation::KeepWorldPosition);
-				}
-			}
-		}
-
-		// Add the errors to the message log, notifications will also be displayed as needed.
-		EditorErrors.Notify(NSLOCTEXT("ActorAttachmentError", "AttachmentsFailed", "Attachments Failed!"));
-	}
-}
-
-void UEditorEngine::ReplaceSelectedActors(UActorFactory* Factory, const FAssetData& AssetData)
-{
-	UObject* ObjectForFactory = NULL;
-
-	// Provide the option to abort the delete
-	if (ShouldAbortActorDeletion())
-	{
-		return;
-	}
-	else if (Factory != nullptr)
-	{
-		FText ActorErrorMsg;
-		if (!Factory->CanCreateActorFrom( AssetData, ActorErrorMsg))
-		{
-			FMessageDialog::Open( EAppMsgType::Ok, ActorErrorMsg );
-			return;
-		}
-	}
-	else
-	{
-		UE_LOG(LogEditor, Error, TEXT("UEditorEngine::ReplaceSelectedActors() called with NULL parameters!"));
-		return;
-	}
-
-	const FScopedTransaction Transaction( NSLOCTEXT("UnrealEd", "Replace Actors", "Replace Actor(s)") );
-
-	// construct a list of Actors to replace in a separate pass so we can modify the selection set as we perform the replacement
-	TArray<AActor*> ActorsToReplace;
-	for (FSelectionIterator It = GetSelectedActorIterator(); It; ++It)
-	{
-		AActor* Actor = Cast<AActor>(*It);
-		if ( Actor && !FActorEditorUtils::IsABuilderBrush(Actor) )
-		{
-			ActorsToReplace.Add(Actor);
-		}
-	}
-
-	ReplaceActors(Factory, AssetData, ActorsToReplace);
-}
-
-void UEditorEngine::ReplaceActors(UActorFactory* Factory, const FAssetData& AssetData, const TArray<AActor*> ActorsToReplace)
-{
-	// Cache for attachment info of all actors being converted.
-	TArray<ReattachActorsHelper::FActorAttachmentCache> AttachmentInfo;
-
-	// Maps actors from old to new for quick look-up.
-	TMap<AActor*, AActor*> ConvertedMap;
-
-	// Cache the current attachment states.
-	ReattachActorsHelper::CacheAttachments(ActorsToReplace, AttachmentInfo);
-
-	USelection* SelectedActors = GetSelectedActors();
-	SelectedActors->BeginBatchSelectOperation();
-	SelectedActors->Modify();
-
-	UObject* Asset = AssetData.GetAsset();
-	for(int32 ActorIdx = 0; ActorIdx < ActorsToReplace.Num(); ++ActorIdx)
-	{
-		AActor* OldActor = ActorsToReplace[ActorIdx];//.Pop();
-		check(OldActor);
-		UWorld* World = OldActor->GetWorld();
-		ULevel* Level = OldActor->GetLevel();
-		AActor* NewActor = NULL;
-
-		const FName OldActorName = OldActor->GetFName();
-		const FName OldActorReplacedNamed = MakeUniqueObjectName(OldActor->GetOuter(), OldActor->GetClass(), *FString::Printf(TEXT("%s_REPLACED"), *OldActorName.ToString()));
-		OldActor->Rename(*OldActorReplacedNamed.ToString());
-
-		const FTransform OldTransform = OldActor->ActorToWorld();
-
-		// create the actor
-		NewActor = Factory->CreateActor( Asset, Level, OldTransform, RF_Transactional, OldActorName );
-		// For blueprints, try to copy over properties
-		if (Factory->IsA(UActorFactoryBlueprint::StaticClass()))
-		{
-			UBlueprint* Blueprint = CastChecked<UBlueprint>(Asset);
-			// Only try to copy properties if this blueprint is based on the actor
-			UClass* OldActorClass = OldActor->GetClass();
-			if (Blueprint->GeneratedClass->IsChildOf(OldActorClass) && NewActor != NULL)
-			{
-				NewActor->UnregisterAllComponents();
-				UEditorEngine::CopyPropertiesForUnrelatedObjects(OldActor, NewActor);
-				NewActor->RegisterAllComponents();
-			}
-		}
-
-		if ( NewActor != NULL )
-		{
-			// The new actor might not have a root component
-			USceneComponent* const NewActorRootComponent = NewActor->GetRootComponent();
-			if(NewActorRootComponent)
-			{
-				if(!GetDefault<ULevelEditorMiscSettings>()->bReplaceRespectsScale || OldActor->GetRootComponent() == NULL )
-				{
-					NewActorRootComponent->SetRelativeScale3D(FVector(1.0f, 1.0f, 1.0f));
-				}
-				else
-				{
-					NewActorRootComponent->SetRelativeScale3D( OldActor->GetRootComponent()->RelativeScale3D );
-				}
-			}
-
-			NewActor->Layers.Empty();
-			GEditor->Layers->AddActorToLayers( NewActor, OldActor->Layers );
-
-			// Preserve the label and tags from the old actor
-			NewActor->SetActorLabel( OldActor->GetActorLabel() );
-			NewActor->Tags = OldActor->Tags;
-
-			// Allow actor derived classes a chance to replace properties.
-			NewActor->EditorReplacedActor(OldActor);
-
-			// Caches information for finding the new actor using the pre-converted actor.
-			ReattachActorsHelper::CacheActorConvert(OldActor, NewActor, ConvertedMap, AttachmentInfo[ActorIdx]);
-
-			if (SelectedActors->IsSelected(OldActor))
-			{
-				SelectActor(OldActor, false, true);
-				SelectActor(NewActor, true, true);
-			}
-
-			// Find compatible static mesh components and copy instance colors between them.
-			UStaticMeshComponent* NewActorStaticMeshComponent = NewActor->FindComponentByClass<UStaticMeshComponent>();
-			UStaticMeshComponent* OldActorStaticMeshComponent = OldActor->FindComponentByClass<UStaticMeshComponent>();
-			if ( NewActorStaticMeshComponent != NULL && OldActorStaticMeshComponent != NULL )
-			{
-				NewActorStaticMeshComponent->CopyInstanceVertexColorsIfCompatible( OldActorStaticMeshComponent );
-			}
-
-			NewActor->InvalidateLightingCache();
-			NewActor->PostEditMove(true);
-			NewActor->MarkPackageDirty();
-
-			// Replace references in the level script Blueprint with the new Actor
-			const bool bDontCreate = true;
-			ULevelScriptBlueprint* LSB = NewActor->GetLevel()->GetLevelScriptBlueprint(bDontCreate);
-			if( LSB )
-			{
-				// Only if the level script blueprint exists would there be references.  
-				FBlueprintEditorUtils::ReplaceAllActorRefrences(LSB, OldActor, NewActor);
-			}
-
-			GEditor->Layers->DisassociateActorFromLayers( OldActor );
-			World->EditorDestroyActor(OldActor, true);
-		}
-		else
-		{
-			// If creating the new Actor failed, put the old Actor's name back
-			OldActor->Rename(*OldActorName.ToString());
-		}
-	}
-
-	SelectedActors->EndBatchSelectOperation();
-
-	// Reattaches actors based on their previous parent child relationship.
-	ReattachActorsHelper::ReattachActors(ConvertedMap, AttachmentInfo);
-
-	RedrawLevelEditingViewports();
-
-	ULevel::LevelDirtiedEvent.Broadcast();
-}
-
-
-/* Gets the common components of a specific type between two actors so that they may be copied.
- * 
- * @param InOldActor		The actor to copy component properties from
- * @param InNewActor		The actor to copy to
- */
-static void CopyLightComponentProperties( const AActor& InOldActor, AActor& InNewActor )
-{
-	// Since this is only being used for lights, make sure only the light component can be copied.
-	const UClass* CopyableComponentClass =  ULightComponent::StaticClass();
-
-	// Get the light component from the default actor of source actors class.
-	// This is so we can avoid copying properties that have not changed. 
-	// using ULightComponent::StaticClass()->GetDefaultObject() will not work since each light actor sets default component properties differently.
-	ALight* OldActorDefaultObject = InOldActor.GetClass()->GetDefaultObject<ALight>();
-	check(OldActorDefaultObject);
-	UActorComponent* DefaultLightComponent = OldActorDefaultObject->GetLightComponent();
-	check(DefaultLightComponent);
-
-	// The component we are copying from class
-	UClass* CompToCopyClass = NULL;
-	UActorComponent* LightComponentToCopy = NULL;
-
-	// Go through the old actor's components and look for a light component to copy.
-	TInlineComponentArray<UActorComponent*> OldActorComponents;
-	InOldActor.GetComponents(OldActorComponents);
-
-	for( int32 CompToCopyIdx = 0; CompToCopyIdx < OldActorComponents.Num(); ++CompToCopyIdx )
-	{
-		UActorComponent* Component = OldActorComponents[CompToCopyIdx];
-
-		if( Component->IsRegistered() && Component->IsA( CopyableComponentClass ) ) 
-		{
-			// A light component has been found. 
-			CompToCopyClass = Component->GetClass();
-			LightComponentToCopy = Component;
-			break;
-		}
-	}
-
-	// The light component from the new actor
-	UActorComponent* NewActorLightComponent = NULL;
-	// The class of the new actors light component
-	const UClass* CommonLightComponentClass = NULL;
-
-	// Dont do anything if there is no valid light component to copy from
-	if( LightComponentToCopy )
-	{
-		TInlineComponentArray<UActorComponent*> NewActorComponents;
-		InNewActor.GetComponents(NewActorComponents);
-
-		// Find a light component to overwrite in the new actor
-		for( int32 NewCompIdx = 0; NewCompIdx < NewActorComponents.Num(); ++NewCompIdx )
-		{
-			UActorComponent* Component = NewActorComponents[ NewCompIdx ];
-			if(Component->IsRegistered())
-			{
-				// Find a common component class between the new and old actor.   
-				// This needs to be done so we can copy as many properties as possible. 
-				// For example: if we are converting from a point light to a spot light, the point light component will be the common superclass.
-				// That way we can copy properties like light radius, which would have been impossible if we just took the base LightComponent as the common class.
-				const UClass* CommonSuperclass = Component->FindNearestCommonBaseClass( CompToCopyClass );
-
-				if( CommonSuperclass->IsChildOf( CopyableComponentClass ) )
-				{
-					NewActorLightComponent = Component;
-					CommonLightComponentClass = CommonSuperclass;
-				}
-			}
-		}
-	}
-
-	// Don't do anything if there is no valid light component to copy to
-	if( NewActorLightComponent )
-	{
-		bool bCopiedAnyProperty = false;
-
-		// Find and copy the lightmass settings directly as they need to be examined and copied individually and not by the entire light mass settings struct
-		const FString LightmassPropertyName = TEXT("LightmassSettings");
-
-		UProperty* PropertyToCopy = NULL;
-		for( UProperty* Property = CompToCopyClass->PropertyLink; Property != NULL; Property = Property->PropertyLinkNext )
-		{
-			if( Property->GetName() == LightmassPropertyName )
-			{
-				// Get the offset in the old actor where lightmass properties are stored.
-				PropertyToCopy = Property;
-				break;
-			}
-		}
-
-		if( PropertyToCopy != NULL )
-		{
-			void* PropertyToCopyBaseLightComponentToCopy = PropertyToCopy->ContainerPtrToValuePtr<void>(LightComponentToCopy);
-			void* PropertyToCopyBaseDefaultLightComponent = PropertyToCopy->ContainerPtrToValuePtr<void>(DefaultLightComponent);
-			// Find the location of the lightmass settings in the new actor (if any)
-			for( UProperty* NewProperty = NewActorLightComponent->GetClass()->PropertyLink; NewProperty != NULL; NewProperty = NewProperty->PropertyLinkNext )
-			{
-				if( NewProperty->GetName() == LightmassPropertyName )
-				{
-					UStructProperty* OldLightmassProperty = Cast<UStructProperty>(PropertyToCopy);
-					UStructProperty* NewLightmassProperty = Cast<UStructProperty>(NewProperty);
-
-					void* NewPropertyBaseNewActorLightComponent = NewProperty->ContainerPtrToValuePtr<void>(NewActorLightComponent);
-					// The lightmass settings are a struct property so the cast should never fail.
-					check(OldLightmassProperty);
-					check(NewLightmassProperty);
-
-					// Iterate through each property field in the lightmass settings struct that we are copying from...
-					for( TFieldIterator<UProperty> OldIt(OldLightmassProperty->Struct); OldIt; ++OldIt)
-					{
-						UProperty* OldLightmassField = *OldIt;
-
-						// And search for the same field in the lightmass settings struct we are copying to.
-						// We should only copy to fields that exist in both structs.
-						// Even though their offsets match the structs may be different depending on what type of light we are converting to
-						bool bPropertyFieldFound = false;
-						for( TFieldIterator<UProperty> NewIt(NewLightmassProperty->Struct); NewIt; ++NewIt)
-						{
-							UProperty* NewLightmassField = *NewIt;
-							if( OldLightmassField->GetName() == NewLightmassField->GetName() )
-							{
-								// The field is in both structs.  Ok to copy
-								bool bIsIdentical = OldLightmassField->Identical_InContainer(PropertyToCopyBaseLightComponentToCopy, PropertyToCopyBaseDefaultLightComponent);
-								if( !bIsIdentical )
-								{
-									// Copy if the value has changed
-									OldLightmassField->CopySingleValue(NewLightmassField->ContainerPtrToValuePtr<void>(NewPropertyBaseNewActorLightComponent), OldLightmassField->ContainerPtrToValuePtr<void>(PropertyToCopyBaseLightComponentToCopy));
-									bCopiedAnyProperty = true;
-								}
-								break;
-							}
-						}
-					}
-					// No need to continue once we have found the lightmass settings
-					break;
-				}
-			}
-		}
-
-
-
-		// Now Copy the light component properties.
-		for( UProperty* Property = CommonLightComponentClass->PropertyLink; Property != NULL; Property = Property->PropertyLinkNext )
-		{
-			bool bIsTransient = !!(Property->PropertyFlags & (CPF_Transient | CPF_DuplicateTransient | CPF_NonPIEDuplicateTransient));
-			// Properties are identical if they have not changed from the light component on the default source actor
-			bool bIsIdentical = Property->Identical_InContainer(LightComponentToCopy, DefaultLightComponent);
-			bool bIsComponent = !!(Property->PropertyFlags & (CPF_InstancedReference | CPF_ContainsInstancedReference));
-
-			if ( !bIsTransient && !bIsIdentical && !bIsComponent && Property->GetName() != LightmassPropertyName )
-			{
-				bCopiedAnyProperty = true;
-				// Copy only if not native, not transient, not identical, not a component (at this time don't copy components within components)
-				// Also dont copy lightmass settings, those were examined and taken above
-				Property->CopyCompleteValue_InContainer(NewActorLightComponent, LightComponentToCopy);
-			}
-		}	
-
-		if (bCopiedAnyProperty)
-		{
-			NewActorLightComponent->PostEditChange();
-		}
-	}
-}
-
-
-void UEditorEngine::ConvertLightActors( UClass* ConvertToClass )
-{
-	// Provide the option to abort the conversion
-	if ( GEditor->ShouldAbortActorDeletion() )
-	{
-		return;
-	}
-
-	// List of actors to convert
-	TArray< AActor* > ActorsToConvert;
-
-	// Get a list of valid actors to convert.
-	for( FSelectionIterator It( GEditor->GetSelectedActorIterator() ) ; It ; ++It )
-	{
-		AActor* ActorToConvert = static_cast<AActor*>( *It );
-		// Prevent non light actors from being converted
-		// Also prevent light actors from being converted if they are the same time as the new class
-		if( ActorToConvert->IsA( ALight::StaticClass() ) && ActorToConvert->GetClass() != ConvertToClass )
-		{
-			ActorsToConvert.Add( ActorToConvert );
-		}
-	}
-
-	if (ActorsToConvert.Num())
-	{
-		GEditor->GetSelectedActors()->BeginBatchSelectOperation();
-
-		// Undo/Redo support
-		const FScopedTransaction Transaction( NSLOCTEXT("UnrealEd", "ConvertLights", "Convert Light") );
-
-		int32 NumLightsConverted = 0;
-		int32 NumLightsToConvert = ActorsToConvert.Num();
-
-		// Convert each light 
-		for( int32 ActorIdx = 0; ActorIdx < ActorsToConvert.Num(); ++ActorIdx )
-		{
-			AActor* ActorToConvert = ActorsToConvert[ ActorIdx ];
-
-			check( ActorToConvert );
-			// The class of the actor we are about to replace
-			UClass* ClassToReplace = ActorToConvert->GetClass();
-
-			// Set the current level to the level where the convertible actor resides
-			UWorld* World = ActorToConvert->GetWorld();
-			check( World );
-			ULevel* ActorLevel = ActorToConvert->GetLevel();
-			checkSlow( ActorLevel != NULL );
-
-			// Find a common superclass between the actors so we know what properties to copy
-			const UClass* CommonSuperclass = ActorToConvert->FindNearestCommonBaseClass( ConvertToClass );
-			check ( CommonSuperclass );
-
-			// spawn the new actor
-			AActor* NewActor = NULL;	
-
-			// Take the old actors location always, not rotation.  If rotation was changed on the source actor, it will be copied below.
-			FVector const SpawnLoc = ActorToConvert->GetActorLocation();
-			FActorSpawnParameters SpawnInfo;
-			SpawnInfo.OverrideLevel = ActorLevel;
-			NewActor = World->SpawnActor( ConvertToClass, &SpawnLoc, NULL, SpawnInfo );
-			// The new actor must exist
-			check(NewActor);
-
-			// Copy common light component properties
-			CopyLightComponentProperties( *ActorToConvert, *NewActor );
-
-			// Select the new actor
-			GEditor->SelectActor( ActorToConvert, false, true );
-			GEditor->SelectActor( NewActor, true, true );
-
-			NewActor->InvalidateLightingCache();
-			NewActor->PostEditChange();
-			NewActor->PostEditMove( true );
-			NewActor->Modify();
-			GEditor->Layers->InitializeNewActorLayers( NewActor );
-
-			// We have converted another light.
-			++NumLightsConverted;
-
-			UE_LOG(LogEditor, Log, TEXT("Converted: %s to %s"), *ActorToConvert->GetActorLabel(), *NewActor->GetActorLabel() );
-
-			// Destroy the old actor.
-			GEditor->Layers->DisassociateActorFromLayers( ActorToConvert );
-			World->EditorDestroyActor( ActorToConvert, true );
-		}
-
-		GEditor->GetSelectedActors()->EndBatchSelectOperation();
-		GEditor->RedrawLevelEditingViewports();
-
-		ULevel::LevelDirtiedEvent.Broadcast();
-
-		CollectGarbage( GARBAGE_COLLECTION_KEEPFLAGS );
-	}
-}
-
-/**
- * Internal helper function to copy component properties from one actor to another. Only copies properties
- * from components if the source actor, source actor class default object, and destination actor all contain
- * a component of the same name (specified by parameter) and all three of those components share a common base
- * class, at which point properties from the common base are copied. Component template names are used instead of
- * component classes because an actor could potentially have multiple components of the same class.
- *
- * @param	SourceActor		Actor to copy component properties from
- * @param	DestActor		Actor to copy component properties to
- * @param	ComponentNames	Set of component template names to attempt to copy
- */
-void CopyActorComponentProperties( const AActor* SourceActor, AActor* DestActor, const TSet<FString>& ComponentNames )
-{
-	// Don't attempt to copy anything if the user didn't specify component names to copy
-	if ( ComponentNames.Num() > 0 )
-	{
-		check( SourceActor && DestActor );
-		const AActor* SrcActorDefaultActor = SourceActor->GetClass()->GetDefaultObject<AActor>();
-		check( SrcActorDefaultActor );
-
-		// Construct a mapping from the default actor of its relevant component names to its actual components. Here relevant component
-		// names are those that match a name provided as a parameter.
-		TInlineComponentArray<UActorComponent*> CDOComponents;
-		SrcActorDefaultActor->GetComponents(CDOComponents);
-
-		TMap<FString, const UActorComponent*> NameToDefaultComponentMap; 
-		for ( TInlineComponentArray<UActorComponent*>::TConstIterator CompIter( CDOComponents ); CompIter; ++CompIter )
-		{
-			const UActorComponent* CurComp = *CompIter;
-			check( CurComp );
-
-			const FString CurCompName = CurComp->GetName();
-			if ( ComponentNames.Contains( CurCompName ) )
-			{
-				NameToDefaultComponentMap.Add( CurCompName, CurComp );
-			}
-		}
-
-		// Construct a mapping from the source actor of its relevant component names to its actual components. Here relevant component names
-		// are those that match a name provided as a parameter.
-		TInlineComponentArray<UActorComponent*> SourceComponents;
-		SourceActor->GetComponents(SourceComponents);
-
-		TMap<FString, const UActorComponent*> NameToSourceComponentMap;
-		for ( TInlineComponentArray<UActorComponent*>::TConstIterator CompIter( SourceComponents ); CompIter; ++CompIter )
-		{
-			const UActorComponent* CurComp = *CompIter;
-			check( CurComp );
-
-			const FString CurCompName = CurComp->GetName();
-			if ( ComponentNames.Contains( CurCompName ) )
-			{
-				NameToSourceComponentMap.Add( CurCompName, CurComp );
-			}
-		}
-
-		bool bCopiedAnyProperty = false;
-
-		TInlineComponentArray<UActorComponent*> DestComponents;
-		DestActor->GetComponents(DestComponents);
-
-		// Iterate through all of the destination actor's components to find the ones which should have properties copied into them.
-		for ( TInlineComponentArray<UActorComponent*>::TIterator DestCompIter( DestComponents ); DestCompIter; ++DestCompIter )
-		{
-			UActorComponent* CurComp = *DestCompIter;
-			check( CurComp );
-
-			const FString CurCompName = CurComp->GetName();
-
-			// Check if the component is one that the user wanted to copy properties into
-			if ( ComponentNames.Contains( CurCompName ) )
-			{
-				const UActorComponent** DefaultComponent = NameToDefaultComponentMap.Find( CurCompName );
-				const UActorComponent** SourceComponent = NameToSourceComponentMap.Find( CurCompName );
-
-				// Make sure that both the default actor and the source actor had a component of the same name
-				if ( DefaultComponent && SourceComponent )
-				{
-					const UClass* CommonBaseClass = NULL;
-					const UClass* DefaultCompClass = (*DefaultComponent)->GetClass();
-					const UClass* SourceCompClass = (*SourceComponent)->GetClass();
-
-					// Handle the unlikely case of the default component and the source actor component not being the exact same class by finding
-					// the common base class across all three components (default, source, and destination)
-					if ( DefaultCompClass != SourceCompClass )
-					{
-						const UClass* CommonBaseClassWithDefault = CurComp->FindNearestCommonBaseClass( DefaultCompClass );
-						const UClass* CommonBaseClassWithSource = CurComp->FindNearestCommonBaseClass( SourceCompClass );
-						if ( CommonBaseClassWithDefault && CommonBaseClassWithSource )
-						{
-							// If both components yielded the same common base, then that's the common base of all three
-							if ( CommonBaseClassWithDefault == CommonBaseClassWithSource )
-							{
-								CommonBaseClass = CommonBaseClassWithDefault;
-							}
-							// If not, find a common base across all three components
-							else
-							{
-								CommonBaseClass = const_cast<UClass*>(CommonBaseClassWithDefault)->GetDefaultObject()->FindNearestCommonBaseClass( CommonBaseClassWithSource );
-							}
-						}
-					}
-					else
-					{
-						CommonBaseClass = CurComp->FindNearestCommonBaseClass( DefaultCompClass );
-					}
-
-					// If all three components have a base class in common, copy the properties from that base class from the source actor component
-					// to the destination
-					if ( CommonBaseClass )
-					{
-						// Iterate through the properties, only copying those which are non-native, non-transient, non-component, and not identical
-						// to the values in the default component
-						for ( UProperty* Property = CommonBaseClass->PropertyLink; Property != NULL; Property = Property->PropertyLinkNext )
-						{
-							const bool bIsTransient = !!( Property->PropertyFlags & CPF_Transient );
-							const bool bIsIdentical = Property->Identical_InContainer(*SourceComponent, *DefaultComponent);
-							const bool bIsComponent = !!( Property->PropertyFlags & (CPF_InstancedReference | CPF_ContainsInstancedReference) );
-
-							if ( !bIsTransient && !bIsIdentical && !bIsComponent )
-							{
-								bCopiedAnyProperty = true;
-								Property->CopyCompleteValue_InContainer(CurComp, *SourceComponent);
-							}
-						}
-					}
-				}
-			}
-		}
-
-		// If any properties were copied at all, alert the actor to the changes
-		if ( bCopiedAnyProperty )
-		{
-			DestActor->PostEditChange();
-		}
-	}
-}
-
-AActor* UEditorEngine::ConvertBrushesToStaticMesh(const FString& InStaticMeshPackageName, TArray<ABrush*>& InBrushesToConvert, const FVector& InPivotLocation)
-{
-	AActor* NewActor(NULL);
-
-	FName ObjName = *FPackageName::GetLongPackageAssetName(InStaticMeshPackageName);
-
-
-	UPackage* Pkg = CreatePackage(NULL, *InStaticMeshPackageName);
-	check(Pkg != nullptr);
-
-	FVector Location(0.0f, 0.0f, 0.0f);
-	FRotator Rotation(0.0f, 0.0f, 0.0f);
-	for(int32 BrushesIdx = 0; BrushesIdx < InBrushesToConvert.Num(); ++BrushesIdx )
-	{
-		// Cache the location and rotation.
-		Location = InBrushesToConvert[BrushesIdx]->GetActorLocation();
-		Rotation = InBrushesToConvert[BrushesIdx]->GetActorRotation();
-
-		// Leave the actor's rotation but move it to origin so the Static Mesh will generate correctly.
-		InBrushesToConvert[BrushesIdx]->TeleportTo(Location - InPivotLocation, Rotation, false, true);
-	}
-
-	GEditor->RebuildModelFromBrushes(ConversionTempModel, true );
-	GEditor->bspBuildFPolys(ConversionTempModel, true, 0);
-
-	if (0 < ConversionTempModel->Polys->Element.Num())
-	{
-		UStaticMesh* NewMesh = CreateStaticMeshFromBrush(Pkg, ObjName, NULL, ConversionTempModel);
-		NewActor = FActorFactoryAssetProxy::AddActorForAsset( NewMesh );
-
-		NewActor->Modify();
-
-		NewActor->InvalidateLightingCache();
-		NewActor->PostEditChange();
-		NewActor->PostEditMove( true );
-		NewActor->Modify();
-		GEditor->Layers->InitializeNewActorLayers( NewActor );
-
-		// Teleport the new actor to the old location but not the old rotation. The static mesh is built to the rotation already.
-		NewActor->TeleportTo(InPivotLocation, FRotator(0.0f, 0.0f, 0.0f), false, true);
-
-		// Destroy the old brushes.
-		for( int32 BrushIdx = 0; BrushIdx < InBrushesToConvert.Num(); ++BrushIdx )
-		{
-			GEditor->Layers->DisassociateActorFromLayers( InBrushesToConvert[BrushIdx] );
-			GWorld->EditorDestroyActor( InBrushesToConvert[BrushIdx], true );
-		}
-
-		// Notify the asset registry
-		FAssetRegistryModule::AssetCreated(NewMesh);
-	}
-
-	ConversionTempModel->EmptyModel(1, 1);
-	GEditor->RebuildAlteredBSP();
-	GEditor->RedrawLevelEditingViewports();
-
-	return NewActor;
-}
-
-struct TConvertData
-{
-	const TArray<AActor*> ActorsToConvert;
-	UClass* ConvertToClass;
-	const TSet<FString> ComponentsToConsider;
-	bool bUseSpecialCases;
-
-	TConvertData(const TArray<AActor*>& InActorsToConvert, UClass* InConvertToClass, const TSet<FString>& InComponentsToConsider, bool bInUseSpecialCases)
-		: ActorsToConvert(InActorsToConvert)
-		, ConvertToClass(InConvertToClass)
-		, ComponentsToConsider(InComponentsToConsider)
-		, bUseSpecialCases(bInUseSpecialCases)
-	{
-
-	}
-};
-
-namespace ConvertHelpers
-{
-	void OnBrushToStaticMeshNameCommitted(const FString& InSettingsPackageName, TConvertData InConvertData)
-	{
-		GEditor->DoConvertActors(InConvertData.ActorsToConvert, InConvertData.ConvertToClass, InConvertData.ComponentsToConsider, InConvertData.bUseSpecialCases, InSettingsPackageName);
-	}
-
-	void GetBrushList(const TArray<AActor*>& InActorsToConvert, UClass* InConvertToClass, TArray<ABrush*>& OutBrushList, int32& OutBrushIndexForReattachment)
-	{
-		for( int32 ActorIdx = 0; ActorIdx < InActorsToConvert.Num(); ++ActorIdx )
-		{
-			AActor* ActorToConvert = InActorsToConvert[ActorIdx];
-			if (ActorToConvert->GetClass()->IsChildOf(ABrush::StaticClass()) && InConvertToClass == AStaticMeshActor::StaticClass())
-			{
-				GEditor->SelectActor(ActorToConvert, true, true);
-				OutBrushList.Add(Cast<ABrush>(ActorToConvert));
-
-				// If this is a single brush conversion then this index will be used for re-attachment.
-				OutBrushIndexForReattachment = ActorIdx;
-			}
-		}
-	}
-}
-
-void UEditorEngine::ConvertActors( const TArray<AActor*>& ActorsToConvert, UClass* ConvertToClass, const TSet<FString>& ComponentsToConsider, bool bUseSpecialCases )
-{
-	// Early out if actor deletion is currently forbidden
-	if (GEditor->ShouldAbortActorDeletion())
-	{
-		return;
-	}
-
-	GEditor->SelectNone(true, true);
-
-	// List of brushes being converted.
-	TArray<ABrush*> BrushList;
-	int32 BrushIndexForReattachment;
-	ConvertHelpers::GetBrushList(ActorsToConvert, ConvertToClass, BrushList, BrushIndexForReattachment);
-
-	if( BrushList.Num() )
-	{
-		TConvertData ConvertData(ActorsToConvert, ConvertToClass, ComponentsToConsider, bUseSpecialCases);
-
-		TSharedPtr<SWindow> CreateAssetFromActorWindow =
-			SNew(SWindow)
-			.Title(LOCTEXT("SelectPath", "Select Path"))
-			.ToolTipText(LOCTEXT("SelectPathTooltip", "Select the path where the static mesh will be created"))
-			.ClientSize(FVector2D(400, 400));
-
-		TSharedPtr<SCreateAssetFromObject> CreateAssetFromActorWidget;
-		CreateAssetFromActorWindow->SetContent
-			(
-			SAssignNew(CreateAssetFromActorWidget, SCreateAssetFromObject, CreateAssetFromActorWindow)
-			.AssetFilenameSuffix(TEXT("StaticMesh"))
-			.HeadingText(LOCTEXT("ConvertBrushesToStaticMesh_Heading", "Static Mesh Name:"))
-			.CreateButtonText(LOCTEXT("ConvertBrushesToStaticMesh_ButtonLabel", "Create Static Mesh"))
-			.OnCreateAssetAction(FOnPathChosen::CreateStatic(ConvertHelpers::OnBrushToStaticMeshNameCommitted, ConvertData))
-			);
-
-		TSharedPtr<SWindow> RootWindow = FGlobalTabmanager::Get()->GetRootWindow();
-		if (RootWindow.IsValid())
-		{
-			FSlateApplication::Get().AddWindowAsNativeChild(CreateAssetFromActorWindow.ToSharedRef(), RootWindow.ToSharedRef());
-		}
-		else
-		{
-			FSlateApplication::Get().AddWindow(CreateAssetFromActorWindow.ToSharedRef());
-		}
-	}
-	else
-	{
-		DoConvertActors(ActorsToConvert, ConvertToClass, ComponentsToConsider, bUseSpecialCases, TEXT(""));
-	}
-}
-
-void UEditorEngine::DoConvertActors( const TArray<AActor*>& ActorsToConvert, UClass* ConvertToClass, const TSet<FString>& ComponentsToConsider, bool bUseSpecialCases, const FString& InStaticMeshPackageName )
-{
-	// Early out if actor deletion is currently forbidden
-	if (GEditor->ShouldAbortActorDeletion())
-	{
-		return;
-	}
-
-	GWarn->BeginSlowTask( NSLOCTEXT("UnrealEd", "ConvertingActors", "Converting Actors"), true );
-
-	// Scope the transaction - we need it to end BEFORE we finish the slow task we just started
-	{
-		const FScopedTransaction Transaction( NSLOCTEXT("EditorEngine", "ConvertActors", "Convert Actors") );
-
-		GEditor->GetSelectedActors()->BeginBatchSelectOperation();
-
-		TArray<AActor*> ConvertedActors;
-		int32 NumActorsToConvert = ActorsToConvert.Num();
-
-		// Cache for attachment info of all actors being converted.
-		TArray<ReattachActorsHelper::FActorAttachmentCache> AttachmentInfo;
-
-		// Maps actors from old to new for quick look-up.
-		TMap<AActor*, AActor*> ConvertedMap;
-
-		GEditor->SelectNone(true, true);
-		ReattachActorsHelper::CacheAttachments(ActorsToConvert, AttachmentInfo);
-
-		// List of brushes being converted.
-		TArray<ABrush*> BrushList;
-
-		// The index of a brush, utilized for re-attachment purposes when a single brush is being converted.
-		int32 BrushIndexForReattachment = 0;
-
-		FVector CachePivotLocation = GEditor->GetPivotLocation();
-		for( int32 ActorIdx = 0; ActorIdx < ActorsToConvert.Num(); ++ActorIdx )
-		{
-			AActor* ActorToConvert = ActorsToConvert[ActorIdx];
-			if (ActorToConvert->GetClass()->IsChildOf(ABrush::StaticClass()) && ConvertToClass == AStaticMeshActor::StaticClass())
-			{
-				GEditor->SelectActor(ActorToConvert, true, true);
-				BrushList.Add(Cast<ABrush>(ActorToConvert));
-
-				// If this is a single brush conversion then this index will be used for re-attachment.
-				BrushIndexForReattachment = ActorIdx;
-			}
-		}
-
-		// If no package name is supplied, ask the user
-		if( BrushList.Num() )
-		{
-			AActor* ConvertedBrushActor = ConvertBrushesToStaticMesh(InStaticMeshPackageName, BrushList, CachePivotLocation);
-			ConvertedActors.Add(ConvertedBrushActor);
-
-			// If only one brush is being converted, reattach it to whatever it was attached to before.
-			// Multiple brushes become impossible to reattach due to the single actor returned.
-			if(BrushList.Num() == 1)
-			{
-				ReattachActorsHelper::CacheActorConvert(BrushList[0], ConvertedBrushActor, ConvertedMap, AttachmentInfo[BrushIndexForReattachment]);
-			}
-
-		}
-
-		for( int32 ActorIdx = 0; ActorIdx < ActorsToConvert.Num(); ++ActorIdx )
-		{
-			AActor* ActorToConvert = ActorsToConvert[ ActorIdx ];
-			// Source actor display label
-			FString ActorLabel = ActorToConvert->GetActorLabel();
-			// Low level source actor object name
-			FName ActorObjectName = ActorToConvert->GetFName();
-	
-			// The class of the actor we are about to replace
-			UClass* ClassToReplace = ActorToConvert->GetClass();
-
-			// Spawn the new actor
-			AActor* NewActor = NULL;
-
-			ABrush* Brush = Cast< ABrush >( ActorToConvert );
-			if ( ( Brush && FActorEditorUtils::IsABuilderBrush(Brush) ) ||
-				(ClassToReplace->IsChildOf(ABrush::StaticClass()) && ConvertToClass == AStaticMeshActor::StaticClass()) )
-			{
-				continue;
-			}
-
-			if (bUseSpecialCases)
-			{
-				// Disable grouping temporarily as the following code assumes only one actor will be selected at any given time
-				const bool bGroupingActiveSaved = GEditor->bGroupingActive;
-
-				GEditor->bGroupingActive = false;
-				GEditor->SelectNone(true, true);
-				GEditor->SelectActor(ActorToConvert, true, true);
-
-				// Each of the following 'special case' conversions will convert ActorToConvert to ConvertToClass if possible.
-				// If it does it will mark the original for delete and select the new actor
-				if (ClassToReplace->IsChildOf(ALight::StaticClass()))
-				{
-					ConvertLightActors(ConvertToClass);
-				}
-				else if (ClassToReplace->IsChildOf(ABrush::StaticClass()) && ConvertToClass->IsChildOf(AVolume::StaticClass()))
-				{
-					ConvertSelectedBrushesToVolumes(ConvertToClass);
-				}
-				else
-				{
-					ConvertActorsFromClass(ClassToReplace, ConvertToClass);
-				}
-				if (ActorToConvert->IsPendingKill())
-				{
-					// Converted by one of the above
-					check (1 == GEditor->GetSelectedActorCount());
-					NewActor = CastChecked< AActor >(GEditor->GetSelectedActors()->GetSelectedObject(0));
-
-					// Caches information for finding the new actor using the pre-converted actor.
-					ReattachActorsHelper::CacheActorConvert(ActorToConvert, NewActor, ConvertedMap, AttachmentInfo[ActorIdx]);
-				}
-				else
-				{
-					// Failed to convert, make sure the actor is unselected
-					GEditor->SelectActor(ActorToConvert, false, true);
-				}
-
-				// Restore previous grouping setting
-				GEditor->bGroupingActive = bGroupingActiveSaved;
-			}
-
-
-			if (!NewActor)
-			{
-				// Set the current level to the level where the convertible actor resides
-				check(ActorToConvert);
-				UWorld* World = ActorToConvert->GetWorld();
-				ULevel* ActorLevel = ActorToConvert->GetLevel();
-				check(World);
-				checkSlow( ActorLevel );
-				// Find a common base class between the actors so we know what properties to copy
-				const UClass* CommonBaseClass = ActorToConvert->FindNearestCommonBaseClass( ConvertToClass );
-				check ( CommonBaseClass );	
-
-				// Take the old actors location always, not rotation.  If rotation was changed on the source actor, it will be copied below.
-				FVector SpawnLoc = ActorToConvert->GetActorLocation();
-				FRotator SpawnRot = ActorToConvert->GetActorRotation();
-				{
-					FActorSpawnParameters SpawnInfo;
-					SpawnInfo.OverrideLevel = ActorLevel;
-					SpawnInfo.bNoCollisionFail = true;
-					NewActor = World->SpawnActor( ConvertToClass, &SpawnLoc, &SpawnRot, SpawnInfo );
-
-					if (NewActor)
-					{
-						// Copy non component properties from the old actor to the new actor
-						for( UProperty* Property = CommonBaseClass->PropertyLink; Property != NULL; Property = Property->PropertyLinkNext )
-						{
-							const bool bIsTransient = !!(Property->PropertyFlags & CPF_Transient);
-							const bool bIsComponentProp = !!(Property->PropertyFlags & (CPF_InstancedReference | CPF_ContainsInstancedReference));
-							const bool bIsIdentical = Property->Identical_InContainer(ActorToConvert, ClassToReplace->GetDefaultObject());
-
-							if ( !bIsTransient && !bIsIdentical && !bIsComponentProp && Property->GetName() != TEXT("Tag") )
-							{
-								// Copy only if not native, not transient, not identical, and not a component.
-								// Copying components directly here is a bad idea because the next garbage collection will delete the component since we are deleting its outer.  
-
-								// Also do not copy the old actors tag.  That will always come up as not identical since the default actor's Tag is "None" and SpawnActor uses the actor's class name
-								// The tag will be examined for changes later.
-								Property->CopyCompleteValue_InContainer(NewActor, ActorToConvert);
-							}
-						}
-
-						// Copy properties from actor components
-						CopyActorComponentProperties( ActorToConvert, NewActor, ComponentsToConsider );
-
-
-						// Caches information for finding the new actor using the pre-converted actor.
-						ReattachActorsHelper::CacheActorConvert(ActorToConvert, NewActor, ConvertedMap, AttachmentInfo[ActorIdx]);
-
-						NewActor->Modify();
-						NewActor->InvalidateLightingCache();
-						NewActor->PostEditChange();
-						NewActor->PostEditMove( true );
-						GEditor->Layers->InitializeNewActorLayers( NewActor );
-
-						// Destroy the old actor.
-						ActorToConvert->Modify();
-						GEditor->Layers->DisassociateActorFromLayers( ActorToConvert );
-						World->EditorDestroyActor( ActorToConvert, true );	
-					}
-				}
-			}
-
-			if (NewActor)
-			{
-				// If the actor label isn't actually anything custom allow the name to be changed
-				// to avoid cases like converting PointLight->SpotLight still being called PointLight after conversion
-				FString ClassName = ClassToReplace->GetName();
-				
-				// Remove any number off the end of the label
-				int32 Number = 0;
-				if( !ActorLabel.StartsWith( ClassName ) || !FParse::Value(*ActorLabel, *ClassName, Number)  )
-				{
-					NewActor->SetActorLabel(ActorLabel);
-				}
-
-				ConvertedActors.Add(NewActor);
-
-				UE_LOG(LogEditor, Log, TEXT("Converted: %s to %s"), *ActorLabel, *NewActor->GetActorLabel() );
-
-				FFormatNamedArguments Args;
-				Args.Add( TEXT("OldActorName"), FText::FromString( ActorLabel ) );
-				Args.Add( TEXT("NewActorName"), FText::FromString( NewActor->GetActorLabel() ) );
-				const FText StatusUpdate = FText::Format( LOCTEXT("ConvertActorsTaskStatusUpdateMessageFormat", "Converted: {OldActorName} to {NewActorName}"), Args);
-
-				GWarn->StatusUpdate( ConvertedActors.Num(), NumActorsToConvert, StatusUpdate );				
-			}
-		}
-
-		// Reattaches actors based on their previous parent child relationship.
-		ReattachActorsHelper::ReattachActors(ConvertedMap, AttachmentInfo);
-
-		// Select the new actors
-		GEditor->SelectNone( false, true );
-		for( TArray<AActor*>::TConstIterator it(ConvertedActors); it; ++it )
-		{
-			GEditor->SelectActor(*it, true, true);
-		}
-
-		GEditor->GetSelectedActors()->EndBatchSelectOperation();
-
-		GEditor->RedrawLevelEditingViewports();
-
-		ULevel::LevelDirtiedEvent.Broadcast();
-		
-		// Clean up
-		CollectGarbage( GARBAGE_COLLECTION_KEEPFLAGS );
-	}
-	// End the slow task
-	GWarn->EndSlowTask();
-}
-
-void UEditorEngine::NotifyToolsOfObjectReplacement(const TMap<UObject*, UObject*>& OldToNewInstanceMap)
-{
-	// This can be called early on during startup if blueprints need to be compiled.  
-	// If the property module isn't loaded then there aren't any property windows to update
-	if( FModuleManager::Get().IsModuleLoaded( "PropertyEditor" ) )
-	{
-		FPropertyEditorModule& PropertyEditorModule = FModuleManager::GetModuleChecked<FPropertyEditorModule>( "PropertyEditor" );
-		PropertyEditorModule.ReplaceViewedObjects( OldToNewInstanceMap );
-	}
-
-	// Check to see if any selected components were reinstanced
-	USelection* ComponentSelection = GetSelectedComponents();
-	if (ComponentSelection)
-	{
-		TArray<TWeakObjectPtr<UObject> > SelectedComponents;
-		ComponentSelection->GetSelectedObjects(SelectedComponents);
-
-		ComponentSelection->BeginBatchSelectOperation();
-		for (int32 i = 0; i < SelectedComponents.Num(); ++i)
-		{
-			UObject* Component = SelectedComponents[i].GetEvenIfUnreachable();
-
-			// If the component corresponds to a new instance in the map, update the selection accordingly
-			if (OldToNewInstanceMap.Contains(Component))
-			{
-				if (UActorComponent* NewComponent = CastChecked<UActorComponent>(OldToNewInstanceMap[Component], ECastCheckedType::NullAllowed))
-				{
-					ComponentSelection->Deselect(Component);
-					SelectComponent(NewComponent, true, false);
-				}
-			}
-		}
-		ComponentSelection->EndBatchSelectOperation();
-	}
-	
-	BroadcastObjectsReplaced(OldToNewInstanceMap);
-}
-
-void UEditorEngine::DisableRealtimeViewports()
-{
-	for( int32 x = 0 ; x < AllViewportClients.Num() ; ++x)
-	{
-		FEditorViewportClient* VC = AllViewportClients[x];
-		if( VC )
-		{
-			VC->SetRealtime( false, true );
-		}
-	}
-
-	RedrawAllViewports();
-
-	FEditorSupportDelegates::UpdateUI.Broadcast();
-}
-
-
-void UEditorEngine::RestoreRealtimeViewports()
-{
-	for( int32 x = 0 ; x < AllViewportClients.Num() ; ++x)
-	{
-		FEditorViewportClient* VC = AllViewportClients[x];
-		if( VC )
-		{
-			VC->RestoreRealtime(true);
-		}
-	}
-
-	RedrawAllViewports();
-
-	FEditorSupportDelegates::UpdateUI.Broadcast();
-}
-
-
-bool UEditorEngine::IsAnyViewportRealtime()
-{
-	for( int32 x = 0 ; x < AllViewportClients.Num() ; ++x)
-	{
-		FEditorViewportClient* VC = AllViewportClients[x];
-		if( VC )
-		{
-			if( VC->IsRealtime() )
-			{
-				return true;
-			}
-		}
-	}
-	return false;
-}
-
-bool UEditorEngine::ShouldThrottleCPUUsage() const
-{
-	bool bShouldThrottle = false;
-
-	bool bIsForeground = FPlatformProcess::IsThisApplicationForeground();
-
-	if( !bIsForeground )
-	{
-		const UEditorUserSettings* Settings = GetDefault<UEditorUserSettings>();
-		bShouldThrottle = Settings->bThrottleCPUWhenNotForeground;
-
-		// Check if we should throttle due to all windows being minimized
-		if ( !bShouldThrottle )
-		{
-			return bShouldThrottle = AreAllWindowsHidden();
-		}
-	}
-
-	// Don't throttle during amortized export, greatly increases export time
-	if (IsLightingBuildCurrentlyExporting())
-	{
-		return false;
-	}
-
-    return bShouldThrottle && !IsRunningCommandlet();
-}
-
-bool UEditorEngine::AreAllWindowsHidden() const
-{
-	const TArray< TSharedRef<SWindow> > AllWindows = FSlateApplication::Get().GetInteractiveTopLevelWindows();
-
-	bool bAllHidden = true;
-	for( const TSharedRef<SWindow>& Window : AllWindows )
-	{
-		if( !Window->IsWindowMinimized() && Window->IsVisible() )
-		{
-			bAllHidden = false;
-			break;
-		}
-	}
-
-	return bAllHidden;
-}
-
-AActor* UEditorEngine::AddActor(ULevel* InLevel, UClass* Class, const FTransform& Transform, bool bSilent, EObjectFlags ObjectFlags)
-{
-	check( Class );
-
-	if( !bSilent )
-	{
-		const auto Location = Transform.GetLocation();
-		UE_LOG(LogEditor, Log,
-			TEXT("Attempting to add actor of class '%s' to level at %0.2f,%0.2f,%0.2f"),
-			*Class->GetName(), Location.X, Location.Y, Location.Z );
-	}
-
-	///////////////////////////////
-	// Validate class flags.
-
-	if( Class->HasAnyClassFlags(CLASS_Abstract) )
-	{
-		UE_LOG(LogEditor, Error, TEXT("Class %s is abstract.  You can't add actors of this class to the world."), *Class->GetName() );
-		return NULL;
-	}
-	if( Class->HasAnyClassFlags(CLASS_NotPlaceable) )
-	{
-		UE_LOG(LogEditor, Error, TEXT("Class %s isn't placeable.  You can't add actors of this class to the world."), *Class->GetName() );
-		return NULL;
-	}
-	if( Class->HasAnyClassFlags(CLASS_Transient) )
-	{
-		UE_LOG(LogEditor, Error, TEXT("Class %s is transient.  You can't add actors of this class in UnrealEd."), *Class->GetName() );
-		return NULL;
-	}
-
-
-	UWorld* World = InLevel->OwningWorld;
-	ULevel* DesiredLevel = InLevel;
-
-	// Don't spawn the actor if the current level is locked.
-	if ( FLevelUtils::IsLevelLocked(DesiredLevel) )
-	{
-		FNotificationInfo Info( NSLOCTEXT("UnrealEd", "Error_OperationDisallowedOnLockedLevel", "The requested operation could not be completed because the level is locked.") );
-		Info.ExpireDuration = 3.0f;
-		FSlateNotificationManager::Get().AddNotification(Info);
-		return NULL;
-	}
-
-	// Transactionally add the actor.
-	AActor* Actor = NULL;
-	{
-		FScopedTransaction Transaction( NSLOCTEXT("UnrealEd", "AddActor", "Add Actor") );
-		if ( !(ObjectFlags & RF_Transactional) )
-		{
-			// Don't attempt a transaction if the actor we are spawning isn't transactional
-			Transaction.Cancel();
-		}
-		SelectNone( false, true );
-
-		AActor* Default = Class->GetDefaultObject<AActor>();
-
-		FActorSpawnParameters SpawnInfo;
-		SpawnInfo.OverrideLevel = DesiredLevel;
-		SpawnInfo.bNoCollisionFail = true;
-		SpawnInfo.ObjectFlags = ObjectFlags;
-		const auto Location = Transform.GetLocation();
-		const auto Rotation = Transform.GetRotation().Rotator();
-		Actor = World->SpawnActor( Class, &Location, &Rotation, SpawnInfo );
-
-		if( Actor )
-		{
-			SelectActor( Actor, 1, 0 );
-			Actor->InvalidateLightingCache();
-			Actor->PostEditMove( true );
-		}
-		else
-		{
-			FMessageDialog::Open( EAppMsgType::Ok, NSLOCTEXT("UnrealEd", "Error_Couldn'tSpawnActor", "Couldn't spawn actor. Please check the log.") );
-		}
-	}
-
-	if( Actor )
-	{
-		// If this actor is part of any layers (set in its default properties), add them into the visible layers list.
-		GEditor->Layers->SetLayersVisibility( Actor->Layers, true );
-
-		// Clean up.
-		Actor->MarkPackageDirty();
-		ULevel::LevelDirtiedEvent.Broadcast();
-	}
-
-	NoteSelectionChange();
-
-	return Actor;
-}
-
-TArray<AActor*> UEditorEngine::AddExportTextActors(const FString& ExportText, bool bSilent, EObjectFlags ObjectFlags)
-{
-	TArray<AActor*> NewActors;
-
-	// Don't spawn the actor if the current level is locked.
-	ULevel* CurrentLevel = GWorld->GetCurrentLevel();
-	if ( FLevelUtils::IsLevelLocked( CurrentLevel ) )
-	{
-		FMessageDialog::Open(EAppMsgType::Ok, NSLOCTEXT("UnrealEd", "Error_OperationDisallowedOnLockedLevelAddExportTextActors", "AddExportTextActors: The requested operation could not be completed because the level is locked."));
-		return NewActors;
-	}
-
-	// Use a level factory to spawn all the actors using the ExportText
-	auto Factory = NewObject<ULevelFactory>();
-	FVector Location;
-	{
-		FScopedTransaction Transaction( NSLOCTEXT("UnrealEd", "AddActor", "Add Actor") );
-		if ( !(ObjectFlags & RF_Transactional) )
-		{
-			// Don't attempt a transaction if the actor we are spawning isn't transactional
-			Transaction.Cancel();
-		}
-		// Remove the selection to detect the actors that were created during FactoryCreateText. They will be selected when the operation in complete
-		GEditor->SelectNone( false, true );
-		const TCHAR* Text = *ExportText;
-		if ( Factory->FactoryCreateText( ULevel::StaticClass(), CurrentLevel, CurrentLevel->GetFName(), ObjectFlags, NULL, TEXT("paste"), Text, Text + FCString::Strlen(Text), GWarn ) != NULL )
-		{
-			// Now get the selected actors and calculate a center point between all their locations.
-			USelection* ActorSelection = GEditor->GetSelectedActors();
-			FVector Origin = FVector::ZeroVector;
-			for ( int32 ActorIdx = 0; ActorIdx < ActorSelection->Num(); ++ActorIdx )
-			{
-				AActor* Actor = CastChecked<AActor>(ActorSelection->GetSelectedObject(ActorIdx));
-				NewActors.Add(Actor);
-				Origin += Actor->GetActorLocation();
-			}
-
-			if ( NewActors.Num() > 0 )
-			{
-				// Finish the Origin calculations now that we know we are not going to divide by zero
-				Origin /= NewActors.Num();
-
-				// Set up the spawn location
-				FSnappingUtils::SnapPointToGrid( GEditor->ClickLocation, FVector(0, 0, 0) );
-				Location = GEditor->ClickLocation;
-				FVector Collision = NewActors[0]->GetPlacementExtent();
-				Location += GEditor->ClickPlane * (FVector::BoxPushOut(GEditor->ClickPlane, Collision) + 0.1f);
-				FSnappingUtils::SnapPointToGrid( Location, FVector(0, 0, 0) );
-
-				// For every spawned actor, teleport to the target loction, preserving the relative translation to the other spawned actors.
-				for ( int32 ActorIdx = 0; ActorIdx < NewActors.Num(); ++ActorIdx )
-				{
-					AActor* Actor = NewActors[ActorIdx];
-					FVector OffsetToOrigin = Actor->GetActorLocation() - Origin;
-
-					Actor->TeleportTo(Location + OffsetToOrigin, Actor->GetActorRotation(), false, true );
-					Actor->InvalidateLightingCache();
-					Actor->PostEditMove( true );
-
-					GEditor->Layers->SetLayersVisibility( Actor->Layers, true );
-
-					Actor->MarkPackageDirty();
-				}
-
-				// Send notification about a new actor being created
-				ULevel::LevelDirtiedEvent.Broadcast();
-				GEditor->NoteSelectionChange();
-			}
-		}
-	}
-
-	if( NewActors.Num() > 0 && !bSilent )
-	{
-		UE_LOG(LogEditor, Log,
-			TEXT("Added '%d' actor(s) to level at %0.2f,%0.2f,%0.2f"),
-			NewActors.Num(), Location.X, Location.Y, Location.Z );
-	}
-
-	return NewActors;
-}
-
-UActorFactory* UEditorEngine::FindActorFactoryForActorClass( const UClass* InClass )
-{
-	for( int32 i = 0 ; i < ActorFactories.Num() ; ++i )
-	{
-		UActorFactory* Factory = ActorFactories[i];
-
-		// force NewActorClass update
-		const UObject* const ActorCDO = Factory->GetDefaultActor( FAssetData() );
-		if( ActorCDO != NULL && ActorCDO->GetClass() == InClass )
-		{
-			return Factory;
-		}
-	}
-
-	return NULL;
-}
-
-UActorFactory* UEditorEngine::FindActorFactoryByClass( const UClass* InClass ) const
-{
-	for( int32 i = 0 ; i < ActorFactories.Num() ; ++i )
-	{
-		UActorFactory* Factory = ActorFactories[i];
-
-		if( Factory != NULL && Factory->GetClass() == InClass )
-		{
-			return Factory;
-		}
-	}
-
-	return NULL;
-}
-
-UActorFactory* UEditorEngine::FindActorFactoryByClassForActorClass( const UClass* InFactoryClass, const UClass* InActorClass )
-{
-	for ( int32 i = 0; i < ActorFactories.Num(); ++i )
-	{
-		UActorFactory* Factory = ActorFactories[i];
-
-		if ( Factory != NULL && Factory->GetClass() == InFactoryClass )
-		{
-			// force NewActorClass update
-			const UObject* const ActorCDO = Factory->GetDefaultActor( FAssetData() );
-			if ( ActorCDO != NULL && ActorCDO->GetClass() == InActorClass )
-			{
-				return Factory;
-			}
-		}
-	}
-
-	return NULL;
-}
-
-void UEditorEngine::PreWorldOriginOffset(UWorld* InWorld, FIntVector InSrcOrigin, FIntVector InDstOrigin)
-{
-	// In case we simulating world in the editor, 
-	// we need to shift current viewport as well, 
-	// so the streaming procedure will receive correct view location
-	if (bIsSimulatingInEditor && 
-		GCurrentLevelEditingViewportClient &&
-		GCurrentLevelEditingViewportClient->IsVisible())
-	{
-		FVector ViewLocation = GCurrentLevelEditingViewportClient->GetViewLocation();
-		GCurrentLevelEditingViewportClient->SetViewLocation(ViewLocation - FVector(InDstOrigin - InSrcOrigin));
-	}
-}
-
-void UEditorEngine::SetPreviewMeshMode( bool bState )
-{
-	// Only change the state if it's different than the current.
-	if( bShowPreviewMesh != bState )
-	{
-		// Set the preview mesh mode state. 
-		bShowPreviewMesh = !bShowPreviewMesh;
-
-		bool bHavePreviewMesh = (PreviewMeshComp != NULL);
-
-		// It's possible that the preview mesh hasn't been loaded yet,
-		// such as on first-use for the preview mesh mode or there 
-		// could be valid mesh names provided in the INI. 
-		if( !bHavePreviewMesh )
-		{
-			bHavePreviewMesh = LoadPreviewMesh( GUnrealEd->PreviewMeshIndex );
-		}
-
-		// If we have a	preview mesh, change it's visibility based on the preview state. 
-		if( bHavePreviewMesh )
-		{
-			PreviewMeshComp->SetVisibility( bShowPreviewMesh );
-			PreviewMeshComp->SetCastShadow( bShowPreviewMesh );
-			RedrawLevelEditingViewports();
-		}
-		else
-		{
-			// Without a preview mesh, we can't really use the preview mesh mode. 
-			// So, disable it even if the caller wants to enable it. 
-			bShowPreviewMesh = false;
-		}
-	}
-}
-
-
-void UEditorEngine::UpdatePreviewMesh()
-{
-	if( bShowPreviewMesh )
-	{
-		// The component should exist by now. Is the bPlayerHeight state 
-		// manually changed instead of calling SetPreviewMeshMode()?
-		check(PreviewMeshComp);
-
-		// Use the cursor world location as the starting location for the line check. 
-		FViewportCursorLocation CursorLocation = GCurrentLevelEditingViewportClient->GetCursorWorldLocationFromMousePos();
-		FVector LineCheckStart = CursorLocation.GetOrigin();
-		FVector LineCheckEnd = CursorLocation.GetOrigin() + CursorLocation.GetDirection() * HALF_WORLD_MAX;
-
-		// Perform a line check from the camera eye to the surface to place the preview mesh. 
-		FHitResult Hit(ForceInit);
-		static FName UpdatePreviewMeshTrace = FName(TEXT("UpdatePreviewMeshTrace"));
-		FCollisionQueryParams LineParams(UpdatePreviewMeshTrace, true);
-		LineParams.bTraceComplex = false;
-		if ( GWorld->LineTraceSingleByObjectType(Hit, LineCheckStart, LineCheckEnd, FCollisionObjectQueryParams(ECC_WorldStatic), LineParams) ) 
-		{
-			// Dirty the transform so UpdateComponent will actually update the transforms. 
-			PreviewMeshComp->SetRelativeLocation(Hit.Location);
-		}
-
-		// Redraw the viewports because the mesh won't 
-		// be shown or hidden until that happens. 
-		RedrawLevelEditingViewports();
-	}
-}
-
-
-void UEditorEngine::CyclePreviewMesh()
-{
-	const ULevelEditorViewportSettings& ViewportSettings = *GetDefault<ULevelEditorViewportSettings>();
-	if( !ViewportSettings.PreviewMeshes.Num() )
-	{
-		return;
-	}
-
-	const int32 StartingPreviewMeshIndex = FMath::Min(GUnrealEd->PreviewMeshIndex, ViewportSettings.PreviewMeshes.Num() - 1);
-	int32 CurrentPreviewMeshIndex = StartingPreviewMeshIndex;
-	bool bPreviewMeshFound = false;
-
-	do
-	{
-		// Cycle to the next preview mesh. 
-		CurrentPreviewMeshIndex++;
-
-		// If we reached the max index, start at index zero.
-		if( CurrentPreviewMeshIndex == ViewportSettings.PreviewMeshes.Num() )
-		{
-			CurrentPreviewMeshIndex = 0;
-		}
-
-		// Load the mesh (if not already) onto the mesh actor. 
-		bPreviewMeshFound = LoadPreviewMesh(CurrentPreviewMeshIndex);
-
-		if( bPreviewMeshFound )
-		{
-			// Save off the index so we can reference it later when toggling the preview mesh mode. 
-			GUnrealEd->PreviewMeshIndex = CurrentPreviewMeshIndex;
-		}
-
-		// Keep doing this until we found another valid mesh, or we cycled through all possible preview meshes. 
-	} while( !bPreviewMeshFound && (StartingPreviewMeshIndex != CurrentPreviewMeshIndex) );
-}
-
-bool UEditorEngine::LoadPreviewMesh( int32 Index )
-{
-	bool bMeshLoaded = false;
-
-	// Don't register the preview mesh into the PIE world!
-	if(GWorld->IsPlayInEditor())
-	{
-		UE_LOG(LogEditorViewport, Warning, TEXT("LoadPreviewMesh called while PIE world is GWorld."));
-		return false;
-	}
-
-	const ULevelEditorViewportSettings& ViewportSettings = *GetDefault<ULevelEditorViewportSettings>();
-	if( ViewportSettings.PreviewMeshes.IsValidIndex(Index) )
-	{
-		const FStringAssetReference& MeshName = ViewportSettings.PreviewMeshes[Index];
-
-		// If we don't have a preview mesh component in the world yet, create one. 
-		if( !PreviewMeshComp )
-		{
-			PreviewMeshComp = NewObject<UStaticMeshComponent>();
-			check(PreviewMeshComp);
-
-			// Attach the component to the scene even if the preview mesh doesn't load.
-			PreviewMeshComp->RegisterComponentWithWorld(GWorld);
-		}
-
-		// Load the new mesh, if not already loaded. 
-		UStaticMesh* PreviewMesh = LoadObject<UStaticMesh>(NULL, *MeshName.ToString(), NULL, LOAD_None, NULL);
-
-		// Swap out the meshes if we loaded or found the given static mesh. 
-		if( PreviewMesh )
-		{
-			bMeshLoaded = true;
-			PreviewMeshComp->StaticMesh = PreviewMesh;
-		}
-		else
-		{
-			UE_LOG(LogEditorViewport, Warning, TEXT("Couldn't load the PreviewMeshNames for the player at index, %d, with the name, %s."), Index, *MeshName.ToString());
-		}
-	}
-	else
-	{
-		UE_LOG(LogEditorViewport, Log,  TEXT("Invalid array index, %d, provided for PreviewMeshNames in UEditorEngine::LoadPreviewMesh"), Index );
-	}
-
-	return bMeshLoaded;
-}
-
-void UEditorEngine::OnLevelAddedToWorld(ULevel* InLevel, UWorld* InWorld)
-{
-	if (InLevel)
-	{
-		// Update the editorworld list, and make sure this actor is selected if it was when we began pie/sie
-		for (int32 ActorIdx = 0; ActorIdx < InLevel->Actors.Num(); ActorIdx++)
-		{
-			AActor* LevelActor = InLevel->Actors[ActorIdx];
-			if ( LevelActor )
-			{
-				ObjectsThatExistInEditorWorld.Set(LevelActor);
-
-				if ( ActorsThatWereSelected.Num() > 0 )
-				{
-					AActor* EditorActor = EditorUtilities::GetEditorWorldCounterpartActor( LevelActor );
-					if ( EditorActor && ActorsThatWereSelected.Contains( EditorActor ) )
-					{
-						SelectActor( LevelActor, true, false );
+						FFormatNamedArguments Args;
+						Args.Add( TEXT("ObjectName"), FText::FromString( Obj->GetName() ) );
+						Args.Add( TEXT("ObjectType"), FText::FromString( Obj->GetClass()->GetName() ) );
+						Args.Add( TEXT("SourceFile"), FText::FromString( FirstLeafFilename ) );
+						Args.Add( TEXT("Number"), SourceFilenames.Num() - 1 );
+						NotificationText = FText::Format( LOCTEXT("ReimportSuccessfulMultiple", "Successfuly Reimported: {ObjectName} ({ObjectType}) from file ({SourceFile}) and {Number} more"), Args );
 					}
 				}
-			}
-		}
-	}
-}
-
-void UEditorEngine::OnLevelRemovedFromWorld(ULevel* InLevel, UWorld* InWorld)
-{
-	if (InLevel)
-	{
-		// Update the editorworld list and deselect actors belonging to removed level
-		for (int32 ActorIdx = 0; ActorIdx < InLevel->Actors.Num(); ActorIdx++)
-		{
-			AActor* LevelActor = InLevel->Actors[ActorIdx];
-			if ( LevelActor )
-			{
-				ObjectsThatExistInEditorWorld.Clear(LevelActor);
-
-				SelectActor(LevelActor, false, false);
-			}
-		}
-	}
-	else
-	{
-		// UEngine::LoadMap broadcast this event with InLevel==NULL, before cleaning up the world
-		// Reset transactions buffer, to ensure that there are no references to a world which is about to be destroyed
-		ResetTransaction( NSLOCTEXT("UnrealEd", "LoadMapTransReset", "Loading a New Map") );
-	}
-}
-
-void UEditorEngine::OnGCStreamedOutLevels()
-{
-	// Reset transaction buffer because it may hold references to streamed out levels
-	ResetTransaction( NSLOCTEXT("UnrealEd", "GCStreamedOutLevelsTransReset", "GC Streaming Levels") );
-}
-
-void UEditorEngine::UpdateRecentlyLoadedProjectFiles()
-{
-	if ( FPaths::IsProjectFilePathSet() )
-	{
-		const FString AbsoluteProjectPath = IFileManager::Get().ConvertToAbsolutePathForExternalAppForRead(*FPaths::GetProjectFilePath());
-		// Update the recently loaded project files. Move this project file to the front of the list
-		TArray<FString>& RecentlyOpenedProjectFiles = AccessGameAgnosticSettings().RecentlyOpenedProjectFiles;
-		RecentlyOpenedProjectFiles.Remove( AbsoluteProjectPath );
-		RecentlyOpenedProjectFiles.Insert( AbsoluteProjectPath, 0 );
-
-		// Trim any project files that do not have the current game project file extension
-		for ( int32 FileIdx = RecentlyOpenedProjectFiles.Num() - 1; FileIdx >= 0; --FileIdx )
-		{
-			const FString FileExtension = FPaths::GetExtension(RecentlyOpenedProjectFiles[FileIdx]);
-			if ( FileExtension != FProjectDescriptor::GetExtension() )
-			{
-				RecentlyOpenedProjectFiles.RemoveAt(FileIdx, 1);
-			}
-		}
-
-		// Trim the list in case we have more than the max
-		const int32 MaxRecentProjectFiles = 1024;
-		if ( RecentlyOpenedProjectFiles.Num() > MaxRecentProjectFiles )
-		{
-			RecentlyOpenedProjectFiles.RemoveAt(MaxRecentProjectFiles, RecentlyOpenedProjectFiles.Num() - MaxRecentProjectFiles);
-		}
-
-		AccessGameAgnosticSettings().PostEditChange();
-	}
-}
-
-void UEditorEngine::UpdateAutoLoadProject()
-{
-	// If the recent project file exists and is non-empty, update the contents with the currently loaded .uproject
-	// If the recent project file exists and is empty, recent project files should not be auto-loaded
-	// If the recent project file does not exist, auto-populate it with the currently loaded project in Rocket and auto-populate empty in non-rocket
-	//		In Rocket we default to auto-loading, in non-Rocket we default to opting out of auto loading
-	const FString& AutoLoadProjectFileName = IProjectManager::Get().GetAutoLoadProjectFileName();
-	FString RecentProjectFileContents;
-	bool bShouldLoadRecentProjects = false;
-	if ( FFileHelper::LoadFileToString(RecentProjectFileContents, *AutoLoadProjectFileName) )
-	{
-		// Update to the most recently loaded project and continue auto-loading
-		if ( FPaths::IsProjectFilePathSet() )
-		{
-			FFileHelper::SaveStringToFile(FPaths::GetProjectFilePath(), *AutoLoadProjectFileName);
-		}
-
-		bShouldLoadRecentProjects = true;
-	}
-	else
-	{
-		// We do not default to auto-loading project files.
-		bShouldLoadRecentProjects = false;
-	}
-
-	AccessGameAgnosticSettings().bLoadTheMostRecentlyLoadedProjectAtStartup = bShouldLoadRecentProjects;
-
-#if PLATFORM_MAC
-	if ( !GIsBuildMachine )
-	{
-		FString OSVersion, OSSubVersion;
-		FPlatformMisc::GetOSVersions(OSVersion, OSSubVersion);
-		
-		TArray<FString> Components;
-		OSVersion.ParseIntoArray(Components, TEXT("."), true);
-		uint8 ComponentValues[3] = {0};
-		
-		for(uint32 i = 0; i < Components.Num() && i < 3; i++)
-		{
-			TTypeFromString<uint8>::FromString(ComponentValues[i], *Components[i]);
-		}
-		
-		if(ComponentValues[0] < 10 || ComponentValues[1] < 9 || (ComponentValues[1] == 9 && ComponentValues[2] < 4))
-		{
-			if(FSlateApplication::IsInitialized())
-			{
-				FSuppressableWarningDialog::FSetupInfo Info( LOCTEXT("UpdateMacOSX_Body","Please update to the latest version of Mac OS X for best performance."), LOCTEXT("UpdateMacOSX_Title","Update Mac OS X"), TEXT("UpdateMacOSX"), GEditorGameAgnosticIni );
-				Info.ConfirmText = LOCTEXT( "OK", "OK");
-				Info.bDefaultToSupressInTheFuture = true;
-				FSuppressableWarningDialog OSUpdateWarning( Info );
-				OSUpdateWarning.ShowModal();
-			}
-			else
-			{
-				UE_LOG(LogEditor, Warning, TEXT("Please update to the latest version of Mac OS X for best performance."));
-			}
-		}
-		
-		NSOpenGLContext* Context = [NSOpenGLContext currentContext];
-		const bool bTemporaryContext = (!Context);
-		if(bTemporaryContext)
-		{
-			NSOpenGLPixelFormatAttribute Attribs[] = {NSOpenGLPFAOpenGLProfile, (NSOpenGLPixelFormatAttribute)NSOpenGLProfileVersion3_2Core, (NSOpenGLPixelFormatAttribute)0};
-			NSOpenGLPixelFormat* PixelFormat = [[NSOpenGLPixelFormat alloc] initWithAttributes:Attribs];
-			if(PixelFormat)
-			{
-				Context = [[NSOpenGLContext alloc] initWithFormat:PixelFormat shareContext:nil];
-				[PixelFormat release];
-				if(Context)
+				else
 				{
-					[Context makeCurrentContext];
-				}
-			}
-		}
-		
-		GLint RendererID = 0;
-		if(Context)
-		{
-			[Context getValues:&RendererID forParameter:NSOpenGLCPCurrentRendererID];
-		}
-		
-		if(bTemporaryContext && Context)
-		{
-			[Context release];
-		}
-		
-		// This is quite a coarse test, alerting users who are running on GPUs which run on older drivers only - this is what we want to tell people running the only barely working OpenGL 3.3 cards.
-		// They can run, but they are largely on their own in terms of rendering performance and bugs as Apple/vendors won't fix OpenGL drivers for them.
-		const bool bOldGpuDriver = ((RendererID & 0x00022000 && RendererID < kCGLRendererGeForceID) || (RendererID & 0x00021000 && RendererID < kCGLRendererATIRadeonX3000ID) || (RendererID & 0x00024000 && RendererID < kCGLRendererIntelHD4000ID));
-		if(bOldGpuDriver)
-		{
-			if(FSlateApplication::IsInitialized())
-			{
-				FSuppressableWarningDialog::FSetupInfo Info( LOCTEXT("UnsupportedGPUWarning_Body","The current graphics card does not meet the minimum specification, for best performance an NVIDIA GeForce 470 GTX or AMD Radeon 6870 HD series card or higher is recommended. Rendering performance and compatibility are not guaranteed with this graphics card."), LOCTEXT("UnsupportedGPUWarning_Title","Unsupported Graphics Card"), TEXT("UnsupportedGPUWarning"), GEditorGameAgnosticIni );
-				Info.ConfirmText = LOCTEXT( "OK", "OK");
-				Info.bDefaultToSupressInTheFuture = true;
-				FSuppressableWarningDialog OSUpdateWarning( Info );
-				OSUpdateWarning.ShowModal();
-			}
-			else
-			{
-				UE_LOG(LogEditor, Warning, TEXT("The current graphics card does not meet the minimum specification, for best performance an NVIDIA GeForce 470 GTX or AMD Radeon 6870 HD series card or higher is recommended. Rendering performance and compatibility are not guaranteed with this graphics card."));
-			}
-		}
-		
-		// The editor is much more demanding than people realise, so warn them when they are beneath the official recommended specs.
-		// Any GPU slower than the Nvidia 470 GTX or AMD 6870 will receive a warning so that it is clear that the performance will not be stellar.
-		// This will affect all MacBooks, MacBook Pros and Mac Minis, none of which have shipped with a powerful GPU.
-		// Low-end or older iMacs and some older Mac Pros will also see this warning, though Mac Pro owners can upgrade their GPUs to something better if they choose.
-		// Newer, high-end iMacs and the new Mac Pros, or old Mac Pros with top-end/aftermarket GPUs shouldn't see this at all.
-
-		bool bSlowGpu = false;
-		// The 6970 and the 5870 were the only Apple supplied cards that are fast enough for the Editor supported by the kCGLRendererATIRadeonX3000ID driver
-		// All the kCGLRendererATIRadeonX4000ID supported cards should be more than sufficiently fast AFAICT
-		if(RendererID & 0x00021000 && RendererID == kCGLRendererATIRadeonX3000ID) // AMD
-		{
-			if(!GRHIAdapterName.Contains(TEXT("6970")) && !GRHIAdapterName.Contains(TEXT("5870")))
-			{
-				bSlowGpu = true;
-			}
-		}
-		else if(RendererID & 0x00022000 && RendererID == kCGLRendererGeForceID) // Nvidia
-		{
-			// Most of the 600 & 700 series Nvidia or later cards Apple use in iMacs should be fast enough
-			// but the ones in MacBook Pros are not.
-			if(GRHIAdapterName.Contains(TEXT("640")) || GRHIAdapterName.Contains(TEXT("650")) || GRHIAdapterName.Contains(TEXT("660")) || GRHIAdapterName.Contains(TEXT("750")) || GRHIAdapterName.Contains(TEXT("755")))
-			{
-				bSlowGpu = true;
-			}
-		}
-		else if(RendererID & 0x00024000 && RendererID == kCGLRendererIntelHD5000ID) // Intel
-		{
-			// Only the Iris Pro is even approaching fast enough - but even that is slower than our recommended specs by a wide margin
-			bSlowGpu = true;
-		}
-		
-		if(bSlowGpu)
-		{
-			if(FSlateApplication::IsInitialized())
-			{
-				FSuppressableWarningDialog::FSetupInfo Info( LOCTEXT("SlowGPUWarning_Body","The current graphics card is slower than the recommanded specification of an NVIDIA GeForce 470 GTX or AMD Radeon 6870 HD series card or higher, performance may be low."), LOCTEXT("SlowGPUWarning_Title","Slow Graphics Card"), TEXT("SlowGPUWarning"), GEditorGameAgnosticIni );
-				Info.ConfirmText = LOCTEXT( "OK", "OK");
-				Info.bDefaultToSupressInTheFuture = true;
-				FSuppressableWarningDialog OSUpdateWarning( Info );
-				OSUpdateWarning.ShowModal();
-			}
-			else
-			{
-				UE_LOG(LogEditor, Warning, TEXT("The current graphics card is slower than the recommanded specification of an NVIDIA GeForce 470 GTX or AMD Radeon 6870 HD series card or higher, performance may be low."));
-			}
-		}
-		
-		// Warn about low-memory configurations as they harm performance in the Editor
-		if(FPlatformMemory::GetPhysicalGBRam() < 8)
-		{
-			if(FSlateApplication::IsInitialized())
-			{
-				FSuppressableWarningDialog::FSetupInfo Info( LOCTEXT("LowRAMWarning_Body","For best performance install at least 8GB of RAM."), LOCTEXT("LowRAMWarning_Title","Low RAM"), TEXT("LowRAMWarning"), GEditorGameAgnosticIni );
-				Info.ConfirmText = LOCTEXT( "OK", "OK");
-				Info.bDefaultToSupressInTheFuture = true;
-				FSuppressableWarningDialog OSUpdateWarning( Info );
-				OSUpdateWarning.ShowModal();
+					FFormatNamedArguments Args;
+					Args.Add( TEXT("ObjectName"), FText::FromString( Obj->GetName() ) );
+					Args.Add( TEXT("ObjectType"), FText::FromString( Obj->GetClass()->GetName() ) );
+					NotificationText = FText::Format( LOCTEXT("ReimportSuccessful", "Successfully Reimported: {ObjectName} ({ObjectType})"), Args );
+				}
 			}
 			else
 			{
-				UE_LOG(LogEditor, Warning, TEXT("For best performance install at least 8GB of RAM."));
-			}
-		}
-		
-		// And also warn about machines with fewer than 4 cores as they will also struggle
-		if(FPlatformMisc::NumberOfCores() < 4)
-		{
-			if(FSlateApplication::IsInitialized())
-			{
-				FSuppressableWarningDialog::FSetupInfo Info( LOCTEXT("SlowCPUWarning_Body","For best performance a Quad-core Intel or AMD processor, 2.5 GHz or faster is recommended."), LOCTEXT("SlowCPUWarning_Title","CPU Performance Warning"), TEXT("SlowCPUWarning"), GEditorGameAgnosticIni );
-				Info.ConfirmText = LOCTEXT( "OK", "OK");
-				Info.bDefaultToSupressInTheFuture = true;
-				FSuppressableWarningDialog OSUpdateWarning( Info );
-				OSUpdateWarning.ShowModal();
+				FFormatNamedArguments Args;
+				Args.Add( TEXT("ObjectName"), FText::FromString( Obj->GetName() ) );
+				Args.Add( TEXT("ObjectType"), FText::FromString( Obj->GetClass()->GetName() ) );
+				NotificationText = FText::Format( LOCTEXT("ReimportFailed", "Failed to Reimport: {ObjectName} ({ObjectType})"), Args );
 			}
-			else
+
+			FNotificationInfo Info(NotificationText);
+			Info.ExpireDuration = 3.0f;
+			Info.bUseLargeFont = false;
+			TSharedPtr<SNotificationItem> Notification = FSlateNotificationManager::Get().AddNotification(Info);
+			if ( Notification.IsValid() )
 			{
-				UE_LOG(LogEditor, Warning, TEXT("For best performance a Quad-core Intel or AMD processor, 2.5 GHz or faster is recommended."));
+				Notification->SetCompletionState( bSuccess ? SNotificationItem::CS_Success : SNotificationItem::CS_Fail );
 			}
 		}
 	}
-#endif
-	
-	// Clean up the auto-load-in-progress file, if it exists. This file prevents auto-loading of projects and must be deleted here to indicate the load was successful
-	const FString AutoLoadInProgressFilename = AutoLoadProjectFileName + TEXT(".InProgress");
-	const bool bRequireExists = false;
-	const bool bEvenIfReadOnly = true;
-	const bool bQuiet = true;
-	IFileManager::Get().Delete(*AutoLoadInProgressFilename, bRequireExists, bEvenIfReadOnly, bQuiet);
-}
 
-void ExecuteInvalidateCachedShaders(const TArray< FString >& Args)
-{
-	if(Args.Num() == 0)
-	{
-		// todo: log error, at least one command is needed
-		UE_LOG(LogConsoleResponse, Display, TEXT("r.InvalidateCachedShaders failed\nAs this command should not be executed accidentally it requires you to specify an extra parameter."));
-		return;
-	}
+	// Let listeners know whether the reimport was successful or not
+	PostReimport.Broadcast( Obj, bSuccess );
+
+	GEditor->RedrawAllViewports();
 
-	FString FileName = FPaths::EngineDir() + TEXT("Shaders/ShaderVersion.usf");
+	return bSuccess;
+}
 
-	FileName = IFileManager::Get().ConvertToAbsolutePathForExternalAppForRead(*FileName);
+void FReimportManager::GetNewReimportPath(UObject* Obj, TArray<FString>& InOutFilenames)
+{
+	TArray<UObject*> ReturnObjects;
+	FString FileTypes;
+	FString AllExtensions;
+	TArray<UFactory*> Factories;
 
-	ISourceControlProvider& SourceControlProvider = ISourceControlModule::Get().GetProvider();
-	SourceControlProvider.Init();
+	// Determine whether we will allow multi select and clear old filenames
+	bool bAllowMultiSelect = InOutFilenames.Num() > 1;
+	InOutFilenames.Empty();
 
-	FSourceControlStatePtr SourceControlState = SourceControlProvider.GetState(FileName, EStateCacheUsage::ForceUpdate);
-	if(SourceControlState.IsValid())
+	// Get the list of valid factories
+	for( TObjectIterator<UClass> It ; It ; ++It )
 	{
-		if( SourceControlState->CanCheckout() || SourceControlState->IsCheckedOutOther() )
+		UClass* CurrentClass = (*It);
+
+		if( CurrentClass->IsChildOf(UFactory::StaticClass()) && !(CurrentClass->HasAnyClassFlags(CLASS_Abstract)) )
 		{
-			if(SourceControlProvider.Execute(ISourceControlOperation::Create<FCheckOut>(), FileName) == ECommandResult::Failed)
+			UFactory* Factory = Cast<UFactory>( CurrentClass->GetDefaultObject() );
+			if( Factory->bEditorImport && Obj->GetClass()->IsChildOf(Factory->GetSupportedClass()) )
 			{
-				UE_LOG(LogConsoleResponse, Display, TEXT("r.InvalidateCachedShaders failed\nCouldn't check out \"ShaderVersion.usf\""));
-				return;
+				Factories.Add( Factory );
 			}
 		}
-		else if(!SourceControlState->IsSourceControlled())
-		{
-			UE_LOG(LogConsoleResponse, Display, TEXT("r.InvalidateCachedShaders failed\n\"ShaderVersion.usf\" is not under source control."));
-		}
-		else if(SourceControlState->IsCheckedOutOther())
-		{
-			UE_LOG(LogConsoleResponse, Display, TEXT("r.InvalidateCachedShaders failed\n\"ShaderVersion.usf\" is already checked out by someone else\n(UE4 SourceControl needs to be fixed to allow multiple checkout.)"));
-			return;
-		}
-		else if(SourceControlState->IsDeleted())
-		{
-			UE_LOG(LogConsoleResponse, Display, TEXT("r.InvalidateCachedShaders failed\n\"ShaderVersion.usf\" is marked for delete"));
-			return;
-		}
 	}
 
-	IPlatformFile& PlatformFile = FPlatformFileManager::Get().GetPlatformFile();
-
-	IFileHandle* FileHandle = PlatformFile.OpenWrite(*FileName);
-	if(FileHandle)
-	{
-		FString Guid = FString(
-			TEXT("// This file is automatically generated by the console command r.InvalidateCachedShaders\n")
-			TEXT("// Each time the console command is executed it generates a new GUID. As this file is included\n")
-			TEXT("// in common.usf (which should be included in any shader) it allows to invalidate the shader DDC.\n")
-			TEXT("// \n")
-			TEXT("// GUID = "))
-			+ FGuid::NewGuid().ToString();
-
-		FileHandle->Write((const uint8*)TCHAR_TO_ANSI(*Guid), Guid.Len());
-		delete FileHandle;
-
-		UE_LOG(LogConsoleResponse, Display, TEXT("r.InvalidateCachedShaders succeeded\n\"ShaderVersion.usf\" was updated.\n"));
-	}
-	else
+	if ( Factories.Num() <= 0 )
 	{
-		UE_LOG(LogConsoleResponse, Display, TEXT("r.InvalidateCachedShaders failed\nCouldn't open \"ShaderVersion.usf\".\n"));
+		// No matching factories for this asset, fail
+		return;
 	}
-}
 
-FAutoConsoleCommand InvalidateCachedShaders(
-	TEXT("r.InvalidateCachedShaders"),
-	TEXT("Invalidate shader cache by making a unique change to ShaderVersion.usf which is included in common.usf.")
-	TEXT("To initiate actual the recompile of all shaders use \"recompileshaders changed\" or press \"Ctrl Shift .\".\n")
-	TEXT("The ShaderVersion.usf file should be automatically checked out but  it needs to be checked in to have effect on other machines."),
-	FConsoleCommandWithArgsDelegate::CreateStatic(ExecuteInvalidateCachedShaders)
-	);
+	TMultiMap<uint32, UFactory*> DummyFilterIndexToFactory;
+
+	// Generate the file types and extensions represented by the selected factories
+	ObjectTools::GenerateFactoryFileExtensions( Factories, FileTypes, AllExtensions, DummyFilterIndexToFactory );
 
-#undef LOCTEXT_NAMESPACE 
+	FileTypes = FString::Printf(TEXT("All Files (%s)|%s|%s"),*AllExtensions,*AllExtensions,*FileTypes);
 
-FORCEINLINE bool NetworkRemapPath_local(FWorldContext &Context, FString &Str, bool reading)
-{
-	if (reading)
+	// Prompt the user for the filenames
+	TArray<FString> OpenFilenames;
+	IDesktopPlatform* DesktopPlatform = FDesktopPlatformModule::Get();
+	bool bOpened = false;
+	if ( DesktopPlatform )
 	{
-		return (Str.ReplaceInline( *Context.PIERemapPrefix, *Context.PIEPrefix, ESearchCase::IgnoreCase) > 0);
+		void* ParentWindowWindowHandle = NULL;
+
+		IMainFrameModule& MainFrameModule = FModuleManager::LoadModuleChecked<IMainFrameModule>(TEXT("MainFrame"));
+		const TSharedPtr<SWindow>& MainFrameParentWindow = MainFrameModule.GetParentWindow();
+		if ( MainFrameParentWindow.IsValid() && MainFrameParentWindow->GetNativeWindow().IsValid() )
+		{
+			ParentWindowWindowHandle = MainFrameParentWindow->GetNativeWindow()->GetOSWindowHandle();
+		}
+
+		const FString Title = FString::Printf(TEXT("%s: %s"), *NSLOCTEXT("ReimportManager", "ImportDialogTitle", "Import For").ToString(), *Obj->GetName());
+		bOpened = DesktopPlatform->OpenFileDialog(
+			ParentWindowWindowHandle,
+			Title,
+			TEXT(""),
+			TEXT(""),
+			FileTypes,
+			bAllowMultiSelect ? EFileDialogFlags::Multiple : EFileDialogFlags::None,
+			OpenFilenames
+			);
 	}
-	else
+
+	if ( bOpened )
 	{
-		return (Str.ReplaceInline( *Context.PIEPrefix, *Context.PIERemapPrefix, ESearchCase::IgnoreCase) > 0);
+		for (int32 FileIndex = 0; FileIndex < OpenFilenames.Num(); ++FileIndex)
+		{
+			InOutFilenames.Add(OpenFilenames[FileIndex]);
+		}
 	}
 }
 
-bool UEditorEngine::NetworkRemapPath( UWorld *InWorld, FString &Str, bool reading)
+FString FReimportManager::SanitizeImportFilename(const FString& InPath, const UObject* Obj)
 {
-	FWorldContext &Context = GetWorldContextFromWorldChecked(InWorld);
-	if (Context.PIEPrefix.IsEmpty() || Context.PIERemapPrefix.IsEmpty())
+	const UPackage* Package = Obj->GetOutermost();
+	if (Package)
 	{
-		return false;
+		const bool		bIncludeDot = true;
+		const FString	PackagePath	= Package->GetPathName();
+		const FName		MountPoint	= FPackageName::GetPackageMountPoint(PackagePath);
+		const FString	PackageFilename = FPackageName::LongPackageNameToFilename(PackagePath, FPaths::GetExtension(InPath, bIncludeDot));
+		const FString	AbsolutePath = FPaths::ConvertRelativePathToFull(InPath);
+
+		if ((MountPoint == FName("Engine") && AbsolutePath.StartsWith(FPaths::ConvertRelativePathToFull(FPaths::EngineContentDir()))) ||
+			(MountPoint == FName("Game") &&	AbsolutePath.StartsWith(FPaths::ConvertRelativePathToFull(FPaths::GameDir()))))
+		{
+			FString RelativePath = InPath;
+			FPaths::MakePathRelativeTo(RelativePath, *PackageFilename);
+			return RelativePath;
+		}
 	}
 
-	return NetworkRemapPath_local(Context, Str, reading);
+	return IFileManager::Get().ConvertToRelativePath(*InPath);
 }
 
-bool UEditorEngine::NetworkRemapPath( UPendingNetGame *PendingNetGame, FString &Str, bool reading)
+
+FString FReimportManager::ResolveImportFilename(const FString& InRelativePath, const UObject* Obj)
 {
-	FWorldContext &Context = GetWorldContextFromPendingNetGameChecked(PendingNetGame);
-	if (Context.PIEPrefix.IsEmpty() || Context.PIERemapPrefix.IsEmpty())
+	FString RelativePath = InRelativePath;
+
+	const UPackage* Package = Obj->GetOutermost();
+	if (Package)
 	{
-		return false;
+		// Relative to the package filename?
+		const FString PathRelativeToPackage = FPaths::GetPath(FPackageName::LongPackageNameToFilename(Package->GetPathName())) / InRelativePath;
+		if (FPaths::FileExists(PathRelativeToPackage))
+		{
+			RelativePath = PathRelativeToPackage;
+		}
 	}
 
-	return NetworkRemapPath_local(Context, Str, reading);
+	// Convert relative paths
+	return FPaths::ConvertRelativePathToFull(RelativePath);
 }
 
-
-void UEditorEngine::VerifyLoadMapWorldCleanup()
+FReimportManager::FReimportManager()
 {
-	// This does the same as UEngine::VerifyLoadMapWorldCleanup except it also allows Editor Worlds as a valid world.
+	// Create reimport handler for textures
+	// NOTE: New factories can be created anywhere, inside or outside of editor
+	// This is just here for convenience
+	UReimportTextureFactory::StaticClass();
 
-	// All worlds at this point should be the CurrentWorld of some context or preview worlds.
-	
-	for( TObjectIterator<UWorld> It; It; ++It )
-	{
-		UWorld* World = *It;
-		if (World->WorldType != EWorldType::Preview && World->WorldType != EWorldType::Editor && World->WorldType != EWorldType::Inactive)
-		{
-			TArray<UWorld*> OtherEditorWorlds;
-			EditorLevelUtils::GetWorlds(EditorWorld, OtherEditorWorlds, true, false);
-			if (OtherEditorWorlds.Contains(World))
-				continue;
+	// Create reimport handler for FBX static meshes
+	UReimportFbxStaticMeshFactory::StaticClass();
 
-			bool ValidWorld = false;
-			for (int32 idx=0; idx < WorldList.Num(); ++idx)
-			{
-				FWorldContext& WorldContext = WorldList[idx];
-				if (World == WorldContext.SeamlessTravelHandler.GetLoadedWorld())
-				{
-					// World valid, but not loaded yet
-					ValidWorld = true;
-					break;
-				}
-				else if (WorldContext.World())
-				{
-					TArray<UWorld*> OtherWorlds;
-					EditorLevelUtils::GetWorlds(WorldContext.World(), OtherWorlds, true, false);
+	// Create reimport handler for FBX skeletal meshes
+	UReimportFbxSkeletalMeshFactory::StaticClass();
 
-					if (OtherWorlds.Contains(World))
-					{
-						// Some other context is referencing this 
-						ValidWorld = true;
-						break;
-					}
-				}
-			}
+	// Create reimport handler for APEX destructible meshes
+	UReimportDestructibleMeshFactory::StaticClass();
 
-			if (!ValidWorld)
-			{
-				// Print some debug information...
-				UE_LOG(LogLoad, Log, TEXT("%s not cleaned up by garbage collection! "), *World->GetFullName());
-				StaticExec(World, *FString::Printf(TEXT("OBJ REFS CLASS=WORLD NAME=%s"), *World->GetPathName()));
-				TMap<UObject*,UProperty*>	Route		= FArchiveTraceRoute::FindShortestRootPath( World, true, GARBAGE_COLLECTION_KEEPFLAGS );
-				FString						ErrorString	= FArchiveTraceRoute::PrintRootPath( Route, World );
-				UE_LOG(LogLoad, Log, TEXT("%s"),*ErrorString);
-				// before asserting.
-				UE_LOG(LogLoad, Fatal, TEXT("%s not cleaned up by garbage collection!") LINE_TERMINATOR TEXT("%s") , *World->GetFullName(), *ErrorString );
-			}
-		}
-	}
+	// Create reimport handler for sound node waves
+	UReimportSoundFactory::StaticClass();
+
+	// Create reimport handler for surround sound waves
+	UReimportSoundSurroundFactory::StaticClass();
 }
 
-void UEditorEngine::TriggerStreamingDataRebuild()
+FReimportManager::~FReimportManager()
 {
-	for (int32 WorldIndex = 0; WorldIndex < WorldList.Num(); ++WorldIndex)
-	{
-		UWorld* World = WorldList[WorldIndex].World();
-		if (World && World->WorldType == EWorldType::Editor)
-		{
-			// Recalculate in a few seconds.
-			World->TriggerStreamingDataRebuild();
-		}
-	}
+	Handlers.Empty();
 }
 
-FWorldContext& UEditorEngine::GetEditorWorldContext(bool bEnsureIsGWorld)
+int32 FReimportHandler::GetPriority() const
 {
-	for (int32 i=0; i < WorldList.Num(); ++i)
-	{
-		if (WorldList[i].WorldType == EWorldType::Editor)
-		{
-			ensure(!bEnsureIsGWorld || WorldList[i].World() == GWorld);
-			return WorldList[i];
-		}
-	}
+	return UFactory::DefaultImportPriority;
+}
+
+/*-----------------------------------------------------------------------------
+	PIE helpers.
+-----------------------------------------------------------------------------*/
+
+/**
+ * Sets GWorld to the passed in PlayWorld and sets a global flag indicating that we are playing
+ * in the Editor.
+ *
+ * @param	PlayInEditorWorld		PlayWorld
+ * @return	the original GWorld
+ */
+UWorld* SetPlayInEditorWorld( UWorld* PlayInEditorWorld )
+{
+	check(!GIsPlayInEditorWorld);
+	UWorld* SavedWorld = GWorld;
+	GIsPlayInEditorWorld = true;
+	GWorld = PlayInEditorWorld;
+
+	return SavedWorld;
+}
 
-	check(false); // There should have already been one created in UEngine::Init
-	return CreateNewWorldContext(EWorldType::Editor);
+/**
+ * Restores GWorld to the passed in one and reset the global flag indicating whether we are a PIE
+ * world or not.
+ *
+ * @param EditorWorld	original world being edited
+ */
+void RestoreEditorWorld( UWorld* EditorWorld )
+{
+	check(GIsPlayInEditorWorld);
+	GIsPlayInEditorWorld = false;
+	GWorld = EditorWorld;
 }
 
-FWorldContext* UEditorEngine::GetPIEWorldContext()
+/**
+ * Takes an FName and checks to see that it is unique among all loaded objects.
+ *
+ * @param	InName		The name to check
+ * @param	Outer		The context for validating this object name. Should be a group/package, but could be ANY_PACKAGE if you want to check across the whole system (not recommended)
+ * @param	InReason	If the check fails, this string is filled in with the reason why.
+ *
+ * @return	1 if the name is valid, 0 if it is not
+ */
+
+bool IsUniqueObjectName( const FName& InName, UObject* Outer, FText* InReason )
 {
-	for(auto& WorldContext : WorldList)
+	// See if the name is already in use.
+	if( StaticFindObject( UObject::StaticClass(), Outer, *InName.ToString() ) != NULL )
 	{
-		if(WorldContext.WorldType == EWorldType::PIE)
+		if ( InReason != NULL )
 		{
-			return &WorldContext;
+			*InReason = NSLOCTEXT("UnrealEd", "NameAlreadyInUse", "Name is already in use by another object.");
 		}
+		return false;
 	}
 
-	return nullptr;
+	return true;
+}
+
+/**
+ * Takes an FName and checks to see that it is unique among all loaded objects.
+ *
+ * @param	InName		The name to check
+ * @param	Outer		The context for validating this object name. Should be a group/package, but could be ANY_PACKAGE if you want to check across the whole system (not recommended)
+ * @param	InReason	If the check fails, this string is filled in with the reason why.
+ *
+ * @return	1 if the name is valid, 0 if it is not
+ */
+
+bool IsUniqueObjectName( const FName& InName, UObject* Outer, FText& InReason )
+{
+	return IsUniqueObjectName(InName,Outer,&InReason);
 }
 
+//////////////////////////////////////////////////////////////////////////
+// EditorUtilities
+
 namespace EditorUtilities
 {
 	AActor* GetEditorWorldCounterpartActor( AActor* Actor )
@@ -7556,38 +1182,8 @@ namespace EditorUtilities
 	}
 }
 
-bool UEditorEngine::IsUsingWorldAssets()
-{
-	return !FParse::Param(FCommandLine::Get(), TEXT("DisableWorldAssets"));
-}
-
-void UEditorEngine::OnAssetLoaded(UObject* Asset)
-{
-	UWorld* WorldAsset = Cast<UWorld>(Asset);
-	if (WorldAsset != NULL)
-	{
-		// Init inactive worlds here instead of UWorld::PostLoad because it is illegal to call UpdateWorldComponents while IsRoutingPostLoad
-		if (WorldAsset->WorldType == EWorldType::Inactive)
-		{
-			// Create the world without a physics scene because creating too many physics scenes causes deadlock issues in PhysX. The scene will be created when it is opened in the level editor.
-			// Also, don't create an FXSystem because it consumes too much video memory. This is also created when the level editor opens this world.
-			WorldAsset->InitWorld(GetEditorWorldInitializationValues()
-				.CreatePhysicsScene(false)
-				.CreateFXSystem(false)
-				);
-
-			// Update components so the scene is populated
-			WorldAsset->UpdateWorldComponents(true, true);
-		}
-	}
-}
-
-UWorld::InitializationValues UEditorEngine::GetEditorWorldInitializationValues() const
-{
-	return UWorld::InitializationValues()
-		.ShouldSimulatePhysics(false)
-		.EnableTraceCollision(true);
-}
+//////////////////////////////////////////////////////////////////////////
+// FCachedActorLabels
 
 FCachedActorLabels::FCachedActorLabels()
 {
