Split the UEditorEngine implementation out from Editor.cpp into EditorEngine.cpp, removed some GWorld references, and purged some dead code

diff --git a/Engine/Source/Editor/UnrealEd/Private/EditorServer.cpp b/Engine/Source/Editor/UnrealEd/Private/EditorServer.cpp
--- a/Engine/Source/Editor/UnrealEd/Private/EditorServer.cpp
+++ b/Engine/Source/Editor/UnrealEd/Private/EditorServer.cpp
@@ -1798,7 +1798,8 @@ void UEditorEngine::EditorDestroyWorld( FWorldContext & Context, const FText& Cl
 
 	GUnrealEd->CurrentLODParentActor = NULL;
 	SelectNone( true, true );
-	EditorClearComponents();
+
+	ContextWorld->ClearWorldComponents();
 	ClearPreviewComponents();
 	// Remove all active groups, they belong to a map being unloaded
 	ContextWorld->ActiveGroupActors.Empty();
@@ -2517,7 +2518,7 @@ bool UEditorEngine::Map_Import( UWorld* InWorld, const TCHAR* Str, FOutputDevice
 		
 		ResetTransaction( LocalizedImportingMap );
 		GWarn->BeginSlowTask( LocalizedImportingMap, true );
-		EditorClearComponents();
+		InWorld->ClearWorldComponents();
 		InWorld->CleanupWorld();
 		ImportObject<UWorld>(InWorld->GetOuter(), InWorld->GetFName(), RF_Transactional, TempFname );
 		GWarn->EndSlowTask();
@@ -5490,10 +5491,6 @@ bool UEditorEngine::Exec( UWorld* InWorld, const TCHAR* Stream, FOutputDevice& A
 	{		
 		bProcessed = HandleCleanBSPMaterialCommand( Str, Ar, InWorld );
 	}
-	else if( FParse::Command(&Str,TEXT("CREATESMFROMBSP")) )
-	{
-		bProcessed = HandleCreateMeshFromBSPCommand( Str, Ar );
-	}
 	else if( FParse::Command(&Str,TEXT("AUTOMERGESM")) )
 	{		
 		bProcessed = HandleAutoMergeStaticMeshCommand( Str, Ar );
@@ -6181,12 +6178,6 @@ bool UEditorEngine::HandleCleanBSPMaterialCommand( const TCHAR* Str, FOutputDevi
 	return true;
 }
 
-bool UEditorEngine::HandleCreateMeshFromBSPCommand( const TCHAR* Str, FOutputDevice& Ar )
-{
-	// TODO_STATICMESH: Remove this command.
-	return true;
-}
-
 bool UEditorEngine::HandleAutoMergeStaticMeshCommand( const TCHAR* Str, FOutputDevice& Ar )
 {
 	AutoMergeStaticMeshes();
@@ -6293,27 +6284,6 @@ bool UEditorEngine::HandleRemoveArchtypeFlagCommand( const TCHAR* Str, FOutputDe
 	return true;
 }
 
-UClass* UEditorEngine::GetClassFromPairMap( FString ClassName )
-{
-	UClass* ClassToUse = FindObject<UClass>(ANY_PACKAGE,*ClassName);
-	for( int32 i = 0; i < ClassMapPair.Num(); i++ )
-	{
-		const FString& Str = ClassMapPair[i];
-		int32 DelimIdx = Str.Find(TEXT("/"), ESearchCase::CaseSensitive);
-		if( DelimIdx >= 0 )
-		{
-			const FString Key = Str.Left(DelimIdx);
-			const FString Value = Str.Mid(DelimIdx+1);
-			if( FCString::Stricmp( *ClassName, *Key ) == 0 )
-			{
-				ClassToUse = FindObject<UClass>(ANY_PACKAGE,*Value);
-				break;
-			}
-		}
-	}	
-	return ClassToUse;
-}
-
 /**
  * @return true if the given component's StaticMesh can be merged with other StaticMeshes
  */
diff --git a/Engine/Source/Editor/UnrealEd/Private/LevelViewportClickHandlers.cpp b/Engine/Source/Editor/UnrealEd/Private/LevelViewportClickHandlers.cpp
--- a/Engine/Source/Editor/UnrealEd/Private/LevelViewportClickHandlers.cpp
+++ b/Engine/Source/Editor/UnrealEd/Private/LevelViewportClickHandlers.cpp
@@ -198,17 +198,6 @@ namespace ClickHandlers
 
 				return true;
 			}
-			else if( Click.GetKey() == EKeys::LeftMouseButton && ViewportClient->Viewport->KeyState(EKeys::Period) )
-			{
-				if(Click.IsControlDown())
-				{
-					// create a pylon
-					UClass* PylonClass = GEditor->GetClassFromPairMap( FString(TEXT("Pylon")) );
-					PrivateAddActor( PylonClass );
-				}
-
-				return true;
-			}
 			else if ( Actor )
 			{
 				if( bAllowSelectionChange )
@@ -831,15 +820,6 @@ namespace ClickHandlers
 			// Create a static mesh.
 			PrivateAddActor( AStaticMeshActor::StaticClass() );
 		}
-		else if( Click.GetKey() == EKeys::LeftMouseButton && ViewportClient->Viewport->KeyState(EKeys::Period) )
-		{
-			if(Click.IsControlDown())
-			{
-				// create a pylon
-				UClass* PylonClass = GEditor->GetClassFromPairMap( FString(TEXT("Pylon")) );
-				PrivateAddActor( PylonClass );
-			}
-		}
 		else if( Click.GetKey() == EKeys::LeftMouseButton && ViewportClient->Viewport->KeyState(EKeys::Semicolon) )
 		{
 			PrivateAddActor( ATargetPoint::StaticClass() );
@@ -1025,16 +1005,6 @@ namespace ClickHandlers
 			// Create a static mesh.
 			PrivateAddActor( AStaticMeshActor::StaticClass() );
 		}
-		else if( Click.GetKey() == EKeys::LeftMouseButton && ViewportClient->Viewport->KeyState(EKeys::Period) )
-		{
-			// Create a pathnode.
-			if(Click.IsControlDown())
-			{
-				// create a pylon
-				UClass* PylonClass = GEditor->GetClassFromPairMap( FString(TEXT("Pylon")) );
-				PrivateAddActor( PylonClass );
-			}
-		}
 		else if( Click.GetKey() == EKeys::RightMouseButton && !Click.IsControlDown() && !ViewportClient->Viewport->KeyState(EKeys::LeftMouseButton) )
 		{
 			// NOTE: We intentionally do not deselect selected actors here even though the user right
diff --git a/Engine/Source/Editor/UnrealEd/Private/SoundPreviewThread.cpp b/Engine/Source/Editor/UnrealEd/Private/SoundPreviewThread.cpp
--- a/Engine/Source/Editor/UnrealEd/Private/SoundPreviewThread.cpp
+++ b/Engine/Source/Editor/UnrealEd/Private/SoundPreviewThread.cpp
@@ -1,40 +1,38 @@
 // Copyright 1998-2015 Epic Games, Inc. All Rights Reserved.
 
-
 #include "UnrealEd.h"
 #include "SoundPreviewThread.h"
 #include "SoundDefinitions.h"
 #include "Sound/SoundWave.h"
+#include "TargetPlatform.h"
+#include "Dialogs/DlgSoundWaveOptions.h"
+
+//////////////////////////////////////////////////////////////////////////
+// FSoundPreviewThread
 
-/**
- *
- * @param PreviewCount	number of sounds to preview
- * @param Node			sound node to perform compression/decompression
- * @param Info			preview info class
- */
-FSoundPreviewThread::FSoundPreviewThread( int32 PreviewCount, USoundWave *Wave, FPreviewInfo *Info ) : 
-	Count( PreviewCount ),
-	SoundWave( Wave ),
-	PreviewInfo( Info ),
-	TaskFinished( false ),
-	CancelCalculations( false )
+FSoundPreviewThread::FSoundPreviewThread(int32 PreviewCount, USoundWave* Wave, FPreviewInfo* Info)
+	: Count(PreviewCount)
+	, SoundWave(Wave)
+	, PreviewInfo(Info)
+	, TaskFinished(false)
+	, CancelCalculations(false)
 {
 }
 
-bool FSoundPreviewThread::Init( void )
+bool FSoundPreviewThread::Init()
 {
 	return true;
 }
 
-uint32 FSoundPreviewThread::Run( void )
+uint32 FSoundPreviewThread::Run()
 {
 	// Get the original wave
 	SoundWave->RemoveAudioResource();
-	SoundWave->InitAudioResource( SoundWave->RawData );
+	SoundWave->InitAudioResource(SoundWave->RawData);
 
-	for( Index = 0; Index < Count && !CancelCalculations; Index++ )
+	for (Index = 0; Index < Count && !CancelCalculations; Index++)
 	{
-		SoundWaveQualityPreview( SoundWave, PreviewInfo + Index );
+		SoundWaveQualityPreview(SoundWave, PreviewInfo + Index);
 	}
 
 	SoundWave->RemoveAudioResource();
@@ -42,12 +40,64 @@ uint32 FSoundPreviewThread::Run( void )
 	return 0;
 }
 
-void FSoundPreviewThread::Stop( void )
+void FSoundPreviewThread::Stop()
 {
 	CancelCalculations = true;
 }
 
-void FSoundPreviewThread::Exit( void )
+void FSoundPreviewThread::Exit()
 {
 }
 
+void FSoundPreviewThread::SoundWaveQualityPreview(USoundWave* SoundWave, FPreviewInfo* PreviewInfo)
+{
+	FWaveModInfo WaveInfo;
+	FSoundQualityInfo QualityInfo = { 0 };
+
+	// Extract the info from the wave header
+	if (!WaveInfo.ReadWaveHeader((uint8*)SoundWave->ResourceData, SoundWave->ResourceSize, 0))
+	{
+		return;
+	}
+
+	SoundWave->NumChannels = *WaveInfo.pChannels;
+	SoundWave->RawPCMDataSize = WaveInfo.SampleDataSize;
+
+	// Extract the stats
+	PreviewInfo->OriginalSize = SoundWave->ResourceSize;
+	PreviewInfo->OggVorbisSize = 0;
+	PreviewInfo->PS3Size = 0;
+	PreviewInfo->XMASize = 0;
+
+	QualityInfo.Quality = PreviewInfo->QualitySetting;
+	QualityInfo.NumChannels = *WaveInfo.pChannels;
+	QualityInfo.SampleRate = SoundWave->SampleRate;
+	QualityInfo.SampleDataSize = WaveInfo.SampleDataSize;
+	QualityInfo.DebugName = SoundWave->GetFullName();
+
+	// PCM -> Vorbis -> PCM 
+	ITargetPlatformManagerModule& TPM = GetTargetPlatformManagerRef();
+	static FName NAME_OGG(TEXT("OGG"));
+	const IAudioFormat* Compressor = TPM.FindAudioFormat(NAME_OGG); // why is this ogg specific?
+	PreviewInfo->OggVorbisSize = 0;
+
+	if (Compressor)
+	{
+		TArray<uint8> Input;
+		Input.AddUninitialized(WaveInfo.SampleDataSize);
+		FMemory::Memcpy(Input.GetData(), WaveInfo.SampleDataStart, Input.Num());
+		TArray<uint8> Output;
+		Compressor->Recompress(NAME_OGG, Input, QualityInfo, Output);
+		if (Output.Num())
+		{
+			PreviewInfo->OggVorbisSize = Output.Num();
+			PreviewInfo->DecompressedOggVorbis = (uint8*)FMemory::Malloc(Output.Num());
+			FMemory::Memcpy(PreviewInfo->DecompressedOggVorbis, Output.GetData(), Output.Num());
+		}
+	}
+
+	if (PreviewInfo->OggVorbisSize <= 0)
+	{
+		UE_LOG(LogAudio, Log, TEXT("PC decompression failed"));
+	}
+}
diff --git a/Engine/Source/Editor/UnrealEd/Private/UnrealEdMisc.cpp b/Engine/Source/Editor/UnrealEd/Private/UnrealEdMisc.cpp
--- a/Engine/Source/Editor/UnrealEd/Private/UnrealEdMisc.cpp
+++ b/Engine/Source/Editor/UnrealEd/Private/UnrealEdMisc.cpp
@@ -921,8 +921,8 @@ void FUnrealEdMisc::CB_MapChange( uint32 InFlags )
 	// Minor things like brush subtraction will set it to "0".
 
 	if( InFlags != MapChangeEventFlags::Default )
-	{	
-		GEditor->EditorClearComponents();
+	{
+		World->ClearWorldComponents();
 
 		// Note: CleanupWorld is being abused here to detach components and some other stuff
 		// CleanupWorld should only be called before destroying the world
diff --git a/Engine/Source/Editor/UnrealEd/Public/Dialogs/DlgSoundWaveOptions.h b/Engine/Source/Editor/UnrealEd/Public/Dialogs/DlgSoundWaveOptions.h
--- a/Engine/Source/Editor/UnrealEd/Public/Dialogs/DlgSoundWaveOptions.h
+++ b/Engine/Source/Editor/UnrealEd/Public/Dialogs/DlgSoundWaveOptions.h
@@ -5,6 +5,34 @@ typedef TSharedPtr< struct FPreviewInfoContainer > FSoundWaveCompressionListItem
 
 typedef SListView< FSoundWaveCompressionListItemPtr > SSoundWaveCompressionListView;
 
+/** 
+ * Info used to setup the rows of the sound quality previewer
+ */
+class FPreviewInfo
+{
+public:
+	FPreviewInfo(int32 Quality);
+
+	~FPreviewInfo() 
+	{ 
+		Cleanup(); 
+	}
+
+	void Cleanup();
+
+	int32 QualitySetting;
+
+	int32 OriginalSize;
+
+	int32 OggVorbisSize;
+	int32 XMASize;
+	int32 PS3Size;
+
+	uint8* DecompressedOggVorbis;
+	uint8* DecompressedXMA;
+	uint8* DecompressedPS3;
+};
+
 class UNREALED_API SSoundWaveCompressionOptions : public SCompoundWidget
 {
 public:
diff --git a/Engine/Source/Editor/UnrealEd/Public/Editor.h b/Engine/Source/Editor/UnrealEd/Public/Editor.h
--- a/Engine/Source/Editor/UnrealEd/Public/Editor.h
+++ b/Engine/Source/Editor/UnrealEd/Public/Editor.h
@@ -383,7 +383,7 @@ const TCHAR* ImportObjectProperties( FImportObjectParams& InParams );
  *								if SubobjectOuter is a subobject, corresponds to the first object in SubobjectOuter's Outer chain that is not a subobject itself.
  *								if SubobjectOuter is not a subobject, should normally be the same value as SubobjectOuter
  * @param	SubobjectOuter		the object corresponding to DestData; this is the object that will used as the outer when creating subobjects from definitions contained in SourceText
- * @param	Warn				ouptut device to use for log messages
+ * @param	Warn				output device to use for log messages
  * @param	Depth				current nesting level
  * @param	LineNumber			used when importing defaults during script compilation for tracking which line the defaultproperties block begins on
  * @param	InstanceGraph		contains the mappings of instanced objects and components to their templates; used when recursively calling ImportObjectProperties; generally
@@ -510,43 +510,6 @@ bool GetREMOVE( const TCHAR** Stream, const TCHAR* Match );
 bool GetSUBSTRING(const TCHAR*	Stream, const TCHAR* Match, TCHAR* Value, int32 MaxLen);
 TCHAR* SetFVECTOR( TCHAR* Dest, const FVector* Value );
 
-
-/** 
- * Info used to setup the rows of the sound quality previewer
- */
-class FPreviewInfo
-{
-public:
-	FPreviewInfo( int32 Quality );
-	~FPreviewInfo( void ) 
-	{ 
-		Cleanup(); 
-	}
-
-	void Cleanup( void );
-
-	int32			QualitySetting;
-
-	int32			OriginalSize;
-
-	int32			OggVorbisSize;
-	int32			XMASize;
-	int32			PS3Size;
-
-	uint8*		DecompressedOggVorbis;
-	uint8*		DecompressedXMA;
-	uint8*		DecompressedPS3;
-};
-
-
-/**
- * Compresses SoundWave for all available platforms, and then decompresses to PCM 
- *
- * @param	SoundWave				Wave file to compress
- * @param	PreviewInfo				Compressed stats and decompressed data
- */
-void SoundWaveQualityPreview( USoundWave* SoundWave, FPreviewInfo * PreviewInfo );
-
 /**
  * Takes an FName and checks to see that it is unique among all loaded objects.
  *
@@ -554,7 +517,7 @@ void SoundWaveQualityPreview( USoundWave* SoundWave, FPreviewInfo * PreviewInfo
  * @param	Outer		The context for validating this object name. Should be a group/package, but could be ANY_PACKAGE if you want to check across the whole system (not recommended)
  * @param	InReason	If the check fails, this string is filled in with the reason why.
  *
- * @return	1 if the name is valid, 0 if it is not
+ * @return	true if the name is valid
  */
 
 UNREALED_API bool IsUniqueObjectName( const FName& InName, UObject* Outer, FText* InReason=NULL );
@@ -566,7 +529,7 @@ UNREALED_API bool IsUniqueObjectName( const FName& InName, UObject* Outer, FText
  * @param	Outer		The context for validating this object name. Should be a group/package, but could be ANY_PACKAGE if you want to check across the whole system (not recommended)
  * @param	InReason	If the check fails, this string is filled in with the reason why.
  *
- * @return	1 if the name is valid, 0 if it is not
+ * @return	true if the name is valid
  */
 
 UNREALED_API bool IsUniqueObjectName( const FName& InName, UObject* Outer, FText& InReason );
diff --git a/Engine/Source/Editor/UnrealEd/Public/SoundPreviewThread.h b/Engine/Source/Editor/UnrealEd/Public/SoundPreviewThread.h
--- a/Engine/Source/Editor/UnrealEd/Public/SoundPreviewThread.h
+++ b/Engine/Source/Editor/UnrealEd/Public/SoundPreviewThread.h
@@ -5,8 +5,9 @@
  *
  */
 
-#ifndef __SOUNDPREVIETHREAD_H__
-#define __SOUNDPREVIETHREAD_H__
+#pragma once
+
+class FPreviewInfo;
 
 class FSoundPreviewThread : public FRunnable
 {
@@ -16,25 +17,27 @@ public:
 	 * @param Node			sound node to perform compression/decompression
 	 * @param Info			preview info class
 	 */
-	FSoundPreviewThread( int32 PreviewCount, USoundWave *Node, FPreviewInfo *Info );
-
-	virtual bool Init( void );
-	virtual uint32 Run( void );
-	virtual void Stop( void );
-	virtual void Exit( void );
+	FSoundPreviewThread(int32 PreviewCount, USoundWave* Node, FPreviewInfo* Info);
+	
+	// FRunnable interface
+	virtual bool Init() override;
+	virtual uint32 Run() override;
+	virtual void Stop() override;
+	virtual void Exit() override;
+	// End of FRunnable interface
 
 	/**
-	 * @return true if job is finished
-	 */
-	bool IsTaskFinished( void ) const 
-	{ 
-		return TaskFinished; 
+	* @return true if job is finished
+	*/
+	bool IsTaskFinished() const
+	{
+		return TaskFinished;
 	}
 
 	/**
 	 * @return number of sounds to conversion.
 	 */
-	int32 GetCount( void ) const 
+	int32 GetCount() const
 	{ 
 		return Count; 
 	}
@@ -42,24 +45,31 @@ public:
 	/**
 	 * @return index of last converted sound.
 	 */
-	int32 GetIndex( void ) const 
+	int32 GetIndex() const 
 	{ 
 		return Index; 
 	}
 
+private:
+	/**
+	 * Compresses SoundWave for all available platforms, and then decompresses to PCM 
+	 *
+	 * @param	SoundWave				Wave file to compress
+	 * @param	PreviewInfo				Compressed stats and decompressed data
+	 */
+	static void SoundWaveQualityPreview( USoundWave* SoundWave, FPreviewInfo* PreviewInfo );
+
 private:
 	/** Number of sound to preview.	*/
-	int32					Count;					
+	int32 Count;					
 	/** Last converted sound index, also using to conversion loop. */
-	int32  Index;					
+	int32 Index;					
 	/** Pointer to sound wave to perform preview. */
-	USoundWave*			SoundWave;
+	USoundWave* SoundWave;
 	/** Pointer to preview info table. */
-	FPreviewInfo*		PreviewInfo; 
+	FPreviewInfo* PreviewInfo; 
 	/** True if job is finished. */
-	bool				TaskFinished;
+	bool TaskFinished;
 	/** True if cancel calculations and stop thread. */
-	bool				CancelCalculations;	
+	bool CancelCalculations;	
 };
-	
-#endif
