Editor: Deleted some dead code in StaticMeshEditor and UnrealEd (lightmap mesh import/export)

diff --git a/Engine/Source/Editor/StaticMeshEditor/Private/StaticMeshEditor.cpp b/Engine/Source/Editor/StaticMeshEditor/Private/StaticMeshEditor.cpp
--- a/Engine/Source/Editor/StaticMeshEditor/Private/StaticMeshEditor.cpp
+++ b/Engine/Source/Editor/StaticMeshEditor/Private/StaticMeshEditor.cpp
@@ -1703,47 +1703,6 @@ void FStaticMeshEditor::DoDecomp(float InAccuracy, int32 InMaxHullVerts)
 	}
 }
 
-void FStaticMeshEditor::OnExportLightmapMesh( bool IsFBX )
-{
-	check(StaticMesh);		//must have a mesh to perform this operation
-
-	TArray <UStaticMesh*> StaticMeshArray;
-	StaticMeshArray.Add(StaticMesh);
-
-	FString Directory;
-	if (PromptUserForDirectory(Directory, *LOCTEXT("StaticMeshEditor_ExportToPromptTitle", "Export to...").ToString(), *FEditorDirectories::Get().GetLastDirectory(ELastDirectory::MESH_IMPORT_EXPORT)))
-	{
-		FEditorDirectories::Get().SetLastDirectory(ELastDirectory::MESH_IMPORT_EXPORT, Directory); // Save path as default for next time.
-
-		bool bAnyErrorOccurred = FbxMeshUtils::ExportAllLightmapModels(StaticMeshArray, Directory, IsFBX);
-		if (bAnyErrorOccurred)
-		{
-			FMessageDialog::Open( EAppMsgType::Ok, LOCTEXT("StaticMeshEditor_LightmapExportFailure", "Some static meshes failed to export or provided warnings.  Please check the Log Window for details."));
-		}
-	}
-}
-
-void FStaticMeshEditor::OnImportLightmapMesh( bool IsFBX )
-{
-	check(StaticMesh);		//must have a mesh to perform this operation
-
-	TArray <UStaticMesh*> StaticMeshArray;
-	StaticMeshArray.Add(StaticMesh);
-
-	FString Directory;
-	if (PromptUserForDirectory(Directory, *LOCTEXT("ImportFrom", "Import from...").ToString(), *FEditorDirectories::Get().GetLastDirectory(ELastDirectory::MESH_IMPORT_EXPORT)))
-	{
-		FEditorDirectories::Get().SetLastDirectory(ELastDirectory::MESH_IMPORT_EXPORT, Directory); // Save path as default for next time.
-
-		bool bAnyErrorOccurred = FbxMeshUtils::ImportAllLightmapModels(StaticMeshArray, Directory, IsFBX);
-		if (bAnyErrorOccurred)
-		{
-			FMessageDialog::Open( EAppMsgType::Ok, LOCTEXT("StaticMeshEditor_LightmapImportFailure", "Some static meshes failed to import or provided warnings.  Please check the Log Window for details."));
-		}
-		Viewport->RefreshViewport();
-	}
-}
-
 TSet< int32 >& FStaticMeshEditor::GetSelectedEdges()
 {
 	return Viewport->GetSelectedEdges();
diff --git a/Engine/Source/Editor/StaticMeshEditor/Private/StaticMeshEditor.h b/Engine/Source/Editor/StaticMeshEditor/Private/StaticMeshEditor.h
--- a/Engine/Source/Editor/StaticMeshEditor/Private/StaticMeshEditor.h
+++ b/Engine/Source/Editor/StaticMeshEditor/Private/StaticMeshEditor.h
@@ -156,12 +156,6 @@ private:
 	/** Callback for creating sphyl collision. */
 	void OnCollisionSphyl();
 
-	/** Event for exporting the light map channel of a static mesh to an intermediate file for Max/Maya */
-	void OnExportLightmapMesh( bool IsFBX );
-
-	/** Event for importing the light map channel to a static mesh from an intermediate file from Max/Maya */
-	void OnImportLightmapMesh( bool IsFBX );
-
 	/**
 	* Quick and dirty way of creating box vertices from a box collision representation
 	* Grossly inefficient, but not time critical
diff --git a/Engine/Source/Editor/UnrealEd/Private/FbxMeshUtils.cpp b/Engine/Source/Editor/UnrealEd/Private/FbxMeshUtils.cpp
--- a/Engine/Source/Editor/UnrealEd/Private/FbxMeshUtils.cpp
+++ b/Engine/Source/Editor/UnrealEd/Private/FbxMeshUtils.cpp
@@ -26,226 +26,6 @@ DEFINE_LOG_CATEGORY_STATIC(LogExportMeshUtils, Log, All);
 
 namespace FbxMeshUtils
 {
-	/*
-	*Writes out all StaticMeshes passed in (and all LoDs) for light map editing
-	* @param InAllMeshes - All Static Mesh objects that need their LOD's lightmaps to be exported
-	* @param InDirectory - The directory to attempt the export to
-	* @return If any errors occured during export
-	*/
-	bool ExportAllLightmapModels (TArray<UStaticMesh*>& InAllMeshes, const FString& InDirectory, bool IsFBX)
-	{
-		bool bAnyErrorOccured = false;
-		for (int32 MeshIndex = 0; MeshIndex < InAllMeshes.Num(); ++MeshIndex)
-		{
-			UStaticMesh* CurrentMesh = InAllMeshes[MeshIndex];
-			check (CurrentMesh);
-
-			if (FMath::IsWithin(CurrentMesh->LightMapCoordinateIndex, 0, (int32)MAX_TEXCOORDS))
-			{
-				for (int32 LODIndex = 0; LODIndex < CurrentMesh->GetNumLODs(); ++LODIndex)
-				{	
-					if (IsFBX)
-					{
-						FString Filename = InDirectory + TEXT("/") + CurrentMesh->GetName() + FString::Printf( TEXT("_UVs_LOD_%d.FBX"), LODIndex);
-
-						bAnyErrorOccured |= ExportSingleLightmapModelFBX(CurrentMesh, Filename, LODIndex);
-					}
-					else
-					{
-						FString Filename = InDirectory + TEXT("/") + CurrentMesh->GetName() + FString::Printf( TEXT("_UVs_LOD_%d.OBJ"), LODIndex);
-						bool bDisableBackup = true;
-						FOutputDevice* Buffer = new FOutputDeviceFile(*Filename, bDisableBackup);
-						if (Buffer)
-						{
-							Buffer->SetSuppressEventTag(true);
-							Buffer->SetAutoEmitLineTerminator(false);
-
-							bAnyErrorOccured |= ExportSingleLightmapModel (CurrentMesh, Buffer, LODIndex);
-
-							Buffer->TearDown();
-							//IFileManager::Get().Move(ExportParams.Filename, *TempFile, 1, 1);  //if we were to rename the file
-
-						} 
-						else 
-						{
-							UE_LOG(LogExportMeshUtils, Display, TEXT("Unable to load static mesh from file '%s'"), *Filename);
-							bAnyErrorOccured = true;
-						}
-					}
-				}
-			}
-			else
-			{
-				bAnyErrorOccured |= true;
-				UE_LOG(LogExportMeshUtils, Warning, TEXT("Lightmap Coordinate is out of bounds for '%s' with a value of %d"), *CurrentMesh->GetName(), CurrentMesh->LightMapCoordinateIndex);
-			}
-		}
-		return bAnyErrorOccured;
-	}
-
-
-	/*
-	*Reads in all StaticMeshes passed in (and all LoDs) for light map editing
-	* @param InAllMeshes - All Static Mesh objects that need their LOD's lightmaps to be imported
-	* @param InDirectory - The directory to attempt the import from
-	* @return If any errors occured during import
-	*/
-	bool ImportAllLightmapModels(TArray<UStaticMesh*>& InAllMeshes, const FString& InDirectory, bool IsFBX)
-	{
-
-		bool bAnyErrorOccured = false;
-		// TODO_STATICMESH: Remove this function.
-		return bAnyErrorOccured;
-	}
-
-
-	/*
-	*Writes out a single LOD of a StaticMesh(LODIndex) for light map editing
-	* @param InStaticMesh - The mesh to be exported
-	* @param InOutputDevice - The output device to write the data to.  Current usage is a FOutputDeviceFile
-	* @param InLODIndex - The LOD of the model we want exported
-	* @return If any errors occured during export
-	*/
-	bool ExportSingleLightmapModel (UStaticMesh* InStaticMesh, FOutputDevice* InOutputDevice, const int32 InLODIndex)
-	{
-		bool bAnyErrorOccured = false;
-
-		check(InStaticMesh);
-		check(InLODIndex < InStaticMesh->GetNumLODs());
-		check(InOutputDevice);
-
-		const int32 LightMapChannelIndex = InStaticMesh->LightMapCoordinateIndex;
-
-		//OBJ File Header
-		InOutputDevice->Log( TEXT("# UnrealEd Lightmap OBJ exporter\r\n") );
-
-
-		// Collect all the data about the mesh
-		FStaticMeshLODResources& RenderData = InStaticMesh->GetLODForExport(InLODIndex);
-
-		//inform the user on export
-		if (RenderData.GetNumTriangles())
-		{
-			const int32 MaxUVS = RenderData.VertexBuffer.GetNumTexCoords();
-
-			//if no light map has been authored yet, OR the light map is out of bounds of the current NumUVS
-			if ((LightMapChannelIndex == 0) || (LightMapChannelIndex >= MaxUVS)) 
-			{
-				UE_LOG(LogExportMeshUtils, Warning, TEXT("Lightmap Coordinate is out of bounds for '%s' with a value of %d.  On import, a new set of textures coordinates will be added."),
-					*InStaticMesh->GetName(), LightMapChannelIndex);
-				bAnyErrorOccured = true;
-			}
-		} 
-		else 
-		{
-			UE_LOG(LogExportMeshUtils, Warning, TEXT("Source data missing from model '%s'."), *InStaticMesh->GetName());
-			bAnyErrorOccured = true;
-		}
-
-		TArray<FVector> Verts;				// The verts in the mesh
-		TArray<FVector2D> UVs;			// Lightmap UVs from channel 1
-
-		FIndexArrayView Indices = RenderData.IndexBuffer.GetArrayView();
-
-		// Collect all the data about the mesh
-		for( int32 tri = 0 ; tri < RenderData.GetNumTriangles() ; tri++ )
-		{
-			uint32 Index1 = Indices[(tri * 3) + 0];
-			uint32 Index2 = Indices[(tri * 3) + 1];
-			uint32 Index3 = Indices[(tri * 3) + 2];
-
-			FVector Vertex1 = RenderData.PositionVertexBuffer.VertexPosition(Index1);
-			FVector Vertex2 = RenderData.PositionVertexBuffer.VertexPosition(Index2);
-			FVector Vertex3 = RenderData.PositionVertexBuffer.VertexPosition(Index3);
-
-			// Vertices
-			Verts.Add( RenderData.PositionVertexBuffer.VertexPosition(Index1) );
-			Verts.Add( RenderData.PositionVertexBuffer.VertexPosition(Index2) );
-			Verts.Add( RenderData.PositionVertexBuffer.VertexPosition(Index3) );
-
-			// Lightmap UVs
-			UVs.Add( RenderData.VertexBuffer.GetVertexUV(Index1, LightMapChannelIndex) );
-			UVs.Add( RenderData.VertexBuffer.GetVertexUV(Index2, LightMapChannelIndex) );
-			UVs.Add( RenderData.VertexBuffer.GetVertexUV(Index3, LightMapChannelIndex) );
-		}
-
-		// Write out the vertex data
-		InOutputDevice->Log( TEXT("\r\n") );
-		for( int32 v = 0 ; v < Verts.Num() ; ++v )
-		{
-			// Transform to Lightwave's coordinate system
-			InOutputDevice->Logf( TEXT("v %f %f %f\r\n"), Verts[v].X, Verts[v].Z, Verts[v].Y );
-		}
-		InOutputDevice->Logf( TEXT("# %d vertices\r\n"), Verts.Num() );
-
-		// Write out the lightmap UV data 
-		InOutputDevice->Log( TEXT("\r\n") );
-		for( int32 uv = 0 ; uv < UVs.Num() ; ++uv )
-		{
-			// Invert the y-coordinate (Lightwave has their bitmaps upside-down from us).
-			InOutputDevice->Logf( TEXT("vt %f %f\r\n"), UVs[uv].X, 1.0f - UVs[uv].Y);
-		}
-		InOutputDevice->Logf( TEXT("# %d uvs\r\n"), UVs.Num() );
-
-		// Write object header
-		InOutputDevice->Log( TEXT("\r\n") );
-		InOutputDevice->Log( TEXT("g UnrealEdObject\r\n") );
-		InOutputDevice->Log( TEXT("\r\n") );
-
-		//NOTE - JB - Do not save out the smoothing groups separately.  If you do, the vertices get swizzled on import.
-		//This is because the vertices are exported in RawTriangle Ordering, so to match again on import the face list has to be perfectly sequential.
-		InOutputDevice->Logf( TEXT("s %i\r\n"), 0);		//only smoothing group
-		for( int32 tri = 0 ; tri < RenderData.GetNumTriangles(); tri++ )
-		{
-			int32 VertexIndex = 1 + tri*3;
-			InOutputDevice->Logf(TEXT("f %d/%d %d/%d %d/%d\r\n"), VertexIndex, VertexIndex, VertexIndex+1, VertexIndex+1, VertexIndex+2, VertexIndex+2 );
-		}
-
-		// Write out footer
-		InOutputDevice->Log( TEXT("\r\n") );
-		InOutputDevice->Log( TEXT("g\r\n") );
-
-		return bAnyErrorOccured;	//no errors occured at this point
-	};
-
-	/**
-	 * Writes out a single LOD of a StaticMesh(LODIndex) for light map editing
-	 * @param InCurrentStaticMesh - The mesh to be exported
-	 * @param InFilename - The name of the file to write the data to.
-	 * @param InLODIndex - The LOD of the model we want exported
-	 * @return If any errors occured during export
-	 */
-	bool ExportSingleLightmapModelFBX (UStaticMesh* InStaticMesh, FString& InFilename, const int32 InLODIndex)
-	{
-		bool bAnyErrorOccured = false;
-
-		check(InStaticMesh);
-		check(InLODIndex < InStaticMesh->GetNumLODs());
-
-		int32 LightMapChannelIndex = InStaticMesh->LightMapCoordinateIndex;
-
-		UnFbx::FFbxExporter* Exporter = UnFbx::FFbxExporter::GetInstance();
-		Exporter->CreateDocument();
-		Exporter->ExportStaticMeshLightMap(InStaticMesh, InLODIndex, LightMapChannelIndex);
-		Exporter->WriteToFile(*InFilename);
-
-
-		return bAnyErrorOccured;
-	}
-
-	static bool ExtractLightmapUVsFromMesh( UStaticMesh* InDestStaticMesh, UStaticMesh* MeshToExtractFrom, int32 InLODIndex )
-	{
-		// TODO_STATICMESH: Remove this function.
-		return false;
-	}
-
-
-	bool ImportSingleLightmapModelFBX(UStaticMesh* InStaticMesh, FString& Filename, const int32 InLODIndex )
-	{
-		// TODO_STATICMESH: Remove this function.
-		return false;
-	}
-
 	/** Helper function used for retrieving data required for importing static mesh LODs */
 	void PopulateFBXStaticMeshLODList(UnFbx::FFbxImporter* FFbxImporter, FbxNode* Node, TArray< TArray<FbxNode*>* >& LODNodeList, int32& MaxLODCount, bool bUseLODs)
 	{
@@ -381,7 +161,7 @@ namespace FbxMeshUtils
 				delete LODNodeList[i];
 			}
 		}
-	FFbxImporter->ReleaseScene();
+		FFbxImporter->ReleaseScene();
 	}
 
 	void ImportSkeletalMeshLOD( class USkeletalMesh* SelectedSkelMesh, const FString& Filename, int32 LODLevel )
diff --git a/Engine/Source/Editor/UnrealEd/Public/FbxMeshUtils.h b/Engine/Source/Editor/UnrealEd/Public/FbxMeshUtils.h
--- a/Engine/Source/Editor/UnrealEd/Public/FbxMeshUtils.h
+++ b/Engine/Source/Editor/UnrealEd/Public/FbxMeshUtils.h
@@ -4,50 +4,6 @@
 
 namespace FbxMeshUtils
 {
-	/**
-	 * Writes out all StaticMeshes passed in (and all LoDs) for light map editing
-	 * @param InAllMeshes - All Static Mesh objects that need their LOD's lightmaps to be exported
-	 * @param InDirectory - The directory to attempt the export to
-	 * @return If any errors occured during export
-	 */
-	UNREALED_API bool ExportAllLightmapModels (TArray<UStaticMesh*>& AllMeshes, const FString& Directory, bool IsFBX);
-
-	/**
-	 * Reads in all StaticMeshes passed in (and all LoDs) for light map editing
-	 * @param InAllMeshes - All Static Mesh objects that need their LOD's lightmaps to be imported
-	 * @param InDirectory - The directory to attempt the import from
-	 * @return If any errors occured during import
-	 */
-	UNREALED_API bool ImportAllLightmapModels (TArray<UStaticMesh*>& AllMeshes, const FString& Directory, bool IsFBX);
-
-	/**
-	 * Writes out a single LOD of a StaticMesh(LODIndex) for light map editing
-	 * @param InStaticMesh - The mesh to be exported
-	 * @param InOutputDevice - The output device to write the data to.  Current usage is a FOutputDeviceFile
-	 * @param InLODIndex - The LOD of the model we want exported
-	 * @return If any errors occured during export
-	 */
-	bool ExportSingleLightmapModel (UStaticMesh* CurrentStaticMesh, FOutputDevice* Out, const int32 LODIndex);
-
-	/**
-	 * Writes out a single LOD of a StaticMesh(LODIndex) for light map editing
-	 * @param InCurrentStaticMesh - The mesh to be exported
-	 * @param InFilename - The name of the file to write the data to.
-	 * @param InLODIndex - The LOD of the model we want exported
-	 * @return If any errors occured during export
-	 */
-	bool ExportSingleLightmapModelFBX (UStaticMesh* CurrentStaticMesh, FString& Filename, const int32 LODIndex);
-
-
-	/**
-	 * Reads in a temp FBX formatted mesh who's texture coordinates will REPLACE that of the active static mesh's LightMapCoordinates
-	 * @param InStaticMesh - The mesh to be inported
-	 * @param InData- An array of bytes that contains the files contents
-	 * @param InLODIndex - The LOD of the model we want imported
-	 * @return If any errors occured during import
-	 */
-	bool ImportSingleLightmapModelFBX(UStaticMesh* CurrentStaticMesh, FString& Filename, const int32 LODIndex );
-
 	/**
 	 * Imports a mesh LOD to the given static mesh
 	 *
@@ -57,7 +13,7 @@ namespace FbxMeshUtils
 	 */
 	UNREALED_API void ImportStaticMeshLOD( UStaticMesh* BaseStaticMesh, const FString& Filename, int32 LODLevel );
 
-		/**
+	/**
 	 * Imports a skeletal mesh LOD to the given skeletal mesh
 	 *
 	 * @param Mesh				The skeletal mesh to import the LOD to
