Checking in a missing file for change 2515154 (Editor Modes: Added support for displaying properties as widgets for selected components)

diff --git a/Engine/Source/Editor/UnrealEd/Private/EdMode.cpp b/Engine/Source/Editor/UnrealEd/Private/EdMode.cpp
--- a/Engine/Source/Editor/UnrealEd/Private/EdMode.cpp
+++ b/Engine/Source/Editor/UnrealEd/Private/EdMode.cpp
@@ -143,11 +143,10 @@ namespace
 	 * Returns the value of the property with the given name in the given Actor instance.
 	 */
 	template<typename T>
-	T GetPropertyValueByName(AActor* Actor, FString PropertyName, int32 PropertyIndex)
+	T GetPropertyValueByName(UObject* Object, FString PropertyName, int32 PropertyIndex)
 	{
 		T Value;
-		T* ValuePtr = GetPropertyValuePtrByName<T>(Actor->GetClass(), Actor, PropertyName, PropertyIndex);
-		if(ValuePtr)
+		if (T* ValuePtr = GetPropertyValuePtrByName<T>(Object->GetClass(), Object, PropertyName, PropertyIndex))
 		{
 			Value = *ValuePtr;
 		}
@@ -158,10 +157,9 @@ namespace
 	 * Sets the property with the given name in the given Actor instance to the given value.
 	 */
 	template<typename T>
-	void SetPropertyValueByName(AActor* Actor, FString PropertyName, int32 PropertyIndex, const T& InValue)
+	void SetPropertyValueByName(UObject* Object, FString PropertyName, int32 PropertyIndex, const T& InValue)
 	{
-		T* ValuePtr = GetPropertyValuePtrByName<T>(Actor->GetClass(), Actor, PropertyName, PropertyIndex);
-		if(ValuePtr)
+		if (T* ValuePtr = GetPropertyValuePtrByName<T>(Object->GetClass(), Object, PropertyName, PropertyIndex))
 		{
 			*ValuePtr = InValue;
 		}
@@ -171,6 +169,35 @@ namespace
 //////////////////////////////////
 // FEdMode
 
+void FEdMode::FPropertyWidgetInfo::GetTransformAndColor(UObject* BestSelectedItem, bool bIsSelected, FTransform& OutLocalTransform, FString& OutValidationMessage, FColor& OutDrawColor) const
+{
+	// Determine the desired position
+	if (bIsTransform)
+	{
+		OutLocalTransform = GetPropertyValueByName<FTransform>(BestSelectedItem, PropertyName, PropertyIndex);
+	}
+	else
+	{
+		OutLocalTransform = FTransform(GetPropertyValueByName<FVector>(BestSelectedItem, PropertyName, PropertyIndex));
+	}
+
+	// Determine the desired color
+	OutDrawColor = bIsSelected ? FColor::White : FColor(128, 128, 255);
+	if (PropertyValidationName != NAME_None)
+	{
+		if (UFunction* ValidateFunc = BestSelectedItem->FindFunction(PropertyValidationName))
+		{
+			BestSelectedItem->ProcessEvent(ValidateFunc, &OutValidationMessage);
+
+			// if we have a negative result, the widget color is red.
+			OutDrawColor = OutValidationMessage.IsEmpty() ? OutDrawColor : FColor::Red;
+		}
+	}
+}
+
+//////////////////////////////////
+// FEdMode
+
 const FName FEdMode::MD_MakeEditWidget(TEXT("MakeEditWidget"));
 const FName FEdMode::MD_ValidateWidgetUsing(TEXT("ValidateWidgetUsing"));
 
@@ -303,94 +330,88 @@ bool FEdMode::InputAxis(FEditorViewportClient* InViewportClient, FViewport* View
 
 bool FEdMode::InputDelta(FEditorViewportClient* InViewportClient, FViewport* InViewport, FVector& InDrag, FRotator& InRot, FVector& InScale)
 {	
-	if(UsesPropertyWidgets())
+	if (UsesPropertyWidgets())
 	{
-		AActor* SelectedActor = GetFirstSelectedActorInstance();
-		if(SelectedActor != NULL && InViewportClient->GetCurrentWidgetAxis() != EAxisList::None)
+		FTransform DisplayWidgetToWorld;
+		UObject* BestSelectedItem = GetItemToTryDisplayingWidgetsFor(/*out*/ DisplayWidgetToWorld);
+
+		if ((BestSelectedItem != nullptr) && (InViewportClient->GetCurrentWidgetAxis() != EAxisList::None))
 		{
 			GEditor->NoteActorMovement();
 
-			if (EditedPropertyName != TEXT(""))
+			if (!EditedPropertyName.IsEmpty())
 			{
 				FTransform LocalTM = FTransform::Identity;
 
-				if(bEditedPropertyIsTransform)
+				if (bEditedPropertyIsTransform)
 				{
-					LocalTM = GetPropertyValueByName<FTransform>(SelectedActor, EditedPropertyName, EditedPropertyIndex);
+					LocalTM = GetPropertyValueByName<FTransform>(BestSelectedItem, EditedPropertyName, EditedPropertyIndex);
 				}
 				else
 				{					
-					FVector LocalPos = GetPropertyValueByName<FVector>(SelectedActor, EditedPropertyName, EditedPropertyIndex);
+					FVector LocalPos = GetPropertyValueByName<FVector>(BestSelectedItem, EditedPropertyName, EditedPropertyIndex);
 					LocalTM = FTransform(LocalPos);
 				}
 
-				// Get actor transform (actor to world)
-				FTransform ActorTM = SelectedActor->ActorToWorld();
 				// Calculate world transform
-				FTransform WorldTM = LocalTM * ActorTM;
+				FTransform WorldTM = LocalTM * DisplayWidgetToWorld;
 				// Calc delta specified by drag
 				//FTransform DeltaTM(InRot.Quaternion(), InDrag);
 				// Apply delta in world space
 				WorldTM.SetTranslation(WorldTM.GetTranslation() + InDrag);
 				WorldTM.SetRotation(InRot.Quaternion() * WorldTM.GetRotation());
 				// Convert new world transform back into local space
-				LocalTM = WorldTM.GetRelativeTransform(ActorTM);
+				LocalTM = WorldTM.GetRelativeTransform(DisplayWidgetToWorld);
 				// Apply delta scale
 				LocalTM.SetScale3D(LocalTM.GetScale3D() + InScale);
 
-				SelectedActor->PreEditChange(NULL);
+				BestSelectedItem->PreEditChange(NULL);
 
-				if(bEditedPropertyIsTransform)
+				if (bEditedPropertyIsTransform)
 				{
-					SetPropertyValueByName<FTransform>(SelectedActor, EditedPropertyName, EditedPropertyIndex, LocalTM);
+					SetPropertyValueByName<FTransform>(BestSelectedItem, EditedPropertyName, EditedPropertyIndex, LocalTM);
 				}
 				else
 				{
-					SetPropertyValueByName<FVector>(SelectedActor, EditedPropertyName, EditedPropertyIndex, LocalTM.GetLocation());
+					SetPropertyValueByName<FVector>(BestSelectedItem, EditedPropertyName, EditedPropertyIndex, LocalTM.GetLocation());
 				}
 
-				SelectedActor->PostEditChange();
+				BestSelectedItem->PostEditChange();
 
 				return true;
 			}
 		}
 	}
 
-	if( GetCurrentTool() )
+	if (GetCurrentTool())
 	{
 		return GetCurrentTool()->InputDelta(InViewportClient,InViewport,InDrag,InRot,InScale);
 	}
 
-	return 0;
+	return false;
 }
 
-/**
- * Lets each tool determine if it wants to use the editor widget or not.  If the tool doesn't want to use it, it will be
- * fed raw mouse delta information (not snapped or altered in any way).
- */
-
 bool FEdMode::UsesTransformWidget() const
 {
-	if( GetCurrentTool() )
+	if (GetCurrentTool())
 	{
 		return GetCurrentTool()->UseWidget();
 	}
 
-	return 1;
+	return true;
 }
 
-/**
- * Lets each mode selectively exclude certain widget types.
- */
 bool FEdMode::UsesTransformWidget(FWidget::EWidgetMode CheckMode) const
 {
-	if(UsesPropertyWidgets())
+	if (UsesPropertyWidgets())
 	{
-		AActor* SelectedActor = GetFirstSelectedActorInstance();
-		if(SelectedActor != NULL)
+		FTransform DisplayWidgetToWorld;
+		UObject* BestSelectedItem = GetItemToTryDisplayingWidgetsFor(/*out*/ DisplayWidgetToWorld);
+
+		if (BestSelectedItem != nullptr)
 		{
 			// If editing a vector (not a transform)
-			if(EditedPropertyName != TEXT("") && !bEditedPropertyIsTransform)
+			if (!EditedPropertyName.IsEmpty() && !bEditedPropertyIsTransform)
 			{
 				return (CheckMode == FWidget::WM_Translate);
 			}
@@ -400,39 +421,34 @@ bool FEdMode::UsesTransformWidget(FWidget::EWidgetMode CheckMode) const
 	return true;
 }
 
-/**
- * Allows each mode/tool to determine a good location for the widget to be drawn at.
- */
-
 FVector FEdMode::GetWidgetLocation() const
 {
-	if(UsesPropertyWidgets())
+	if (UsesPropertyWidgets())
 	{
-		AActor* SelectedActor = GetFirstSelectedActorInstance();
-		if(SelectedActor != NULL)
+		FTransform DisplayWidgetToWorld;
+		UObject* BestSelectedItem = GetItemToTryDisplayingWidgetsFor(/*out*/ DisplayWidgetToWorld);
+		if (BestSelectedItem)
 		{
-			if(EditedPropertyName != TEXT(""))
+			if (!EditedPropertyName.IsEmpty())
 			{
 				FVector LocalPos = FVector::ZeroVector;
 
-				if(bEditedPropertyIsTransform)
+				if (bEditedPropertyIsTransform)
 				{
-					FTransform LocalTM = GetPropertyValueByName<FTransform>(SelectedActor, EditedPropertyName, EditedPropertyIndex);
+					FTransform LocalTM = GetPropertyValueByName<FTransform>(BestSelectedItem, EditedPropertyName, EditedPropertyIndex);
 					LocalPos = LocalTM.GetLocation();
 				}
 				else
 				{
-					LocalPos = GetPropertyValueByName<FVector>(SelectedActor, EditedPropertyName, EditedPropertyIndex);
+					LocalPos = GetPropertyValueByName<FVector>(BestSelectedItem, EditedPropertyName, EditedPropertyIndex);
 				}
 
-				FTransform ActorToWorld = SelectedActor->ActorToWorld();
-				FVector WorldPos = ActorToWorld.TransformPosition(LocalPos);
+				const FVector WorldPos = DisplayWidgetToWorld.TransformPosition(LocalPos);
 				return WorldPos;
 			}
 		}
 	}
 
-	//UE_LOG(LogEditorModes, Log, TEXT("In FEdMode::GetWidgetLocation"));
 	return Owner->PivotLocation;
 }
 
@@ -440,20 +456,21 @@ bool FEdMode::GetCustomDrawingCoordinateSystem(FMatrix& InMatrix, void* InData)
 {
 	if (UsesPropertyWidgets())
 	{
-		AActor* SelectedActor = GetFirstSelectedActorInstance();
-		if (SelectedActor != NULL)
+		FTransform DisplayWidgetToWorld;
+		UObject* BestSelectedItem = GetItemToTryDisplayingWidgetsFor(/*out*/ DisplayWidgetToWorld);
+		if (BestSelectedItem)
 		{
 			if (EditedPropertyName != TEXT(""))
 			{
 				if (bEditedPropertyIsTransform)
 				{
-					FTransform LocalTM = GetPropertyValueByName<FTransform>(SelectedActor, EditedPropertyName, EditedPropertyIndex);
-					InMatrix = FRotationMatrix::Make((LocalTM * SelectedActor->GetTransform()).GetRotation());
+					FTransform LocalTM = GetPropertyValueByName<FTransform>(BestSelectedItem, EditedPropertyName, EditedPropertyIndex);
+					InMatrix = FRotationMatrix::Make((LocalTM * DisplayWidgetToWorld).GetRotation());
 					return true;
 				}
 				else
 				{
-					InMatrix = FRotationMatrix::Make(SelectedActor->GetTransform().GetRotation());
+					InMatrix = FRotationMatrix::Make(DisplayWidgetToWorld.GetRotation());
 					return true;
 				}
 			}
@@ -463,17 +480,17 @@ bool FEdMode::GetCustomDrawingCoordinateSystem(FMatrix& InMatrix, void* InData)
 	return false;
 }
 
-/**
- * Lets the mode determine if it wants to draw the widget or not.
- */
-
 bool FEdMode::ShouldDrawWidget() const
 {
 	bool bDrawWidget = false;
-	if (GEditor->GetSelectedComponentCount() > 0) // when components are selected, only show the widget when one or more are scene components
+
+	bool bHadSelectableComponents = false;
+	if (Owner->GetSelectedComponents()->Num() > 0)
 	{
-		for (FSelectedEditableComponentIterator It(GEditor->GetSelectedEditableComponentIterator()); It; ++It)
+		// when components are selected, only show the widget when one or more are scene components
+		for (FSelectedEditableComponentIterator It(*Owner->GetSelectedComponents()); It; ++It)
 		{
+			bHadSelectableComponents = true;
 			if (It->IsA<USceneComponent>())
 			{
 				bDrawWidget = true;
@@ -481,14 +498,15 @@ bool FEdMode::ShouldDrawWidget() const
 			}
 		}
 	}
-	else // when actors are selected, only show the widget when all selected actors have scene components
-	{
-		bDrawWidget = GEditor->GetSelectedActorCount() > 0;
 
-		for (FSelectionIterator It(GEditor->GetSelectedActorIterator()); It; ++It)
+	if (!bHadSelectableComponents)
+	{
+		// when actors are selected, only show the widget when all selected actors have scene components
+		bDrawWidget = Owner->GetSelectedActors()->Num() > 0;
+		for (FSelectionIterator It(*Owner->GetSelectedActors()); It; ++It)
 		{
 			AActor* Actor = Cast<AActor>(*It);
-			if (Actor == nullptr || Actor->FindComponentByClass<USceneComponent>() == nullptr)
+			if ((Actor == nullptr) || (Actor->FindComponentByClass<USceneComponent>() == nullptr))
 			{
 				bDrawWidget = false;
 				break;
@@ -499,25 +517,11 @@ bool FEdMode::ShouldDrawWidget() const
 	return bDrawWidget;
 }
 
-/**
- * Allows each mode to customize the axis pieces of the widget they want drawn.
- *
- * @param	InwidgetMode	The current widget mode
- *
- * @return	A bitfield comprised of AXIS_ values
- */
-
 EAxisList::Type FEdMode::GetWidgetAxisToDraw( FWidget::EWidgetMode InWidgetMode ) const
 {
 	return EAxisList::All;
 }
 
-/**
- * Lets each mode/tool handle box selection in its own way.
- *
- * @param	InBox	The selection box to use, in worldspace coordinates.
- * @return		true if something was selected/deselected, false otherwise.
- */
 bool FEdMode::BoxSelect( FBox& InBox, bool InSelect )
 {
 	bool bResult = false;
@@ -528,12 +532,6 @@ bool FEdMode::BoxSelect( FBox& InBox, bool InSelect )
 	return bResult;
 }
 
-/**
- * Lets each mode/tool handle frustum selection in its own way.
- *
- * @param	InFrustum	The selection frustum to use, in worldspace coordinates.
- * @return	true if something was selected/deselected, false otherwise.
- */
 bool FEdMode::FrustumSelect( const FConvexVolume& InFrustum, bool InSelect )
 {
 	bool bResult = false;
@@ -569,7 +567,7 @@ void FEdMode::ActorSelectionChangeNotify()
 
 bool FEdMode::HandleClick(FEditorViewportClient* InViewportClient, HHitProxy *HitProxy, const FViewportClick &Click)
 {
-	if(UsesPropertyWidgets() && HitProxy)
+	if (UsesPropertyWidgets() && (HitProxy != nullptr))
 	{
 		if( HitProxy->IsA(HPropertyWidgetProxy::StaticGetType()) )
 		{
@@ -582,7 +580,7 @@ bool FEdMode::HandleClick(FEditorViewportClient* InViewportClient, HHitProxy *Hi
 		// Left clicking on an actor, stop editing a property
 		else if( HitProxy->IsA(HActor::StaticGetType()) )
 		{
-			EditedPropertyName = TEXT("");
+			EditedPropertyName = FString();
 			EditedPropertyIndex = INDEX_NONE;
 			bEditedPropertyIsTransform = false;
 		}
@@ -648,64 +646,41 @@ FModeTool* FEdMode::FindTool( EModeTools InID )
 void FEdMode::Render(const FSceneView* View,FViewport* Viewport,FPrimitiveDrawInterface* PDI)
 {
 	// Let the current mode tool render if it wants to
-	FModeTool* tool = GetCurrentTool();
-	if( tool )
+	if (FModeTool* Tool = GetCurrentTool())
 	{
-		tool->Render( View, Viewport, PDI );
+		Tool->Render( View, Viewport, PDI );
 	}
 
-	if(UsesPropertyWidgets())
+	if (UsesPropertyWidgets())
 	{
-		bool bHitTesting = PDI->IsHitTesting();
-		AActor* SelectedActor = GetFirstSelectedActorInstance();
-		if (SelectedActor != NULL)
+		const bool bHitTesting = PDI->IsHitTesting();
+
+		FTransform DisplayWidgetToWorld;
+		UObject* BestSelectedItem = GetItemToTryDisplayingWidgetsFor(/*out*/ DisplayWidgetToWorld);
+		
+		if (BestSelectedItem != nullptr)
 		{
-			UClass* Class = SelectedActor->GetClass();
+			UClass* Class = BestSelectedItem->GetClass();
 			TArray<FPropertyWidgetInfo> WidgetInfos;
-			GetPropertyWidgetInfos(Class, SelectedActor, WidgetInfos);
+			GetPropertyWidgetInfos(Class, BestSelectedItem, WidgetInfos);
 			FEditorScriptExecutionGuard ScriptGuard;
-			for(int32 i=0; i<WidgetInfos.Num(); i++)
+			for (const FPropertyWidgetInfo& WidgetInfo : WidgetInfos)
 			{
-				FString WidgetName = WidgetInfos[i].PropertyName;
-				FName WidgetValidator = WidgetInfos[i].PropertyValidationName;
-				int32 WidgetIndex = WidgetInfos[i].PropertyIndex;
-				bool bIsTransform = WidgetInfos[i].bIsTransform;
-
-				bool bSelected = (WidgetName == EditedPropertyName) && (WidgetIndex == EditedPropertyIndex);
-				FColor WidgetColor = bSelected ? FColor::White : FColor(128, 128, 255);
-
-				FVector LocalPos = FVector::ZeroVector;
-				if(bIsTransform)
-				{
-					FTransform LocalTM = GetPropertyValueByName<FTransform>(SelectedActor, WidgetName, WidgetIndex);
-					LocalPos = LocalTM.GetLocation();
-				}
-				else
-				{
-					LocalPos = GetPropertyValueByName<FVector>(SelectedActor, WidgetName, WidgetIndex);
-				}
-
-				FTransform ActorToWorld = SelectedActor->ActorToWorld();
-				FVector WorldPos = ActorToWorld.TransformPosition(LocalPos);
-
-				UFunction* ValidateFunc= NULL;
-				if(WidgetValidator != NAME_None && 
-					(ValidateFunc = SelectedActor->FindFunction(WidgetValidator)) != NULL)
-				{
-					FString ReturnText;
-					SelectedActor->ProcessEvent(ValidateFunc, &ReturnText);
+				const bool bSelected = (WidgetInfo.PropertyName == EditedPropertyName) && (WidgetInfo.PropertyIndex == EditedPropertyIndex);
 
-					//if we have a negative result, the widget color is red.
-					WidgetColor = ReturnText.IsEmpty() ? WidgetColor : FColor::Red;
-				}
+				FTransform LocalWidgetTransform;
+				FString ValidationMessage;
+				FColor WidgetColor;
+				WidgetInfo.GetTransformAndColor(BestSelectedItem, bSelected, /*out*/ LocalWidgetTransform, /*out*/ ValidationMessage, /*out*/ WidgetColor);
 
-				FTranslationMatrix WidgetTM(WorldPos);
+				const FTransform WorldWidgetTransform = LocalWidgetTransform * DisplayWidgetToWorld;
+				const FMatrix WidgetTM(WorldWidgetTransform.ToMatrixWithScale());
 
 				const float WidgetSize = 0.035f;
 				const float ZoomFactor = FMath::Min<float>(View->ViewMatrices.ProjMatrix.M[0][0], View->ViewMatrices.ProjMatrix.M[1][1]);
-				const float WidgetRadius = View->Project(WorldPos).W * (WidgetSize / ZoomFactor);
+				const float WidgetRadius = View->Project(WorldWidgetTransform.GetTranslation()).W * (WidgetSize / ZoomFactor);
 
-				if(bHitTesting) PDI->SetHitProxy( new HPropertyWidgetProxy(WidgetName, WidgetIndex, bIsTransform) );
+				if (bHitTesting) PDI->SetHitProxy(new HPropertyWidgetProxy(WidgetInfo.PropertyName, WidgetInfo.PropertyIndex, WidgetInfo.bIsTransform));
 				DrawWireDiamond(PDI, WidgetTM, WidgetRadius, WidgetColor, SDPG_Foreground );
 				if(bHitTesting) PDI->SetHitProxy( NULL );
 			}
@@ -817,60 +792,41 @@ void FEdMode::DrawHUD(FEditorViewportClient* ViewportClient,FViewport* Viewport,
 		}
 	}
 
-	if(UsesPropertyWidgets())
+	if (UsesPropertyWidgets())
 	{
-		AActor* SelectedActor = GetFirstSelectedActorInstance();
-		if (SelectedActor != NULL)
+		FTransform DisplayWidgetToWorld;
+		UObject* BestSelectedItem = GetItemToTryDisplayingWidgetsFor(/*out*/ DisplayWidgetToWorld);
+		if (BestSelectedItem != nullptr)
 		{
 			FEditorScriptExecutionGuard ScriptGuard;
 
 			const int32 HalfX = 0.5f * Viewport->GetSizeXY().X;
 			const int32 HalfY = 0.5f * Viewport->GetSizeXY().Y;
 
-			UClass* Class = SelectedActor->GetClass();		
+			UClass* Class = BestSelectedItem->GetClass();		
 			TArray<FPropertyWidgetInfo> WidgetInfos;
-			GetPropertyWidgetInfos(Class, SelectedActor, WidgetInfos);
-			for(int32 i=0; i<WidgetInfos.Num(); i++)
+			GetPropertyWidgetInfos(Class, BestSelectedItem, WidgetInfos);
+			for (const FPropertyWidgetInfo& WidgetInfo : WidgetInfos)
 			{
-				FString WidgetName = WidgetInfos[i].PropertyName;
-				FName WidgetValidator = WidgetInfos[i].PropertyValidationName;
-				int32 WidgetIndex = WidgetInfos[i].PropertyIndex;
-				bool bIsTransform = WidgetInfos[i].bIsTransform;
+				FTransform LocalWidgetTransform;
+				FString ValidationMessage;
+				FColor IgnoredWidgetColor;
+				WidgetInfo.GetTransformAndColor(BestSelectedItem, /*bSelected=*/ false, /*out*/ LocalWidgetTransform, /*out*/ ValidationMessage, /*out*/ IgnoredWidgetColor);
 
-				FVector LocalPos = FVector::ZeroVector;
-				if(bIsTransform)
-				{
-					FTransform LocalTM = GetPropertyValueByName<FTransform>(SelectedActor, WidgetName, WidgetIndex);
-					LocalPos = LocalTM.GetLocation();
-				}
-				else
-				{
-					LocalPos = GetPropertyValueByName<FVector>(SelectedActor, WidgetName, WidgetIndex);
-				}
+				const FTransform WorldWidgetTransform = LocalWidgetTransform * DisplayWidgetToWorld;
 
-				FTransform ActorToWorld = SelectedActor->ActorToWorld();
-				FVector WorldPos = ActorToWorld.TransformPosition(LocalPos);
-
-				UFunction* ValidateFunc = NULL;
-				FString FinalString;
-				if(WidgetValidator != NAME_None && 
-					(ValidateFunc = SelectedActor->FindFunction(WidgetValidator)) != NULL)
+				const FPlane Proj = View->Project(WorldWidgetTransform.GetTranslation());
+				if (Proj.W > 0.f)
 				{
-					SelectedActor->ProcessEvent(ValidateFunc, &FinalString);
-				}
-
-				const FPlane Proj = View->Project( WorldPos );
+					// do some string fixing
+					const uint32 VectorIndex = WidgetInfo.PropertyIndex;
+					const FString WidgetDisplayName = WidgetInfo.DisplayName + ((VectorIndex != INDEX_NONE) ? FString::Printf(TEXT("[%d]"), VectorIndex) : TEXT(""));
+					const FString DisplayString = ValidationMessage.IsEmpty() ? WidgetDisplayName : ValidationMessage;
 
-				//do some string fixing
-				const uint32 VectorIndex = WidgetInfos[i].PropertyIndex;
-				const FString WidgetDisplayName = WidgetInfos[i].DisplayName + ((VectorIndex != INDEX_NONE) ? FString::Printf(TEXT("[%d]"), VectorIndex) : TEXT(""));
-				FinalString = FinalString.IsEmpty() ? WidgetDisplayName : FinalString;
-
-				if(Proj.W > 0.f)
-				{
 					const int32 XPos = HalfX + ( HalfX * Proj.X );
 					const int32 YPos = HalfY + ( HalfY * (Proj.Y * -1.f) );
-					FCanvasTextItem TextItem( FVector2D( XPos + 5, YPos), FText::FromString( FinalString ), GEngine->GetSmallFont(), FLinearColor::White );
+					FCanvasTextItem TextItem( FVector2D( XPos + 5, YPos), FText::FromString( DisplayString ), GEngine->GetSmallFont(), FLinearColor::White );
+					TextItem.EnableShadow(FLinearColor::Black);
 					Canvas->DrawItem( TextItem );
 				}
 			}
@@ -878,7 +834,6 @@ void FEdMode::DrawHUD(FEditorViewportClient* ViewportClient,FViewport* Viewport,
 	}
 }
 
-// Draw brackets around all selected objects
 void FEdMode::DrawBrackets( FEditorViewportClient* ViewportClient, FViewport* Viewport, const FSceneView* View, FCanvas* Canvas )
 {
 	USelection& SelectedActors = *Owner->GetSelectedActors();
@@ -1130,6 +1085,29 @@ struct FPropertyWidgetInfoChainElement
 	}
 };
 
+UObject* FEdMode::GetItemToTryDisplayingWidgetsFor(FTransform& OutLocalToWorld) const
+{
+	// Determine what is selected, preferring a component over an actor
+	USceneComponent* SelectedComponent = Owner->GetSelectedComponents()->GetTop<USceneComponent>();
+	UObject* BestSelectedItem = SelectedComponent;
+
+	if (SelectedComponent == nullptr)
+	{
+		AActor* SelectedActor = GetFirstSelectedActorInstance();
+		if (SelectedActor != nullptr)
+		{
+			BestSelectedItem = SelectedActor;
+			OutLocalToWorld = SelectedActor->GetRootComponent()->GetComponentToWorld();
+		}
+	}
+	else
+	{
+		OutLocalToWorld = SelectedComponent->GetComponentToWorld();
+	}
+
+	return BestSelectedItem;
+}
+
 void FEdMode::GetPropertyWidgetInfos(const UStruct* InStruct, const void* InContainer, TArray<FPropertyWidgetInfo>& OutInfos) const
 {
 	TArray<FPropertyWidgetInfoChainElement> Chain;
