Editor: Replaced direct calls to SetActorLabelUnique in UActorFactory::PostSpawnActor overrides with calls to the super which does just that

diff --git a/Engine/Plugins/2D/Paper2D/Source/Paper2DEditor/Private/PaperAnimatedSpriteActorFactory.cpp b/Engine/Plugins/2D/Paper2D/Source/Paper2DEditor/Private/PaperAnimatedSpriteActorFactory.cpp
--- a/Engine/Plugins/2D/Paper2D/Source/Paper2DEditor/Private/PaperAnimatedSpriteActorFactory.cpp
+++ b/Engine/Plugins/2D/Paper2D/Source/Paper2DEditor/Private/PaperAnimatedSpriteActorFactory.cpp
@@ -15,10 +15,10 @@ UPaperFlipbookActorFactory::UPaperFlipbookActorFactory(const FObjectInitializer&
 
 void UPaperFlipbookActorFactory::PostSpawnActor(UObject* Asset, AActor* NewActor)
 {
+	Super::PostSpawnActor(Asset, NewActor);
+
 	if (UPaperFlipbook* Flipbook = Cast<UPaperFlipbook>(Asset))
 	{
-		GEditor->SetActorLabelUnique(NewActor, Flipbook->GetName());
-
 		APaperFlipbookActor* TypedActor = CastChecked<APaperFlipbookActor>(NewActor);
 		UPaperFlipbookComponent* RenderComponent = TypedActor->GetRenderComponent();
 		check(RenderComponent);
diff --git a/Engine/Plugins/2D/Paper2D/Source/Paper2DEditor/Private/PaperSpriteActorFactory.cpp b/Engine/Plugins/2D/Paper2D/Source/Paper2DEditor/Private/PaperSpriteActorFactory.cpp
--- a/Engine/Plugins/2D/Paper2D/Source/Paper2DEditor/Private/PaperSpriteActorFactory.cpp
+++ b/Engine/Plugins/2D/Paper2D/Source/Paper2DEditor/Private/PaperSpriteActorFactory.cpp
@@ -15,10 +15,10 @@ UPaperSpriteActorFactory::UPaperSpriteActorFactory(const FObjectInitializer& Obj
 
 void UPaperSpriteActorFactory::PostSpawnActor(UObject* Asset, AActor* NewActor)
 {
+	Super::PostSpawnActor(Asset, NewActor);
+
 	if (UPaperSprite* Sprite = Cast<UPaperSprite>(Asset))
 	{
-		GEditor->SetActorLabelUnique(NewActor, Sprite->GetName());
-
 		APaperSpriteActor* TypedActor = CastChecked<APaperSpriteActor>(NewActor);
 		UPaperSpriteComponent* RenderComponent = TypedActor->GetRenderComponent();
 		check(RenderComponent);
diff --git a/Engine/Plugins/2D/Paper2D/Source/Paper2DEditor/Private/TerrainSplineEditing/TerrainSplineActorFactory.cpp b/Engine/Plugins/2D/Paper2D/Source/Paper2DEditor/Private/TerrainSplineEditing/TerrainSplineActorFactory.cpp
--- a/Engine/Plugins/2D/Paper2D/Source/Paper2DEditor/Private/TerrainSplineEditing/TerrainSplineActorFactory.cpp
+++ b/Engine/Plugins/2D/Paper2D/Source/Paper2DEditor/Private/TerrainSplineEditing/TerrainSplineActorFactory.cpp
@@ -15,10 +15,10 @@ UTerrainSplineActorFactory::UTerrainSplineActorFactory(const FObjectInitializer&
 
 void UTerrainSplineActorFactory::PostSpawnActor(UObject* Asset, AActor* NewActor)
 {
+	Super::PostSpawnActor(Asset, NewActor);
+
 	if (UPaperTerrainMaterial* TerrainMaterial = Cast<UPaperTerrainMaterial>(Asset))
 	{
-		GEditor->SetActorLabelUnique(NewActor, TerrainMaterial->GetName());
-
 		APaperTerrainActor* TypedActor = CastChecked<APaperTerrainActor>(NewActor);
 		UPaperTerrainComponent* RenderComponent = TypedActor->GetRenderComponent();
 		check(RenderComponent);
diff --git a/Engine/Plugins/2D/Paper2D/Source/Paper2DEditor/Private/TileMapEditing/TileMapActorFactory.cpp b/Engine/Plugins/2D/Paper2D/Source/Paper2DEditor/Private/TileMapEditing/TileMapActorFactory.cpp
--- a/Engine/Plugins/2D/Paper2D/Source/Paper2DEditor/Private/TileMapEditing/TileMapActorFactory.cpp
+++ b/Engine/Plugins/2D/Paper2D/Source/Paper2DEditor/Private/TileMapEditing/TileMapActorFactory.cpp
@@ -15,22 +15,20 @@ UTileMapActorFactory::UTileMapActorFactory(const FObjectInitializer& ObjectIniti
 
 void UTileMapActorFactory::PostSpawnActor(UObject* Asset, AActor* NewActor)
 {
+	Super::PostSpawnActor(Asset, NewActor);
+
 	APaperTileMapActor* TypedActor = CastChecked<APaperTileMapActor>(NewActor);
 	UPaperTileMapComponent* RenderComponent = TypedActor->GetRenderComponent();
 	check(RenderComponent);
 
 	if (UPaperTileMap* TileMap = Cast<UPaperTileMap>(Asset))
 	{
-		GEditor->SetActorLabelUnique(NewActor, TileMap->GetName());
-
 		RenderComponent->UnregisterComponent();
 		RenderComponent->TileMap = TileMap;
 		RenderComponent->RegisterComponent();
 	}
 	else if (UPaperTileSet* TileSet = Cast<UPaperTileSet>(Asset))
 	{
-		GEditor->SetActorLabelUnique(NewActor, TileSet->GetName());
-
 		if (RenderComponent->TileMap != nullptr)
 		{
 			RenderComponent->UnregisterComponent();
diff --git a/Engine/Source/Editor/UnrealEd/Classes/ActorFactories/ActorFactoryBlueprint.h b/Engine/Source/Editor/UnrealEd/Classes/ActorFactories/ActorFactoryBlueprint.h
--- a/Engine/Source/Editor/UnrealEd/Classes/ActorFactories/ActorFactoryBlueprint.h
+++ b/Engine/Source/Editor/UnrealEd/Classes/ActorFactories/ActorFactoryBlueprint.h
@@ -12,7 +12,6 @@ class UActorFactoryBlueprint : public UActorFactory
 
 	// Begin UActorFactory Interface
 	virtual bool PreSpawnActor( UObject* Asset, FTransform& InOutLocation ) override;
-	virtual void PostSpawnActor( UObject* Asset, AActor* NewActor ) override;
 	virtual bool CanCreateActorFrom( const FAssetData& AssetData, FText& OutErrorMsg ) override;
 	virtual AActor* GetDefaultActor( const FAssetData& AssetData ) override;
 	// End UActorFactory Interface	
diff --git a/Engine/Source/Editor/UnrealEd/Classes/ActorFactories/ActorFactoryClass.h b/Engine/Source/Editor/UnrealEd/Classes/ActorFactories/ActorFactoryClass.h
--- a/Engine/Source/Editor/UnrealEd/Classes/ActorFactories/ActorFactoryClass.h
+++ b/Engine/Source/Editor/UnrealEd/Classes/ActorFactories/ActorFactoryClass.h
@@ -12,7 +12,6 @@ class UActorFactoryClass : public UActorFactory
 
 	// Begin UActorFactory Interface
 	virtual bool PreSpawnActor( UObject* Asset, FTransform& InOutLocation ) override;
-	virtual void PostSpawnActor( UObject* Asset, AActor* NewActor ) override;
 	virtual bool CanCreateActorFrom( const FAssetData& AssetData, FText& OutErrorMsg ) override;
 	virtual AActor* GetDefaultActor( const FAssetData& AssetData ) override;
 	// End UActorFactory Interface	
diff --git a/Engine/Source/Editor/UnrealEd/Private/Factories/ActorFactory.cpp b/Engine/Source/Editor/UnrealEd/Private/Factories/ActorFactory.cpp
--- a/Engine/Source/Editor/UnrealEd/Private/Factories/ActorFactory.cpp
+++ b/Engine/Source/Editor/UnrealEd/Private/Factories/ActorFactory.cpp
@@ -280,8 +280,9 @@ bool UActorFactoryStaticMesh::CanCreateActorFrom( const FAssetData& AssetData, F
 
 void UActorFactoryStaticMesh::PostSpawnActor( UObject* Asset, AActor* NewActor)
 {
-	UStaticMesh* StaticMesh = CastChecked<UStaticMesh>( Asset );
-	GEditor->SetActorLabelUnique(NewActor, StaticMesh->GetName());
+	Super::PostSpawnActor(Asset, NewActor);
+
+	UStaticMesh* StaticMesh = CastChecked<UStaticMesh>(Asset);
 
 	UE_LOG(LogActorFactory, Log, TEXT("Actor Factory created %s"), *StaticMesh->GetName());
 
@@ -356,9 +357,10 @@ bool UActorFactoryBasicShape::CanCreateActorFrom( const FAssetData& AssetData, F
 
 void UActorFactoryBasicShape::PostSpawnActor(UObject* Asset, AActor* NewActor)
 {
+	Super::PostSpawnActor(Asset, NewActor);
+
 	// Change properties
 	UStaticMesh* StaticMesh = CastChecked<UStaticMesh>(Asset);
-	GEditor->SetActorLabelUnique(NewActor, StaticMesh->GetName());
 
 	AStaticMeshActor* StaticMeshActor = CastChecked<AStaticMeshActor>(NewActor);
 	UStaticMeshComponent* StaticMeshComponent = StaticMeshActor->GetStaticMeshComponent();
@@ -449,12 +451,12 @@ bool UActorFactoryDeferredDecal::CanCreateActorFrom( const FAssetData& AssetData
 
 void UActorFactoryDeferredDecal::PostSpawnActor(UObject* Asset, AActor* NewActor)
 {
+	Super::PostSpawnActor(Asset, NewActor);
+
 	UMaterialInterface* Material = GetMaterial( Asset );
 
 	if ( Material != NULL )
 	{
-		GEditor->SetActorLabelUnique(NewActor, Material->GetName());
-
 		// Change properties
 		TInlineComponentArray<UDecalComponent*> DecalComponents;
 		NewActor->GetComponents(DecalComponents);
@@ -474,10 +476,6 @@ void UActorFactoryDeferredDecal::PostSpawnActor(UObject* Asset, AActor* NewActor
 		// Init Component
 		DecalComponent->RegisterComponent();
 	}
-	else if (Asset)
-	{
-		GEditor->SetActorLabelUnique(NewActor, Asset->GetName());
-	}
 }
 
 void UActorFactoryDeferredDecal::PostCreateBlueprint( UObject* Asset, AActor* CDO )
@@ -550,11 +548,11 @@ bool UActorFactoryEmitter::CanCreateActorFrom( const FAssetData& AssetData, FTex
 
 void UActorFactoryEmitter::PostSpawnActor(UObject* Asset, AActor* NewActor)
 {
+	Super::PostSpawnActor(Asset, NewActor);
+
 	UParticleSystem* ParticleSystem = CastChecked<UParticleSystem>(Asset);
 	AEmitter* NewEmitter = CastChecked<AEmitter>(NewActor);
 
-	GEditor->SetActorLabelUnique(NewActor, ParticleSystem->GetName());
-
 	// Term Component
 	NewEmitter->GetParticleSystemComponent()->UnregisterComponent();
 
@@ -626,11 +624,11 @@ bool UActorFactoryNiagara::CanCreateActorFrom(const FAssetData& AssetData, FText
 
 void UActorFactoryNiagara::PostSpawnActor(UObject* Asset, AActor* NewActor)
 {
+	Super::PostSpawnActor(Asset, NewActor);
+
 	UNiagaraEffect* Effect = CastChecked<UNiagaraEffect>(Asset);
 	ANiagaraActor* NiagaraActor = CastChecked<ANiagaraActor>(NewActor);
 
-	GEditor->SetActorLabelUnique(NewActor, Effect->GetName());
-
 	// Term Component
 	NiagaraActor->GetNiagaraComponent()->UnregisterComponent();
 
@@ -742,11 +740,12 @@ bool UActorFactoryPhysicsAsset::PreSpawnActor(UObject* Asset, FTransform& InOutL
 
 void UActorFactoryPhysicsAsset::PostSpawnActor(UObject* Asset, AActor* NewActor)
 {
+	Super::PostSpawnActor(Asset, NewActor);
+
 	UPhysicsAsset* PhysicsAsset = CastChecked<UPhysicsAsset>(Asset);
 	USkeletalMesh* UseSkelMesh = PhysicsAsset->PreviewSkeletalMesh.Get();
 
 	ASkeletalMeshActor* NewSkelActor = CastChecked<ASkeletalMeshActor>(NewActor);
-	GEditor->SetActorLabelUnique(NewActor, PhysicsAsset->GetName());
 
 	// Term Component
 	NewSkelActor->GetSkeletalMeshComponent()->UnregisterComponent();
@@ -938,8 +937,6 @@ void UActorFactoryAnimationAsset::PostSpawnActor( UObject* Asset, AActor* NewAct
 			NewSASComponent->AnimationData.AnimToPlay = AnimationAsset;
 			// set runtime data
 			NewSASComponent->SetAnimation(AnimationAsset);
-
-			GEditor->SetActorLabelUnique(NewActor, AnimationAsset->GetName());
 		}
 		else if( VertexAnimation )
 		{
@@ -948,8 +945,6 @@ void UActorFactoryAnimationAsset::PostSpawnActor( UObject* Asset, AActor* NewAct
 
 			// set runtime data
 			NewSASComponent->SetVertexAnimation(VertexAnimation);
-
-			GEditor->SetActorLabelUnique(NewActor, VertexAnimation->GetName());
 		}
 	}
 }
@@ -1114,7 +1109,7 @@ void UActorFactorySkeletalMesh::PostSpawnActor( UObject* Asset, AActor* NewActor
 	UAnimBlueprint* AnimBlueprint = Cast<UAnimBlueprint>( Asset );
 	ASkeletalMeshActor* NewSMActor = CastChecked<ASkeletalMeshActor>(NewActor);
 
-	GEditor->SetActorLabelUnique(NewActor, SkeletalMesh->GetName());
+	Super::PostSpawnActor(SkeletalMesh, NewActor);
 
 	// Term Component
 	NewSMActor->GetSkeletalMeshComponent()->UnregisterComponent();
@@ -1290,12 +1285,13 @@ bool UActorFactoryAmbientSound::CanCreateActorFrom( const FAssetData& AssetData,
 
 void UActorFactoryAmbientSound::PostSpawnActor( UObject* Asset, AActor* NewActor)
 {
-	USoundBase* AmbientSound = Cast<USoundBase>( Asset );
+	Super::PostSpawnActor(Asset, NewActor);
+
+	USoundBase* AmbientSound = Cast<USoundBase>(Asset);
 
 	if ( AmbientSound != NULL )
 	{
 		AAmbientSound* NewSound = CastChecked<AAmbientSound>( NewActor );
-		GEditor->SetActorLabelUnique(NewSound, AmbientSound->GetName());
 		NewSound->GetAudioComponent()->SetSound(AmbientSound);
 	}
 }
@@ -1375,12 +1371,6 @@ bool UActorFactoryClass::PreSpawnActor( UObject* Asset, FTransform& InOutLocatio
 	return false;
 }
 
-void UActorFactoryClass::PostSpawnActor( UObject* Asset, AActor* NewActor)
-{
-	UClass* ActualClass = CastChecked<UClass>(Asset);
-	GEditor->SetActorLabelUnique(NewActor, ActualClass->GetName());
-}
-
 AActor* UActorFactoryClass::SpawnActor( UObject* Asset, ULevel* InLevel, const FVector& Location, const FRotator& Rotation, EObjectFlags ObjectFlags, const FName& Name )
 {
 	UClass* ActualClass = Cast<UClass>(Asset);
@@ -1490,12 +1480,6 @@ bool UActorFactoryBlueprint::PreSpawnActor( UObject* Asset, FTransform& InOutLoc
 	return true;
 }
 
-void UActorFactoryBlueprint::PostSpawnActor( UObject* Asset, AActor* NewActor)
-{
-	UBlueprint* Blueprint = CastChecked<UBlueprint>(Asset);
-	GEditor->SetActorLabelUnique(NewActor, Blueprint->GetName());
-}
-
 
 
 /*-----------------------------------------------------------------------------
@@ -1527,18 +1511,19 @@ bool UActorFactoryMatineeActor::CanCreateActorFrom( const FAssetData& AssetData,
 
 void UActorFactoryMatineeActor::PostSpawnActor( UObject* Asset, AActor* NewActor )
 {
+	Super::PostSpawnActor(Asset, NewActor);
+
 	UInterpData* MatineeData = Cast<UInterpData>( Asset );
 	AMatineeActor* MatineeActor = CastChecked<AMatineeActor>( NewActor );
 
 	if( MatineeData )
 	{
-		GEditor->SetActorLabelUnique(NewActor, MatineeData->GetName());
 		MatineeActor->MatineeData = MatineeData;
 	}
 	else
 	{
 		// if MatineeData isn't set yet, create default one
-		UInterpData * NewMatineeData = NewObject<UInterpData>(NewActor);
+		UInterpData* NewMatineeData = NewObject<UInterpData>(NewActor);
 		MatineeActor->MatineeData = NewMatineeData;
 	}
 }
@@ -1724,11 +1709,11 @@ bool UActorFactoryDestructible::CanCreateActorFrom( const FAssetData& AssetData,
 
 void UActorFactoryDestructible::PostSpawnActor( UObject* Asset, AActor* NewActor )
 {
+	Super::PostSpawnActor(Asset, NewActor);
+
 	UDestructibleMesh* DestructibleMesh = CastChecked<UDestructibleMesh>( Asset );
 	ADestructibleActor* NewDestructibleActor = CastChecked<ADestructibleActor>(NewActor);
 
-	GEditor->SetActorLabelUnique(NewActor, DestructibleMesh->GetName());
-
 	// Term Component
 	NewDestructibleActor->GetDestructibleComponent()->UnregisterComponent();
 
@@ -1789,11 +1774,11 @@ bool UActorFactoryVectorFieldVolume::CanCreateActorFrom( const FAssetData& Asset
 
 void UActorFactoryVectorFieldVolume::PostSpawnActor( UObject* Asset, AActor* NewActor )
 {
+	Super::PostSpawnActor(Asset, NewActor);
+
 	UVectorField* VectorField = CastChecked<UVectorField>(Asset);
 	AVectorFieldVolume* VectorFieldVolumeActor = CastChecked<AVectorFieldVolume>(NewActor);
 
-	GEditor->SetActorLabelUnique(NewActor, VectorField->GetName());
-
 	if ( VectorFieldVolumeActor && VectorFieldVolumeActor->GetVectorFieldComponent() )
 	{
 		VectorFieldVolumeActor->GetVectorFieldComponent()->VectorField = VectorField;
@@ -1868,16 +1853,13 @@ bool UActorFactoryBoxVolume::CanCreateActorFrom( const FAssetData& AssetData, FT
 
 void UActorFactoryBoxVolume::PostSpawnActor( UObject* Asset, AActor* NewActor )
 {
+	Super::PostSpawnActor(Asset, NewActor);
+
 	AVolume* VolumeActor = CastChecked<AVolume>(NewActor);
 	if ( VolumeActor != NULL )
 	{
 		UCubeBuilder* Builder = NewObject<UCubeBuilder>();
 		CreateBrushForVolumeActor( VolumeActor, Builder );
-
-		if (Asset)
-		{
-			GEditor->SetActorLabelUnique(NewActor, Asset->GetName());
-		}
 	}
 }
 
@@ -1908,6 +1890,8 @@ bool UActorFactorySphereVolume::CanCreateActorFrom( const FAssetData& AssetData,
 
 void UActorFactorySphereVolume::PostSpawnActor( UObject* Asset, AActor* NewActor )
 {
+	Super::PostSpawnActor(Asset, NewActor);
+
 	AVolume* VolumeActor = CastChecked<AVolume>(NewActor);
 	if ( VolumeActor != NULL )
 	{
@@ -1915,11 +1899,6 @@ void UActorFactorySphereVolume::PostSpawnActor( UObject* Asset, AActor* NewActor
 		Builder->SphereExtrapolation = 2;
 		Builder->Radius = 192.0f;
 		CreateBrushForVolumeActor( VolumeActor, Builder );
-
-		if (Asset)
-		{
-			GEditor->SetActorLabelUnique(NewActor, Asset->GetName());
-		}
 	}
 }
 
@@ -1949,17 +1928,14 @@ bool UActorFactoryCylinderVolume::CanCreateActorFrom( const FAssetData& AssetDat
 }
 void UActorFactoryCylinderVolume::PostSpawnActor( UObject* Asset, AActor* NewActor )
 {
+	Super::PostSpawnActor(Asset, NewActor);
+
 	AVolume* VolumeActor = CastChecked<AVolume>(NewActor);
 	if ( VolumeActor != NULL )
 	{
 		UCylinderBuilder* Builder = NewObject<UCylinderBuilder>();
 		Builder->OuterRadius = 128.0f;
 		CreateBrushForVolumeActor( VolumeActor, Builder );
-
-		if (Asset)
-		{
-			GEditor->SetActorLabelUnique(NewActor, Asset->GetName());
-		}
 	}
 }
 
