Editor: Deprecated GEditor->SetActorLabelUnique and GEditor->SplitActorLabel. Please replace calls with the static methods  FActorLabelUtilities::SetActorLabelUnique and FActorLabelUtilities::SplitActorLabel

diff --git a/Engine/Source/Editor/KismetWidgets/Private/CreateBlueprintFromActorDialog.cpp b/Engine/Source/Editor/KismetWidgets/Private/CreateBlueprintFromActorDialog.cpp
--- a/Engine/Source/Editor/KismetWidgets/Private/CreateBlueprintFromActorDialog.cpp
+++ b/Engine/Source/Editor/KismetWidgets/Private/CreateBlueprintFromActorDialog.cpp
@@ -89,7 +89,7 @@ void FCreateBlueprintFromActorDialog::OnCreateBlueprint(const FString& InAssetPa
 			AActor* Actor = Cast<AActor>(SelectedActors->GetSelectedObject(0));
 			if(Actor)
 			{
-				GEditor->SetActorLabelUnique(Actor, FPackageName::GetShortName(InAssetPath));
+				FActorLabelUtilities::SetActorLabelUnique(Actor, FPackageName::GetShortName(InAssetPath));
 			}
 		}
 
diff --git a/Engine/Source/Editor/Sequencer/Private/SequencerActorBindingManager.cpp b/Engine/Source/Editor/Sequencer/Private/SequencerActorBindingManager.cpp
--- a/Engine/Source/Editor/Sequencer/Private/SequencerActorBindingManager.cpp
+++ b/Engine/Source/Editor/Sequencer/Private/SequencerActorBindingManager.cpp
@@ -216,7 +216,7 @@ void FSequencerActorBindingManager::SpawnOrDestroyObjectsForInstance( TSharedRef
 						if( NewActor )
 						{
 							// @todo sequencer: We're naming the actor based off of the spawnable's name.  Is that really what we want?
-							GEditor->SetActorLabelUnique( NewActor, Spawnable.GetName() );
+							FActorLabelUtilities::SetActorLabelUnique(NewActor, Spawnable.GetName());
 
 							// Actor was spawned OK!
 
diff --git a/Engine/Source/Editor/UnrealEd/Classes/Editor/EditorEngine.h b/Engine/Source/Editor/UnrealEd/Classes/Editor/EditorEngine.h
--- a/Engine/Source/Editor/UnrealEd/Classes/Editor/EditorEngine.h
+++ b/Engine/Source/Editor/UnrealEd/Classes/Editor/EditorEngine.h
@@ -2594,7 +2594,7 @@ private:
 		FProcHandle ProcessHandle;
 	};
 
-	TArray<FPlayOnPCInfo>	PlayOnLocalPCSessions;
+	TArray<FPlayOnPCInfo> PlayOnLocalPCSessions;
 
 	bool bPlayOnLocalPcSession;
 
@@ -2624,7 +2624,7 @@ private:
 	FRotator LastViewRotation;
 	
 	/** Are the lastview/rotation variables valid */
-	bool	bLastViewAndLocationValid;
+	bool bLastViewAndLocationValid;
 
 protected:
 
@@ -2702,7 +2702,7 @@ public:
 	 *
 	 * @return	true if the label ends with a number.
 	 */
-	//DEPRECATED(4.8, "This function is deprecated.  Please use FActorLabelUtilities::SplitActorLabel instead")
+	DEPRECATED(4.8, "This function is deprecated.  Please use FActorLabelUtilities::SplitActorLabel instead")
 	bool SplitActorLabel( FString& InOutLabel, int32& OutIdx ) const;
 
 	/**
@@ -2712,7 +2712,7 @@ public:
 	 * @param	NewActorLabel			The new label string to assign to the actor.  If empty, the actor will have a default label.
 	 * @param	InExistingActorLabels	(optional) Pointer to a set of actor labels that are currently in use
 	 */
-	//DEPRECATED(4.8, "This function is deprecated.  Please use FActorLabelUtilities::SetActorLabelUnique instead")
+	DEPRECATED(4.8, "This function is deprecated.  Please use FActorLabelUtilities::SetActorLabelUnique instead")
 	void SetActorLabelUnique( AActor* Actor, const FString& NewActorLabel, const FCachedActorLabels* InExistingActorLabels = nullptr ) const;
 
 private:
@@ -2720,7 +2720,7 @@ private:
 };
 
 //////////////////////////////////////////////////////////////////////////
-// UEditorEngine
+// FActorLabelUtilities
 
 struct UNREALED_API FActorLabelUtilities
 {
@@ -2742,7 +2742,6 @@ public:
 	 * @param	NewActorLabel			The new label string to assign to the actor.  If empty, the actor will have a default label.
 	 * @param	InExistingActorLabels	(optional) Pointer to a set of actor labels that are currently in use
 	 */
-	//DEPRECATED(4.8, "This function is deprecated.  Please use FActorLabelUtilities::SetActorLabelUnique instead")
 	static void SetActorLabelUnique(AActor* Actor, const FString& NewActorLabel, const FCachedActorLabels* InExistingActorLabels = nullptr);
 
 private:
diff --git a/Engine/Source/Editor/UnrealEd/Private/BSPOps.cpp b/Engine/Source/Editor/UnrealEd/Private/BSPOps.cpp
--- a/Engine/Source/Editor/UnrealEd/Private/BSPOps.cpp
+++ b/Engine/Source/Editor/UnrealEd/Private/BSPOps.cpp
@@ -614,7 +614,7 @@ ABrush*	FBSPOps::csgAddOperation( ABrush* Actor, uint32 PolyFlags, EBrushType Br
 
 	if( Result->GetBrushBuilder() )
 	{
-		GEditor->SetActorLabelUnique(Result, FText::Format(NSLOCTEXT("BSPBrushOps", "BrushName", "{0} Brush"), FText::FromString(Result->GetBrushBuilder()->GetClass()->GetDescription())).ToString());
+		FActorLabelUtilities::SetActorLabelUnique(Result, FText::Format(NSLOCTEXT("BSPBrushOps", "BrushName", "{0} Brush"), FText::FromString(Result->GetBrushBuilder()->GetClass()->GetDescription())).ToString());
 	}
 	// Assign the default material to the brush's polys.
 	for( int32 i=0; i<Result->Brush->Polys->Element.Num(); i++ )
diff --git a/Engine/Source/Editor/UnrealEd/Private/EditorActor.cpp b/Engine/Source/Editor/UnrealEd/Private/EditorActor.cpp
--- a/Engine/Source/Editor/UnrealEd/Private/EditorActor.cpp
+++ b/Engine/Source/Editor/UnrealEd/Private/EditorActor.cpp
@@ -330,7 +330,7 @@ void UUnrealEdEngine::edactPasteSelected(UWorld* InWorld, bool bDuplicate, bool
 			Actor->TeleportTo(Actor->GetActorLocation() + ActorLocationOffset, Actor->GetActorRotation(), false, true);
 
 			// Re-label duplicated actors so that labels become unique
-			GEditor->SetActorLabelUnique(Actor, Actor->GetActorLabel(), &ActorLabels);
+			FActorLabelUtilities::SetActorLabelUnique(Actor, Actor->GetActorLabel(), &ActorLabels);
 			ActorLabels.Add(Actor->GetActorLabel());
 
 			GEditor->Layers->InitializeNewActorLayers(Actor);
diff --git a/Engine/Source/Editor/UnrealEd/Private/Factories/ActorFactory.cpp b/Engine/Source/Editor/UnrealEd/Private/Factories/ActorFactory.cpp
--- a/Engine/Source/Editor/UnrealEd/Private/Factories/ActorFactory.cpp
+++ b/Engine/Source/Editor/UnrealEd/Private/Factories/ActorFactory.cpp
@@ -246,7 +246,7 @@ void UActorFactory::PostSpawnActor( UObject* Asset, AActor* NewActor)
 	// Subclasses may implement this to modify the actor after it has been spawned
 	if (Asset)
 	{
-		GEditor->SetActorLabelUnique(NewActor, Asset->GetName());
+		FActorLabelUtilities::SetActorLabelUnique(NewActor, Asset->GetName());
 	}
 }
 
diff --git a/Engine/Source/Editor/UnrealEd/Private/Factories/EditorFactories.cpp b/Engine/Source/Editor/UnrealEd/Private/Factories/EditorFactories.cpp
--- a/Engine/Source/Editor/UnrealEd/Private/Factories/EditorFactories.cpp
+++ b/Engine/Source/Editor/UnrealEd/Private/Factories/EditorFactories.cpp
@@ -636,7 +636,7 @@ UObject* ULevelFactory::FactoryCreateText
 
 									// Place the group in the map so we can find it later.
 									NewGroups.Add( GroupName, SpawnedGroupActor);
-									GEditor->SetActorLabelUnique(SpawnedGroupActor, GroupName);
+									FActorLabelUtilities::SetActorLabelUnique(SpawnedGroupActor, GroupName);
 								}
 
 								// If we're copying a sub-group, add add duplicated group to original parent
