Editor: Fixed a warning about LOCTEXT_NAMESPACE being redefined due to a missing undef

diff --git a/Engine/Source/Editor/UnrealEd/Private/Editor.cpp b/Engine/Source/Editor/UnrealEd/Private/Editor.cpp
--- a/Engine/Source/Editor/UnrealEd/Private/Editor.cpp
+++ b/Engine/Source/Editor/UnrealEd/Private/Editor.cpp
@@ -1208,3 +1208,7 @@ void FCachedActorLabels::Populate(UWorld* World, const TSet<AActor*>& IgnoredAct
 	}
 	ActorLabels.Shrink();
 }
+
+//////////////////////////////////////////////////////////////////////////
+
+#undef LOCTEXT_NAMESPACE
