Editor: Fixed a crash in the default editor mode when an actor with no root component is selected

diff --git a/Engine/Source/Editor/UnrealEd/Private/EdMode.cpp b/Engine/Source/Editor/UnrealEd/Private/EdMode.cpp
--- a/Engine/Source/Editor/UnrealEd/Private/EdMode.cpp
+++ b/Engine/Source/Editor/UnrealEd/Private/EdMode.cpp
@@ -1096,8 +1096,11 @@ UObject* FEdMode::GetItemToTryDisplayingWidgetsFor(FTransform& OutLocalToWorld)
 		AActor* SelectedActor = GetFirstSelectedActorInstance();
 		if (SelectedActor != nullptr)
 		{
-			BestSelectedItem = SelectedActor;
-			OutLocalToWorld = SelectedActor->GetRootComponent()->GetComponentToWorld();
+			if (USceneComponent* RootComponent = SelectedActor->GetRootComponent())
+			{
+				BestSelectedItem = SelectedActor;
+				OutLocalToWorld = RootComponent->GetComponentToWorld();
+			}
 		}
 	}
 	else
