Fixed material flattening which was produced bad normal maps and base colors in some cases
Fixed mesh merging which was not transforming mesh normals

diff --git a/Engine/Source/Developer/MeshUtilities/Private/MeshUtilities.cpp b/Engine/Source/Developer/MeshUtilities/Private/MeshUtilities.cpp
--- a/Engine/Source/Developer/MeshUtilities/Private/MeshUtilities.cpp
+++ b/Engine/Source/Developer/MeshUtilities/Private/MeshUtilities.cpp
@@ -3966,6 +3966,18 @@ bool FMeshUtilities::ConstructRawMesh(
 	{
 		Vertex = CtoM.TransformFVector4(Vertex);
 	}
+	for (FVector& TangentX : OutRawMesh.WedgeTangentX)
+	{
+		TangentX = CtoM.TransformVectorNoScale(TangentX);
+	}
+	for (FVector& TangentY : OutRawMesh.WedgeTangentY)
+	{
+		TangentY = CtoM.TransformVectorNoScale(TangentY);
+	}
+	for (FVector& TangentZ : OutRawMesh.WedgeTangentZ)
+	{
+		TangentZ = CtoM.TransformVectorNoScale(TangentZ);
+	}
 
 	const bool bIsMirrored = CtoM.GetDeterminant() < 0.f;
 	if (bIsMirrored)
diff --git a/Engine/Source/Editor/UnrealEd/Private/EditorExporters.cpp b/Engine/Source/Editor/UnrealEd/Private/EditorExporters.cpp
--- a/Engine/Source/Editor/UnrealEd/Private/EditorExporters.cpp
+++ b/Engine/Source/Editor/UnrealEd/Private/EditorExporters.cpp
@@ -2456,7 +2456,7 @@ UExportTextContainer::UExportTextContainer(const FObjectInitializer& ObjectIniti
 -----------------------------------------------------------------------------*/
 namespace MaterialExportUtils
 {
-	void RenderMaterialTile(UWorld* InWorld, FMaterialRenderProxy* InMaterialProxy, UTextureRenderTarget2D* InRenderTarget)
+	void RenderMaterialTile(UWorld* InWorld, FMaterialRenderProxy* InMaterialProxy, UTextureRenderTarget2D* InRenderTarget, bool bFrontView)
 	{
 		float CurrentRealTime = 0.f;
 		float CurrentWorldTime = 0.f;
@@ -2485,13 +2485,25 @@ namespace MaterialExportUtils
 
 		FIntPoint ViewSize = RenderTargetResource->GetSizeXY();
 		FIntRect ViewRect(FIntPoint(0, 0), ViewSize);
-		// YZ - Front view
+		// By default render tile in top view
 		FMatrix ViewMatrix = FMatrix(
-						FPlane(1, 0, 0, 0),
-						FPlane(0, 0, -1, 0),
-						FPlane(0, 1, 0, 0),
-						FPlane(0, 0, 0, 1));
-					
+					FPlane(1, 0, 0, 0),
+					FPlane(0, -1, 0, 0),
+					FPlane(0, 0, -1, 0),
+					FPlane(0, 0, 0, 1));
+		FQuat Rotation = FQuat::Identity;
+
+		if (bFrontView)
+		{
+			ViewMatrix = FMatrix(
+					FPlane(1, 0, 0, 0),
+					FPlane(0, 0, -1, 0),
+					FPlane(0, 1, 0, 0),
+					FPlane(0, 0, 0, 1));
+			// Tile mesh was created on XY plane
+			Rotation = FRotator(0.0f, 0.0f, 90.0f).Quaternion();
+		}
+						
 		// make a temporary view
 		FSceneViewInitOptions ViewInitOptions;
 		ViewInitOptions.ViewFamily = ViewFamily;
@@ -2504,21 +2516,19 @@ namespace MaterialExportUtils
 				
 		FSceneView* View = new FSceneView(ViewInitOptions);
 						
-		ENQUEUE_UNIQUE_RENDER_COMMAND_FOURPARAMETER(
+		ENQUEUE_UNIQUE_RENDER_COMMAND_FIVEPARAMETER(
 			RenderMaterialTileCommand,
 			FSceneView*, InView, View,
 			FMaterialRenderProxy*, InProxy, InMaterialProxy,
 			const FRenderTarget*, InRenderTarget, RenderTargetResource,
 			FIntPoint, InSize, ViewSize,
+			FQuat, InRotation, Rotation,
 			{
 				::SetRenderTarget(RHICmdList, InRenderTarget->GetRenderTargetTexture(), FTextureRHIRef());
 				FIntRect ViewportRect = FIntRect(FIntPoint::ZeroValue, InSize);
 				RHICmdList.SetViewport(ViewportRect.Min.X, ViewportRect.Min.Y, 0.0f, ViewportRect.Max.X, ViewportRect.Max.Y, 1.0f);
-				
-				// We want to render material into tile that has vertices on YZ plane
-				// By default tile mesh created on XY plane, so rotate it
-				FQuat Rotation = FRotator(0.f, 0.f, 90).Quaternion();
-				FTileRenderer::DrawRotatedTile(RHICmdList, *InView, InProxy, false, Rotation, 0.f, 0.f, InSize.X, InSize.Y,  0.f, 0.f, 1.f, 1.f);
+		
+				FTileRenderer::DrawRotatedTile(RHICmdList, *InView, InProxy, false, InRotation, 0.f, 0.f, InSize.X, InSize.Y,  0.f, 0.f, 1.f, 1.f);
 			});
 
 		FlushRenderingCommands();
@@ -2550,9 +2560,10 @@ namespace MaterialExportUtils
 	
 		check(InRenderTarget);
 
-		RenderMaterialTile(InWorld, MaterialProxy, InRenderTarget);
-				
 		bool bNormalmap = (InMaterialProperty == MP_Normal);
+
+		RenderMaterialTile(InWorld, MaterialProxy, InRenderTarget, !bNormalmap);
+					
 		FReadSurfaceDataFlags ReadPixelFlags(bNormalmap ? RCM_SNorm : RCM_UNorm);
 		ReadPixelFlags.SetLinearToGamma(false);
 
