Fixed moving actors into a hidden level.

diff --git a/Engine/Source/Editor/UnrealEd/Private/EditorServer.cpp b/Engine/Source/Editor/UnrealEd/Private/EditorServer.cpp
--- a/Engine/Source/Editor/UnrealEd/Private/EditorServer.cpp
+++ b/Engine/Source/Editor/UnrealEd/Private/EditorServer.cpp
@@ -9,6 +9,7 @@
 #include "ParticleDefinitions.h"
 #include "AnimationUtils.h"
 #include "LevelUtils.h"
+#include "EditorLevelUtils.h"
 #include "Layers/ILayers.h"
 #include "ScopedTransaction.h"
 #include "SurfaceIterators.h"
@@ -2851,11 +2852,23 @@ void UEditorEngine::DoMoveSelectedActorsToLevel( ULevel* InDestLevel )
 	// Copy the actors we have selected to the clipboard
 	CopySelectedActorsToClipboard( World, true );
 
-	// Set the new level
+	// Set the new level and force it visible while we do the paste
 	World->SetCurrentLevel( InDestLevel );
+	const bool bLevelVisible = InDestLevel->bIsVisible;
+	if (!bLevelVisible)
+	{
+		EditorLevelUtils::SetLevelVisibility(InDestLevel, true, false);
+	}
+
 	// Paste the actors into the new level
 	edactPasteSelected( World, false, false, false );
 
+	// Restore new level visibility to previous state
+	if (!bLevelVisible)
+	{
+		EditorLevelUtils::SetLevelVisibility(InDestLevel, false, false);
+	}
+
 	// Restore the original current level
 	World->SetCurrentLevel( OldCurrentLevel );
 
