UE-13697: Fixes camera position in material editor's view

diff --git a/Engine/Source/Editor/UnrealEd/Private/EditorViewportClient.cpp b/Engine/Source/Editor/UnrealEd/Private/EditorViewportClient.cpp
--- a/Engine/Source/Editor/UnrealEd/Private/EditorViewportClient.cpp
+++ b/Engine/Source/Editor/UnrealEd/Private/EditorViewportClient.cpp
@@ -649,8 +649,7 @@ FSceneView* FEditorViewportClient::CalcSceneView(FSceneViewFamily* ViewFamily)
 		{
 			if (bUsingOrbitCamera)
 			{
-				ViewInitOptions.ViewRotationMatrix = ViewTransform.ComputeOrbitMatrix();
-				ViewInitOptions.ViewOrigin = FVector::ZeroVector;
+				ViewInitOptions.ViewRotationMatrix = FTranslationMatrix(ViewLocation) * ViewTransform.ComputeOrbitMatrix();
 			}
 			else
 			{
