Fixed "Align Brush Vertices" operation on rotated or scaled brushes.

diff --git a/Engine/Source/Editor/UnrealEd/Private/EditorActor.cpp b/Engine/Source/Editor/UnrealEd/Private/EditorActor.cpp
--- a/Engine/Source/Editor/UnrealEd/Private/EditorActor.cpp
+++ b/Engine/Source/Editor/UnrealEd/Private/EditorActor.cpp
@@ -2240,6 +2240,7 @@ void UUnrealEdEngine::edactAlignVertices()
 			Brush->PreEditChange(NULL);
 			Brush->Modify();
 			FVector BrushLocation = Brush->GetActorLocation();
+			const FTransform BrushTransform = Brush->GetRootComponent()->GetComponentTransform();
 
 			// Snap each vertex in the brush to an integer grid.
 			UPolys* Polys = Brush->Brush->Polys;
@@ -2248,10 +2249,17 @@ void UUnrealEdEngine::edactAlignVertices()
 				FPoly* Poly = &Polys->Element[PolyIdx];
 				for( int32 VertIdx=0; VertIdx<Poly->Vertices.Num(); VertIdx++ )
 				{
+					const float GridSize = GetGridSize();
+
 					// Snap each vertex to the nearest grid.
-					Poly->Vertices[VertIdx].X = FMath::RoundToFloat( ( Poly->Vertices[VertIdx].X + BrushLocation.X )  / GetGridSize() ) * GetGridSize() - BrushLocation.X;
-					Poly->Vertices[VertIdx].Y = FMath::RoundToFloat( ( Poly->Vertices[VertIdx].Y + BrushLocation.Y )  / GetGridSize() ) * GetGridSize() - BrushLocation.Y;
-					Poly->Vertices[VertIdx].Z = FMath::RoundToFloat( ( Poly->Vertices[VertIdx].Z + BrushLocation.Z )  / GetGridSize() ) * GetGridSize() - BrushLocation.Z;
+					const FVector Vertex = Poly->Vertices[VertIdx];
+					const FVector VertexWorld = BrushTransform.TransformPosition(Vertex);
+					const FVector VertexSnapped(FMath::RoundToFloat(VertexWorld.X / GridSize) * GridSize,
+												FMath::RoundToFloat(VertexWorld.Y / GridSize) * GridSize,
+												FMath::RoundToFloat(VertexWorld.Z / GridSize) * GridSize);
+					const FVector VertexSnappedLocal = BrushTransform.InverseTransformPosition(VertexSnapped);
+
+					Poly->Vertices[VertIdx] = VertexSnappedLocal;
 				}
 
 				// If the snapping resulted in an off plane polygon, triangulate it to compensate.
