Dragging material/texture assets over the level viewport now retains the cursor decorator window.

Changed AttemptApplyObjToComponent so that dragging textures or materials over components returns false (meaning they were not applied, and hence the cursor decorator window remains visible).  They will still be successfully applied on drop.

diff --git a/Engine/Source/Editor/UnrealEd/Private/LevelEditorViewport.cpp b/Engine/Source/Editor/UnrealEd/Private/LevelEditorViewport.cpp
--- a/Engine/Source/Editor/UnrealEd/Private/LevelEditorViewport.cpp
+++ b/Engine/Source/Editor/UnrealEd/Private/LevelEditorViewport.cpp
@@ -521,7 +521,7 @@ static bool AttemptApplyObjToComponent(UObject* ObjToUse, USceneComponent* Compo
 			{
 				if (bTest)
 				{
-					bResult = true;
+					bResult = false;
 				}
 				else
 				{
@@ -536,7 +536,7 @@ static bool AttemptApplyObjToComponent(UObject* ObjToUse, USceneComponent* Compo
 			{
 				if (bTest)
 				{
-					bResult = true;
+					bResult = false;
 				}
 				else
 				{
