Editor: Added a note to look further at how SStandaloneAssetEditorToolkitHost process command bindings

diff --git a/Engine/Source/Editor/UnrealEd/Private/Toolkits/SStandaloneAssetEditorToolkitHost.cpp b/Engine/Source/Editor/UnrealEd/Private/Toolkits/SStandaloneAssetEditorToolkitHost.cpp
--- a/Engine/Source/Editor/UnrealEd/Private/Toolkits/SStandaloneAssetEditorToolkitHost.cpp
+++ b/Engine/Source/Editor/UnrealEd/Private/Toolkits/SStandaloneAssetEditorToolkitHost.cpp
@@ -322,6 +322,7 @@ FReply SStandaloneAssetEditorToolkitHost::OnKeyDown( const FGeometry& MyGeometry
 			}
 		}
 
+		//@TODO: This seems wrong (should prioritize it but not totally block it)
 		if (ActiveToolkit.IsValid())
 		{
 			// A toolkit tab is active, so direct all command processing to it
