UEPLAT-543 - added a build option to the Editor Preferences for Launch On
#editor

diff --git a/Engine/Source/Editor/UnrealEd/Classes/Settings/LevelEditorPlaySettings.h b/Engine/Source/Editor/UnrealEd/Classes/Settings/LevelEditorPlaySettings.h
--- a/Engine/Source/Editor/UnrealEd/Classes/Settings/LevelEditorPlaySettings.h
+++ b/Engine/Source/Editor/UnrealEd/Classes/Settings/LevelEditorPlaySettings.h
@@ -190,6 +190,12 @@ public:
 	UPROPERTY(config , EditAnywhere, Category=PlayInStandaloneGame, AdvancedDisplay)
 	FString AdditionalLaunchParameters;
 
+public:
+
+	/** The width of the new view port window in pixels (0 = use the desktop's screen resolution). */
+	UPROPERTY(config, EditAnywhere, Category = PlayOnDevice)
+	bool BuildGameBeforeLaunch;
+
 private:
 
 	/** NetMode to use for Play In Editor. */
diff --git a/Engine/Source/Editor/UnrealEd/Private/PlayLevel.cpp b/Engine/Source/Editor/UnrealEd/Private/PlayLevel.cpp
--- a/Engine/Source/Editor/UnrealEd/Private/PlayLevel.cpp
+++ b/Engine/Source/Editor/UnrealEd/Private/PlayLevel.cpp
@@ -1462,7 +1462,22 @@ void UEditorEngine::PlayUsingLauncher()
 
 		// Setup launch profile, keep the setting here to a minimum.
 		ILauncherProfileRef LauncherProfile = LauncherServicesModule.CreateProfile(TEXT("Play On Device"));
-		LauncherProfile->SetBuildGame(bPlayUsingLauncherHasCode && bPlayUsingLauncherHasCompiler);
+		bool bAlwaysBuild = GetDefault<ULevelEditorPlaySettings>()->BuildGameBeforeLaunch;
+		if ((bPlayUsingLauncherHasCode || bAlwaysBuild) && bPlayUsingLauncherHasCompiler)
+		{
+			LauncherProfile->SetBuildGame((bPlayUsingLauncherHasCode || bAlwaysBuild) && bPlayUsingLauncherHasCompiler);
+
+			// set the build configuration to be the same as the running editor
+			FString ExeName = FUnrealEdMisc::Get().GetExecutableForCommandlets();
+			if (ExeName.Contains(TEXT("Debug")))
+			{
+				LauncherProfile->SetBuildConfiguration(EBuildConfigurations::Debug);
+			}
+			else
+			{
+				LauncherProfile->SetBuildConfiguration(EBuildConfigurations::Development);
+			}
+		}
 
 		// select the quickest cook mode based on which in editor cook mode is enabled
 		bool bIncrimentalCooking = true;
