Rationalized the filename building in SaveAsImplementation to account for ambiguities when passing paths or fully-qualified filenames to various FPaths and FPackageName methods.

diff --git a/Engine/Source/Editor/UnrealEd/Private/FileHelpers.cpp b/Engine/Source/Editor/UnrealEd/Private/FileHelpers.cpp
--- a/Engine/Source/Editor/UnrealEd/Private/FileHelpers.cpp
+++ b/Engine/Source/Editor/UnrealEd/Private/FileHelpers.cpp
@@ -602,16 +602,7 @@ static bool OpenLevelSaveAsDialog(const FString& InDefaultPath, const FString& I
 	}
 
 	FString NewNameSuggestion = InNewNameSuggestion;
-	if (NewNameSuggestion.IsEmpty())
-	{
-		const FString DefaultName = TEXT("NewMap");
-		FString PackageName = DefaultPath / DefaultName;
-		FString Name;
-		FAssetToolsModule& AssetToolsModule = FModuleManager::LoadModuleChecked<FAssetToolsModule>("AssetTools");
-		AssetToolsModule.Get().CreateUniqueAssetName(PackageName, TEXT(""), PackageName, Name);
-
-		NewNameSuggestion = FPaths::GetCleanFilename(PackageName);
-	}
+	check(!NewNameSuggestion.IsEmpty());
 
 	FSaveAssetDialogConfig SaveAssetDialogConfig;
 	SaveAssetDialogConfig.DialogTitleOverride = LOCTEXT("SaveLevelDialogTitle", "Save Level As");
@@ -638,6 +629,21 @@ static bool SaveAsImplementation( UWorld* InWorld, const FString& DefaultFilenam
 {
 	UEditorLoadingSavingSettings* LoadingSavingSettings = GetMutableDefault<UEditorLoadingSavingSettings>();
 
+	// Get default path and filename. If no default filename was supplied, create one.
+	FString DefaultDirectory = FEditorDirectories::Get().GetLastDirectory(ELastDirectory::LEVEL);
+	FString Filename = FPaths::GetCleanFilename(DefaultFilename);
+	if (Filename.IsEmpty())
+	{
+		const FString DefaultName = TEXT("NewMap");
+		FString PackageName;
+		FPackageName::TryConvertFilenameToLongPackageName(DefaultDirectory / DefaultName, PackageName);
+		FString Name;
+		FAssetToolsModule& AssetToolsModule = FModuleManager::LoadModuleChecked<FAssetToolsModule>("AssetTools");
+		AssetToolsModule.Get().CreateUniqueAssetName(PackageName, TEXT(""), PackageName, Name);
+
+		Filename = FPaths::GetCleanFilename(FPackageName::LongPackageNameToFilename(PackageName));
+	}
+
 	// Disable autosaving while the "Save As..." dialog is up.
 	const bool bOldAutoSaveState = LoadingSavingSettings->bAutoSaveEnable;
 	LoadingSavingSettings->bAutoSaveEnable = false;
@@ -646,7 +652,7 @@ static bool SaveAsImplementation( UWorld* InWorld, const FString& DefaultFilenam
 
 	// Loop through until a valid filename is given or the user presses cancel
 	bool bFilenameIsValid = false;
-	FString DefaultDirectory = FEditorDirectories::Get().GetLastDirectory(ELastDirectory::LEVEL);
+
 	while( !bFilenameIsValid )
 	{
 		FString SaveFilename;
@@ -654,12 +660,12 @@ static bool SaveAsImplementation( UWorld* InWorld, const FString& DefaultFilenam
 		if (UEditorEngine::IsUsingWorldAssets())
 		{
 			FString DefaultPackagePath;
-			FPackageName::TryConvertFilenameToLongPackageName(DefaultDirectory, DefaultPackagePath);
+			FPackageName::TryConvertFilenameToLongPackageName(DefaultDirectory / Filename, DefaultPackagePath);
 
 			FString PackageName;
 			bSaveFileLocationSelected = OpenLevelSaveAsDialog(
-				DefaultPackagePath,
-				FPaths::GetBaseFilename(DefaultFilename),
+				FPackageName::GetLongPackagePath(DefaultPackagePath),
+				FPaths::GetBaseFilename(Filename),
 				PackageName);
 
 			if ( bSaveFileLocationSelected )
@@ -673,7 +679,7 @@ static bool SaveAsImplementation( UWorld* InWorld, const FString& DefaultFilenam
 				NSLOCTEXT("UnrealEd", "SaveAs", "Save As").ToString(),
 				FEditorFileUtils::GetFilterString(FI_Save),
 				DefaultDirectory,
-				FPaths::GetCleanFilename(DefaultFilename),
+				FPaths::GetCleanFilename(Filename),
 				SaveFilename
 				);
 		}
@@ -722,7 +728,7 @@ static bool SaveAsImplementation( UWorld* InWorld, const FString& DefaultFilenam
 
 			if( bAllowStreamingLevelRename )
 			{
-				const FString OldLevelName = FPaths::GetBaseFilename(DefaultFilename);
+				const FString OldLevelName = FPaths::GetBaseFilename(Filename);
 				const FString NewLevelName = FPaths::GetBaseFilename(SaveFilename);
 
 				// The old and new level names must have a common suffix.  We'll detect that now.
@@ -800,7 +806,7 @@ static bool SaveAsImplementation( UWorld* InWorld, const FString& DefaultFilenam
 				const int32 DlgResult =
 					FMessageDialog::Open( EAppMsgType::YesNoCancel,
 					FText::Format( NSLOCTEXT("UnrealEd", "SaveLevelAs_PromptToRenameStreamingLevels_F", "Would you like to update references to streaming levels and rename those as well?\n\nIf you select Yes, references to streaming levels in {0} will be renamed to {1} (including Level Blueprint level name references.)  You should also do this for each of your streaming level maps.\n\nIf you select No, the level will be saved with the specified name and no other changes will be made." ),
-					FText::FromString(FPaths::GetBaseFilename(DefaultFilename)), FText::FromString(FPaths::GetBaseFilename(SaveFilename)) ) );
+					FText::FromString(FPaths::GetBaseFilename(Filename)), FText::FromString(FPaths::GetBaseFilename(SaveFilename)) ) );
 
 				if( DlgResult != EAppReturnType::Cancel )	// Cancel?
 				{
@@ -2346,11 +2352,11 @@ static int32 InternalSavePackage( UPackage* PackageToSave, bool& bOutPackageLoca
 			if( UEditorEngine::IsUsingWorldAssets() )
 			{
 				FString DefaultPackagePath;
-				FPackageName::TryConvertFilenameToLongPackageName(DefaultLocation, DefaultPackagePath);
+				FPackageName::TryConvertFilenameToLongPackageName(DefaultLocation / FinalPackageFilename, DefaultPackagePath);
 
 				FString SaveAsPackageName;
 				bSaveFile = OpenLevelSaveAsDialog(
-					DefaultPackagePath,
+					FPackageName::GetLongPackagePath(DefaultPackagePath),
 					FPaths::GetBaseFilename(FinalPackageFilename),
 					SaveAsPackageName);
 
