Added config migration path for newer versions of the engine.

diff --git a/Engine/Source/Editor/UnrealEd/Private/Settings/EditorSettings.cpp b/Engine/Source/Editor/UnrealEd/Private/Settings/EditorSettings.cpp
new file mode 100644
--- /dev/null
+++ b/Engine/Source/Editor/UnrealEd/Private/Settings/EditorSettings.cpp
@@ -0,0 +1,90 @@
+// Copyright 1998-2015 Epic Games, Inc. All Rights Reserved.
+
+#include "UnrealEd.h"
+#include "EditorSettings.h"
+
+UEditorSettings::UEditorSettings(const FObjectInitializer& ObjectInitializer)
+	: Super(ObjectInitializer)
+{
+	bCopyStarterContentPreference = false;
+	bEditorAnalyticsEnabled = true;
+	AutoScalabilityWorkScaleAmount = 1;
+}
+
+void UEditorSettings::PostEditChangeProperty( struct FPropertyChangedEvent& PropertyChangedEvent)
+{
+	UProperty* PropertyThatChanged = PropertyChangedEvent.Property;
+
+	const FName Name = PropertyThatChanged ? PropertyThatChanged->GetFName() : NAME_None;
+	if (Name == FName(TEXT("bLoadTheMostRecentlyLoadedProjectAtStartup")))
+	{
+		const FString& AutoLoadProjectFileName = IProjectManager::Get().GetAutoLoadProjectFileName();
+		if ( bLoadTheMostRecentlyLoadedProjectAtStartup )
+		{
+			// Form or overwrite the file that is read at load to determine the most recently loaded project file
+			FFileHelper::SaveStringToFile(FPaths::GetProjectFilePath(), *AutoLoadProjectFileName);
+		}
+		else
+		{
+			// Remove the file. It's possible for bLoadTheMostRecentlyLoadedProjectAtStartup to be set before FPaths::GetProjectFilePath() is valid, so we need to distinguish the two cases.
+			IFileManager::Get().Delete(*AutoLoadProjectFileName);
+		}
+	}
+
+
+	SaveConfig(CPF_Config);
+}
+
+void UEditorSettings::LoadScalabilityBenchmark()
+{
+	check(!GEditorSettingsIni.IsEmpty());
+
+	const TCHAR* Section = TEXT("EngineBenchmarkResult");
+
+	Scalability::FQualityLevels Temporary;
+
+	if (IsScalabilityBenchmarkValid())
+	{
+		GConfig->GetInt(Section, TEXT("ResolutionQuality"),		Temporary.ResolutionQuality,	GEditorSettingsIni);
+		GConfig->GetInt(Section, TEXT("ViewDistanceQuality"),		Temporary.ViewDistanceQuality,	GEditorSettingsIni);
+		GConfig->GetInt(Section, TEXT("AntiAliasingQuality"),		Temporary.AntiAliasingQuality,	GEditorSettingsIni);
+		GConfig->GetInt(Section, TEXT("ShadowQuality"),			Temporary.ShadowQuality,		GEditorSettingsIni);
+		GConfig->GetInt(Section, TEXT("PostProcessQuality"),		Temporary.PostProcessQuality,	GEditorSettingsIni);
+		GConfig->GetInt(Section, TEXT("TextureQuality"),			Temporary.TextureQuality,		GEditorSettingsIni);
+		GConfig->GetInt(Section, TEXT("EffectsQuality"),			Temporary.EffectsQuality,		GEditorSettingsIni);
+		EngineBenchmarkResult = Temporary;
+	}
+}
+
+void UEditorSettings::AutoApplyScalabilityBenchmark()
+{
+	const TCHAR* Section = TEXT("EngineBenchmarkResult");
+
+	FScopedSlowTask SlowTask(0, NSLOCTEXT("UnrealEd", "RunningEngineBenchmark", "Running engine benchmark..."));
+	SlowTask.MakeDialog();
+
+
+	Scalability::FQualityLevels Temporary = Scalability::BenchmarkQualityLevels(AutoScalabilityWorkScaleAmount);
+
+	GConfig->SetBool(Section, TEXT("Valid"), true, GEditorSettingsIni);
+	GConfig->SetInt(Section, TEXT("ResolutionQuality"), Temporary.ResolutionQuality, GEditorSettingsIni);
+	GConfig->SetInt(Section, TEXT("ViewDistanceQuality"), Temporary.ViewDistanceQuality, GEditorSettingsIni);
+	GConfig->SetInt(Section, TEXT("AntiAliasingQuality"), Temporary.AntiAliasingQuality, GEditorSettingsIni);
+	GConfig->SetInt(Section, TEXT("ShadowQuality"), Temporary.ShadowQuality, GEditorSettingsIni);
+	GConfig->SetInt(Section, TEXT("PostProcessQuality"), Temporary.PostProcessQuality, GEditorSettingsIni);
+	GConfig->SetInt(Section, TEXT("TextureQuality"), Temporary.TextureQuality, GEditorSettingsIni);
+	GConfig->SetInt(Section, TEXT("EffectsQuality"), Temporary.EffectsQuality, GEditorSettingsIni);
+
+	Scalability::SetQualityLevels(Temporary);
+	Scalability::SaveState(GEditorSettingsIni);
+}
+
+bool UEditorSettings::IsScalabilityBenchmarkValid() const
+{
+	const TCHAR* Section = TEXT("EngineBenchmarkResult");
+
+	bool bIsValid = false;
+	GConfig->GetBool(Section, TEXT("Valid"), bIsValid, GEditorSettingsIni);
+
+	return bIsValid;
+}
\ No newline at end of file
diff --git a/Engine/Source/Editor/UnrealEd/Private/Settings/SettingsClasses.cpp b/Engine/Source/Editor/UnrealEd/Private/Settings/SettingsClasses.cpp
--- a/Engine/Source/Editor/UnrealEd/Private/Settings/SettingsClasses.cpp
+++ b/Engine/Source/Editor/UnrealEd/Private/Settings/SettingsClasses.cpp
@@ -204,7 +204,7 @@ void UEditorLoadingSavingSettings::CheckSourceControlCompatability()
 			if (SettingsModule != nullptr)
 			{
 				// Ensure that the advanced properties are visible
-				GConfig->SetBool(TEXT("DetailCategoriesAdvanced"), TEXT("EditorLoadingSavingSettings.AutoReimport"), true, GEditorUserSettingsIni);
+				GConfig->SetBool(TEXT("DetailCategoriesAdvanced"), TEXT("EditorLoadingSavingSettings.AutoReimport"), true, GEditorPerProjectIni);
 				SettingsModule->ShowViewer("Editor", "General", "LoadingSaving");
 			}
 		});
diff --git a/Engine/Source/Editor/UnrealEd/Private/SnappingUtils.cpp b/Engine/Source/Editor/UnrealEd/Private/SnappingUtils.cpp
--- a/Engine/Source/Editor/UnrealEd/Private/SnappingUtils.cpp
+++ b/Engine/Source/Editor/UnrealEd/Private/SnappingUtils.cpp
@@ -91,7 +91,6 @@ public:
 	bool SnapToBSPVertex( FVector& Location, FVector GridBase, FRotator& Rotation );
 
 private:
-	static const UEditorUserSettings& GetEditorUserSettings();
 
 	/** Vertex snapping implementation */
 	FVertexSnappingImpl VertexSnappingImpl;
@@ -358,11 +357,6 @@ bool FEditorViewportSnapping::SnapToBSPVertex(FVector& Location, FVector GridBas
 	return bSnapped;
 }
 
-const UEditorUserSettings& FEditorViewportSnapping::GetEditorUserSettings()
-{
-	return GEditor->GetEditorUserSettings();
-}
-
 void FEditorViewportSnapping::ClearSnappingHelpers( bool bClearImmediately )
 {
 	VertexSnappingImpl.ClearSnappingHelpers( bClearImmediately );
diff --git a/Engine/Source/Editor/UnrealEd/Private/StaticLightingSystem/StaticLightingSystem.cpp b/Engine/Source/Editor/UnrealEd/Private/StaticLightingSystem/StaticLightingSystem.cpp
--- a/Engine/Source/Editor/UnrealEd/Private/StaticLightingSystem/StaticLightingSystem.cpp
+++ b/Engine/Source/Editor/UnrealEd/Private/StaticLightingSystem/StaticLightingSystem.cpp
@@ -1953,7 +1953,7 @@ bool FStaticLightingSystem::FinishLightmassProcess()
 			}
 		
 			bool bShowLightingBuildInfo = false;
-			GConfig->GetBool( TEXT("LightingBuildOptions"), TEXT("ShowLightingBuildInfo"), bShowLightingBuildInfo, GEditorUserSettingsIni );
+			GConfig->GetBool( TEXT("LightingBuildOptions"), TEXT("ShowLightingBuildInfo"), bShowLightingBuildInfo, GEditorPerProjectIni );
 			if( bShowLightingBuildInfo )
 			{
 				StatsViewerModule.GetPage(EStatsPage::LightingBuildInfo)->Show();
diff --git a/Engine/Source/Editor/UnrealEd/Private/Tests/AutomationEditorCommon.cpp b/Engine/Source/Editor/UnrealEd/Private/Tests/AutomationEditorCommon.cpp
--- a/Engine/Source/Editor/UnrealEd/Private/Tests/AutomationEditorCommon.cpp
+++ b/Engine/Source/Editor/UnrealEd/Private/Tests/AutomationEditorCommon.cpp
@@ -698,14 +698,14 @@ bool FBuildLightingCommand::Update()
 	FLightingBuildOptions LightingBuildOptions;
 
 	// Retrieve settings from ini.
-	GConfig->GetBool(TEXT("LightingBuildOptions"), TEXT("OnlyBuildSelected"), LightingBuildOptions.bOnlyBuildSelected, GEditorUserSettingsIni);
-	GConfig->GetBool(TEXT("LightingBuildOptions"), TEXT("OnlyBuildCurrentLevel"), LightingBuildOptions.bOnlyBuildCurrentLevel, GEditorUserSettingsIni);
-	GConfig->GetBool(TEXT("LightingBuildOptions"), TEXT("OnlyBuildSelectedLevels"), LightingBuildOptions.bOnlyBuildSelectedLevels, GEditorUserSettingsIni);
-	GConfig->GetBool(TEXT("LightingBuildOptions"), TEXT("OnlyBuildVisibility"), LightingBuildOptions.bOnlyBuildVisibility, GEditorUserSettingsIni);
-	GConfig->GetBool(TEXT("LightingBuildOptions"), TEXT("UseErrorColoring"), LightingBuildOptions.bUseErrorColoring, GEditorUserSettingsIni);
-	GConfig->GetBool(TEXT("LightingBuildOptions"), TEXT("ShowLightingBuildInfo"), LightingBuildOptions.bShowLightingBuildInfo, GEditorUserSettingsIni);
+	GConfig->GetBool(TEXT("LightingBuildOptions"), TEXT("OnlyBuildSelected"), LightingBuildOptions.bOnlyBuildSelected, GEditorPerProjectIni);
+	GConfig->GetBool(TEXT("LightingBuildOptions"), TEXT("OnlyBuildCurrentLevel"), LightingBuildOptions.bOnlyBuildCurrentLevel, GEditorPerProjectIni);
+	GConfig->GetBool(TEXT("LightingBuildOptions"), TEXT("OnlyBuildSelectedLevels"), LightingBuildOptions.bOnlyBuildSelectedLevels, GEditorPerProjectIni);
+	GConfig->GetBool(TEXT("LightingBuildOptions"), TEXT("OnlyBuildVisibility"), LightingBuildOptions.bOnlyBuildVisibility, GEditorPerProjectIni);
+	GConfig->GetBool(TEXT("LightingBuildOptions"), TEXT("UseErrorColoring"), LightingBuildOptions.bUseErrorColoring, GEditorPerProjectIni);
+	GConfig->GetBool(TEXT("LightingBuildOptions"), TEXT("ShowLightingBuildInfo"), LightingBuildOptions.bShowLightingBuildInfo, GEditorPerProjectIni);
 	int32 QualityLevel;
-	GConfig->GetInt(TEXT("LightingBuildOptions"), TEXT("QualityLevel"), QualityLevel, GEditorUserSettingsIni);
+	GConfig->GetInt(TEXT("LightingBuildOptions"), TEXT("QualityLevel"), QualityLevel, GEditorPerProjectIni);
 	QualityLevel = FMath::Clamp<int32>(QualityLevel, Quality_Preview, Quality_Production);
 	LightingBuildOptions.QualityLevel = Quality_Production;
 
diff --git a/Engine/Source/Editor/UnrealEd/Private/Tests/EditorBuildPromotionTests.cpp b/Engine/Source/Editor/UnrealEd/Private/Tests/EditorBuildPromotionTests.cpp
--- a/Engine/Source/Editor/UnrealEd/Private/Tests/EditorBuildPromotionTests.cpp
+++ b/Engine/Source/Editor/UnrealEd/Private/Tests/EditorBuildPromotionTests.cpp
@@ -244,8 +244,8 @@ namespace EditorBuildPromotionTestUtils
 	static void ExportEditorSettings(const FString& TargetFilename)
 	{
 		FInputBindingManager::Get().SaveInputBindings();
-		GConfig->Flush(false, GEditorUserSettingsIni);
-		IFileManager::Get().Copy(*TargetFilename, *GEditorUserSettingsIni);
+		GConfig->Flush(false, GEditorPerProjectIni);
+		IFileManager::Get().Copy(*TargetFilename, *GEditorPerProjectIni);
 	}
 
 	/**
@@ -255,9 +255,9 @@ namespace EditorBuildPromotionTestUtils
 	*/
 	static void ImportEditorSettings(const FString& TargetFilename)
 	{
-		GConfig->Flush(true, GEditorUserSettingsIni);
-		IFileManager::Get().Copy(*GEditorUserSettingsIni, *TargetFilename);
-		GConfig->LoadFile(GEditorUserSettingsIni);
+		GConfig->Flush(true, GEditorPerProjectIni);
+		IFileManager::Get().Copy(*GEditorPerProjectIni, *TargetFilename);
+		GConfig->LoadFile(GEditorPerProjectIni);
 	}
 
 	/**
@@ -922,14 +922,14 @@ namespace EditorBuildPromotionTestUtils
 		FLightingBuildOptions LightingBuildOptions;
 
 		// Retrieve settings from ini.
-		GConfig->GetBool(TEXT("LightingBuildOptions"), TEXT("OnlyBuildSelected"), LightingBuildOptions.bOnlyBuildSelected, GEditorUserSettingsIni);
-		GConfig->GetBool(TEXT("LightingBuildOptions"), TEXT("OnlyBuildCurrentLevel"), LightingBuildOptions.bOnlyBuildCurrentLevel, GEditorUserSettingsIni);
-		GConfig->GetBool(TEXT("LightingBuildOptions"), TEXT("OnlyBuildSelectedLevels"), LightingBuildOptions.bOnlyBuildSelectedLevels, GEditorUserSettingsIni);
-		GConfig->GetBool(TEXT("LightingBuildOptions"), TEXT("OnlyBuildVisibility"), LightingBuildOptions.bOnlyBuildVisibility, GEditorUserSettingsIni);
-		GConfig->GetBool(TEXT("LightingBuildOptions"), TEXT("UseErrorColoring"), LightingBuildOptions.bUseErrorColoring, GEditorUserSettingsIni);
-		GConfig->GetBool(TEXT("LightingBuildOptions"), TEXT("ShowLightingBuildInfo"), LightingBuildOptions.bShowLightingBuildInfo, GEditorUserSettingsIni);
+		GConfig->GetBool(TEXT("LightingBuildOptions"), TEXT("OnlyBuildSelected"), LightingBuildOptions.bOnlyBuildSelected, GEditorPerProjectIni);
+		GConfig->GetBool(TEXT("LightingBuildOptions"), TEXT("OnlyBuildCurrentLevel"), LightingBuildOptions.bOnlyBuildCurrentLevel, GEditorPerProjectIni);
+		GConfig->GetBool(TEXT("LightingBuildOptions"), TEXT("OnlyBuildSelectedLevels"), LightingBuildOptions.bOnlyBuildSelectedLevels, GEditorPerProjectIni);
+		GConfig->GetBool(TEXT("LightingBuildOptions"), TEXT("OnlyBuildVisibility"), LightingBuildOptions.bOnlyBuildVisibility, GEditorPerProjectIni);
+		GConfig->GetBool(TEXT("LightingBuildOptions"), TEXT("UseErrorColoring"), LightingBuildOptions.bUseErrorColoring, GEditorPerProjectIni);
+		GConfig->GetBool(TEXT("LightingBuildOptions"), TEXT("ShowLightingBuildInfo"), LightingBuildOptions.bShowLightingBuildInfo, GEditorPerProjectIni);
 		int32 QualityLevel;
-		GConfig->GetInt(TEXT("LightingBuildOptions"), TEXT("QualityLevel"), QualityLevel, GEditorUserSettingsIni);
+		GConfig->GetInt(TEXT("LightingBuildOptions"), TEXT("QualityLevel"), QualityLevel, GEditorPerProjectIni);
 		QualityLevel = FMath::Clamp<int32>(QualityLevel, Quality_Preview, Quality_Production);
 		LightingBuildOptions.QualityLevel = (ELightingBuildQuality)QualityLevel;
 
@@ -2319,7 +2319,7 @@ namespace BuildPromotionTestHelper
 		bool Lighting_BuildLighting_Part1()
 		{
 			//Set production quality
-			GConfig->SetInt(TEXT("LightingBuildOptions"), TEXT("QualityLevel"), (int32)ELightingBuildQuality::Quality_Production, GEditorUserSettingsIni);
+			GConfig->SetInt(TEXT("LightingBuildOptions"), TEXT("QualityLevel"), (int32)ELightingBuildQuality::Quality_Production, GEditorPerProjectIni);
 			UE_LOG(LogEditorBuildPromotionTests, Display, TEXT("Set the lighting quality to Production"));
 
 			//Force AutoApplyLighting on
diff --git a/Engine/Source/Editor/UnrealEd/Private/Toolkits/AssetEditorManager.cpp b/Engine/Source/Editor/UnrealEd/Private/Toolkits/AssetEditorManager.cpp
--- a/Engine/Source/Editor/UnrealEd/Private/Toolkits/AssetEditorManager.cpp
+++ b/Engine/Source/Editor/UnrealEd/Private/Toolkits/AssetEditorManager.cpp
@@ -486,10 +486,10 @@ void FAssetEditorManager::RequestRestorePreviouslyOpenAssets()
 void FAssetEditorManager::RestorePreviouslyOpenAssets()
 {
 	TArray<FString> OpenAssets;
-	GConfig->GetArray(TEXT("AssetEditorManager"), TEXT("OpenAssetsAtExit"), OpenAssets, GEditorUserSettingsIni);
+	GConfig->GetArray(TEXT("AssetEditorManager"), TEXT("OpenAssetsAtExit"), OpenAssets, GEditorPerProjectIni);
 
 	bool bCleanShutdown =  false;
-	GConfig->GetBool(TEXT("AssetEditorManager"), TEXT("CleanShutdown"), bCleanShutdown, GEditorUserSettingsIni);
+	GConfig->GetBool(TEXT("AssetEditorManager"), TEXT("CleanShutdown"), bCleanShutdown, GEditorPerProjectIni);
 
 	SaveOpenAssetEditors(false);
 
@@ -509,7 +509,7 @@ void FAssetEditorManager::RestorePreviouslyOpenAssets()
 			{
 				// Has this notification previously been suppressed by the user?
 				bool bSuppressNotification = false;
-				GConfig->GetBool(TEXT("AssetEditorManager"), TEXT("SuppressRestorePreviouslyOpenAssetsNotification"), bSuppressNotification, GEditorUserSettingsIni);
+				GConfig->GetBool(TEXT("AssetEditorManager"), TEXT("SuppressRestorePreviouslyOpenAssetsNotification"), bSuppressNotification, GEditorPerProjectIni);
 
 				if(!bSuppressNotification)
 				{
@@ -534,14 +534,14 @@ void FAssetEditorManager::SpawnRestorePreviouslyOpenAssetsNotification(const boo
 		static ECheckBoxState GetDontAskAgainCheckBoxState()
 		{
 			bool bSuppressNotification = false;
-			GConfig->GetBool(TEXT("AssetEditorManager"), TEXT("SuppressRestorePreviouslyOpenAssetsNotification"), bSuppressNotification, GEditorUserSettingsIni);
+			GConfig->GetBool(TEXT("AssetEditorManager"), TEXT("SuppressRestorePreviouslyOpenAssetsNotification"), bSuppressNotification, GEditorPerProjectIni);
 			return bSuppressNotification ? ECheckBoxState::Checked : ECheckBoxState::Unchecked;
 		}
 
 		static void OnDontAskAgainCheckBoxStateChanged(ECheckBoxState NewState)
 		{
 			const bool bSuppressNotification = (NewState == ECheckBoxState::Checked);
-			GConfig->SetBool(TEXT("AssetEditorManager"), TEXT("SuppressRestorePreviouslyOpenAssetsNotification"), bSuppressNotification, GEditorUserSettingsIni);
+			GConfig->SetBool(TEXT("AssetEditorManager"), TEXT("SuppressRestorePreviouslyOpenAssetsNotification"), bSuppressNotification, GEditorPerProjectIni);
 		}
 	};
 
@@ -604,7 +604,7 @@ void FAssetEditorManager::OnConfirmRestorePreviouslyOpenAssets(TArray<FString> A
 		// If the user suppressed the notification for future sessions, make sure this is reflected in their settings
 		// Note: We do that inside this if statement so that we only do it if we were showing a UI they could interact with
 		bool bSuppressNotification = false;
-		GConfig->GetBool(TEXT("AssetEditorManager"), TEXT("SuppressRestorePreviouslyOpenAssetsNotification"), bSuppressNotification, GEditorUserSettingsIni);
+		GConfig->GetBool(TEXT("AssetEditorManager"), TEXT("SuppressRestorePreviouslyOpenAssetsNotification"), bSuppressNotification, GEditorPerProjectIni);
 		UEditorLoadingSavingSettings& Settings = *GetMutableDefault<UEditorLoadingSavingSettings>();
 		Settings.bRestoreOpenAssetTabsOnRestart = bSuppressNotification;
 		Settings.PostEditChange();
@@ -655,9 +655,9 @@ void FAssetEditorManager::SaveOpenAssetEditors(bool bOnShutdown)
 			}
 		}
 
-		GConfig->SetArray(TEXT("AssetEditorManager"), TEXT("OpenAssetsAtExit"), OpenAssets, GEditorUserSettingsIni);
-		GConfig->SetBool(TEXT("AssetEditorManager"), TEXT("CleanShutdown"), bOnShutdown, GEditorUserSettingsIni);
-		GConfig->Flush(false, GEditorUserSettingsIni);
+		GConfig->SetArray(TEXT("AssetEditorManager"), TEXT("OpenAssetsAtExit"), OpenAssets, GEditorPerProjectIni);
+		GConfig->SetBool(TEXT("AssetEditorManager"), TEXT("CleanShutdown"), bOnShutdown, GEditorPerProjectIni);
+		GConfig->Flush(false, GEditorPerProjectIni);
 	}
 }
 
diff --git a/Engine/Source/Editor/UnrealEd/Private/Toolkits/AssetEditorToolkit.cpp b/Engine/Source/Editor/UnrealEd/Private/Toolkits/AssetEditorToolkit.cpp
--- a/Engine/Source/Editor/UnrealEd/Private/Toolkits/AssetEditorToolkit.cpp
+++ b/Engine/Source/Editor/UnrealEd/Private/Toolkits/AssetEditorToolkit.cpp
@@ -100,7 +100,7 @@ void FAssetEditorToolkit::InitAssetEditor( const EToolkitMode::Type Mode, const
 				TEXT("AssetEditorToolkitTabLocation"), 
 				*ObjectsToEdit[0]->GetPathName(), 
 				reinterpret_cast<int32&>(SavedAssetEditorToolkitTabLocation), 
-				GEditorUserSettingsIni
+				GEditorPerProjectIni
 				);
 
 			const FName AssetEditorToolkitTab = (SavedAssetEditorToolkitTabLocation == EAssetEditorToolkitTabLocation::Docked) ? "DockedToolkit" : "StandaloneToolkit";
diff --git a/Engine/Source/Editor/UnrealEd/Private/Toolkits/SStandaloneAssetEditorToolkitHost.cpp b/Engine/Source/Editor/UnrealEd/Private/Toolkits/SStandaloneAssetEditorToolkitHost.cpp
--- a/Engine/Source/Editor/UnrealEd/Private/Toolkits/SStandaloneAssetEditorToolkitHost.cpp
+++ b/Engine/Source/Editor/UnrealEd/Private/Toolkits/SStandaloneAssetEditorToolkitHost.cpp
@@ -368,7 +368,7 @@ void SStandaloneAssetEditorToolkitHost::OnTabClosed(TSharedRef<SDockTab> TabClos
 						TEXT("AssetEditorToolkitTabLocation"), 
 						*ObjectBeingEdited->GetPathName(), 
 						static_cast<int32>(AssetEditorToolkitTabLocation), 
-						GEditorUserSettingsIni
+						GEditorPerProjectIni
 						);
 				}
 			}
diff --git a/Engine/Source/Editor/UnrealEd/Private/UnrealEdEngine.cpp b/Engine/Source/Editor/UnrealEd/Private/UnrealEdEngine.cpp
--- a/Engine/Source/Editor/UnrealEd/Private/UnrealEdEngine.cpp
+++ b/Engine/Source/Editor/UnrealEd/Private/UnrealEdEngine.cpp
@@ -51,7 +51,7 @@ void UUnrealEdEngine::Init(IEngineLoop* InEngineLoop)
 	PackageAutoSaver->LoadRestoreFile();
 
 #if !UE_BUILD_DEBUG
-	if( !GEditorGameAgnosticIni.IsEmpty() )
+	if( !GEditorSettingsIni.IsEmpty() )
 	{
 		// We need the game agnostic ini for this code
 		PerformanceMonitor = new FPerformanceMonitor;
@@ -820,7 +820,7 @@ void UUnrealEdEngine::CloseEditor()
 
 bool UUnrealEdEngine::AllowSelectTranslucent() const
 {
-	return GEditor->GetEditorUserSettings().bAllowSelectTranslucent;
+	return GetDefault<UEditorPerProjectUserSettings>()->bAllowSelectTranslucent;
 }
 
 
diff --git a/Engine/Source/Editor/UnrealEd/Private/UnrealEdMisc.cpp b/Engine/Source/Editor/UnrealEd/Private/UnrealEdMisc.cpp
--- a/Engine/Source/Editor/UnrealEd/Private/UnrealEdMisc.cpp
+++ b/Engine/Source/Editor/UnrealEd/Private/UnrealEdMisc.cpp
@@ -1405,7 +1405,7 @@ void FUnrealEdMisc::SwitchProject(const FString& GameOrProjectFileName, bool bWa
 		const FText Message = FText::Format( LOCTEXT( "SwitchProjectWarning", "The editor will restart to switch to the {CurrentProjectName} project.  You will be prompted to save any changes before the editor restarts.  Continue switching projects?" ), Arguments ); 
 
 		// Present the user with a warning that changing projects has to restart the editor
-		FSuppressableWarningDialog::FSetupInfo Info( Message, Title, "Warning_SwitchProject", GEditorGameAgnosticIni );
+		FSuppressableWarningDialog::FSetupInfo Info( Message, Title, "Warning_SwitchProject", GEditorSettingsIni );
 		Info.ConfirmText = LOCTEXT( "Yes", "Yes" );
 		Info.CancelText = LOCTEXT( "No", "No" );
 
diff --git a/Engine/Source/Editor/UnrealEd/Private/UnrealEdSrv.cpp b/Engine/Source/Editor/UnrealEd/Private/UnrealEdSrv.cpp
--- a/Engine/Source/Editor/UnrealEd/Private/UnrealEdSrv.cpp
+++ b/Engine/Source/Editor/UnrealEd/Private/UnrealEdSrv.cpp
@@ -276,7 +276,7 @@ UPackage* UUnrealEdEngine::GeneratePackageThumbnailsIfRequired( const TCHAR* Str
 				}
 
 				bool bPrintThumbnailDiagnostics = false;
-				GConfig->GetBool(TEXT("Thumbnails"), TEXT("Debug"), bPrintThumbnailDiagnostics, GEditorUserSettingsIni);
+				GConfig->GetBool(TEXT("Thumbnails"), TEXT("Debug"), bPrintThumbnailDiagnostics, GEditorPerProjectIni);
 
 				const FObjectThumbnail* ExistingThumbnail = ThumbnailTools::FindCachedThumbnail( CurObject->GetFullName() );
 				if (bPrintThumbnailDiagnostics)
diff --git a/Engine/Source/Editor/UnrealEd/Public/Dialogs/Dialogs.h b/Engine/Source/Editor/UnrealEd/Public/Dialogs/Dialogs.h
--- a/Engine/Source/Editor/UnrealEd/Public/Dialogs/Dialogs.h
+++ b/Engine/Source/Editor/UnrealEd/Public/Dialogs/Dialogs.h
@@ -98,9 +98,9 @@ public:
 		* @param Prompt					The message that appears to the user
 		* @param Title					The title of the dialog
 		* @param InIniSettingName		The name of the entry in the INI where the state of the "Disable this warning" check box is stored
-		* @param InIniSettingFileName	The name of the INI where the state of the InIniSettingName flag is stored (defaults to GEditorUserSettingsIni)
+		* @param InIniSettingFileName	The name of the INI where the state of the InIniSettingName flag is stored (defaults to GEditorPerProjectIni)
 		*/
-		FSetupInfo(const FText& InMessage, const FText& InTitle, const FString& InIniSettingName, const FString& InIniSettingFileName=GEditorUserSettingsIni )
+		FSetupInfo(const FText& InMessage, const FText& InTitle, const FString& InIniSettingName, const FString& InIniSettingFileName=GEditorPerProjectIni )
 			: Message(InMessage)
 			, Title(InTitle)
 			, IniSettingName(InIniSettingName)
diff --git a/Engine/Source/Editor/UnrealEd/Public/Settings/EditorSettings.h b/Engine/Source/Editor/UnrealEd/Public/Settings/EditorSettings.h
new file mode 100644
--- /dev/null
+++ b/Engine/Source/Editor/UnrealEd/Public/Settings/EditorSettings.h
@@ -0,0 +1,61 @@
+// Copyright 1998-2015 Epic Games, Inc. All Rights Reserved.
+
+#pragma once
+#include "Scalability.h"
+#include "EditorSettings.generated.h"
+
+UCLASS(config=EditorSettings)
+class UNREALED_API UEditorSettings : public UObject
+{
+	GENERATED_UCLASS_BODY()
+
+	/** When checked, the most recently loaded project will be auto-loaded at editor startup if no other project was specified on the command line */
+	UPROPERTY(EditAnywhere, Category=General)
+	bool bLoadTheMostRecentlyLoadedProjectAtStartup; // Note that this property is NOT config since it is not necessary to save the value to ini. It is determined at startup in UEditorEngine::InitEditor().
+
+	/** Can the editor report usage analytics (types of assets being spawned, etc...) back to Epic in order for us to improve the editor user experience?  Note: The editor must be restarted for changes to take effect. */
+	UPROPERTY(EditAnywhere, config, Category=General, AdvancedDisplay)
+	bool bEditorAnalyticsEnabled;
+
+	// =====================================================================
+	// The following options are NOT exposed in the preferences Editor
+	// (usually because there is a different way to set them interactively!)
+
+	/** Game project files that were recently opened in the editor */
+	UPROPERTY(config)
+	TArray<FString> RecentlyOpenedProjectFiles;
+
+	/** The paths of projects created with the new project wizard. This is used to populate the "Path" field of the new project dialog. */
+	UPROPERTY(config)
+	TArray<FString> CreatedProjectPaths;
+
+	UPROPERTY(config)
+	bool bCopyStarterContentPreference;
+
+	/** The id's of the surveys completed */
+	UPROPERTY(config)
+	TArray<FGuid> CompletedSurveys;
+
+	/** The id's of the surveys currently in-progress */
+	UPROPERTY(config)
+	TArray<FGuid> InProgressSurveys;
+
+	UPROPERTY(config)
+	float AutoScalabilityWorkScaleAmount;
+
+	/** Engine scalability benchmark results */
+	Scalability::FQualityLevels EngineBenchmarkResult;
+
+	/** Load the engine scalability benchmark results. Performs a benchmark if not yet valid. */
+	void LoadScalabilityBenchmark();
+
+	/** Auto detects and applies the scalability benchmark */
+	void AutoApplyScalabilityBenchmark();
+
+	/** @return true if the scalability benchmark is valid */
+	bool IsScalabilityBenchmarkValid() const;
+
+	// Begin UObject Interface
+	virtual void PostEditChangeProperty( struct FPropertyChangedEvent& PropertyChangedEvent) override;
+	// End UObject Interface
+};
\ No newline at end of file
diff --git a/Engine/Source/Editor/WorldBrowser/Private/LevelCollectionModel.cpp b/Engine/Source/Editor/WorldBrowser/Private/LevelCollectionModel.cpp
--- a/Engine/Source/Editor/WorldBrowser/Private/LevelCollectionModel.cpp
+++ b/Engine/Source/Editor/WorldBrowser/Private/LevelCollectionModel.cpp
@@ -1289,22 +1289,22 @@ void FLevelCollectionModel::LoadSettings()
 {
 	// Display paths
 	bool bDisplayPathsSetting = false;
-	GConfig->GetBool(*ConfigIniSection, TEXT("DisplayPaths"), bDisplayPathsSetting, GEditorUserSettingsIni);
+	GConfig->GetBool(*ConfigIniSection, TEXT("DisplayPaths"), bDisplayPathsSetting, GEditorPerProjectIni);
 	SetDisplayPathsState(bDisplayPathsSetting);
 
 	// Display actors count
 	bool bDisplayActorsCountSetting = false;
-	GConfig->GetBool(*ConfigIniSection, TEXT("DisplayActorsCount"), bDisplayActorsCountSetting, GEditorUserSettingsIni);
+	GConfig->GetBool(*ConfigIniSection, TEXT("DisplayActorsCount"), bDisplayActorsCountSetting, GEditorPerProjectIni);
 	SetDisplayActorsCountState(bDisplayActorsCountSetting);
 }
 	
 void FLevelCollectionModel::SaveSettings()
 {
 	// Display paths
-	GConfig->SetBool(*ConfigIniSection, TEXT("DisplayPaths"), GetDisplayPathsState(), GEditorUserSettingsIni);
+	GConfig->SetBool(*ConfigIniSection, TEXT("DisplayPaths"), GetDisplayPathsState(), GEditorPerProjectIni);
 
 	// Display actors count
-	GConfig->SetBool(*ConfigIniSection, TEXT("DisplayActorsCount"), GetDisplayActorsCountState(), GEditorUserSettingsIni);
+	GConfig->SetBool(*ConfigIniSection, TEXT("DisplayActorsCount"), GetDisplayActorsCountState(), GEditorPerProjectIni);
 }
 
 void FLevelCollectionModel::RefreshBrowser_Executed()
diff --git a/Engine/Source/Runtime/AIModule/Private/BehaviorTree/BehaviorTreeComponent.cpp b/Engine/Source/Runtime/AIModule/Private/BehaviorTree/BehaviorTreeComponent.cpp
--- a/Engine/Source/Runtime/AIModule/Private/BehaviorTree/BehaviorTreeComponent.cpp
+++ b/Engine/Source/Runtime/AIModule/Private/BehaviorTree/BehaviorTreeComponent.cpp
@@ -2243,7 +2243,7 @@ bool UBehaviorTreeComponent::IsDebuggerActive()
 
 		if (GFrameCounter != PrevFrameCounter)
 		{
-			GConfig->GetBool(TEXT("/Script/UnrealEd.EditorUserSettings"), TEXT("bAlwaysGatherBehaviorTreeDebuggerData"), bAlwaysGatherData, GEditorUserSettingsIni);
+			GConfig->GetBool(TEXT("/Script/UnrealEd.EditorPerProjectUserSettings"), TEXT("bAlwaysGatherBehaviorTreeDebuggerData"), bAlwaysGatherData, GEditorPerProjectIni);
 			PrevFrameCounter = GFrameCounter;
 		}
 
diff --git a/Engine/Source/Runtime/AppFramework/Private/Widgets/Colors/SColorPicker.cpp b/Engine/Source/Runtime/AppFramework/Private/Widgets/Colors/SColorPicker.cpp
--- a/Engine/Source/Runtime/AppFramework/Private/Widgets/Colors/SColorPicker.cpp
+++ b/Engine/Source/Runtime/AppFramework/Private/Widgets/Colors/SColorPicker.cpp
@@ -81,13 +81,13 @@ void SColorPicker::Construct( const FArguments& InArgs )
 
 	bool bAdvancedSectionExpanded = false;
 
-	if (!FPaths::FileExists(GEditorUserSettingsIni))
+	if (!FPaths::FileExists(GEditorPerProjectIni))
 	{
 		bool WheelMode = true;
 
-		GConfig->GetBool(TEXT("ColorPickerUI"), TEXT("bWheelMode"), WheelMode, GEditorUserSettingsIni);
-		GConfig->GetBool(TEXT("ColorPickerUI"), TEXT("bAdvancedSectionExpanded"), bAdvancedSectionExpanded, GEditorUserSettingsIni);
-		GConfig->GetBool(TEXT("ColorPickerUI"), TEXT("bSRGBEnabled"), SColorThemesViewer::bSRGBEnabled, GEditorUserSettingsIni);
+		GConfig->GetBool(TEXT("ColorPickerUI"), TEXT("bWheelMode"), WheelMode, GEditorPerProjectIni);
+		GConfig->GetBool(TEXT("ColorPickerUI"), TEXT("bAdvancedSectionExpanded"), bAdvancedSectionExpanded, GEditorPerProjectIni);
+		GConfig->GetBool(TEXT("ColorPickerUI"), TEXT("bSRGBEnabled"), SColorThemesViewer::bSRGBEnabled, GEditorPerProjectIni);
 
 		CurrentMode = WheelMode ? EColorPickerModes::Wheel : EColorPickerModes::Spectrum;
 	}
@@ -1124,9 +1124,9 @@ FLinearColor SColorPicker::GetGradientStartColor( EColorPickerChannels Channel )
 
 void SColorPicker::HandleAdvancedAreaExpansionChanged( bool Expanded )
 {
-	if (FPaths::FileExists(GEditorUserSettingsIni))
+	if (FPaths::FileExists(GEditorPerProjectIni))
 	{
-		GConfig->SetBool(TEXT("ColorPickerUI"), TEXT("bAdvancedSectionExpanded"), Expanded, GEditorUserSettingsIni);
+		GConfig->SetBool(TEXT("ColorPickerUI"), TEXT("bAdvancedSectionExpanded"), Expanded, GEditorPerProjectIni);
 	}
 }
 
@@ -1323,9 +1323,9 @@ FReply SColorPicker::HandleColorPickerModeButtonClicked()
 {
 	CycleMode();
 
-	if (FPaths::FileExists(GEditorUserSettingsIni))
+	if (FPaths::FileExists(GEditorPerProjectIni))
 	{
-		GConfig->SetBool(TEXT("ColorPickerUI"), TEXT("bWheelMode"), CurrentMode == EColorPickerModes::Wheel, GEditorUserSettingsIni);
+		GConfig->SetBool(TEXT("ColorPickerUI"), TEXT("bWheelMode"), CurrentMode == EColorPickerModes::Wheel, GEditorPerProjectIni);
 	}
 
 	return FReply::Handled();
@@ -1426,9 +1426,9 @@ void SColorPicker::HandleSRGBCheckBoxCheckStateChanged( ECheckBoxState InIsCheck
 {
 	SColorThemesViewer::bSRGBEnabled = (InIsChecked == ECheckBoxState::Checked);
 
-	if (FPaths::FileExists(GEditorUserSettingsIni))
+	if (FPaths::FileExists(GEditorPerProjectIni))
 	{
-		GConfig->SetBool(TEXT("ColorPickerUI"), TEXT("bSRGBEnabled"), SColorThemesViewer::bSRGBEnabled, GEditorUserSettingsIni);
+		GConfig->SetBool(TEXT("ColorPickerUI"), TEXT("bSRGBEnabled"), SColorThemesViewer::bSRGBEnabled, GEditorPerProjectIni);
 	}
 }
 
diff --git a/Engine/Source/Runtime/AppFramework/Private/Widgets/Colors/SColorThemes.cpp b/Engine/Source/Runtime/AppFramework/Private/Widgets/Colors/SColorThemes.cpp
--- a/Engine/Source/Runtime/AppFramework/Private/Widgets/Colors/SColorThemes.cpp
+++ b/Engine/Source/Runtime/AppFramework/Private/Widgets/Colors/SColorThemes.cpp
@@ -1108,13 +1108,13 @@ TSharedRef<ITableRow> SColorThemesViewer::OnGenerateColorThemeBars(TSharedPtr<FC
 
 void SColorThemesViewer::LoadColorThemesFromIni()
 {
-	if (FPaths::FileExists(GEditorUserSettingsIni))
+	if (FPaths::FileExists(GEditorPerProjectIni))
 	{
 		bool bThemesRemaining = true;
 		int32 ThemeID = 0;
 		while (bThemesRemaining)
 		{
-			const FString ThemeName = GConfig->GetStr(TEXT("ColorThemes"), *FString::Printf(TEXT("Theme%i"), ThemeID), GEditorUserSettingsIni);
+			const FString ThemeName = GConfig->GetStr(TEXT("ColorThemes"), *FString::Printf(TEXT("Theme%i"), ThemeID), GEditorPerProjectIni);
 			if (!ThemeName.IsEmpty())
 			{
 				TSharedPtr<FColorTheme> ColorTheme = GetColorTheme( ThemeName );
@@ -1123,7 +1123,7 @@ void SColorThemesViewer::LoadColorThemesFromIni()
 				int32 ColorID = 0;
 				while (bColorsRemaining)
 				{
-					const FString ColorString = GConfig->GetStr(TEXT("ColorThemes"), *FString::Printf(TEXT("Theme%iColor%i"), ThemeID, ColorID), GEditorUserSettingsIni);
+					const FString ColorString = GConfig->GetStr(TEXT("ColorThemes"), *FString::Printf(TEXT("Theme%iColor%i"), ThemeID, ColorID), GEditorPerProjectIni);
 					if (!ColorString.IsEmpty())
 					{
 						// Add the color if it hasn't already
@@ -1158,19 +1158,19 @@ void SColorThemesViewer::LoadColorThemesFromIni()
 
 void SColorThemesViewer::SaveColorThemesToIni()
 {
-	if (FPaths::FileExists(GEditorUserSettingsIni))
+	if (FPaths::FileExists(GEditorPerProjectIni))
 	{
-		GConfig->EmptySection(TEXT("ColorThemes"), GEditorUserSettingsIni);
+		GConfig->EmptySection(TEXT("ColorThemes"), GEditorPerProjectIni);
 		for (int32 ThemeIndex = 0; ThemeIndex < ColorThemes.Num(); ++ThemeIndex)
 		{
 			const TSharedPtr<FColorTheme>& Theme = ColorThemes[ThemeIndex];
-			GConfig->SetString(TEXT("ColorThemes"), *FString::Printf(TEXT("Theme%i"), ThemeIndex), *Theme->Name, GEditorUserSettingsIni);
+			GConfig->SetString(TEXT("ColorThemes"), *FString::Printf(TEXT("Theme%i"), ThemeIndex), *Theme->Name, GEditorPerProjectIni);
 
 			const TArray< TSharedPtr<FLinearColor> >& Colors = Theme->GetColors();
 			for (int32 ColorIndex = 0; ColorIndex < Colors.Num(); ++ColorIndex)
 			{
 				const TSharedPtr<FLinearColor>& Color = Colors[ColorIndex];
-				GConfig->SetString(TEXT("ColorThemes"), *FString::Printf(TEXT("Theme%iColor%i"), ThemeIndex, ColorIndex), *Color->ToString(), GEditorUserSettingsIni);
+				GConfig->SetString(TEXT("ColorThemes"), *FString::Printf(TEXT("Theme%iColor%i"), ThemeIndex, ColorIndex), *Color->ToString(), GEditorPerProjectIni);
 			}
 		}
 	}
diff --git a/Engine/Source/Runtime/Core/Private/Internationalization/TextLocalizationManager.cpp b/Engine/Source/Runtime/Core/Private/Internationalization/TextLocalizationManager.cpp
--- a/Engine/Source/Runtime/Core/Private/Internationalization/TextLocalizationManager.cpp
+++ b/Engine/Source/Runtime/Core/Private/Internationalization/TextLocalizationManager.cpp
@@ -57,7 +57,7 @@ void EndInitTextLocalization()
 #endif // !UE_BUILD_SHIPPING
 #if WITH_EDITOR
 		// See if we've been provided a culture override in the editor
-		if(GConfig->GetString( TEXT("Internationalization"), TEXT("Culture"), RequestedCultureName, GEditorGameAgnosticIni ))
+		if(GConfig->GetString( TEXT("Internationalization"), TEXT("Culture"), RequestedCultureName, GEditorSettingsIni ))
 		{
 			//UE_LOG(LogInit, Log, TEXT("Overriding culture %s w/ editor configuration."), *CultureName);
 		}
@@ -705,7 +705,7 @@ void FTextLocalizationManager::LoadLocalizationResourcesForCulture(const FString
 		LocalizationPaths += FPaths::GetToolTipLocalizationPaths();
 
 		bool bShouldLoadLocalizedPropertyNames = true;
-		if( !GConfig->GetBool( TEXT("Internationalization"), TEXT("ShouldLoadLocalizedPropertyNames"), bShouldLoadLocalizedPropertyNames, GEditorGameAgnosticIni ) )
+		if( !GConfig->GetBool( TEXT("Internationalization"), TEXT("ShouldLoadLocalizedPropertyNames"), bShouldLoadLocalizedPropertyNames, GEditorSettingsIni ) )
 		{
 			GConfig->GetBool( TEXT("Internationalization"), TEXT("ShouldLoadLocalizedPropertyNames"), bShouldLoadLocalizedPropertyNames, GEngineIni );
 		}
diff --git a/Engine/Source/Runtime/Core/Private/Misc/ConfigCacheIni.cpp b/Engine/Source/Runtime/Core/Private/Misc/ConfigCacheIni.cpp
--- a/Engine/Source/Runtime/Core/Private/Misc/ConfigCacheIni.cpp
+++ b/Engine/Source/Runtime/Core/Private/Misc/ConfigCacheIni.cpp
@@ -9,6 +9,7 @@
 #include "DefaultValueHelper.h"
 #include "EngineBuildSettings.h"
 #include "Paths.h"
+#include "ConfigManifest.h"
 
 #if WITH_EDITOR
 	#define INI_CACHE 1
@@ -2237,6 +2238,11 @@ void FConfigCacheIni::SetArray
 )
 {
 	FConfigFile* File = Find( Filename, 1 );
+	if (!File)
+	{
+		return;
+	}
+	
 	FConfigSection* Sec  = File->Find( Section );
 	if( !Sec )
 		Sec = &File->Add( Section, FConfigSection() );
@@ -2816,6 +2822,9 @@ static FString GetDestIniFilename(const TCHAR* BaseIniName, const TCHAR* Platfor
 
 void FConfigCacheIni::InitializeConfigSystem()
 {
+	// Perform any upgrade we need before we load any configuration files
+	FConfigManifest::UpgradeFromPreviousVersions();
+
 	// create GConfig
 	GConfig = new FConfigCacheIni(EConfigCacheType::DiskBacked);
 
@@ -2849,16 +2858,21 @@ void FConfigCacheIni::InitializeConfigSystem()
 	FConfigCacheIni::LoadGlobalIniFile(GInputIni, TEXT("Input"));
 #if WITH_EDITOR
 	// load some editor specific .ini files
-	FConfigCacheIni::LoadGlobalIniFile(GEditorIni, TEXT("Editor"));
-	FConfigCacheIni::LoadGlobalIniFile(GEditorLayoutIni, TEXT("EditorLayout"));
-	FConfigCacheIni::LoadGlobalIniFile(GEditorUserSettingsIni, TEXT("EditorUserSettings"));
-	FConfigCacheIni::LoadGlobalIniFile(GEditorKeyBindingsIni, TEXT("EditorKeyBindings"));
 
-	// Game agnostic editor ini file. Only load this config if we are not the build machine as it would introduce a non-deterministic state
+	// Upgrade editor user settings before loading the editor per project user settings
+	FConfigManifest::MigrateEditorUserSettings();
+	FConfigCacheIni::LoadGlobalIniFile(GEditorPerProjectIni, TEXT("EditorPerProjectUserSettings"));
+
+
+	// Project agnostic editor ini files. Only load these configs if we are not the build machine as it would introduce a non-deterministic state
 	if ( !GIsBuildMachine )
 	{
-		const FString GameAgnosticEditorSettingsDir = FPaths::Combine(*FPaths::GameAgnosticSavedDir(), TEXT("Config")) + TEXT("/");
-		FConfigCacheIni::LoadGlobalIniFile(GEditorGameAgnosticIni, TEXT("EditorGameAgnostic"), NULL, false, false, true, *GameAgnosticEditorSettingsDir);
+		const FString EditorSettingsDir = FPaths::Combine(*FPaths::GameAgnosticSavedDir(), TEXT("Config")) + TEXT("/");
+
+		FConfigCacheIni::LoadGlobalIniFile(GEditorIni, TEXT("Editor"), NULL, false, false, true, *EditorSettingsDir);
+		FConfigCacheIni::LoadGlobalIniFile(GEditorSettingsIni, TEXT("EditorSettings"), NULL, false, false, true, *EditorSettingsDir);
+		FConfigCacheIni::LoadGlobalIniFile(GEditorLayoutIni, TEXT("EditorLayout"), NULL, false, false, true, *EditorSettingsDir);
+		FConfigCacheIni::LoadGlobalIniFile(GEditorKeyBindingsIni, TEXT("EditorKeyBindings"), NULL, false, false, true, *EditorSettingsDir);
 	}
 #endif
 #if PLATFORM_DESKTOP
diff --git a/Engine/Source/Runtime/Core/Private/Misc/Core.cpp b/Engine/Source/Runtime/Core/Private/Misc/Core.cpp
--- a/Engine/Source/Runtime/Core/Private/Misc/Core.cpp
+++ b/Engine/Source/Runtime/Core/Private/Misc/Core.cpp
@@ -128,11 +128,16 @@ uint32					GScreenshotResolutionY			= 0;						/* Y Resolution for high res shots
 uint64					GMakeCacheIDIndex				= 0;						/* Cache ID */
 
 FString				GEngineIni;													/* Engine ini filename */
+
+/** Editor ini file locations - stored per engine version (shared across all projects). Migrated between versions on first run. */
 FString				GEditorIni;													/* Editor ini filename */
 FString				GEditorKeyBindingsIni;										/* Editor Key Bindings ini file */
 FString				GEditorLayoutIni;											/* Editor UI Layout ini filename */
-FString				GEditorUserSettingsIni;										/* Editor User Settings ini filename */
-FString				GEditorGameAgnosticIni;										/* Editor Settings (shared between games) ini filename */
+FString				GEditorSettingsIni;											/* Editor Settings ini filename */
+
+/** Editor per-project ini files - stored per project. */
+FString				GEditorPerProjectIni;										/* Editor User Settings ini filename */
+
 FString				GCompatIni;
 FString				GLightmassIni;												/* Lightmass settings ini filename */
 FString				GScalabilityIni;											/* Scalability settings ini filename */
@@ -140,7 +145,7 @@ FString				GInputIni;													/* Input ini filename */
 FString				GGameIni;													/* Game ini filename */
 FString				GGameUserSettingsIni;										/* User Game Settings ini filename */
 
-float					GNearClippingPlane				= 10.0f;					/* Near clipping plane */
+float					GNearClippingPlane				= 10.0f;				/* Near clipping plane */
 
 bool					GExitPurge						= false;
 
diff --git a/Engine/Source/Runtime/Core/Public/CoreGlobals.h b/Engine/Source/Runtime/Core/Public/CoreGlobals.h
--- a/Engine/Source/Runtime/Core/Public/CoreGlobals.h
+++ b/Engine/Source/Runtime/Core/Public/CoreGlobals.h
@@ -144,12 +144,18 @@ extern CORE_API bool GIsHighResScreenshot;
 extern CORE_API uint32 GScreenshotResolutionX;
 extern CORE_API uint32 GScreenshotResolutionY;
 extern CORE_API uint64 GMakeCacheIDIndex;
-extern CORE_API FString GEditorKeyBindingsIni;
+
 extern CORE_API FString GEngineIni;
+
+/** Editor ini file locations - stored per engine version (shared across all projects). Migrated between versions on first run. */
 extern CORE_API FString GEditorIni;
 extern CORE_API FString GEditorLayoutIni;
-extern CORE_API FString GEditorUserSettingsIni;
-extern CORE_API FString GEditorGameAgnosticIni;
+extern CORE_API FString GEditorKeyBindingsIni;
+extern CORE_API FString GEditorSettingsIni;
+
+/** Editor per-project ini files - stored per project. */
+extern CORE_API FString GEditorPerProjectIni;
+
 extern CORE_API FString GCompatIni;
 extern CORE_API FString GLightmassIni;
 extern CORE_API FString GScalabilityIni;
diff --git a/Engine/Source/Runtime/Core/Public/UObject/UnrealNames.inl b/Engine/Source/Runtime/Core/Public/UObject/UnrealNames.inl
--- a/Engine/Source/Runtime/Core/Public/UObject/UnrealNames.inl
+++ b/Engine/Source/Runtime/Core/Public/UObject/UnrealNames.inl
@@ -119,7 +119,7 @@ REGISTER_NAME(247,OtherChildren)
 REGISTER_NAME(248,Location)
 REGISTER_NAME(249,Rotation)
 REGISTER_NAME(250,BSP)
-REGISTER_NAME(251,EditorGameAgnostic)
+REGISTER_NAME(251,EditorSettings)
 
 // Online
 REGISTER_NAME(280,DGram)
@@ -172,3 +172,7 @@ REGISTER_NAME(450,AI)
 REGISTER_NAME(451,NavMesh)
 
 REGISTER_NAME(500,PerformanceCapture)
+
+// Special config names - not required to be consistent for network replication
+REGISTER_NAME(600,EditorLayout)
+REGISTER_NAME(601,EditorKeyBindings)
\ No newline at end of file
diff --git a/Engine/Source/Runtime/CoreUObject/Private/UObject/Class.cpp b/Engine/Source/Runtime/CoreUObject/Private/UObject/Class.cpp
--- a/Engine/Source/Runtime/CoreUObject/Private/UObject/Class.cpp
+++ b/Engine/Source/Runtime/CoreUObject/Private/UObject/Class.cpp
@@ -3605,9 +3605,17 @@ const FString UClass::GetConfigName() const
 	{
 		return GGameIni;
 	}
-	else if ( ClassConfigName == NAME_EditorGameAgnostic )
+	else if ( ClassConfigName == NAME_EditorSettings )
 	{
-		return GEditorGameAgnosticIni;
+		return GEditorSettingsIni;
+	}
+	else if ( ClassConfigName == NAME_EditorLayout )
+	{
+		return GEditorLayoutIni;
+	}
+	else if ( ClassConfigName == NAME_EditorKeyBindings )
+	{
+		return GEditorKeyBindingsIni;
 	}
 	else if( ClassConfigName == NAME_None )
 	{
diff --git a/Engine/Source/Runtime/CoreUObject/Private/UObject/Obj.cpp b/Engine/Source/Runtime/CoreUObject/Private/UObject/Obj.cpp
--- a/Engine/Source/Runtime/CoreUObject/Private/UObject/Obj.cpp
+++ b/Engine/Source/Runtime/CoreUObject/Private/UObject/Obj.cpp
@@ -1477,6 +1477,13 @@ void UObject::LoadConfig( UClass* ConfigClass/*=NULL*/, const TCHAR* InFilename/
 		? InFilename
 		: GetConfigFilename(this);
 
+	// If the file name is empty on build machines, this means the config system never initialized them
+	// as this would potentially create indeterminate on the builder. As such we don't load these configs.
+	if (Filename.IsEmpty() && GIsBuildMachine)
+	{
+		return;
+	}
+
 	const bool bPerObject = UsesPerObjectConfig(this);
 
 	FString ClassSection;
@@ -1675,6 +1682,13 @@ void UObject::SaveConfig( uint64 Flags, const TCHAR* InFilename, FConfigCacheIni
 		? InFilename
 		: GetConfigFilename(this);
 
+	// If the file name is empty on build machines, this means the config system never initialized them
+	// as this would potentially create indeterminate on the builder. As such we don't save these configs.
+	if (Filename.IsEmpty() && GIsBuildMachine)
+	{
+		return;
+	}
+
 	// Determine whether the file we are writing is a default file config.
 	const bool bIsADefaultIniWrite = !Filename.Contains(FPaths::GameSavedDir())
 		&& !Filename.Contains(FPaths::EngineSavedDir())
diff --git a/Engine/Source/Runtime/Engine/Private/EdGraph/EdGraphPin.cpp b/Engine/Source/Runtime/Engine/Private/EdGraph/EdGraphPin.cpp
--- a/Engine/Source/Runtime/Engine/Private/EdGraph/EdGraphPin.cpp
+++ b/Engine/Source/Runtime/Engine/Private/EdGraph/EdGraphPin.cpp
@@ -264,7 +264,7 @@ FText UEdGraphPin::GetDisplayName() const
 		DisplayName = (!PinFriendlyName.IsEmpty()) ? PinFriendlyName : FText::FromString(PinName);
 
 		bool bShowNodesAndPinsUnlocalized;
-		GConfig->GetBool( TEXT("Internationalization"), TEXT("ShowNodesAndPinsUnlocalized"), bShowNodesAndPinsUnlocalized, GEditorGameAgnosticIni );
+		GConfig->GetBool( TEXT("Internationalization"), TEXT("ShowNodesAndPinsUnlocalized"), bShowNodesAndPinsUnlocalized, GEditorSettingsIni );
 		if (bShowNodesAndPinsUnlocalized)
 		{
 			return FText::FromString(DisplayName.BuildSourceString());
diff --git a/Engine/Source/Runtime/Engine/Private/EdGraph/EdGraphSchema.cpp b/Engine/Source/Runtime/Engine/Private/EdGraph/EdGraphSchema.cpp
--- a/Engine/Source/Runtime/Engine/Private/EdGraph/EdGraphSchema.cpp
+++ b/Engine/Source/Runtime/Engine/Private/EdGraph/EdGraphSchema.cpp
@@ -510,7 +510,7 @@ FText UEdGraphSchema::GetPinDisplayName(const UEdGraphPin* Pin) const
 	{
 		ResultPinName = Pin->PinFriendlyName;
 		bool bShowNodesAndPinsUnlocalized;
-		GConfig->GetBool( TEXT("Internationalization"), TEXT("ShowNodesAndPinsUnlocalized"), bShowNodesAndPinsUnlocalized, GEditorGameAgnosticIni );
+		GConfig->GetBool( TEXT("Internationalization"), TEXT("ShowNodesAndPinsUnlocalized"), bShowNodesAndPinsUnlocalized, GEditorSettingsIni );
 		if (bShowNodesAndPinsUnlocalized)
 		{
 			ResultPinName = FText::FromString(ResultPinName.BuildSourceString());
diff --git a/Engine/Source/Runtime/Engine/Private/GameMode.cpp b/Engine/Source/Runtime/Engine/Private/GameMode.cpp
--- a/Engine/Source/Runtime/Engine/Private/GameMode.cpp
+++ b/Engine/Source/Runtime/Engine/Private/GameMode.cpp
@@ -290,13 +290,13 @@ bool AGameMode::ShouldStartInCinematicMode(bool& OutHidePlayer,bool& OutHideHUD,
 	bool StartInCinematicMode = false;
 	if(GEngine->MatineeScreenshotOptions.bStartWithMatineeCapture)
 	{
-		GConfig->GetBool( TEXT("MatineeCreateMovieOptions"), TEXT("CinematicMode"), StartInCinematicMode, GEditorUserSettingsIni );
+		GConfig->GetBool( TEXT("MatineeCreateMovieOptions"), TEXT("CinematicMode"), StartInCinematicMode, GEditorPerProjectIni );
 		if(StartInCinematicMode)
 		{
-			GConfig->GetBool( TEXT("MatineeCreateMovieOptions"), TEXT("DisableMovement"), OutDisableMovement, GEditorUserSettingsIni );
-			GConfig->GetBool( TEXT("MatineeCreateMovieOptions"), TEXT("DisableTurning"), OutDisableTurning, GEditorUserSettingsIni );
-			GConfig->GetBool( TEXT("MatineeCreateMovieOptions"), TEXT("HidePlayer"), OutHidePlayer, GEditorUserSettingsIni );
-			GConfig->GetBool( TEXT("MatineeCreateMovieOptions"), TEXT("HideHUD"), OutHideHUD, GEditorUserSettingsIni );
+			GConfig->GetBool( TEXT("MatineeCreateMovieOptions"), TEXT("DisableMovement"), OutDisableMovement, GEditorPerProjectIni );
+			GConfig->GetBool( TEXT("MatineeCreateMovieOptions"), TEXT("DisableTurning"), OutDisableTurning, GEditorPerProjectIni );
+			GConfig->GetBool( TEXT("MatineeCreateMovieOptions"), TEXT("HidePlayer"), OutHidePlayer, GEditorPerProjectIni );
+			GConfig->GetBool( TEXT("MatineeCreateMovieOptions"), TEXT("HideHUD"), OutHideHUD, GEditorPerProjectIni );
 			return StartInCinematicMode;
 		}
 	}
diff --git a/Engine/Source/Runtime/Engine/Private/PreviewScene.cpp b/Engine/Source/Runtime/Engine/Private/PreviewScene.cpp
--- a/Engine/Source/Runtime/Engine/Private/PreviewScene.cpp
+++ b/Engine/Source/Runtime/Engine/Private/PreviewScene.cpp
@@ -190,7 +190,7 @@ void FPreviewScene::SetSkyBrightness(float SkyBrightness)
 void FPreviewScene::LoadSettings(const TCHAR* Section)
 {
 	FRotator LightDir;
-	if ( GConfig->GetRotator( Section, TEXT("LightDir"), LightDir, GEditorUserSettingsIni ) )
+	if ( GConfig->GetRotator( Section, TEXT("LightDir"), LightDir, GEditorPerProjectIni ) )
 	{
 		SetLightDirection( LightDir );
 	}
@@ -198,5 +198,5 @@ void FPreviewScene::LoadSettings(const TCHAR* Section)
 
 void FPreviewScene::SaveSettings(const TCHAR* Section)
 {
-	GConfig->SetRotator( Section, TEXT("LightDir"), GetLightDirection(), GEditorUserSettingsIni );
+	GConfig->SetRotator( Section, TEXT("LightDir"), GetLightDirection(), GEditorPerProjectIni );
 }
diff --git a/Engine/Source/Runtime/Engine/Private/Scalability.cpp b/Engine/Source/Runtime/Engine/Private/Scalability.cpp
--- a/Engine/Source/Runtime/Engine/Private/Scalability.cpp
+++ b/Engine/Source/Runtime/Engine/Private/Scalability.cpp
@@ -293,7 +293,7 @@ void ProcessCommand(const TCHAR* Cmd, FOutputDevice& Ar)
 		{
 			FQualityLevels State = Scalability::BenchmarkQualityLevels();
 			Scalability::SetQualityLevels(State);
-			Scalability::SaveState(GIsEditor ? GEditorGameAgnosticIni : GGameUserSettingsIni);
+			Scalability::SaveState(GIsEditor ? GEditorSettingsIni : GGameUserSettingsIni);
 			bPrintUsage = false;
 			bPrintCurrentSettings = true;
 		}
@@ -310,7 +310,7 @@ void ProcessCommand(const TCHAR* Cmd, FOutputDevice& Ar)
 			int32 RequestedQualityLevel = FCString::Atoi(*Token);
 			QualityLevels.SetFromSingleQualityLevel(RequestedQualityLevel);
 			SetQualityLevels(QualityLevels);
-			Scalability::SaveState(GIsEditor ? GEditorGameAgnosticIni : GGameUserSettingsIni);
+			Scalability::SaveState(GIsEditor ? GEditorSettingsIni : GGameUserSettingsIni);
 
 			bPrintUsage = false;
 		}
diff --git a/Engine/Source/Runtime/Engine/Private/UnrealEngine.cpp b/Engine/Source/Runtime/Engine/Private/UnrealEngine.cpp
--- a/Engine/Source/Runtime/Engine/Private/UnrealEngine.cpp
+++ b/Engine/Source/Runtime/Engine/Private/UnrealEngine.cpp
@@ -1218,7 +1218,7 @@ void UEngine::ParseCommandline()
 
 	if( !GIsEditor && MatineeScreenshotOptions.bStartWithMatineeCapture )
 	{
-		GConfig->GetBool( TEXT("MatineeCreateMovieOptions"), TEXT("HideHUD"), MatineeScreenshotOptions.bHideHud, GEditorUserSettingsIni );
+		GConfig->GetBool( TEXT("MatineeCreateMovieOptions"), TEXT("HideHUD"), MatineeScreenshotOptions.bHideHud, GEditorPerProjectIni );
 	}
 
 	// If we are capturing a matinee movie and we want to dump the buffer visualization shots too, for on all required functionality
diff --git a/Engine/Source/Runtime/Engine/Public/Scalability.h b/Engine/Source/Runtime/Engine/Public/Scalability.h
--- a/Engine/Source/Runtime/Engine/Public/Scalability.h
+++ b/Engine/Source/Runtime/Engine/Public/Scalability.h
@@ -55,10 +55,10 @@ namespace Scalability
 	/**  */
 	ENGINE_API void InitScalabilitySystem();
 
-	/** @param IniName e.g. GEditorUserSettingsIni or GGameUserSettingsIni */
+	/** @param IniName e.g. GEditorPerProjectIni or GGameUserSettingsIni */
 	ENGINE_API void LoadState(const FString& IniName);
 	
-	/** @param IniName e.g. GEditorUserSettingsIni or GGameUserSettingsIni */
+	/** @param IniName e.g. GEditorPerProjectIni or GGameUserSettingsIni */
 	ENGINE_API void SaveState(const FString& IniName);
 
 	/**
diff --git a/Engine/Source/Runtime/Launch/Private/LaunchEngineLoop.cpp b/Engine/Source/Runtime/Launch/Private/LaunchEngineLoop.cpp
--- a/Engine/Source/Runtime/Launch/Private/LaunchEngineLoop.cpp
+++ b/Engine/Source/Runtime/Launch/Private/LaunchEngineLoop.cpp
@@ -1111,7 +1111,7 @@ int32 FEngineLoop::PreInit( const TCHAR* CmdLine )
 	// after SystemSettings.ini file loading so we get the right state,
 	// before ConsoleVariables.ini so the local developer can always override.
 	// before InitializeCVarsForActiveDeviceProfile() so the platform can override user settings
-	Scalability::LoadState((bHasEditorToken && !GEditorGameAgnosticIni.IsEmpty()) ? GEditorGameAgnosticIni : GGameUserSettingsIni);
+	Scalability::LoadState((bHasEditorToken && !GEditorSettingsIni.IsEmpty()) ? GEditorSettingsIni : GGameUserSettingsIni);
 
 	// Set all CVars which have been setup in the device profiles.
 	UDeviceProfileManager::InitializeCVarsForActiveDeviceProfile();
@@ -2578,7 +2578,7 @@ bool FEngineLoop::AppInit( )
 		if(CurrentProject != nullptr && CurrentProject->Modules.Num() > 0)
 		{
 			bool bNeedCompile = false;
-			GConfig->GetBool(TEXT("/Script/UnrealEd.EditorLoadingSavingSettings"), TEXT("bForceCompilationAtStartup"), bNeedCompile, GEditorUserSettingsIni);
+			GConfig->GetBool(TEXT("/Script/UnrealEd.EditorLoadingSavingSettings"), TEXT("bForceCompilationAtStartup"), bNeedCompile, GEditorPerProjectIni);
 
 			if(!bNeedCompile)
 			{
diff --git a/Engine/Source/Runtime/OpenGLDrv/Private/Mac/MacOpenGLContext.cpp b/Engine/Source/Runtime/OpenGLDrv/Private/Mac/MacOpenGLContext.cpp
--- a/Engine/Source/Runtime/OpenGLDrv/Private/Mac/MacOpenGLContext.cpp
+++ b/Engine/Source/Runtime/OpenGLDrv/Private/Mac/MacOpenGLContext.cpp
@@ -15,7 +15,7 @@
 
 #if WITH_EDITOR
 #define MAC_OPENGL_SETTINGS TEXT("/Script/MacGraphicsSwitching.MacGraphicsSwitchingSettings")
-#define MAC_OPENGL_INI GEditorGameAgnosticIni
+#define MAC_OPENGL_INI GEditorSettingsIni
 #else
 #define MAC_OPENGL_SETTINGS TEXT("OpenGL")
 #define MAC_OPENGL_INI GEngineIni
diff --git a/Engine/Source/Runtime/Slate/Private/Framework/Docking/LayoutService.cpp b/Engine/Source/Runtime/Slate/Private/Framework/Docking/LayoutService.cpp
--- a/Engine/Source/Runtime/Slate/Private/Framework/Docking/LayoutService.cpp
+++ b/Engine/Source/Runtime/Slate/Private/Framework/Docking/LayoutService.cpp
@@ -66,8 +66,8 @@ void FLayoutSaveRestore::MigrateConfig( const FString& OldConfigFileName, const
 
 	// remove old configuration
 	GConfig->EmptySection(EditorLayoutsSectionName, OldConfigFileName);
-	GConfig->Flush(false, GEditorUserSettingsIni);
-	GConfig->Flush(false, GEditorLayoutIni);
+	GConfig->Flush(false, OldConfigFileName);
+	GConfig->Flush(false, NewConfigFileName);
 }
 
 
diff --git a/Engine/Source/Runtime/Slate/Private/Framework/MultiBox/MultiBox.cpp b/Engine/Source/Runtime/Slate/Private/Framework/MultiBox/MultiBox.cpp
--- a/Engine/Source/Runtime/Slate/Private/Framework/MultiBox/MultiBox.cpp
+++ b/Engine/Source/Runtime/Slate/Private/Framework/MultiBox/MultiBox.cpp
@@ -47,7 +47,7 @@ void FMultiBoxSettings::ToggleToolbarEditing()
 	bool bCanBeEnabled = false;
 	if( GIsEditor )
 	{
-		GConfig->GetBool(TEXT("/Script/UnrealEd.EditorExperimentalSettings"), TEXT("bToolbarCustomization"), bCanBeEnabled, GEditorUserSettingsIni);
+		GConfig->GetBool(TEXT("/Script/UnrealEd.EditorExperimentalSettings"), TEXT("bToolbarCustomization"), bCanBeEnabled, GEditorPerProjectIni);
 	}
 
 	bInToolbarEditMode = !bInToolbarEditMode && bCanBeEnabled;
@@ -268,7 +268,7 @@ bool FMultiBox::IsCustomizable() const
 	bool bIsCustomizable = false;
 	if( CustomizationData->GetCustomizationName() != NAME_None )
 	{
-		GConfig->GetBool(TEXT("/Script/UnrealEd.EditorExperimentalSettings"), TEXT("bToolbarCustomization"), bIsCustomizable, GEditorUserSettingsIni);
+		GConfig->GetBool(TEXT("/Script/UnrealEd.EditorExperimentalSettings"), TEXT("bToolbarCustomization"), bIsCustomizable, GEditorPerProjectIni);
 	}
 
 	return bIsCustomizable;
diff --git a/Engine/Source/Runtime/Slate/Private/Framework/MultiBox/MultiBoxCustomization.cpp b/Engine/Source/Runtime/Slate/Private/Framework/MultiBox/MultiBoxCustomization.cpp
--- a/Engine/Source/Runtime/Slate/Private/Framework/MultiBox/MultiBoxCustomization.cpp
+++ b/Engine/Source/Runtime/Slate/Private/Framework/MultiBox/MultiBoxCustomization.cpp
@@ -107,7 +107,7 @@ void FMultiBoxCustomizationData::LoadCustomizedBlocks()
 
 	FString Content;
 	
-	GConfig->GetString(*GetConfigSectionName(), *CustomizationName.ToString(), Content, GEditorUserSettingsIni);
+	GConfig->GetString(*GetConfigSectionName(), *CustomizationName.ToString(), Content, GEditorPerProjectIni);
 
 	TSharedPtr<FJsonObject> SavedData;
 	TSharedRef< TJsonReader<> > Reader = TJsonReaderFactory<>::Create( FRemoteConfig::ReplaceIniSpecialCharWithChar(Content).ReplaceEscapedCharWithChar() );
@@ -344,7 +344,7 @@ void FMultiBoxCustomizationData::SaveCustomizedBlocks()
 	Writer->WriteObjectEnd();
 	Writer->Close();
 
-	GConfig->SetString( *GetConfigSectionName(), *CustomizationName.ToString(), *FRemoteConfig::ReplaceIniCharWithSpecialChar(SaveData).ReplaceCharWithEscapedChar(), GEditorUserSettingsIni );
+	GConfig->SetString( *GetConfigSectionName(), *CustomizationName.ToString(), *FRemoteConfig::ReplaceIniCharWithSpecialChar(SaveData).ReplaceCharWithEscapedChar(), GEditorPerProjectIni );
 }
 
 void FMultiBoxCustomizationData::BlockRemoved( TSharedRef< const FMultiBlock> RemovedBlock, int32 Index, const TArray< TSharedRef< const FMultiBlock > >& AllBlocks )
