Replaced some TArray::Empty calls with TArray:Reset calls to potentially avoid some reallocations

diff --git a/Engine/Source/Editor/CurveTableEditor/Private/CurveTableEditor.cpp b/Engine/Source/Editor/CurveTableEditor/Private/CurveTableEditor.cpp
--- a/Engine/Source/Editor/CurveTableEditor/Private/CurveTableEditor.cpp
+++ b/Engine/Source/Editor/CurveTableEditor/Private/CurveTableEditor.cpp
@@ -288,7 +288,7 @@ void FCurveTableEditor::CacheDataTableForEditing()
 	static const float CellPadding = 10.0f;
 
 	// Column titles, taken from the longest curve
-	AvailableColumns.Empty(Curves[LongestCurveIndex]->GetNumKeys());
+	AvailableColumns.Reset(Curves[LongestCurveIndex]->GetNumKeys());
 	for (auto It(Curves[LongestCurveIndex]->GetKeyIterator()); It; ++It)
 	{
 		const FText ColumnText = FText::AsNumber(It->Time);
@@ -302,7 +302,7 @@ void FCurveTableEditor::CacheDataTableForEditing()
 	}
 
 	// Each curve is a row entry
-	AvailableRows.Empty(Curves.Num());
+	AvailableRows.Reset(Curves.Num());
 	for (int32 CurvesIdx = 0; CurvesIdx < Curves.Num(); ++CurvesIdx)
 	{
 		const FName& CurveName = Names[CurvesIdx];
diff --git a/Engine/Source/Editor/UnrealEd/Private/DataTableEditorUtils.cpp b/Engine/Source/Editor/UnrealEd/Private/DataTableEditorUtils.cpp
--- a/Engine/Source/Editor/UnrealEd/Private/DataTableEditorUtils.cpp
+++ b/Engine/Source/Editor/UnrealEd/Private/DataTableEditorUtils.cpp
@@ -206,7 +206,7 @@ void FDataTableEditorUtils::CacheDataTableForEditing(const UDataTable* DataTable
 	static const float CellPadding = 10.0f;
 
 	// Populate the column data
-	OutAvailableColumns.Empty(StructProps.Num());
+	OutAvailableColumns.Reset(StructProps.Num());
 	for (const UProperty* Prop : StructProps)
 	{
 		FDataTableEditorColumnHeaderDataPtr CachedColumnData = MakeShareable(new FDataTableEditorColumnHeaderData());
@@ -218,7 +218,7 @@ void FDataTableEditorUtils::CacheDataTableForEditing(const UDataTable* DataTable
 	}
 
 	// Populate the row data
-	OutAvailableRows.Empty(DataTable->RowMap.Num());
+	OutAvailableRows.Reset(DataTable->RowMap.Num());
 	for (auto RowIt = DataTable->RowMap.CreateConstIterator(); RowIt; ++RowIt)
 	{
 		FDataTableEditorRowListViewDataPtr CachedRowData = MakeShareable(new FDataTableEditorRowListViewData());
