Improved performance for startup cook time.
Improved caching routine for cook in editor modes.
Improved cooker generating dependencies performance.
Added new error when dlc is cooking and engine content is encountered (error disabled by default).

diff --git a/Engine/Source/Editor/UnrealEd/Classes/CookOnTheSide/CookOnTheFlyServer.h b/Engine/Source/Editor/UnrealEd/Classes/CookOnTheSide/CookOnTheFlyServer.h
--- a/Engine/Source/Editor/UnrealEd/Classes/CookOnTheSide/CookOnTheFlyServer.h
+++ b/Engine/Source/Editor/UnrealEd/Classes/CookOnTheSide/CookOnTheFlyServer.h
@@ -22,7 +22,7 @@ enum class ECookInitializationFlags
 	AutoTick = 0x10,				// enable ticking (only works in the editor)
 	AsyncSave = 0x20,				// save packages async
 	IncludeServerMaps = 0x80,		// should we include the server maps when cooking
-	GenerateStreamingInstallManifest = 0x100,  // should we generate streaming install manifest
+	UseSerializationForPackageDependencies = 0x100, // should we use the serialization code path for generating package dependencies (old method will be depricated)
 };
 ENUM_CLASS_FLAGS(ECookInitializationFlags);
 
@@ -585,17 +585,21 @@ private:
 	struct FCookByTheBookOptions
 	{
 	public:
-		FCookByTheBookOptions() : bGenerateStreamingInstallManifests(false),
+		FCookByTheBookOptions() : bLeakTest(false),
+			bGenerateStreamingInstallManifests(false),
+			bGenerateDependenciesForMaps(false),
 			bRunning(false),
 			CookTime( 0.0 ),
 			CookStartTime( 0.0 ), 
-			bGenerateDependeciesForMaps(false)
+			bErrorOnEngineContentUse(false)
 		{ }
 
 		/** Should we test for UObject leaks */
 		bool bLeakTest;
 		/** Should we generate streaming install manifests (only valid option in cook by the book) */
 		bool bGenerateStreamingInstallManifests;
+		/** Should we generate a seperate manifest for map dependencies */
+		bool bGenerateDependenciesForMaps;
 		/** Is cook by the book currently running */
 		bool bRunning;
 		/** Cancel has been queued will be processed next tick */
@@ -608,12 +612,13 @@ private:
 		TSet<FWeakObjectPtr> LastGCItems;
 		/** Map of platform name to manifest generator, manifest is only used in cook by the book however it needs to be maintained across multiple cook by the books. */
 		TMap<FName, FChunkManifestGenerator*> ManifestGenerators;
+		/** Dependency graph of maps as root objects. */
+		TMap< FName, TSet <FName> > MapDependencyGraph; 
 		/** If a cook is cancelled next cook will need to resume cooking */ 
 		TArray<FFilePlatformRequest> PreviousCookRequests; 
 		double CookTime;
 		double CookStartTime;
-		/** Generate Map dependencies */
-		bool bGenerateDependeciesForMaps;
+		bool bErrorOnEngineContentUse;
 	};
 	FCookByTheBookOptions* CookByTheBookOptions;
 	
@@ -638,8 +643,6 @@ private:
 	bool bIsSavingPackage; // used to stop recursive mark package dirty functions
 
 	//////////////////////////////////////////////////////////////////////////
-	// Dependency graph of maps as root objects. 
-	TMap< FName, TSet <FName> > MapDependencyGraph; 
 
 	// data about the current package being processed
 	struct FReentryData
@@ -745,12 +748,16 @@ public:
 		FString DLCName;
 		FString CreateReleaseVersion;
 		FString BasedOnReleaseVersion;
-		bool bGenerateDependeciesForMaps; 
+		bool bGenerateStreamingInstallManifests; 
+		bool bGenerateDependenciesForMaps; 
+		bool bErrorOnEngineContentUse; // this is a flag for dlc, will cause the cooker to error if the dlc references engine content
 
 		FCookByTheBookStartupOptions() :
 			CookOptions(ECookByTheBookOptions::None),
 			DLCName(FString()),
-			bGenerateDependeciesForMaps(false)
+			bGenerateStreamingInstallManifests(false),
+			bGenerateDependenciesForMaps(false),
+			bErrorOnEngineContentUse(false)
 		{ }
 	};
 
@@ -954,6 +961,16 @@ private:
 	 * @param Found return value, all objects which package is dependent on
 	 */
 	void GetDependencies( const TSet<UPackage*>& Packages, TSet<UObject*>& Found);
+
+
+	/**
+	 * GetDependencies
+	 * 
+	 * @param Packages List of packages to use as the root set for dependency checking
+	 * @param Found return value, all objects which package is dependent on
+	 */
+	void GetDependentPackages( const TSet<UPackage*>& Packages, TSet<FName>& Found);
+
 	/**
 	 * GenerateManifestInfo
 	 * generate the manifest information for a given package
diff --git a/Engine/Source/Editor/UnrealEd/Private/Commandlets/ChunkManifestGenerator.cpp b/Engine/Source/Editor/UnrealEd/Private/Commandlets/ChunkManifestGenerator.cpp
--- a/Engine/Source/Editor/UnrealEd/Private/Commandlets/ChunkManifestGenerator.cpp
+++ b/Engine/Source/Editor/UnrealEd/Private/Commandlets/ChunkManifestGenerator.cpp
@@ -63,6 +63,14 @@ void FChunkManifestGenerator::OnLastPackageLoaded( UPackage* Package )
 	}
 }
 
+void FChunkManifestGenerator::OnLastPackageLoaded( const FName& PackageName )
+{
+	if ( !AssetsLoadedWithLastPackage.Contains(PackageName))
+	{
+		AssetsLoadedWithLastPackage.Add(PackageName);
+	}
+}
+
 bool FChunkManifestGenerator::CleanTempPackagingDirectory(const FString& Platform) const
 {
 	FString TmpPackagingDir = GetTempPackagingDirectoryForPlatform(Platform);
@@ -200,8 +208,76 @@ void FChunkManifestGenerator::Initialize(bool InGenerateChunks)
 	}
 }
 
+void FChunkManifestGenerator::AddPackageToChunkManifest(const FName& PackageFName, const FString& PackagePathName, const FString& SandboxFilename, const FString& LastLoadedMapName, FSandboxPlatformFile* InSandboxFile)
+{
+	TArray<int32> TargetChunks;
+	TArray<int32> ExistingChunkIDs;
+
+	if (!bGenerateChunks)
+	{
+		TargetChunks.AddUnique(0);
+		ExistingChunkIDs.AddUnique(0);
+	}
+
+	if (bGenerateChunks)
+	{
+		// Try to determine if this package has been loaded as a result of loading a map package.
+		FString MapThisAssetWasLoadedWith;
+		if (!LastLoadedMapName.IsEmpty())
+		{
+			if (AssetsLoadedWithLastPackage.Contains(PackageFName))
+			{
+				MapThisAssetWasLoadedWith = LastLoadedMapName;
+			}
+		}
+
+		// Collect all chunk IDs associated with this package from the asset registry
+		TArray<int32> RegistryChunkIDs = GetAssetRegistryChunkAssignments(PackageFName);
+		ExistingChunkIDs = GetExistingPackageChunkAssignments(PackageFName);
+
+		// Try to call game-specific delegate to determine the target chunk ID
+		// FString Name = Package->GetPathName();
+		if (FGameDelegates::Get().GetAssignStreamingChunkDelegate().IsBound())
+		{
+			FGameDelegates::Get().GetAssignStreamingChunkDelegate().ExecuteIfBound(PackagePathName, MapThisAssetWasLoadedWith, RegistryChunkIDs, ExistingChunkIDs, TargetChunks);
+		}
+		else
+		{
+			//Take asset registry assignments and existing assignments
+			TargetChunks.Append(RegistryChunkIDs);
+			TargetChunks.Append(ExistingChunkIDs);
+		}
+	}
+
+	NotifyPackageWasCooked(SandboxFilename, PackageFName);
+
+	bool bAssignedToChunk = false;
+	// if the delegate requested a specific chunk assignment, add them package to it now.
+	for (const auto& PackageChunk : TargetChunks)
+	{
+		AddPackageToManifest(SandboxFilename, PackageFName, PackageChunk);
+		bAssignedToChunk = true;
+	}
+	// If the delegate requested to remove the package from any chunk, remove it now
+	for (const auto& PackageChunk : ExistingChunkIDs)
+	{
+		if (!TargetChunks.Contains(PackageChunk))
+		{
+			RemovePackageFromManifest(PackageFName, PackageChunk);
+		}
+	}
+
+	if (!bAssignedToChunk)
+	{
+		NotifyPackageWasNotAssigned(SandboxFilename, PackageFName);
+	}
+}
+
 void FChunkManifestGenerator::AddPackageToChunkManifest(UPackage* Package, const FString& SandboxFilename, const FString& LastLoadedMapName, FSandboxPlatformFile* InSandboxFile)
 {
+#if 0
+	AddPackageToChunkManifest(Package->GetFName(), Package->GetPathName(), SandboxFilename, LastLoadedMapName, InSandboxFile);
+#else
 	TArray<int32> TargetChunks;
 	TArray<int32> ExistingChunkIDs;
 	
@@ -264,6 +340,7 @@ void FChunkManifestGenerator::AddPackageToChunkManifest(UPackage* Package, const
 	{
 		NotifyPackageWasNotAssigned(SandboxFilename, PackageFName);
 	}
+#endif
 }
 
 void FChunkManifestGenerator::PrepareToLoadNewPackage(const FString& Filename)
diff --git a/Engine/Source/Editor/UnrealEd/Private/Commandlets/ChunkManifestGenerator.h b/Engine/Source/Editor/UnrealEd/Private/Commandlets/ChunkManifestGenerator.h
--- a/Engine/Source/Editor/UnrealEd/Private/Commandlets/ChunkManifestGenerator.h
+++ b/Engine/Source/Editor/UnrealEd/Private/Commandlets/ChunkManifestGenerator.h
@@ -162,6 +162,23 @@ class FChunkManifestGenerator
 		return RegistryChunkIDs;
 	}
 
+	/**
+	* Returns an array of chunks IDs for a package that have been assigned in the editor.
+	*/
+	FORCEINLINE TArray<int32> GetAssetRegistryChunkAssignments(const FName& PackageFName)
+	{
+		//get the objects in this package
+		TArray<int32> RegistryChunkIDs;
+		for ( const auto& AssetData : AssetRegistryData )
+		{
+			if ( AssetData.PackageName == PackageFName ) 
+			{
+				RegistryChunkIDs.Append( AssetData.ChunkIDs );
+			}
+		}
+		return RegistryChunkIDs;
+	}
+
 	/**
 	 * Deletes the temporary packaging directory for the specified platform.
 	 */
@@ -236,6 +253,7 @@ public:
 	 * @param the SandboxPlatformFile used during cook
 	 */
 	void AddPackageToChunkManifest(UPackage* Package, const FString& SandboxFilename, const FString& LastLoadedMapName, FSandboxPlatformFile* SandboxFile);
+	void AddPackageToChunkManifest(const FName& PackageFName, const FString& PackagePathName, const FString& SandboxFilename, const FString& LastLoadedMapName, FSandboxPlatformFile* InSandboxFile);
 	
 
 	/**
@@ -253,6 +271,15 @@ public:
 	 */
 	void OnLastPackageLoaded( UPackage* Package );
 
+
+	/**
+	 * Collects all the packages loaded 
+	 * Does the same as other overload of OnLastPackageLoaded except takes in the name of the package as a fname 
+	 * 
+	 * @param Package name
+	 */
+	void OnLastPackageLoaded( const FName& PackageName );
+
 	/**
 	 * The cooker is about to load a new package from the list, reset AssetsLoadedWithLastPackage
 	 */
diff --git a/Engine/Source/Editor/UnrealEd/Private/Commandlets/CookCommandlet.cpp b/Engine/Source/Editor/UnrealEd/Private/Commandlets/CookCommandlet.cpp
--- a/Engine/Source/Editor/UnrealEd/Private/Commandlets/CookCommandlet.cpp
+++ b/Engine/Source/Editor/UnrealEd/Private/Commandlets/CookCommandlet.cpp
@@ -522,7 +522,8 @@ int32 UCookCommandlet::Main(const FString& CmdLineParams)
 	bCompressed = Switches.Contains(TEXT("COMPRESSED"));
 	bIterativeCooking = Switches.Contains(TEXT("ITERATE"));
 	bSkipEditorContent = Switches.Contains(TEXT("SKIPEDITORCONTENT")); // This won't save out any packages in Engine/COntent/Editor*
-
+	bErrorOnEngineContentUse = Switches.Contains(TEXT("ERRORONENGINECONTENTUSE"));
+	bUseSerializationForGeneratingPackageDependencies = Switches.Contains(TEXT("UseSerializationForGeneratingPackageDependencies"));
 	if (bLeakTest)
 	{
 		for (FObjectIterator It; It; ++It)
@@ -1029,8 +1030,8 @@ bool UCookCommandlet::NewCook( const TArray<ITargetPlatform*>& Platforms, TArray
 	ECookInitializationFlags CookFlags = ECookInitializationFlags::IncludeServerMaps;
 	CookFlags |= bCompressed ? ECookInitializationFlags::Compressed : ECookInitializationFlags::None;
 	CookFlags |= bIterativeCooking ? ECookInitializationFlags::Iterative : ECookInitializationFlags::None;
-	CookFlags |= bSkipEditorContent ? ECookInitializationFlags::SkipEditorContent : ECookInitializationFlags::None;
-	CookFlags |= bGenerateStreamingInstallManifests ? ECookInitializationFlags::GenerateStreamingInstallManifest : ECookInitializationFlags::None;
+	CookFlags |= bSkipEditorContent ? ECookInitializationFlags::SkipEditorContent : ECookInitializationFlags::None;	
+	CookFlags |= bUseSerializationForGeneratingPackageDependencies ? ECookInitializationFlags::UseSerializationForPackageDependencies : ECookInitializationFlags::None;
 
 	TArray<UClass*> FullGCAssetClasses;
 	if ( FullGCAssetClassNames.Num() )
@@ -1182,7 +1183,10 @@ bool UCookCommandlet::NewCook( const TArray<ITargetPlatform*>& Platforms, TArray
 	Swap( StartupOptions.BasedOnReleaseVersion, BasedOnReleaseVersion );
 	Swap( StartupOptions.CreateReleaseVersion, CreateReleaseVersion );
 	StartupOptions.CookOptions = CookOptions;
-	StartupOptions.bGenerateDependeciesForMaps = Switches.Contains(TEXT("GenerateDependenciesForMaps"));
+	StartupOptions.bErrorOnEngineContentUse = bErrorOnEngineContentUse;
+	StartupOptions.bGenerateDependenciesForMaps = Switches.Contains(TEXT("GenerateDependenciesForMaps"));
+	StartupOptions.bGenerateStreamingInstallManifests = bGenerateStreamingInstallManifests;
+
 
 	CookOnTheFlyServer->StartCookByTheBook( StartupOptions );
 
diff --git a/Engine/Source/Editor/UnrealEd/Private/CookOnTheFlyServer.cpp b/Engine/Source/Editor/UnrealEd/Private/CookOnTheFlyServer.cpp
--- a/Engine/Source/Editor/UnrealEd/Private/CookOnTheFlyServer.cpp
+++ b/Engine/Source/Editor/UnrealEd/Private/CookOnTheFlyServer.cpp
@@ -42,6 +42,9 @@ DEFINE_LOG_CATEGORY_STATIC(LogCook, Log, All);
 #define DEBUG_COOKONTHEFLY 0
 #define OUTPUT_TIMING 0
 
+#define USEASSETREGISTRYFORDEPENDENTPACKAGES 1
+#define VERIFY_GETDEPENDENTPACKAGES 0 // verify has false hits because old serialization method for generating dependencies had errors (included transient objects which shouldn't be in asset registry), but you can still use verify to build a list then cross check against transient objects.  
+
 #if OUTPUT_TIMING
 
 
@@ -628,6 +631,7 @@ public:
 		, RootSet(InRootSet)
 	{
 		ArIsObjectReferenceCollector = true;
+		ArIsSaving = true;
 
 		for ( const auto& Object : RootSet )
 		{
@@ -657,83 +661,117 @@ public:
 	virtual FString GetArchiveName() const override { return TEXT("FArchiveFindReferences"); }
 };
 
-void UCookOnTheFlyServer::GetDependencies( const TSet<UPackage*>& Packages, TSet<UObject*>& Found)
+
+void UCookOnTheFlyServer::GetDependentPackages( const TSet<UPackage*>& RootPackages, TSet<FName>& FoundPackages )
 {
+static const FName AssetRegistryName("AssetRegistry");
+	FAssetRegistryModule& AssetRegistryModule = FModuleManager::LoadModuleChecked<FAssetRegistryModule>(AssetRegistryName);
+	IAssetRegistry& AssetRegistry = AssetRegistryModule.Get();
 
-	TSet<UObject*> RootSet;
+	TArray<FName> FoundPackagesArray;
+	for ( const UPackage* RootPackage : RootPackages)
+	{
+		FoundPackagesArray.Add( RootPackage->GetFName() );
+		FoundPackages.Add( RootPackage->GetFName() );
+	}
 
-	for (UPackage* Package : Packages)
+
+	int FoundPackagesCounter = 0;
+	while ( FoundPackagesCounter < FoundPackagesArray.Num() )
+	{
+
+		TArray<FName> PackageDependencies;
+		verify( AssetRegistry.GetDependencies(FoundPackagesArray[FoundPackagesCounter], PackageDependencies) );
+		++FoundPackagesCounter;
+		for ( const auto& PackageDependency : PackageDependencies )
+		{
+			if ( FoundPackages.Contains(PackageDependency) == false )
+			{
+				FoundPackages.Add(PackageDependency);
+				FoundPackagesArray.Add( PackageDependency );
+			}
+			
+		}
+	}	
+
+
+#if VERIFY_GETDEPENDENTPACKAGES
+
+	TSet<UObject*> RootSet;
+	TSet<UObject*> FoundObjects;
+	for (UPackage* Package : RootPackages)
 	{
 		TArray<UObject*> ObjectsInPackage;
 		GetObjectsWithOuter(Package, ObjectsInPackage, true);
 		for (UObject* Obj : ObjectsInPackage)
 		{
 			RootSet.Add(Obj);
-			Found.Add(Obj);
+			FoundObjects.Add(Obj);
 		}
 	}
 
 	TArray<UObject*> Exclude;
-	FArchiveFindReferences ArFindReferences( RootSet, Found, Exclude );
+	FArchiveFindReferences ArFindReferences( RootSet, FoundObjects, Exclude );
+
+	TSet<UPackage*> NewFoundPackages;
+	for ( const auto& Obj : FoundObjects )
+	{
+		NewFoundPackages.Add( Obj->GetOutermost() );
+	}
+
+
+	// verify that the old code (directly above) functions the same as the new code
 
-	// Iterate through the object list
-	/*for( FObjectIterator It; It; ++It )
+	for (const UPackage* NewPackage : NewFoundPackages)
 	{
-		// if this object is within the package specified, serialize the object
-		// into a specialized archive which logs object names encountered during
-		// serialization -- rjp
-		if ( It->IsIn(Pkg) )
+		UE_LOG(LogCook, Display, TEXT("Old Found package %s"), *NewPackage->GetFName().ToString());
+	}
+	for (const auto& CurrentPackage : FoundPackages )
+	{
+		UE_LOG(LogCook, Display, TEXT("New Found package %s"), *CurrentPackage.ToString());
+	}
+
+	for ( const UPackage* NewPackage : NewFoundPackages )
+	{
+		if (FoundPackages.Contains(NewPackage->GetFName()) == false)
 		{
-			if ( It->GetOuter() == Pkg )
-			{
-				FArchiveFindReferences ArFindReferences( Pkg, *It, Found, Exclude );
-			}
-			else if ( bRecurse )
-			{
-				// Two options -
-				// a) this object is a function or something (which we don't care about)
-				// b) this object is inside a group inside the specified package (which we do care about)
-				UObject* CurrentObject = *It;
-				UObject* CurrentOuter = It->GetOuter();
-				while ( CurrentObject && CurrentOuter )
-				{
-					// this object is a UPackage (a group inside a package)
-					// abort
-					if ( CurrentObject->GetClass() == UPackage::StaticClass() )
-						break;
+			UE_LOG(LogCook, Display, TEXT("Package was missing from new generated packages list %s"), *NewPackage->GetFName().ToString());
+		}
+		//check( FoundPackages.Contains( NewPackage->GetFName() ) );
+	}
+	check(NewFoundPackages.Num() == FoundPackages.Num());
 
-					// see if this object's outer is a UPackage
-					if ( CurrentOuter->GetClass() == UPackage::StaticClass() )
-					{
-						// if this object's outer is our original package, the original object (It)
-						// wasn't inside a group, it just wasn't at the base level of the package
-						// (its Outer wasn't the Pkg, it was something else e.g. a function, state, etc.)
-						/// ....just skip it
-						if ( CurrentOuter == Pkg )
-							break;
+#endif
 
-						// otherwise, we've successfully found an object that was in the package we
-						// were searching, but would have been hidden within a group - let's log it
-						FArchiveFindReferences ArFindReferences( CurrentOuter, CurrentObject, Found, Exclude );
-						break;
-					}
+}
 
-					CurrentObject = CurrentOuter;
-					CurrentOuter = CurrentObject->GetOuter();
-				}
-			}
+void UCookOnTheFlyServer::GetDependencies( const TSet<UPackage*>& Packages, TSet<UObject*>& Found)
+{
+	TSet<UObject*> RootSet;
+
+	for (UPackage* Package : Packages)
+	{
+		TArray<UObject*> ObjectsInPackage;
+		GetObjectsWithOuter(Package, ObjectsInPackage, true);
+		for (UObject* Obj : ObjectsInPackage)
+		{
+			RootSet.Add(Obj);
+			Found.Add(Obj);
 		}
-	}*/
+	}
+
+	TArray<UObject*> Exclude;
+	FArchiveFindReferences ArFindReferences( RootSet, Found, Exclude );
 }
 
 
-void UCookOnTheFlyServer::GenerateManifestInfo( UPackage* Package, const TArray<FName>& TargetPlatformNames )
+void UCookOnTheFlyServer::GenerateManifestInfo(UPackage* Package, const TArray<FName>& TargetPlatformNames)
 {
-	if ( !CookByTheBookOptions )
+	if (!CookByTheBookOptions)
 		return;
-	
+
 	// generate dependency information for this package
-	
+
 
 	TSet<UPackage*> RootPackages;
 	RootPackages.Add(Package);
@@ -743,13 +781,13 @@ void UCookOnTheFlyServer::GenerateManifestInfo( UPackage* Package, const TArray<
 	// load sublevels
 	UWorld* World = UWorld::FindWorldInPackage(Package);
 
-	if ( World )
+	if (World)
 	{
 		for (const auto& StreamingLevel : World->StreamingLevels)
 		{
-			if ( StreamingLevel->GetLoadedLevel() )
+			if (StreamingLevel->GetLoadedLevel())
 			{
-				RootPackages.Add(StreamingLevel->GetLoadedLevel()->GetOutermost() );
+				RootPackages.Add(StreamingLevel->GetLoadedLevel()->GetOutermost());
 			}
 		}
 
@@ -761,11 +799,11 @@ void UCookOnTheFlyServer::GenerateManifestInfo( UPackage* Package, const TArray<
 			World->WorldComposition->CollectTilesToCook(NewPackagesToCook);
 		}
 
-		for ( const auto& PackageName : NewPackagesToCook)
+		for (const auto& PackageName : NewPackagesToCook)
 		{
-			UPackage* PackageToCook = LoadPackage( NULL, *PackageName, LOAD_None );
+			UPackage* PackageToCook = LoadPackage(NULL, *PackageName, LOAD_None);
 
-			RootPackages.Add( PackageToCook );
+			RootPackages.Add(PackageToCook);
 			//GetDependencies( PackageToCook, Dependencies );
 
 			// Dependencies.Add(PackageToCook);
@@ -773,62 +811,116 @@ void UCookOnTheFlyServer::GenerateManifestInfo( UPackage* Package, const TArray<
 
 		LastLoadedMapName = Package->GetName();
 	}
+#if USEASSETREGISTRYFORDEPENDENTPACKAGES
+
+	static const FName AssetRegistryName("AssetRegistry");
+	FAssetRegistryModule& AssetRegistryModule = FModuleManager::LoadModuleChecked<FAssetRegistryModule>(AssetRegistryName);
+	IAssetRegistry& AssetRegistry = AssetRegistryModule.Get();
+
+	TSet<FName> Packages;
+
+	GetDependentPackages(RootPackages, Packages);
+
+	auto GetPackageFName = [](const FName& Name) { return Name; };
+
+#if VERIFY_GETDEPENDENTPACKAGES
 
 	TSet<UObject*> Dependencies; 
 	GetDependencies( RootPackages, Dependencies );
 
+	TSet<UPackage*> VerifyPackages;
+	for ( const auto& Object : Dependencies )
+	{
+		VerifyPackages.Add( Object->GetOutermost() );
+	}
+	check( VerifyPackages.Num() == Packages.Num() );
+
 
-	FName StandardFilename = GetCachedStandardPackageFileFName(Package);
+
+	for ( const auto& VerifyPackage : VerifyPackages )
+	{
+		check( Packages.Contains( VerifyPackage->GetFName() ) );
+
+		TArray<FName> PackagePaths;
+		AssetRegistry.GetPackagePathsByPackageName(VerifyPackage->GetFName(), PackagePaths);
+		check( PackagePaths.Num() == 0 );
+
+		check( PackagePaths[0].ToString() == VerifyPackage->GetPathName() );
+	}
+
+
+#endif
+
+#else
+	TSet<UObject*> Dependencies;
+	GetDependencies(RootPackages, Dependencies);
 
 	TSet<UPackage*> Packages;
-	for ( const auto& Object : Dependencies )
+	for (const auto& Object : Dependencies)
 	{
-		Packages.Add( Object->GetOutermost() );
+		Packages.Add(Object->GetOutermost());
 	}
 
-	if (CookByTheBookOptions->bGenerateDependeciesForMaps
-		&& Package->ContainsMap()
-		&& (World != nullptr)
-		)
+	auto GetPackageFName = [](const UPackage* Package) { return Package->GetFName(); };
+#endif
+
+
+
+
+	FName StandardFilename = GetCachedStandardPackageFileFName(Package);
+
+	if (CookByTheBookOptions &&
+		CookByTheBookOptions->bGenerateDependenciesForMaps &&
+		Package->ContainsMap() &&
+		(World != nullptr))
 	{
 		TSet <FName> Names;
 		for (auto Object : Packages)
 		{
-			Names.Add(Object->GetFName());
+			Names.Add(GetPackageFName(Object));
 		}
-		MapDependencyGraph.Add(Package->GetFName(), Names);
+		CookByTheBookOptions->MapDependencyGraph.Add(Package->GetFName(), Names);
 	}
 
 	// update the manifests with generated dependencies
-	for ( const auto& PlatformName : TargetPlatformNames )
+	for (const auto& PlatformName : TargetPlatformNames)
 	{
 		FChunkManifestGenerator* ManifestGenerator = CookByTheBookOptions->ManifestGenerators.FindChecked(PlatformName);
 
-		if ( CookByTheBookOptions->bGenerateStreamingInstallManifests )
+		if (CookByTheBookOptions && CookByTheBookOptions->bGenerateStreamingInstallManifests)
 		{
-			ManifestGenerator->PrepareToLoadNewPackage( StandardFilename.ToString() );
+			ManifestGenerator->PrepareToLoadNewPackage(StandardFilename.ToString());
 		}
 
-		for ( const auto& DependentPackage : Packages )
+		for (const auto& DependentPackage : Packages)
 		{
-			ManifestGenerator->OnLastPackageLoaded( DependentPackage  );
+			ManifestGenerator->OnLastPackageLoaded(DependentPackage);
 		}
-	
-		for ( const auto& DependentPackage  : Packages )
+
+		for (const auto& DependentPackage : Packages)
 		{
-			FString Filename = GetCachedPackageFilename( DependentPackage );
+			FString Filename = GetCachedPackageFilename(DependentPackage);
 			if (!Filename.IsEmpty())
 			{
 				// Populate streaming install manifests
 				FString SandboxFilename = ConvertToFullSandboxPath(*Filename, true);
 				//UE_LOG(LogCook, Display, TEXT("Adding package to manifest %s, %s, %s"), *DependentPackage->GetName(), *SandboxFilename, *LastLoadedMapName);
+#if USEASSETREGISTRYFORDEPENDENTPACKAGES
+				// Determine the package name and path
+				// FString PackageName = FPackageName::FilenameToLongPackageName(PackageFilename);
+				// FString PackagePath = FPackageName::GetLongPackagePath(DependentPackage.ToString());
+				FString PackagePath = DependentPackage.ToString();
+
+				ManifestGenerator->AddPackageToChunkManifest(DependentPackage, PackagePath, SandboxFilename, LastLoadedMapName, SandboxFile.GetOwnedPointer());
+#else
 				ManifestGenerator->AddPackageToChunkManifest(DependentPackage, SandboxFilename, LastLoadedMapName, SandboxFile.GetOwnedPointer());
+#endif
 			}
 		}
 	}
 }
 
-bool UCookOnTheFlyServer::IsCookingInEditor() const 
+bool UCookOnTheFlyServer::IsCookingInEditor() const
 {
 	return CurrentCookMode == ECookMode::CookByTheBookFromTheEditor || CurrentCookMode == ECookMode::CookOnTheFlyFromTheEditor;;
 }
@@ -1002,6 +1094,16 @@ uint32 UCookOnTheFlyServer::TickCookOnTheSide( const float TimeSlice, uint32 &Co
 			}
 		}
 		
+		if( CookByTheBookOptions && CookByTheBookOptions->bErrorOnEngineContentUse )
+		{
+			FString DLCPath = FPaths::GamePluginsDir() / CookByTheBookOptions->DlcName / FString(TEXT("Content"));
+			if ( ToBuild.GetFilename().ToString().StartsWith(DLCPath) == false ) // if we don't start with the dlc path then we shouldn't be cooking this data 
+			{
+				UE_LOG(LogCook, Error, TEXT("Engine content %s is being referenced by DLC!"), *ToBuild.GetFilename().ToString() );
+				bShouldCook = false;
+			}
+		}
+
 		if ( bShouldCook ) // if we should cook the package then cook it otherwise add it to the list of already cooked packages below
 		{
 			SCOPE_TIMER(AllOfLoadPackage);
@@ -2002,7 +2104,6 @@ void UCookOnTheFlyServer::Initialize( ECookMode::Type DesiredCookMode, ECookInit
 	if ( IsCookByTheBookMode() )
 	{
 		CookByTheBookOptions = new FCookByTheBookOptions();
-		CookByTheBookOptions->bGenerateStreamingInstallManifests = IsCookFlagSet(ECookInitializationFlags::GenerateStreamingInstallManifest);
 	}
 
 	ITargetPlatformManagerModule& TPM = GetTargetPlatformManagerRef();
@@ -2744,8 +2845,10 @@ void UCookOnTheFlyServer::CollectFilesToCook(TArray<FString>& FilesInPath, const
 			}
 		}
 
-		if (CookByTheBookOptions->bGenerateDependeciesForMaps) 
-			MapDependencyGraph.Add(FName(TEXT("ContentDirectoryAssets")), ContentDirectoryAssets);
+		if (CookByTheBookOptions && CookByTheBookOptions->bGenerateDependenciesForMaps) 
+		{
+			CookByTheBookOptions->MapDependencyGraph.Add(FName(TEXT("ContentDirectoryAssets")), ContentDirectoryAssets);
+		}
 	}
 }
 
@@ -2873,8 +2976,10 @@ void UCookOnTheFlyServer::CookByTheBookFinished()
 		}
 	}
 
-	if (CookByTheBookOptions->bGenerateDependeciesForMaps)
+	if (CookByTheBookOptions->bGenerateDependenciesForMaps)
 	{
+		auto& MapDependencyGraph = CookByTheBookOptions->MapDependencyGraph;
+
 		FString MapDependencyGraphFile = FPaths::GameDir()/ TEXT("MapDependencyGraph.json");
 		// dump dependency graph. 
 		FString DependencyString; 
@@ -3050,8 +3155,8 @@ void UCookOnTheFlyServer::StartCookByTheBook( const FCookByTheBookStartupOptions
 	CookByTheBookOptions->bCancel = false;
 	CookByTheBookOptions->CookTime = 0.0f;
 	CookByTheBookOptions->CookStartTime = FPlatformTime::Seconds();
-	CookByTheBookOptions->bGenerateDependeciesForMaps = CookByTheBookStartupOptions.bGenerateDependeciesForMaps;
-	
+	CookByTheBookOptions->bGenerateStreamingInstallManifests = CookByTheBookStartupOptions.bGenerateStreamingInstallManifests;
+	CookByTheBookOptions->bGenerateDependenciesForMaps = CookByTheBookStartupOptions.bGenerateDependenciesForMaps;
 	CookByTheBookOptions->CreateReleaseVersion = CreateReleaseVersion;
 
 	bool bCookAll = (CookOptions & ECookByTheBookOptions::CookAll) != ECookByTheBookOptions::None;
@@ -3085,7 +3190,7 @@ void UCookOnTheFlyServer::StartCookByTheBook( const FCookByTheBookStartupOptions
 			Platforms.Add( Platform );
 			ManifestGenerator = new FChunkManifestGenerator(Platforms);
 			ManifestGenerator->CleanManifestDirectories();
-			ManifestGenerator->Initialize( CookByTheBookOptions->bGenerateStreamingInstallManifests);
+			ManifestGenerator->Initialize( CookByTheBookOptions->bGenerateStreamingInstallManifests );
 
 			CookByTheBookOptions->ManifestGenerators.Add(PlatformName, ManifestGenerator);
 		}
