Missed file for CL# 2518076

diff --git a/Engine/Source/Editor/UnrealEd/Classes/Commandlets/CookCommandlet.h b/Engine/Source/Editor/UnrealEd/Classes/Commandlets/CookCommandlet.h
--- a/Engine/Source/Editor/UnrealEd/Classes/Commandlets/CookCommandlet.h
+++ b/Engine/Source/Editor/UnrealEd/Classes/Commandlets/CookCommandlet.h
@@ -36,6 +36,8 @@ class UCookCommandlet
 	bool bGenerateStreamingInstallManifests;
 	/** Error if we access engine content (useful for dlc) */
 	bool bErrorOnEngineContentUse;
+	/** Use historical serialization system for generating package dependencies (use for historical reasons only this method has been depricated, only affects cooked manifests) */
+	bool bUseSerializationForGeneratingPackageDependencies;
 	/** All commandline tokens */
 	TArray<FString> Tokens;
 	/** All commandline switches */
