Recompile dependent blueprints

UE-11941 When an event on a component that is called from an actor's constructor is deleted, the editor crashes

diff --git a/Engine/Source/Editor/KismetCompiler/Private/KismetCompilerModule.cpp b/Engine/Source/Editor/KismetCompiler/Private/KismetCompilerModule.cpp
--- a/Engine/Source/Editor/KismetCompiler/Private/KismetCompilerModule.cpp
+++ b/Engine/Source/Editor/KismetCompiler/Private/KismetCompilerModule.cpp
@@ -11,6 +11,7 @@
 
 #include "Editor/UnrealEd/Public/Kismet2/KismetDebugUtilities.h"
 #include "Editor/UnrealEd/Public/Kismet2/KismetReinstanceUtilities.h"
+#include "Editor/UnrealEd/Public/Kismet2/BlueprintEditorUtils.h"
 
 #include "UserDefinedStructureCompilerUtils.h"
 #include "Engine/UserDefinedStruct.h"
@@ -47,7 +48,7 @@ public:
 };
 
 // Compiles a blueprint.
-void FKismet2CompilerModule::CompileBlueprintInner(class UBlueprint* Blueprint, const FKismetCompilerOptions& CompileOptions, FCompilerResultsLog& Results, TArray<UObject*>* ObjLoaded)
+void FKismet2CompilerModule::CompileBlueprintInner(class UBlueprint* Blueprint, const FKismetCompilerOptions& CompileOptions, FCompilerResultsLog& Results, TSharedPtr<FBlueprintCompileReinstancer> Reinstancer, TArray<UObject*>* ObjLoaded)
 {
 	FBlueprintIsBeingCompiledHelper BeingCompiled(Blueprint);
 
@@ -60,6 +61,16 @@ void FKismet2CompilerModule::CompileBlueprintInner(class UBlueprint* Blueprint,
 	}
 	else
 	{
+		const uint32 PreviousSignatureCrc = Blueprint->CrcLastCompiledSignature;
+		const bool bIsFullCompile = CompileOptions.DoesRequireBytecodeGeneration() && (Blueprint->BlueprintType != BPTYPE_Interface);
+		const bool bRecompileDependencies = bIsFullCompile && !Blueprint->bIsRegeneratingOnLoad && Reinstancer.IsValid();
+
+		TArray<UBlueprint*> StoredDependentBlueprints;
+		if (bRecompileDependencies)
+		{
+			FBlueprintEditorUtils::GetDependentBlueprints(Blueprint, StoredDependentBlueprints);
+		}
+
 		// Loop through all external compiler delegates attempting to compile the blueprint.
 		bool Compiled = false;
 		for ( IBlueprintCompiler* Compiler : Compilers )
@@ -88,6 +99,22 @@ void FKismet2CompilerModule::CompileBlueprintInner(class UBlueprint* Blueprint,
 				check(Compiler.NewClass);
 			}
 		}
+
+		if (bRecompileDependencies)
+		{
+			Reinstancer->BlueprintWasRecompiled(Blueprint, CompileOptions.CompileType == EKismetCompileType::BytecodeOnly);
+
+			const bool bSignatureWasChanged = PreviousSignatureCrc != Blueprint->CrcLastCompiledSignature;
+			UE_LOG(LogK2Compiler, Verbose, TEXT("Signature of Blueprint '%s' %s changed"), *GetNameSafe(Blueprint), bSignatureWasChanged ? TEXT("was") : TEXT("was not"));
+
+			if (bSignatureWasChanged)
+			{
+				for (auto CurrentBP : StoredDependentBlueprints)
+				{
+					Reinstancer->EnlistDependentBlueprintToRecompile(CurrentBP, !CurrentBP->IsPossiblyDirty());
+				}
+			}
+		}
 	}
 
 	Blueprint->CurrentMessageLog = NULL;
@@ -125,7 +152,7 @@ void FKismet2CompilerModule::CompileBlueprint(class UBlueprint* Blueprint, const
 		SkeletonResults.bSilentMode = true;
 		FKismetCompilerOptions SkeletonCompileOptions;
 		SkeletonCompileOptions.CompileType = EKismetCompileType::SkeletonOnly;
-		CompileBlueprintInner(Blueprint, SkeletonCompileOptions, SkeletonResults, ObjLoaded);
+		CompileBlueprintInner(Blueprint, SkeletonCompileOptions, SkeletonResults, ParentReinstancer, ObjLoaded);
 	}
 
 	// If this was a full compile, take appropriate actions depending on the success of failure of the compile
@@ -134,7 +161,7 @@ void FKismet2CompilerModule::CompileBlueprint(class UBlueprint* Blueprint, const
 		BP_SCOPED_COMPILER_EVENT_STAT(EKismetCompilerStats_CompileGeneratedClass);
 
 		// Perform the full compile
-		CompileBlueprintInner(Blueprint, CompileOptions, Results, ObjLoaded);
+		CompileBlueprintInner(Blueprint, CompileOptions, Results, ParentReinstancer, ObjLoaded);
 
 		if (Results.NumErrors == 0)
 		{
@@ -177,7 +204,7 @@ void FKismet2CompilerModule::CompileBlueprint(class UBlueprint* Blueprint, const
 			FKismetCompilerOptions StubCompileOptions(CompileOptions);
 			StubCompileOptions.CompileType = EKismetCompileType::StubAfterFailure;
 
-			CompileBlueprintInner(Blueprint, StubCompileOptions, StubResults, ObjLoaded);
+			CompileBlueprintInner(Blueprint, StubCompileOptions, StubResults, StubReinstancer, ObjLoaded);
 
 			StubReinstancer->UpdateBytecodeReferences();
 			if( !Blueprint->bIsRegeneratingOnLoad )
diff --git a/Engine/Source/Editor/KismetCompiler/Public/KismetCompilerModule.h b/Engine/Source/Editor/KismetCompiler/Public/KismetCompilerModule.h
--- a/Engine/Source/Editor/KismetCompiler/Public/KismetCompilerModule.h
+++ b/Engine/Source/Editor/KismetCompiler/Public/KismetCompilerModule.h
@@ -80,7 +80,7 @@ public:
 	virtual TArray<IBlueprintCompiler*>& GetCompilers() override { return Compilers; }
 	// End implementation
 private:
-	void CompileBlueprintInner(class UBlueprint* Blueprint, const FKismetCompilerOptions& CompileOptions, FCompilerResultsLog& Results, TArray<UObject*>* ObjLoaded);
+	void CompileBlueprintInner(class UBlueprint* Blueprint, const FKismetCompilerOptions& CompileOptions, FCompilerResultsLog& Results, TSharedPtr<FBlueprintCompileReinstancer> Reinstancer, TArray<UObject*>* ObjLoaded);
 
 	TArray<IBlueprintCompiler*> Compilers;
 };
diff --git a/Engine/Source/Editor/UnrealEd/Private/Kismet2/KismetReinstanceUtilities.cpp b/Engine/Source/Editor/UnrealEd/Private/Kismet2/KismetReinstanceUtilities.cpp
--- a/Engine/Source/Editor/UnrealEd/Private/Kismet2/KismetReinstanceUtilities.cpp
+++ b/Engine/Source/Editor/UnrealEd/Private/Kismet2/KismetReinstanceUtilities.cpp
@@ -127,6 +127,8 @@ struct FReplaceReferenceHelper
 // FBlueprintCompileReinstancer
 
 TSet<TWeakObjectPtr<UBlueprint>> FBlueprintCompileReinstancer::DependentBlueprintsToRefresh = TSet<TWeakObjectPtr<UBlueprint>>();
+TSet<TWeakObjectPtr<UBlueprint>> FBlueprintCompileReinstancer::DependentBlueprintsToRecompile = TSet<TWeakObjectPtr<UBlueprint>>();
+TSet<TWeakObjectPtr<UBlueprint>> FBlueprintCompileReinstancer::DependentBlueprintsToByteRecompile = TSet<TWeakObjectPtr<UBlueprint>>();
 
 FBlueprintCompileReinstancer::FBlueprintCompileReinstancer(UClass* InClassToReinstance, bool bIsBytecodeOnly, bool bSkipGC)
 	: ClassToReinstance(InClassToReinstance)
@@ -237,7 +239,7 @@ FBlueprintCompileReinstancer::FBlueprintCompileReinstancer(UClass* InClassToRein
 		check(GeneratingBP || GIsAutomationTesting);
 		if(GeneratingBP)
 		{
-			ClassToReinstanceDefaultValuesCRC = GeneratingBP->CrcPreviousCompiledCDO;
+			ClassToReinstanceDefaultValuesCRC = GeneratingBP->CrcLastCompiledCDO;
 			FBlueprintEditorUtils::GetDependentBlueprints(GeneratingBP, Dependencies);
 		}
 	}
@@ -431,7 +433,7 @@ TSharedPtr<FReinstanceFinalizer> FBlueprintCompileReinstancer::ReinstanceInner(b
 			const UBlueprintGeneratedClass* BPClassB = Cast<const UBlueprintGeneratedClass>(ClassToReinstance);
 			const UBlueprint* BP = Cast<const UBlueprint>(ClassToReinstance->ClassGeneratedBy);
 
-			const bool bTheSameDefaultValues = (BP != nullptr) && (ClassToReinstanceDefaultValuesCRC != 0) && (BP->CrcPreviousCompiledCDO == ClassToReinstanceDefaultValuesCRC);
+			const bool bTheSameDefaultValues = (BP != nullptr) && (ClassToReinstanceDefaultValuesCRC != 0) && (BP->CrcLastCompiledCDO == ClassToReinstanceDefaultValuesCRC);
 			const bool bTheSameLayout = (BPClassA != nullptr) && (BPClassB != nullptr) && FStructUtils::TheSameLayout(BPClassA, BPClassB, true);
 			const bool bAllowedToDoFastPath = bTheSameDefaultValues && bTheSameLayout;
 			if (bAllowedToDoFastPath)
@@ -458,6 +460,34 @@ void FBlueprintCompileReinstancer::ListDependentBlueprintsToRefresh(const TArray
 	}
 }
 
+void FBlueprintCompileReinstancer::EnlistDependentBlueprintToRecompile(UBlueprint* BP, bool bBytecodeOnly)
+{
+	if (IsValid(BP))
+	{
+		if (bBytecodeOnly)
+		{
+			DependentBlueprintsToByteRecompile.Add(BP);
+		}
+		else
+		{
+			DependentBlueprintsToRecompile.Add(BP);
+		}
+	}
+}
+
+void FBlueprintCompileReinstancer::BlueprintWasRecompiled(UBlueprint* BP, bool bBytecodeOnly)
+{
+	if (IsValid(BP))
+	{
+		DependentBlueprintsToRefresh.Remove(BP);
+		DependentBlueprintsToByteRecompile.Remove(BP);
+		if (!bBytecodeOnly)
+		{
+			DependentBlueprintsToRecompile.Remove(BP);
+		}
+	}
+}
+
 void FBlueprintCompileReinstancer::ReinstanceObjects(bool bForceAlwaysReinstance)
 {
 	BP_SCOPED_COMPILER_EVENT_STAT(EKismetCompilerStats_ReinstanceObjects);
@@ -477,6 +507,30 @@ void FBlueprintCompileReinstancer::ReinstanceObjects(bool bForceAlwaysReinstance
 
 		if (QueueToReinstance.Num() && (QueueToReinstance[0] == SharedThis))
 		{
+			while (DependentBlueprintsToRecompile.Num())
+			{
+				auto Iter = DependentBlueprintsToRecompile.CreateIterator();
+				TWeakObjectPtr<UBlueprint> BPPtr = *Iter;
+				Iter.RemoveCurrent();
+				if (auto BP = BPPtr.Get())
+				{
+					FKismetEditorUtilities::CompileBlueprint(BP, false, bSkipGarbageCollection);
+				}
+			}
+
+			while (DependentBlueprintsToByteRecompile.Num())
+			{
+				auto Iter = DependentBlueprintsToByteRecompile.CreateIterator();
+				TWeakObjectPtr<UBlueprint> BPPtr = *Iter;
+				Iter.RemoveCurrent();
+				if (auto BP = BPPtr.Get())
+				{
+					FKismetEditorUtilities::RecompileBlueprintBytecode(BP);
+				}
+			}
+
+			ensure(0 == DependentBlueprintsToRecompile.Num());
+
 			TArray<TSharedPtr<FReinstanceFinalizer>> Finalizers;
 
 			// All children were recompiled. It's safe to reinstance.
@@ -488,13 +542,6 @@ void FBlueprintCompileReinstancer::ReinstanceObjects(bool bForceAlwaysReinstance
 					Finalizers.Push(Finalizer);
 				}
 				QueueToReinstance[Idx]->bHasReinstanced = true;
-
-				auto LocalClassToReinstance = QueueToReinstance[Idx]->ClassToReinstance;
-				UBlueprint* BP = LocalClassToReinstance ? Cast<UBlueprint>(LocalClassToReinstance->ClassGeneratedBy) : nullptr;
-				if (BP)
-				{
-					DependentBlueprintsToRefresh.Remove(BP);
-				}
 			}
 			QueueToReinstance.Empty();
 
diff --git a/Engine/Source/Editor/UnrealEd/Public/Kismet2/KismetReinstanceUtilities.h b/Engine/Source/Editor/UnrealEd/Public/Kismet2/KismetReinstanceUtilities.h
--- a/Engine/Source/Editor/UnrealEd/Public/Kismet2/KismetReinstanceUtilities.h
+++ b/Engine/Source/Editor/UnrealEd/Public/Kismet2/KismetReinstanceUtilities.h
@@ -17,6 +17,8 @@ class UNREALED_API FBlueprintCompileReinstancer : public TSharedFromThis<FBluepr
 protected:
 
 	static TSet<TWeakObjectPtr<UBlueprint>> DependentBlueprintsToRefresh;
+	static TSet<TWeakObjectPtr<UBlueprint>> DependentBlueprintsToRecompile;
+	static TSet<TWeakObjectPtr<UBlueprint>> DependentBlueprintsToByteRecompile;
 
 	/** Reference to the class we're actively reinstancing */
 	UClass* ClassToReinstance;
@@ -57,6 +59,8 @@ public:
 	// End of FSerializableObject interface
 
 	void ListDependentBlueprintsToRefresh(const TArray<UBlueprint*>& DependentBPs);
+	void EnlistDependentBlueprintToRecompile(UBlueprint* BP, bool bBytecodeOnly);
+	void BlueprintWasRecompiled(UBlueprint* BP, bool bBytecodeOnly);
 
 	static TSharedPtr<FBlueprintCompileReinstancer> Create(UClass* InClassToReinstance, bool bIsBytecodeOnly = false, bool bSkipGC = false)
 	{
