Editor: Brought back the 'r.InvalidateCachedShaders' exec command, which was inadvertently removed in CL# 2515212

diff --git a/Engine/Source/Editor/UnrealEd/Private/Editor.cpp b/Engine/Source/Editor/UnrealEd/Private/Editor.cpp
--- a/Engine/Source/Editor/UnrealEd/Private/Editor.cpp
+++ b/Engine/Source/Editor/UnrealEd/Private/Editor.cpp
@@ -1213,4 +1213,81 @@ void FCachedActorLabels::Populate(UWorld* World, const TSet<AActor*>& IgnoredAct
 
 //////////////////////////////////////////////////////////////////////////
 
+void ExecuteInvalidateCachedShaders(const TArray< FString >& Args)
+{
+	if(Args.Num() == 0)
+	{
+		// todo: log error, at least one command is needed
+		UE_LOG(LogConsoleResponse, Display, TEXT("r.InvalidateCachedShaders failed\nAs this command should not be executed accidentally it requires you to specify an extra parameter."));
+		return;
+	}
+
+	FString FileName = FPaths::EngineDir() + TEXT("Shaders/ShaderVersion.usf");
+
+	FileName = IFileManager::Get().ConvertToAbsolutePathForExternalAppForRead(*FileName);
+
+	ISourceControlProvider& SourceControlProvider = ISourceControlModule::Get().GetProvider();
+	SourceControlProvider.Init();
+
+	FSourceControlStatePtr SourceControlState = SourceControlProvider.GetState(FileName, EStateCacheUsage::ForceUpdate);
+	if(SourceControlState.IsValid())
+	{
+		if( SourceControlState->CanCheckout() || SourceControlState->IsCheckedOutOther() )
+		{
+			if(SourceControlProvider.Execute(ISourceControlOperation::Create<FCheckOut>(), FileName) == ECommandResult::Failed)
+			{
+				UE_LOG(LogConsoleResponse, Display, TEXT("r.InvalidateCachedShaders failed\nCouldn't check out \"ShaderVersion.usf\""));
+				return;
+			}
+		}
+		else if(!SourceControlState->IsSourceControlled())
+		{
+			UE_LOG(LogConsoleResponse, Display, TEXT("r.InvalidateCachedShaders failed\n\"ShaderVersion.usf\" is not under source control."));
+		}
+		else if(SourceControlState->IsCheckedOutOther())
+		{
+			UE_LOG(LogConsoleResponse, Display, TEXT("r.InvalidateCachedShaders failed\n\"ShaderVersion.usf\" is already checked out by someone else\n(UE4 SourceControl needs to be fixed to allow multiple checkout.)"));
+			return;
+		}
+		else if(SourceControlState->IsDeleted())
+		{
+			UE_LOG(LogConsoleResponse, Display, TEXT("r.InvalidateCachedShaders failed\n\"ShaderVersion.usf\" is marked for delete"));
+			return;
+		}
+	}
+
+	IPlatformFile& PlatformFile = FPlatformFileManager::Get().GetPlatformFile();
+
+	IFileHandle* FileHandle = PlatformFile.OpenWrite(*FileName);
+	if(FileHandle)
+	{
+		FString Guid = FString(
+			TEXT("// This file is automatically generated by the console command r.InvalidateCachedShaders\n")
+			TEXT("// Each time the console command is executed it generates a new GUID. As this file is included\n")
+			TEXT("// in common.usf (which should be included in any shader) it allows to invalidate the shader DDC.\n")
+			TEXT("// \n")
+			TEXT("// GUID = "))
+			+ FGuid::NewGuid().ToString();
+
+		FileHandle->Write((const uint8*)TCHAR_TO_ANSI(*Guid), Guid.Len());
+		delete FileHandle;
+
+		UE_LOG(LogConsoleResponse, Display, TEXT("r.InvalidateCachedShaders succeeded\n\"ShaderVersion.usf\" was updated.\n"));
+	}
+	else
+	{
+		UE_LOG(LogConsoleResponse, Display, TEXT("r.InvalidateCachedShaders failed\nCouldn't open \"ShaderVersion.usf\".\n"));
+	}
+}
+
+FAutoConsoleCommand InvalidateCachedShaders(
+	TEXT("r.InvalidateCachedShaders"),
+	TEXT("Invalidate shader cache by making a unique change to ShaderVersion.usf which is included in common.usf.")
+	TEXT("To initiate actual the recompile of all shaders use \"recompileshaders changed\" or press \"Ctrl Shift .\".\n")
+	TEXT("The ShaderVersion.usf file should be automatically checked out but  it needs to be checked in to have effect on other machines."),
+	FConsoleCommandWithArgsDelegate::CreateStatic(ExecuteInvalidateCachedShaders)
+	);
+
+//////////////////////////////////////////////////////////////////////////
+
 #undef LOCTEXT_NAMESPACE
