UE-13859 Setting inherited variables in child blueprint crash editor

Children Blueprint are handled as dependent.

diff --git a/Engine/Source/Editor/UnrealEd/Private/Kismet2/BlueprintEditorUtils.cpp b/Engine/Source/Editor/UnrealEd/Private/Kismet2/BlueprintEditorUtils.cpp
--- a/Engine/Source/Editor/UnrealEd/Private/Kismet2/BlueprintEditorUtils.cpp
+++ b/Engine/Source/Editor/UnrealEd/Private/Kismet2/BlueprintEditorUtils.cpp
@@ -2615,6 +2615,9 @@ void FBlueprintEditorUtils::GatherDependencies(const UBlueprint* InBlueprint, TS
 
 	check(InBlueprint);
 	Dependencies.Empty();
+
+	FGatherDependenciesHelper::ProcessHierarchy(InBlueprint->ParentClass, Dependencies);
+
 	for (const auto& InterfaceDesc : InBlueprint->ImplementedInterfaces)
 	{
 		UBlueprint* InterfaceBP = InterfaceDesc.Interface ? Cast<UBlueprint>(InterfaceDesc.Interface->ClassGeneratedBy) : NULL;
diff --git a/Engine/Source/Editor/UnrealEd/Private/Kismet2/Kismet2.cpp b/Engine/Source/Editor/UnrealEd/Private/Kismet2/Kismet2.cpp
--- a/Engine/Source/Editor/UnrealEd/Private/Kismet2/Kismet2.cpp
+++ b/Engine/Source/Editor/UnrealEd/Private/Kismet2/Kismet2.cpp
@@ -923,6 +923,18 @@ void FKismetEditorUtilities::RecompileBlueprintBytecode(UBlueprint* BlueprintObj
 	{
 		BlueprintPackage->SetDirtyFlag(bStartedWithUnsavedChanges);
 	}
+
+	if (!BlueprintObj->bIsRegeneratingOnLoad)
+	{
+		BP_SCOPED_COMPILER_EVENT_STAT(EKismetCompilerStats_NotifyBlueprintChanged);
+
+		BlueprintObj->BroadcastCompiled();
+
+		if (GEditor)
+		{
+			GEditor->BroadcastBlueprintCompiled();
+		}
+	}
 }
 
 /** Recompiles the bytecode of a blueprint only.  Should only be run for recompiling dependencies during compile on load */
diff --git a/Engine/Source/Editor/UnrealEd/Private/Kismet2/KismetReinstanceUtilities.cpp b/Engine/Source/Editor/UnrealEd/Private/Kismet2/KismetReinstanceUtilities.cpp
--- a/Engine/Source/Editor/UnrealEd/Private/Kismet2/KismetReinstanceUtilities.cpp
+++ b/Engine/Source/Editor/UnrealEd/Private/Kismet2/KismetReinstanceUtilities.cpp
@@ -592,6 +592,7 @@ void FBlueprintCompileReinstancer::UpdateBytecodeReferences()
 				continue;
 			}
 
+			bool bBPWasChanged = false;
 			// For each function defined in this blueprint, run through the bytecode, and update any refs from the old properties to the new
 			for( TFieldIterator<UFunction> FuncIter(BPClass, EFieldIteratorFlags::ExcludeSuper); FuncIter; ++FuncIter )
 			{
@@ -599,14 +600,22 @@ void FBlueprintCompileReinstancer::UpdateBytecodeReferences()
 				if( CurrentFunction->Script.Num() > 0 )
 				{
 					FArchiveReplaceObjectRef<UObject> ReplaceAr(CurrentFunction, FieldMappings, /*bNullPrivateRefs=*/ false, /*bIgnoreOuterRef=*/ true, /*bIgnoreArchetypeRef=*/ true);
+					bBPWasChanged |= (0 != ReplaceAr.GetCount());
 				}
 			}
 
 			FArchiveReplaceObjectRef<UObject> ReplaceInBPAr(*DependentBP, FieldMappings, false, true, true);
 			if (ReplaceInBPAr.GetCount())
 			{
+				bBPWasChanged = true;
 				UE_LOG(LogBlueprint, Log, TEXT("UpdateBytecodeReferences: %d references from %s was replaced in BP %s"), ReplaceInBPAr.GetCount(), *GetPathNameSafe(ClassToReinstance), *GetPathNameSafe(*DependentBP));
 			}
+
+			auto CompiledBlueprint = UBlueprint::GetBlueprintFromClass(ClassToReinstance);
+			if (bBPWasChanged && CompiledBlueprint && !CompiledBlueprint->bIsRegeneratingOnLoad)
+			{
+				DependentBlueprintsToRefresh.Add(*DependentBP);
+			}
 		}
 	}
 }
