Fix crashing when reimporting skeletal meshes which have been transacted.  This is unsafe and requires a transaction buffer reset.  Added a new method to examine the list of objects in a transaction buffer to determine if we should  reset the transaction buffer.

diff --git a/Engine/Source/Editor/UnrealEd/Classes/Editor/EditorEngine.h b/Engine/Source/Editor/UnrealEd/Classes/Editor/EditorEngine.h
--- a/Engine/Source/Editor/UnrealEd/Classes/Editor/EditorEngine.h
+++ b/Engine/Source/Editor/UnrealEd/Classes/Editor/EditorEngine.h
@@ -824,6 +824,7 @@ public:
 	bool RedoTransaction();
 	bool IsTransactionActive();
 	FText GetTransactionName() const;
+	bool IsObjectInTransactionBuffer( const UObject* Object ) const;
 
 	/**
 	 * Rebuilds the map.
diff --git a/Engine/Source/Editor/UnrealEd/Classes/Editor/TransBuffer.h b/Engine/Source/Editor/UnrealEd/Classes/Editor/TransBuffer.h
--- a/Engine/Source/Editor/UnrealEd/Classes/Editor/TransBuffer.h
+++ b/Engine/Source/Editor/UnrealEd/Classes/Editor/TransBuffer.h
@@ -109,6 +109,7 @@ public:
 	virtual bool IsObjectSerializationEnabled() override { return DisallowObjectSerialization == 0; }
 	virtual void SetPrimaryUndoObject( UObject* Object ) override;
 	virtual ITransaction* CreateInternalTransaction() override;
+	virtual bool IsObjectInTransationBuffer( const UObject* Object ) const override;
 	virtual bool IsActive() override
 	{
 		return ActiveCount > 0;
diff --git a/Engine/Source/Editor/UnrealEd/Classes/Editor/Transactor.h b/Engine/Source/Editor/UnrealEd/Classes/Editor/Transactor.h
--- a/Engine/Source/Editor/UnrealEd/Classes/Editor/Transactor.h
+++ b/Engine/Source/Editor/UnrealEd/Classes/Editor/Transactor.h
@@ -394,6 +394,7 @@ public:
 	void RemoveRecords( int32 Count = 1 );
 	int32 GetRecordCount() const;
 
+	const UObject* GetPrimaryObject() const { return PrimaryObject; }
 	/**
 	 * Outputs the contents of the ObjectMap to the specified output device.
 	 */
@@ -556,6 +557,9 @@ class UTransactor : public UObject
 	 */
 	virtual void SetPrimaryUndoObject( UObject* Object ) PURE_VIRTUAL(UTransactor::MakePrimaryUndoObject,);
 
+	virtual bool IsObjectInTransationBuffer( const UObject* Object ) const { return false; }
+
+
 	// @todo document
 	virtual ITransaction* CreateInternalTransaction() PURE_VIRTUAL(UTransactor::CreateInternalTransaction,return NULL;);
 };
diff --git a/Engine/Source/Editor/UnrealEd/Private/EditorServer.cpp b/Engine/Source/Editor/UnrealEd/Private/EditorServer.cpp
--- a/Engine/Source/Editor/UnrealEd/Private/EditorServer.cpp
+++ b/Engine/Source/Editor/UnrealEd/Private/EditorServer.cpp
@@ -1272,6 +1272,11 @@ FText UEditorEngine::GetTransactionName() const
 	return Trans->GetUndoContext(false).Title;
 }
 
+bool UEditorEngine::IsObjectInTransactionBuffer( const UObject* Object ) const
+{
+	return Trans->IsObjectInTransationBuffer(Object);
+}
+
 bool UEditorEngine::Map_Select( UWorld* InWorld, const TCHAR* Str, FOutputDevice& Ar)
 {
 	const FScopedTransaction Transaction( NSLOCTEXT("UnrealEd", "SelectBrushes", "Select Brushes") );
diff --git a/Engine/Source/Editor/UnrealEd/Private/EditorTransaction.cpp b/Engine/Source/Editor/UnrealEd/Private/EditorTransaction.cpp
--- a/Engine/Source/Editor/UnrealEd/Private/EditorTransaction.cpp
+++ b/Engine/Source/Editor/UnrealEd/Private/EditorTransaction.cpp
@@ -727,3 +727,21 @@ void UTransBuffer::SetPrimaryUndoObject(UObject* PrimaryObject)
 		}
 	}
 }
+
+bool UTransBuffer::IsObjectInTransationBuffer( const UObject* Object ) const
+{
+	TArray<UObject*> TransactionObjects;
+	for( const FTransaction& Transaction : UndoBuffer )
+	{
+		Transaction.GetTransactionObjects(TransactionObjects);
+
+		if( TransactionObjects.Contains(Object) )
+		{
+			return true;
+		}
+		
+		TransactionObjects.Reset();
+	}
+
+	return false;
+}
\ No newline at end of file
diff --git a/Engine/Source/Editor/UnrealEd/Private/Factories/EditorFactories.cpp b/Engine/Source/Editor/UnrealEd/Private/Factories/EditorFactories.cpp
--- a/Engine/Source/Editor/UnrealEd/Private/Factories/EditorFactories.cpp
+++ b/Engine/Source/Editor/UnrealEd/Private/Factories/EditorFactories.cpp
@@ -5945,6 +5945,12 @@ EReimportResult::Type UReimportFbxSkeletalMeshFactory::Reimport( UObject* Obj )
 
 		CleanUp();
 
+		// Reimporting can have dangerous effects if the mesh is still in the transaction buffer.  Reset the transaction buffer if this is the case
+		if( GEditor->IsObjectInTransactionBuffer( SkeletalMesh ) )
+		{
+			GEditor->ResetTransaction( LOCTEXT("ReimportSkeletalMeshTransactionReset", "Reimporting a skeletal mesh which was in the undo buffer") );
+		}
+
 		return bSuccess ? EReimportResult::Succeeded : EReimportResult::Failed;
 	}
 	else
