Fix crash when applying instance changes to blueprint

diff --git a/Engine/Source/Editor/UnrealEd/Private/Editor.cpp b/Engine/Source/Editor/UnrealEd/Private/Editor.cpp
--- a/Engine/Source/Editor/UnrealEd/Private/Editor.cpp
+++ b/Engine/Source/Editor/UnrealEd/Private/Editor.cpp
@@ -700,7 +700,7 @@ namespace EditorUtilities
 									StartIndex = FindTargetComponentIndex;
 									break;
 								}
-								CheckComponent = CastChecked<UActorComponent>(CheckComponent->GetArchetype(), ECastCheckedType::NullAllowed);
+								CheckComponent = Cast<UActorComponent>(CheckComponent->GetArchetype());
 							}
 							if (TargetComponent)
 							{
