Optimize compile time of EditorAssetAutomationTests by converting a macro to a template with a lambda to get the definition out of line.
This reduced compile time for this file from ~1m to ~1s

diff --git a/Engine/Source/Editor/UnrealEd/Private/Tests/EditorAssetAutomationTests.cpp b/Engine/Source/Editor/UnrealEd/Private/Tests/EditorAssetAutomationTests.cpp
--- a/Engine/Source/Editor/UnrealEd/Private/Tests/EditorAssetAutomationTests.cpp
+++ b/Engine/Source/Editor/UnrealEd/Private/Tests/EditorAssetAutomationTests.cpp
@@ -422,16 +422,20 @@ bool FLogAssetCreationStatsCommand::Update()
 	return true;
 }
 
+template<typename TAssetClass, typename TFactoryClass, typename ExtraCommandsFunc>
+void AssetTestCreate(const TCHAR* NamePrefix, const FString& CurrentTimestamp, const FString& GamePath, TArray< TSharedPtr<CreateAssetHelper::FCreateAssetInfo> >& AssetInfos, const TSharedPtr<CreateAssetHelper::FCreateAssetStats>& BuildStats, ExtraCommandsFunc ExtraCommands)
+{
+	FString NameString = FString::Printf(TEXT("%s_%s"), NamePrefix, *CurrentTimestamp);
+	TFactoryClass* FactoryInst = NewObject<TFactoryClass>();
+	ExtraCommands(FactoryInst);
+	TSharedPtr<CreateAssetHelper::FCreateAssetInfo> CreateInfo = MakeShareable(new CreateAssetHelper::FCreateAssetInfo(NameString, GamePath, TAssetClass::StaticClass(), FactoryInst, BuildStats));
+	AssetInfos.Add(CreateInfo);
+	ADD_LATENT_AUTOMATION_COMMAND(FCreateNewAssetCommand(CreateInfo));
+}
+
 //Macro to create a factory and queue up a command to create a given asset type
 #define ASSET_TEST_CREATE( TAssetClass, TFactoryClass, NamePrefix, ExtraCommands ) \
-{ \
-	FString NameString = FString::Printf(TEXT("%s_%s"), TEXT(#NamePrefix), *CurrentTimestamp); \
-	TFactoryClass* FactoryInst = NewObject<TFactoryClass>(); \
-	ExtraCommands \
-	TSharedPtr<CreateAssetHelper::FCreateAssetInfo> CreateInfo = MakeShareable(new CreateAssetHelper::FCreateAssetInfo(NameString, GamePath, TAssetClass::StaticClass(), FactoryInst, BuildStats)); \
-	AssetInfos.Add(CreateInfo); \
-	ADD_LATENT_AUTOMATION_COMMAND(FCreateNewAssetCommand(CreateInfo)); \
-}
+	AssetTestCreate<TAssetClass, TFactoryClass>(TEXT(#NamePrefix), CurrentTimestamp, GamePath, AssetInfos, BuildStats, [=](TFactoryClass* FactoryInst) { ExtraCommands });
 
 //Macro to create a factory by name and queue up a command to create an asset
 #define ASSET_TEST_CREATE_BY_NAME(TAssetClassName, TFactoryClassName, NamePrefix, ExtraCommands) \
