[UE-11968]  Fix for shutdown of diff and merge tool, as comment explains we are running too much logic on diff tool shutdown

diff --git a/Engine/Source/Editor/UnrealEd/Private/EditorCommandLineUtils.cpp b/Engine/Source/Editor/UnrealEd/Private/EditorCommandLineUtils.cpp
--- a/Engine/Source/Editor/UnrealEd/Private/EditorCommandLineUtils.cpp
+++ b/Engine/Source/Editor/UnrealEd/Private/EditorCommandLineUtils.cpp
@@ -426,11 +426,10 @@ static void EditorCommandLineUtilsImpl::RaiseEditorMessageBox(const FText& Title
 //------------------------------------------------------------------------------
 static void EditorCommandLineUtilsImpl::ForceCloseEditor()
 {
-	IMainFrameModule& MainFrameModule = IMainFrameModule::Get();
-	if (MainFrameModule.IsWindowInitialized())
-	{
-		MainFrameModule.RequestCloseEditor();
-	}
+	// We used to call IMainFrameModule::RequestCloseEditor, but that runs a lot of logic that should only be
+	// run for the real project editor (notably UThumbnailManager::CaptureProjectThumbnail was causing a crash on shutdown
+	// but INI serialization was running when it should not have as well). Instead, we just raise the QUIT_EDITOR command:
+	GEngine->DeferredCommands.Add(TEXT("QUIT_EDITOR"));
 }
 
 //------------------------------------------------------------------------------
