UE-11313 Changing the value of a parent class variable does not update after hot reload

diff --git a/Engine/Source/Editor/UnrealEd/Private/Kismet2/KismetReinstanceUtilities.cpp b/Engine/Source/Editor/UnrealEd/Private/Kismet2/KismetReinstanceUtilities.cpp
--- a/Engine/Source/Editor/UnrealEd/Private/Kismet2/KismetReinstanceUtilities.cpp
+++ b/Engine/Source/Editor/UnrealEd/Private/Kismet2/KismetReinstanceUtilities.cpp
@@ -156,7 +156,13 @@ FBlueprintCompileReinstancer::FBlueprintCompileReinstancer(UClass* InClassToRein
 		ClassToReinstance->ClassFlags &= ~CLASS_NewerVersionExists;
 		GIsDuplicatingClassForReinstancing = false;
 
+		auto BPClassToReinstance = Cast<UBlueprintGeneratedClass>(ClassToReinstance);
 		auto BPGDuplicatedClass = Cast<UBlueprintGeneratedClass>(DuplicatedClass);
+		if (BPGDuplicatedClass && BPClassToReinstance && BPClassToReinstance->OverridenArchetypeForCDO)
+		{
+			BPGDuplicatedClass->OverridenArchetypeForCDO = BPClassToReinstance->OverridenArchetypeForCDO;
+		}
+
 		auto DuplicatedClassUberGraphFunction = BPGDuplicatedClass ? BPGDuplicatedClass->UberGraphFunction : nullptr;
 		if (DuplicatedClassUberGraphFunction)
 		{
