UE-9450: Recalculates the view's location when rotating with the orbit mode in the editor viewport client.

diff --git a/Engine/Source/Editor/UnrealEd/Private/EditorViewportClient.cpp b/Engine/Source/Editor/UnrealEd/Private/EditorViewportClient.cpp
--- a/Engine/Source/Editor/UnrealEd/Private/EditorViewportClient.cpp
+++ b/Engine/Source/Editor/UnrealEd/Private/EditorViewportClient.cpp
@@ -1868,6 +1868,11 @@ void FEditorViewportClient::InputAxisForOrbit(FViewport* InViewport, const FVect
 	{
 		SetViewRotation( GetViewRotation() + FRotator( Rot.Pitch, -Rot.Yaw, Rot.Roll ) );
 		FEditorViewportStats::Using(IsPerspective() ? FEditorViewportStats::CAT_PERSPECTIVE_MOUSE_ORBIT_ROTATION : FEditorViewportStats::CAT_ORTHOGRAPHIC_MOUSE_ORBIT_ROTATION);
+		
+		/*
+		 * Recalculates the view location according to the new SetViewRotation() did earlier.
+		 */
+		SetViewLocation(ViewTransform.ComputeOrbitMatrix().Inverse().GetOrigin());
 	}
 	else if ( IsOrbitPanMode( InViewport ) )
 	{
