[UE-13370] - Blueprints that have a SCS that imply calls its parent can now be data only, also we clean up data left over from blueprints that were formerly not blueprint only but now are

diff --git a/Engine/Source/Editor/UnrealEd/Private/Kismet2/BlueprintEditorUtils.cpp b/Engine/Source/Editor/UnrealEd/Private/Kismet2/BlueprintEditorUtils.cpp
--- a/Engine/Source/Editor/UnrealEd/Private/Kismet2/BlueprintEditorUtils.cpp
+++ b/Engine/Source/Editor/UnrealEd/Private/Kismet2/BlueprintEditorUtils.cpp
@@ -1290,6 +1290,43 @@ struct FEditoronlyBlueprintHelper
 	}
 };
 
+/**
+	Procedure used to remove old function implementations and child properties from data only blueprints.
+	These blueprints have a 'fast path' compilation path but we need to make sure that any data regenerated 
+	by normal blueprint compilation is cleared here. If we don't then these functions and properties will
+	hang around hwen a class is converted from a real blueprint to a data only blueprint.
+*/
+static void RemoveStaleFunctions(UBlueprintGeneratedClass* Class, UBlueprint* Blueprint)
+{
+	if (Class == nullptr)
+	{
+		return;
+	}
+
+	// Removes all existing functions from the class, currently used 
+	TFieldIterator<UFunction> Fn(Class, EFieldIteratorFlags::ExcludeSuper);
+	if (Fn)
+	{
+		FString OrphanedClassString = FString::Printf(TEXT("ORPHANED_DATA_ONLY_%s"), *Class->GetName());
+		FName OrphanedClassName = MakeUniqueObjectName(GetTransientPackage(), UBlueprintGeneratedClass::StaticClass(), FName(*OrphanedClassString));
+		UClass* OrphanedClass = NewObject<UBlueprintGeneratedClass>(GetTransientPackage(), OrphanedClassName, RF_Public | RF_Transient);
+		OrphanedClass->ClassAddReferencedObjects = Class->AddReferencedObjects;
+
+		const ERenameFlags RenFlags = REN_DontCreateRedirectors | (Blueprint->bIsRegeneratingOnLoad ? REN_ForceNoResetLoaders : 0) | REN_NonTransactional | REN_DoNotDirty;
+
+		while (Fn)
+		{
+			Class->RemoveFunctionFromFunctionMap(*Fn);
+			Fn->Rename(nullptr, OrphanedClass, RenFlags);
+			++Fn;
+		}
+	}
+
+	Blueprint->GeneratedClass->Children = nullptr;
+	Blueprint->GeneratedClass->Bind();
+	Blueprint->GeneratedClass->StaticLink(true);
+}
+
 UClass* FBlueprintEditorUtils::RegenerateBlueprintClass(UBlueprint* Blueprint, UClass* ClassToRegenerate, UObject* PreviousCDO, TArray<UObject*>& ObjLoaded)
 {
 	bool bRegenerated = false;
@@ -1422,6 +1459,8 @@ UClass* FBlueprintEditorUtils::RegenerateBlueprintClass(UBlueprint* Blueprint, U
 		{
 			if (Blueprint->IsGeneratedClassAuthoritative() && (Blueprint->GeneratedClass != NULL))
 			{
+				RemoveStaleFunctions(Cast<UBlueprintGeneratedClass>(Blueprint->GeneratedClass), Blueprint);
+
 				check(PreviousCDO != NULL);
 				check(Blueprint->SkeletonGeneratedClass != NULL);
 
@@ -2747,10 +2786,32 @@ bool FBlueprintEditorUtils::IsDataOnlyBlueprint(const UBlueprint* Blueprint)
 	if (UserConstructionScript)
 	{
 		//Call parent construction script may be added automatically
-		if (UserConstructionScript->Nodes.Num() > 1)
+		UBlueprint* BlueprintParent = Cast<UBlueprint>(Blueprint->ParentClass->ClassGeneratedBy);
+		// just 1 entry node or just one entry node and a call to our super, which is DataOnly:
+		if ( !BlueprintParent && UserConstructionScript->Nodes.Num() > 1 )
 		{
 			return false;
 		}
+		else if (BlueprintParent)
+		{
+			// More than two nodes.. one of them must do something (same logic as above, but we have a call to super as well)
+			if (UserConstructionScript->Nodes.Num() > 2)
+			{
+				return false;
+			}
+			else
+			{
+				// Just make sure the nodes are trivial, if they aren't then we're not data only:
+				for (auto Node : UserConstructionScript->Nodes)
+				{
+					if (!Cast<UK2Node_FunctionEntry>(Node) &&
+						!Cast<UK2Node_CallParentFunction>(Node))
+					{
+						return false;
+					}
+				}
+			}
+		}
 	}
 
 	// No extra eventgraphs
diff --git a/Engine/Source/Runtime/CoreUObject/Public/UObject/Class.h b/Engine/Source/Runtime/CoreUObject/Public/UObject/Class.h
--- a/Engine/Source/Runtime/CoreUObject/Public/UObject/Class.h
+++ b/Engine/Source/Runtime/CoreUObject/Public/UObject/Class.h
@@ -1968,6 +1968,12 @@ public:
 		FuncMap.Add(NewFunction->GetFName(), NewFunction);
 	}
 
+	// Remove a function from the function map
+	void RemoveFunctionFromFunctionMap(UFunction* Function)
+	{
+		FuncMap.Remove(Function->GetFName());
+	}
+
 	UFunction* FindFunctionByName(FName InName, EIncludeSuperFlag::Type IncludeSuper = EIncludeSuperFlag::IncludeSuper) const;
 
 	// UObject interface.
