Fixes a bug similar to UE-9450 with planar translations in the editor viewport client.

diff --git a/Engine/Source/Editor/UnrealEd/Private/EditorViewportClient.cpp b/Engine/Source/Editor/UnrealEd/Private/EditorViewportClient.cpp
--- a/Engine/Source/Editor/UnrealEd/Private/EditorViewportClient.cpp
+++ b/Engine/Source/Editor/UnrealEd/Private/EditorViewportClient.cpp
@@ -1889,8 +1889,8 @@ void FEditorViewportClient::InputAxisForOrbit(FViewport* InViewport, const FVect
 
 		FVector TransformedDelta = RotMat.InverseFast().TransformVector(DeltaLocation);
 
-		SetViewLocation( GetViewLocation() + TransformedDelta );
 		SetLookAtLocation( GetLookAtLocation() + TransformedDelta );
+		SetViewLocation(ViewTransform.ComputeOrbitMatrix().Inverse().GetOrigin());
 
 		FEditorViewportStats::Using(IsPerspective() ? FEditorViewportStats::CAT_PERSPECTIVE_MOUSE_ORBIT_PAN : FEditorViewportStats::CAT_ORTHOGRAPHIC_MOUSE_ORBIT_PAN);
 	}
