Changed slow task widget scope stack pointer to be a weak pointer
This should hopefull alleviate any crashes caused by invalid pointers

This addresses UE-12075.

diff --git a/Engine/Source/Editor/UnrealEd/Private/FeedbackContextEditor.cpp b/Engine/Source/Editor/UnrealEd/Private/FeedbackContextEditor.cpp
--- a/Engine/Source/Editor/UnrealEd/Private/FeedbackContextEditor.cpp
+++ b/Engine/Source/Editor/UnrealEd/Private/FeedbackContextEditor.cpp
@@ -29,7 +29,7 @@ public:
 		SLATE_EVENT( FOnCancelClickedDelegate, OnCancelClickedDelegate )
 
 		/** The feedback scope stack that we are presenting to the user */
-		SLATE_ARGUMENT( const FScopedSlowTaskStack*, ScopeStack )
+		SLATE_ARGUMENT( TWeakPtr<FScopedSlowTaskStack>, ScopeStack )
 
 	SLATE_END_ARGS()
 
@@ -37,7 +37,7 @@ public:
 	void Construct( const FArguments& InArgs )
 	{
 		OnCancelClickedDelegate = InArgs._OnCancelClickedDelegate;
-		ScopeStack = InArgs._ScopeStack;
+		WeakStack = InArgs._ScopeStack;
 
 		// This is a temporary widget that needs to be updated over its entire lifespan => has an active timer registered for its entire lifespan
 		RegisterActiveTimer( 0.f, FWidgetActiveTimerDelegate::CreateSP( this, &SSlowTaskWidget::UpdateProgress ) );
@@ -144,6 +144,12 @@ private:
 	/** Updates the dynamic progress bars for this widget */
 	void UpdateDynamicProgressBars()
 	{
+		auto ScopeStack = WeakStack.Pin();
+		if (!ScopeStack.IsValid())
+		{
+			return;
+		}
+
 		static const double VisibleScopeThreshold = 0.5;
 
 		DynamicProgressIndices.Empty();
@@ -226,12 +232,17 @@ private:
 	/** The main text that we will display in the window */
 	FText GetPercentageText() const
 	{
-		const float ProgressInterp = ScopeStack->GetProgressFraction(0);
+		auto ScopeStack = WeakStack.Pin();
+		if (ScopeStack.IsValid())
+		{
+			const float ProgressInterp = ScopeStack->GetProgressFraction(0);
 
-		FFormatOrderedArguments Args;
-		Args.Add(int(ProgressInterp * 100));
+			FFormatOrderedArguments Args;
+			Args.Add(int(ProgressInterp * 100));
 
-		return FText::Format(NSLOCTEXT("FeedbackContextEditor", "MainProgressDisplayText", "{0}%"), Args);
+			return FText::Format(NSLOCTEXT("FeedbackContextEditor", "MainProgressDisplayText", "{0}%"), Args);
+		}
+		return FText();
 	}
 
 	/** Calculate the best font to display the main text with */
@@ -255,9 +266,13 @@ private:
 	/** Get the tint for a secondary progress bar */
 	FLinearColor GetSecondaryProgressBarTint(int32 Index) const
 	{
-		if (!DynamicProgressIndices.IsValidIndex(Index) || !ScopeStack->IsValidIndex(DynamicProgressIndices[Index]))
+		auto ScopeStack = WeakStack.Pin();
+		if (ScopeStack.IsValid())
 		{
-			return FLinearColor::White.CopyWithNewOpacity(0.25f);
+			if (!DynamicProgressIndices.IsValidIndex(Index) || !ScopeStack->IsValidIndex(DynamicProgressIndices[Index]))
+			{
+				return FLinearColor::White.CopyWithNewOpacity(0.25f);
+			}
 		}
 		return FLinearColor::White;
 	}
@@ -265,9 +280,13 @@ private:
 	/** Get the fractional percentage of completion for a progress bar */
 	TOptional<float> GetProgressFraction(int32 Index) const
 	{
-		if (DynamicProgressIndices.IsValidIndex(Index) && ScopeStack->IsValidIndex(DynamicProgressIndices[Index]))
+		auto ScopeStack = WeakStack.Pin();
+		if (ScopeStack.IsValid())
 		{
-			return ScopeStack->GetProgressFraction(DynamicProgressIndices[Index]);
+			if (DynamicProgressIndices.IsValidIndex(Index) && ScopeStack->IsValidIndex(DynamicProgressIndices[Index]))
+			{
+				return ScopeStack->GetProgressFraction(DynamicProgressIndices[Index]);
+			}
 		}
 		return TOptional<float>();
 	}
@@ -275,9 +294,13 @@ private:
 	/** Get the text to display for a progress bar */
 	FText GetProgressText(int32 Index) const
 	{
-		if (DynamicProgressIndices.IsValidIndex(Index) && ScopeStack->IsValidIndex(DynamicProgressIndices[Index]))
+		auto ScopeStack = WeakStack.Pin();
+		if (ScopeStack.IsValid())
 		{
-			return (*ScopeStack)[DynamicProgressIndices[Index]]->GetCurrentMessage();
+			if (DynamicProgressIndices.IsValidIndex(Index) && ScopeStack->IsValidIndex(DynamicProgressIndices[Index]))
+			{
+				return (*ScopeStack)[DynamicProgressIndices[Index]]->GetCurrentMessage();
+			}
 		}
 
 		return FText();
@@ -296,7 +319,7 @@ private:
 	FOnCancelClickedDelegate OnCancelClickedDelegate;
 
 	/** The scope stack that we are reflecting */
-	const FScopedSlowTaskStack* ScopeStack;
+	TWeakPtr<FScopedSlowTaskStack> WeakStack;
 
 	/** The vertical box containing the secondary progress bars */
 	TSharedPtr<SVerticalBox> SecondaryBars;
@@ -397,7 +420,7 @@ void FFeedbackContextEditor::StartSlowTask( const FText& Task, bool bShowCancelB
 
 			SlowTaskWindowRef->SetContent(
 				SNew(SSlowTaskWidget)
-				.ScopeStack(&ScopeStack)
+				.ScopeStack(ScopeStack)
 				.OnCancelClickedDelegate( OnCancelClicked )
 			);
 
@@ -464,7 +487,7 @@ void FFeedbackContextEditor::ProgressReported( const float TotalProgressInterp,
 	else
 	{
 		// Always show the top-most message
-		for (auto& Scope : ScopeStack)
+		for (auto& Scope : *ScopeStack)
 		{
 			const FText ThisMessage = Scope->GetCurrentMessage();
 			if (!ThisMessage.IsEmpty())
diff --git a/Engine/Source/Runtime/Core/Private/Misc/FeedbackContext.cpp b/Engine/Source/Runtime/Core/Private/Misc/FeedbackContext.cpp
--- a/Engine/Source/Runtime/Core/Private/Misc/FeedbackContext.cpp
+++ b/Engine/Source/Runtime/Core/Private/Misc/FeedbackContext.cpp
@@ -20,9 +20,9 @@ void FFeedbackContext::RequestUpdateUI(bool bForceUpdate)
 
 void FFeedbackContext::UpdateUI()
 {
-	if (ScopeStack.Num() != 0)
+	if (ScopeStack->Num() != 0)
 	{
-		ProgressReported(ScopeStack.GetProgressFraction(0), ScopeStack[0]->GetCurrentMessage());
+		ProgressReported(ScopeStack->GetProgressFraction(0), (*ScopeStack)[0]->GetCurrentMessage());
 	}
 }
 
diff --git a/Engine/Source/Runtime/Core/Public/Misc/FeedbackContext.h b/Engine/Source/Runtime/Core/Public/Misc/FeedbackContext.h
--- a/Engine/Source/Runtime/Core/Public/Misc/FeedbackContext.h
+++ b/Engine/Source/Runtime/Core/Public/Misc/FeedbackContext.h
@@ -31,7 +31,7 @@ public:
 	/** Public const access to the current state of the scope stack */
 	FORCEINLINE const FScopedSlowTaskStack& GetScopeStack() const
 	{
-		return ScopeStack;
+		return *ScopeStack;
 	}
 
 	/**** Legacy API - not deprecated as it's still in heavy use, but superceded by FScopedSlowTask ****/
@@ -80,8 +80,9 @@ public:
 
 	bool	TreatWarningsAsErrors;
 
-	FFeedbackContext() :
-		 TreatWarningsAsErrors( 0 )
+	FFeedbackContext()
+		: TreatWarningsAsErrors(0)
+		, ScopeStack(MakeShareable(new FScopedSlowTaskStack))
 	{}
 
 private:
@@ -93,7 +94,7 @@ protected:
 	friend FScopedSlowTask;
 
 	/** Stack of pointers to feedback scopes that are currently open */
-	FScopedSlowTaskStack ScopeStack;
+	TSharedRef<FScopedSlowTaskStack> ScopeStack;
 	TArray<TUniquePtr<FScopedSlowTask>> LegacyAPIScopes;
 
 	/** Ask that the UI be updated as a result of the scope stack changing */
@@ -193,7 +194,7 @@ public:
 
 		if (bEnabled)
 		{
-			Context.ScopeStack.Push(this);
+			Context.ScopeStack->Push(this);
 		}
 	}
 
@@ -208,7 +209,7 @@ public:
 				Context.FinalizeSlowTask();
 			}
 
-			FScopedSlowTaskStack& Stack = Context.ScopeStack;
+			FScopedSlowTaskStack& Stack = *Context.ScopeStack;
 			checkSlow(Stack.Num() != 0 && Stack.Last() == this);
 
 			Stack.Pop(false);
@@ -245,7 +246,7 @@ public:
 
 		if (bEnabled)
 		{
-			Context.RequestUpdateUI(bCreatedDialog || Context.ScopeStack[0] == this);
+			Context.RequestUpdateUI(bCreatedDialog || (*Context.ScopeStack)[0] == this);
 		}
 	}
 
