Auto reimport improvements

Added user-controlled threshold to wait before importing a particular file after it has changed.

Change includes two changes from GitHub user Sektor (pr 1024 - many thanks!):
    - Fixed failed assertion in LongPackageNameToFilename when renaming assets at the root of the content path
    - Fixed auto-reimport not detecting existing unloaded packages for a given name

diff --git a/Engine/Source/Editor/UnrealEd/Classes/Settings/EditorLoadingSavingSettings.h b/Engine/Source/Editor/UnrealEd/Classes/Settings/EditorLoadingSavingSettings.h
--- a/Engine/Source/Editor/UnrealEd/Classes/Settings/EditorLoadingSavingSettings.h
+++ b/Engine/Source/Editor/UnrealEd/Classes/Settings/EditorLoadingSavingSettings.h
@@ -108,6 +108,8 @@ public:
 	UPROPERTY(EditAnywhere, config, AdvancedDisplay,Category=AutoReimport, meta=(DisplayName="Directories to Monitor", ToolTip="Lists every directory to monitor for content changes. Can be virtual package paths (eg /Game/ or /MyPlugin/), or absolute paths on disk."))
 	TArray<FAutoReimportDirectoryConfig> AutoReimportDirectorySettings;
 
+	UPROPERTY(EditAnywhere, config, AdvancedDisplay,Category=AutoReimport, meta=(ClampMin=0, ClampMax=60, DisplayName="Import Threshold Time", ToolTip="Specifies an amount of time to wait before a specific file change is considered for auto reimport"))
+	float AutoReimportThreshold;
 	UPROPERTY(EditAnywhere, config, AdvancedDisplay, Category=AutoReimport, meta=(DisplayName="Auto Create Assets", ToolTip="When enabled, newly added source content files will be automatically imported into new assets."))
 	bool bAutoCreateAssets;
 	UPROPERTY(EditAnywhere, config, AdvancedDisplay, Category=AutoReimport, meta=(DisplayName="Auto Delete Assets", ToolTip="When enabled, deleting a source content file will automatically prompt the deletion of any related assets."))
diff --git a/Engine/Source/Editor/UnrealEd/Private/AutoReimport/AutoReimportManager.cpp b/Engine/Source/Editor/UnrealEd/Private/AutoReimport/AutoReimportManager.cpp
--- a/Engine/Source/Editor/UnrealEd/Private/AutoReimport/AutoReimportManager.cpp
+++ b/Engine/Source/Editor/UnrealEd/Private/AutoReimport/AutoReimportManager.cpp
@@ -185,12 +185,6 @@ private:
 	/** A list of packages to save when we've added a bunch of assets */
 	TArray<UPackage*> PackagesToSave;
 
-	/** The time when the last change to the cache was reported */
-	FTimeLimit StartProcessingDelay;
-
-	/** The cached number of unprocessed changes we currently have to process */
-	int32 CachedNumUnprocessedChanges;
-
 	/** Reentracy guard for when we are making changes to assets */
 	bool bGuardAssetChanges;
 
@@ -200,7 +194,6 @@ private:
 
 FAutoReimportManager::FAutoReimportManager()
 	: StateMachine(ECurrentState::Idle)
-	, StartProcessingDelay(0.5)
 	, bGuardAssetChanges(false)
 {
 	auto* Settings = GetMutableDefault<UEditorLoadingSavingSettings>();
@@ -348,8 +341,8 @@ void FAutoReimportManager::OnAssetRenamed(const FAssetData& AssetData, const FSt
 	// We move the file with the asset provided it is the only file referenced, and sits right beside the uasset file
 	if (SourceFilesRelativeToOldPath.Num() == 1 && !SourceFilesRelativeToOldPath[0].GetCharArray().ContainsByPredicate([](const TCHAR Char) { return Char == '/' || Char == '\\'; }))
 	{
-		const FString AbsoluteSrcPath = FPaths::ConvertRelativePathToFull(FPackageName::LongPackageNameToFilename(FPackageName::GetLongPackagePath(OldPath)));
-		const FString AbsoluteDstPath = FPaths::ConvertRelativePathToFull(FPackageName::LongPackageNameToFilename(AssetData.PackagePath.ToString()));
+		const FString AbsoluteSrcPath = FPaths::ConvertRelativePathToFull(FPackageName::LongPackageNameToFilename(FPackageName::GetLongPackagePath(OldPath) / TEXT("")));
+		const FString AbsoluteDstPath = FPaths::ConvertRelativePathToFull(FPackageName::LongPackageNameToFilename(AssetData.PackagePath.ToString() / TEXT("")));
 
 		const FString OldAssetName = FPackageName::GetLongPackageAssetName(FPackageName::ObjectPathToPackageName(OldPath));
 		FString NewFileName = FPaths::GetBaseFilename(SourceFilesRelativeToOldPath[0]);
@@ -578,40 +571,32 @@ TOptional<ECurrentState> FAutoReimportManager::Idle()
 		Monitor.Tick();
 	}
 
-	const int32 NumUnprocessedChanges = GetNumUnprocessedChanges();
-	if (NumUnprocessedChanges > 0)
+	const IAssetRegistry& AssetRegistry = FModuleManager::LoadModuleChecked<FAssetRegistryModule>("AssetRegistry").Get();
+	if (AssetRegistry.IsLoadingAssets())
+	{
+		return TOptional<ECurrentState>();		
+	}
+
+	if (GetNumUnprocessedChanges() > 0)
 	{
-		// Check if we have any more unprocessed changes. If we have, we wait a delay before processing them.
-		if (NumUnprocessedChanges != CachedNumUnprocessedChanges)
+		// We have some changes so kick off the process
+		for (auto& Monitor : DirectoryMonitors)
 		{
-			CachedNumUnprocessedChanges = NumUnprocessedChanges;
-			StartProcessingDelay.Reset();
+			Monitor.StartProcessing();
 		}
-		else
-		{
-			const IAssetRegistry& AssetRegistry = FModuleManager::LoadModuleChecked<FAssetRegistryModule>("AssetRegistry").Get();
-			if (StartProcessingDelay.Exceeded() && !AssetRegistry.IsLoadingAssets())
-			{
-				// Deal with changes to the file system second
-				for (auto& Monitor : DirectoryMonitors)
-				{
-					Monitor.StartProcessing();
-				}
 
-				// Check that we actually have anything to do before kicking off a process
-				const int32 WorkTotal = Utils::Reduce(DirectoryMonitors, [](const FContentDirectoryMonitor& Monitor, int32 Total){
-					return Total + Monitor.GetTotalWork();
-				}, 0);
+		// Check that we actually have anything to do before kicking off a process
+		const int32 WorkTotal = Utils::Reduce(DirectoryMonitors, [](const FContentDirectoryMonitor& Monitor, int32 Total){
+			return Total + Monitor.GetTotalWork();
+		}, 0);
 
-				if (WorkTotal > 0)
-				{
-					// Create a new feedback context override
-					FeedbackContextOverride = MakeShareable(new FReimportFeedbackContext);
-					FeedbackContextOverride->Initialize(SNew(SReimportFeedback, GetProgressText()));
+		if (WorkTotal > 0)
+		{
+			// Create a new feedback context override
+			FeedbackContextOverride = MakeShareable(new FReimportFeedbackContext);
+			FeedbackContextOverride->Initialize(SNew(SReimportFeedback, GetProgressText()));
 
-					return ECurrentState::ProcessAdditions;
-				}
-			}
+			return ECurrentState::ProcessAdditions;
 		}
 	}
 
@@ -620,8 +605,6 @@ TOptional<ECurrentState> FAutoReimportManager::Idle()
 
 void FAutoReimportManager::Cleanup()
 {
-	CachedNumUnprocessedChanges = 0;
-
 	FeedbackContextOverride->Destroy();
 	FeedbackContextOverride = nullptr;
 }
diff --git a/Engine/Source/Editor/UnrealEd/Private/AutoReimport/ContentDirectoryMonitor.cpp b/Engine/Source/Editor/UnrealEd/Private/AutoReimport/ContentDirectoryMonitor.cpp
--- a/Engine/Source/Editor/UnrealEd/Private/AutoReimport/ContentDirectoryMonitor.cpp
+++ b/Engine/Source/Editor/UnrealEd/Private/AutoReimport/ContentDirectoryMonitor.cpp
@@ -85,7 +85,13 @@ void FContentDirectoryMonitor::StartProcessing()
 {
 	WorkProgress = TotalWork = 0;
 
-	auto OutstandingChanges = Cache.GetOutstandingChanges();
+	// We only process things that haven't changed for a given threshold
+	auto& FileManager = IFileManager::Get();
+	const FDateTime Threshold = FDateTime::UtcNow() - FTimespan(0, 0, GetDefault<UEditorLoadingSavingSettings>()->AutoReimportThreshold);
+	auto OutstandingChanges = Cache.FilterOutstandingChanges([&](const FUpdateCacheTransaction& Transaction, const FDateTime& TimeOfChange){
+		return TimeOfChange < Threshold;
+	});
+
 	if (OutstandingChanges.Num() != 0)
 	{
 		const auto* Settings = GetDefault<UEditorLoadingSavingSettings>();
@@ -185,11 +191,7 @@ void FContentDirectoryMonitor::ProcessAdditions(const IAssetRegistry& Registry,
 		FString NewAssetName = ObjectTools::SanitizeObjectName(FPaths::GetBaseFilename(FullFilename));
 		FString PackagePath = PackageTools::SanitizePackageName(MountedContentPath / FPaths::GetPath(Addition.Filename.Get()) / NewAssetName);
 
-		if (FEditorFileUtils::IsMapPackageAsset(PackagePath))
-		{
-			Context.AddMessage(EMessageSeverity::Error, FText::Format(LOCTEXT("Error_ExistingMap", "Attempted to create asset with the same name as a map ({0})."), FText::FromString(PackagePath)));
-		}
-		else if (FindPackage(nullptr, *PackagePath))
+		if (FPackageName::DoesPackageExist(*PackagePath))
 		{
 			Context.AddMessage(EMessageSeverity::Warning, FText::Format(LOCTEXT("Warning_ExistingAsset", "Can't create a new asset at {0} - one already exists."), FText::FromString(PackagePath)));
 		}
@@ -300,13 +302,14 @@ void FContentDirectoryMonitor::ProcessModifications(const IAssetRegistry& Regist
 				{
 					const bool bAssetWasDirty = IsAssetDirty(Asset);
 
-					FString NewAssetName = ObjectTools::SanitizeObjectName(FPaths::GetBaseFilename(Change.Filename.Get()));
-					FString PackagePath = PackageTools::SanitizePackageName(MountedContentPath / FPaths::GetPath(Change.Filename.Get()));
+					const FString NewAssetName = ObjectTools::SanitizeObjectName(FPaths::GetBaseFilename(Change.Filename.Get()));
+					const FString PackagePath = PackageTools::SanitizePackageName(MountedContentPath / FPaths::GetPath(Change.Filename.Get()));
+					const FString FullDestPath = PackagePath / NewAssetName;
 
-					const FText SrcPathText = FText::FromString(Assets[0].PackagePath.ToString() / Assets[0].PackageName.ToString()),
-						DstPathText = FText::FromString(PackagePath / NewAssetName);
+					const FText SrcPathText = FText::FromString(Assets[0].PackageName.ToString()),
+						DstPathText = FText::FromString(FullDestPath);
 
-					if (FEditorFileUtils::IsMapPackageAsset(PackagePath) || FindPackage(nullptr, *PackagePath))
+					if (FPackageName::DoesPackageExist(*FullDestPath))
 					{
 						Context.AddMessage(EMessageSeverity::Warning, FText::Format(LOCTEXT("MoveWarning_ExistingAsset", "Can't move {0} to {1} - one already exists."), SrcPathText, DstPathText));
 					}
diff --git a/Engine/Source/Editor/UnrealEd/Private/AutoReimport/FileCache.cpp b/Engine/Source/Editor/UnrealEd/Private/AutoReimport/FileCache.cpp
--- a/Engine/Source/Editor/UnrealEd/Private/AutoReimport/FileCache.cpp
+++ b/Engine/Source/Editor/UnrealEd/Private/AutoReimport/FileCache.cpp
@@ -426,7 +426,6 @@ void FFileCache::Destroy()
 	DirectoryReader = nullptr;
 	AsyncFileHasher = nullptr;
 
-	PendingTransactions.Empty();
 	DirtyFiles.Empty();
 	CachedDirectoryState = FDirectoryState();
 
@@ -580,19 +579,25 @@ TOptional<FString> FFileCache::GetTransactionPath(const FString& InAbsolutePath)
 	}
 }
 
-void FFileCache::DiffDirtyFiles(const TSet<FImmutableString>& InDirtyFiles, TArray<FUpdateCacheTransaction>& InOutTransactions, const FDirectoryState* InFileSystemState) const
+void FFileCache::DiffDirtyFiles(TMap<FImmutableString, FDateTime>& InDirtyFiles, TArray<FUpdateCacheTransaction>& OutTransactions, const FDirectoryState* InFileSystemState, const FDateTime& ThresholdTime) const
 {
 	TArray<uint8> ScratchBuffer;
 	ScratchBuffer.SetNumUninitialized(1024*1024);
 
 	TMap<FImmutableString, FFileData> AddedFiles, ModifiedFiles;
-	TSet<FImmutableString> RemovedFiles;
+	TSet<FImmutableString> RemovedFiles, InvalidDirtyFiles;
 
 	auto& FileManager = IFileManager::Get();
 	auto& PlatformFile = FPlatformFileManager::Get().GetPlatformFile();
 
-	for (const auto& File : InDirtyFiles)
+	for (const auto& Pair : InDirtyFiles)
 	{
+		if (Pair.Value > ThresholdTime)
+		{
+			continue;
+		}
+
+		const auto& File = Pair.Key;
 		FString AbsoluteFilename = GetAbsolutePath(File.Get());
 		
 		const auto* CachedState = CachedDirectoryState.Files.Find(File);
@@ -610,6 +615,11 @@ void FFileCache::DiffDirtyFiles(const TSet<FImmutableString>& InDirtyFiles, TArr
 				{
 					ModifiedFiles.Add(File, FileData);
 				}
+				else
+				{
+					// File hasn't changed
+					InvalidDirtyFiles.Add(File);
+				}
 			}
 			else
 			{
@@ -621,6 +631,17 @@ void FFileCache::DiffDirtyFiles(const TSet<FImmutableString>& InDirtyFiles, TArr
 		{
 			RemovedFiles.Add(File);
 		}
+		else
+		{
+			// File doesn't exist, and isn't in the cache
+			InvalidDirtyFiles.Add(File);
+		}
+	}
+
+	// Remove any dirty files that aren't dirty
+	for (auto& Filename : InvalidDirtyFiles)
+	{
+		InDirtyFiles.Remove(Filename);
 	}
 
 	// Rename / move detection
@@ -636,7 +657,7 @@ void FFileCache::DiffDirtyFiles(const TSet<FImmutableString>& InDirtyFiles, TArr
 					if (AdIt.Value().FileHash == CachedState->FileHash)
 					{
 						// Found a move destination!
-						InOutTransactions.Add(FUpdateCacheTransaction(*RemoveIt, AdIt.Key(), AdIt.Value()));
+						OutTransactions.Add(FUpdateCacheTransaction(*RemoveIt, AdIt.Key(), AdIt.Value()));
 
 						AdIt.RemoveCurrent();
 						RemoveIt.RemoveCurrent();
@@ -649,31 +670,51 @@ void FFileCache::DiffDirtyFiles(const TSet<FImmutableString>& InDirtyFiles, TArr
 
 	for (auto& RemovedFile : RemovedFiles)
 	{
-		InOutTransactions.Add(FUpdateCacheTransaction(MoveTemp(RemovedFile), EFileAction::Removed));
+		OutTransactions.Add(FUpdateCacheTransaction(MoveTemp(RemovedFile), EFileAction::Removed));
 	}
 	// RemovedFiles is now bogus
 
 	for (auto& Pair : AddedFiles)
 	{
-		InOutTransactions.Add(FUpdateCacheTransaction(MoveTemp(Pair.Key), EFileAction::Added, Pair.Value));
+		OutTransactions.Add(FUpdateCacheTransaction(MoveTemp(Pair.Key), EFileAction::Added, Pair.Value));
 	}
 	// AddedFiles is now bogus
 	
 	for (auto& Pair : ModifiedFiles)
 	{
-		InOutTransactions.Add(FUpdateCacheTransaction(MoveTemp(Pair.Key), EFileAction::Modified, Pair.Value));
+		OutTransactions.Add(FUpdateCacheTransaction(MoveTemp(Pair.Key), EFileAction::Modified, Pair.Value));
 	}
 	// ModifiedFiles is now bogus
 }
 
 TArray<FUpdateCacheTransaction> FFileCache::GetOutstandingChanges()
 {
+	TArray<FUpdateCacheTransaction> PendingTransactions;
 	DiffDirtyFiles(DirtyFiles, PendingTransactions);
 	DirtyFiles.Empty();
+	return PendingTransactions;	
+}
+
+TArray<FUpdateCacheTransaction> FFileCache::FilterOutstandingChanges(const TFunctionRef<bool(const FUpdateCacheTransaction&, const FDateTime&)>& InPredicate)
+{
+	// We don't diff things that have only just changed, to ensure that add/delete pairs correctly get picked up as renames
+	FDateTime Threshold = FDateTime::UtcNow() - FTimespan(0,0,0.5);
 
-	TArray<FUpdateCacheTransaction> Moved;
-	Swap(PendingTransactions, Moved);
-	return Moved;
+	TArray<FUpdateCacheTransaction> AllTransactions;
+	DiffDirtyFiles(DirtyFiles, AllTransactions, nullptr, Threshold);
+
+	TArray<FUpdateCacheTransaction> FilteredTransactions;
+	for (auto& Transaction : AllTransactions)
+	{
+		auto TimeOfChange = DirtyFiles[Transaction.Filename];
+		if (InPredicate(Transaction, TimeOfChange))
+		{
+			DirtyFiles.Remove(Transaction.Filename);
+			FilteredTransactions.Add(MoveTemp(Transaction));
+		}
+	}
+	// Anything left in AllTransactions is discarded
+	return FilteredTransactions;
 }
 
 void FFileCache::IgnoreNewFile(const FString& Filename)
@@ -863,7 +904,8 @@ void FFileCache::ReadStateFromAsyncReader()
 	}
 
 	// Build up a list of dirty files
-	TSet<FImmutableString> LocalDirtyFiles;
+	TMap<FImmutableString, FDateTime> LocalDirtyFiles;
+	auto Now = FDateTime::UtcNow();
 
 	// We already have cached data so we need to compare it with the harvested data
 	// to detect additions, modifications, and removals
@@ -874,7 +916,7 @@ void FFileCache::ReadStateFromAsyncReader()
 		const auto* CachedData = CachedDirectoryState.Files.Find(Filename);
 		if (!CachedData || CachedData->Timestamp != FilenameAndData.Value.Timestamp)
 		{
-			LocalDirtyFiles.Add(Filename);
+			LocalDirtyFiles.Add(Filename, Now);
 		}
 	}
 
@@ -884,7 +926,7 @@ void FFileCache::ReadStateFromAsyncReader()
 		const FImmutableString& Filename = It.Key();
 		if (!LiveState->Files.Contains(Filename))
 		{
-			LocalDirtyFiles.Add(Filename);
+			LocalDirtyFiles.Add(Filename, Now);
 		}
 	}
 
@@ -940,7 +982,13 @@ void FFileCache::ReadStateFromAsyncReader()
 	}
 
 	// Now the only files left in LocalDirtyFiles are things that definitely apply to the current cache, and used to apply to the old cache
-	DirtyFiles = DirtyFiles.Union(LocalDirtyFiles);
+	for (auto& Pair : LocalDirtyFiles)
+	{
+		if (!DirtyFiles.Contains(Pair.Key))
+		{
+			DirtyFiles.Add(MoveTemp(Pair.Key), Pair.Value);
+		}
+	}
 
 	// Update the applicable extensions now that we've updated the cache
 	CachedDirectoryState.Rules = LiveState->Rules;
@@ -948,12 +996,13 @@ void FFileCache::ReadStateFromAsyncReader()
 
 void FFileCache::OnDirectoryChanged(const TArray<FFileChangeData>& FileChanges)
 {
+	FDateTime Now = FDateTime::UtcNow();
 	for (const auto& ThisEntry : FileChanges)
 	{
 		auto TransactionPath = GetTransactionPath(ThisEntry.Filename);
 		if (TransactionPath.IsSet())
 		{
-			DirtyFiles.Add(MoveTemp(TransactionPath.GetValue()));
+			DirtyFiles.Add(MoveTemp(TransactionPath.GetValue()), Now);
 		}
 	}
 }
\ No newline at end of file
diff --git a/Engine/Source/Editor/UnrealEd/Private/AutoReimport/FileCache.h b/Engine/Source/Editor/UnrealEd/Private/AutoReimport/FileCache.h
--- a/Engine/Source/Editor/UnrealEd/Private/AutoReimport/FileCache.h
+++ b/Engine/Source/Editor/UnrealEd/Private/AutoReimport/FileCache.h
@@ -377,9 +377,10 @@ public:
 	void IgnoreDeletedFile(const FString& Filename);
 
 	/** Get the number of pending changes to the cache. */
-	int32 GetNumOutstandingChanges() const { return DirtyFiles.Num() + PendingTransactions.Num(); }
+	int32 GetNumOutstandingChanges() const { return DirtyFiles.Num(); }
 
-	/** Get all pending changes to the cache. Transactions must be returned to CompleteTransaction to update the cache. */
+	/** Get pending changes to the cache. Transactions must be returned to CompleteTransaction to update the cache. */
+	TArray<FUpdateCacheTransaction> FilterOutstandingChanges(const TFunctionRef<bool(const FUpdateCacheTransaction&, const FDateTime&)>& InPredicate);
 	TArray<FUpdateCacheTransaction> GetOutstandingChanges();
 
 private:
@@ -389,8 +390,9 @@ private:
 
 	/** Diff the specified set of dirty files (absolute paths, or relative to the monitor directory), adding transactions to the specified array if necessary.
 	 *	Optionally takes a directory state from which we can retrieve current file system state, without having to ask the FS directly.
+	 *  Optionally exclude files that have changed since the specified threshold, to ensure that related events get grouped together correctly.
 	 */
-	void DiffDirtyFiles(const TSet<FImmutableString>& InOutDirtyFiles, TArray<FUpdateCacheTransaction>& InOutTransactions, const FDirectoryState* InFileSystemState = nullptr) const;
+	void DiffDirtyFiles(TMap<FImmutableString, FDateTime>& InDirtyFiles, TArray<FUpdateCacheTransaction>& OutTransactions, const FDirectoryState* InFileSystemState = nullptr, const FDateTime& ThreasholdTime = FDateTime::UtcNow()) const;
 
 	/** Get the absolute path from a transaction filename */
 	FString GetAbsolutePath(const FString& InTransactionPath) const;
@@ -416,11 +418,8 @@ private:
 	TSharedPtr<FAsyncDirectoryReader, ESPMode::ThreadSafe> DirectoryReader;
 	TSharedPtr<FAsyncFileHasher, ESPMode::ThreadSafe> AsyncFileHasher;
 
-	/** A set of dirty files that we will use to report changes to the user */
-	TSet<FImmutableString> DirtyFiles;
-
-	/** A list of pending change transactions to be returned to the user */
-	TArray<FUpdateCacheTransaction> PendingTransactions;
+	/** A map of dirty file times that we will use to report changes to the user */
+	TMap<FImmutableString, FDateTime> DirtyFiles;
 
 	/** Our in-memory view of the cached directory state. */
 	FDirectoryState CachedDirectoryState;
@@ -433,4 +432,5 @@ private:
 
 	/** The time we last retrieved file hashes from the thread */
 	double LastFileHashGetTime;
+
 };
diff --git a/Engine/Source/Editor/UnrealEd/Private/Settings/SettingsClasses.cpp b/Engine/Source/Editor/UnrealEd/Private/Settings/SettingsClasses.cpp
--- a/Engine/Source/Editor/UnrealEd/Private/Settings/SettingsClasses.cpp
+++ b/Engine/Source/Editor/UnrealEd/Private/Settings/SettingsClasses.cpp
@@ -102,6 +102,7 @@ UEditorLoadingSavingSettings::UEditorLoadingSavingSettings( const FObjectInitial
 	: Super(ObjectInitializer)
 	, bEnableSourceControlCompatabilityCheck(true)
 	, bMonitorContentDirectories(false)
+	, AutoReimportThreshold(3.f)
 	, bAutoCreateAssets(true)
 	, bAutoDeleteAssets(true)
 	, bDetectChangesOnRestart(true)
