Mac compile error fix

diff --git a/Engine/Source/Editor/UnrealEd/Private/AutoReimport/FileCache.cpp b/Engine/Source/Editor/UnrealEd/Private/AutoReimport/FileCache.cpp
--- a/Engine/Source/Editor/UnrealEd/Private/AutoReimport/FileCache.cpp
+++ b/Engine/Source/Editor/UnrealEd/Private/AutoReimport/FileCache.cpp
@@ -698,7 +698,7 @@ TArray<FUpdateCacheTransaction> FFileCache::GetOutstandingChanges()
 TArray<FUpdateCacheTransaction> FFileCache::FilterOutstandingChanges(const TFunctionRef<bool(const FUpdateCacheTransaction&, const FDateTime&)>& InPredicate)
 {
 	// We don't diff things that have only just changed, to ensure that add/delete pairs correctly get picked up as renames
-	FDateTime Threshold = FDateTime::UtcNow() - FTimespan(0,0,0.5);
+	FDateTime Threshold = FDateTime::UtcNow() - FTimespan(0,0,0,500);
 
 	TArray<FUpdateCacheTransaction> AllTransactions;
 	DiffDirtyFiles(DirtyFiles, AllTransactions, nullptr, Threshold);
