Added some  automated tests for collision

diff --git a/Engine/Config/BaseEditor.ini b/Engine/Config/BaseEditor.ini
--- a/Engine/Config/BaseEditor.ini
+++ b/Engine/Config/BaseEditor.ini
@@ -43,6 +43,61 @@ MaxNumberOfLogLines=2000
 DefaultSizeX=512
 DefaultSizeY=512
 
+[/Script/UnrealEd.CollisionAutomationTestConfigData]
+// Component sweeps
+// Capsule
++ComponentSweepMultiTests=(RootShapeAsset="/Engine/BasicShapes/Cone.Cone",ShapeType="TriggerCapsule",HitResult=(Time=0.25,Location=(X=0.0,Y=0.0,Z=149.999939),ImpactPoint=(X=0.070709,Y=0.0,Z=49.999985),Normal=(X=-0.001414,Y=-0.0,Z=0.999999),ImpactNormal=(X=0.0,Y=0.0,Z=1.0),TraceStart=(X=0.0,Y=0.0,Z=200.0),TraceEnd=(X=0.0,Y=0.0,Z=0.0)))
++ComponentSweepMultiTests=(RootShapeAsset="/Engine/BasicShapes/Cone.Cone",ShapeType="TriggerCapsule",HitResult=(Time=0.50,Location=(X=-100.0,Y=0.0,Z=0.0),ImpactPoint=(X=-50.0,Y=-0.000019,Z=-50.0),Normal=(X=-1.0,Y=0.0,Z=-0.0),ImpactNormal=(X=-0.888290,Y=0.116945,Z=0.444145),TraceStart=(X=-200.0,Y=0.0,Z=0.0),TraceEnd=(X=0.0,Y=0.0,Z=0.0)))
++ComponentSweepMultiTests=(RootShapeAsset="/Engine/BasicShapes/Cone.Cone",ShapeType="TriggerCapsule",HitResult=(Time=0.646994,Location=(X=-70.601135,Y=0.0,Z=70.601135),ImpactPoint=(X=-25.879774,Y=-0.000010,Z=-1.759551),Normal=(X=-0.894427,Y=0.0,Z=0.447214),ImpactNormal=(X=-0.888290,Y=0.116945,Z=0.444145),TraceStart=(X=-200.0,Y=0.0,Z=200.0),TraceEnd=(X=0.0,Y=0.0,Z=0.0)))
+
++ComponentSweepMultiTests=(RootShapeAsset="/Engine/BasicShapes/Cube.Cube",ShapeType="TriggerCapsule",HitResult=(Time=0.25,Location=(X=0.0,Y=0.0,Z=150.0),ImpactPoint=(X=0.0,Y=0.0,Z=50.0),Normal=(X=0.0,Y=-0.0,Z=1.0),ImpactNormal=(X=0.0,Y=0.0,Z=1.0),TraceStart=(X=0.0,Y=0.0,Z=200.0),TraceEnd=(X=0.0,Y=0.0,Z=0.0)))
++ComponentSweepMultiTests=(RootShapeAsset="/Engine/BasicShapes/Cube.Cube",ShapeType="TriggerCapsule",HitResult=(Time=0.50,Location=(X=-100.000008,Y=0.0,Z=0.0),ImpactPoint=(X=-50.0,Y=0.0,Z=0.0),Normal=(X=-1.0,Y=0.0,Z=0.0),ImpactNormal=(X=-1.0,Y=0.0,Z=0.0),TraceStart=(X=-200.0,Y=0.0,Z=0.0),TraceEnd=(X=0.0,Y=0.0,Z=0.0),PenetrationDepth=0.0)))
++ComponentSweepMultiTests=(RootShapeAsset="/Engine/BasicShapes/Cube.Cube",ShapeType="TriggerCapsule",HitResult=(Time=0.50,Location=(X=-100.000015,Y=0.0,Z=100.000015),ImpactPoint=(X=-50.0,Y=0.0,Z=50.0),Normal=(X=-1.0,Y=-0.0,Z=0.0),ImpactNormal=(X=0.0,Y=0.0,Z=1.0),TraceStart=(X=-200.0,Y=0.0,Z=200.0),TraceEnd=(X=0.0,Y=0.0,Z=0.0))
+
++ComponentSweepMultiTests=(RootShapeAsset="/Engine/BasicShapes/Cylinder.Cylinder",ShapeType="TriggerCapsule",HitResult=(Time=0.25,Location=(X=0.0,Y=0.0,Z=150.0),ImpactPoint=(X=0.0,Y=0.0,Z=50.0),Normal=(X=0.0,Y=-0.0,Z=1.0),ImpactNormal=(X=0.0,Y=0.0,Z=1.0),TraceStart=(X=0.0,Y=0.0,Z=200.0),TraceEnd=(X=0.0,Y=0.0,Z=0.0)))
++ComponentSweepMultiTests=(RootShapeAsset="/Engine/BasicShapes/Cylinder.Cylinder",ShapeType="TriggerCapsule",HitResult=(Time=0.50,Location=(X=-100.0,Y=0.0,Z=0.0),ImpactPoint=(X=-50.0,Y=0.000016,Z=30.0),Normal=(X=-1.0,Y=-0.0,Z=0.0),ImpactNormal=(X=-0.995185,Y=0.098017,Z=0.0),TraceStart=(X=-200.0,Y=0.0,Z=0.0),TraceEnd=(X=0.0,Y=0.0,Z=0.0)))
++ComponentSweepMultiTests=(RootShapeAsset="/Engine/BasicShapes/Cylinder.Cylinder",ShapeType="TriggerCapsule",HitResult=(Time=0.50,Location=(X=-100.000015,Y=0.0,Z=100.000015),ImpactPoint=(X=-50.0,Y=0.000016,Z=50.0),Normal=(X=-1.0,Y=-0.0,Z=0.0),ImpactNormal=(X=0.0,Y=0.0,Z=1.0),TraceStart=(X=-200.0,Y=0.0,Z=200.0),TraceEnd=(X=0.0,Y=0.0,Z=0.0)))
+
++ComponentSweepMultiTests=(RootShapeAsset="/Engine/BasicShapes/Sphere.Sphere",ShapeType="TriggerCapsule",HitResult=(Time=0.25,Location=(X=0.0,Y=0.0,Z=150.0),ImpactPoint=(X=0.0,Y=0.000002,Z=50.0),Normal=(X=0.0,Y=-0.0,Z=1.0),ImpactNormal=(X=0.028589,Y=0.094246,Z=0.995138),TraceStart=(X=0.0,Y=0.0,Z=200.0),TraceEnd=(X=0.0,Y=0.0,Z=0.0)))
++ComponentSweepMultiTests=(RootShapeAsset="/Engine/BasicShapes/Sphere.Sphere",ShapeType="TriggerCapsule",HitResult=(Time=0.50,Location=(X=-100.0,Y=0.0,Z=0.0),ImpactPoint=(X=-50.0,Y=0.000018,Z=0.000004),Normal=(X=-1.0,Y=-0.0,Z=0.0),ImpactNormal=(X=-0.990438,Y=0.097550,Z=-0.097550),TraceStart=(X=-200.0,Y=0.0,Z=0.0),TraceEnd=(X=0.0,Y=0.0,Z=0.0)))
++ComponentSweepMultiTests=(RootShapeAsset="/Engine/BasicShapes/Sphere.Sphere",ShapeType="TriggerCapsule",HitResult=(Time=0.544464,Location=(X=-91.107117,Y=0.0,Z=91.107117),ImpactPoint=(X=-46.193974,Y=0.000017,Z=19.134171),Normal=(X=-0.898263,Y=-0.0,Z=0.439459),ImpactNormal=(X=-0.878613,Y=-0.086536,Z=0.469628),TraceStart=(X=-200.0,Y=0.0,Z=200.0),TraceEnd=(X=0.0,Y=0.0,Z=0.0)))
+
++ComponentSweepMultiTests=(RootShapeAsset="/Engine/BasicShapes/Cone.Cone",ShapeType="TriggerSphere",HitResult=(Time=0.55,Location=(X=0.0,Y=0.0,Z=89.999924),ImpactPoint=(X=0.070709,Y=0.0,Z=49.999985),Normal=(X=-0.001768,Y=-0.0,Z=0.999998),ImpactNormal=(X=0.0,Y=0.0,Z=1.0),TraceStart=(X=0.0,Y=0.0,Z=200.0),TraceEnd=(X=0.0,Y=0.0,Z=0.0)))
++ComponentSweepMultiTests=(RootShapeAsset="/Engine/BasicShapes/Cone.Cone",ShapeType="TriggerSphere",HitResult=(Time=0.651393,Location=(X=-69.721359,Y=0.0,Z=0.0),ImpactPoint=(X=-33.944271,Y=-0.000013,Z=-17.888544),Normal=(X=-0.894427,Y=0.0,Z=0.447214),ImpactNormal=(X=-0.888290,Y=0.116945,Z=0.444145),TraceStart=(X=-200.0,Y=0.0,Z=0.0),TraceEnd=(X=0.0,Y=0.0,Z=0.0)))
++ComponentSweepMultiTests=(RootShapeAsset="/Engine/BasicShapes/Cone.Cone",ShapeType="TriggerSphere",HitResult=(Time=0.767595,Location=(X=-46.480911,Y=0.0,Z=46.480911),ImpactPoint=(X=-10.703819,Y=-0.000004,Z=28.592360),Normal=(X=-0.894427,Y=0.0,Z=0.447214),ImpactNormal=(X=-0.888290,Y=0.116945,Z=0.444145),TraceStart=(X=-200.0,Y=0.0,Z=200.0),TraceEnd=(X=0.0,Y=0.0,Z=0.0)))
+
++ComponentSweepMultiTests=(RootShapeAsset="/Engine/BasicShapes/Cube.Cube",ShapeType="TriggerSphere",HitResult=(Time=0.55,Location=(X=0.0,Y=0.0,Z=90.0),ImpactPoint=(X=0.0,Y=0.0,Z=50.0),Normal=(X=0.0,Y=-0.0,Z=1.0),ImpactNormal=(X=0.0,Y=0.0,Z=1.0),TraceStart=(X=0.0,Y=0.0,Z=200.0),TraceEnd=(X=0.0,Y=0.0,Z=0.0)))
++ComponentSweepMultiTests=(RootShapeAsset="/Engine/BasicShapes/Cube.Cube",ShapeType="TriggerSphere",HitResult=(Time=0.55,Location=(X=-90.0,Y=0.0,Z=0.0),ImpactPoint=(X=-50.0,Y=0.0,Z=-0.0),Normal=(X=-1.0,Y=-0.0,Z=0.0),ImpactNormal=(X=-1.0,Y=0.0,Z=0.0),TraceStart=(X=-200.0,Y=0.0,Z=0.0),TraceEnd=(X=0.0,Y=0.0,Z=0.0)))
++ComponentSweepMultiTests=(RootShapeAsset="/Engine/BasicShapes/Cube.Cube",ShapeType="TriggerSphere",HitResult=(Time=0.608579,Location=(X=-78.284286,Y=0.0,Z=78.284286),ImpactPoint=(X=-50.0,Y=0.0,Z=50.0),Normal=(X=-0.707107,Y=-0.0,Z=0.707107),ImpactNormal=(X=0.0,Y=0.0,Z=1.0),TraceStart=(X=-200.0,Y=0.0,Z=200.0),TraceEnd=(X=0.0,Y=0.0,Z=0.0)))
+
++ComponentSweepMultiTests=(RootShapeAsset="/Engine/BasicShapes/Cylinder.Cylinder",ShapeType="TriggerSphere",HitResult=(Time=0.55,Location=(X=0.0,Y=0.0,Z=90.0),ImpactPoint=(X=0.0,Y=0.0,Z=50.0),Normal=(X=0.0,Y=-0.0,Z=1.0),ImpactNormal=(X=0.0,Y=0.0,Z=1.0),TraceStart=(X=0.0,Y=0.0,Z=200.0),TraceEnd=(X=0.0,Y=0.0,Z=0.0)))
++ComponentSweepMultiTests=(RootShapeAsset="/Engine/BasicShapes/Cylinder.Cylinder",ShapeType="TriggerSphere",HitResult=(Time=0.55,Location=(X=-90.0,Y=0.0,Z=0.0),ImpactPoint=(X=-50.0,Y=0.000016,Z=-0.000001),Normal=(X=-1.0,Y=-0.0,Z=0.0),ImpactNormal=(X=-0.995185,Y=-0.098017,Z=0.0),TraceStart=(X=-200.0,Y=0.0,Z=0.0),TraceEnd=(X=0.0,Y=0.0,Z=0.0)))
++ComponentSweepMultiTests=(RootShapeAsset="/Engine/BasicShapes/Cylinder.Cylinder",ShapeType="TriggerSphere",HitResult=(Time=0.608579,Location=(X=-78.284286,Y=0.0,Z=78.284286),ImpactPoint=(X=-50.0,Y=0.000016,Z=50.0),Normal=(X=-0.707107,Y=-0.0,Z=0.707107),ImpactNormal=(X=0.0,Y=0.0,Z=1.0),TraceStart=(X=-200.0,Y=0.0,Z=200.0),TraceEnd=(X=0.0,Y=0.0,Z=0.0)))
+
++ComponentSweepMultiTests=(RootShapeAsset="/Engine/BasicShapes/Sphere.Sphere",ShapeType="TriggerSphere",HitResult=(Time=0.55,Location=(X=0.0,Y=0.0,Z=90.0),ImpactPoint=(X=0.0,Y=0.000002,Z=50.0),Normal=(X=0.0,Y=-0.0,Z=1.0),ImpactNormal=(X=0.028589,Y=0.094246,Z=0.995138),TraceStart=(X=0.0,Y=0.0,Z=200.0),TraceEnd=(X=0.0,Y=0.0,Z=0.0)))
++ComponentSweepMultiTests=(RootShapeAsset="/Engine/BasicShapes/Sphere.Sphere",ShapeType="TriggerSphere",HitResult=(Time=0.55,Location=(X=-90.0,Y=0.0,Z=0.0),ImpactPoint=(X=-50.0,Y=0.000018,Z=0.000004),Normal=(X=-1.0,Y=-0.0,Z=-0.0),ImpactNormal=(X=-0.990438,Y=0.097550,Z=-0.097550),TraceStart=(X=-200.0,Y=0.0,Z=0.0),TraceEnd=(X=0.0,Y=0.0,Z=0.0)))
++ComponentSweepMultiTests=(RootShapeAsset="/Engine/BasicShapes/Sphere.Sphere",ShapeType="TriggerSphere",HitResult=(Time=0.681802,Location=(X=-63.639618,Y=0.0,Z=63.639618),ImpactPoint=(X=-35.355339,Y=0.000014,Z=35.355339),Normal=(X=-0.707107,Y=-0.0,Z=0.707107),ImpactNormal=(X=-0.633158,Y=-0.062360,Z=0.771506),TraceStart=(X=-200.0,Y=0.0,Z=200.0),TraceEnd=(X=0.0,Y=0.0,Z=0.0)))
+
+
+//Line traces
++LineTraceSingleByChannelTests=(RootShapeAsset="/Engine/BasicShapes/Cone.Cone",HitResult=(Time=0.25,Location=(X=0.0,Y=0.0,Z=149.999939),ImpactPoint=(X=0.070709,Y=0.0,Z=49.999985),Normal=(X=-0.001414,Y=-0.0,Z=0.999999),ImpactNormal=(X=0.0,Y=0.0,Z=1.0),TraceStart=(X=0.0,Y=0.0,Z=200.0),TraceEnd=(X=0.0,Y=0.0,Z=0.0)))
++LineTraceSingleByChannelTests=(RootShapeAsset="/Engine/BasicShapes/Cone.Cone",HitResult=(Time=0.50,Location=(X=-100.0,Y=0.0,Z=0.0),ImpactPoint=(X=-50.0,Y=-0.000019,Z=-50.0),Normal=(X=-1.0,Y=0.0,Z=-0.0),ImpactNormal=(X=-0.888290,Y=0.116945,Z=0.444145),TraceStart=(X=-200.0,Y=0.0,Z=0.0),TraceEnd=(X=0.0,Y=0.0,Z=0.0)))
++LineTraceSingleByChannelTests=(RootShapeAsset="/Engine/BasicShapes/Cone.Cone",HitResult=(Time=0.646994,Location=(X=-70.601135,Y=0.0,Z=70.601135),ImpactPoint=(X=-25.879774,Y=-0.000010,Z=-1.759551),Normal=(X=-0.894427,Y=0.0,Z=0.447214),ImpactNormal=(X=-0.888290,Y=0.116945,Z=0.444145),TraceStart=(X=-200.0,Y=0.0,Z=200.0),TraceEnd=(X=0.0,Y=0.0,Z=0.0)))
+
++LineTraceSingleByChannelTests=(RootShapeAsset="/Engine/BasicShapes/Cube.Cube",HitResult=(Time=0.25,Location=(X=0.0,Y=0.0,Z=150.0),ImpactPoint=(X=0.0,Y=0.0,Z=50.0),Normal=(X=0.0,Y=-0.0,Z=1.0),ImpactNormal=(X=0.0,Y=0.0,Z=1.0),TraceStart=(X=0.0,Y=0.0,Z=200.0),TraceEnd=(X=0.0,Y=0.0,Z=0.0)))
++LineTraceSingleByChannelTests=(RootShapeAsset="/Engine/BasicShapes/Cube.Cube",HitResult=(Time=0.50,Location=(X=-100.000008,Y=0.0,Z=0.0),ImpactPoint=(X=-50.0,Y=0.0,Z=0.0),Normal=(X=-1.0,Y=0.0,Z=0.0),ImpactNormal=(X=-1.0,Y=0.0,Z=0.0),TraceStart=(X=-200.0,Y=0.0,Z=0.0),TraceEnd=(X=0.0,Y=0.0,Z=0.0),PenetrationDepth=0.0)))
++LineTraceSingleByChannelTests=(RootShapeAsset="/Engine/BasicShapes/Cube.Cube",HitResult=(Time=0.50,Location=(X=-100.000015,Y=0.0,Z=100.000015),ImpactPoint=(X=-50.0,Y=0.0,Z=50.0),Normal=(X=-1.0,Y=-0.0,Z=0.0),ImpactNormal=(X=0.0,Y=0.0,Z=1.0),TraceStart=(X=-200.0,Y=0.0,Z=200.0),TraceEnd=(X=0.0,Y=0.0,Z=0.0))
+
++LineTraceSingleByChannelTests=(RootShapeAsset="/Engine/BasicShapes/Cylinder.Cylinder",HitResult=(Time=0.25,Location=(X=0.0,Y=0.0,Z=150.0),ImpactPoint=(X=0.0,Y=0.0,Z=50.0),Normal=(X=0.0,Y=-0.0,Z=1.0),ImpactNormal=(X=0.0,Y=0.0,Z=1.0),TraceStart=(X=0.0,Y=0.0,Z=200.0),TraceEnd=(X=0.0,Y=0.0,Z=0.0)))
++LineTraceSingleByChannelTests=(RootShapeAsset="/Engine/BasicShapes/Cylinder.Cylinder",HitResult=(Time=0.50,Location=(X=-100.0,Y=0.0,Z=0.0),ImpactPoint=(X=-50.0,Y=0.000016,Z=30.0),Normal=(X=-1.0,Y=-0.0,Z=0.0),ImpactNormal=(X=-0.995185,Y=0.098017,Z=0.0),TraceStart=(X=-200.0,Y=0.0,Z=0.0),TraceEnd=(X=0.0,Y=0.0,Z=0.0)))
++LineTraceSingleByChannelTests=(RootShapeAsset="/Engine/BasicShapes/Cylinder.Cylinder",HitResult=(Time=0.50,Location=(X=-100.000015,Y=0.0,Z=100.000015),ImpactPoint=(X=-50.0,Y=0.000016,Z=50.0),Normal=(X=-1.0,Y=-0.0,Z=0.0),ImpactNormal=(X=0.0,Y=0.0,Z=1.0),TraceStart=(X=-200.0,Y=0.0,Z=200.0),TraceEnd=(X=0.0,Y=0.0,Z=0.0)))
+
++LineTraceSingleByChannelTests=(RootShapeAsset="/Engine/BasicShapes/Sphere.Sphere",HitResult=(Time=0.25,Location=(X=0.0,Y=0.0,Z=150.0),ImpactPoint=(X=0.0,Y=0.000002,Z=50.0),Normal=(X=0.0,Y=-0.0,Z=1.0),ImpactNormal=(X=0.028589,Y=0.094246,Z=0.995138),TraceStart=(X=0.0,Y=0.0,Z=200.0),TraceEnd=(X=0.0,Y=0.0,Z=0.0)))
++LineTraceSingleByChannelTests=(RootShapeAsset="/Engine/BasicShapes/Sphere.Sphere",HitResult=(Time=0.50,Location=(X=-100.0,Y=0.0,Z=0.0),ImpactPoint=(X=-50.0,Y=0.000018,Z=0.000004),Normal=(X=-1.0,Y=-0.0,Z=0.0),ImpactNormal=(X=-0.990438,Y=0.097550,Z=-0.097550),TraceStart=(X=-200.0,Y=0.0,Z=0.0),TraceEnd=(X=0.0,Y=0.0,Z=0.0)))
++LineTraceSingleByChannelTests=(RootShapeAsset="/Engine/BasicShapes/Sphere.Sphere",HitResult=(Time=0.544464,Location=(X=-91.107117,Y=0.0,Z=91.107117),ImpactPoint=(X=-46.193974,Y=0.000017,Z=19.134171),Normal=(X=-0.898263,Y=-0.0,Z=0.439459),ImpactNormal=(X=-0.878613,Y=-0.086536,Z=0.469628),TraceStart=(X=-200.0,Y=0.0,Z=200.0),TraceEnd=(X=0.0,Y=0.0,Z=0.0)))
+
+
+
 [/Script/UnrealEd.WorldThumbnailRenderer]
 GlobalOrbitPitchOffset=0
 GlobalOrbitYawOffset=0
@@ -97,7 +152,6 @@ MenuPriority=10
 [/Script/BlueprintGraph.EdGraphSchema_K2]
 +EditoronlyBPFunctionRedirects=(ClassName="GameplayStatics", OldFunctionName="CreateSaveGameObjectFromBlueprint", NewFunctionName="CreateSaveGameObject", BlueprintParamName="SaveGameBlueprint", ClassParamName="SaveGameClass")
 
-
 [BlueprintEditor.Palette]
 bUseLegacyLayout=false
 
diff --git a/Engine/Source/Editor/UnrealEd/Private/Tests/CollisionAutomationTests.cpp b/Engine/Source/Editor/UnrealEd/Private/Tests/CollisionAutomationTests.cpp
new file mode 100644
--- /dev/null
+++ b/Engine/Source/Editor/UnrealEd/Private/Tests/CollisionAutomationTests.cpp
@@ -0,0 +1,305 @@
+// Copyright 1998-2015 Epic Games, Inc. All Rights Reserved.
+
+
+#include "UnrealEd.h"
+#include "CollisionAutomationTests.h"
+
+#include "Tests/AutomationTestSettings.h"
+
+#if WITH_EDITOR
+#include "FileHelpers.h"
+#include "LevelEditor.h"
+#include "AutomationEditorCommon.h"
+#endif
+
+#include "AssetSelection.h"
+#include "AutomationCommon.h"
+#include "AutomationTest.h"
+#include "Engine/TriggerBase.h"
+#include "Engine/TriggerCapsule.h"
+#include "Engine/TriggerBox.h"
+#include "Engine/TriggerSphere.h"
+
+DEFINE_LOG_CATEGORY_STATIC(CollisionAutomationTestLog, Log, All);
+
+
+namespace CollisionAutomationTests
+{
+	FAutomationTestBase* TestBase;
+	// Return the currently active world
+	UWorld* GetAutomationWorld(const int32 TestFlags)
+ 	{
+		UWorld* World = nullptr;
+		if( TestFlags & EAutomationTestFlags::ATF_Game)
+		{
+			check(GEngine->GetWorldContexts().Num() == 1);
+			World = GEngine->GetWorldContexts()[0].World();
+		}
+		else
+		{
+			for (auto It = GEngine->GetWorldContexts().CreateConstIterator(); It; ++It)
+			{
+				const FWorldContext& Context = *It;
+
+				if (Context.WorldType == EWorldType::Editor )
+				{
+					World = Context.World();
+					break;
+				}
+			}
+		}
+
+ 		return World;
+ 	}
+	
+	// Create a shape mesh actor from a given asset name
+	AStaticMeshActor* CreateShapeMeshActor( const FString ShapeAssetName, FVector Location )
+	{
+		UStaticMesh* StaticMeshAsset = (UStaticMesh*)StaticLoadObject(UStaticMesh::StaticClass(), nullptr, *ShapeAssetName, nullptr, LOAD_None, nullptr);
+		TestBase->TestNotNull<UStaticMesh>(FString::Printf(TEXT("Failed to find mesh object %s."), *ShapeAssetName), StaticMeshAsset);
+
+		AStaticMeshActor* ShapeMeshActor = Cast<AStaticMeshActor>(FActorFactoryAssetProxy::AddActorForAsset(StaticMeshAsset));
+		ShapeMeshActor->SetActorLocation(Location);
+		TestBase->TestNotNull<AActor>(FString::Printf(TEXT("Failed to create MeshActor for %s."), *ShapeAssetName), ShapeMeshActor); 
+		return ShapeMeshActor;
+	}
+
+	// Create a collision shape actor from a given shape name
+	AActor* CreateCollisionShape( UWorld* World, const FString ShapeTypeName, FVector Location )
+	{
+		AActor* TestRayCollisionActor = nullptr;
+		FTransform CollisionTransform;
+		CollisionTransform.AddToTranslation(Location);
+		
+		if (ShapeTypeName == TEXT("TriggerCapsule"))
+		{
+			TestRayCollisionActor = Cast<ATriggerCapsule>(GEditor->AddActor(World->GetCurrentLevel(), ATriggerCapsule::StaticClass(), CollisionTransform));
+			if (TestRayCollisionActor != nullptr)
+			{
+				UCapsuleComponent* Capsule = Cast<UCapsuleComponent>(TestRayCollisionActor->GetRootComponent());
+				if (Capsule != nullptr)
+				{
+					Capsule->SetCapsuleHalfHeight(100.0f);
+					Capsule->SetCapsuleRadius(50.0f);
+				}
+			}
+		}
+		else if (ShapeTypeName == TEXT("TriggerBox"))
+		{
+			TestRayCollisionActor = Cast<ATriggerBox>(GEditor->AddActor(World->GetCurrentLevel(), ATriggerBox::StaticClass(), CollisionTransform));
+		}
+		else if (ShapeTypeName == TEXT("TriggerSphere"))
+		{
+			TestRayCollisionActor = Cast<ATriggerSphere>(GEditor->AddActor(World->GetCurrentLevel(), ATriggerSphere::StaticClass(), CollisionTransform));
+		}
+		TestBase->TestNotNull<AActor>(FString::Printf(TEXT("Failed to create Collision trigger%s."), *ShapeTypeName), TestRayCollisionActor); 
+		return TestRayCollisionActor;
+	}
+
+	void CheckVector( FVector ResultVector, FVector ExpectedVector, FString TestName, FString ParameterName, int32 TestIndex, float Tolerance = KINDA_SMALL_NUMBER )
+	{
+		FVector Delta = ExpectedVector - ResultVector;
+		float Diff = Delta.Size();
+		if (Diff > Tolerance)
+		{
+			//UE_LOG(CollisionAutomationTestLog, Log, TEXT("%d:HitResult=(%s)"), iTest+1, *OutHits[iHits].ToString());
+			TestBase->AddError(FString::Printf(TEXT("Test %d:%s %s mismatch. Should be %s but is actually %s."), TestIndex, *TestName, *ParameterName, *ExpectedVector.ToString(), *ResultVector.ToString()));
+		}
+	}
+
+	void CheckFloat(float ResultFloat, float ExpectedFloat, FString TestName, FString ParameterName, int32 TestIndex, float Tolerance = KINDA_SMALL_NUMBER)
+	{
+		float Diff = ExpectedFloat - ResultFloat;
+		if (Diff > Tolerance)
+		{
+			//UE_LOG(CollisionAutomationTestLog, Log, TEXT("%d:HitResult=(%s)"), iTest+1, *OutHits[iHits].ToString());
+			TestBase->AddError(FString::Printf(TEXT("Test %d:%s %s mismatch. Should be %df but is actually %f."), TestIndex, *TestName, *ParameterName, ExpectedFloat, ResultFloat));
+		}
+	}
+
+	FString HitToString(FHitResult& HitResult)
+	{
+		return FString::Printf(TEXT("Time=%f,Location=(X=%f,Y=%f,Z=%f),ImpactPoint=(X=%f,Y=%f,Z=%f),Normal=(X=%f,Y=%f,Z=%f),ImpactNormal=(X=%f,Y=%f,Z=%f),TraceStart=(X=%f,Y=%f,Z=%f),TraceEnd=(X=%f,Y=%f,Z=%f),PenetrationDepth=%f"),
+			HitResult.Time,
+			HitResult.Location.X,HitResult.Location.Y,HitResult.Location.Z,
+			HitResult.ImpactPoint.X,HitResult.ImpactPoint.Y,HitResult.ImpactPoint.Z,
+			HitResult.Normal.X,HitResult.Normal.Y,HitResult.Normal.Z,
+			HitResult.ImpactNormal.X,HitResult.ImpactNormal.Y,HitResult.ImpactNormal.Z,
+			HitResult.TraceStart.X,HitResult.TraceStart.Y,HitResult.TraceStart.Z,
+			HitResult.TraceEnd.X,HitResult.TraceEnd.Y,HitResult.TraceEnd.Z,
+			HitResult.PenetrationDepth
+		);
+	}
+}
+
+/**
+ * ComponentSweepMultiTest Verification
+ */
+IMPLEMENT_SIMPLE_AUTOMATION_TEST( FComponentSweepMultiTest, "System.Physics.Collision.ComponentSweepMulti", EAutomationTestFlags::ATF_Editor )
+
+/** 
+ * Perform some collision sweep tests. Creates a given shape mesh and checks collision normal against a collision shape type.
+ * Data for tests is in the [/Script/UnrealEd.CollisionAutomationTestConfigData] section of BaseEditor.ini
+ *
+ * @param Parameters - Unused for this test
+ * @return	TRUE if the test was successful, FALSE otherwise
+ */
+bool FComponentSweepMultiTest::RunTest(const FString& Parameters)
+{	
+	CollisionAutomationTests::TestBase = this;
+	// Create map
+	UWorld* World = AutomationEditorCommonUtils::CreateNewMap();
+	TestNotNull( TEXT("Failed to create world for Physics.Collision.Ray Test. Tests aborted."), World );
+
+	static FName TraceIdent = FName(TEXT("TestTrace"));
+	
+	FVector StartPos;
+	FVector EndPos;
+	ECollisionChannel Channel = ECC_WorldStatic;
+	
+	UCollisionAutomationTestConfigData* Data = UCollisionAutomationTestConfigData::StaticClass()->GetDefaultObject<UCollisionAutomationTestConfigData>();
+	
+	// Get the tests
+	for (int32 iTest = 0; iTest < Data->ComponentSweepMultiTests.Num(); iTest++)
+	{
+		FCollisionTestEntry OneElement = Data->ComponentSweepMultiTests[iTest];
+
+		// Create the Actor to check against
+		AStaticMeshActor* TestRayMeshActor = CollisionAutomationTests::CreateShapeMeshActor( *OneElement.RootShapeAsset, OneElement.HitResult.TraceEnd);
+		// Create the collision component
+		AActor* TestRayCollisionActor = CollisionAutomationTests::CreateCollisionShape( World, OneElement.ShapeType, OneElement.HitResult.TraceStart);
+		
+		if ((TestRayMeshActor != nullptr) && (TestRayCollisionActor != nullptr))
+		{			
+			// Set the collision profile and enable collision and physics
+			TestRayMeshActor->GetStaticMeshComponent()->BodyInstance.SetCollisionProfileName(TEXT("BlockAll"));
+			TestRayMeshActor->SetActorEnableCollision(true);
+			TestRayMeshActor->GetStaticMeshComponent()->BodyInstance.bSimulatePhysics = true;
+			UShapeComponent* CollisionComponent = Cast<UShapeComponent>(TestRayCollisionActor->GetRootComponent());
+			TestRayCollisionActor->SetActorEnableCollision(true);
+			if( CollisionComponent != nullptr )
+			{
+				CollisionComponent->SetCollisionProfileName(TEXT("BlockAll"));
+				CollisionComponent->SetSimulatePhysics(true);
+			}
+			
+			// Setup positions
+			StartPos = TestRayCollisionActor->GetActorLocation();
+			EndPos = TestRayMeshActor->GetActorLocation();
+			// Setup the query
+			FComponentQueryParams ShapeQueryParameters(TraceIdent, nullptr);
+			ShapeQueryParameters.bTraceComplex = true;
+			ShapeQueryParameters.bTraceAsyncScene = true;
+
+			// Perform test
+			TArray<FHitResult> OutHits;
+			bool WasBlocked = World->ComponentSweepMulti(OutHits, CollisionComponent, StartPos, EndPos, FRotator::ZeroRotator, ShapeQueryParameters);
+			bool BlockedBySpecified = false;
+			if (WasBlocked == true)
+			{
+				for (int32 iHits = 0; iHits < OutHits.Num(); iHits++)
+				{
+					AActor* EachActor = OutHits[iHits].GetActor();
+					if (EachActor == TestRayMeshActor)
+					{
+						BlockedBySpecified = true;	
+						// This generates a snippet you can copy/paste into the ini file for test validation
+						//UE_LOG(CollisionAutomationTestLog, Log, TEXT("%d:HitResult=(%s)"), iTest+1, *(CollisionAutomationTests::HitToString(OutHits[iHits])));
+
+						CollisionAutomationTests::CheckVector( OutHits[iHits].ImpactNormal, OneElement.HitResult.ImpactNormal, TEXT("ComponentSweepMulti"), TEXT("ImpactNormal"), iTest );
+						CollisionAutomationTests::CheckVector( OutHits[iHits].Normal, OneElement.HitResult.Normal, TEXT("ComponentSweepMulti"), TEXT("Normal"), iTest );
+						CollisionAutomationTests::CheckVector( OutHits[iHits].ImpactPoint, OneElement.HitResult.ImpactPoint, TEXT("ComponentSweepMulti"), TEXT("ImpactPoint"), iTest );
+						CollisionAutomationTests::CheckFloat( OutHits[iHits].Time, OneElement.HitResult.Time, TEXT("ComponentSweepMulti"), TEXT("Time"), iTest );
+					}
+				}
+			}
+			TestTrue(FString::Printf(TEXT("Test %d:ComponentSweepMulti from %s to %s failed. Should return blocking hit"), iTest+1, *TestRayMeshActor->GetName(), *TestRayCollisionActor->GetName()), BlockedBySpecified);
+		}
+		// Remove the actors
+		TestRayMeshActor->Destroy();
+		TestRayCollisionActor->Destroy();
+	}
+
+	return true;
+}
+
+IMPLEMENT_SIMPLE_AUTOMATION_TEST( FLineTraceSingleByChannel, "System.Physics.Collision.LineTraceSingleByChannel", EAutomationTestFlags::ATF_Editor )
+
+
+/** 
+ * Perform LineTraceSingleByChannel tests. Does a ray trace from a given point to a given shape mesh and verifies blocking is correct.
+ * Data for tests is in the [/Script/UnrealEd.CollisionAutomationTestConfigData] section of BaseEditor.ini
+ *
+ * @param Parameters - Unused for this test
+ * @return	TRUE if the test was successful, FALSE otherwise
+ */
+bool FLineTraceSingleByChannel::RunTest(const FString& Parameters)
+{	
+	// Create map
+	UWorld* World = AutomationEditorCommonUtils::CreateNewMap();
+	TestNotNull(TEXT("Failed to create world for Physics.Collision.Ray Test. Tests aborted."), World);
+
+	static FName TraceIdent = FName(TEXT("TestTrace"));
+
+	FVector StartPos;
+	FVector EndPos;
+	ECollisionChannel Channel = ECC_WorldStatic;
+
+	UCollisionAutomationTestConfigData* Data = UCollisionAutomationTestConfigData::StaticClass()->GetDefaultObject<UCollisionAutomationTestConfigData>();
+
+	for (int32 iTest = 0; iTest < Data->LineTraceSingleByChannelTests.Num(); iTest++)
+	{
+		FCollisionTestEntry OneElement = Data->LineTraceSingleByChannelTests[iTest];
+
+		AStaticMeshActor* TestRayMeshActor = CollisionAutomationTests::CreateShapeMeshActor(*OneElement.RootShapeAsset, OneElement.HitResult.TraceEnd);
+		
+		if (TestRayMeshActor != nullptr)
+		{
+			// Create the Actor to check against			
+			TestRayMeshActor->GetStaticMeshComponent()->BodyInstance.SetCollisionProfileName(TEXT("BlockAll"));
+			TestRayMeshActor->SetActorEnableCollision(true);
+			TestRayMeshActor->GetStaticMeshComponent()->BodyInstance.bSimulatePhysics = true;
+
+			// Setup trace start/end
+			StartPos = OneElement.HitResult.TraceStart;			
+			EndPos = TestRayMeshActor->GetActorLocation();
+						
+			// Do the trace
+			FHitResult OutHit;
+			bool WasBlocked = World->LineTraceSingleByChannel(OutHit,StartPos,EndPos,Channel);
+			bool BlockedBySpecified = false;
+			if (WasBlocked == true)
+			{
+				if (OutHit.GetActor() == TestRayMeshActor)
+				{
+					BlockedBySpecified = true;
+					// This generates a snippet you can copy/paste into the ini file for test validation
+					//UE_LOG(CollisionAutomationTestLog, Log, TEXT("%d:HitResult=(%s)"), iTest + 1, *(CollisionAutomationTests::HitToString(OutHit)));
+
+					CollisionAutomationTests::CheckVector(OutHit.ImpactNormal, OneElement.HitResult.ImpactNormal, TEXT("LineTraceSingleByChannel"), TEXT("ImpactNormal"), iTest);
+					CollisionAutomationTests::CheckVector(OutHit.Normal, OneElement.HitResult.Normal, TEXT("LineTraceSingleByChannel"), TEXT("Normal"), iTest);
+					CollisionAutomationTests::CheckVector(OutHit.ImpactPoint, OneElement.HitResult.ImpactPoint, TEXT("LineTraceSingleByChannel"), TEXT("ImpactPoint"), iTest);
+					CollisionAutomationTests::CheckFloat(OutHit.Time, OneElement.HitResult.Time, TEXT("LineTraceSingleByChannel"), TEXT("Time"), iTest);									
+				}
+			}
+			TestTrue(FString::Printf(TEXT("Test %d:LineTraceSingleByChannel to %s failed. Should return blocking hit"), iTest + 1, *TestRayMeshActor->GetName()), BlockedBySpecified);
+
+			// Change the collision profile and ensure we dont get a blocking hit
+			UShapeComponent* CollisionComponent = Cast<UShapeComponent>(TestRayMeshActor->GetRootComponent());			
+			if (CollisionComponent != nullptr)
+			{
+				CollisionComponent->SetCollisionProfileName(TEXT("OverlapAll"));
+				CollisionComponent->SetSimulatePhysics(true);
+			}
+				
+			TestRayMeshActor->GetStaticMeshComponent()->BodyInstance.SetCollisionProfileName(TEXT("OverlapAll"));
+			WasBlocked = World->LineTraceSingleByChannel(OutHit, StartPos, EndPos, Channel);
+			TestFalse(FString::Printf(TEXT("Test %d:LineTraceSingleByChannel to %s failed. Should not return blocking hit"), iTest + 1, *TestRayMeshActor->GetName()), WasBlocked);			
+		}
+		// Remove the actor
+		TestRayMeshActor->Destroy();
+	}
+
+	return true;
+}
+
