Fix progress update during build process.

diff --git a/Engine/Source/Editor/UnrealEd/Private/FeedbackContextEditor.cpp b/Engine/Source/Editor/UnrealEd/Private/FeedbackContextEditor.cpp
--- a/Engine/Source/Editor/UnrealEd/Private/FeedbackContextEditor.cpp
+++ b/Engine/Source/Editor/UnrealEd/Private/FeedbackContextEditor.cpp
@@ -477,7 +477,7 @@ void FFeedbackContextEditor::ProgressReported( const float TotalProgressInterp,
 			}
 
 			BuildProgressWidget->SetBuildProgressPercent(TotalProgressInterp * 100, 100);
-			TickSlate(SlowTaskWindow.Pin());
+			TickSlate(BuildProgressWindow.Pin());
 		}
 		else if (SlowTaskWindow.IsValid())
 		{
