Missing .h file

diff --git a/Engine/Source/Editor/UnrealEd/Private/Tests/CollisionAutomationTests.h b/Engine/Source/Editor/UnrealEd/Private/Tests/CollisionAutomationTests.h
new file mode 100644
--- /dev/null
+++ b/Engine/Source/Editor/UnrealEd/Private/Tests/CollisionAutomationTests.h
@@ -0,0 +1,64 @@
+// Copyright 1998-2015 Epic Games, Inc. All Rights Reserved.
+
+#pragma once
+
+#include "CollisionAutomationTests.generated.h"
+
+/**
+ * Container for detailing collision automated test data.
+ */
+USTRUCT()
+struct FCollisionTestEntry
+{
+	GENERATED_USTRUCT_BODY()
+
+	UPROPERTY()
+	FString RootShapeAsset;
+
+	UPROPERTY()
+	FString ShapeType;
+
+	UPROPERTY()
+	FHitResult	HitResult;
+
+	FCollisionTestEntry()
+	{
+	}
+};
+
+USTRUCT()
+struct FCollisionPerfTest
+{
+	GENERATED_USTRUCT_BODY()
+
+	UPROPERTY()
+	FString RootShapeAsset;
+
+	UPROPERTY()
+	FString ShapeType;
+
+	UPROPERTY()
+	FVector CreationBounds;
+
+	UPROPERTY()
+	FVector CreationElements;
+
+	FCollisionPerfTest()
+	{
+	}
+};
+
+UCLASS(config=Editor)
+class UCollisionAutomationTestConfigData : public UObject
+{
+public:
+	GENERATED_BODY()
+		
+	UPROPERTY(config)
+	TArray<FCollisionTestEntry>	ComponentSweepMultiTests;
+	UPROPERTY(config)
+	TArray<FCollisionTestEntry>	LineTraceSingleByChannelTests;
+	
+	UPROPERTY(config)
+	TArray<FCollisionPerfTest>	LineTracePerformanceTests;
+};
\ No newline at end of file
