#OpenWorld: HLOD

- Improved clustering by using sphere
- Fixed progress window to work during build process
- Optimized for a big map
- Disabled for build all since it is exceptionally slow, it might be painful to add there
- Fixed Merge Actor to work with material flattening

diff --git a/Engine/Source/Editor/UnrealEd/Private/EditorBuildUtils.cpp b/Engine/Source/Editor/UnrealEd/Private/EditorBuildUtils.cpp
--- a/Engine/Source/Editor/UnrealEd/Private/EditorBuildUtils.cpp
+++ b/Engine/Source/Editor/UnrealEd/Private/EditorBuildUtils.cpp
@@ -331,10 +331,10 @@ bool FEditorBuildUtils::EditorBuild( UWorld* InWorld, EBuildOptions::Type Id, co
 
 				// disable this in build all path, it's very slow, it can be painful
 				// they don't have any collision, but make sure it's added first before doing anything
- 				{
- 					BuildProgressWidget.Pin()->SetBuildType(SBuildProgressWidget::BUILDTYPE_LODs);
- 					TriggerHierarchicalLODBuilder(InWorld, Id);
- 				}
+// 				{
+// 					BuildProgressWidget.Pin()->SetBuildType(SBuildProgressWidget::BUILDTYPE_LODs);
+//					TriggerHierarchicalLODBuilder(InWorld, Id);
+// 				}
 
 				{
 					BuildProgressWidget.Pin()->SetBuildType(SBuildProgressWidget::BUILDTYPE_Paths);
diff --git a/Engine/Source/Editor/UnrealEd/Private/HierarchicalLOD.cpp b/Engine/Source/Editor/UnrealEd/Private/HierarchicalLOD.cpp
--- a/Engine/Source/Editor/UnrealEd/Private/HierarchicalLOD.cpp
+++ b/Engine/Source/Editor/UnrealEd/Private/HierarchicalLOD.cpp
@@ -19,7 +19,7 @@
 #include "GameFramework/WorldSettings.h"
 #include "Components/InstancedStaticMeshComponent.h"
 
-DEFINE_LOG_CATEGORY_STATIC(LogLODGenerator, Log, All);
+DEFINE_LOG_CATEGORY_STATIC(LogLODGenerator, Warning, All);
 
 #define LOCTEXT_NAMESPACE "HierarchicalLOD"
 
@@ -31,48 +31,71 @@ DEFINE_LOG_CATEGORY_STATIC(LogLODGenerator, Log, All);
 // Hierarchical LOD System implementation 
 ////////////////////////////////////////////////////////////////////////////////////////////
 
-// utility functions
-float CalculateFillingFactor(const FBox& ABox, const float AFillingFactor, const FBox& BBox, const float BFillingFactor)
+// utility function to calculate overlap of two spheres
+float CalculateOverlap(const FSphere& ASphere, const float AFillingFactor, const FSphere& BSphere, const float BFillingFactor)
 {
-	FBox OverlapBox = ABox.Overlap(BBox);
-	FBox UnionBox = ABox + BBox;
-	// it shouldn't be zero or it should be checked outside
-	ensure(UnionBox.GetSize() != FVector::ZeroVector);
+	// if it doesn't intersect, return zero 
+	if (!ASphere.Intersects(BSphere))
+	{
+		return 0.f;
+	}
 
-	// http://deim.urv.cat/~rivi/pub/3d/icra04b.pdf
-	// cost is calculated based on r^3 / filling factor
-	// since it subtract by AFillingFactor * 1/2 overlap volume + BfillingFactor * 1/2 overlap volume
-	return (AFillingFactor * ABox.GetVolume() + BFillingFactor * BBox.GetVolume() - 0.5f * (AFillingFactor+BFillingFactor) * OverlapBox.GetVolume()) / UnionBox.GetVolume();
-}
+	if (ASphere.IsInside(BSphere, 1.f))
+	{
+		return ASphere.GetVolume();
+	}
 
-float GetDistance(float AMin, float AMax, float BMin, float BMax)
-{
-	if(AMin > BMax)
+	if(BSphere.IsInside(ASphere, 1.f))
 	{
-		return (AMin-BMax);
+		return BSphere.GetVolume();
 	}
 
-	if(BMin > AMax)
+	if (ASphere.Equals(BSphere, 1.f))
 	{
-		return (BMin-AMax);
+		return ASphere.GetVolume();
 	}
 
-	return 0.f;
+	float Distance = (ASphere.Center-BSphere.Center).Size();
+	check (Distance > 0.f);
+
+	float  ARadius = ASphere.W;
+	float  BRadius = BSphere.W;
+
+	float  ACapHeight = (BRadius*BRadius - (ARadius - Distance)*(ARadius - Distance)) / (2*Distance);
+	float  BCapHeight = (ARadius*ARadius - (BRadius - Distance)*(BRadius - Distance)) / (2*Distance);
+
+	check (ACapHeight>0.f);
+	check (BCapHeight>0.f);
+
+	float  OverlapRadius1 = ((ARadius+BRadius)*(ARadius+BRadius) - Distance*Distance) * (Distance*Distance - (ARadius-BRadius)*(ARadius-BRadius));
+	float  OverlapRadius2 = 2*Distance;
+	float  OverlapRedius = FMath::Sqrt(OverlapRadius1) / OverlapRadius2;
+	float  OverlapRediusSq = OverlapRedius*OverlapRedius;
+
+	check (OverlapRadius1 >= 0.f);
+
+	float ConstPI = PI/6.0f;
+	float AVolume = ConstPI*(3*OverlapRediusSq + ACapHeight*ACapHeight) * ACapHeight;
+	float BVolume = ConstPI*(3*OverlapRediusSq + BCapHeight*BCapHeight) * BCapHeight;
+
+	check (AVolume > 0.f &&  BVolume > 0.f);
+
+	float TotalVolume = AFillingFactor*AVolume + BFillingFactor*BVolume;
+	return TotalVolume;
 }
-// utility functions
-float CalculateDistance(const FLODCluster& A, const FLODCluster& B)
-{
-	FVector AMax = A.Bound.Max;
-	FVector AMin = A.Bound.Min;
-	FVector BMax = B.Bound.Max;
-	FVector BMin = B.Bound.Min;
-	FVector DistanceVector = FVector::ZeroVector;
 
-	DistanceVector.X = GetDistance(AMin.X, AMax.X, BMin.X, BMax.X);
-	DistanceVector.Y = GetDistance(AMin.Y, AMax.Y, BMin.Y, BMax.Y);
-	DistanceVector.Z = GetDistance(AMin.Z, AMax.Z, BMin.Z, BMax.Z);
+// utility functions that calculates filling factor
+float CalculateFillingFactor(const FSphere& ASphere, const float AFillingFactor, const FSphere& BSphere, const float BFillingFactor)
+{
+	float OverlapVolume = CalculateOverlap( ASphere, AFillingFactor, BSphere, BFillingFactor);
+	FSphere UnionSphere = ASphere + BSphere;
+	// it shouldn't be zero or it should be checked outside
+	ensure(UnionSphere.W != 0.f);
 
-	return DistanceVector.SizeSquared();
+	// http://deim.urv.cat/~rivi/pub/3d/icra04b.pdf
+	// cost is calculated based on r^3 / filling factor
+	// since it subtract by AFillingFactor * 1/2 overlap volume + BfillingFactor * 1/2 overlap volume
+	return FMath::Max(0.f, (AFillingFactor * ASphere.GetVolume() + BFillingFactor * BSphere.GetVolume() - OverlapVolume) / UnionSphere.GetVolume());
 }
 
 ////////////////////////////////////////////////////////////////////////////////////////////
@@ -104,10 +127,13 @@ void FHierarchicalLODBuilder::Build()
 	}
 }
 
-
 // for now it only builds per level, it turns to one actor at the end
 void FHierarchicalLODBuilder::BuildClusters(class ULevel* InLevel)
 {
+	// I'm using stack mem within this scope of the function
+	// so we need this
+	FMemMark Mark(FMemStack::Get());
+
 	// you still have to delete all objects just in case they had it and didn't want it anymore
 	TArray<UObject*> AssetsToDelete;
 	for (int32 ActorId=InLevel->Actors.Num()-1; ActorId >= 0; --ActorId)
@@ -139,22 +165,46 @@ void FHierarchicalLODBuilder::BuildClusters(class ULevel* InLevel)
 
 		for(int32 LODId=0; LODId<TotalNumLOD; ++LODId)
 		{
-			// initialize Clusters
-			InitializeClusters(InLevel, LODId);
+			AWorldSettings* WorldSetting = InLevel->GetWorld()->GetWorldSettings();
+			check(WorldSetting->HierarchicalLODSetup.IsValidIndex(LODId));
+			float DesiredBoundBoxHalfSize = WorldSetting->HierarchicalLODSetup[LODId].DesiredBoundSize * 0.5f;
+			float DesiredFillingRatio = WorldSetting->HierarchicalLODSetup[LODId].DesiredFillingPercentage * 0.01f;
+			ensure(DesiredFillingRatio!=0.f);
+			float HighestCost = FMath::Pow(DesiredBoundBoxHalfSize, 3)/(DesiredFillingRatio);
+			// test parameter I was playing with to cull adding to the array
+			// intialization can have too many elements, decided to cull
+			// the problem can be that we can create disconnected tree
+			// my assumption is that if the merge cost is too high, then it's not worth merge anyway
+			static int32 CullMultiplier=1;
+
+			// since to show progress of initialization, I'm scoping it
+			{
+				FFormatNamedArguments Arguments;
+				Arguments.Add(TEXT("LODIndex"), FText::AsNumber(LODId));
 
-			// now we have all pair of nodes
-			TArray<FLODCluster> MSTClusters;
-			FindMST(MSTClusters);
-			// replace Clusters to new culled cluster
-			Clusters = MSTClusters;
+				FScopedSlowTask SlowTask(100, FText::Format(LOCTEXT("HierarchicalLOD_InitializeCluster", "Initializing Clusters for LOD {LODIndex}..."), Arguments));
+				SlowTask.MakeDialog();
 
-			// now we have to calculate merge
-			MergeClusters(InLevel, LODId);
+				// initialize Clusters
+				InitializeClusters(InLevel, LODId, HighestCost*CullMultiplier);
+
+				// move a half way - I know we can do this better but as of now this is small progress
+				SlowTask.EnterProgressFrame(50);
+
+				// now we have all pair of nodes
+				FindMST();
+			}
+
+			// now we have to calculate merge clusters and build actors
+			MergeClustersAndBuildActors(InLevel, LODId, HighestCost);
 		}
 	}
+
+	// Clear Clusters. It is using stack mem, so it won't be good after this
+	Clusters.Empty();
 }
 
-void FHierarchicalLODBuilder::FindMST(TArray<FLODCluster>& OutClusters) 
+void FHierarchicalLODBuilder::FindMST() 
 {
 	if (Clusters.Num() > 0)
 	{
@@ -168,118 +218,102 @@ void FHierarchicalLODBuilder::FindMST(TArray<FLODCluster>& OutClusters)
 		};
 
 		Clusters.Sort(FCompareCluster());
-
-		// now start from small to large, find the MST
-		OutClusters.Empty();
-
-		TArray<class AActor*> Processed;
-		for(auto& Cluster: Clusters)
-		{
-			if (Cluster.IsValid())
-			{
-				bool bNeedToAdd=false;
-
-				TArray<class AActor*>& Actors=Cluster.Actors;
-				for(auto& Actor : Actors)
-				{
-					if(!Processed.Contains(Actor))
-					{
-						bNeedToAdd = true;
-						Processed.Add(Actor);
-					}
-				}
-				if(bNeedToAdd)
-				{
-					OutClusters.Add(Cluster);
-				}
-			}
-		}
 	}
 }
 
-void FHierarchicalLODBuilder::MergeClusters(class ULevel* InLevel, const int32 LODIdx)
+void FHierarchicalLODBuilder::MergeClustersAndBuildActors(class ULevel* InLevel, const int32 LODIdx, float HighestCost)
 {
 	if (Clusters.Num() > 0)
 	{
-		int32 LastIndex = Clusters.Num()-1; //FMath::TruncToInt(Clusters.Num() * 0.70); /*(10% ratio)*/
-		AWorldSettings* WorldSetting = InLevel->GetWorld()->GetWorldSettings();
-		check (WorldSetting->HierarchicalLODSetup.IsValidIndex(LODIdx));
-		float DesiredBoundBoxHalfSize = WorldSetting->HierarchicalLODSetup[LODIdx].DesiredBoundBoxSize * 0.5f;
-		float DesiredFillingRatio = WorldSetting->HierarchicalLODSetup[LODIdx].DesiredFillingPercentage * 0.01f;
-		ensure (DesiredFillingRatio!=0.f);
-		float HighestCost = FMath::Pow(DesiredBoundBoxHalfSize, 3)/(DesiredFillingRatio);
-
-		const int32 TotalCluster = Clusters.Num();
-		// now we have minimum Clusters
-		for(int32 ClusterId=0; ClusterId < TotalCluster; ++ClusterId)
+		FFormatNamedArguments Arguments;
+		Arguments.Add(TEXT("LODIndex"), FText::AsNumber(LODIdx));
+
+		// merge clusters first
 		{
-			auto& Cluster = Clusters[ClusterId];
-			UE_LOG(LogLODGenerator, Log, TEXT("%d. %0.2f {%s}"), ClusterId+1, Cluster.GetCost(), *Cluster.ToString());
-
-			for (int32 MergedClusterId=0; MergedClusterId < ClusterId; ++MergedClusterId)
-			{ 
-				// compare with previous clusters
-				auto& MergedCluster = Clusters[MergedClusterId];
-				// see if it's valid, if it contains, check the cost
-				if (MergedCluster.IsValid() && Cluster.IsValid())
+			static int32 TotalIteration=3;
+			const int32 TotalCluster = Clusters.Num();
+
+			FScopedSlowTask SlowTask(TotalIteration*TotalCluster, FText::Format( LOCTEXT("HierarchicalLOD_BuildClusters", "Building Clusters for LOD {LODIndex}..."), Arguments) );
+			SlowTask.MakeDialog();
+
+			for(int32 Iteration=0; Iteration<TotalIteration; ++Iteration)
+			{
+				// now we have minimum Clusters
+				for(int32 ClusterId=0; ClusterId < TotalCluster; ++ClusterId)
 				{
-					// if valid, see if it contains any of this actors
-					// merge whole clusters
-					FLODCluster NewCluster = Cluster + MergedCluster;
-					float MergeCost = NewCluster.GetCost();
+					auto& Cluster = Clusters[ClusterId];
+					UE_LOG(LogLODGenerator, Verbose, TEXT("%d. %0.2f {%s}"), ClusterId+1, Cluster.GetCost(), *Cluster.ToString());
 
-					// merge two clusters
-					if (MergeCost <= HighestCost)
-					{
-						UE_LOG(LogLODGenerator, Log, TEXT("Merging of Cluster (%d) and (%d) with merge cost (%0.2f) "), ClusterId+1, MergedClusterId+1, MergeCost);
-						
-						MergedCluster = NewCluster;
-						// now this cluster is invalid
-						Cluster.Invalidate();
-						break;
-					}
-					else if (Cluster.Contains(MergedCluster))
+					// progress bar update
+					SlowTask.EnterProgressFrame();
+
+					if(Cluster.IsValid())
 					{
-						UE_LOG(LogLODGenerator, Log, TEXT("Removing actor (%d) from Cluster (%d)"), MergedClusterId+1, ClusterId+1);
-						// if higher than high cost, then we break them apart
-						Cluster -= MergedCluster;
-						break;
+						for(int32 MergedClusterId=0; MergedClusterId < ClusterId; ++MergedClusterId)
+						{
+							// compare with previous clusters
+							auto& MergedCluster = Clusters[MergedClusterId];
+							// see if it's valid, if it contains, check the cost
+							if(MergedCluster.IsValid())
+							{
+								if(MergedCluster.Contains(Cluster))
+								{
+									// if valid, see if it contains any of this actors
+									// merge whole clusters
+									FLODCluster NewCluster = Cluster + MergedCluster;
+									float MergeCost = NewCluster.GetCost();
+
+									// merge two clusters
+									if(MergeCost <= HighestCost)
+									{
+										UE_LOG(LogLODGenerator, Log, TEXT("Merging of Cluster (%d) and (%d) with merge cost (%0.2f) "), ClusterId+1, MergedClusterId+1, MergeCost);
+
+										MergedCluster = NewCluster;
+										// now this cluster is invalid
+										Cluster.Invalidate();
+										break;
+									}
+									else
+									{
+										Cluster -= MergedCluster;
+									}
+								}
+							}
+						}
+
+						UE_LOG(LogLODGenerator, Verbose, TEXT("Processed(%s): %0.2f {%s}"), Cluster.IsValid()? TEXT("Valid"):TEXT("Invalid"), Cluster.GetCost(), *Cluster.ToString());
 					}
 				}
 			}
-
-			UE_LOG(LogLODGenerator, Log, TEXT("Processed(%s): %0.2f {%s}"), Cluster.IsValid()? TEXT("Valid"):TEXT("Invalid"), Cluster.GetCost(), *Cluster.ToString());
 		}
 
 		static bool bBuildActor=true;
 
-		GWarn->BeginSlowTask(LOCTEXT("HierarchicalLOD_CreatingSimpleMesh", "Creating Simple Mesh"), true);
-		GEditor->BeginTransaction(LOCTEXT("HierarchicalLOD_CreatingSimpleMesh", "Creating Simple Mesh"));
-		// print data
-		int32 Index=0, TotalValidCluster=0;
-		for (auto& Cluster: Clusters)
+		if (bBuildActor)
 		{
-			if (Cluster.IsValid())
+			// print data
+			int32 TotalValidCluster=0;
+			for(auto& Cluster: Clusters)
 			{
-				++TotalValidCluster;
+				if(Cluster.IsValid())
+				{
+					++TotalValidCluster;
+				}
 			}
-		}
 
-		for(auto& Cluster: Clusters)
-		{
-			if(Cluster.IsValid())
+			FScopedSlowTask SlowTask(TotalValidCluster, FText::Format( LOCTEXT("HierarchicalLOD_MergeActors", "Merging Actors for LOD {LODIndex}..."), Arguments) );
+			SlowTask.MakeDialog();
+
+			for(auto& Cluster: Clusters)
 			{
-				GWarn->StatusForceUpdate(Index, TotalValidCluster, LOCTEXT("HierarchicalLOD_BuildActor", "Building Actor"));
-				UE_LOG(LogLODGenerator, Log, TEXT("%d. %0.2f (%s)"), ++Index, Cluster.GetCost(), *Cluster.ToString());
-				if (bBuildActor)
+				if(Cluster.ShouldBuild())
 				{
+					SlowTask.EnterProgressFrame();
+					//Sleep(0.001f);
 					Cluster.BuildActor(InLevel->GetWorld(), InLevel, LODIdx);
 				}
 			}
 		}
-
-		GEditor->EndTransaction();
-		GWarn->EndSlowTask();
 	}
 }
 
@@ -317,13 +351,13 @@ bool ShouldGenerateCluster(class AActor* Actor)
 	return (Components.Num() > 0);
 }
 
-void FHierarchicalLODBuilder::InitializeClusters(class ULevel* InLevel, const int32 LODIdx)
+void FHierarchicalLODBuilder::InitializeClusters(class ULevel* InLevel, const int32 LODIdx, float CullCost)
 {
 	if (InLevel->Actors.Num() > 0)
 	{
 		if (LODIdx == 0)
 		{
-			Clusters.Empty(InLevel->Actors.Num());
+			Clusters.Empty();
 
 			// first we generate graph with 2 pair nodes
 			// this is very expensive when we have so many actors
@@ -341,9 +375,11 @@ void FHierarchicalLODBuilder::InitializeClusters(class ULevel* InLevel, const in
 						{
 							FLODCluster NewClusterCandidate = FLODCluster(Actor1, Actor2);
 							// @todo optimization : filter first - this might create disconnected tree - 
-							// minimal it should be filled up 20% even to consider that
-							// if (NewClusterCandidate.GetFillingFactor() > 0.2f)
-							Clusters.Add(NewClusterCandidate);
+							//@Todo debug but it's too slow to have full tree 
+							if (NewClusterCandidate.GetCost() <= CullCost)
+							{
+								Clusters.Add(NewClusterCandidate);
+							}
 						}
 					}
 				}
@@ -377,9 +413,6 @@ void FHierarchicalLODBuilder::InitializeClusters(class ULevel* InLevel, const in
 						{
 							// create new cluster
 							FLODCluster NewClusterCandidate = FLODCluster(Actor1, Actor2);
-							// @todo optimization : filter first - this might create disconnected tree - 
-							// minimal it should be filled up 20% even to consider that
-							// if (NewClusterCandidate.GetFillingFactor() > 0.2f)
 							Clusters.Add(NewClusterCandidate);
 						}
 					}
@@ -402,24 +435,34 @@ FLODCluster::FLODCluster(const FLODCluster& Other)
 {
 }
 
+FLODCluster::FLODCluster(class AActor* Actor1)
+: Bound(ForceInit)
+, bValid(true)
+{
+	AddActor(Actor1);
+	// calculate new filling factor
+	FillingFactor = 1.f;
+}
+
 FLODCluster::FLODCluster(class AActor* Actor1, class AActor* Actor2)
 : Bound(ForceInit)
 , bValid(true)
 {
-	FBox Actor1Bound = AddActor(Actor1);
-	FBox Actor2Bound = AddActor(Actor2);
+	FSphere Actor1Bound = AddActor(Actor1);
+	FSphere Actor2Bound = AddActor(Actor2);
 
 	// calculate new filling factor
 	FillingFactor = CalculateFillingFactor(Actor1Bound, 1.f, Actor2Bound, 1.f);
 }
 
-FBox FLODCluster::AddActor(class AActor* NewActor)
+FSphere FLODCluster::AddActor(class AActor* NewActor)
 {
 	ensure (Actors.Contains(NewActor) == false);
 	Actors.Add(NewActor);
 	FVector Origin, Extent;
 	NewActor->GetActorBounds(false, Origin, Extent);
-	FBox NewBound = FBox(Origin-Extent, Origin+Extent);
+	
+	FSphere NewBound = FSphere(Origin, Extent.Size());
 	Bound += NewBound;
 
 	return NewBound;
@@ -494,45 +537,33 @@ void FLODCluster::SubtractCluster(const FLODCluster& Other)
 	}
 	else if (NewActors.Num() == 1)
 	{
-		Bound = FBox(ForceInitToZero);
+		Bound = FSphere(ForceInitToZero);
 		AddActor(NewActors[0]);
 		FillingFactor = 1.f;
 	}
-	else if (NewActors.Num() == 2)
+	else if (NewActors.Num() >= 2)
 	{
-		Bound = FBox(ForceInitToZero);
+		Bound = FSphere(ForceInit);
 
-		FBox Actor1Bound = AddActor(NewActors[0]);
-		FBox Actor2Bound = AddActor(NewActors[1]);
+		FSphere Actor1Bound = AddActor(NewActors[0]);
+		FSphere Actor2Bound = AddActor(NewActors[1]);
 
 		// calculate new filling factor
 		FillingFactor = CalculateFillingFactor(Actor1Bound, 1.f, Actor2Bound, 1.f);
-	}
-	else
-	{
-		// this shouldn't happen
-		// if this has to happen, we'll either have to come up with way to subtrac bound or recalculate from scratch
-		ensure (false);
-	}
-}
 
-// invalidate this cluster - happens when this merged to another cluster
-void FLODCluster::Invalidate()
-{
-	bValid = false;
-}
+		// if more actors, we add them manually
+		for (int32 ActorId=2; ActorId<NewActors.Num(); ++ActorId)
+		{
+			// if not contained, it shouldn't be
+			check (!Actors.Contains(NewActors[ActorId]));
 
-bool FLODCluster::IsValid() const
-{
-	return bValid;
+			FSphere NewBound = AddActor(NewActors[ActorId]);
+			FillingFactor = CalculateFillingFactor(NewBound, 1.f, Bound, FillingFactor);
+			Bound += NewBound;
+		}
+	}
 }
 
-const float FLODCluster::GetCost() const
-{
-	FVector BoundExtent = Bound.GetExtent();
-	const float Extent = BoundExtent.GetMax();
-	return (FMath::Pow(Extent, 3) / FillingFactor);
-}
 // if criteria matches, build new LODActor and replace current Actors with that. We don't need 
 // this clears previous actors and sets to this new actor
 // this is required when new LOD is created from these actors, this will be replaced
@@ -573,7 +604,7 @@ void FLODCluster::BuildActor(class UWorld* InWorld, class ULevel* InLevel, const
 				// Generate proxy mesh and proxy material assets
 				IMeshUtilities& MeshUtilities = FModuleManager::Get().LoadModuleChecked<IMeshUtilities>("MeshUtilities");
 				// should give unique name, so use level + actor name
-				const FString PackageName = FString::Printf(TEXT("LOD_%s_%s"), *InLevel->GetName(), *FirstActor->GetName());
+				const FString PackageName = FString::Printf(TEXT("LOD_%s"), *FirstActor->GetName());
 				if (MeshUtilities.GetMeshMergingInterface() && LODSetup.bSimplifyMesh)
 				{
 					MeshUtilities.CreateProxyMesh(Actors, ProxySetting, AssetsOuter, PackageName, OutAssets, OutProxyLocation);
@@ -581,7 +612,8 @@ void FLODCluster::BuildActor(class UWorld* InWorld, class ULevel* InLevel, const
 				else
 				{
 					FMeshMergingSettings MergeSetting;
-					MeshUtilities.MergeActors(Actors, MergeSetting, AssetsPath, OutAssets, OutProxyLocation, true );
+					MergeSetting.bMergeMaterials = true;
+					MeshUtilities.MergeActors(Actors, MergeSetting, AssetsPath + PackageName, OutAssets, OutProxyLocation, true );
 				}
 
 				UStaticMesh* MainMesh=NULL;
diff --git a/Engine/Source/Editor/UnrealEd/Public/HierarchicalLOD.h b/Engine/Source/Editor/UnrealEd/Public/HierarchicalLOD.h
--- a/Engine/Source/Editor/UnrealEd/Public/HierarchicalLOD.h
+++ b/Engine/Source/Editor/UnrealEd/Public/HierarchicalLOD.h
@@ -6,13 +6,22 @@
 	HierarchicalLOD.h: Hierarchical LOD definition.
 =============================================================================*/
 
-
-// implementation 
-//  http://deim.urv.cat/~rivi/pub/3d/icra04b.pdf
+/**
+ *
+ *	This is a LOD cluster struct that holds list of actors with relevant information
+ *
+ *	http://deim.urv.cat/~rivi/pub/3d/icra04b.pdf
+ *
+ *	This is used by Hierarchical LOD Builder to generates list of clusters 
+ *	that are together in vincinity and build as one actor
+ *
+ **/
 
 struct FLODCluster
 {
+	// constructors
 	FLODCluster(const FLODCluster& Other);
+	FLODCluster(class AActor* Actor1);
 	FLODCluster(class AActor* Actor1, class AActor* Actor2);
 
 	FLODCluster operator+( const FLODCluster& Other ) const;
@@ -20,33 +29,58 @@ struct FLODCluster
 	FLODCluster operator-(const FLODCluster& Other) const;
 	FLODCluster operator-=(const FLODCluster& Other);
 	FLODCluster& operator=(const FLODCluster & Other);
-	// invalidate this cluster - happens when this merged to another cluster
-	void Invalidate();
-	bool IsValid() const;
 
-	const float GetCost() const;
+	// Invalidate current cluster
+	void Invalidate() { bValid = false; }
+	// return true if valid
+	bool IsValid() const {	return bValid; }
+	// should build to actor list
+	bool ShouldBuild() const { return IsValid() && Actors.Num() > 1; }
+
+	// return cost of the cluster, lower is better
+	const float GetCost() const
+	{
+		return (FMath::Pow(Bound.W, 3) / FillingFactor);
+	}
+
+	// return true if this cluster contains ANY of actors of Other
 	bool Contains(FLODCluster& Other) const;
+	// return string of data
 	FString ToString() const;
 	
+	// member variable
+	// list of actors
 	TArray<class AActor*>	Actors;
-	FBox					Bound;
+	// bound of this cluster
+	FSphere					Bound;
+	// filling factor - higher means filled more
 	float					FillingFactor;
 
-	// we don't delete the invalid entry, but mark as invalid when it's not any more valid
-	// it's initialized to be valid
-	bool bValid;
-
 	// if criteria matches, build new LODActor and replace current Actors with that. We don't need 
 	// this clears previous actors and sets to this new actor
 	// this is required when new LOD is created from these actors, this will be replaced
 	// to save memory and to reduce memory increase during this process, we discard previous actors and replace with this actor
 	void BuildActor(class UWorld* InWorld, class ULevel* InLevel, const int32 LODIdx);
-	// Merge two clusters
+
+private:
+	// cluster operations
 	void MergeClusters(const FLODCluster& Other);
 	void SubtractCluster(const FLODCluster& Other);
-	FBox AddActor(class AActor* NewActor);
+
+	// add new actor, this doesn't recalculate filling factor
+	FSphere AddActor(class AActor* NewActor);
+
+	// whether it's valid or not
+	bool bValid;
 };
 
+/**
+ *
+ *	This is Hierarchical LOD builder
+ *
+ * This builds list of clusters and make sure it's sorted in the order of lower cost to high and merge clusters
+ **/
+
 struct UNREALED_API FHierarchicalLODBuilder
 {
 	FHierarchicalLODBuilder(class UWorld* InWorld);
@@ -60,20 +94,20 @@ struct UNREALED_API FHierarchicalLODBuilder
 	void Build();
 
 private:
-	// data structure
-	TArray<FLODCluster> Clusters;
+	// data structure - this is only valid within scope since mem stack allocator
+	TArray<FLODCluster, TMemStackAllocator<>> Clusters;
+	// owner world
 	class UWorld*		World;
 
 	// for now it only builds per level, it turns to one actor at the end
 	void BuildClusters(class ULevel* InLevel);
 
 	// initialize Clusters variable - each Actor becomes Cluster
-	void InitializeClusters(class ULevel* InLevel, const int32 LODIdx);
-
-	// Refresh Cluster Distance for ClusterId
-	//void RefreshClusterDistance(int32 ClusterId, const int32 TotalNumClusters, float** InDistanceCacheMatrix);
+	void InitializeClusters(class ULevel* InLevel, const int32 LODIdx, float CullCost);
 
-	void MergeClusters(class ULevel* InLevel, const int32 LODIdx);
+	// merge clusters
+	void MergeClustersAndBuildActors(class ULevel* InLevel, const int32 LODIdx, float HighestCost);
 
-	void FindMST(TArray<FLODCluster>& OutClusters);
+	// find minmal spanning tree of clusters
+	void FindMST();
 };
diff --git a/Engine/Source/Runtime/Core/Private/Math/Sphere.cpp b/Engine/Source/Runtime/Core/Private/Math/Sphere.cpp
--- a/Engine/Source/Runtime/Core/Private/Math/Sphere.cpp
+++ b/Engine/Source/Runtime/Core/Private/Math/Sphere.cpp
@@ -62,4 +62,51 @@ FSphere FSphere::TransformBy(const FTransform& M) const
 	Result.W = M.GetMaximumAxisScale();
 
 	return Result;
+}
+
+float FSphere::GetVolume() const
+{
+	return 1.25f * PI * FMath::Pow(W, 3);
+}
+
+FSphere& FSphere::operator+=(const FSphere &Other)
+{
+	if (W == 0.f)
+	{
+		*this = Other;
+	}
+	else if (IsInside(Other))
+	{
+		*this = Other;
+	}
+	else if (Other.IsInside(*this))
+	{
+		// no change		
+	}
+	else
+	{
+		FSphere NewSphere;
+
+		FVector DirToOther = Other.Center - Center;
+		FVector UnitDirToOther = DirToOther;
+		UnitDirToOther.Normalize();
+
+		float NewRadius = (DirToOther.Size() + Other.W + W) * 0.5f;
+
+		// find end point
+		FVector End1 = Other.Center + UnitDirToOther*Other.W;
+		FVector End2 = Center - UnitDirToOther*W;
+		FVector NewCenter = (End1 + End2)*0.5f;
+
+		NewSphere.Center = NewCenter; 
+		NewSphere.W = NewRadius;
+
+		// make sure both are inside afterwards
+		check (Other.IsInside(NewSphere, 1.f));
+		check (IsInside(NewSphere, 1.f));
+
+		*this = NewSphere;
+	}
+
+	return *this;
 }
\ No newline at end of file
diff --git a/Engine/Source/Runtime/Core/Public/Math/Sphere.h b/Engine/Source/Runtime/Core/Public/Math/Sphere.h
--- a/Engine/Source/Runtime/Core/Public/Math/Sphere.h
+++ b/Engine/Source/Runtime/Core/Public/Math/Sphere.h
@@ -83,12 +83,12 @@ public:
 	 */
 	bool IsInside( const FSphere& Other, float Tolerance = KINDA_SMALL_NUMBER ) const
 	{
-		if (W > Other.W - Tolerance)
+		if (W > Other.W + Tolerance)
 		{
 			return false;
 		}
 
-		return (Center - Other.Center).SizeSquared() <= FMath::Square(Other.W - Tolerance - W);
+		return (Center - Other.Center).SizeSquared() <= FMath::Square(Other.W + Tolerance - W);
 	}
 
 
@@ -132,6 +132,31 @@ public:
 	 */
 	CORE_API FSphere TransformBy( const FTransform& M ) const;
 
+	/**
+	 * Get volume of the current sphere
+	 * 
+	 */
+	CORE_API float GetVolume() const;
+
+	/**
+	 * Adds to this bounding box to include a new bounding volume.
+	 *
+	 * @param Other the bounding volume to increase the bounding volume to.
+	 * @return Reference to this bounding volume after resizing to include the other bounding volume.
+	 */
+	CORE_API FSphere& operator+=( const FSphere& Other );
+
+	/**
+	 * Gets the result of addition to this bounding volume.
+	 *
+	 * @param Other The other volume to add to this.
+	 * @return A new bounding volume.
+	 */
+	FSphere operator+( const FSphere& Other ) const
+	{
+		return FSphere(*this) += Other;
+	}
+
 public:
 
 	/**
diff --git a/Engine/Source/Runtime/Engine/Classes/Engine/LODActor.h b/Engine/Source/Runtime/Engine/Classes/Engine/LODActor.h
--- a/Engine/Source/Runtime/Engine/Classes/Engine/LODActor.h
+++ b/Engine/Source/Runtime/Engine/Classes/Engine/LODActor.h
@@ -45,6 +45,7 @@ public:
 protected:
 	// Begin UObject interface.
 	virtual FString GetDetailedInfoInternal() const override;
+	virtual FBox GetComponentsBoundingBox(bool bNonColliding = false) const override;	
 #if WITH_EDITOR
 	virtual void PostEditChangeChainProperty(FPropertyChangedChainEvent& PropertyChangedEvent) override;
 #endif // WITH_EDITOR	
diff --git a/Engine/Source/Runtime/Engine/Classes/GameFramework/WorldSettings.h b/Engine/Source/Runtime/Engine/Classes/GameFramework/WorldSettings.h
--- a/Engine/Source/Runtime/Engine/Classes/GameFramework/WorldSettings.h
+++ b/Engine/Source/Runtime/Engine/Classes/GameFramework/WorldSettings.h
@@ -207,7 +207,7 @@ struct ENGINE_API FHierarchicalSimplification
 
 	/** Desired Bounding Box Size for clustering */
 	UPROPERTY(EditAnywhere, Category=FHierarchicalSimplification, AdvancedDisplay, meta=(UIMin=10.f, ClampMin=10.f))
-	float DesiredBoundBoxSize;
+	float DesiredBoundSize;
 
 	/** Desired Filling Percentage for clustering */
 	UPROPERTY(EditAnywhere, Category=FHierarchicalSimplification, AdvancedDisplay, meta=(ClampMin = "0", ClampMax = "100", UIMin = "0", UIMax = "100"))
@@ -220,10 +220,10 @@ struct ENGINE_API FHierarchicalSimplification
 	FHierarchicalSimplification()
 		: bSimplifyMesh(true)
 		, DrawDistance(5000)
-		, DesiredBoundBoxSize(10000)
-		, DesiredFillingPercentage(30)
+		, DesiredBoundSize(5000)
+		, DesiredFillingPercentage(50)
 	{
-		Setting.ScreenSize = 300;
+		Setting.ScreenSize = 200;
 	}
 };
 
diff --git a/Engine/Source/Runtime/Engine/Private/LODActor.cpp b/Engine/Source/Runtime/Engine/Private/LODActor.cpp
--- a/Engine/Source/Runtime/Engine/Private/LODActor.cpp
+++ b/Engine/Source/Runtime/Engine/Private/LODActor.cpp
@@ -88,5 +88,22 @@ void ALODActor::CheckForErrors()
 
 #endif // WITH_EDITOR
 
+FBox ALODActor::GetComponentsBoundingBox(bool bNonColliding) const 
+{
+	FBox BoundBox = Super::GetComponentsBoundingBox(bNonColliding);
+
+	if (bNonColliding)
+	{
+		if (StaticMeshComponent && StaticMeshComponent->StaticMesh)
+		{
+			FBoxSphereBounds StaticBound = StaticMeshComponent->StaticMesh->GetBounds();
+			FBox StaticBoundBox(BoundBox.GetCenter()-StaticBound.BoxExtent, BoundBox.GetCenter()+StaticBound.BoxExtent);
+			BoundBox += StaticBoundBox;
+		}
+	}
+
+	return BoundBox;	
+}
+
 #undef LOCTEXT_NAMESPACE
 
