#HLOD
-Fixed LOD count for message is correct.
-Merge Actor now creates asset correctly inside of map with private set up.

diff --git a/Engine/Source/Editor/UnrealEd/Private/HierarchicalLOD.cpp b/Engine/Source/Editor/UnrealEd/Private/HierarchicalLOD.cpp
--- a/Engine/Source/Editor/UnrealEd/Private/HierarchicalLOD.cpp
+++ b/Engine/Source/Editor/UnrealEd/Private/HierarchicalLOD.cpp
@@ -180,7 +180,7 @@ void FHierarchicalLODBuilder::BuildClusters(class ULevel* InLevel)
 			// since to show progress of initialization, I'm scoping it
 			{
 				FFormatNamedArguments Arguments;
-				Arguments.Add(TEXT("LODIndex"), FText::AsNumber(LODId));
+				Arguments.Add(TEXT("LODIndex"), FText::AsNumber(LODId+1));
 
 				FScopedSlowTask SlowTask(100, FText::Format(LOCTEXT("HierarchicalLOD_InitializeCluster", "Initializing Clusters for LOD {LODIndex}..."), Arguments));
 				SlowTask.MakeDialog();
@@ -226,7 +226,7 @@ void FHierarchicalLODBuilder::MergeClustersAndBuildActors(class ULevel* InLevel,
 	if (Clusters.Num() > 0)
 	{
 		FFormatNamedArguments Arguments;
-		Arguments.Add(TEXT("LODIndex"), FText::AsNumber(LODIdx));
+		Arguments.Add(TEXT("LODIndex"), FText::AsNumber(LODIdx+1));
 
 		// merge clusters first
 		{
@@ -596,7 +596,8 @@ void FLODCluster::BuildActor(class UWorld* InWorld, class ULevel* InLevel, const
 			if (ensure(AllComponents.Num() > 0))
 			{
 				// In case we don't have outer generated assets should have same path as LOD level
-				const FString AssetsPath = FPackageName::GetLongPackagePath(AssetsOuter->GetName()) + TEXT("/");
+				
+				const FString AssetsPath = AssetsOuter->GetName() + TEXT("/");
 				class AActor* FirstActor = Actors[0];
 
 				TArray<UObject*> OutAssets;
@@ -616,6 +617,12 @@ void FLODCluster::BuildActor(class UWorld* InWorld, class ULevel* InLevel, const
 					MeshUtilities.MergeActors(Actors, MergeSetting, AssetsPath + PackageName, OutAssets, OutProxyLocation, true );
 				}
 
+				// remove flags of RF_Public since we don't want these assets to be public
+				for (auto& AssetIter : OutAssets)
+				{
+					AssetIter->ClearFlags(RF_Public);
+				}
+
 				UStaticMesh* MainMesh=NULL;
 				// set staticmesh
 				for(auto& Asset: OutAssets)
