Fixing mac command line argument parsing for unreal audio tests

diff --git a/Engine/Source/Editor/UnrealEd/Private/Commandlets/AudioTestCommandlet.cpp b/Engine/Source/Editor/UnrealEd/Private/Commandlets/AudioTestCommandlet.cpp
--- a/Engine/Source/Editor/UnrealEd/Private/Commandlets/AudioTestCommandlet.cpp
+++ b/Engine/Source/Editor/UnrealEd/Private/Commandlets/AudioTestCommandlet.cpp
@@ -139,10 +139,26 @@ UAudioTestCommandlet::UAudioTestCommandlet(const FObjectInitializer& ObjectIniti
 int32 UAudioTestCommandlet::Main(const FString& InParams)
 {
 #if ENABLE_UNREAL_AUDIO
+    
+    // Mac stupidly is adds "-NSDocumentRevisionsDebugMode YES" to command line args so lets remove that
+    FString Params;
+#if PLATFORM_MAC
+    FString BadString(TEXT("-NSDocumentRevisionsDebugMode YES"));
+    int32 DebugModeIndex = InParams.Find(BadString);
+    if (DebugModeIndex != INDEX_NONE)
+    {
+        Params = InParams.LeftChop(BadString.Len());
+    }
+    else
+#endif // #if PLATFORM_MAC
+    {
+        Params = InParams;
+    }
+   
 	// Parse command line.
 	TArray<FString> Tokens;
 	TArray<FString> Switches;
-	UCommandlet::ParseCommandLine(*InParams, Tokens, Switches);
+	UCommandlet::ParseCommandLine(*Params, Tokens, Switches);
 
 	if (Tokens.Num() < 3 || Tokens.Num() > 4)
 	{
@@ -202,9 +218,9 @@ int32 UAudioTestCommandlet::Main(const FString& InParams)
 	}
 
 	UnrealAudioUnload();
-#else
+#else // #if ENABLE_UNREAL_AUDIO
 	UE_LOG(AudioTestCommandlet, Display, TEXT("Unreal Audio Module Not Enabled For This Platform"));
-#endif // #if ENABLE_UNREAL_AUDIO
+#endif // #else // #if ENABLE_UNREAL_AUDIO
 
 
 	return 0;
