Fix crashes with missing OwnerOfTemporaries when running some Blueprint commandlets.

UE-14280 - Crash running GenerateBlueprintAPI Commandlet

diff --git a/Engine/Source/Editor/UnrealEd/Private/Commandlets/DumpBlueprintsInfoCommandlet.cpp b/Engine/Source/Editor/UnrealEd/Private/Commandlets/DumpBlueprintsInfoCommandlet.cpp
--- a/Engine/Source/Editor/UnrealEd/Private/Commandlets/DumpBlueprintsInfoCommandlet.cpp
+++ b/Engine/Source/Editor/UnrealEd/Private/Commandlets/DumpBlueprintsInfoCommandlet.cpp
@@ -1553,6 +1553,9 @@ static void DumpBlueprintInfoUtils::DumpPalette(uint32 Indent, UBlueprint* Bluep
 	else
 	{
 		FCategorizedGraphActionListBuilder PaletteBuilder;
+		PaletteBuilder.OwnerOfTemporaries = NewObject<UEdGraph>((UObject*)Blueprint);
+		PaletteBuilder.OwnerOfTemporaries->Schema = UEdGraphSchema_K2::StaticClass();
+		PaletteBuilder.OwnerOfTemporaries->SetFlags(RF_Transient);
 		double MenuBuildDuration = GetPaletteMenuActions(PaletteBuilder, Blueprint, ClassFilter);
 
 		BeginPaletteEntry += NestedIndent + "\"FilterClass\" : \"" + FilterClassName + "\",\n";
diff --git a/Engine/Source/Editor/UnrealEd/Private/Commandlets/GenerateBlueprintAPICommandlet.cpp b/Engine/Source/Editor/UnrealEd/Private/Commandlets/GenerateBlueprintAPICommandlet.cpp
--- a/Engine/Source/Editor/UnrealEd/Private/Commandlets/GenerateBlueprintAPICommandlet.cpp
+++ b/Engine/Source/Editor/UnrealEd/Private/Commandlets/GenerateBlueprintAPICommandlet.cpp
@@ -688,6 +688,10 @@ static void GenerateBlueprintAPIUtils::DumpPalette(uint32 Indent, UBlueprint* Bl
 	else
 	{
 		FCategorizedGraphActionListBuilder PaletteBuilder;
+		PaletteBuilder.OwnerOfTemporaries = NewObject<UEdGraph>((UObject*)Blueprint);
+		PaletteBuilder.OwnerOfTemporaries->Schema = UEdGraphSchema_K2::StaticClass();
+		PaletteBuilder.OwnerOfTemporaries->SetFlags(RF_Transient);
+
 		double MenuBuildDuration = GetPaletteMenuActions(PaletteBuilder, Blueprint, ClassFilter);
 
 		BeginPaletteEntry += NestedIndent + "\"FilterClass\" : \"" + FilterClassName + "\",\n";
