Fix for cook errors in Ocean to do with asset names being unable to be resolved.

diff --git a/Engine/Source/Editor/UnrealEd/Private/CookOnTheFlyServer.cpp b/Engine/Source/Editor/UnrealEd/Private/CookOnTheFlyServer.cpp
--- a/Engine/Source/Editor/UnrealEd/Private/CookOnTheFlyServer.cpp
+++ b/Engine/Source/Editor/UnrealEd/Private/CookOnTheFlyServer.cpp
@@ -664,7 +664,7 @@ public:
 
 void UCookOnTheFlyServer::GetDependentPackages( const TSet<UPackage*>& RootPackages, TSet<FName>& FoundPackages )
 {
-static const FName AssetRegistryName("AssetRegistry");
+	static const FName AssetRegistryName("AssetRegistry");
 	FAssetRegistryModule& AssetRegistryModule = FModuleManager::LoadModuleChecked<FAssetRegistryModule>(AssetRegistryName);
 	IAssetRegistry& AssetRegistry = AssetRegistryModule.Get();
 
@@ -683,8 +683,20 @@ static const FName AssetRegistryName("AssetRegistry");
 		TArray<FName> PackageDependencies;
 		verify( AssetRegistry.GetDependencies(FoundPackagesArray[FoundPackagesCounter], PackageDependencies) );
 		++FoundPackagesCounter;
-		for ( const auto& PackageDependency : PackageDependencies )
+		for ( const auto& OriginalPackageDependency : PackageDependencies )
 		{
+			// check(PackageDependency.ToString().StartsWith(TEXT("/")));
+			FName PackageDependency = OriginalPackageDependency;
+			FString PackageDepdencyString = PackageDependency.ToString();
+
+			if (!FPackageName::IsValidLongPackageName(PackageDepdencyString))
+			{			
+				LogCookerMessage(FString::Printf(TEXT("Unable to generate long package name for %s because %s"), *PackageDepdencyString, *FailureReason), EMessageSeverity::Warning);
+				UE_LOG(LogCook, Warning, TEXT("Unable to generate long package name for %s because %s"), *PackageDepdencyString, *FailureReason); 
+				continue;
+			}
+
+
 			if ( FoundPackages.Contains(PackageDependency) == false )
 			{
 				FoundPackages.Add(PackageDependency);
