Fixed undeclared identifier compilation error.

diff --git a/Engine/Source/Editor/UnrealEd/Private/CookOnTheFlyServer.cpp b/Engine/Source/Editor/UnrealEd/Private/CookOnTheFlyServer.cpp
--- a/Engine/Source/Editor/UnrealEd/Private/CookOnTheFlyServer.cpp
+++ b/Engine/Source/Editor/UnrealEd/Private/CookOnTheFlyServer.cpp
@@ -691,8 +691,8 @@ void UCookOnTheFlyServer::GetDependentPackages( const TSet<UPackage*>& RootPacka
 
 			if (!FPackageName::IsValidLongPackageName(PackageDepdencyString))
 			{			
-				LogCookerMessage(FString::Printf(TEXT("Unable to generate long package name for %s because %s"), *PackageDepdencyString, *FailureReason), EMessageSeverity::Warning);
-				UE_LOG(LogCook, Warning, TEXT("Unable to generate long package name for %s because %s"), *PackageDepdencyString, *FailureReason); 
+				LogCookerMessage(FString::Printf(TEXT("Unable to generate long package name for %s because %s"), *PackageDepdencyString), EMessageSeverity::Warning);
+				UE_LOG(LogCook, Warning, TEXT("Unable to generate long package name for %s because %s"), *PackageDepdencyString); 
 				continue;
 			}
 
