#HLOD : fixed issue with streaming level

diff --git a/Engine/Source/Editor/UnrealEd/Private/HierarchicalLOD.cpp b/Engine/Source/Editor/UnrealEd/Private/HierarchicalLOD.cpp
--- a/Engine/Source/Editor/UnrealEd/Private/HierarchicalLOD.cpp
+++ b/Engine/Source/Editor/UnrealEd/Private/HierarchicalLOD.cpp
@@ -166,11 +166,12 @@ void FHierarchicalLODBuilder::BuildClusters(class ULevel* InLevel)
 		for(int32 LODId=0; LODId<TotalNumLOD; ++LODId)
 		{
 			AWorldSettings* WorldSetting = InLevel->GetWorld()->GetWorldSettings();
-			check(WorldSetting->HierarchicalLODSetup.IsValidIndex(LODId));
 			float DesiredBoundBoxHalfSize = WorldSetting->HierarchicalLODSetup[LODId].DesiredBoundSize * 0.5f;
 			float DesiredFillingRatio = WorldSetting->HierarchicalLODSetup[LODId].DesiredFillingPercentage * 0.01f;
 			ensure(DesiredFillingRatio!=0.f);
 			float HighestCost = FMath::Pow(DesiredBoundBoxHalfSize, 3)/(DesiredFillingRatio);
+			int32 MinNumActors = WorldSetting->HierarchicalLODSetup[LODId].MinNumberOfActorsToBuild;
+			check (MinNumActors > 0);
 			// test parameter I was playing with to cull adding to the array
 			// intialization can have too many elements, decided to cull
 			// the problem can be that we can create disconnected tree
@@ -196,7 +197,7 @@ void FHierarchicalLODBuilder::BuildClusters(class ULevel* InLevel)
 			}
 
 			// now we have to calculate merge clusters and build actors
-			MergeClustersAndBuildActors(InLevel, LODId, HighestCost);
+			MergeClustersAndBuildActors(InLevel, LODId, HighestCost, MinNumActors);
 		}
 	}
 
@@ -221,7 +222,7 @@ void FHierarchicalLODBuilder::FindMST()
 	}
 }
 
-void FHierarchicalLODBuilder::MergeClustersAndBuildActors(class ULevel* InLevel, const int32 LODIdx, float HighestCost)
+void FHierarchicalLODBuilder::MergeClustersAndBuildActors(class ULevel* InLevel, const int32 LODIdx, float HighestCost, int32 MinNumActors)
 {
 	if (Clusters.Num() > 0)
 	{
@@ -306,11 +307,14 @@ void FHierarchicalLODBuilder::MergeClustersAndBuildActors(class ULevel* InLevel,
 
 			for(auto& Cluster: Clusters)
 			{
-				if(Cluster.ShouldBuild())
+				if(Cluster.IsValid())
 				{
 					SlowTask.EnterProgressFrame();
-					//Sleep(0.001f);
-					Cluster.BuildActor(InLevel->GetWorld(), InLevel, LODIdx);
+
+					if (Cluster.Actors.Num() >= MinNumActors)
+					{
+						Cluster.BuildActor(InLevel->GetWorld(), InLevel, LODIdx);
+					}
 				}
 			}
 		}
@@ -637,9 +641,12 @@ void FLODCluster::BuildActor(class UWorld* InWorld, class ULevel* InLevel, const
 
 				if (MainMesh)
 				{
+					UWorld* LevelWorld = Cast<UWorld>(InLevel->GetOuter());
+
+					check (LevelWorld);
 					////////////////////////////////////////////////////////////////////////////////////
 					// create LODActors using the current Actors
-					class ALODActor* NewActor = Cast<ALODActor>(InWorld->SpawnActor(ALODActor::StaticClass(), &OutProxyLocation, &FRotator::ZeroRotator));
+					class ALODActor* NewActor = Cast<ALODActor>(LevelWorld->SpawnActor(ALODActor::StaticClass(), &OutProxyLocation, &FRotator::ZeroRotator));
 
 					NewActor->SubObjects = OutAssets;
 					NewActor->SubActors = Actors;
