UE-14164 Crash dragging off Get Character Movement node

diff --git a/Engine/Source/Editor/UnrealEd/Private/Kismet2/BlueprintEditorUtils.cpp b/Engine/Source/Editor/UnrealEd/Private/Kismet2/BlueprintEditorUtils.cpp
--- a/Engine/Source/Editor/UnrealEd/Private/Kismet2/BlueprintEditorUtils.cpp
+++ b/Engine/Source/Editor/UnrealEd/Private/Kismet2/BlueprintEditorUtils.cpp
@@ -5870,7 +5870,7 @@ bool FBlueprintEditorUtils::IsSCSComponentProperty(UObjectProperty* MemberProper
 		// associated with a natively added component (which would still be  
 		// accessible through the SCS tree)
 
-		if (OwnerClass == nullptr)
+		if ((OwnerClass == nullptr) || !OwnerClass->IsChildOf<AActor>())
 		{
 			return false;
 		}
