Fix asset thumbnails in the content browser being very dark

diff --git a/Engine/Source/Editor/UnrealEd/Private/AssetThumbnail.cpp b/Engine/Source/Editor/UnrealEd/Private/AssetThumbnail.cpp
--- a/Engine/Source/Editor/UnrealEd/Private/AssetThumbnail.cpp
+++ b/Engine/Source/Editor/UnrealEd/Private/AssetThumbnail.cpp
@@ -124,7 +124,6 @@ public:
 			TSharedPtr<SViewport> Viewport = 
 				SNew( SViewport )
 				.EnableGammaCorrection(false)
-				.IgnoreTextureAlpha(false)
 				.EnableBlending(true);
 
 			Viewport->SetViewportInterface( AssetThumbnail.ToSharedRef() );
