Removing commented-out code

diff --git a/Engine/Source/Editor/UnrealEd/Private/Settings/EditorProjectSettings.cpp b/Engine/Source/Editor/UnrealEd/Private/Settings/EditorProjectSettings.cpp
--- a/Engine/Source/Editor/UnrealEd/Private/Settings/EditorProjectSettings.cpp
+++ b/Engine/Source/Editor/UnrealEd/Private/Settings/EditorProjectSettings.cpp
@@ -23,17 +23,6 @@ EUnit ConvertDefaultInputUnits(UEditorProjectAppearanceSettings::EDefaultLocatio
 	}
 }
 
-// EGlobalUnitDisplay ConvertGlobalUnitDisplay(EUnitDisplay In)
-// {
-// 	switch(In)
-// 	{
-// 	case EUnitDisplay::None:		return EGlobalUnitDisplay::None;
-// 	case EUnitDisplay::Metric:		return EGlobalUnitDisplay::Metric;
-// 	case EUnitDisplay::Imperial:	return EGlobalUnitDisplay::Imperial;
-// 	default:				 		return EGlobalUnitDisplay::Metric;
-// 	}
-// }
-
 TArray<EUnit> ToRawUnits(const TArray<TEnumAsByte<EUnit>>& In)
 {
 	TArray<EUnit> Units;
