#HLOD

- Changed default setting for draw distance
- Added message when user clicks Build LOD without build set ups
- Cleaned up code
- Hide Static Mesh Component from LOD Actor

diff --git a/Engine/Source/Editor/UnrealEd/Private/HierarchicalLOD.cpp b/Engine/Source/Editor/UnrealEd/Private/HierarchicalLOD.cpp
--- a/Engine/Source/Editor/UnrealEd/Private/HierarchicalLOD.cpp
+++ b/Engine/Source/Editor/UnrealEd/Private/HierarchicalLOD.cpp
@@ -119,6 +119,18 @@ void FHierarchicalLODBuilder::Build()
 {
 	check (World);
 
+	if (World->GetWorldSettings()->bEnableHierarchicalLODSystem == false)
+	{
+		if (EAppReturnType::No == 
+			FMessageDialog::Open( EAppMsgType::YesNo, LOCTEXT("HierarchicalLOD_Warning", 
+			"Enable Hierarchical LOD System is disabled in World Setting. This process will delete previously generated LODActors if exists.\n\nWould you like to continue?")) )
+		{
+			// has been canceled, return
+			GEditor->SetMapBuildCancelled(true);
+			return;
+		}
+	}
+	
 	const TArray<class ULevel*>& Levels = World->GetLevels();
 
 	for (const auto& LevelIter : Levels)
@@ -288,9 +300,10 @@ void FHierarchicalLODBuilder::MergeClustersAndBuildActors(class ULevel* InLevel,
 			}
 		}
 
-		static bool bBuildActor=true;
+		// debug flag, so that I can just see clustered data since no visualization in the editor yet
+		//static bool bBuildActor=true;
 
-		if (bBuildActor)
+		//if (bBuildActor)
 		{
 			// print data
 			int32 TotalValidCluster=0;
@@ -313,7 +326,7 @@ void FHierarchicalLODBuilder::MergeClustersAndBuildActors(class ULevel* InLevel,
 
 					if (Cluster.Actors.Num() >= MinNumActors)
 					{
-						Cluster.BuildActor(InLevel->GetWorld(), InLevel, LODIdx);
+						Cluster.BuildActor(InLevel, LODIdx);
 					}
 				}
 			}
@@ -572,14 +585,14 @@ void FLODCluster::SubtractCluster(const FLODCluster& Other)
 // this clears previous actors and sets to this new actor
 // this is required when new LOD is created from these actors, this will be replaced
 // to save memory and to reduce memory increase during this process, we discard previous actors and replace with this actor
-void FLODCluster::BuildActor(class UWorld* InWorld, class ULevel* InLevel, const int32 LODIdx)
+void FLODCluster::BuildActor(class ULevel* InLevel, const int32 LODIdx)
 {
 	// do big size
-	if (InWorld && InLevel)
+	if (InLevel && InLevel->GetWorld())
 	{
 		////////////////////////////////////////////////////////////////////////////////////
 		// create asset using Actors
-		const FHierarchicalSimplification& LODSetup = InWorld->GetWorldSettings()->HierarchicalLODSetup[LODIdx];
+		const FHierarchicalSimplification& LODSetup = InLevel->GetWorld()->GetWorldSettings()->HierarchicalLODSetup[LODIdx];
 		const FMeshProxySettings& ProxySetting = LODSetup.Setting;
 		// Where generated assets will be stored
 		UPackage* AssetsOuter = InLevel->GetOutermost(); // this asset is going to save with map, this means, I'll have to delete with it
diff --git a/Engine/Source/Editor/UnrealEd/Public/HierarchicalLOD.h b/Engine/Source/Editor/UnrealEd/Public/HierarchicalLOD.h
--- a/Engine/Source/Editor/UnrealEd/Public/HierarchicalLOD.h
+++ b/Engine/Source/Editor/UnrealEd/Public/HierarchicalLOD.h
@@ -58,7 +58,7 @@ struct FLODCluster
 	// this clears previous actors and sets to this new actor
 	// this is required when new LOD is created from these actors, this will be replaced
 	// to save memory and to reduce memory increase during this process, we discard previous actors and replace with this actor
-	void BuildActor(class UWorld* InWorld, class ULevel* InLevel, const int32 LODIdx);
+	void BuildActor(class ULevel* InLevel, const int32 LODIdx);
 
 private:
 	// cluster operations
@@ -66,6 +66,7 @@ private:
 	void SubtractCluster(const FLODCluster& Other);
 
 	// add new actor, this doesn't recalculate filling factor
+	// the filling factor has to be calculated outside
 	FSphere AddActor(class AActor* NewActor);
 
 	// whether it's valid or not
@@ -94,6 +95,7 @@ struct UNREALED_API FHierarchicalLODBuilder
 private:
 	// data structure - this is only valid within scope since mem stack allocator
 	TArray<FLODCluster, TMemStackAllocator<>> Clusters;
+
 	// owner world
 	class UWorld*		World;
 
diff --git a/Engine/Source/Runtime/Engine/Classes/Engine/LODActor.h b/Engine/Source/Runtime/Engine/Classes/Engine/LODActor.h
--- a/Engine/Source/Runtime/Engine/Classes/Engine/LODActor.h
+++ b/Engine/Source/Runtime/Engine/Classes/Engine/LODActor.h
@@ -16,7 +16,8 @@ class ENGINE_API ALODActor : public AActor
 	GENERATED_UCLASS_BODY()
 
 private_subobject:
-	UPROPERTY(Category=LODActor, VisibleAnywhere)
+	// disable display of this component
+	UPROPERTY(/*Category=LODActor, VisibleAnywhere*/)
 	class UStaticMeshComponent* StaticMeshComponent;
 
 public:
diff --git a/Engine/Source/Runtime/Engine/Classes/GameFramework/WorldSettings.h b/Engine/Source/Runtime/Engine/Classes/GameFramework/WorldSettings.h
--- a/Engine/Source/Runtime/Engine/Classes/GameFramework/WorldSettings.h
+++ b/Engine/Source/Runtime/Engine/Classes/GameFramework/WorldSettings.h
@@ -226,7 +226,7 @@ struct ENGINE_API FHierarchicalSimplification
 
 	FHierarchicalSimplification()
 		: bSimplifyMesh(true)
-		, DrawDistance(5000)
+		, DrawDistance(2000)
 		, DesiredBoundSize(5000)
 		, DesiredFillingPercentage(50)
 		, MinNumberOfActorsToBuild(2)
