Fixed rotation widget when selected with Ctrl+RMB so that it doesn't use the cursor location relative to the widget to determine which direction to rotate.

diff --git a/Engine/Source/Editor/UnrealEd/Private/MouseDeltaTracker.cpp b/Engine/Source/Editor/UnrealEd/Private/MouseDeltaTracker.cpp
--- a/Engine/Source/Editor/UnrealEd/Private/MouseDeltaTracker.cpp
+++ b/Engine/Source/Editor/UnrealEd/Private/MouseDeltaTracker.cpp
@@ -621,7 +621,7 @@ void FMouseDeltaTracker::ConvertMovementDeltaToDragRot(FEditorViewportClient* In
 
 	if( InViewportClient->GetCurrentWidgetAxis() != EAxisList::None )
 	{
-		InViewportClient->Widget->ConvertMouseMovementToAxisMovement( InViewportClient, InViewportClient->GetWidgetLocation(), InOutDragDelta, OutDrag, OutRotation, OutScale );
+		InViewportClient->Widget->ConvertMouseMovementToAxisMovement( InViewportClient, bUsedDragModifier, InOutDragDelta, OutDrag, OutRotation, OutScale );
 	}
 	else
 	{
diff --git a/Engine/Source/Editor/UnrealEd/Private/UnrealWidget.cpp b/Engine/Source/Editor/UnrealEd/Private/UnrealWidget.cpp
--- a/Engine/Source/Editor/UnrealEd/Private/UnrealWidget.cpp
+++ b/Engine/Source/Editor/UnrealEd/Private/UnrealWidget.cpp
@@ -930,7 +930,7 @@ void FWidget::Render_2D(const FSceneView* View, FPrimitiveDrawInterface* PDI, FE
 /**
  * Converts mouse movement on the screen to widget axis movement/rotation.
  */
-void FWidget::ConvertMouseMovementToAxisMovement( FEditorViewportClient* InViewportClient, const FVector& InLocation, FVector& InOutDelta, FVector& OutDrag, FRotator& OutRotation, FVector& OutScale )
+void FWidget::ConvertMouseMovementToAxisMovement( FEditorViewportClient* InViewportClient, bool bInUsedDragModifier, FVector& InOutDelta, FVector& OutDrag, FRotator& OutRotation, FVector& OutScale )
 {
 	OutDrag = FVector::ZeroVector;
 	OutRotation = FRotator::ZeroRotator;
@@ -944,8 +944,19 @@ void FWidget::ConvertMouseMovementToAxisMovement( FEditorViewportClient* InViewp
 	// Get offset of the drag start position from the widget origin
 	const FVector2D DirectionToMousePos = FVector2D(DragStartPos - Origin).GetSafeNormal();
 
-	// Get direction of tangent of circle formed from the widget origin and the drag start position
-	const FVector2D TangentDir = FVector2D(-DirectionToMousePos.Y, DirectionToMousePos.X);
+	// For rotations which display as a full ring, calculate the tangent direction representing a clockwise movement
+	FVector2D TangentDir;
+	if (bInUsedDragModifier)
+	{
+		// If a drag modifier has been used, this implies we are not actually touching the widget, so don't attempt to
+		// calculate the tangent dir based on the relative offset of the cursor from the widget location.
+		TangentDir = FVector2D(1, 1).GetSafeNormal();
+	}
+	else
+	{
+		// Treat the tangent dir as perpendicular to the relative offset of the cursor from the widget location.
+		TangentDir = FVector2D(-DirectionToMousePos.Y, DirectionToMousePos.X);
+	}
 
 	switch (WidgetMode)
 	{
diff --git a/Engine/Source/Editor/UnrealEd/Public/UnrealWidget.h b/Engine/Source/Editor/UnrealEd/Public/UnrealWidget.h
--- a/Engine/Source/Editor/UnrealEd/Public/UnrealWidget.h
+++ b/Engine/Source/Editor/UnrealEd/Public/UnrealWidget.h
@@ -80,7 +80,7 @@ public:
 	/**
 	 * Converts mouse movement on the screen to widget axis movement/rotation.
 	 */
-	void ConvertMouseMovementToAxisMovement( FEditorViewportClient* InViewportClient, const FVector& InLocation, FVector& InDiff, FVector& OutDrag, FRotator& OutRotation, FVector& OutScale );
+	void ConvertMouseMovementToAxisMovement( FEditorViewportClient* InViewportClient, bool bInUsedDragModifier, FVector& InDiff, FVector& OutDrag, FRotator& OutRotation, FVector& OutScale );
 
 	/**
 	 * Absolute Translation conversion from mouse movement on the screen to widget axis movement/rotation.
