Allowed the UDS editor to set asset properties

UE-11714 - User defined structures cannot reference asset/object/class types in defaults

diff --git a/Engine/Source/Editor/Kismet/Private/UserDefinedStructureEditor.cpp b/Engine/Source/Editor/Kismet/Private/UserDefinedStructureEditor.cpp
--- a/Engine/Source/Editor/Kismet/Private/UserDefinedStructureEditor.cpp
+++ b/Engine/Source/Editor/Kismet/Private/UserDefinedStructureEditor.cpp
@@ -36,7 +36,7 @@ public:
 
 		FStructureDetailsViewArgs StructureViewArgs;
 		StructureViewArgs.bShowObjects = false;
-		StructureViewArgs.bShowAssets = false;
+		StructureViewArgs.bShowAssets = true;
 		StructureViewArgs.bShowClasses = true;
 		StructureViewArgs.bShowInterfaces = false;
 
diff --git a/Engine/Source/Editor/UnrealEd/Private/Kismet2/BlueprintEditorUtils.cpp b/Engine/Source/Editor/UnrealEd/Private/Kismet2/BlueprintEditorUtils.cpp
--- a/Engine/Source/Editor/UnrealEd/Private/Kismet2/BlueprintEditorUtils.cpp
+++ b/Engine/Source/Editor/UnrealEd/Private/Kismet2/BlueprintEditorUtils.cpp
@@ -7115,35 +7115,17 @@ bool FBlueprintEditorUtils::PropertyValueFromString(const UProperty* Property, c
 		{
 			CastChecked<UTextProperty>(Property)->SetPropertyValue_InContainer(DefaultObject, FText::FromString(Value));
 		}
-		else if( Property->IsA(UClassProperty::StaticClass()) )
-		{
-			FStringOutputDevice ImportError;
-			const auto EndOfParsedBuff = Property->ImportText(Value.IsEmpty() ? TEXT("()") : *Value, Property->ContainerPtrToValuePtr<uint8>(DefaultObject), 0, NULL, &ImportError);
-			bParseSucceeded = EndOfParsedBuff && ImportError.IsEmpty();
-		}
-		else if( Property->IsA(UObjectPropertyBase::StaticClass()) )
-		{
-			CastChecked<UObjectPropertyBase>(Property)->SetObjectPropertyValue_InContainer(DefaultObject, FindObject<UObject>(ANY_PACKAGE, *Value));
-		}
-		else if (Property->IsA(UArrayProperty::StaticClass()) 
-			|| Property->IsA(UMulticastDelegateProperty::StaticClass()))
-		{
-			FStringOutputDevice ImportError;
-			const auto EndOfParsedBuff = Property->ImportText(Value.IsEmpty() ? TEXT("()") : *Value, Property->ContainerPtrToValuePtr<uint8>(DefaultObject), 0, NULL, &ImportError);
-			bParseSucceeded = EndOfParsedBuff && ImportError.IsEmpty();
-		}
-		else if (Property->IsA(UInterfaceProperty::StaticClass()))
+		else
 		{
+			// Empty array-like properties need to use "()" in order to import correctly (as array properties export comma separated within a set of brackets)
+			const TCHAR* const ValueToImport = (Value.IsEmpty() && (Property->IsA(UArrayProperty::StaticClass()) || Property->IsA(UMulticastDelegateProperty::StaticClass())))
+				? TEXT("()")
+				: *Value;
+
 			FStringOutputDevice ImportError;
-			const auto EndOfParsedBuff = Property->ImportText(*Value, Property->ContainerPtrToValuePtr<uint8>(DefaultObject), 0, NULL, &ImportError);
+			const auto EndOfParsedBuff = Property->ImportText(*Value, Property->ContainerPtrToValuePtr<uint8>(DefaultObject), 0, nullptr, &ImportError);
 			bParseSucceeded = EndOfParsedBuff && ImportError.IsEmpty();
 		}
-		else
-		{
-			// HOOK UP NEW TYPES HERE
-			check(0);
-		}
-
 	}
 	else 
 	{
@@ -7190,7 +7172,7 @@ bool FBlueprintEditorUtils::PropertyValueFromString(const UProperty* Property, c
 			bParseSucceeded = FStructureEditorUtils::Fill_MakeStructureDefaultValue(Cast<const UUserDefinedStruct>(Struct), StructData);
 
 			FStringOutputDevice ImportError;
-			const auto EndOfParsedBuff = StructProperty->ImportText(Value.IsEmpty() ? TEXT("()") : *Value, StructData, 0, NULL, &ImportError);
+			const auto EndOfParsedBuff = StructProperty->ImportText(Value.IsEmpty() ? TEXT("()") : *Value, StructData, 0, nullptr, &ImportError);
 			bParseSucceeded &= EndOfParsedBuff && ImportError.IsEmpty();
 		}
 	}
@@ -7201,7 +7183,7 @@ bool FBlueprintEditorUtils::PropertyValueFromString(const UProperty* Property, c
 bool FBlueprintEditorUtils::PropertyValueToString(const UProperty* Property, const uint8* Container, FString& OutForm)
 {
 	check(Property && Container);
-	OutForm.Empty();
+	OutForm.Reset();
 	if (Property->IsA(UStructProperty::StaticClass()))
 	{
 		static UScriptStruct* VectorStruct = GetBaseStructure(TEXT("Vector"));
@@ -7241,7 +7223,7 @@ bool FBlueprintEditorUtils::PropertyValueToString(const UProperty* Property, con
 	bool bSuccedded = true;
 	if (OutForm.IsEmpty())
 	{
-		bSuccedded = Property->ExportText_InContainer(0, OutForm, Container, Container, NULL, 0);
+		bSuccedded = Property->ExportText_InContainer(0, OutForm, Container, Container, nullptr, 0);
 	}
 	return bSuccedded;
 }
