Missing include.

diff --git a/Engine/Source/Editor/UnrealEd/Private/EditorCommandLineUtils.cpp b/Engine/Source/Editor/UnrealEd/Private/EditorCommandLineUtils.cpp
--- a/Engine/Source/Editor/UnrealEd/Private/EditorCommandLineUtils.cpp
+++ b/Engine/Source/Editor/UnrealEd/Private/EditorCommandLineUtils.cpp
@@ -13,6 +13,7 @@
 #include "AssetToolsModule.h"
 #include "AssetEditorManager.h"
 #include "Editor.h"					// for OnShutdownPostPackagesSaved
+#include "ProjectDescriptor.h"
 
 #define LOCTEXT_NAMESPACE "EditorCommandLineUtils"
 
