Fixed crash in gathering text from assets, specifically dialogue waves, if the dialogue wave has no targets.

diff --git a/Engine/Source/Editor/UnrealEd/Private/Commandlets/GatherTextFromAssetsCommandlet.cpp b/Engine/Source/Editor/UnrealEd/Private/Commandlets/GatherTextFromAssetsCommandlet.cpp
--- a/Engine/Source/Editor/UnrealEd/Private/Commandlets/GatherTextFromAssetsCommandlet.cpp
+++ b/Engine/Source/Editor/UnrealEd/Private/Commandlets/GatherTextFromAssetsCommandlet.cpp
@@ -114,8 +114,8 @@ bool UGatherTextFromAssetsCommandlet::FDialogueHelper::ProcessDialogueWave( cons
 			}
 		}
 
-		FString FinalTargetGender = GetGrammaticalGenderString( AccumulatedTargetGender );
-		FString FinalTargetPlurality = GetGrammaticalNumberString( AccumulatedTargetPlurality );
+		FString FinalTargetGender = DialogueContext.Targets.Num() ? GetGrammaticalGenderString( AccumulatedTargetGender ) : TEXT("");
+		FString FinalTargetPlurality = DialogueContext.Targets.Num() ? GetGrammaticalNumberString( AccumulatedTargetPlurality ) : TEXT("");
 
 		// Add the context specific variation
 		{
