Editor - Only showing notifications when cooking if it's an error, preventing warning spam.  Fixing a bug in the message log, we were not properly checking a 1 higher severity level, like the comment indicated the severity filter parameter was for.

diff --git a/Engine/Source/Developer/MessageLog/Private/Presentation/MessageLogListingViewModel.cpp b/Engine/Source/Developer/MessageLog/Private/Presentation/MessageLogListingViewModel.cpp
--- a/Engine/Source/Developer/MessageLog/Private/Presentation/MessageLogListingViewModel.cpp
+++ b/Engine/Source/Developer/MessageLog/Private/Presentation/MessageLogListingViewModel.cpp
@@ -286,7 +286,10 @@ void FMessageLogListingViewModel::NotifyIfAnyMessages( const FText& Message, EMe
 	// Note we use page 0 in this function, as that is the page that will
 	// have most recently had messages added to it.
 
-	if ( bForce || NumMessagesPresent(0, SeverityFilter) > 0 )
+	// SeverityFilter represents only logging items Higher severity than it.
+	EMessageSeverity::Type HigherSeverity = ( EMessageSeverity::Type )( FMath::Max(SeverityFilter - 1, 0) );
+
+	if ( bForce || NumMessagesPresent(0, HigherSeverity) > 0 )
 	{
 		FText NotificationMessage;
 		if(Message.IsEmpty())
diff --git a/Engine/Source/Editor/UnrealEd/Private/CookOnTheFlyServer.cpp b/Engine/Source/Editor/UnrealEd/Private/CookOnTheFlyServer.cpp
--- a/Engine/Source/Editor/UnrealEd/Private/CookOnTheFlyServer.cpp
+++ b/Engine/Source/Editor/UnrealEd/Private/CookOnTheFlyServer.cpp
@@ -36,6 +36,8 @@
 #include "Engine/LevelStreaming.h"
 #include "Engine/TextureLODSettings.h"
 
+#define LOCTEXT_NAMESPACE "Cooker"
+
 DEFINE_LOG_CATEGORY_STATIC(LogCook, Log, All);
 
 
@@ -308,7 +310,7 @@ void LogCookerMessage( const FString& MessageText, EMessageSeverity::Type Severi
 	// Message->AddToken(FDocumentationToken::Create(TEXT("https://docs.unrealengine.com/latest/INT/Platforms/iOS/QuickStart/6/index.html"))); 
 	MessageLog.AddMessage(Message);
 
-	MessageLog.Notify();
+	MessageLog.Notify(FText(), EMessageSeverity::Warning, false);
 }
 
 
@@ -653,7 +655,7 @@ public:
 	}
 
 	/**
-  	 * Returns the name of the Archive.  Useful for getting the name of the package a struct or object
+	 * Returns the name of the Archive.  Useful for getting the name of the package a struct or object
 	 * is in when a loading error occurs.
 	 *
 	 * This is overridden for the specific Archive Types
@@ -679,7 +681,6 @@ void UCookOnTheFlyServer::GetDependentPackages( const TSet<UPackage*>& RootPacka
 	int FoundPackagesCounter = 0;
 	while ( FoundPackagesCounter < FoundPackagesArray.Num() )
 	{
-
 		TArray<FName> PackageDependencies;
 		verify( AssetRegistry.GetDependencies(FoundPackagesArray[FoundPackagesCounter], PackageDependencies) );
 		++FoundPackagesCounter;
@@ -689,20 +690,22 @@ void UCookOnTheFlyServer::GetDependentPackages( const TSet<UPackage*>& RootPacka
 			FName PackageDependency = OriginalPackageDependency;
 			FString PackageDepdencyString = PackageDependency.ToString();
 
-			if (!FPackageName::IsValidLongPackageName(PackageDepdencyString))
-			{			
-				LogCookerMessage(FString::Printf(TEXT("Unable to generate long package name for %s"), *PackageDepdencyString), EMessageSeverity::Warning);
-				UE_LOG(LogCook, Warning, TEXT("Unable to generate long package name for %s"), *PackageDepdencyString); 
+			FText OutReason;
+			if ( !FPackageName::IsValidLongPackageName(PackageDepdencyString, false, &OutReason) )
+			{
+				const FText FailMessage = FText::Format(LOCTEXT("UnableToGeneratePackageName", "Unable to generate long package name for {0}. {1}"),
+					FText::FromString(PackageDepdencyString), OutReason);
+
+				LogCookerMessage(FailMessage.ToString(), EMessageSeverity::Warning);
+				UE_LOG(LogCook, Warning, TEXT("%s"), *( FailMessage.ToString() ));
 				continue;
 			}
 
-
 			if ( FoundPackages.Contains(PackageDependency) == false )
 			{
 				FoundPackages.Add(PackageDependency);
 				FoundPackagesArray.Add( PackageDependency );
 			}
-			
 		}
 	}	
 
@@ -1900,8 +1903,8 @@ bool UCookOnTheFlyServer::SaveCookedPackage( UPackage* Package, uint32 SaveFlags
 		Filename = ConvertToFullSandboxPath(*Filename, true);
 
 		uint32 OriginalPackageFlags = Package->PackageFlags;
- 		UWorld* World = NULL;
- 		EObjectFlags Flags = RF_NoFlags;
+		UWorld* World = NULL;
+		EObjectFlags Flags = RF_NoFlags;
 		bool bPackageFullyLoaded = false;
 
 		if (IsCookFlagSet(ECookInitializationFlags::Compressed) )
@@ -3506,4 +3509,6 @@ bool UCookOnTheFlyServer::GetAllPackagesFromAssetRegistry( const FString& AssetR
 		return true;
 	}
 	return false;
-}
\ No newline at end of file
+}
+
+#undef LOCTEXT_NAMESPACE
\ No newline at end of file
