Fix for cooked asset registry getting out of date in iterative case.
Improved performance of cook startup time.

diff --git a/Engine/Source/Editor/UnrealEd/Classes/CookOnTheSide/CookOnTheFlyServer.h b/Engine/Source/Editor/UnrealEd/Classes/CookOnTheSide/CookOnTheFlyServer.h
--- a/Engine/Source/Editor/UnrealEd/Classes/CookOnTheSide/CookOnTheFlyServer.h
+++ b/Engine/Source/Editor/UnrealEd/Classes/CookOnTheSide/CookOnTheFlyServer.h
@@ -641,6 +641,7 @@ private:
 	ECookInitializationFlags CookFlags;
 	TAutoPtr<class FSandboxPlatformFile> SandboxFile;
 	bool bIsSavingPackage; // used to stop recursive mark package dirty functions
+	TSet<FName> PackagesKeptFromPreviousCook; // used for iterative cooking this is a list of the packages which were kept from the previous cook
 
 	//////////////////////////////////////////////////////////////////////////
 
@@ -888,14 +889,14 @@ private:
 	/**
 	 * Collect all the files which need to be cooked for a cook by the book session
 	 */
-	void CollectFilesToCook(TArray<FString>& FilesInPath, 
+	void CollectFilesToCook(TArray<FName>& FilesInPath, 
 		const TArray<FString>& CookMaps, const TArray<FString>& CookDirectories, const TArray<FString>& CookCultures, 
 		const TArray<FString>& IniMapSections, bool bCookAll, bool bMapsOnly, bool bNoDev);
 
 	/**
 	 * AddFileToCook add file to cook list 
 	 */
-	void AddFileToCook( TArray<FString>& InOutFilesToCook, const FString &InFilename ) const;
+	void AddFileToCook( TArray<FName>& InOutFilesToCook, const FString &InFilename ) const;
 
 	/**
 	 * Call back from the TickCookOnTheSide when a cook by the book finishes (when started form StartCookByTheBook)
@@ -1003,6 +1004,18 @@ private:
 	 */
 	FString ConvertToFullSandboxPath( const FString &FileName, bool bForWrite = false ) const;
 	FString ConvertToFullSandboxPath( const FString &FileName, bool bForWrite, const FString& PlatformName ) const;
+
+
+
+	/**
+	 * GetSandboxAssetRegistryFilename
+	 * 
+	 * return full path of the asset registry in the sandbox
+	 */
+	const FString GetSandboxAssetRegistryFilename();
+
+	const FString GetCookedAssetRegistryFilename(const FString& PlatformName);
+
 	/**
 	 * Get the sandbox root directory for that platform
 	 * is effected by the CookingDlc settings
@@ -1124,7 +1137,7 @@ private:
 	void GenerateAssetRegistry(const TArray<ITargetPlatform*>& Platforms);
 
 	/** Generates long package names for all files to be cooked */
-	void GenerateLongPackageNames(TArray<FString>& FilesInPath);
+	void GenerateLongPackageNames(TArray<FName>& FilesInPath);
 
 
 	void GetDependencies( UPackage* Package, TArray<UPackage*> Dependencies );
diff --git a/Engine/Source/Editor/UnrealEd/Private/Commandlets/ChunkManifestGenerator.cpp b/Engine/Source/Editor/UnrealEd/Private/Commandlets/ChunkManifestGenerator.cpp
--- a/Engine/Source/Editor/UnrealEd/Private/Commandlets/ChunkManifestGenerator.cpp
+++ b/Engine/Source/Editor/UnrealEd/Private/Commandlets/ChunkManifestGenerator.cpp
@@ -407,6 +407,50 @@ bool FChunkManifestGenerator::SaveManifests(FSandboxPlatformFile* InSandboxFile)
 	return true;
 }
 
+bool FChunkManifestGenerator::LoadAssetRegistry(const FString& SandboxPath, const TSet<FName>* PackagesToKeep)
+{
+	UE_LOG(LogChunkManifestGenerator, Display, TEXT("Saving asset registry."));
+
+	// Load generated registry for each platform
+	check(Platforms.Num() == 1);
+
+	for (auto Platform : Platforms)
+	{
+		FString PlatformSandboxPath = SandboxPath.Replace(TEXT("[Platform]"), *Platform->PlatformName());
+		FArchive* AssetRegistryReader = IFileManager::Get().CreateFileReader(*PlatformSandboxPath);
+
+		TMap<FName, FAssetData*> SavedAssetRegistryData;
+		AssetRegistry.LoadRegistryData(*AssetRegistryReader, SavedAssetRegistryData);
+
+		for (FAssetData LoadedAssetData : AssetRegistryData)
+		{
+			if (PackagesToKeep &&
+				PackagesToKeep->Contains(LoadedAssetData.PackageName) == false)
+			{
+				continue;
+			}
+
+			FAssetData* FoundAssetData = SavedAssetRegistryData.FindRef(LoadedAssetData.ObjectPath);
+			if ( FoundAssetData )
+			{
+				LoadedAssetData.ChunkIDs.Append(FoundAssetData->ChunkIDs);
+				
+				SavedAssetRegistryData.Remove(LoadedAssetData.ObjectPath);
+				delete FoundAssetData;
+			}
+		}
+
+		for (const auto& SavedAsset : SavedAssetRegistryData)
+		{
+			if (PackagesToKeep && PackagesToKeep->Contains(SavedAsset.Value->PackageName))
+			AssetRegistryData.Add(*SavedAsset.Value);
+			delete SavedAsset.Value;
+		}
+		SavedAssetRegistryData.Empty();
+	}
+	return true;
+}
+
 bool FChunkManifestGenerator::SaveAssetRegistry(const FString& SandboxPath)
 {
 	UE_LOG(LogChunkManifestGenerator, Display, TEXT("Saving asset registry."));
@@ -422,7 +466,7 @@ bool FChunkManifestGenerator::SaveAssetRegistry(const FString& SandboxPath)
 			GeneratedAssetRegistryData.Add(AssetData.ObjectPath, &AssetData);
 		}
 	}
-	AssetRegistry.SaveRegistryData(SerializedAssetRegistry, GeneratedAssetRegistryData, GeneratedAssetRegistryData.Num());
+	AssetRegistry.SaveRegistryData(SerializedAssetRegistry, GeneratedAssetRegistryData);
 	UE_LOG(LogChunkManifestGenerator, Display, TEXT("Generated asset registry num assets %d, size is %5.2fkb"), GeneratedAssetRegistryData.Num(), (float)SerializedAssetRegistry.Num() / 1024.f);
 
 	// Save the generated registry for each platform
diff --git a/Engine/Source/Editor/UnrealEd/Private/Commandlets/ChunkManifestGenerator.h b/Engine/Source/Editor/UnrealEd/Private/Commandlets/ChunkManifestGenerator.h
--- a/Engine/Source/Editor/UnrealEd/Private/Commandlets/ChunkManifestGenerator.h
+++ b/Engine/Source/Editor/UnrealEd/Private/Commandlets/ChunkManifestGenerator.h
@@ -303,6 +303,14 @@ public:
 	bool SaveAssetRegistry(const FString& SandboxPath);
 
 
+	/**
+	 * Loads asset registry data and incorporates it into the current asset registry data
+	 * 
+	 * @param the SandboxPath to load teh asset registry from
+	 * @param PackagesToLoadMask is used to mask out any packages which we don't want to incorporate into the internal asset registry when we load
+	 */
+	bool LoadAssetRegistry(const FString& SandboxPath, const TSet<FName>* PackagesToLoadMask = nullptr);
+
 	/**
 	 * Saves cooked package and asset information about all the cooked packages and assets contained within for stats purposes
 	 * in json format
diff --git a/Engine/Source/Editor/UnrealEd/Private/CookOnTheFlyServer.cpp b/Engine/Source/Editor/UnrealEd/Private/CookOnTheFlyServer.cpp
--- a/Engine/Source/Editor/UnrealEd/Private/CookOnTheFlyServer.cpp
+++ b/Engine/Source/Editor/UnrealEd/Private/CookOnTheFlyServer.cpp
@@ -2482,6 +2482,9 @@ void UCookOnTheFlyServer::CleanSandbox( const bool bIterative )
 			TArray<FString> DirectoriesToSkip;
 			TArray<FString> DirectoriesToNotRecurse;
 
+
+			PackagesKeptFromPreviousCook.Empty();
+
 			// See what files are out of date in the sandbox folder
 			for (int32 Index = 0; Index < Platforms.Num(); Index++)
 			{
@@ -2500,6 +2503,7 @@ void UCookOnTheFlyServer::CleanSandbox( const bool bIterative )
 					FDateTime CookedTimestamp = TimestampIt.Value();
 					FString StandardCookedFilename = CookedFilename.Replace(*SandboxDirectory, *(FPaths::GetRelativePathToRoot()));
 					FDateTime DependentTimestamp;
+					FName StandardCookedFileFName = FName(*StandardCookedFilename);
 
 					if (PDInfoModule.DeterminePackageDependentTimeStamp(*(FPaths::GetBaseFilename(StandardCookedFilename, false)), DependentTimestamp) == true)
 					{
@@ -2512,9 +2516,18 @@ void UCookOnTheFlyServer::CleanSandbox( const bool bIterative )
 #endif
 							IFileManager::Get().Delete(*CookedFilename);
 
-							CookedPackages.RemoveFileForPlatform(FName(*StandardCookedFilename), PlatformFName);
+							PackagesKeptFromPreviousCook.Remove(StandardCookedFileFName);
+							CookedPackages.RemoveFileForPlatform(StandardCookedFileFName, PlatformFName);
+						}
+						else
+						{
+							PackagesKeptFromPreviousCook.Add(StandardCookedFileFName);
 						}
 					}
+					else
+					{
+						PackagesKeptFromPreviousCook.Add(StandardCookedFileFName);
+					}
 				}
 			}
 
@@ -2561,7 +2574,7 @@ void UCookOnTheFlyServer::GenerateAssetRegistry(const TArray<ITargetPlatform*>&
 	UE_LOG(LogCook, Display, TEXT("Done creating registry. It took %5.2fs."), GenerateAssetRegistryTime);
 }
 
-void UCookOnTheFlyServer::GenerateLongPackageNames(TArray<FString>& FilesInPath)
+void UCookOnTheFlyServer::GenerateLongPackageNames(TArray<FName>& FilesInPath)
 {
 #if 0
 	TArray<FString> FilesInPathReverse;
@@ -2594,7 +2607,7 @@ void UCookOnTheFlyServer::GenerateLongPackageNames(TArray<FString>& FilesInPath)
 
 	for( int32 FileIndex = 0; FileIndex < FilesInPath.Num(); FileIndex++ )
 	{
-		const FString& FileInPath = FilesInPath[FilesInPath.Num() - FileIndex - 1];
+		const FString& FileInPath = FilesInPath[FilesInPath.Num() - FileIndex - 1].ToString();
 		if (FPackageName::IsValidLongPackageName(FileInPath))
 		{
 			const FName FileInPathFName(*FileInPath);
@@ -2616,14 +2629,15 @@ void UCookOnTheFlyServer::GenerateLongPackageNames(TArray<FString>& FilesInPath)
 		}
 	}
 	// Exchange(FilesInPathReverse, FilesInPath);
+	FilesInPath.Empty(FilesInPathReverse.Num());
 	for ( const auto& Files : FilesInPathReverse )
 	{
-		FilesInPath.Add(Files.ToString());
+		FilesInPath.Add(Files);
 	}
 #endif
 }
 
-void UCookOnTheFlyServer::AddFileToCook( TArray<FString>& InOutFilesToCook, const FString &InFilename ) const
+void UCookOnTheFlyServer::AddFileToCook( TArray<FName>& InOutFilesToCook, const FString &InFilename ) const
 { 
 	if (!FPackageName::IsScriptPackage(InFilename))
 	{
@@ -2634,12 +2648,12 @@ void UCookOnTheFlyServer::AddFileToCook( TArray<FString>& InOutFilesToCook, cons
 			InOutFilesToCook.Insert(InFilename, Index );
 		}
 #else
-		InOutFilesToCook.AddUnique(InFilename);
+		InOutFilesToCook.AddUnique(FName(*InFilename));
 #endif
 	}
 }
 
-void UCookOnTheFlyServer::CollectFilesToCook(TArray<FString>& FilesInPath, const TArray<FString> &CookMaps, const TArray<FString> &InCookDirectories, const TArray<FString> &CookCultures, const TArray<FString> &IniMapSections, bool bCookAll, bool bMapsOnly, bool bNoDev)
+void UCookOnTheFlyServer::CollectFilesToCook(TArray<FName>& FilesInPath, const TArray<FString> &CookMaps, const TArray<FString> &InCookDirectories, const TArray<FString> &CookCultures, const TArray<FString> &IniMapSections, bool bCookAll, bool bMapsOnly, bool bNoDev)
 {
 	TArray<FString> CookDirectories = InCookDirectories;
 	
@@ -2948,6 +2962,19 @@ FString UCookOnTheFlyServer::ConvertToFullSandboxPath( const FString &FileName,
 	return Result;
 }
 
+const FString UCookOnTheFlyServer::GetSandboxAssetRegistryFilename()
+{
+	static const FString RegistryFilename = FPaths::GameDir() / TEXT("AssetRegistry.bin");
+	const FString SandboxRegistryFilename = ConvertToFullSandboxPath(*RegistryFilename, true);
+	return SandboxRegistryFilename;
+}
+
+const FString UCookOnTheFlyServer::GetCookedAssetRegistryFilename(const FString& PlatformName )
+{
+	const FString CookedAssetRegistryFilename = GetSandboxAssetRegistryFilename().Replace(TEXT("[Platform]"), *PlatformName);
+	return CookedAssetRegistryFilename;
+}
+
 void UCookOnTheFlyServer::CookByTheBookFinished()
 {
 	check( IsInGameThread() );
@@ -2960,8 +2987,7 @@ void UCookOnTheFlyServer::CookByTheBookFinished()
 
 	{
 		// Save modified asset registry with all streaming chunk info generated during cook
-		const FString RegistryFilename = FPaths::GameDir() / TEXT("AssetRegistry.bin");
-		const FString SandboxRegistryFilename = ConvertToFullSandboxPath(*RegistryFilename, true);
+		const FString& SandboxRegistryFilename = GetSandboxAssetRegistryFilename();
 		// the registry filename will be modified when we call save asset registry
 		
 		const FString CookedAssetRegistry = FPaths::GameDir() / TEXT("CookedAssetRegistry.json");
@@ -3207,6 +3233,14 @@ void UCookOnTheFlyServer::StartCookByTheBook( const FCookByTheBookStartupOptions
 			ManifestGenerator->CleanManifestDirectories();
 			ManifestGenerator->Initialize( CookByTheBookOptions->bGenerateStreamingInstallManifests );
 
+
+			if (IsCookFlagSet(ECookInitializationFlags::Iterative) && (PackagesKeptFromPreviousCook.Num() > 0) )
+			{
+				const FString& SandboxRegistryFilename = GetSandboxAssetRegistryFilename();
+				ManifestGenerator->LoadAssetRegistry(GetCookedAssetRegistryFilename(PlatformName.ToString()), &PackagesKeptFromPreviousCook);
+			}
+
+
 			CookByTheBookOptions->ManifestGenerators.Add(PlatformName, ManifestGenerator);
 		}
 	}
@@ -3236,8 +3270,14 @@ void UCookOnTheFlyServer::StartCookByTheBook( const FCookByTheBookStartupOptions
 	}
 	
 	// allow the game to fill out the asset registry, as well as get a list of objects to always cook
-	TArray<FString> FilesInPath;
-	FGameDelegates::Get().GetCookModificationDelegate().ExecuteIfBound(FilesInPath);
+	TArray<FString> FilesInPathStrings;
+	FGameDelegates::Get().GetCookModificationDelegate().ExecuteIfBound(FilesInPathStrings);
+
+	TArray<FName> FilesInPath;
+	for (const auto& FileString : FilesInPathStrings)
+	{
+		FilesInPath.Add(FName(*FileString));
+	}
 
 	// don't resave the global shader map files in dlc
 	if ( !IsCookingDLC() )
@@ -3245,7 +3285,7 @@ void UCookOnTheFlyServer::StartCookByTheBook( const FCookByTheBookStartupOptions
 		SaveGlobalShaderMapFiles(TargetPlatforms);
 	}
 	
-
+	
 	CollectFilesToCook(FilesInPath, CookMaps, CookDirectories, CookCultures, IniMapSections, bCookAll, bMapsOnly, bNoDev );
 	if (FilesInPath.Num() == 0)
 	{
@@ -3260,9 +3300,10 @@ void UCookOnTheFlyServer::StartCookByTheBook( const FCookByTheBookStartupOptions
 		GenerateLongPackageNames(FilesInPath);
 	}
 	// add all the files for the requested platform to the cook list
-	for ( const auto& FileName : FilesInPath )
+	for ( const auto& FileFName : FilesInPath )
 	{
-		FName FileFName = FName(*FileName);
+		// FName FileFName = FName(*FileName);
+		FString FileName = FileFName.ToString();
 		FName PackageFileFName = GetCachedStandardPackageFileFName(FileFName);
 		
 		if (PackageFileFName != NAME_None)
diff --git a/Engine/Source/Runtime/AssetRegistry/Private/AssetRegistry.cpp b/Engine/Source/Runtime/AssetRegistry/Private/AssetRegistry.cpp
--- a/Engine/Source/Runtime/AssetRegistry/Private/AssetRegistry.cpp
+++ b/Engine/Source/Runtime/AssetRegistry/Private/AssetRegistry.cpp
@@ -1313,7 +1313,8 @@ void FAssetRegistry::Serialize(FArchive& Ar)
 {
 	if (Ar.IsSaving())
 	{
-		SaveRegistryData(Ar, CachedAssetsByObjectPath, NumAssets);
+		check(CachedAssetsByObjectPath.Num() == NumAssets);
+		SaveRegistryData(Ar, CachedAssetsByObjectPath);
 	}
 	// load in by building the TMap
 	else
@@ -1340,9 +1341,10 @@ void FAssetRegistry::Serialize(FArchive& Ar)
 	}
 }
 
-void FAssetRegistry::SaveRegistryData(FArchive& Ar, TMap<FName, FAssetData*>& Data, int32 AssetCount)
+void FAssetRegistry::SaveRegistryData(FArchive& Ar, TMap<FName, FAssetData*>& Data )
 {
 	// serialize number of objects
+	int32 AssetCount = Data.Num();
 	Ar << AssetCount;
 
 	// save out by walking the TMap
@@ -1352,6 +1354,25 @@ void FAssetRegistry::SaveRegistryData(FArchive& Ar, TMap<FName, FAssetData*>& Da
 	}
 }
 
+void FAssetRegistry::LoadRegistryData(FArchive& Ar, TMap<FName, FAssetData*>& Data )
+{
+	check(Ar.IsLoading());
+	// serialize number of objects
+	int AssetCount = 0;
+	Ar << AssetCount;
+
+	for (int32 AssetIndex = 0; AssetIndex < AssetCount; AssetIndex++)
+	{
+		// make a new asset data object
+		FAssetData *NewAssetData = new FAssetData();
+
+		// load it
+		Ar << *NewAssetData;
+
+		Data.Add(NewAssetData->ObjectPath, NewAssetData);
+	}
+}
+
 void FAssetRegistry::ScanPathsSynchronous_Internal(const TArray<FString>& InPaths, bool bForceRescan, bool bUseCache)
 {
 	const double SearchStartTime = FPlatformTime::Seconds();
diff --git a/Engine/Source/Runtime/AssetRegistry/Private/AssetRegistry.h b/Engine/Source/Runtime/AssetRegistry/Private/AssetRegistry.h
--- a/Engine/Source/Runtime/AssetRegistry/Private/AssetRegistry.h
+++ b/Engine/Source/Runtime/AssetRegistry/Private/AssetRegistry.h
@@ -38,7 +38,8 @@ public:
 	virtual void ScanPathsSynchronous(const TArray<FString>& InPaths, bool bForceRescan = false) override;
 	virtual void PrioritizeSearchPath(const FString& PathToPrioritize) override;
 	virtual void Serialize(FArchive& Ar) override;
-	virtual void SaveRegistryData(FArchive& Ar, TMap<FName, FAssetData*>& Data, int32 AssetCount) override;
+	virtual void SaveRegistryData(FArchive& Ar, TMap<FName, FAssetData*>& Data) override;
+	virtual void LoadRegistryData(FArchive& Ar, TMap<FName, FAssetData*>& Data) override;
 
 	DECLARE_DERIVED_EVENT( FAssetRegistry, IAssetRegistry::FPathAddedEvent, FPathAddedEvent);
 	virtual FPathAddedEvent& OnPathAdded() override { return PathAddedEvent; }
diff --git a/Engine/Source/Runtime/AssetRegistry/Public/IAssetRegistry.h b/Engine/Source/Runtime/AssetRegistry/Public/IAssetRegistry.h
--- a/Engine/Source/Runtime/AssetRegistry/Public/IAssetRegistry.h
+++ b/Engine/Source/Runtime/AssetRegistry/Public/IAssetRegistry.h
@@ -222,6 +222,10 @@ public:
 	/** Serialize the registry to/from a file, skipping editor only data */
 	virtual void Serialize(FArchive& Ar) = 0;
 
-	/** Serialize raw registry data to/from a file, skipping editor only data */
-	virtual void SaveRegistryData(FArchive& Ar, TMap<FName, FAssetData*>& Data, int32 AssetCount) = 0;
+	/** Serialize raw registry data to a file, skipping editor only data */
+	virtual void SaveRegistryData(FArchive& Ar, TMap<FName, FAssetData*>& Data) = 0;
+
+
+	/** Serialize registry data from a file */
+	virtual void LoadRegistryData(FArchive& Ar, TMap<FName, FAssetData*>& Data) = 0;
 };
