Fix curve editor Time and Value boxes always showing "Time" for the label

diff --git a/Engine/Source/Editor/UnrealEd/Private/SCurveEditor.cpp b/Engine/Source/Editor/UnrealEd/Private/SCurveEditor.cpp
--- a/Engine/Source/Editor/UnrealEd/Private/SCurveEditor.cpp
+++ b/Engine/Source/Editor/UnrealEd/Private/SCurveEditor.cpp
@@ -143,7 +143,7 @@ void SCurveEditor::Construct(const FArguments& InArgs)
 	CurveSelectionWidget = CurveSelector;
 
 	InputAxisName = InArgs._XAxisName.IsSet() ? FText::FromString(InArgs._XAxisName.GetValue()) : LOCTEXT("Time", "Time");
-	InputAxisName = InArgs._YAxisName.IsSet() ? FText::FromString(InArgs._YAxisName.GetValue()) : LOCTEXT("Value", "Value");
+	OutputAxisName = InArgs._YAxisName.IsSet() ? FText::FromString(InArgs._YAxisName.GetValue()) : LOCTEXT("Value", "Value");
 
 	this->ChildSlot
 	[
@@ -269,7 +269,7 @@ void SCurveEditor::Construct(const FArguments& InArgs)
 							.Label()
 							[
 								SNew(STextBlock)
-								.Text(InputAxisName)
+								.Text(OutputAxisName)
 							]
 						]
 					]
