#HLOD

- showing the mesh again since it is good for debugging mesh
- fixed overflow issue when bound is huge
- added level name for clarification of building process

diff --git a/Engine/Source/Editor/UnrealEd/Private/HierarchicalLOD.cpp b/Engine/Source/Editor/UnrealEd/Private/HierarchicalLOD.cpp
--- a/Engine/Source/Editor/UnrealEd/Private/HierarchicalLOD.cpp
+++ b/Engine/Source/Editor/UnrealEd/Private/HierarchicalLOD.cpp
@@ -22,6 +22,7 @@
 DEFINE_LOG_CATEGORY_STATIC(LogLODGenerator, Warning, All);
 
 #define LOCTEXT_NAMESPACE "HierarchicalLOD"
+#define CM_TO_METER		0.01f
 
 /*-----------------------------------------------------------------------------
 	HierarchicalLOD implementation.
@@ -178,7 +179,9 @@ void FHierarchicalLODBuilder::BuildClusters(class ULevel* InLevel)
 		for(int32 LODId=0; LODId<TotalNumLOD; ++LODId)
 		{
 			AWorldSettings* WorldSetting = InLevel->GetWorld()->GetWorldSettings();
-			float DesiredBoundBoxHalfSize = WorldSetting->HierarchicalLODSetup[LODId].DesiredBoundSize * 0.5f;
+			// we use meter for bound. Otherwise it's very easy to get to overflow and have problem with filling ratio because
+			// bound is too huge
+			float DesiredBoundBoxHalfSize = WorldSetting->HierarchicalLODSetup[LODId].DesiredBoundSize * CM_TO_METER * 0.5f;
 			float DesiredFillingRatio = WorldSetting->HierarchicalLODSetup[LODId].DesiredFillingPercentage * 0.01f;
 			ensure(DesiredFillingRatio!=0.f);
 			float HighestCost = FMath::Pow(DesiredBoundBoxHalfSize, 3)/(DesiredFillingRatio);
@@ -192,10 +195,12 @@ void FHierarchicalLODBuilder::BuildClusters(class ULevel* InLevel)
 
 			// since to show progress of initialization, I'm scoping it
 			{
+				FString LevelName = FPackageName::GetShortName(InLevel->GetOutermost()->GetName());
 				FFormatNamedArguments Arguments;
 				Arguments.Add(TEXT("LODIndex"), FText::AsNumber(LODId+1));
+				Arguments.Add(TEXT("LevelName"), FText::FromString(LevelName));
 
-				FScopedSlowTask SlowTask(100, FText::Format(LOCTEXT("HierarchicalLOD_InitializeCluster", "Initializing Clusters for LOD {LODIndex}..."), Arguments));
+				FScopedSlowTask SlowTask(100, FText::Format(LOCTEXT("HierarchicalLOD_InitializeCluster", "Initializing Clusters for LOD {LODIndex} of {LevelName}..."), Arguments));
 				SlowTask.MakeDialog();
 
 				// initialize Clusters
@@ -238,15 +243,16 @@ void FHierarchicalLODBuilder::MergeClustersAndBuildActors(class ULevel* InLevel,
 {
 	if (Clusters.Num() > 0)
 	{
+		FString LevelName = FPackageName::GetShortName(InLevel->GetOutermost()->GetName());
 		FFormatNamedArguments Arguments;
 		Arguments.Add(TEXT("LODIndex"), FText::AsNumber(LODIdx+1));
-
+		Arguments.Add(TEXT("LevelName"), FText::FromString(LevelName));
 		// merge clusters first
 		{
 			static int32 TotalIteration=3;
 			const int32 TotalCluster = Clusters.Num();
 
-			FScopedSlowTask SlowTask(TotalIteration*TotalCluster, FText::Format( LOCTEXT("HierarchicalLOD_BuildClusters", "Building Clusters for LOD {LODIndex}..."), Arguments) );
+			FScopedSlowTask SlowTask(TotalIteration*TotalCluster, FText::Format( LOCTEXT("HierarchicalLOD_BuildClusters", "Building Clusters for LOD {LODIndex} of {LevelName}..."), Arguments) );
 			SlowTask.MakeDialog();
 
 			for(int32 Iteration=0; Iteration<TotalIteration; ++Iteration)
@@ -315,7 +321,7 @@ void FHierarchicalLODBuilder::MergeClustersAndBuildActors(class ULevel* InLevel,
 				}
 			}
 
-			FScopedSlowTask SlowTask(TotalValidCluster, FText::Format( LOCTEXT("HierarchicalLOD_MergeActors", "Merging Actors for LOD {LODIndex}..."), Arguments) );
+			FScopedSlowTask SlowTask(TotalValidCluster, FText::Format( LOCTEXT("HierarchicalLOD_MergeActors", "Merging Actors for LOD {LODIndex} of {LevelName}..."), Arguments) );
 			SlowTask.MakeDialog();
 
 			for(auto& Cluster: Clusters)
@@ -391,9 +397,10 @@ void FHierarchicalLODBuilder::InitializeClusters(class ULevel* InLevel, const in
 						if(ShouldGenerateCluster(Actor2))
 						{
 							FLODCluster NewClusterCandidate = FLODCluster(Actor1, Actor2);
+							float NewClusterCost = NewClusterCandidate.GetCost();
 							// @todo optimization : filter first - this might create disconnected tree - 
 							//@Todo debug but it's too slow to have full tree 
-							if (NewClusterCandidate.GetCost() <= CullCost)
+							if ( NewClusterCost <= CullCost)
 							{
 								Clusters.Add(NewClusterCandidate);
 							}
@@ -479,7 +486,8 @@ FSphere FLODCluster::AddActor(class AActor* NewActor)
 	FVector Origin, Extent;
 	NewActor->GetActorBounds(false, Origin, Extent);
 	
-	FSphere NewBound = FSphere(Origin, Extent.Size());
+	// scale 0.01 (change to meter from centermeter)
+	FSphere NewBound = FSphere(Origin*CM_TO_METER, Extent.Size()*CM_TO_METER);
 	Bound += NewBound;
 
 	return NewBound;
@@ -634,7 +642,7 @@ void FLODCluster::BuildActor(class ULevel* InLevel, const int32 LODIdx)
 					MeshUtilities.MergeActors(Actors, MergeSetting, AssetsOuter, PackageName, LODIdx+1, OutAssets, OutProxyLocation, true );
 				}
 
-				// this is ideally the way we should do, but 
+				// we make it private, so it can't be used by outside of map since it's useless, and then remove standalone
 				for (auto& AssetIter : OutAssets)
 				{
 					AssetIter->ClearFlags(RF_Public | RF_Standalone);
diff --git a/Engine/Source/Runtime/Engine/Classes/Engine/LODActor.h b/Engine/Source/Runtime/Engine/Classes/Engine/LODActor.h
--- a/Engine/Source/Runtime/Engine/Classes/Engine/LODActor.h
+++ b/Engine/Source/Runtime/Engine/Classes/Engine/LODActor.h
@@ -17,7 +17,7 @@ class ENGINE_API ALODActor : public AActor
 
 private_subobject:
 	// disable display of this component
-	UPROPERTY(/*Category=LODActor, VisibleAnywhere*/)
+	UPROPERTY(Category=LODActor, VisibleAnywhere)
 	class UStaticMeshComponent* StaticMeshComponent;
 
 public:
