Many fixes related to EventDispatcher:
UE-13190 Cannot Delete Event Dispatcher inside of a child blueprint when the parent has an event dispatcher as well
Parameter "const FName&" was changed into "const FName" in some functions in BlueprintEditorUtils.
Fix for TRASHCLASS_ crash in UBlueprintEventNodeSpawner.
In ResolveSimpleMemberReference the GUID is primary data, the Name is secondary.

diff --git a/Engine/Source/Editor/BlueprintGraph/Private/BlueprintEventNodeSpawner.cpp b/Engine/Source/Editor/BlueprintGraph/Private/BlueprintEventNodeSpawner.cpp
--- a/Engine/Source/Editor/BlueprintGraph/Private/BlueprintEventNodeSpawner.cpp
+++ b/Engine/Source/Editor/BlueprintGraph/Private/BlueprintEventNodeSpawner.cpp
@@ -3,6 +3,7 @@
 #include "BlueprintGraphPrivatePCH.h"
 #include "BlueprintEventNodeSpawner.h"
 #include "EdGraphSchema_K2.h" // for GetFriendlySignatureName()
+#include "BlueprintNodeTemplateCache.h" // for IsTemplateOuter()
 
 #define LOCTEXT_NAMESPACE "BlueprintEventNodeSpawner"
 
@@ -169,11 +170,17 @@ UEdGraphNode* UBlueprintEventNodeSpawner::Invoke(UEdGraph* ParentGraph, FBinding
 {
 	check(ParentGraph != nullptr);
 	UBlueprint* Blueprint = FBlueprintEditorUtils::FindBlueprintForGraphChecked(ParentGraph);
-	// look to see if a node for this event already exists (only one node is
-	// allowed per event, per blueprint)
-	UK2Node_Event const* PreExistingNode = FindPreExistingEvent(Blueprint, Bindings);
-	// @TODO: casting away the const is bad form!
-	UK2Node_Event* EventNode = const_cast<UK2Node_Event*>(PreExistingNode);
+
+	UK2Node_Event* EventNode = nullptr;
+	bool const bIsTemplateNode = FBlueprintNodeTemplateCache::IsTemplateOuter(ParentGraph);
+	if (!bIsTemplateNode)
+	{
+		// look to see if a node for this event already exists (only one node is
+		// allowed per event, per blueprint)
+		UK2Node_Event const* PreExistingNode = FindPreExistingEvent(Blueprint, Bindings);
+		// @TODO: casting away the const is bad form!
+		EventNode = const_cast<UK2Node_Event*>(PreExistingNode);
+	}
 
 	bool const bIsCustomEvent = IsForCustomEvent();
 	check(bIsCustomEvent || (EventFunc != nullptr));
diff --git a/Engine/Source/Editor/BlueprintGraph/Private/K2Node_MCDelegate.cpp b/Engine/Source/Editor/BlueprintGraph/Private/K2Node_MCDelegate.cpp
--- a/Engine/Source/Editor/BlueprintGraph/Private/K2Node_MCDelegate.cpp
+++ b/Engine/Source/Editor/BlueprintGraph/Private/K2Node_MCDelegate.cpp
@@ -106,16 +106,10 @@ UFunction* UK2Node_BaseMCDelegate::GetDelegateSignature(bool bForceNotFromSkelCl
 	}
 
 	FMemberReference ReferenceToUse;
-	FGuid DelegateGuid;
+	ReferenceToUse.SetDirect(DelegateReference.GetMemberName(), DelegateReference.GetMemberGuid(), OwnerClass, /*bIsConsideredSelfContext =*/false);
 
-	if (OwnerClass != nullptr)
-	{
-		UBlueprint::GetGuidFromClassByFieldName<UFunction>(OwnerClass, DelegateReference.GetMemberName(), DelegateGuid);
-	}
-	ReferenceToUse.SetDirect(DelegateReference.GetMemberName(), DelegateGuid, OwnerClass, /*bIsConsideredSelfContext =*/false);
-
-	UMulticastDelegateProperty* DelegateProperty = ReferenceToUse.ResolveMember<UMulticastDelegateProperty>(GetBlueprintClassFromNode());
-	return (DelegateProperty != NULL) ? DelegateProperty->SignatureFunction : NULL;
+	UMulticastDelegateProperty* DelegateProperty = ReferenceToUse.ResolveMember<UMulticastDelegateProperty>((UClass*)nullptr);
+	return (DelegateProperty != nullptr) ? DelegateProperty->SignatureFunction : nullptr;
 }
 
 UEdGraphPin* UK2Node_BaseMCDelegate::GetDelegatePin() const
@@ -199,9 +193,9 @@ bool UK2Node_BaseMCDelegate::HasExternalBlueprintDependencies(TArray<class UStru
 	const bool bResult = (SourceClass != NULL) && (SourceClass->ClassGeneratedBy != NULL) && (SourceClass->ClassGeneratedBy != SourceBlueprint);
 	if (bResult && OptionalOutput)
 	{
-		OptionalOutput->Add(GetDelegateSignature());
+		OptionalOutput->AddUnique(GetDelegateSignature());
 	}
-	return bResult || Super::HasExternalBlueprintDependencies(OptionalOutput);
+	return Super::HasExternalBlueprintDependencies(OptionalOutput) || bResult;
 }
 
 void UK2Node_BaseMCDelegate::GetNodeAttributes( TArray<TKeyValuePair<FString, FString>>& OutNodeAttributes ) const
diff --git a/Engine/Source/Editor/Kismet/Private/SBlueprintPalette.cpp b/Engine/Source/Editor/Kismet/Private/SBlueprintPalette.cpp
--- a/Engine/Source/Editor/Kismet/Private/SBlueprintPalette.cpp
+++ b/Engine/Source/Editor/Kismet/Private/SBlueprintPalette.cpp
@@ -1207,20 +1207,6 @@ void SBlueprintPaletteItem::OnNameTextCommitted(const FText& NewText, ETextCommi
 
 				UBlueprint* BlueprintObj = BlueprintEditorPtr.Pin()->GetBlueprintObj();
 				FBlueprintEditorUtils::RenameMemberVariable(BlueprintObj, OldName, NewName);
-
-				TArray<UK2Node_BaseMCDelegate*> NodeUsingDelegate;
-				FBlueprintEditorUtils::GetAllNodesOfClass<UK2Node_BaseMCDelegate>(BlueprintObj, NodeUsingDelegate);
-				for (auto FuncIt = NodeUsingDelegate.CreateIterator(); FuncIt; ++FuncIt)
-				{
-					UK2Node_BaseMCDelegate* FunctionNode = *FuncIt;
-					if (FunctionNode->DelegateReference.IsSelfContext() && (FunctionNode->DelegateReference.GetMemberName() == OldName))
-					{
-						FunctionNode->Modify();
-						FunctionNode->DelegateReference.SetSelfMember(NewName);
-					}
-				}
-
-				FBlueprintEditorUtils::RenameGraph(Graph, NewNameString );
 			}
 		}
 	}
diff --git a/Engine/Source/Editor/UnrealEd/Private/Kismet2/BlueprintEditorUtils.cpp b/Engine/Source/Editor/UnrealEd/Private/Kismet2/BlueprintEditorUtils.cpp
--- a/Engine/Source/Editor/UnrealEd/Private/Kismet2/BlueprintEditorUtils.cpp
+++ b/Engine/Source/Editor/UnrealEd/Private/Kismet2/BlueprintEditorUtils.cpp
@@ -522,7 +522,8 @@ void FBlueprintEditorUtils::RefreshExternalBlueprintDependencyNodes(UBlueprint*
 					if (!bShouldRefresh)
 					{
 						UClass* OwnerClass = Struct->GetOwnerClass();
-						bShouldRefresh |= OwnerClass && OwnerClass->IsChildOf(RefreshOnlyChild);
+						bShouldRefresh |= OwnerClass && 
+							(OwnerClass->IsChildOf(RefreshOnlyChild) || OwnerClass->GetAuthoritativeClass()->IsChildOf(RefreshOnlyChild));
 					}
 					if (bShouldRefresh)
 					{
@@ -3887,7 +3888,7 @@ bool FBlueprintEditorUtils::AddMemberVariable(UBlueprint* Blueprint, const FName
 }
 
 // Removes a member variable if it was declared in this blueprint and not in a base class.
-void FBlueprintEditorUtils::RemoveMemberVariable(UBlueprint* Blueprint, const FName& VarName)
+void FBlueprintEditorUtils::RemoveMemberVariable(UBlueprint* Blueprint, const FName VarName)
 {
 	const int32 VarIndex = FBlueprintEditorUtils::FindNewVariableIndex(Blueprint, VarName);
 	if (VarIndex != INDEX_NONE)
@@ -3952,7 +3953,7 @@ FName FBlueprintEditorUtils::FindMemberVariableNameByGuid(UBlueprint* InBlueprin
 	return NAME_None;
 }
 
-void FBlueprintEditorUtils::RemoveVariableNodes(UBlueprint* Blueprint, const FName& VarName, bool const bForSelfOnly, UEdGraph* LocalGraphScope)
+void FBlueprintEditorUtils::RemoveVariableNodes(UBlueprint* Blueprint, const FName VarName, bool const bForSelfOnly, UEdGraph* LocalGraphScope)
 {
 	TArray<UEdGraph*> AllGraphs;
 	Blueprint->GetAllGraphs(AllGraphs);
@@ -3985,7 +3986,7 @@ void FBlueprintEditorUtils::RemoveVariableNodes(UBlueprint* Blueprint, const FNa
 	}
 }
 
-void FBlueprintEditorUtils::RenameComponentMemberVariable(UBlueprint* Blueprint, USCS_Node* Node, const FName& NewName)
+void FBlueprintEditorUtils::RenameComponentMemberVariable(UBlueprint* Blueprint, USCS_Node* Node, const FName NewName)
 {
 	// Should not allow renaming to "none" (UI should prevent this)
 	check(!NewName.IsNone());
@@ -4032,7 +4033,7 @@ void FBlueprintEditorUtils::RenameComponentMemberVariable(UBlueprint* Blueprint,
 	}
 }
 
-void FBlueprintEditorUtils::RenameMemberVariable(UBlueprint* Blueprint, const FName& OldName, const FName& NewName)
+void FBlueprintEditorUtils::RenameMemberVariable(UBlueprint* Blueprint, const FName OldName, const FName NewName)
 {
 	if (!NewName.IsNone() && !NewName.IsEqual(OldName, ENameCase::CaseSensitive))
 	{
@@ -4093,6 +4094,29 @@ void FBlueprintEditorUtils::RenameMemberVariable(UBlueprint* Blueprint, const FN
 					NewTargetProperty->CopyCompleteValue(NewPropertyAddr, OldPropertyAddr);
 				}
 			}
+
+			{
+				UClass* NewSkelGeneratedClass = Blueprint->SkeletonGeneratedClass;
+				UMulticastDelegateProperty* MCProperty = NewSkelGeneratedClass ? FindField<UMulticastDelegateProperty>(NewSkelGeneratedClass, NewName) : nullptr;
+				UEdGraph* DelegateSignatureGraph = MCProperty ? FindObject<UEdGraph>(Blueprint, *OldName.ToString()) : nullptr;
+				if (MCProperty && DelegateSignatureGraph)
+				{
+					FBlueprintEditorUtils::RenameGraph(DelegateSignatureGraph, NewName.ToString());
+
+					// this code should not be necessary, because the GUID remains valid, but let it be for backward compatibility.
+					TArray<UK2Node_BaseMCDelegate*> NodeUsingDelegate;
+					FBlueprintEditorUtils::GetAllNodesOfClass<UK2Node_BaseMCDelegate>(Blueprint, NodeUsingDelegate);
+					for (auto FuncIt = NodeUsingDelegate.CreateIterator(); FuncIt; ++FuncIt)
+					{
+						UK2Node_BaseMCDelegate* FunctionNode = *FuncIt;
+						if (FunctionNode->DelegateReference.IsSelfContext() && (FunctionNode->DelegateReference.GetMemberName() == OldName))
+						{
+							FunctionNode->Modify();
+							FunctionNode->DelegateReference.SetSelfMember(NewName);
+						}
+					}
+				}
+			}
 		}
 		else if (Blueprint && Blueprint->SimpleConstructionScript)
 		{
@@ -4187,7 +4211,7 @@ void FBlueprintEditorUtils::GetLoadedChildBlueprints(UBlueprint* InBlueprint, TA
 	}
 }
 
-void FBlueprintEditorUtils::ChangeMemberVariableType(UBlueprint* Blueprint, const FName& VariableName, const FEdGraphPinType& NewPinType)
+void FBlueprintEditorUtils::ChangeMemberVariableType(UBlueprint* Blueprint, const FName VariableName, const FEdGraphPinType& NewPinType)
 {
 	if (VariableName != NAME_None)
 	{
@@ -4404,7 +4428,7 @@ FBPVariableDescription FBlueprintEditorUtils::DuplicateVariableDescription(UBlue
 	return NewVar;
 }
 
-bool FBlueprintEditorUtils::AddLocalVariable(UBlueprint* Blueprint, UEdGraph* InTargetGraph, const FName& InNewVarName, const FEdGraphPinType& InNewVarType)
+bool FBlueprintEditorUtils::AddLocalVariable(UBlueprint* Blueprint, UEdGraph* InTargetGraph, const FName InNewVarName, const FEdGraphPinType& InNewVarType)
 {
 	if(InTargetGraph != NULL && InTargetGraph->GetSchema()->GetGraphType(InTargetGraph) == GT_Function)
 	{
@@ -4438,7 +4462,7 @@ bool FBlueprintEditorUtils::AddLocalVariable(UBlueprint* Blueprint, UEdGraph* In
 	return false;
 }
 
-void FBlueprintEditorUtils::RemoveLocalVariable(UBlueprint* InBlueprint, const UStruct* InScope, const FName& InVarName)
+void FBlueprintEditorUtils::RemoveLocalVariable(UBlueprint* InBlueprint, const UStruct* InScope, const FName InVarName)
 {
 	UEdGraph* ScopeGraph = FindScopeGraph(InBlueprint, InScope);
 
@@ -4533,7 +4557,7 @@ UEdGraph* FBlueprintEditorUtils::FindScopeGraph(const UBlueprint* InBlueprint, c
 	return ScopeGraph;
 }
 
-void FBlueprintEditorUtils::RenameLocalVariable(UBlueprint* InBlueprint, const UStruct* InScope, const FName& InOldName, const FName& InNewName)
+void FBlueprintEditorUtils::RenameLocalVariable(UBlueprint* InBlueprint, const UStruct* InScope, const FName InOldName, const FName InNewName)
 {
 	if (!InNewName.IsNone() && !InNewName.IsEqual(InOldName, ENameCase::CaseSensitive))
 	{
@@ -4570,13 +4594,13 @@ void FBlueprintEditorUtils::RenameLocalVariable(UBlueprint* InBlueprint, const U
 	}
 }
 
-FBPVariableDescription* FBlueprintEditorUtils::FindLocalVariable(UBlueprint* InBlueprint, const UStruct* InScope, const FName& InVariableName)
+FBPVariableDescription* FBlueprintEditorUtils::FindLocalVariable(UBlueprint* InBlueprint, const UStruct* InScope, const FName InVariableName)
 {
 	UK2Node_FunctionEntry* DummyFunctionEntry = nullptr;
 	return FindLocalVariable(InBlueprint, InScope, InVariableName, &DummyFunctionEntry);
 }
 
-FBPVariableDescription* FBlueprintEditorUtils::FindLocalVariable(const UBlueprint* InBlueprint, const UEdGraph* InScopeGraph, const FName& InVariableName, class UK2Node_FunctionEntry** OutFunctionEntry)
+FBPVariableDescription* FBlueprintEditorUtils::FindLocalVariable(const UBlueprint* InBlueprint, const UEdGraph* InScopeGraph, const FName InVariableName, class UK2Node_FunctionEntry** OutFunctionEntry)
 {
 	FBPVariableDescription* ReturnVariable = NULL;
 	if(DoesSupportLocalVariables(InScopeGraph))
@@ -4611,7 +4635,7 @@ FBPVariableDescription* FBlueprintEditorUtils::FindLocalVariable(const UBlueprin
 	return ReturnVariable;
 }
 
-FBPVariableDescription* FBlueprintEditorUtils::FindLocalVariable(const UBlueprint* InBlueprint, const UStruct* InScope, const FName& InVariableName, class UK2Node_FunctionEntry** OutFunctionEntry)
+FBPVariableDescription* FBlueprintEditorUtils::FindLocalVariable(const UBlueprint* InBlueprint, const UStruct* InScope, const FName InVariableName, class UK2Node_FunctionEntry** OutFunctionEntry)
 {
 	UEdGraph* ScopeGraph = FindScopeGraph(InBlueprint, InScope);
 
@@ -4667,7 +4691,7 @@ FGuid FBlueprintEditorUtils::FindLocalVariableGuidByName(UBlueprint* InBlueprint
 	return ReturnGuid;
 }
 
-void FBlueprintEditorUtils::ChangeLocalVariableType(UBlueprint* InBlueprint, const UStruct* InScope, const FName& InVariableName, const FEdGraphPinType& NewPinType)
+void FBlueprintEditorUtils::ChangeLocalVariableType(UBlueprint* InBlueprint, const UStruct* InScope, const FName InVariableName, const FEdGraphPinType& NewPinType)
 {
 	if (InVariableName != NAME_None)
 	{
@@ -4750,7 +4774,7 @@ void FBlueprintEditorUtils::ChangeLocalVariableType(UBlueprint* InBlueprint, con
 	}
 }
 
-void FBlueprintEditorUtils::ReplaceVariableReferences(UBlueprint* Blueprint, const FName& OldName, const FName& NewName)
+void FBlueprintEditorUtils::ReplaceVariableReferences(UBlueprint* Blueprint, const FName OldName, const FName NewName)
 {
 	check((OldName != NAME_None) && (NewName != NAME_None));
 
@@ -4810,15 +4834,15 @@ bool FBlueprintEditorUtils::IsVariableUsed(const UBlueprint* Blueprint, const FN
 	return false;
 }
 
-bool FBlueprintEditorUtils::ValidateAllMemberVariables(UBlueprint* InBlueprint, UBlueprint* InParentBlueprint, const FName& InVariableName)
+bool FBlueprintEditorUtils::ValidateAllMemberVariables(UBlueprint* InBlueprint, UBlueprint* InParentBlueprint, const FName InVariableName)
 {
 	for(int32 VariableIdx = 0; VariableIdx < InBlueprint->NewVariables.Num(); ++VariableIdx)
 	{
 		if(InBlueprint->NewVariables[VariableIdx].VarName == InVariableName)
 		{
-			FName NewChildName = FBlueprintEditorUtils::FindUniqueKismetName(InBlueprint, InVariableName.ToString());
+			FName NewChildName = FBlueprintEditorUtils::FindUniqueKismetName(InBlueprint, InVariableName.ToString(), InParentBlueprint ? InParentBlueprint->SkeletonGeneratedClass : InBlueprint->ParentClass);
 
-			UE_LOG(LogBlueprint, Warning, TEXT("Blueprint %s (child of/implements %s) has a member variable with a conflicting name (%s). Changing to %s."), *InBlueprint->GetName(), *InParentBlueprint->GetName(), *InVariableName.ToString(), *NewChildName.ToString());
+			UE_LOG(LogBlueprint, Warning, TEXT("Blueprint %s (child of/implements %s) has a member variable with a conflicting name (%s). Changing to %s."), *InBlueprint->GetName(), *GetNameSafe(InParentBlueprint), *InVariableName.ToString(), *NewChildName.ToString());
 
 			FBlueprintEditorUtils::RenameMemberVariable(InBlueprint, InBlueprint->NewVariables[VariableIdx].VarName, NewChildName);
 			return true;
@@ -4888,7 +4912,7 @@ bool FBlueprintEditorUtils::ValidateAllFunctionGraphs(UBlueprint* InBlueprint, U
 	return false;
 }
 
-void FBlueprintEditorUtils::ValidateBlueprintChildVariables(UBlueprint* InBlueprint, const FName& InVariableName)
+void FBlueprintEditorUtils::ValidateBlueprintChildVariables(UBlueprint* InBlueprint, const FName InVariableName)
 {
 	// Iterate over currently-loaded Blueprints and potentially adjust their variable names if they conflict with the parent
 	for(TObjectIterator<UBlueprint> BlueprintIt; BlueprintIt; ++BlueprintIt)
@@ -6082,7 +6106,7 @@ UK2Node_Timeline* FBlueprintEditorUtils::FindNodeForTimeline(UBlueprint* Bluepri
 	return NULL; // no node found
 }
 
-bool FBlueprintEditorUtils::RenameTimeline(UBlueprint* Blueprint, const FName& OldNameRef, const FName& NewName)
+bool FBlueprintEditorUtils::RenameTimeline(UBlueprint* Blueprint, const FName OldNameRef, const FName NewName)
 {
 	check(Blueprint);
 
@@ -7424,15 +7448,23 @@ FName FBlueprintEditorUtils::GetFunctionNameFromClassByGuid(const UClass* InClas
 	for (int32 BPIndex = 0; BPIndex < Blueprints.Num(); ++BPIndex)
 	{
 		UBlueprint* Blueprint = Blueprints[BPIndex];
-		for (int32 FunctionIndex = 0; FunctionIndex < Blueprint->FunctionGraphs.Num(); ++FunctionIndex)
+		for (auto FunctionGraph : Blueprint->FunctionGraphs)
 		{
-			UEdGraph* FunctionGraph = Blueprint->FunctionGraphs[FunctionIndex];
 			if (FunctionGraph && FunctionGraph->GraphGuid == FunctionGuid)
 			{
 				return FunctionGraph->GetFName();
 			}
 		}
 
+		for (auto FunctionGraph : Blueprint->DelegateSignatureGraphs)
+		{
+			if (FunctionGraph && FunctionGraph->GraphGuid == FunctionGuid)
+			{
+				const FString Name = FunctionGraph->GetName() + HEADER_GENERATED_DELEGATE_SIGNATURE_SUFFIX;
+				return FName(*Name);
+			}
+		}
+
 		//FUNCTIONS BASED ON CUSTOM EVENTS:
 		TArray<UEdGraph*> UberGraphs;
 		FUberGraphHelper::GetAll(Blueprint, UberGraphs);
@@ -7464,9 +7496,8 @@ bool FBlueprintEditorUtils::GetFunctionGuidFromClassByFieldName(const UClass* In
 		for (int32 BPIndex = 0; BPIndex < Blueprints.Num(); ++BPIndex)
 		{
 			UBlueprint* Blueprint = Blueprints[BPIndex];
-			for (int32 FunctionIndex = 0; FunctionIndex < Blueprint->FunctionGraphs.Num(); ++FunctionIndex)
+			for (auto FunctionGraph : Blueprint->FunctionGraphs)
 			{
-				UEdGraph* FunctionGraph = Blueprint->FunctionGraphs[FunctionIndex];
 				if (FunctionGraph && FunctionGraph->GetFName() == FunctionName)
 				{
 					FunctionGuid = FunctionGraph->GraphGuid;
@@ -7474,6 +7505,20 @@ bool FBlueprintEditorUtils::GetFunctionGuidFromClassByFieldName(const UClass* In
 				}
 			}
 
+			FString BaseDelegateSignatureName = FunctionName.ToString();
+			if (BaseDelegateSignatureName.RemoveFromEnd(HEADER_GENERATED_DELEGATE_SIGNATURE_SUFFIX))
+			{
+				const FName GraphName(*BaseDelegateSignatureName);
+				for (auto FunctionGraph : Blueprint->DelegateSignatureGraphs)
+				{
+					if (FunctionGraph && FunctionGraph->GetFName() == GraphName)
+					{
+						FunctionGuid = FunctionGraph->GraphGuid;
+						return true;
+					}
+				}
+			}
+
 			TArray<UEdGraph*> UberGraphs;
 			FUberGraphHelper::GetAll(Blueprint, UberGraphs);
 			for (const auto UberGraph : UberGraphs)
diff --git a/Engine/Source/Editor/UnrealEd/Public/Kismet2/BlueprintEditorUtils.h b/Engine/Source/Editor/UnrealEd/Public/Kismet2/BlueprintEditorUtils.h
--- a/Engine/Source/Editor/UnrealEd/Public/Kismet2/BlueprintEditorUtils.h
+++ b/Engine/Source/Editor/UnrealEd/Public/Kismet2/BlueprintEditorUtils.h
@@ -641,7 +641,7 @@ public:
 	 *
 	 * @param	VarName	Name of the variable to be removed.
 	 */
-	static void RemoveMemberVariable(UBlueprint* Blueprint, const FName& VarName);
+	static void RemoveMemberVariable(UBlueprint* Blueprint, const FName VarName);
 	
 	/**
 	 * Removes member variables if they were declared in this blueprint and not in a base class.
@@ -677,16 +677,16 @@ public:
 	 * @param	LocalGraphScope		Local scope graph of variables
 	 *								false if you just want everything with the specified name removed (variables from other classes too).
 	 */
-	static void RemoveVariableNodes(UBlueprint* Blueprint, const FName& VarName, bool const bForSelfOnly = true, UEdGraph* LocalGraphScope = nullptr);
+	static void RemoveVariableNodes(UBlueprint* Blueprint, const FName VarName, bool const bForSelfOnly = true, UEdGraph* LocalGraphScope = nullptr);
 
 	/**Rename a member variable*/
-	static void RenameMemberVariable(UBlueprint* Blueprint, const FName& OldName, const FName& NewName);
+	static void RenameMemberVariable(UBlueprint* Blueprint, const FName OldName, const FName NewName);
 
 	/** Rename a member variable created by a SCS entry */
-	static void RenameComponentMemberVariable(UBlueprint* Blueprint, USCS_Node* Node, const FName& NewName);
+	static void RenameComponentMemberVariable(UBlueprint* Blueprint, USCS_Node* Node, const FName NewName);
 	
 	/** Changes the type of a member variable */
-	static void ChangeMemberVariableType(UBlueprint* Blueprint, const FName& VariableName, const FEdGraphPinType& NewPinType);
+	static void ChangeMemberVariableType(UBlueprint* Blueprint, const FName VariableName, const FEdGraphPinType& NewPinType);
 
 	/**
 	 * Finds the scope's associated graph for local variables (or any passed UFunction)
@@ -705,7 +705,7 @@ public:
 	 *
 	 * @return	true if it succeeds, false if it fails.
 	 */
-	static bool AddLocalVariable(UBlueprint* Blueprint, UEdGraph* InTargetGraph, const FName& InNewVarName, const FEdGraphPinType& InNewVarType);
+	static bool AddLocalVariable(UBlueprint* Blueprint, UEdGraph* InTargetGraph, const FName InNewVarName, const FEdGraphPinType& InNewVarType);
 
 	/**
 	 * Removes a member variable if it was declared in this blueprint and not in a base class.
@@ -714,7 +714,7 @@ public:
 	 * @param	InScope				Local variable's scope
 	 * @param	InVarName			Name of the variable to be removed.
 	 */
-	static void RemoveLocalVariable(UBlueprint* InBlueprint, const UStruct* InScope, const FName& InVarName);
+	static void RemoveLocalVariable(UBlueprint* InBlueprint, const UStruct* InScope, const FName InVarName);
 
 	/**
 	 * Returns a local variable with the function entry it was found in
@@ -723,7 +723,7 @@ public:
 	 * @param InVariableName	Name of the variable to search for
 	 * @return					The local variable description
 	 */
-	static FBPVariableDescription* FindLocalVariable(UBlueprint* InBlueprint, const UStruct* InScope, const FName& InVariableName);
+	static FBPVariableDescription* FindLocalVariable(UBlueprint* InBlueprint, const UStruct* InScope, const FName InVariableName);
 
 	/**
 	 * Returns a local variable
@@ -734,7 +734,7 @@ public:
 	 * @param OutFunctionEntry	Optional output parameter. If not null, the found function entry is returned.
 	 * @return					The local variable description
 	 */
-	static FBPVariableDescription* FindLocalVariable(const UBlueprint* InBlueprint, const UEdGraph* InScopeGraph, const FName& InVariableName, class UK2Node_FunctionEntry** OutFunctionEntry = NULL);
+	static FBPVariableDescription* FindLocalVariable(const UBlueprint* InBlueprint, const UEdGraph* InScopeGraph, const FName InVariableName, class UK2Node_FunctionEntry** OutFunctionEntry = NULL);
 
 	/**
 	 * Returns a local variable
@@ -745,7 +745,7 @@ public:
 	 * @param OutFunctionEntry	Optional output parameter. If not null, the found function entry is returned.
 	 * @return					The local variable description
 	 */
-	static FBPVariableDescription* FindLocalVariable(const UBlueprint* InBlueprint, const UStruct* InScope, const FName& InVariableName, class UK2Node_FunctionEntry** OutFunctionEntry = NULL);
+	static FBPVariableDescription* FindLocalVariable(const UBlueprint* InBlueprint, const UStruct* InScope, const FName InVariableName, class UK2Node_FunctionEntry** OutFunctionEntry = NULL);
 
 	/**
 	 * Finds a local variable name using the variable's Guid
@@ -784,7 +784,7 @@ public:
 	 * @param InOldName			The name of the local variable to change
 	 * @param InNewName			The new name of the local variable
 	 */
-	static void RenameLocalVariable(UBlueprint* InBlueprint, const UStruct* InScope, const FName& InOldName, const FName& InNewName);
+	static void RenameLocalVariable(UBlueprint* InBlueprint, const UStruct* InScope, const FName InOldName, const FName InNewName);
 
 	/**
 	 * Changes the type of a local variable
@@ -794,19 +794,19 @@ public:
 	 * @param InVariableName	Name of the local variable to change the type of
 	 * @param InNewPinType		The pin type to change the local variable type to
 	 */
-	static void ChangeLocalVariableType(UBlueprint* InBlueprint, const UStruct* InScope, const FName& InVariableName, const FEdGraphPinType& InNewPinType);
+	static void ChangeLocalVariableType(UBlueprint* InBlueprint, const UStruct* InScope, const FName InVariableName, const FEdGraphPinType& InNewPinType);
 
 	/** Replaces all variable references in the specified blueprint */
-	static void ReplaceVariableReferences(UBlueprint* Blueprint, const FName& OldName, const FName& NewName);
+	static void ReplaceVariableReferences(UBlueprint* Blueprint, const FName OldName, const FName NewName);
 
 	/** Replaces all variable references in the specified blueprint */
 	static void ReplaceVariableReferences(UBlueprint* Blueprint, const UProperty* OldVariable, const UProperty* NewVariable);
 
 	/** Validate child blueprint component member variables, member variables, and timelines, and function graphs against the given variable name */
-	static void ValidateBlueprintChildVariables(UBlueprint* InBlueprint, const FName& InVariableName);
+	static void ValidateBlueprintChildVariables(UBlueprint* InBlueprint, const FName InVariableName);
 
 	/** Rename a Timeline. If bRenameNodes is true, will also rename any timeline nodes associated with this timeline */
-	static bool RenameTimeline (UBlueprint* Blueprint, const FName& OldVarName, const FName& NewVarName);
+	static bool RenameTimeline (UBlueprint* Blueprint, const FName OldVarName, const FName NewVarName);
 
 	/**
 	 * Sets the Blueprint edit-only flag on the variable with the specified name
@@ -1204,7 +1204,7 @@ protected:
 	static bool IsObjectADebugCandidate( AActor* InActorObject, UBlueprint* InBlueprint , bool bInDisallowDerivedBlueprints );
 
 	/** Validate child blueprint member variables against the given variable name */
-	static bool ValidateAllMemberVariables(UBlueprint* InBlueprint, UBlueprint* InParentBlueprint, const FName& InVariableName);
+	static bool ValidateAllMemberVariables(UBlueprint* InBlueprint, UBlueprint* InParentBlueprint, const FName InVariableName);
 
 	/** Validate child blueprint component member variables against the given variable name */
 	static bool ValidateAllComponentMemberVariables(UBlueprint* InBlueprint, UBlueprint* InParentBlueprint, const FName& InVariableName);
diff --git a/Engine/Source/Runtime/Engine/Classes/Engine/MemberReference.h b/Engine/Source/Runtime/Engine/Classes/Engine/MemberReference.h
--- a/Engine/Source/Runtime/Engine/Classes/Engine/MemberReference.h
+++ b/Engine/Source/Runtime/Engine/Classes/Engine/MemberReference.h
@@ -416,7 +416,16 @@ public:
 	static TFieldType* ResolveSimpleMemberReference(const FSimpleMemberReference& Reference)
 	{
 		FMemberReference TempMemberReference;
-		TempMemberReference.SetDirect(Reference.MemberName, Reference.MemberGuid, Reference.MemberParentClass, false);
-		return TempMemberReference.ResolveMember<TFieldType>((UClass*)NULL);
+		const FName Name = Reference.MemberGuid.IsValid() ? NAME_None : Reference.MemberName; // if the guid is valid don't check the name, it could be renamed 
+		TempMemberReference.SetDirect(Name, Reference.MemberGuid, Reference.MemberParentClass, false);
+		auto Result = TempMemberReference.ResolveMember<TFieldType>((UClass*)nullptr);
+
+		if (!Result && (Name != Reference.MemberName))
+		{
+			TempMemberReference.SetDirect(Reference.MemberName, Reference.MemberGuid, Reference.MemberParentClass, false);
+			Result = TempMemberReference.ResolveMember<TFieldType>((UClass*)nullptr);
+		}
+
+		return Result;
 	}
 };
diff --git a/Engine/Source/Runtime/Engine/Private/Blueprint.cpp b/Engine/Source/Runtime/Engine/Private/Blueprint.cpp
--- a/Engine/Source/Runtime/Engine/Private/Blueprint.cpp
+++ b/Engine/Source/Runtime/Engine/Private/Blueprint.cpp
@@ -1042,6 +1042,13 @@ void UBlueprint::GetAllGraphs(TArray<UEdGraph*>& Graphs) const
 		Graph->GetAllChildrenGraphs(Graphs);
 	}
 
+	for (int32 i = 0; i < DelegateSignatureGraphs.Num(); ++i)
+	{
+		UEdGraph* Graph = DelegateSignatureGraphs[i];
+		Graphs.Add(Graph);
+		Graph->GetAllChildrenGraphs(Graphs);
+	}
+
 	for (int32 BPIdx=0; BPIdx<ImplementedInterfaces.Num(); BPIdx++)
 	{
 		const FBPInterfaceDescription& InterfaceDesc = ImplementedInterfaces[BPIdx];
