K2Nodes::HasExternalBlueprintDependencies and K2Nodes::HasExternalUserDefinedStructDependencies were changed into HasExternalDependencies and can return native structs

diff --git a/Engine/Source/Editor/BlueprintGraph/Classes/K2Node.h b/Engine/Source/Editor/BlueprintGraph/Classes/K2Node.h
--- a/Engine/Source/Editor/BlueprintGraph/Classes/K2Node.h
+++ b/Engine/Source/Editor/BlueprintGraph/Classes/K2Node.h
@@ -115,16 +115,10 @@ class UK2Node : public UEdGraphNode
 	virtual bool IsNodePure() const { return false; }
 
 	/** 
-	 * Returns whether or not this node has dependencies on an external blueprint 
-	 * If OptionalOutput isn't null, it should be filled with the known dependencies objects (Classes, Functions, etc).
+	 * Returns whether or not this node has dependencies on an external structure 
+	 * If OptionalOutput isn't null, it should be filled with the known dependencies objects (Classes, Structures, Functions, etc).
 	 */
-	virtual bool HasExternalBlueprintDependencies(TArray<class UStruct*>* OptionalOutput = NULL) const { return false; }
-
-	/**
-	* Returns whether or not this node has dependencies on an external user-defined structure
-	* If OptionalOutput isn't null, it should be filled with the known dependencies structures.
-	*/
-	virtual bool HasExternalUserDefinedStructDependencies(TArray<class UStruct*>* OptionalOutput = NULL) const { return false; }
+	virtual bool HasExternalDependencies(TArray<class UStruct*>* OptionalOutput = NULL) const { return false; }
 
 	/** Returns whether this node can have breakpoints placed on it in the debugger */
 	virtual bool CanPlaceBreakpoints() const { return !IsNodePure(); }
diff --git a/Engine/Source/Editor/BlueprintGraph/Classes/K2Node_BaseAsyncTask.h b/Engine/Source/Editor/BlueprintGraph/Classes/K2Node_BaseAsyncTask.h
--- a/Engine/Source/Editor/BlueprintGraph/Classes/K2Node_BaseAsyncTask.h
+++ b/Engine/Source/Editor/BlueprintGraph/Classes/K2Node_BaseAsyncTask.h
@@ -25,7 +25,7 @@ class BLUEPRINTGRAPH_API UK2Node_BaseAsyncTask : public UK2Node
 
 	// UK2Node interface
 	virtual void ExpandNode(class FKismetCompilerContext& CompilerContext, UEdGraph* SourceGraph) override;
-	virtual bool HasExternalBlueprintDependencies(TArray<class UStruct*>* OptionalOutput) const override;
+	virtual bool HasExternalDependencies(TArray<class UStruct*>* OptionalOutput) const override;
 	virtual FName GetCornerIcon() const override;
 	virtual void GetMenuActions(FBlueprintActionDatabaseRegistrar& ActionRegistrar) const override;
 	virtual FText GetMenuCategory() const override;
diff --git a/Engine/Source/Editor/BlueprintGraph/Classes/K2Node_BaseMCDelegate.h b/Engine/Source/Editor/BlueprintGraph/Classes/K2Node_BaseMCDelegate.h
--- a/Engine/Source/Editor/BlueprintGraph/Classes/K2Node_BaseMCDelegate.h
+++ b/Engine/Source/Editor/BlueprintGraph/Classes/K2Node_BaseMCDelegate.h
@@ -23,7 +23,7 @@ public:
 	virtual FString GetDocumentationExcerptName() const override;
 	virtual bool AllowMultipleSelfs(bool bInputAsArray) const override { return true; }
 	virtual void ExpandNode(class FKismetCompilerContext& CompilerContext, UEdGraph* SourceGraph) override;
-	virtual bool HasExternalBlueprintDependencies(TArray<class UStruct*>* OptionalOutput) const override;
+	virtual bool HasExternalDependencies(TArray<class UStruct*>* OptionalOutput) const override;
 	virtual void GetNodeAttributes( TArray<TKeyValuePair<FString, FString>>& OutNodeAttributes ) const override;
 	virtual void AutowireNewNode(UEdGraphPin* FromPin) override;
 	// End of UK2Node interface
diff --git a/Engine/Source/Editor/BlueprintGraph/Classes/K2Node_CallFunction.h b/Engine/Source/Editor/BlueprintGraph/Classes/K2Node_CallFunction.h
--- a/Engine/Source/Editor/BlueprintGraph/Classes/K2Node_CallFunction.h
+++ b/Engine/Source/Editor/BlueprintGraph/Classes/K2Node_CallFunction.h
@@ -80,7 +80,7 @@ public:
 	// UK2Node interface
 	virtual void ReallocatePinsDuringReconstruction(TArray<UEdGraphPin*>& OldPins) override;
 	virtual bool IsNodePure() const override { return bIsPureFunc; }
-	virtual bool HasExternalBlueprintDependencies(TArray<class UStruct*>* OptionalOutput) const override;
+	virtual bool HasExternalDependencies(TArray<class UStruct*>* OptionalOutput) const override;
 	virtual void PostReconstructNode() override;
 	virtual bool ShouldDrawCompact() const override;
 	virtual bool ShouldDrawAsBead() const override;
diff --git a/Engine/Source/Editor/BlueprintGraph/Classes/K2Node_CallFunctionOnMember.h b/Engine/Source/Editor/BlueprintGraph/Classes/K2Node_CallFunctionOnMember.h
--- a/Engine/Source/Editor/BlueprintGraph/Classes/K2Node_CallFunctionOnMember.h
+++ b/Engine/Source/Editor/BlueprintGraph/Classes/K2Node_CallFunctionOnMember.h
@@ -15,7 +15,7 @@ class UK2Node_CallFunctionOnMember : public UK2Node_CallFunction
 	UPROPERTY()
 	FMemberReference				MemberVariableToCallOn;
 
-	virtual bool HasExternalBlueprintDependencies(TArray<class UStruct*>* OptionalOutput) const override;
+	virtual bool HasExternalDependencies(TArray<class UStruct*>* OptionalOutput) const override;
 
 	// Begin UK2Node_CallFunction interface
 	virtual UEdGraphPin* CreateSelfPin(const UFunction* Function) override;
diff --git a/Engine/Source/Editor/BlueprintGraph/Classes/K2Node_ConstructObjectFromClass.h b/Engine/Source/Editor/BlueprintGraph/Classes/K2Node_ConstructObjectFromClass.h
--- a/Engine/Source/Editor/BlueprintGraph/Classes/K2Node_ConstructObjectFromClass.h
+++ b/Engine/Source/Editor/BlueprintGraph/Classes/K2Node_ConstructObjectFromClass.h
@@ -18,7 +18,7 @@ class BLUEPRINTGRAPH_API UK2Node_ConstructObjectFromClass : public UK2Node
 	virtual void PinDefaultValueChanged(UEdGraphPin* Pin) override;
 	virtual FText GetTooltipText() const override;
 	virtual void ExpandNode(class FKismetCompilerContext& CompilerContext, UEdGraph* SourceGraph) override;
-	virtual bool HasExternalBlueprintDependencies(TArray<class UStruct*>* OptionalOutput) const override;
+	virtual bool HasExternalDependencies(TArray<class UStruct*>* OptionalOutput) const override;
 	virtual bool IsCompatibleWithGraph(const UEdGraph* TargetGraph) const override;
 	// End UEdGraphNode interface.
 
diff --git a/Engine/Source/Editor/BlueprintGraph/Classes/K2Node_DynamicCast.h b/Engine/Source/Editor/BlueprintGraph/Classes/K2Node_DynamicCast.h
--- a/Engine/Source/Editor/BlueprintGraph/Classes/K2Node_DynamicCast.h
+++ b/Engine/Source/Editor/BlueprintGraph/Classes/K2Node_DynamicCast.h
@@ -28,7 +28,7 @@ class UK2Node_DynamicCast : public UK2Node
 	// UK2Node interface
 	virtual ERedirectType DoPinsMatchForReconstruction(const UEdGraphPin* NewPin, int32 NewPinIndex, const UEdGraphPin* OldPin, int32 OldPinIndex) const override;
 	virtual class FNodeHandlingFunctor* CreateNodeHandler(class FKismetCompilerContext& CompilerContext) const override;
-	virtual bool HasExternalBlueprintDependencies(TArray<class UStruct*>* OptionalOutput) const override;
+	virtual bool HasExternalDependencies(TArray<class UStruct*>* OptionalOutput) const override;
 	virtual FText GetMenuCategory() const override;
 	virtual FBlueprintNodeSignature GetSignature() const override;
 	virtual bool IsNodePure() const override { return bIsPureCast; }
diff --git a/Engine/Source/Editor/BlueprintGraph/Classes/K2Node_Event.h b/Engine/Source/Editor/BlueprintGraph/Classes/K2Node_Event.h
--- a/Engine/Source/Editor/BlueprintGraph/Classes/K2Node_Event.h
+++ b/Engine/Source/Editor/BlueprintGraph/Classes/K2Node_Event.h
@@ -75,6 +75,7 @@ class UK2Node_Event : public UK2Node_EditablePinBase
 	BLUEPRINTGRAPH_API virtual FText GetToolTipHeading() const override;
 	BLUEPRINTGRAPH_API virtual void GetNodeAttributes( TArray<TKeyValuePair<FString, FString>>& OutNodeAttributes ) const override;
 	BLUEPRINTGRAPH_API virtual FText GetMenuCategory() const override;
+	BLUEPRINTGRAPH_API virtual bool HasExternalDependencies(TArray<class UStruct*>* OptionalOutput) const override;
 	// End UK2Node interface
 
 	/** Checks whether the parameters for this event node are compatible with the specified function entry node */
diff --git a/Engine/Source/Editor/BlueprintGraph/Classes/K2Node_FunctionTerminator.h b/Engine/Source/Editor/BlueprintGraph/Classes/K2Node_FunctionTerminator.h
--- a/Engine/Source/Editor/BlueprintGraph/Classes/K2Node_FunctionTerminator.h
+++ b/Engine/Source/Editor/BlueprintGraph/Classes/K2Node_FunctionTerminator.h
@@ -30,6 +30,7 @@ class UK2Node_FunctionTerminator : public UK2Node_EditablePinBase
 
 	// Begin UK2Node interface
 	virtual bool NodeCausesStructuralBlueprintChange() const override { return true; }
+	virtual bool HasExternalDependencies(TArray<class UStruct*>* OptionalOutput) const override;
 	// End UK2Node interface
 
 	// Begin UK2Node_EditablePinBase interface
diff --git a/Engine/Source/Editor/BlueprintGraph/Classes/K2Node_MacroInstance.h b/Engine/Source/Editor/BlueprintGraph/Classes/K2Node_MacroInstance.h
--- a/Engine/Source/Editor/BlueprintGraph/Classes/K2Node_MacroInstance.h
+++ b/Engine/Source/Editor/BlueprintGraph/Classes/K2Node_MacroInstance.h
@@ -56,7 +56,7 @@ public:
 	virtual void NotifyPinConnectionListChanged(UEdGraphPin* Pin) override;
 	virtual void PostReconstructNode() override;
 	virtual FText GetActiveBreakpointToolTipText() const override;
-	virtual bool HasExternalBlueprintDependencies(TArray<class UStruct*>* OptionalOutput) const override;
+	virtual bool HasExternalDependencies(TArray<class UStruct*>* OptionalOutput) const override;
 	virtual void GetNodeAttributes( TArray<TKeyValuePair<FString, FString>>& OutNodeAttributes ) const override;
 	virtual FText GetMenuCategory() const override;
 	virtual int32 GetNodeRefreshPriority() const override { return EBaseNodeRefreshPriority::Low_UsesDependentWildcard; }
diff --git a/Engine/Source/Editor/BlueprintGraph/Classes/K2Node_SpawnActor.h b/Engine/Source/Editor/BlueprintGraph/Classes/K2Node_SpawnActor.h
--- a/Engine/Source/Editor/BlueprintGraph/Classes/K2Node_SpawnActor.h
+++ b/Engine/Source/Editor/BlueprintGraph/Classes/K2Node_SpawnActor.h
@@ -29,7 +29,7 @@ class UK2Node_SpawnActor : public UK2Node
 	virtual void ReallocatePinsDuringReconstruction(TArray<UEdGraphPin*>& OldPins) override;
 	virtual class FNodeHandlingFunctor* CreateNodeHandler(class FKismetCompilerContext& CompilerContext) const override;
 	virtual void ExpandNode(class FKismetCompilerContext& CompilerContext, UEdGraph* SourceGraph) override;
-	virtual bool HasExternalBlueprintDependencies(TArray<class UStruct*>* OptionalOutput) const override;
+	virtual bool HasExternalDependencies(TArray<class UStruct*>* OptionalOutput) const override;
 	virtual void GetNodeAttributes( TArray<TKeyValuePair<FString, FString>>& OutNodeAttributes ) const override;
 	// End UK2Node interface
 
diff --git a/Engine/Source/Editor/BlueprintGraph/Classes/K2Node_SpawnActorFromClass.h b/Engine/Source/Editor/BlueprintGraph/Classes/K2Node_SpawnActorFromClass.h
--- a/Engine/Source/Editor/BlueprintGraph/Classes/K2Node_SpawnActorFromClass.h
+++ b/Engine/Source/Editor/BlueprintGraph/Classes/K2Node_SpawnActorFromClass.h
@@ -19,7 +19,7 @@ class BLUEPRINTGRAPH_API UK2Node_SpawnActorFromClass : public UK2Node
 	virtual void PinConnectionListChanged(UEdGraphPin* Pin) override;
 	virtual FText GetTooltipText() const override;
 	virtual void ExpandNode(class FKismetCompilerContext& CompilerContext, UEdGraph* SourceGraph) override;
-	virtual bool HasExternalBlueprintDependencies(TArray<class UStruct*>* OptionalOutput) const override;
+	virtual bool HasExternalDependencies(TArray<class UStruct*>* OptionalOutput) const override;
 	virtual FName GetPaletteIcon(FLinearColor& OutColor) const override{ return TEXT("GraphEditor.SpawnActor_16x"); }
 	virtual bool IsCompatibleWithGraph(const UEdGraph* TargetGraph) const override;
 	virtual void PostPlacedNewNode() override;
diff --git a/Engine/Source/Editor/BlueprintGraph/Classes/K2Node_StructOperation.h b/Engine/Source/Editor/BlueprintGraph/Classes/K2Node_StructOperation.h
--- a/Engine/Source/Editor/BlueprintGraph/Classes/K2Node_StructOperation.h
+++ b/Engine/Source/Editor/BlueprintGraph/Classes/K2Node_StructOperation.h
@@ -22,7 +22,7 @@ class UK2Node_StructOperation : public UK2Node_Variable
 	//virtual bool DrawNodeAsVariable() const override { return true; }
 	virtual bool ShouldShowNodeProperties() const override { return true; }
 	virtual void ValidateNodeDuringCompilation(class FCompilerResultsLog& MessageLog) const override {}
-	virtual bool HasExternalUserDefinedStructDependencies(TArray<class UStruct*>* OptionalOutput) const override;
+	virtual bool HasExternalDependencies(TArray<class UStruct*>* OptionalOutput = NULL) const override;
 	// End of UK2Node interface
 
 protected:
diff --git a/Engine/Source/Editor/BlueprintGraph/Classes/K2Node_Variable.h b/Engine/Source/Editor/BlueprintGraph/Classes/K2Node_Variable.h
--- a/Engine/Source/Editor/BlueprintGraph/Classes/K2Node_Variable.h
+++ b/Engine/Source/Editor/BlueprintGraph/Classes/K2Node_Variable.h
@@ -68,7 +68,7 @@ public:
 	// End UEdGraphNode interface
 
 	// Begin K2Node interface
-	virtual bool HasExternalBlueprintDependencies(TArray<class UStruct*>* OptionalOutput) const override;
+	virtual bool HasExternalDependencies(TArray<class UStruct*>* OptionalOutput) const override;
 	virtual bool DrawNodeAsVariable() const override { return true; }
 	virtual ERedirectType DoPinsMatchForReconstruction(const UEdGraphPin* NewPin, int32 NewPinIndex, const UEdGraphPin* OldPin, int32 OldPinIndex)  const override;
 	virtual void ValidateNodeDuringCompilation(class FCompilerResultsLog& MessageLog) const override;
diff --git a/Engine/Source/Editor/BlueprintGraph/Private/K2Node_BaseAsyncTask.cpp b/Engine/Source/Editor/BlueprintGraph/Private/K2Node_BaseAsyncTask.cpp
--- a/Engine/Source/Editor/BlueprintGraph/Private/K2Node_BaseAsyncTask.cpp
+++ b/Engine/Source/Editor/BlueprintGraph/Private/K2Node_BaseAsyncTask.cpp
@@ -334,23 +334,24 @@ void UK2Node_BaseAsyncTask::ExpandNode(class FKismetCompilerContext& CompilerCon
 	BreakAllNodeLinks();
 }
 
-bool UK2Node_BaseAsyncTask::HasExternalBlueprintDependencies(TArray<class UStruct*>* OptionalOutput) const
+bool UK2Node_BaseAsyncTask::HasExternalDependencies(TArray<class UStruct*>* OptionalOutput) const
 {
 	const UBlueprint* SourceBlueprint = GetBlueprint();
 
-	const bool bProxyFactoryResult = (ProxyFactoryClass != NULL) && (ProxyFactoryClass->ClassGeneratedBy != NULL) && (ProxyFactoryClass->ClassGeneratedBy != SourceBlueprint);
+	const bool bProxyFactoryResult = (ProxyFactoryClass != NULL) && (ProxyFactoryClass->ClassGeneratedBy != SourceBlueprint);
 	if (bProxyFactoryResult && OptionalOutput)
 	{
-		OptionalOutput->Add(ProxyFactoryClass);
+		OptionalOutput->AddUnique(ProxyFactoryClass);
 	}
 
-	const bool bProxyResult = (ProxyClass != NULL) && (ProxyClass->ClassGeneratedBy != NULL) && (ProxyClass->ClassGeneratedBy != SourceBlueprint);
+	const bool bProxyResult = (ProxyClass != NULL) && (ProxyClass->ClassGeneratedBy != SourceBlueprint);
 	if (bProxyResult && OptionalOutput)
 	{
-		OptionalOutput->Add(ProxyClass);
+		OptionalOutput->AddUnique(ProxyClass);
 	}
 
-	return bProxyFactoryResult || bProxyResult || Super::HasExternalBlueprintDependencies(OptionalOutput);
+	const bool bSuperResult = Super::HasExternalDependencies(OptionalOutput);
+	return bProxyFactoryResult || bProxyResult || bSuperResult;
 }
 
 FName UK2Node_BaseAsyncTask::GetCornerIcon() const
diff --git a/Engine/Source/Editor/BlueprintGraph/Private/K2Node_CallFunction.cpp b/Engine/Source/Editor/BlueprintGraph/Private/K2Node_CallFunction.cpp
--- a/Engine/Source/Editor/BlueprintGraph/Private/K2Node_CallFunction.cpp
+++ b/Engine/Source/Editor/BlueprintGraph/Private/K2Node_CallFunction.cpp
@@ -2168,7 +2168,7 @@ FText UK2Node_CallFunction::GetMenuCategory() const
 	return FText::GetEmpty();
 }
 
-bool UK2Node_CallFunction::HasExternalBlueprintDependencies(TArray<class UStruct*>* OptionalOutput) const
+bool UK2Node_CallFunction::HasExternalDependencies(TArray<class UStruct*>* OptionalOutput) const
 {
 	UFunction* Function = GetTargetFunction();
 	const UClass* SourceClass = Function ? Function->GetOwnerClass() : nullptr;
@@ -2178,7 +2178,9 @@ bool UK2Node_CallFunction::HasExternalBlueprintDependencies(TArray<class UStruct
 	{
 		OptionalOutput->AddUnique(Function);
 	}
-	return Super::HasExternalBlueprintDependencies(OptionalOutput) || bResult;
+
+	const bool bSuperResult = Super::HasExternalDependencies(OptionalOutput);
+	return bSuperResult || bResult;
 }
 
 UEdGraph* UK2Node_CallFunction::GetFunctionGraph(const UEdGraphNode*& OutGraphNode) const
diff --git a/Engine/Source/Editor/BlueprintGraph/Private/K2Node_CallFunctionOnMember.cpp b/Engine/Source/Editor/BlueprintGraph/Private/K2Node_CallFunctionOnMember.cpp
--- a/Engine/Source/Editor/BlueprintGraph/Private/K2Node_CallFunctionOnMember.cpp
+++ b/Engine/Source/Editor/BlueprintGraph/Private/K2Node_CallFunctionOnMember.cpp
@@ -122,16 +122,20 @@ void UK2Node_CallFunctionOnMember::ExpandNode(class FKismetCompilerContext& Comp
 	BreakAllNodeLinks();
 }
 
-bool UK2Node_CallFunctionOnMember::HasExternalBlueprintDependencies(TArray<class UStruct*>* OptionalOutput) const
+bool UK2Node_CallFunctionOnMember::HasExternalDependencies(TArray<class UStruct*>* OptionalOutput) const
 {
-	UClass* SourceClass = MemberVariableToCallOn.GetMemberParentClass(GetBlueprintClassFromNode());
 	const UBlueprint* SourceBlueprint = GetBlueprint();
-	const bool bResult = (SourceClass != NULL) && (SourceClass->ClassGeneratedBy != NULL) && (SourceClass->ClassGeneratedBy != SourceBlueprint);
+
+	auto VarProperty = MemberVariableToCallOn.ResolveMember<UProperty>(GetBlueprintClassFromNode());
+	UClass* SourceClass = VarProperty ? VarProperty->GetOwnerClass() : nullptr;
+	const bool bResult = (SourceClass != NULL) && (SourceClass->ClassGeneratedBy != SourceBlueprint);
 	if (bResult && OptionalOutput)
 	{
-		OptionalOutput->Add(SourceClass);
+		OptionalOutput->AddUnique(SourceClass);
 	}
-	return bResult || Super::HasExternalBlueprintDependencies(OptionalOutput);
+
+	const bool bSuperResult = Super::HasExternalDependencies(OptionalOutput);
+	return bSuperResult || bResult;
 }
 
 #undef LOCTEXT_NAMESPACE
diff --git a/Engine/Source/Editor/BlueprintGraph/Private/K2Node_ConstructObjectFromClass.cpp b/Engine/Source/Editor/BlueprintGraph/Private/K2Node_ConstructObjectFromClass.cpp
--- a/Engine/Source/Editor/BlueprintGraph/Private/K2Node_ConstructObjectFromClass.cpp
+++ b/Engine/Source/Editor/BlueprintGraph/Private/K2Node_ConstructObjectFromClass.cpp
@@ -336,16 +336,17 @@ void UK2Node_ConstructObjectFromClass::ExpandNode(class FKismetCompilerContext&
 	Super::ExpandNode(CompilerContext, SourceGraph);	
 }
 
-bool UK2Node_ConstructObjectFromClass::HasExternalBlueprintDependencies(TArray<class UStruct*>* OptionalOutput) const
+bool UK2Node_ConstructObjectFromClass::HasExternalDependencies(TArray<class UStruct*>* OptionalOutput) const
 {
 	UClass* SourceClass = GetClassToSpawn();
 	const UBlueprint* SourceBlueprint = GetBlueprint();
-	const bool bResult = (SourceClass != NULL) && (SourceClass->ClassGeneratedBy != NULL) && (SourceClass->ClassGeneratedBy != SourceBlueprint);
+	const bool bResult = (SourceClass != NULL) && (SourceClass->ClassGeneratedBy != SourceBlueprint);
 	if (bResult && OptionalOutput)
 	{
-		OptionalOutput->Add(SourceClass);
+		OptionalOutput->AddUnique(SourceClass);
 	}
-	return bResult || Super::HasExternalBlueprintDependencies(OptionalOutput);
+	const bool bSuperResult = Super::HasExternalDependencies(OptionalOutput);
+	return bSuperResult || bResult;
 }
 
 #undef LOCTEXT_NAMESPACE
\ No newline at end of file
diff --git a/Engine/Source/Editor/BlueprintGraph/Private/K2Node_DynamicCast.cpp b/Engine/Source/Editor/BlueprintGraph/Private/K2Node_DynamicCast.cpp
--- a/Engine/Source/Editor/BlueprintGraph/Private/K2Node_DynamicCast.cpp
+++ b/Engine/Source/Editor/BlueprintGraph/Private/K2Node_DynamicCast.cpp
@@ -271,16 +271,17 @@ FNodeHandlingFunctor* UK2Node_DynamicCast::CreateNodeHandler(FKismetCompilerCont
 	return new FKCHandler_DynamicCast(CompilerContext, KCST_DynamicCast);
 }
 
-bool UK2Node_DynamicCast::HasExternalBlueprintDependencies(TArray<class UStruct*>* OptionalOutput) const
+bool UK2Node_DynamicCast::HasExternalDependencies(TArray<class UStruct*>* OptionalOutput) const
 {
 	const UBlueprint* SourceBlueprint = GetBlueprint();
 	UClass* SourceClass = *TargetType;
-	const bool bResult = (SourceClass != NULL) && (SourceClass->ClassGeneratedBy != NULL) && (SourceClass->ClassGeneratedBy != SourceBlueprint);
+	const bool bResult = (SourceClass != NULL) && (SourceClass->ClassGeneratedBy != SourceBlueprint);
 	if (bResult && OptionalOutput)
 	{
-		OptionalOutput->Add(SourceClass);
+		OptionalOutput->AddUnique(SourceClass);
 	}
-	return bResult || Super::HasExternalBlueprintDependencies(OptionalOutput);
+	const bool bSuperResult = Super::HasExternalDependencies(OptionalOutput);
+	return bSuperResult || bResult;
 }
 
 FText UK2Node_DynamicCast::GetMenuCategory() const
diff --git a/Engine/Source/Editor/BlueprintGraph/Private/K2Node_Event.cpp b/Engine/Source/Editor/BlueprintGraph/Private/K2Node_Event.cpp
--- a/Engine/Source/Editor/BlueprintGraph/Private/K2Node_Event.cpp
+++ b/Engine/Source/Editor/BlueprintGraph/Private/K2Node_Event.cpp
@@ -722,4 +722,20 @@ bool UK2Node_Event::AreEventNodesIdentical(const UK2Node_Event* InNodeA, const U
 		&& InNodeA->EventReference.GetMemberParentClass(InNodeA->GetBlueprintClassFromNode()) == InNodeB->EventReference.GetMemberParentClass(InNodeB->GetBlueprintClassFromNode());
 }
 
+bool UK2Node_Event::HasExternalDependencies(TArray<class UStruct*>* OptionalOutput) const
+{
+	const UBlueprint* SourceBlueprint = GetBlueprint();
+
+	UFunction* Function = EventReference.ResolveMember<UFunction>(GetBlueprintClassFromNode());
+	const UClass* SourceClass = Function ? Function->GetOwnerClass() : nullptr;
+	const bool bResult = (SourceClass != NULL) && (SourceClass->ClassGeneratedBy != SourceBlueprint);
+	if (bResult && OptionalOutput)
+	{
+		OptionalOutput->AddUnique(Function);
+	}
+
+	const bool bSuperResult = Super::HasExternalDependencies(OptionalOutput);
+	return bSuperResult || bResult;
+}
+
 #undef LOCTEXT_NAMESPACE
\ No newline at end of file
diff --git a/Engine/Source/Editor/BlueprintGraph/Private/K2Node_FunctionTerminator.cpp b/Engine/Source/Editor/BlueprintGraph/Private/K2Node_FunctionTerminator.cpp
--- a/Engine/Source/Editor/BlueprintGraph/Private/K2Node_FunctionTerminator.cpp
+++ b/Engine/Source/Editor/BlueprintGraph/Private/K2Node_FunctionTerminator.cpp
@@ -46,4 +46,19 @@ bool UK2Node_FunctionTerminator::CanCreateUserDefinedPin(const FEdGraphPinType&
 	return bIsEditable;
 }
 
+bool UK2Node_FunctionTerminator::HasExternalDependencies(TArray<class UStruct*>* OptionalOutput) const
+{
+	const UBlueprint* SourceBlueprint = GetBlueprint();
+
+	UClass* SourceClass = *SignatureClass;
+	const bool bResult = (SourceClass != NULL) && (SourceClass->ClassGeneratedBy != SourceBlueprint);
+	if (bResult && OptionalOutput)
+	{
+		OptionalOutput->AddUnique(SourceClass);
+	}
+
+	const bool bSuperResult = Super::HasExternalDependencies(OptionalOutput);
+	return bSuperResult || bResult;
+}
+
 #undef LOCTEXT_NAMESPACE
diff --git a/Engine/Source/Editor/BlueprintGraph/Private/K2Node_MCDelegate.cpp b/Engine/Source/Editor/BlueprintGraph/Private/K2Node_MCDelegate.cpp
--- a/Engine/Source/Editor/BlueprintGraph/Private/K2Node_MCDelegate.cpp
+++ b/Engine/Source/Editor/BlueprintGraph/Private/K2Node_MCDelegate.cpp
@@ -186,16 +186,28 @@ bool UK2Node_BaseMCDelegate::IsAuthorityOnly() const
 
 }
 
-bool UK2Node_BaseMCDelegate::HasExternalBlueprintDependencies(TArray<class UStruct*>* OptionalOutput) const
+bool UK2Node_BaseMCDelegate::HasExternalDependencies(TArray<class UStruct*>* OptionalOutput) const
 {
-	const UClass* SourceClass = DelegateReference.GetMemberParentClass(GetBlueprintClassFromNode());
 	const UBlueprint* SourceBlueprint = GetBlueprint();
-	const bool bResult = (SourceClass != NULL) && (SourceClass->ClassGeneratedBy != NULL) && (SourceClass->ClassGeneratedBy != SourceBlueprint);
-	if (bResult && OptionalOutput)
+
+	auto MCProperty = GetProperty();
+	UClass* PropertySourceClass = MCProperty ? MCProperty->GetOwnerClass() : nullptr;
+	const bool bPropertyResult = (PropertySourceClass != NULL) && (PropertySourceClass->ClassGeneratedBy != SourceBlueprint);
+	if (bPropertyResult && OptionalOutput)
 	{
-		OptionalOutput->AddUnique(GetDelegateSignature());
+		OptionalOutput->AddUnique(PropertySourceClass);
 	}
-	return Super::HasExternalBlueprintDependencies(OptionalOutput) || bResult;
+
+	auto Signature = GetDelegateSignature(true);
+	UClass* SignatureSourceClass = Signature ? Signature->GetOwnerClass() : nullptr;
+	const bool bSignatureResult = (SignatureSourceClass != NULL) && (SignatureSourceClass->ClassGeneratedBy != SourceBlueprint);
+	if (bSignatureResult && OptionalOutput)
+	{
+		OptionalOutput->AddUnique(Signature);
+	}
+
+	const bool bSuperResult = Super::HasExternalDependencies(OptionalOutput);
+	return bSignatureResult || bPropertyResult || bSuperResult;
 }
 
 void UK2Node_BaseMCDelegate::GetNodeAttributes( TArray<TKeyValuePair<FString, FString>>& OutNodeAttributes ) const
diff --git a/Engine/Source/Editor/BlueprintGraph/Private/K2Node_MacroInstance.cpp b/Engine/Source/Editor/BlueprintGraph/Private/K2Node_MacroInstance.cpp
--- a/Engine/Source/Editor/BlueprintGraph/Private/K2Node_MacroInstance.cpp
+++ b/Engine/Source/Editor/BlueprintGraph/Private/K2Node_MacroInstance.cpp
@@ -505,7 +505,7 @@ FText UK2Node_MacroInstance::GetActiveBreakpointToolTipText() const
 	return LOCTEXT("ActiveBreakpointToolTip", "Execution will break inside the macro.");
 }
 
-bool UK2Node_MacroInstance::HasExternalBlueprintDependencies(TArray<class UStruct*>* OptionalOutput) const
+bool UK2Node_MacroInstance::HasExternalDependencies(TArray<class UStruct*>* OptionalOutput) const
 {
 	UBlueprint* OtherBlueprint = MacroGraphReference.GetBlueprint();
 	const bool bResult = OtherBlueprint && (OtherBlueprint != GetBlueprint());
@@ -513,10 +513,11 @@ bool UK2Node_MacroInstance::HasExternalBlueprintDependencies(TArray<class UStruc
 	{
 		if (UClass* OtherClass = *OtherBlueprint->GeneratedClass)
 		{
-			OptionalOutput->Add(OtherClass);
+			OptionalOutput->AddUnique(OtherClass);
 		}
 	}
-	return bResult;
+	const bool bSuperResult = Super::HasExternalDependencies(OptionalOutput);
+	return bSuperResult || bResult;
 }
 
 void UK2Node_MacroInstance::GetNodeAttributes( TArray<TKeyValuePair<FString, FString>>& OutNodeAttributes ) const
diff --git a/Engine/Source/Editor/BlueprintGraph/Private/K2Node_SpawnActor.cpp b/Engine/Source/Editor/BlueprintGraph/Private/K2Node_SpawnActor.cpp
--- a/Engine/Source/Editor/BlueprintGraph/Private/K2Node_SpawnActor.cpp
+++ b/Engine/Source/Editor/BlueprintGraph/Private/K2Node_SpawnActor.cpp
@@ -452,16 +452,17 @@ void UK2Node_SpawnActor::ExpandNode(class FKismetCompilerContext& CompilerContex
 	BreakAllNodeLinks();
 }
 
-bool UK2Node_SpawnActor::HasExternalBlueprintDependencies(TArray<class UStruct*>* OptionalOutput) const
+bool UK2Node_SpawnActor::HasExternalDependencies(TArray<class UStruct*>* OptionalOutput) const
 {
 	UClass* SourceClass = GetClassToSpawn();
 	const UBlueprint* SourceBlueprint = GetBlueprint();
-	const bool bResult = (SourceClass != NULL) && (SourceClass->ClassGeneratedBy != NULL) && (SourceClass->ClassGeneratedBy != SourceBlueprint);
+	const bool bResult = (SourceClass != NULL) && (SourceClass->ClassGeneratedBy != SourceBlueprint);
 	if (bResult && OptionalOutput)
 	{
-		OptionalOutput->Add(SourceClass);
+		OptionalOutput->AddUnique(SourceClass);
 	}
-	return bResult || Super::HasExternalBlueprintDependencies(OptionalOutput);
+	const bool bSuperResult = Super::HasExternalDependencies(OptionalOutput);
+	return bSuperResult || bResult;
 }
 
 void UK2Node_SpawnActor::GetNodeAttributes( TArray<TKeyValuePair<FString, FString>>& OutNodeAttributes ) const
diff --git a/Engine/Source/Editor/BlueprintGraph/Private/K2Node_SpawnActorFromClass.cpp b/Engine/Source/Editor/BlueprintGraph/Private/K2Node_SpawnActorFromClass.cpp
--- a/Engine/Source/Editor/BlueprintGraph/Private/K2Node_SpawnActorFromClass.cpp
+++ b/Engine/Source/Editor/BlueprintGraph/Private/K2Node_SpawnActorFromClass.cpp
@@ -535,16 +535,17 @@ void UK2Node_SpawnActorFromClass::ExpandNode(class FKismetCompilerContext& Compi
 	SpawnNode->BreakAllNodeLinks();
 }
 
-bool UK2Node_SpawnActorFromClass::HasExternalBlueprintDependencies(TArray<class UStruct*>* OptionalOutput) const
+bool UK2Node_SpawnActorFromClass::HasExternalDependencies(TArray<class UStruct*>* OptionalOutput) const
 {
 	UClass* SourceClass = GetClassToSpawn();
 	const UBlueprint* SourceBlueprint = GetBlueprint();
-	const bool bResult = (SourceClass != NULL) && (SourceClass->ClassGeneratedBy != NULL) && (SourceClass->ClassGeneratedBy != SourceBlueprint);
+	const bool bResult = (SourceClass != NULL) && (SourceClass->ClassGeneratedBy != SourceBlueprint);
 	if (bResult && OptionalOutput)
 	{
-		OptionalOutput->Add(SourceClass);
+		OptionalOutput->AddUnique(SourceClass);
 	}
-	return bResult || Super::HasExternalBlueprintDependencies(OptionalOutput);
+	const bool bSuperResult = Super::HasExternalDependencies(OptionalOutput);
+	return bSuperResult || bResult;
 }
 
 #undef LOCTEXT_NAMESPACE
\ No newline at end of file
diff --git a/Engine/Source/Editor/BlueprintGraph/Private/K2Node_StructOperation.cpp b/Engine/Source/Editor/BlueprintGraph/Private/K2Node_StructOperation.cpp
--- a/Engine/Source/Editor/BlueprintGraph/Private/K2Node_StructOperation.cpp
+++ b/Engine/Source/Editor/BlueprintGraph/Private/K2Node_StructOperation.cpp
@@ -12,15 +12,16 @@ UK2Node_StructOperation::UK2Node_StructOperation(const FObjectInitializer& Objec
 	: Super(ObjectInitializer)
 {
 }
-
-bool UK2Node_StructOperation::HasExternalUserDefinedStructDependencies(TArray<class UStruct*>* OptionalOutput) const
+bool UK2Node_StructOperation::HasExternalDependencies(TArray<class UStruct*>* OptionalOutput) const
 {
-	const bool bResult = StructType && StructType->IsA<UUserDefinedStruct>();
+	const bool bResult = nullptr != StructType;
 	if (bResult && OptionalOutput)
 	{
-		OptionalOutput->Add(StructType);
+		OptionalOutput->AddUnique(StructType);
 	}
-	return bResult || Super::HasExternalUserDefinedStructDependencies(OptionalOutput);
+
+	const bool bSuperResult = Super::HasExternalDependencies(OptionalOutput);
+	return bSuperResult || bResult;
 }
 
 void UK2Node_StructOperation::FStructOperationOptionalPinManager::CustomizePinData(UEdGraphPin* Pin, FName SourcePropertyName, int32 ArrayIndex, UProperty* Property) const
diff --git a/Engine/Source/Editor/BlueprintGraph/Private/K2Node_Variable.cpp b/Engine/Source/Editor/BlueprintGraph/Private/K2Node_Variable.cpp
--- a/Engine/Source/Editor/BlueprintGraph/Private/K2Node_Variable.cpp
+++ b/Engine/Source/Editor/BlueprintGraph/Private/K2Node_Variable.cpp
@@ -585,16 +585,17 @@ FName UK2Node_Variable::GetCornerIcon() const
 	return Super::GetCornerIcon();
 }
 
-bool UK2Node_Variable::HasExternalBlueprintDependencies(TArray<class UStruct*>* OptionalOutput) const
+bool UK2Node_Variable::HasExternalDependencies(TArray<class UStruct*>* OptionalOutput) const
 {
 	UClass* SourceClass = GetVariableSourceClass();
 	UBlueprint* SourceBlueprint = GetBlueprint();
-	const bool bResult = (SourceClass && (SourceClass->ClassGeneratedBy != NULL) && (SourceClass->ClassGeneratedBy != SourceBlueprint));
+	const bool bResult = (SourceClass && (SourceClass->ClassGeneratedBy != SourceBlueprint));
 	if (bResult && OptionalOutput)
 	{
-		OptionalOutput->Add(SourceClass);
+		OptionalOutput->AddUnique(SourceClass);
 	}
-	return bResult || Super::HasExternalBlueprintDependencies(OptionalOutput);
+	const bool bSuperResult = Super::HasExternalDependencies(OptionalOutput);
+	return bSuperResult || bResult;
 }
 
 FString UK2Node_Variable::GetDocumentationLink() const
diff --git a/Engine/Source/Editor/UnrealEd/Private/Kismet2/BlueprintEditorUtils.cpp b/Engine/Source/Editor/UnrealEd/Private/Kismet2/BlueprintEditorUtils.cpp
--- a/Engine/Source/Editor/UnrealEd/Private/Kismet2/BlueprintEditorUtils.cpp
+++ b/Engine/Source/Editor/UnrealEd/Private/Kismet2/BlueprintEditorUtils.cpp
@@ -500,7 +500,7 @@ void FBlueprintEditorUtils::RefreshExternalBlueprintDependencyNodes(UBlueprint*
 		for (auto NodeIt = AllNodes.CreateIterator(); NodeIt; ++NodeIt)
 		{
 			UK2Node* Node = *NodeIt;
-			if (Node->HasExternalBlueprintDependencies())
+			if (Node->HasExternalDependencies())
 			{
 				//@todo:  Do we really need per-schema refreshing?
 				const UEdGraphSchema* Schema = Node->GetGraph()->GetSchema();
@@ -514,7 +514,7 @@ void FBlueprintEditorUtils::RefreshExternalBlueprintDependencyNodes(UBlueprint*
 		{
 			UK2Node* Node = *NodeIt;
 			TArray<UStruct*> Dependencies;
-			if (Node->HasExternalBlueprintDependencies(&Dependencies))
+			if (Node->HasExternalDependencies(&Dependencies))
 			{
 				for (UStruct* Struct : Dependencies)
 				{
@@ -1017,7 +1017,7 @@ struct FRegenerationHelper
 					if (Node)
 					{
 						TArray<UStruct*> LocalDependentStructures;
-						if (Node->HasExternalBlueprintDependencies(&LocalDependentStructures))
+						if (Node->HasExternalDependencies(&LocalDependentStructures))
 						{
 							for (auto Struct : LocalDependentStructures)
 							{
@@ -1046,13 +1046,6 @@ struct FRegenerationHelper
 							}
 						}
 
-						LocalDependentStructures.Empty(LocalDependentStructures.Max());
-						Node->HasExternalUserDefinedStructDependencies(&LocalDependentStructures);
-						for (auto Struct : LocalDependentStructures)
-						{
-							ProcessHierarchy(Struct, Dependencies);
-						}
-
 						auto FunctionEntry = Cast<UK2Node_FunctionEntry>(Node);
 						if (FunctionEntry && !bIsDelegateSignatureGraph)
 						{
@@ -2660,7 +2653,7 @@ void FBlueprintEditorUtils::GatherDependencies(const UBlueprint* InBlueprint, TS
 
 	for (const auto& InterfaceDesc : InBlueprint->ImplementedInterfaces)
 	{
-		UBlueprint* InterfaceBP = InterfaceDesc.Interface ? Cast<UBlueprint>(InterfaceDesc.Interface->ClassGeneratedBy) : NULL;
+		UBlueprint* InterfaceBP = InterfaceDesc.Interface ? Cast<UBlueprint>(InterfaceDesc.Interface->ClassGeneratedBy) : nullptr;
 		if (InterfaceBP)
 		{
 			Dependencies.Add(InterfaceBP);
@@ -2678,7 +2671,7 @@ void FBlueprintEditorUtils::GatherDependencies(const UBlueprint* InBlueprint, TS
 			for (auto Node : Nodes)
 			{
 				TArray<UStruct*> LocalDependentStructures;
-				if (Node && Node->HasExternalBlueprintDependencies(&LocalDependentStructures))
+				if (Node && Node->HasExternalDependencies(&LocalDependentStructures))
 				{
 					for (auto Struct : LocalDependentStructures)
 					{
