Editor: Changed the texture thumbnail renderer to treat TEXTUREGROUP_Pixels2D like it does TEXTUREGROUP_UI, showing the texture as translucent [UE-14421]

diff --git a/Engine/Source/Editor/UnrealEd/Private/ThumbnailRendering/TextureThumbnailRenderer.cpp b/Engine/Source/Editor/UnrealEd/Private/ThumbnailRendering/TextureThumbnailRenderer.cpp
--- a/Engine/Source/Editor/UnrealEd/Private/ThumbnailRendering/TextureThumbnailRenderer.cpp
+++ b/Engine/Source/Editor/UnrealEd/Private/ThumbnailRendering/TextureThumbnailRenderer.cpp
@@ -49,7 +49,7 @@ void UTextureThumbnailRenderer::Draw(UObject* Object, int32 X, int32 Y, uint32 W
 		// This provides a much better preview than just showing RGB,
 		// Because the RGB content in areas with an alpha of 0 is often garbage that will not be seen in normal conditions.
 		// Non-UI textures often have uncorrelated data in the alpha channel (like a skin mask, specular power, etc) so we only preview UI textures this way.
-		const bool bUseTranslucentBlend = Texture2D && Texture2D->HasAlphaChannel() && Texture2D->LODGroup == TEXTUREGROUP_UI;
+		const bool bUseTranslucentBlend = Texture2D && Texture2D->HasAlphaChannel() && ((Texture2D->LODGroup == TEXTUREGROUP_UI) || (Texture2D->LODGroup == TEXTUREGROUP_Pixels2D));
 
 		UTextureCube* TextureCube = Cast<UTextureCube>(Texture);
 		UTextureRenderTargetCube* RTTextureCube = Cast<UTextureRenderTargetCube>(Texture);
