PgkInfoCommandlet -assetregistry displays size of asset registry data stored in the file.

diff --git a/Engine/Source/Editor/UnrealEd/Private/Commandlets/PackageUtilities.cpp b/Engine/Source/Editor/UnrealEd/Private/Commandlets/PackageUtilities.cpp
--- a/Engine/Source/Editor/UnrealEd/Private/Commandlets/PackageUtilities.cpp
+++ b/Engine/Source/Editor/UnrealEd/Private/Commandlets/PackageUtilities.cpp
@@ -1210,7 +1210,14 @@ void FPkgInfoReporter_Log::GeneratePackageReport( FLinkerLoad* InLinker/*=NULL*/
 
 	if( (InfoFlags&PKGINFO_AssetRegistry) != 0 )
 	{
-		UE_LOG(LogPackageUtilities, Warning, TEXT("--------------------------------------------") );
+		UE_LOG(LogPackageUtilities, Warning, TEXT("--------------------------------------------"));
+
+		{
+			const int32 NextOffset = Linker->Summary.WorldTileInfoDataOffset ? Linker->Summary.WorldTileInfoDataOffset : Linker->Summary.TotalHeaderSize;
+			const int32 AssetRegistrySize = NextOffset - Linker->Summary.AssetRegistryDataOffset;
+			UE_LOG(LogPackageUtilities, Display, TEXT("Asset Registry Size: %10i"), AssetRegistrySize);
+		}
+		
 		GWarn->Log ( TEXT("Asset Registry Data"));
 		GWarn->Log ( TEXT("=========="));
 
