Adjusted MeshMerging dialog window height

diff --git a/Engine/Source/Editor/UnrealEd/Private/EditorGroup.cpp b/Engine/Source/Editor/UnrealEd/Private/EditorGroup.cpp
--- a/Engine/Source/Editor/UnrealEd/Private/EditorGroup.cpp
+++ b/Engine/Source/Editor/UnrealEd/Private/EditorGroup.cpp
@@ -266,7 +266,7 @@ void UUnrealEdEngine::edactMergeActorsByMaterials()
 	/** Create the window to host mesh megin widget */
 	TSharedRef<SWindow> MergeActorsWindow = SNew(SWindow)
 											.Title(FText::FromString("Merge actors"))
-											.ClientSize(FVector2D(400, 330));
+											.ClientSize(FVector2D(400, 350));
 
 	/** Set the content of the window to our package dialog widget */
 	TSharedRef<SMeshMergingDialog> MergingDialog = SNew(SMeshMergingDialog)
