Destructible mesh thumbnail so DM assets dont go through the skeletal mesh thumbnail renderer (causes output spam)

diff --git a/Engine/Config/BaseEditor.ini b/Engine/Config/BaseEditor.ini
--- a/Engine/Config/BaseEditor.ini
+++ b/Engine/Config/BaseEditor.ini
@@ -30,6 +30,7 @@ MaxNumberOfLogLines=2000
 +RenderableThumbnailTypes=(ClassNeedingThumbnailName="/Script/Engine.ParticleSystem",RendererClassName="/Script/UnrealEd.ParticleSystemThumbnailRenderer")
 +RenderableThumbnailTypes=(ClassNeedingThumbnailName="/Script/Engine.StaticMesh",RendererClassName="/Script/UnrealEd.StaticMeshThumbnailRenderer")
 +RenderableThumbnailTypes=(ClassNeedingThumbnailName="/Script/Engine.SkeletalMesh",RendererClassName="/Script/UnrealEd.SkeletalMeshThumbnailRenderer")
++RenderableThumbnailTypes=(ClassNeedingThumbnailName="/Script/Engine.DestructibleMesh",RendererClassName="/Script/UnrealEd.DestructibleMeshThumbnailRenderer")
 +RenderableThumbnailTypes=(ClassNeedingThumbnailName="/Script/Engine.AnimSequenceBase",RendererClassName="/Script/UnrealEd.AnimSequenceThumbnailRenderer")
 +RenderableThumbnailTypes=(ClassNeedingThumbnailName="/Script/Engine.BlendSpaceBase",RendererClassName="/Script/UnrealEd.BlendSpaceThumbnailRenderer")
 +RenderableThumbnailTypes=(ClassNeedingThumbnailName="/Script/Engine.AnimBlueprint",RendererClassName="/Script/UnrealEd.AnimBlueprintThumbnailRenderer")
diff --git a/Engine/Source/Editor/UnrealEd/Classes/ThumbnailRendering/DestructibleMeshThumbnailRenderer.h b/Engine/Source/Editor/UnrealEd/Classes/ThumbnailRendering/DestructibleMeshThumbnailRenderer.h
new file mode 100644
--- /dev/null
+++ b/Engine/Source/Editor/UnrealEd/Classes/ThumbnailRendering/DestructibleMeshThumbnailRenderer.h
@@ -0,0 +1,27 @@
+// Copyright 1998-2015 Epic Games, Inc. All Rights Reserved.
+
+/**
+ *
+ * This thumbnail renderer displays a given skeletal mesh
+ */
+
+#pragma once
+#include "DestructibleMeshThumbnailRenderer.generated.h"
+
+UCLASS(config=Editor, MinimalAPI)
+class UDestructibleMeshThumbnailRenderer : public UDefaultSizedThumbnailRenderer
+{
+	GENERATED_UCLASS_BODY()
+
+
+	// Begin UThumbnailRenderer Object
+	UNREALED_API virtual void Draw(UObject* Object, int32 X, int32 Y, uint32 Width, uint32 Height, FRenderTarget*, FCanvas* Canvas) override;
+	// End UThumbnailRenderer Object
+
+	// UObject implementation
+	UNREALED_API virtual void BeginDestroy() override;
+
+private:
+	class FDestructibleMeshThumbnailScene* ThumbnailScene;
+};
+
diff --git a/Engine/Source/Editor/UnrealEd/Private/ThumbnailHelpers.cpp b/Engine/Source/Editor/UnrealEd/Private/ThumbnailHelpers.cpp
--- a/Engine/Source/Editor/UnrealEd/Private/ThumbnailHelpers.cpp
+++ b/Engine/Source/Editor/UnrealEd/Private/ThumbnailHelpers.cpp
@@ -15,6 +15,8 @@
 #include "Engine/SimpleConstructionScript.h"
 #include "Engine/StaticMesh.h"
 #include "Engine/TextureCube.h"
+#include "Engine/DestructibleMesh.h"
+#include "PhysicsEngine/DestructibleActor.h"
 
 /*
 ***************************************************************
@@ -429,6 +431,75 @@ void FSkeletalMeshThumbnailScene::GetViewMatrixParameters(const float InFOVDegre
 	OutOrbitZoom = TargetDistance + ThumbnailInfo->OrbitZoom;
 }
 
+/*
+***************************************************************
+FDestructibleMeshThumbnailScene
+***************************************************************
+*/
+
+FDestructibleMeshThumbnailScene::FDestructibleMeshThumbnailScene()
+	: FThumbnailPreviewScene()
+{
+	bForceAllUsedMipsResident = false;
+	// Create preview actor
+	// checked
+	FActorSpawnParameters SpawnInfo;
+	SpawnInfo.bNoCollisionFail = true;
+	SpawnInfo.bNoFail = true;
+	SpawnInfo.ObjectFlags = RF_Transient;
+	PreviewActor = GetWorld()->SpawnActor<ADestructibleActor>(SpawnInfo);
+
+	PreviewActor->SetActorEnableCollision(false);
+}
+
+void FDestructibleMeshThumbnailScene::SetDestructibleMesh(class UDestructibleMesh* InMesh)
+{
+	PreviewActor->GetDestructibleComponent()->OverrideMaterials.Empty();
+	PreviewActor->GetDestructibleComponent()->SetDestructibleMesh(InMesh);
+
+	if (InMesh)
+	{
+		FTransform MeshTransform = FTransform::Identity;
+
+		PreviewActor->SetActorLocation(FVector(0, 0, 0), false);
+		PreviewActor->GetDestructibleComponent()->UpdateBounds();
+
+		// Center the mesh at the world origin then offset to put it on top of the plane
+		const float BoundsZOffset = GetBoundsZOffset(PreviewActor->GetDestructibleComponent()->Bounds);
+		PreviewActor->SetActorLocation(-PreviewActor->GetDestructibleComponent()->Bounds.Origin + FVector(0, 0, BoundsZOffset), false);
+		PreviewActor->GetDestructibleComponent()->RecreateRenderState_Concurrent();
+	}
+}
+
+void FDestructibleMeshThumbnailScene::GetViewMatrixParameters(const float InFOVDegrees, FVector& OutOrigin, float& OutOrbitPitch, float& OutOrbitYaw, float& OutOrbitZoom) const
+{
+	check(PreviewActor->GetDestructibleComponent());
+	check(PreviewActor->GetDestructibleComponent()->GetDestructibleMesh());
+
+	const float HalfFOVRadians = FMath::DegreesToRadians<float>(InFOVDegrees) * 0.5f;
+	// No need to add extra size to view slightly outside of the sphere to compensate for perspective since skeletal meshes already buffer bounds.
+	const float HalfMeshSize = PreviewActor->GetDestructibleComponent()->Bounds.SphereRadius;
+	const float BoundsZOffset = GetBoundsZOffset(PreviewActor->GetDestructibleComponent()->Bounds);
+	const float TargetDistance = HalfMeshSize / FMath::Tan(HalfFOVRadians);
+
+	USceneThumbnailInfo* ThumbnailInfo = Cast<USceneThumbnailInfo>(PreviewActor->GetDestructibleComponent()->DestructibleMesh->ThumbnailInfo);
+	if (ThumbnailInfo)
+	{
+		if (TargetDistance + ThumbnailInfo->OrbitZoom < 0)
+		{
+			ThumbnailInfo->OrbitZoom = -TargetDistance;
+		}
+	}
+	else
+	{
+		ThumbnailInfo = USceneThumbnailInfo::StaticClass()->GetDefaultObject<USceneThumbnailInfo>();
+	}
+
+	OutOrigin = FVector(0, 0, -BoundsZOffset);
+	OutOrbitPitch = ThumbnailInfo->OrbitPitch;
+	OutOrbitYaw = ThumbnailInfo->OrbitYaw;
+	OutOrbitZoom = TargetDistance + ThumbnailInfo->OrbitZoom;
+}
 
 /*
 ***************************************************************
diff --git a/Engine/Source/Editor/UnrealEd/Private/ThumbnailRendering/DestructibleMeshThumbnailRenderer.cpp b/Engine/Source/Editor/UnrealEd/Private/ThumbnailRendering/DestructibleMeshThumbnailRenderer.cpp
new file mode 100644
--- /dev/null
+++ b/Engine/Source/Editor/UnrealEd/Private/ThumbnailRendering/DestructibleMeshThumbnailRenderer.cpp
@@ -0,0 +1,50 @@
+// Copyright 1998-2015 Epic Games, Inc. All Rights Reserved.
+
+#include "UnrealEd.h"
+
+// FPreviewScene derived helpers for rendering
+#include "ThumbnailHelpers.h"
+#include "EngineModule.h"
+#include "RendererInterface.h"
+#include "Engine/DestructibleMesh.h"
+
+UDestructibleMeshThumbnailRenderer::UDestructibleMeshThumbnailRenderer(const FObjectInitializer& ObjectInitializer)
+	: Super(ObjectInitializer)
+{
+	ThumbnailScene = nullptr;
+}
+
+void UDestructibleMeshThumbnailRenderer::Draw(UObject* Object, int32 X, int32 Y, uint32 Width, uint32 Height, FRenderTarget* RenderTarget, FCanvas* Canvas)
+{
+	UDestructibleMesh* DestructibleMesh = Cast<UDestructibleMesh>(Object);
+	if (DestructibleMesh != nullptr)
+	{
+		if ( ThumbnailScene == nullptr )
+		{
+			ThumbnailScene = new FDestructibleMeshThumbnailScene();
+		}
+
+		ThumbnailScene->SetDestructibleMesh(DestructibleMesh);
+		FSceneViewFamilyContext ViewFamily( FSceneViewFamily::ConstructionValues( RenderTarget, ThumbnailScene->GetScene(), FEngineShowFlags(ESFIM_Game) )
+			.SetWorldTimes(FApp::GetCurrentTime() - GStartTime, FApp::GetDeltaTime(), FApp::GetCurrentTime() - GStartTime));
+
+		ViewFamily.EngineShowFlags.DisableAdvancedFeatures();
+		ViewFamily.EngineShowFlags.MotionBlur = 0;
+		ViewFamily.EngineShowFlags.LOD = 0;
+
+		ThumbnailScene->GetView(&ViewFamily, X, Y, Width, Height);
+		GetRendererModule().BeginRenderingViewFamily(Canvas,&ViewFamily);
+		ThumbnailScene->SetDestructibleMesh(nullptr);
+	}
+}
+
+void UDestructibleMeshThumbnailRenderer::BeginDestroy()
+{
+	if ( ThumbnailScene != nullptr )
+	{
+		delete ThumbnailScene;
+		ThumbnailScene = nullptr;
+	}
+
+	Super::BeginDestroy();
+}
diff --git a/Engine/Source/Editor/UnrealEd/Public/ThumbnailHelpers.h b/Engine/Source/Editor/UnrealEd/Public/ThumbnailHelpers.h
--- a/Engine/Source/Editor/UnrealEd/Public/ThumbnailHelpers.h
+++ b/Engine/Source/Editor/UnrealEd/Public/ThumbnailHelpers.h
@@ -90,6 +90,24 @@ private:
 	class ASkeletalMeshActor* PreviewActor;
 };
 
+class UNREALED_API FDestructibleMeshThumbnailScene : public FThumbnailPreviewScene
+{
+public:
+	/** Constructor */
+	FDestructibleMeshThumbnailScene();
+
+	/** Sets the skeletal mesh to use in the next GetView() */
+	void SetDestructibleMesh(class UDestructibleMesh* InMesh);
+
+protected:
+	// FThumbnailPreviewScene implementation
+	virtual void GetViewMatrixParameters(const float InFOVDegrees, FVector& OutOrigin, float& OutOrbitPitch, float& OutOrbitYaw, float& OutOrbitZoom) const override;
+
+private:
+	/** The skeletal mesh actor used to display all skeletal mesh thumbnails */
+	class ADestructibleActor* PreviewActor;
+};
+
 class UNREALED_API FStaticMeshThumbnailScene : public FThumbnailPreviewScene
 {
 public:
