Fixing 'Unable to generate long package name for /Script/*' log spam when cooking.

diff --git a/Engine/Source/Editor/UnrealEd/Private/CookOnTheFlyServer.cpp b/Engine/Source/Editor/UnrealEd/Private/CookOnTheFlyServer.cpp
--- a/Engine/Source/Editor/UnrealEd/Private/CookOnTheFlyServer.cpp
+++ b/Engine/Source/Editor/UnrealEd/Private/CookOnTheFlyServer.cpp
@@ -691,7 +691,8 @@ void UCookOnTheFlyServer::GetDependentPackages( const TSet<UPackage*>& RootPacka
 			FString PackageDepdencyString = PackageDependency.ToString();
 
 			FText OutReason;
-			if ( !FPackageName::IsValidLongPackageName(PackageDepdencyString, false, &OutReason) )
+			const bool bIncludeReadOnlyRoots = true; // Dependency packages are often script packages (read-only)
+			if (!FPackageName::IsValidLongPackageName(PackageDepdencyString, bIncludeReadOnlyRoots, &OutReason))
 			{
 				const FText FailMessage = FText::Format(LOCTEXT("UnableToGeneratePackageName", "Unable to generate long package name for {0}. {1}"),
 					FText::FromString(PackageDepdencyString), OutReason);
