Adding or renaming function pins in Blueprints, will correctly refresh all active uses of the function.

diff --git a/Engine/Source/Editor/Kismet/Private/BlueprintDetailsCustomization.cpp b/Engine/Source/Editor/Kismet/Private/BlueprintDetailsCustomization.cpp
--- a/Engine/Source/Editor/Kismet/Private/BlueprintDetailsCustomization.cpp
+++ b/Engine/Source/Editor/Kismet/Private/BlueprintDetailsCustomization.cpp
@@ -3265,7 +3265,7 @@ bool FBaseBlueprintGraphActionDetails::OnPinRenamed(UK2Node_EditablePinBase* Tar
 
 		for (auto BlueprintIt = PinRenamedHelper.ModifiedBlueprints.CreateIterator(); BlueprintIt; ++BlueprintIt)
 		{
-			FBlueprintEditorUtils::MarkBlueprintAsModified(*BlueprintIt);
+			FBlueprintEditorUtils::MarkBlueprintAsStructurallyModified(*BlueprintIt);
 		}
 	}
 	return true;
diff --git a/Engine/Source/Editor/UnrealEd/Private/Kismet2/BlueprintEditorUtils.cpp b/Engine/Source/Editor/UnrealEd/Private/Kismet2/BlueprintEditorUtils.cpp
--- a/Engine/Source/Editor/UnrealEd/Private/Kismet2/BlueprintEditorUtils.cpp
+++ b/Engine/Source/Editor/UnrealEd/Private/Kismet2/BlueprintEditorUtils.cpp
@@ -317,8 +317,8 @@ void FBasePinChangeHelper::Broadcast(UBlueprint* InBlueprint, UK2Node_EditablePi
 
 				const bool bNameMatches = (CallSite->FunctionReference.GetMemberName() == FunctionDefNode->SignatureName);
 				const UClass* MemberParentClass = CallSite->FunctionReference.GetMemberParentClass(CallSite->GetBlueprintClassFromNode());
-				const bool bClassMatchesEasy = (MemberParentClass != NULL) && (MemberParentClass->IsChildOf(FunctionDefNode->SignatureClass));
-				const bool bClassMatchesHard = (CallSiteBlueprint != NULL) && (CallSite->FunctionReference.IsSelfContext()) && (FunctionDefNode->SignatureClass == NULL) && (CallSiteBlueprint == InBlueprint);
+				const bool bClassMatchesEasy = (MemberParentClass != NULL) && (MemberParentClass->IsChildOf(FunctionDefNode->SignatureClass) || MemberParentClass->IsChildOf(InBlueprint->GeneratedClass));
+				const bool bClassMatchesHard = (CallSiteBlueprint != NULL) && (CallSite->FunctionReference.IsSelfContext()) && (FunctionDefNode->SignatureClass == NULL) && (CallSiteBlueprint == InBlueprint || CallSiteBlueprint->SkeletonGeneratedClass->IsChildOf(InBlueprint->SkeletonGeneratedClass));
 				const bool bValidSchema = CallSite->GetSchema() != NULL;
 
 				if (bNameMatches && bValidSchema && (bClassMatchesEasy || bClassMatchesHard))
