DLC cooking removed packages which are in based on release from dlc asset registry.

diff --git a/Engine/Source/Editor/UnrealEd/Classes/CookOnTheSide/CookOnTheFlyServer.h b/Engine/Source/Editor/UnrealEd/Classes/CookOnTheSide/CookOnTheFlyServer.h
--- a/Engine/Source/Editor/UnrealEd/Classes/CookOnTheSide/CookOnTheFlyServer.h
+++ b/Engine/Source/Editor/UnrealEd/Classes/CookOnTheSide/CookOnTheFlyServer.h
@@ -616,6 +616,8 @@ private:
 		TMap< FName, TSet <FName> > MapDependencyGraph; 
 		/** If a cook is cancelled next cook will need to resume cooking */ 
 		TArray<FFilePlatformRequest> PreviousCookRequests; 
+		/** If we are based on a release version of the game this is the set of packages which were cooked in that release */
+		TMap<FName,TArray<FName> > BasedOnReleaseCookedPackages;
 		double CookTime;
 		double CookStartTime;
 		bool bErrorOnEngineContentUse;
@@ -1034,6 +1036,13 @@ private:
 		return false;
 	}
 
+	/**
+	 * GetDLCContentPath
+	 * 
+	 * @return return the path to the source dlc content
+	 */
+	FString GetDLCContentPath();
+	
 	inline bool IsCreatingReleaseVersion()
 	{
 		if ( CookByTheBookOptions )
diff --git a/Engine/Source/Editor/UnrealEd/Private/Commandlets/ChunkManifestGenerator.cpp b/Engine/Source/Editor/UnrealEd/Private/Commandlets/ChunkManifestGenerator.cpp
--- a/Engine/Source/Editor/UnrealEd/Private/Commandlets/ChunkManifestGenerator.cpp
+++ b/Engine/Source/Editor/UnrealEd/Private/Commandlets/ChunkManifestGenerator.cpp
@@ -451,15 +451,31 @@ bool FChunkManifestGenerator::LoadAssetRegistry(const FString& SandboxPath, cons
 	return true;
 }
 
-bool FChunkManifestGenerator::SaveAssetRegistry(const FString& SandboxPath)
+bool FChunkManifestGenerator::SaveAssetRegistry(const FString& SandboxPath, const TArray<FName>* IgnorePackageList)
 {
 	UE_LOG(LogChunkManifestGenerator, Display, TEXT("Saving asset registry."));
 
+
+	TSet<FName> IgnorePackageSet;
+	if (IgnorePackageList != nullptr)
+	{
+		for (const auto& IgnorePackage : *IgnorePackageList)
+		{
+			IgnorePackageSet.Add(IgnorePackage);
+		}
+	}
+	
+
 	// Create asset registry data
 	FArrayWriter SerializedAssetRegistry;
 	TMap<FName, FAssetData*> GeneratedAssetRegistryData;
 	for (auto& AssetData : AssetRegistryData)
 	{
+		if (IgnorePackageSet.Contains(AssetData.PackageName))
+		{
+			continue;
+		}
+
 		// Add only assets that have actually been cooked and belong to any chunk
 		if (AssetData.ChunkIDs.Num() > 0)
 		{
diff --git a/Engine/Source/Editor/UnrealEd/Private/Commandlets/ChunkManifestGenerator.h b/Engine/Source/Editor/UnrealEd/Private/Commandlets/ChunkManifestGenerator.h
--- a/Engine/Source/Editor/UnrealEd/Private/Commandlets/ChunkManifestGenerator.h
+++ b/Engine/Source/Editor/UnrealEd/Private/Commandlets/ChunkManifestGenerator.h
@@ -300,7 +300,7 @@ public:
 	/**
 	* Saves generated asset registry data for each platform.
 	*/
-	bool SaveAssetRegistry(const FString& SandboxPath);
+	bool SaveAssetRegistry(const FString& SandboxPath, const TArray<FName>* IgnorePackageList = nullptr);
 
 
 	/**
diff --git a/Engine/Source/Editor/UnrealEd/Private/CookOnTheFlyServer.cpp b/Engine/Source/Editor/UnrealEd/Private/CookOnTheFlyServer.cpp
--- a/Engine/Source/Editor/UnrealEd/Private/CookOnTheFlyServer.cpp
+++ b/Engine/Source/Editor/UnrealEd/Private/CookOnTheFlyServer.cpp
@@ -391,7 +391,7 @@ const UCookOnTheFlyServer::FCachedPackageFilename& UCookOnTheFlyServer::Cache(co
 	FString StandardFilename;
 	FName StandardFileFName = NAME_None;
 	if (FPackageName::DoesPackageExist(PackageName.ToString(), NULL, &Filename))
-{
+	{
 		StandardFilename = PackageFilename = FPaths::ConvertRelativePathToFull(Filename);
 
 
@@ -956,6 +956,11 @@ bool UCookOnTheFlyServer::IsCookOnTheFlyMode() const
 	return CurrentCookMode == ECookMode::CookOnTheFly || CurrentCookMode == ECookMode::CookOnTheFlyFromTheEditor; 
 }
 
+FString UCookOnTheFlyServer::GetDLCContentPath()
+{
+	return FPaths::GamePluginsDir() / CookByTheBookOptions->DlcName / FString(TEXT("Content"));
+}
+
 COREUOBJECT_API extern bool GOutputCookingWarnings;
 
 uint32 UCookOnTheFlyServer::TickCookOnTheSide( const float TimeSlice, uint32 &CookedPackageCount )
@@ -1112,7 +1117,8 @@ uint32 UCookOnTheFlyServer::TickCookOnTheSide( const float TimeSlice, uint32 &Co
 		
 		if( CookByTheBookOptions && CookByTheBookOptions->bErrorOnEngineContentUse )
 		{
-			FString DLCPath = FPaths::GamePluginsDir() / CookByTheBookOptions->DlcName / FString(TEXT("Content"));
+			check(IsCookingDLC());
+			FString DLCPath = GetDLCContentPath();
 			if ( ToBuild.GetFilename().ToString().StartsWith(DLCPath) == false ) // if we don't start with the dlc path then we shouldn't be cooking this data 
 			{
 				UE_LOG(LogCook, Error, TEXT("Engine content %s is being referenced by DLC!"), *ToBuild.GetFilename().ToString() );
@@ -2617,15 +2623,16 @@ void UCookOnTheFlyServer::GenerateLongPackageNames(TArray<FName>& FilesInPath)
 		else
 		{
 			FString LongPackageName;
-			if (FPackageName::TryConvertFilenameToLongPackageName(FileInPath, LongPackageName))
+			FString FailureReason;
+			if (FPackageName::TryConvertFilenameToLongPackageName(FileInPath, LongPackageName, &FailureReason))
 			{
 				const FName LongPackageFName(*LongPackageName);
 				FilesInPathReverse.AddUnique(LongPackageFName);
 			}
 			else
 			{
-				LogCookerMessage( FString::Printf(TEXT("Unable to generate long package name for %s"), *FileInPath), EMessageSeverity::Warning);
-				UE_LOG(LogCook, Warning, TEXT("Unable to generate long package name for %s"), *FileInPath);
+				LogCookerMessage(FString::Printf(TEXT("Unable to generate long package name for %s because %s"), *FileInPath, *FailureReason), EMessageSeverity::Warning);
+				UE_LOG(LogCook, Warning, TEXT("Unable to generate long package name for %s because %s"), *FileInPath, *FailureReason);
 			}
 		}
 	}
@@ -2659,31 +2666,36 @@ void UCookOnTheFlyServer::CollectFilesToCook(TArray<FName>& FilesInPath, const T
 	TArray<FString> CookDirectories = InCookDirectories;
 	
 	
-	TArray<FString> MapList;
-	// Add the default map section
-	GEditor->LoadMapListFromIni(TEXT("AlwaysCookMaps"), MapList);
-
-	for ( const auto &IniMapSection : IniMapSections )
+	if (!IsCookingDLC())
 	{
-		GEditor->LoadMapListFromIni(*IniMapSection, MapList);
-	}
+		TArray<FString> MapList;
+		// Add the default map section
+		GEditor->LoadMapListFromIni(TEXT("AlwaysCookMaps"), MapList);
 
-	for (int32 MapIdx = 0; MapIdx < MapList.Num(); MapIdx++)
-	{
-		AddFileToCook( FilesInPath, MapList[MapIdx]);
-	}
+		for (const auto &IniMapSection : IniMapSections)
+		{
+			GEditor->LoadMapListFromIni(*IniMapSection, MapList);
+		}
 
+		for (int32 MapIdx = 0; MapIdx < MapList.Num(); MapIdx++)
+		{
+			AddFileToCook(FilesInPath, MapList[MapIdx]);
+		}
 
-	// Also append any cookdirs from the project ini files; these dirs are relative to the game content directory
-	{
-		const FString AbsoluteGameContentDir = FPaths::ConvertRelativePathToFull(FPaths::GameContentDir());
-		const UProjectPackagingSettings* const PackagingSettings = GetDefault<UProjectPackagingSettings>();
-		for(const auto& DirToCook : PackagingSettings->DirectoriesToAlwaysCook)
+
+		// Also append any cookdirs from the project ini files; these dirs are relative to the game content directory
 		{
-			CookDirectories.Add(AbsoluteGameContentDir / DirToCook.Path);
+			const FString AbsoluteGameContentDir = FPaths::ConvertRelativePathToFull(FPaths::GameContentDir());
+			const UProjectPackagingSettings* const PackagingSettings = GetDefault<UProjectPackagingSettings>();
+			for (const auto& DirToCook : PackagingSettings->DirectoriesToAlwaysCook)
+			{
+				CookDirectories.Add(AbsoluteGameContentDir / DirToCook.Path);
+			}
 		}
 	}
 
+	
+
 	for ( const auto &CurrEntry : CookMaps )
 	{
 		if (FPackageName::IsShortPackageName(CurrEntry))
@@ -2986,6 +2998,7 @@ void UCookOnTheFlyServer::CookByTheBookFinished()
 
 	GetDerivedDataCacheRef().WaitForQuiescence(true);
 
+	
 	{
 		// Save modified asset registry with all streaming chunk info generated during cook
 		const FString& SandboxRegistryFilename = GetSandboxAssetRegistryFilename();
@@ -2998,7 +3011,19 @@ void UCookOnTheFlyServer::CookByTheBookFinished()
 		{
 			// Always try to save the manifests, this is required to make the asset registry work, but doesn't necessarily write a file
 			Manifest.Value->SaveManifests(SandboxFile.GetOwnedPointer());
-			Manifest.Value->SaveAssetRegistry(SandboxRegistryFilename);
+
+			const FName& PlatformName = Manifest.Key;
+			const TArray<FName>& PackageFilenames = CookByTheBookOptions->BasedOnReleaseCookedPackages.FindRef(PlatformName);
+			TArray<FName> LongPackageNames;
+			LongPackageNames.Empty(PackageFilenames.Num());
+			for (const auto& PackageFilename : PackageFilenames)
+			{
+				FString LongPackageName;
+				verify(FPackageName::TryConvertFilenameToLongPackageName(PackageFilename.ToString(), LongPackageName));
+				LongPackageNames.Add(FName(*LongPackageName));
+			}
+
+			Manifest.Value->SaveAssetRegistry(SandboxRegistryFilename, &LongPackageNames);
 
 			Manifest.Value->SaveCookedPackageAssetRegistry(SandboxCookedAssetRegistryFilename, true);
 
@@ -3055,6 +3080,8 @@ void UCookOnTheFlyServer::CookByTheBookFinished()
 	const float TotalCookTime = (float)(FPlatformTime::Seconds() - CookByTheBookOptions->CookStartTime);
 	UE_LOG(LogCook, Display, TEXT("Cook by the book total time in tick %fs total time %f"), CookByTheBookOptions->CookTime, TotalCookTime);
 
+	CookByTheBookOptions->BasedOnReleaseCookedPackages.Empty();
+
 	CookByTheBookOptions->bRunning = false;
 }
 
@@ -3199,7 +3226,7 @@ void UCookOnTheFlyServer::StartCookByTheBook( const FCookByTheBookStartupOptions
 	CookByTheBookOptions->CookStartTime = FPlatformTime::Seconds();
 	CookByTheBookOptions->bGenerateStreamingInstallManifests = CookByTheBookStartupOptions.bGenerateStreamingInstallManifests;
 	CookByTheBookOptions->bGenerateDependenciesForMaps = CookByTheBookStartupOptions.bGenerateDependenciesForMaps;
-	CookByTheBookOptions->CreateReleaseVersion = CreateReleaseVersion;
+	CookByTheBookOptions->CreateReleaseVersion = CreateReleaseVersion;	
 
 	bool bCookAll = (CookOptions & ECookByTheBookOptions::CookAll) != ECookByTheBookOptions::None;
 	bool bMapsOnly = (CookOptions & ECookByTheBookOptions::MapsOnly) != ECookByTheBookOptions::None;
@@ -3255,18 +3282,21 @@ void UCookOnTheFlyServer::StartCookByTheBook( const FCookByTheBookStartupOptions
 		{
 			const FString OriginalSandboxRegistryFilename = GetReleaseVersionAssetRegistryPath(BasedOnReleaseVersion, PlatformName ) / GetAssetRegistryFilename();
 
-			TArray<FName> PackageFiles;
+			TArray<FName> PackageList;
 			// if this check fails probably because the asset registry can't be found or read
-			bool bSucceeded = GetAllPackagesFromAssetRegistry(OriginalSandboxRegistryFilename, PackageFiles);
+			bool bSucceeded = GetAllPackagesFromAssetRegistry(OriginalSandboxRegistryFilename, PackageList);
 			check( bSucceeded );
 
 			if ( bSucceeded )
 			{
-				for ( const auto& PackageFilename : PackageFiles )
+				TArray<FName> PlatformNames;
+				PlatformNames.Add(PlatformName);
+				for (const auto& PackageFilename : PackageList)
 				{
-					CookedPackages.Add( MoveTemp( FFilePlatformRequest( PackageFilename, TargetPlatformNames ) ) );
+					CookedPackages.Add( MoveTemp( FFilePlatformRequest( PackageFilename, PlatformNames ) ) );
 				}
 			}
+			CookByTheBookOptions->BasedOnReleaseCookedPackages.Add(PlatformName, MoveTemp(PackageList));
 		}
 	}
 	
