Don't add native script packages to package dependencies when cooking

diff --git a/Engine/Source/Editor/UnrealEd/Private/CookOnTheFlyServer.cpp b/Engine/Source/Editor/UnrealEd/Private/CookOnTheFlyServer.cpp
--- a/Engine/Source/Editor/UnrealEd/Private/CookOnTheFlyServer.cpp
+++ b/Engine/Source/Editor/UnrealEd/Private/CookOnTheFlyServer.cpp
@@ -701,6 +701,10 @@ void UCookOnTheFlyServer::GetDependentPackages( const TSet<UPackage*>& RootPacka
 				UE_LOG(LogCook, Warning, TEXT("%s"), *( FailMessage.ToString() ));
 				continue;
 			}
+			else if (FPackageName::IsScriptPackage(PackageDepdencyString))
+			{
+				continue;
+			}
 
 			if ( FoundPackages.Contains(PackageDependency) == false )
 			{
