Use the output log dialog to show UBT output if a build fails, rather than trying to cram it into a normal message box.

diff --git a/Engine/Source/Editor/GameProjectGeneration/Private/GameProjectUtils.cpp b/Engine/Source/Editor/GameProjectGeneration/Private/GameProjectUtils.cpp
--- a/Engine/Source/Editor/GameProjectGeneration/Private/GameProjectUtils.cpp
+++ b/Engine/Source/Editor/GameProjectGeneration/Private/GameProjectUtils.cpp
@@ -29,6 +29,8 @@
 #include "SourceCodeNavigation.h"
 ///#include "AssetToolsModule.h"
 
+#include "SOutputLogDialog.h"
+
 #define LOCTEXT_NAMESPACE "GameProjectUtils"
 
 #define MAX_PROJECT_PATH_BUFFER_SPACE 130 // Leave a reasonable buffer of additional characters to account for files created in the content directory during or after project generation
@@ -1952,15 +1954,10 @@ bool GameProjectUtils::BuildCodeProject(const FString& ProjectFilename)
 
 		TArray<FText> CompileFailedButtons;
 		int32 OpenIDEButton = CompileFailedButtons.Add(FText::Format(LOCTEXT("CompileFailedOpenIDE", "Open with {0}"), DevEnvName));
-		int32 ViewLogButton = CompileFailedButtons.Add(LOCTEXT("CompileFailedViewLog", "View build log"));
 		CompileFailedButtons.Add(LOCTEXT("CompileFailedCancel", "Cancel"));
 
-		int32 CompileFailedChoice = SVerbChoiceDialog::ShowModal(LOCTEXT("ProjectUpgradeTitle", "Project Conversion Failed"), FText::Format(LOCTEXT("ProjectUpgradeCompileFailed", "The project failed to compile with this version of the engine. Would you like to open the project in {0}?"), DevEnvName), CompileFailedButtons);
-		if(CompileFailedChoice == ViewLogButton)
-		{
-			CompileFailedButtons.RemoveAt(ViewLogButton);
-			CompileFailedChoice = SVerbChoiceDialog::ShowModal(LOCTEXT("ProjectUpgradeTitle", "Project Conversion Failed"), FText::Format(LOCTEXT("ProjectUpgradeCompileFailed", "The project failed to compile with this version of the engine. Build output is as follows:\n\n{0}"), FText::FromString(OutputLog)), CompileFailedButtons);
-		}
+		FText LogText = FText::FromString(OutputLog.Replace(LINE_TERMINATOR, TEXT("\n")).TrimTrailing());
+		int32 CompileFailedChoice = SOutputLogDialog::Open(LOCTEXT("CompileFailedTitle", "Compile Failed"), FText::Format(LOCTEXT("CompileFailedHeader", "The project could not be compiled. Would you like to open it in {0}?"), DevEnvName), LogText, FText::GetEmpty(), CompileFailedButtons);
 
 		FText FailReason;
 		if(CompileFailedChoice == OpenIDEButton && !GameProjectUtils::OpenCodeIDE(ProjectFilename, FailReason))
diff --git a/Engine/Source/Editor/UnrealEd/Private/Dialogs/SOutputLogDialog.cpp b/Engine/Source/Editor/UnrealEd/Private/Dialogs/SOutputLogDialog.cpp
--- a/Engine/Source/Editor/UnrealEd/Private/Dialogs/SOutputLogDialog.cpp
+++ b/Engine/Source/Editor/UnrealEd/Private/Dialogs/SOutputLogDialog.cpp
@@ -5,6 +5,13 @@
 #include "SHyperlink.h"
 
 void SOutputLogDialog::Open( const FText& InTitle, const FText& InHeader, const FText& InLog, const FText& InFooter )
+{
+	TArray<FText> Buttons;
+	Buttons.Add(NSLOCTEXT("SOutputLogDialog", "Ok", "Ok"));
+	Open(InTitle, InHeader, InLog, InFooter, Buttons);
+}
+
+int32 SOutputLogDialog::Open( const FText& InTitle, const FText& InHeader, const FText& InLog, const FText& InFooter, const TArray<FText>& InButtons )
 {
 	TSharedRef<SWindow> ModalWindow = SNew(SWindow)
 		.Title( InTitle )
@@ -17,11 +24,14 @@ void SOutputLogDialog::Open( const FText& InTitle, const FText& InHeader, const
 		.ParentWindow(ModalWindow)
 		.Header( InHeader )
 		.Log( InLog )
-		.Footer( InFooter );
+		.Footer( InFooter )
+		.Buttons( InButtons );
 
 	ModalWindow->SetContent( MessageBox );
 
 	GEditor->EditorAddModalWindow(ModalWindow);
+
+	return MessageBox->Response;
 }
 
 BEGIN_SLATE_FUNCTION_BUILD_OPTIMIZATION
@@ -31,9 +41,10 @@ void SOutputLogDialog::Construct( const FArguments& InArgs )
 	ParentWindow->SetWidgetToFocusOnActivate(SharedThis(this));
 
 	FSlateFontInfo MessageFont( FEditorStyle::GetFontStyle("StandardDialog.LargeFont"));
-	Header = InArgs._Header;
-	Log = InArgs._Log;
-	Footer = InArgs._Footer;
+	Header = InArgs._Header.Get();
+	Log = InArgs._Log.Get();
+	Footer = InArgs._Footer.Get();
+	Buttons = InArgs._Buttons.Get();
 
 	MaxWidth = FSlateApplication::Get().GetPreferredWorkArea().GetSize().X * 0.8f;
 
@@ -53,7 +64,7 @@ void SOutputLogDialog::Construct( const FArguments& InArgs )
 						[
 							SNew(STextBlock)
 								.Text(Header)
-								.Visibility(Header.Get().IsEmptyOrWhitespace()? EVisibility::Hidden : EVisibility::Visible)
+								.Visibility(Header.IsEmptyOrWhitespace()? EVisibility::Hidden : EVisibility::Visible)
 								.Font(MessageFont)
 								.WrapTextAt(MaxWidth - 50.0f)
 						]
@@ -69,7 +80,7 @@ void SOutputLogDialog::Construct( const FArguments& InArgs )
 								.Style(FEditorStyle::Get(), "Log.TextBox")
 								.TextStyle(FEditorStyle::Get(), "Log.Normal")
 								.ForegroundColor(FLinearColor::Gray)
-								.Text(FText::TrimTrailing(Log.Get()))
+								.Text(FText::TrimTrailing(Log))
 								.IsReadOnly(true)
 								.AlwaysShowScrollbars(true)
 						]
@@ -77,11 +88,11 @@ void SOutputLogDialog::Construct( const FArguments& InArgs )
 					+SVerticalBox::Slot()
 						.HAlign(HAlign_Fill)
 						.AutoHeight()
-						.Padding(12.0f, 0.0f, 12.0f, Footer.Get().IsEmptyOrWhitespace()? 0.0f : 12.0f)
+						.Padding(12.0f, 0.0f, 12.0f, Footer.IsEmptyOrWhitespace()? 0.0f : 12.0f)
 						[
 							SNew(STextBlock)
 								.Text(Footer)
-								.Visibility(Footer.Get().IsEmptyOrWhitespace()? EVisibility::Hidden : EVisibility::Visible)
+								.Visibility(Footer.IsEmptyOrWhitespace()? EVisibility::Hidden : EVisibility::Visible)
 								.Font(MessageFont)
 								.WrapTextAt(MaxWidth - 50.0f)
 						]
@@ -104,31 +115,38 @@ void SOutputLogDialog::Construct( const FArguments& InArgs )
 										.ToolTipText( NSLOCTEXT("SOutputLogDialog", "CopyMessageTooltip", "Copy the text in this message to the clipboard (CTRL+C)") )
 								]
 
-							+SHorizontalBox::Slot()
+							+ SHorizontalBox::Slot()
 								.AutoWidth()
 								.HAlign(HAlign_Right)
 								.VAlign(VAlign_Center)
-								.Padding(0.f)
+								.Padding(0.0f)
 								[
-									SNew( SButton )
-										.Text( NSLOCTEXT("SOutputLogDialog", "Ok", "Ok") )
-										.OnClicked( this, &SOutputLogDialog::HandleButtonClicked )
-										.ContentPadding(FEditorStyle::GetMargin("StandardDialog.ContentPadding"))
-										.HAlign(HAlign_Center)
+									SAssignNew( ButtonBox, SUniformGridPanel )
+										.SlotPadding(FEditorStyle::GetMargin("StandardDialog.SlotPadding"))
+										.MinDesiredSlotWidth(FEditorStyle::GetFloat("StandardDialog.MinDesiredSlotWidth"))
+										.MinDesiredSlotHeight(FEditorStyle::GetFloat("StandardDialog.MinDesiredSlotHeight"))
 								]
 						]
 				]
 		];
+
+	for(int32 Idx = 0; Idx < Buttons.Num(); Idx++)
+	{
+		ButtonBox->AddSlot(Idx, 0)
+			[
+				SNew( SButton )
+				.Text( Buttons[Idx] )
+				.OnClicked( this, &SOutputLogDialog::HandleButtonClicked, Idx )
+				.ContentPadding(FEditorStyle::GetMargin("StandardDialog.ContentPadding"))
+				.HAlign(HAlign_Center)
+			];
+	}
 }
 END_SLATE_FUNCTION_BUILD_OPTIMIZATION
 
 FReply SOutputLogDialog::OnKeyDown( const FGeometry& MyGeometry, const FKeyEvent& InKeyEvent )
 {
-	if( InKeyEvent.GetKey() == EKeys::Escape )
-	{
-		return HandleButtonClicked();
-	}
-	else if (InKeyEvent.GetKey() == EKeys::C && InKeyEvent.IsControlDown())
+	if (InKeyEvent.GetKey() == EKeys::C && InKeyEvent.IsControlDown())
 	{
 		CopyMessageToClipboard();
 		return FReply::Handled();
@@ -153,12 +171,13 @@ FVector2D SOutputLogDialog::ComputeDesiredSize(float LayoutScaleMultiplier) cons
 
 void SOutputLogDialog::CopyMessageToClipboard( )
 {
-	FString FullMessage = FString::Printf(TEXT("%s") LINE_TERMINATOR LINE_TERMINATOR TEXT("%s") LINE_TERMINATOR LINE_TERMINATOR TEXT("%s"), *Header.Get().ToString(), *Log.Get().ToString(), *Footer.Get().ToString()).Trim();
+	FString FullMessage = FString::Printf(TEXT("%s") LINE_TERMINATOR LINE_TERMINATOR TEXT("%s") LINE_TERMINATOR LINE_TERMINATOR TEXT("%s"), *Header.ToString(), *Log.ToString(), *Footer.ToString()).Trim();
 	FPlatformMisc::ClipboardCopy( *FullMessage );
 }
 
-FReply SOutputLogDialog::HandleButtonClicked( )
+FReply SOutputLogDialog::HandleButtonClicked( int32 InResponse )
 {
+	Response = InResponse;
 	ParentWindow->RequestDestroyWindow();
 	return FReply::Handled();
 }
diff --git a/Engine/Source/Editor/UnrealEd/Public/Dialogs/SOutputLogDialog.h b/Engine/Source/Editor/UnrealEd/Public/Dialogs/SOutputLogDialog.h
--- a/Engine/Source/Editor/UnrealEd/Public/Dialogs/SOutputLogDialog.h
+++ b/Engine/Source/Editor/UnrealEd/Public/Dialogs/SOutputLogDialog.h
@@ -10,11 +10,15 @@ public:
 		SLATE_ATTRIBUTE(FText, Header)	
 		SLATE_ATTRIBUTE(FText, Log)	
 		SLATE_ATTRIBUTE(FText, Footer)	
+		SLATE_ATTRIBUTE(TArray<FText>, Buttons)
 	SLATE_END_ARGS()
 
 	/** Displays the modal dialog box */
 	static UNREALED_API void Open(const FText& InTitle, const FText& InHeader, const FText& InLog, const FText& InFooter = FText::GetEmpty());
 
+	/** Displays the modal dialog box with a custom set of buttons */
+	static UNREALED_API int32 Open(const FText& InTitle, const FText& InHeader, const FText& InLog, const FText& InFooter, const TArray<FText>& InButtons);
+
 	/** Constructs the dialog */
 	void Construct( const FArguments& InArgs );
 
@@ -33,14 +37,16 @@ protected:
 
 private:
 	/** Handles clicking a message box button. */
-	FReply HandleButtonClicked( );
+	FReply HandleButtonClicked( int32  );
 
 	/** Handles clicking the 'Copy Message' hyper link. */
 	void HandleCopyMessageHyperlinkNavigate( );
 
+	int32 Response;
 	TSharedPtr<SWindow> ParentWindow;
-	TAttribute<FText> Header;
-	TAttribute<FText> Log;
-	TAttribute<FText> Footer;
+	FText Header;
+	FText Log;
+	FText Footer;
+	TArray<FText> Buttons;
 	float MaxWidth;
 };
