Fix potential crash when a scenecomponent no longer exists within an Actor after reinstancing

diff --git a/Engine/Source/Editor/UnrealEd/Private/Kismet2/KismetReinstanceUtilities.cpp b/Engine/Source/Editor/UnrealEd/Private/Kismet2/KismetReinstanceUtilities.cpp
--- a/Engine/Source/Editor/UnrealEd/Private/Kismet2/KismetReinstanceUtilities.cpp
+++ b/Engine/Source/Editor/UnrealEd/Private/Kismet2/KismetReinstanceUtilities.cpp
@@ -766,7 +766,7 @@ void FActorReplacementHelper::Finalize(const TMap<UObject*, UObject*>& OldToNewI
 	if (TargetParentComponent)
 	{
 		UObject* const* NewTargetParentComponent = OldToNewInstanceMap.Find(TargetParentComponent);
-		if (NewTargetParentComponent)
+		if (NewTargetParentComponent && *NewTargetParentComponent)
 		{
 			TargetParentComponent = CastChecked<USceneComponent>(*NewTargetParentComponent);
 		}
