#HLOD :

- fixed overflow issue
- fixed big actor issue
- hide static mesh again because it causes GLEO if editor is opened
- LODParentPrimitive is now private - have setter/getter
- Made sure LODActor isn't placed manually
- Added proper map error
- Removed it from FGlobalComponentReregisterContext because I don't think we'll need it anymore.

diff --git a/Engine/Source/Editor/UnrealEd/Private/EditorServer.cpp b/Engine/Source/Editor/UnrealEd/Private/EditorServer.cpp
--- a/Engine/Source/Editor/UnrealEd/Private/EditorServer.cpp
+++ b/Engine/Source/Editor/UnrealEd/Private/EditorServer.cpp
@@ -5005,7 +5005,7 @@ void UEditorEngine::AssignReplacementComponentsByActors(TArray<AActor*>& ActorsT
 				FComponentReregisterContext ComponentReattch(PrimitiveComponent);
 
 				// set the replacement
-				PrimitiveComponent->LODParentPrimitive = ReplacementComponent;
+				PrimitiveComponent->SetLODParentPrimitive(ReplacementComponent);
 
 				// makr the package as dirty now that we've modified it
 				Actor->MarkPackageDirty();
diff --git a/Engine/Source/Editor/UnrealEd/Private/HierarchicalLOD.cpp b/Engine/Source/Editor/UnrealEd/Private/HierarchicalLOD.cpp
--- a/Engine/Source/Editor/UnrealEd/Private/HierarchicalLOD.cpp
+++ b/Engine/Source/Editor/UnrealEd/Private/HierarchicalLOD.cpp
@@ -65,8 +65,12 @@ float CalculateOverlap(const FSphere& ASphere, const float AFillingFactor, const
 	float  ACapHeight = (BRadius*BRadius - (ARadius - Distance)*(ARadius - Distance)) / (2*Distance);
 	float  BCapHeight = (ARadius*ARadius - (BRadius - Distance)*(BRadius - Distance)) / (2*Distance);
 
-	check (ACapHeight>0.f);
-	check (BCapHeight>0.f);
+	if (ACapHeight<=0.f || BCapHeight<=0.f)
+	{
+		// it's possible to get cap height to be less than 0 
+		// since when we do check intersect, we do have regular tolerance
+		return 0.f;		
+	}
 
 	float  OverlapRadius1 = ((ARadius+BRadius)*(ARadius+BRadius) - Distance*Distance) * (Distance*Distance - (ARadius-BRadius)*(ARadius-BRadius));
 	float  OverlapRadius2 = 2*Distance;
@@ -171,6 +175,9 @@ void FHierarchicalLODBuilder::BuildClusters(class ULevel* InLevel)
 		ObjectTools::DeleteSingleObject(Asset, false);
 	}
 
+	// garbage collect
+	CollectGarbage( GARBAGE_COLLECTION_KEEPFLAGS, true );
+
 	// only build if it's enabled
 	if(InLevel->GetWorld()->GetWorldSettings()->bEnableHierarchicalLODSystem)
 	{
@@ -371,7 +378,21 @@ bool ShouldGenerateCluster(class AActor* Actor)
 	// TODO: support instanced static meshes
 	Components.RemoveAll([](UStaticMeshComponent* Val){ return Val->IsA(UInstancedStaticMeshComponent::StaticClass()); });
 
-	return (Components.Num() > 0);
+	// now make sure you check parent primitive, so that we don't build for the actor that already has built. 
+	if (Components.Num() > 0)
+	{
+		for (auto& ComponentIter: Components)
+		{
+			if (ComponentIter->GetLODParentPrimitive())
+			{
+				return false;
+			}
+		}
+
+		return true;
+	}
+
+	return false;
 }
 
 void FHierarchicalLODBuilder::InitializeClusters(class ULevel* InLevel, const int32 LODIdx, float CullCost)
@@ -412,34 +433,37 @@ void FHierarchicalLODBuilder::InitializeClusters(class ULevel* InLevel, const in
 		else // at this point we only care for LODActors
 		{
 			Clusters.Empty();
-			TArray<ALODActor*> LODActors;
+			// we filter the LOD index first
+			TArray<AActor*> Actors;
 			for(int32 ActorId=0; ActorId<InLevel->Actors.Num(); ++ActorId)
 			{
-				ALODActor* Actor = Cast<ALODActor>(InLevel->Actors[ActorId]);
-				if (Actor)
+				AActor* Actor = (InLevel->Actors[ActorId]);
+				if (Actor && Actor->IsA(ALODActor::StaticClass()))
+				{
+					ALODActor* LODActor = CastChecked<ALODActor>(Actor);
+					if (LODActor->LODLevel == LODIdx)
+					{
+						Actors.Add(Actor);
+					}
+				}
+				else if(ShouldGenerateCluster(Actor))
 				{
-					LODActors.Add(Actor);
+					Actors.Add(Actor);
 				}
 			}
 			// first we generate graph with 2 pair nodes
 			// this is very expensive when we have so many actors
 			// so we'll need to optimize later @todo
-			for(int32 ActorId=0; ActorId<LODActors.Num(); ++ActorId)
+			for(int32 ActorId=0; ActorId<Actors.Num(); ++ActorId)
 			{
-				ALODActor* Actor1 = (LODActors[ActorId]);
-				if (Actor1->LODLevel == LODIdx)
+				AActor* Actor1 = (Actors[ActorId]);
+				for(int32 SubActorId=ActorId+1; SubActorId<Actors.Num(); ++SubActorId)
 				{
-					for(int32 SubActorId=ActorId+1; SubActorId<LODActors.Num(); ++SubActorId)
-					{
-						ALODActor* Actor2 = LODActors[SubActorId];
+					AActor* Actor2 = Actors[SubActorId];
 
-						if (Actor2->LODLevel == LODIdx)
-						{
-							// create new cluster
-							FLODCluster NewClusterCandidate = FLODCluster(Actor1, Actor2);
-							Clusters.Add(NewClusterCandidate);
-						}
-					}
+					// create new cluster
+					FLODCluster NewClusterCandidate = FLODCluster(Actor1, Actor2);
+					Clusters.Add(NewClusterCandidate);
 				}
 			}
 		}
diff --git a/Engine/Source/Runtime/Core/Private/Misc/MapErrors.cpp b/Engine/Source/Runtime/Core/Private/Misc/MapErrors.cpp
--- a/Engine/Source/Runtime/Core/Private/Misc/MapErrors.cpp
+++ b/Engine/Source/Runtime/Core/Private/Misc/MapErrors.cpp
@@ -54,6 +54,9 @@ FName FMapErrors::NoLevelsAssociated(TEXT("NoLevelsAssociated"));
 FName FMapErrors::FilenameIsTooLongForCooking(TEXT("FilenameIsTooLongForCooking"));
 FName FMapErrors::UsingExternalObject(TEXT("UsingExternalObject"));
 FName FMapErrors::RepairedPaintedVertexColors(TEXT("RepairedPaintedVertexColors"));
+FName FMapErrors::LODActorMissingStaticMesh(TEXT("LODActorMissingStaticMesh"));
+FName FMapErrors::LODActorMissingActor(TEXT("LODActorMissingActor"));
+FName FMapErrors::LODActorNoActorFound(TEXT("LODActorNoActor"));
 
 static const FString MapErrorsPath = TEXT("Shared/Editor/MapErrors");
 
diff --git a/Engine/Source/Runtime/Core/Public/Misc/MapErrors.h b/Engine/Source/Runtime/Core/Public/Misc/MapErrors.h
--- a/Engine/Source/Runtime/Core/Public/Misc/MapErrors.h
+++ b/Engine/Source/Runtime/Core/Public/Misc/MapErrors.h
@@ -179,7 +179,6 @@ struct CORE_API FMapErrors
 	 /** {Count} element(s) with zero triangles in static mesh {StaticMesh}: **TODO** */
     static FName ElementsWithZeroTriangles;
   
-
     /**  Volume */
 
 	/** LevelStreamingVolume is not in the persistent level - please delete: This warning is caused when there is a level streaming volume that does not exist in the persistent level.  This can be problematic because the volume will not be considered when checking to see if a streaming level should be loaded or unloaded.  You can fix this problem by deleting the level streaming volume and recreating it. */
@@ -201,6 +200,17 @@ struct CORE_API FMapErrors
 
 	/** {Actor} : Repaired painted vertex colors: Painted vertex colors were repaired on this actor. */
     static FName RepairedPaintedVertexColors;
+
+	/** Hierarchical LOD **/
+
+	/** {LODActor} : Static mesh is missing for the built LODActor.  Did you remove the asset? */
+	static FName LODActorMissingStaticMesh;
+
+	/** {LODActor} : Actor is missing. The actor might have been removed. We recommend you to build LOD again. */
+	static FName LODActorMissingActor;
+
+	/** {LODActor} : NoActor is assigned. We recommend to delete this actor. */
+	static FName LODActorNoActorFound;
 };
 
 /**
diff --git a/Engine/Source/Runtime/Engine/Classes/Components/PrimitiveComponent.h b/Engine/Source/Runtime/Engine/Classes/Components/PrimitiveComponent.h
--- a/Engine/Source/Runtime/Engine/Classes/Components/PrimitiveComponent.h
+++ b/Engine/Source/Runtime/Engine/Classes/Components/PrimitiveComponent.h
@@ -987,14 +987,14 @@ public:
 	FThreadSafeCounter AttachmentCounter;
 
 	// Scene data
-
-	/** Replacement primitive to draw instead of this one (multiple UPrim's will point to the same Replacement) */
+private:
+	/** LOD parent primitive to draw instead of this one (multiple UPrim's will point to the same LODParent ) */
 	UPROPERTY(duplicatetransient)
-	TLazyObjectPtr<class UPrimitiveComponent> LODParentPrimitive;
+	class UPrimitiveComponent* LODParentPrimitive;
 
 public:
-	UFUNCTION(BlueprintCallable, Category="Rendering|LOD")
 	void SetLODParentPrimitive(UPrimitiveComponent * InLODParentPrimitive);
+	UPrimitiveComponent* GetLODParentPrimitive();
 
 #if WITH_EDITOR
 	virtual const int32 GetNumUncachedStaticLightingInteractions() const override; // recursive function
diff --git a/Engine/Source/Runtime/Engine/Classes/Engine/LODActor.h b/Engine/Source/Runtime/Engine/Classes/Engine/LODActor.h
--- a/Engine/Source/Runtime/Engine/Classes/Engine/LODActor.h
+++ b/Engine/Source/Runtime/Engine/Classes/Engine/LODActor.h
@@ -10,14 +10,14 @@
  * @see https://docs.unrealengine.com/latest/INT/Engine/Actors/LODActor/
  * @see UStaticMesh
  */
-UCLASS()
+UCLASS(notplaceable)
 class ENGINE_API ALODActor : public AActor
 {
 	GENERATED_UCLASS_BODY()
 
 private_subobject:
 	// disable display of this component
-	UPROPERTY(Category=LODActor, VisibleAnywhere)
+	UPROPERTY(/*Category=LODActor, VisibleAnywhere*/)
 	class UStaticMeshComponent* StaticMeshComponent;
 
 public:
diff --git a/Engine/Source/Runtime/Engine/Private/Components/ActorComponent.cpp b/Engine/Source/Runtime/Engine/Private/Components/ActorComponent.cpp
--- a/Engine/Source/Runtime/Engine/Private/Components/ActorComponent.cpp
+++ b/Engine/Source/Runtime/Engine/Private/Components/ActorComponent.cpp
@@ -69,26 +69,6 @@ FGlobalComponentReregisterContext::FGlobalComponentReregisterContext(const TArra
 
 	// wait until resources are released
 	FlushRenderingCommands();
-
-	// Detach only actor components that are children of the actors list provided
-	for (auto* Component : TObjectRange<UActorComponent>())
-	{
-		bool bShouldReregister=false;
-		UPrimitiveComponent* PrimitiveComponent = Cast<UPrimitiveComponent>(Component);
-		if (PrimitiveComponent && PrimitiveComponent->LODParentPrimitive.Get())
-		{
-			UPrimitiveComponent* LODParentPrimitive = PrimitiveComponent->LODParentPrimitive.Get();
-			AActor* ParentActor = Cast<AActor>(LODParentPrimitive->GetOuter());
-			if (ParentActor && InParentActors.Contains(ParentActor))
-			{
-				bShouldReregister = true;
-			}
-		}
-		if( bShouldReregister )
-		{
-			new(ComponentContexts) FComponentReregisterContext(Component);		
-		}
-	}
 }
 
 FGlobalComponentReregisterContext::~FGlobalComponentReregisterContext()
diff --git a/Engine/Source/Runtime/Engine/Private/Components/PrimitiveComponent.cpp b/Engine/Source/Runtime/Engine/Private/Components/PrimitiveComponent.cpp
--- a/Engine/Source/Runtime/Engine/Private/Components/PrimitiveComponent.cpp
+++ b/Engine/Source/Runtime/Engine/Private/Components/PrimitiveComponent.cpp
@@ -2561,6 +2561,11 @@ void UPrimitiveComponent::SetLODParentPrimitive(UPrimitiveComponent * InLODParen
 	LODParentPrimitive = InLODParentPrimitive;
 	MarkRenderStateDirty();
 }
+
+UPrimitiveComponent* UPrimitiveComponent::GetLODParentPrimitive()
+{
+	return LODParentPrimitive;
+}
 #if WITH_EDITOR
 const int32 UPrimitiveComponent::GetNumUncachedStaticLightingInteractions() const
 {
diff --git a/Engine/Source/Runtime/Engine/Private/LODActor.cpp b/Engine/Source/Runtime/Engine/Private/LODActor.cpp
--- a/Engine/Source/Runtime/Engine/Private/LODActor.cpp
+++ b/Engine/Source/Runtime/Engine/Private/LODActor.cpp
@@ -75,15 +75,42 @@ void ALODActor::CheckForErrors()
 			->AddToken(FTextToken::Create(LOCTEXT( "MapCheck_Message_StaticMeshComponent", "Static mesh actor has NULL StaticMeshComponent property - please delete" ) ))
 			->AddToken(FMapErrorToken::Create(FMapErrors::StaticMeshComponent));
 	}
-	else if( StaticMeshComponent->StaticMesh == NULL )
+
+	if(StaticMeshComponent && StaticMeshComponent->StaticMesh == NULL)
 	{
-		MapCheck.Warning()
+		FFormatNamedArguments Arguments;
+		Arguments.Add(TEXT("ActorName"), FText::FromString(GetName()));
+		FMessageLog("MapCheck").Error()
 			->AddToken(FUObjectToken::Create(this))
-			->AddToken(FTextToken::Create(LOCTEXT( "MapCheck_Message_StaticMeshNull", "Static mesh actor has NULL StaticMesh property" ) ))
-			->AddToken(FMapErrorToken::Create(FMapErrors::StaticMeshNull));
+			->AddToken(FTextToken::Create(FText::Format(LOCTEXT("MapCheck_Message_InvalidLODActorMissingMesh", "{ActorName} : Static mesh is missing for the built LODActor.  Did you remove the asset? Please delete it and build LOD again. "), Arguments)))
+			->AddToken(FMapErrorToken::Create(FMapErrors::LODActorMissingStaticMesh));
 	}
 
-	// @todo error message when missing Actors
+	if (SubActors.Num() == 0)
+	{
+		FFormatNamedArguments Arguments;
+		Arguments.Add(TEXT("ActorName"), FText::FromString(GetName()));
+		FMessageLog("MapCheck").Error()
+			->AddToken(FUObjectToken::Create(this))
+			->AddToken(FTextToken::Create(FText::Format(LOCTEXT("MapCheck_Message_InvalidLODActorEmptyActor", "{ActorName} : NoActor is assigned. We recommend you to delete this actor. "), Arguments)))
+			->AddToken(FMapErrorToken::Create(FMapErrors::LODActorNoActorFound));
+	}
+	else
+	{
+		for(auto& Actor : SubActors)
+		{
+			// see if it's null, if so it is not good
+			if(Actor == nullptr)
+			{
+				FFormatNamedArguments Arguments;
+				Arguments.Add(TEXT("ActorName"), FText::FromString(GetName()));
+				FMessageLog("MapCheck").Error()
+					->AddToken(FUObjectToken::Create(this))
+					->AddToken(FTextToken::Create(FText::Format(LOCTEXT("MapCheck_Message_InvalidLODActorNullActor", "{ActorName} : Actor is missing. The actor might have been removed. We recommend you to build LOD again. "), Arguments)))
+					->AddToken(FMapErrorToken::Create(FMapErrors::LODActorMissingActor));
+			}
+		}
+	}
 }
 
 #endif // WITH_EDITOR
diff --git a/Engine/Source/Runtime/Renderer/Private/PrimitiveSceneInfo.cpp b/Engine/Source/Runtime/Renderer/Private/PrimitiveSceneInfo.cpp
--- a/Engine/Source/Runtime/Renderer/Private/PrimitiveSceneInfo.cpp
+++ b/Engine/Source/Runtime/Renderer/Private/PrimitiveSceneInfo.cpp
@@ -123,11 +123,10 @@ FPrimitiveSceneInfo::FPrimitiveSceneInfo(UPrimitiveComponent* InComponent,FScene
 	}
 
 	// set LOD parent info if exists
-	if (InComponent->LODParentPrimitive.IsValid())
+	UPrimitiveComponent* LODParent = InComponent->GetLODParentPrimitive();
+	if (LODParent)
 	{
-		// I think this should be checked in editor feature
-		check (InComponent->LODParentPrimitive.Get() != InComponent);
-		LODParentComponentId = InComponent->LODParentPrimitive.Get()->ComponentId;
+		LODParentComponentId = LODParent->ComponentId;
 	}
 }
 
