Don't hardcode platform editor modules

diff --git a/Engine/Source/Editor/UnrealEd/Private/EditorEngine.cpp b/Engine/Source/Editor/UnrealEd/Private/EditorEngine.cpp
--- a/Engine/Source/Editor/UnrealEd/Private/EditorEngine.cpp
+++ b/Engine/Source/Editor/UnrealEd/Private/EditorEngine.cpp
@@ -628,11 +628,6 @@ void UEditorEngine::Init(IEngineLoop* InEngineLoop)
 			TEXT("SettingsEditor"),
 			TEXT("EditorSettingsViewer"),
 			TEXT("ProjectSettingsViewer"),
-			TEXT("AndroidRuntimeSettings"),
-			TEXT("AndroidPlatformEditor"),
-			TEXT("HTML5PlatformEditor"),
-			TEXT("IOSRuntimeSettings"),
-			TEXT("IOSPlatformEditor"),
 			TEXT("Blutility"),
 			TEXT("OnlineBlueprintSupport"),
 			TEXT("XmlParser"),
@@ -656,6 +651,31 @@ void UEditorEngine::Init(IEngineLoop* InEngineLoop)
 			FModuleManager::Get().LoadModule(ModuleName);
 		}
 
+		{
+			// Load platform runtime settings modules
+			TArray<FName> Modules;
+			FModuleManager::Get().FindModules( TEXT( "*RuntimeSettings" ), Modules );
+
+			for( int32 Index = 0; Index < Modules.Num(); Index++ )
+			{
+				FModuleManager::Get().LoadModule( Modules[Index] );
+			}
+		}
+
+		{
+			// Load platform editor modules
+			TArray<FName> Modules;
+			FModuleManager::Get().FindModules( TEXT( "*PlatformEditor" ), Modules );
+
+			for( int32 Index = 0; Index < Modules.Num(); Index++ )
+			{
+				if( Modules[Index] != TEXT("ProjectTargetPlatformEditor") )
+				{
+					FModuleManager::Get().LoadModule( Modules[Index] );
+				}
+			}
+		}
+
 		if (!IsRunningCommandlet())
 		{
 			FModuleManager::Get().LoadModule(TEXT("EditorLiveStreaming"));
