Fixed a crash that was occurring when during the build lighting portion of the promotion and launch on test.

diff --git a/Engine/Source/Editor/UnrealEd/Private/Tests/AutomationEditorCommon.cpp b/Engine/Source/Editor/UnrealEd/Private/Tests/AutomationEditorCommon.cpp
--- a/Engine/Source/Editor/UnrealEd/Private/Tests/AutomationEditorCommon.cpp
+++ b/Engine/Source/Editor/UnrealEd/Private/Tests/AutomationEditorCommon.cpp
@@ -695,6 +695,9 @@ bool FBuildLightingCommand::Update()
 		return true;
 	}
 
+	UWorld* CurrentWorld = GEditor->GetEditorWorldContext().World();
+	GUnrealEd->Exec(CurrentWorld, TEXT("MAP REBUILD"));
+
 	FLightingBuildOptions LightingBuildOptions;
 
 	// Retrieve settings from ini.
diff --git a/Engine/Source/Editor/UnrealEd/Private/Tests/EditorBuildPromotionTests.cpp b/Engine/Source/Editor/UnrealEd/Private/Tests/EditorBuildPromotionTests.cpp
--- a/Engine/Source/Editor/UnrealEd/Private/Tests/EditorBuildPromotionTests.cpp
+++ b/Engine/Source/Editor/UnrealEd/Private/Tests/EditorBuildPromotionTests.cpp
@@ -919,6 +919,9 @@ namespace EditorBuildPromotionTestUtils
 			return;
 		}
 
+		UWorld* CurrentWorld = GEditor->GetEditorWorldContext().World();
+		GUnrealEd->Exec(CurrentWorld, TEXT("MAP REBUILD"));
+
 		FLightingBuildOptions LightingBuildOptions;
 
 		// Retrieve settings from ini.
