Workaround NULL PackageAutoSaver

diff --git a/Engine/Source/Editor/UnrealEd/Private/UnrealEdEngine.cpp b/Engine/Source/Editor/UnrealEd/Private/UnrealEdEngine.cpp
--- a/Engine/Source/Editor/UnrealEd/Private/UnrealEdEngine.cpp
+++ b/Engine/Source/Editor/UnrealEd/Private/UnrealEdEngine.cpp
@@ -657,7 +657,12 @@ void UUnrealEdEngine::OnOpenMatinee()
 
 bool UUnrealEdEngine::IsAutosaving() const
 {
-	return PackageAutoSaver->IsAutoSaving();
+	if (PackageAutoSaver)
+	{
+		return PackageAutoSaver->IsAutoSaving();
+	}
+	
+	return NULL;
 }
 
 
