Fix compile error

diff --git a/Engine/Source/Editor/UnrealEd/Private/UnrealEdEngine.cpp b/Engine/Source/Editor/UnrealEd/Private/UnrealEdEngine.cpp
--- a/Engine/Source/Editor/UnrealEd/Private/UnrealEdEngine.cpp
+++ b/Engine/Source/Editor/UnrealEd/Private/UnrealEdEngine.cpp
@@ -662,7 +662,7 @@ bool UUnrealEdEngine::IsAutosaving() const
 		return PackageAutoSaver->IsAutoSaving();
 	}
 	
-	return NULL;
+	return false;
 }
 
 
