Added the ability to pause / abort auto-reimport operations

diff --git a/Engine/Source/Editor/EditorStyle/Private/SlateEditorStyle.cpp b/Engine/Source/Editor/EditorStyle/Private/SlateEditorStyle.cpp
--- a/Engine/Source/Editor/EditorStyle/Private/SlateEditorStyle.cpp
+++ b/Engine/Source/Editor/EditorStyle/Private/SlateEditorStyle.cpp
@@ -683,6 +683,10 @@ void FSlateEditorStyle::FStyle::SetupGeneralStyles()
 	Set( "GenericLock", new IMAGE_BRUSH( "Icons/padlock_locked_16x", Icon16x16 ) );
 	Set( "GenericUnlock", new IMAGE_BRUSH( "Icons/padlock_unlocked_16x", Icon16x16 ) );
 
+	Set( "GenericPlay", new IMAGE_BRUSH( "Icons/generic_play_16x", Icon16x16 ) );
+	Set( "GenericPause", new IMAGE_BRUSH( "Icons/generic_pause_16x", Icon16x16 ) );
+	Set( "GenericStop", new IMAGE_BRUSH( "Icons/generic_stop_16x", Icon16x16 ) );
+
 	Set( "SoftwareCursor_Grab", new IMAGE_BRUSH( "Icons/cursor_grab", Icon16x16 ) );
 	Set( "SoftwareCursor_CardinalCross", new IMAGE_BRUSH( "Icons/cursor_cardinal_cross", Icon24x24 ) );
 	Set( "SoftwareCursor_UpDown", new IMAGE_BRUSH( "Icons/cursor_updown", Icon16x20 ) );
@@ -1332,7 +1336,7 @@ void FSlateEditorStyle::FStyle::SetupGeneralStyles()
 		Set( "NotificationList.FailImage", new IMAGE_BRUSH( "Old/PropertyEditor/Button_Clear", Icon16x16 ) );
 		Set( "NotificationList.DefaultMessage", new IMAGE_BRUSH( "Old/EventMessage_Default", Icon40x40 ) );
 		Set( "NotificationList.Glow", new FSlateColorBrush( FColor(255, 255, 255, 255) ) );
-		}
+	}
 #endif // WITH_EDITOR || IS_PROGRAM
 
 #if WITH_EDITOR || IS_PROGRAM
diff --git a/Engine/Source/Editor/UnrealEd/Private/AutoReimport/AutoReimportManager.cpp b/Engine/Source/Editor/UnrealEd/Private/AutoReimport/AutoReimportManager.cpp
--- a/Engine/Source/Editor/UnrealEd/Private/AutoReimport/AutoReimportManager.cpp
+++ b/Engine/Source/Editor/UnrealEd/Private/AutoReimport/AutoReimportManager.cpp
@@ -83,7 +83,7 @@ private:
 /** Enum and value to specify the current state of our processing */
 enum class ECurrentState
 {
-	Idle, PendingResponse, ProcessAdditions, ProcessModifications, ProcessDeletions, SavePackages
+	Idle, Paused, Aborting, ProcessAdditions, ProcessModifications, ProcessDeletions, SavePackages
 };
 uint32 GetTypeHash(ECurrentState State)
 {
@@ -169,11 +169,17 @@ private:
 	TOptional<ECurrentState> ProcessDeletions();
 
 	/** Wait for a user's input. Just updates the progress text for now */
-	TOptional<ECurrentState> PendingResponse();
+	TOptional<ECurrentState> Paused();
+
+	/** Abort the process */
+	TOptional<ECurrentState> Abort();
 
 	/** Cleanup an operation that just processed some changes */
 	void Cleanup();
 
+	/** Check whether we should pause the operation or not */
+	TOptional<ECurrentState> HandlePauseAbort(ECurrentState InCurrentState);
+
 private:
 
 	/** Feedback context that can selectively override the global context to consume progress events for saving of assets */
@@ -190,11 +196,32 @@ private:
 
 	/** A timeout used to refresh directory monitors when the user has made an interactive change to the settings */
 	FTimeLimit ResetMonitorsTimeout;
+
+	/** The paused state of the state machine */
+	ECurrentState PausedState;
+
+private:
+
+	void OnPauseClicked()
+	{
+		switch (State)
+		{
+		case EProcessState::Paused: State = EProcessState::Running; break;
+		case EProcessState::Running: State = EProcessState::Paused; break;
+		default: break;
+		}
+	}
+	void OnAbortClicked() { State = EProcessState::Aborted; }
+
+	/** Flags for paused/aborted */
+	enum class EProcessState { Running, Paused, Aborted };
+	EProcessState State;
 };
 
 FAutoReimportManager::FAutoReimportManager()
 	: StateMachine(ECurrentState::Idle)
 	, bGuardAssetChanges(false)
+	, PausedState(ECurrentState::Idle)
 {
 	auto* Settings = GetMutableDefault<UEditorLoadingSavingSettings>();
 	Settings->OnSettingChanged().AddRaw(this, &FAutoReimportManager::HandleLoadingSavingSettingChanged);
@@ -218,7 +245,8 @@ FAutoReimportManager::FAutoReimportManager()
 	}
 
 	StateMachine.Add(ECurrentState::Idle,					EStateMachineNode::CallOnce,	[this](const FTimeLimit& T){ return this->Idle();					});
-	StateMachine.Add(ECurrentState::PendingResponse,		EStateMachineNode::CallOnce,	[this](const FTimeLimit& T){ return this->PendingResponse();		});
+	StateMachine.Add(ECurrentState::Paused,					EStateMachineNode::CallOnce,	[this](const FTimeLimit& T){ return this->Paused();					});
+	StateMachine.Add(ECurrentState::Aborting,				EStateMachineNode::CallOnce,	[this](const FTimeLimit& T){ return this->Abort();					});
 	StateMachine.Add(ECurrentState::ProcessAdditions,		EStateMachineNode::CallMany,	[this](const FTimeLimit& T){ return this->ProcessAdditions(T);		});
 	StateMachine.Add(ECurrentState::ProcessModifications,	EStateMachineNode::CallMany,	[this](const FTimeLimit& T){ return this->ProcessModifications(T);	});
 	StateMachine.Add(ECurrentState::ProcessDeletions,		EStateMachineNode::CallMany,	[this](const FTimeLimit& T){ return this->ProcessDeletions();		});
@@ -413,12 +441,16 @@ FText FAutoReimportManager::GetProgressText() const
 
 TOptional<ECurrentState> FAutoReimportManager::ProcessAdditions(const FTimeLimit& TimeLimit)
 {
+	auto NewState = HandlePauseAbort(ECurrentState::ProcessAdditions);
+	if (NewState.IsSet())
+	{
+		return NewState;
+	}
+
 	// Override the global feedback context while we do this to avoid popping up dialogs
 	TGuardValue<FFeedbackContext*> ScopedContextOverride(GWarn, FeedbackContextOverride.Get());
 	TGuardValue<bool> ScopedAssetChangesGuard(bGuardAssetChanges, true);
 
-	FeedbackContextOverride->GetContent()->SetMainText(GetProgressText());
-
 	TMap<FString, TArray<UFactory*>> Factories;
 
 	TArray<FString> FactoryExtensions;
@@ -464,12 +496,16 @@ TOptional<ECurrentState> FAutoReimportManager::ProcessAdditions(const FTimeLimit
 
 TOptional<ECurrentState> FAutoReimportManager::ProcessModifications(const FTimeLimit& TimeLimit)
 {
+	auto NewState = HandlePauseAbort(ECurrentState::ProcessModifications);
+	if (NewState.IsSet())
+	{
+		return NewState;
+	}
+
 	// Override the global feedback context while we do this to avoid popping up dialogs
 	TGuardValue<FFeedbackContext*> ScopedContextOverride(GWarn, FeedbackContextOverride.Get());
 	TGuardValue<bool> ScopedAssetChangesGuard(bGuardAssetChanges, true);
 
-	FeedbackContextOverride->GetContent()->SetMainText(GetProgressText());
-
 	const IAssetRegistry& Registry = FModuleManager::LoadModuleChecked<FAssetRegistryModule>("AssetRegistry").Get();
 
 	for (auto& Monitor : DirectoryMonitors)
@@ -483,12 +519,16 @@ TOptional<ECurrentState> FAutoReimportManager::ProcessModifications(const FTimeL
 
 TOptional<ECurrentState> FAutoReimportManager::ProcessDeletions()
 {
+	auto NewState = HandlePauseAbort(ECurrentState::ProcessDeletions);
+	if (NewState.IsSet())
+	{
+		return NewState;
+	}
+
 	TGuardValue<bool> ScopedAssetChangesGuard(bGuardAssetChanges, true);
 
 	const IAssetRegistry& Registry = FModuleManager::LoadModuleChecked<FAssetRegistryModule>("AssetRegistry").Get();
 
-	FeedbackContextOverride->GetContent()->SetMainText(GetProgressText());
-
 	TArray<FAssetData> AssetsToDelete;
 
 	for (auto& Monitor : DirectoryMonitors)
@@ -516,8 +556,6 @@ TOptional<ECurrentState> FAutoReimportManager::SavePackages()
 
 	TGuardValue<bool> ScopedAssetChangesGuard(bGuardAssetChanges, true);
 
-	FeedbackContextOverride->GetContent()->SetMainText(GetProgressText());
-
 	if (PackagesToSave.Num() > 0)
 	{
 		const bool bAlreadyCheckedOut = false;
@@ -528,19 +566,48 @@ TOptional<ECurrentState> FAutoReimportManager::SavePackages()
 		PackagesToSave.Empty();
 	}
 
-	// Make sure the progress text is up to date before we close the notification
-	FeedbackContextOverride->GetContent()->SetMainText(GetProgressText());
-
 	Cleanup();
 	return ECurrentState::Idle;
 }
 
-TOptional<ECurrentState> FAutoReimportManager::PendingResponse()
+TOptional<ECurrentState> FAutoReimportManager::HandlePauseAbort(ECurrentState InCurrentState)
 {
-	FeedbackContextOverride->GetContent()->SetMainText(GetProgressText());
+	if (State == EProcessState::Aborted)
+	{
+		return ECurrentState::Aborting;
+	}
+	else if (State == EProcessState::Paused)
+	{
+		PausedState = InCurrentState;
+		return ECurrentState::Paused;
+	}
+
 	return TOptional<ECurrentState>();
 }
 
+TOptional<ECurrentState> FAutoReimportManager::Paused()
+{
+	TOptional<ECurrentState> NewState = HandlePauseAbort(PausedState);
+	if (NewState.IsSet())
+	{
+		return NewState.GetValue();
+	}
+
+	// No longer paused
+	return PausedState;
+}
+
+TOptional<ECurrentState> FAutoReimportManager::Abort()
+{
+	for (auto& Monitor : DirectoryMonitors)
+	{
+		Monitor.Abort();
+	}
+
+	Cleanup();
+	return ECurrentState::Idle;
+}
+
 TOptional<ECurrentState> FAutoReimportManager::Idle()
 {
 	// Check whether we need to reset the monitors or not
@@ -579,6 +646,8 @@ TOptional<ECurrentState> FAutoReimportManager::Idle()
 
 	if (GetNumUnprocessedChanges() > 0)
 	{
+		State = EProcessState::Running;
+
 		// We have some changes so kick off the process
 		for (auto& Monitor : DirectoryMonitors)
 		{
@@ -592,10 +661,17 @@ TOptional<ECurrentState> FAutoReimportManager::Idle()
 
 		if (WorkTotal > 0)
 		{
-			// Create a new feedback context override
-			FeedbackContextOverride = MakeShareable(new FReimportFeedbackContext);
-			FeedbackContextOverride->Initialize(SNew(SReimportFeedback, GetProgressText()));
+			if (!FeedbackContextOverride.IsValid())
+			{
+				// Create a new feedback context override
+				FeedbackContextOverride = MakeShareable(new FReimportFeedbackContext(
+					TAttribute<FText>(this, &FAutoReimportManager::GetProgressText),
+					FSimpleDelegate::CreateSP(this, &FAutoReimportManager::OnPauseClicked),
+					FSimpleDelegate::CreateSP(this, &FAutoReimportManager::OnAbortClicked))
+				);
+			}
 
+			FeedbackContextOverride->Show();
 			return ECurrentState::ProcessAdditions;
 		}
 	}
@@ -605,8 +681,10 @@ TOptional<ECurrentState> FAutoReimportManager::Idle()
 
 void FAutoReimportManager::Cleanup()
 {
-	FeedbackContextOverride->Destroy();
-	FeedbackContextOverride = nullptr;
+	if (GetNumUnprocessedChanges() == 0)
+	{
+		FeedbackContextOverride->Hide();
+	}
 }
 
 void FAutoReimportManager::Tick(float DeltaTime)
diff --git a/Engine/Source/Editor/UnrealEd/Private/AutoReimport/AutoReimportUtilities.h b/Engine/Source/Editor/UnrealEd/Private/AutoReimport/AutoReimportUtilities.h
--- a/Engine/Source/Editor/UnrealEd/Private/AutoReimport/AutoReimportUtilities.h
+++ b/Engine/Source/Editor/UnrealEd/Private/AutoReimport/AutoReimportUtilities.h
@@ -74,8 +74,11 @@ struct FTimeLimit
 	/** Constructor specifying not to run over the specified number of seconds */
 	FTimeLimit(float NumSeconds) : Delay(NumSeconds) { Reset(); }
 
+	/** Return true if this time limit has been set up, false otherwise */
+	bool IsValid() const { return Delay != -1; }
+	
 	/** Check whether we have exceeded the time limit */
-	bool Exceeded() const { return Delay != -1 && FPlatformTime::Seconds() >= StartTime + Delay; }
+	bool Exceeded() const { return IsValid() && FPlatformTime::Seconds() >= StartTime + Delay; }
 
 	/** Reset the time limit to start timing again from the current time */
 	void Reset() { StartTime = FPlatformTime::Seconds(); }
diff --git a/Engine/Source/Editor/UnrealEd/Private/AutoReimport/ContentDirectoryMonitor.cpp b/Engine/Source/Editor/UnrealEd/Private/AutoReimport/ContentDirectoryMonitor.cpp
--- a/Engine/Source/Editor/UnrealEd/Private/AutoReimport/ContentDirectoryMonitor.cpp
+++ b/Engine/Source/Editor/UnrealEd/Private/AutoReimport/ContentDirectoryMonitor.cpp
@@ -392,4 +392,32 @@ void FContentDirectoryMonitor::ExtractAssetsToDelete(const IAssetRegistry& Regis
 	DeletedFiles.Empty();
 }
 
+void FContentDirectoryMonitor::Abort()
+{
+	for (auto& Add : AddedFiles)
+	{
+		Cache.CompleteTransaction(MoveTemp(Add));
+	}
+	AddedFiles.Empty();
+
+	for (auto& Mod : ModifiedFiles)
+	{
+		Cache.CompleteTransaction(MoveTemp(Mod));
+	}
+	ModifiedFiles.Empty();
+
+	for (auto& Del : DeletedFiles)
+	{
+		Cache.CompleteTransaction(MoveTemp(Del));
+	}
+	DeletedFiles.Empty();
+
+	for (auto& Change : Cache.GetOutstandingChanges())
+	{
+		Cache.CompleteTransaction(MoveTemp(Change));
+	}
+
+	WorkProgress = TotalWork;
+}
+
 #undef LOCTEXT_NAMESPACE
diff --git a/Engine/Source/Editor/UnrealEd/Private/AutoReimport/ContentDirectoryMonitor.h b/Engine/Source/Editor/UnrealEd/Private/AutoReimport/ContentDirectoryMonitor.h
--- a/Engine/Source/Editor/UnrealEd/Private/AutoReimport/ContentDirectoryMonitor.h
+++ b/Engine/Source/Editor/UnrealEd/Private/AutoReimport/ContentDirectoryMonitor.h
@@ -34,6 +34,9 @@ public:
 	/** Extract the assets we need to delete from our outstanding changes (happens last) */ 
 	void ExtractAssetsToDelete(const IAssetRegistry& Registry, TArray<FAssetData>& OutAssetsToDelete);
 
+	/** Abort the current processing operation */
+	void Abort();
+
 	/** Report an external change to the manager, such that a subsequent equal change reported by the os be ignored */
 	void IgnoreNewFile(const FString& Filename);
 	void IgnoreFileModification(const FString& Filename);
diff --git a/Engine/Source/Editor/UnrealEd/Private/AutoReimport/ReimportFeedbackContext.cpp b/Engine/Source/Editor/UnrealEd/Private/AutoReimport/ReimportFeedbackContext.cpp
--- a/Engine/Source/Editor/UnrealEd/Private/AutoReimport/ReimportFeedbackContext.cpp
+++ b/Engine/Source/Editor/UnrealEd/Private/AutoReimport/ReimportFeedbackContext.cpp
@@ -9,6 +9,8 @@
 #include "MessageLogModule.h"
 #include "SHyperlink.h"
 
+#define LOCTEXT_NAMESPACE "ReimportContext"
+
 class SWidgetStack : public SVerticalBox
 {
 	SLATE_BEGIN_ARGS(SWidgetStack){}
@@ -216,104 +218,274 @@ class SWidgetStack : public SVerticalBox
 	int32 MaxNumVisible;
 };
 
-void SReimportFeedback::Disable()
-{
-	WidgetStack->SetVisibility(EVisibility::HitTestInvisible);
-}
+class SWidgetStack;
 
-void SReimportFeedback::Add(const TSharedRef<SWidget>& Widget)
+/** Feedback context that overrides GWarn for import operations to prevent popup spam */
+class SReimportFeedback : public SCompoundWidget
 {
-	WidgetStack->Add(Widget);
-}
+public:
+	SLATE_BEGIN_ARGS(SReimportFeedback) : _ExpireDuration(3.f) {}
 
-void SReimportFeedback::SetMainText(FText InText)
-{
-	MainText = InText;
-}
+		SLATE_ATTRIBUTE(FText, MainText)
 
-FText SReimportFeedback::GetMainText() const
-{
-	return MainText;
-}
+		SLATE_ARGUMENT(float, ExpireDuration)
+		SLATE_EVENT(FSimpleDelegate, OnExpired)
 
-EVisibility SReimportFeedback::GetHyperlinkVisibility() const
-{
-	return WidgetStack->NumSlots() != 0 ? EVisibility::Visible : EVisibility::Collapsed;
-}
+		SLATE_EVENT(FSimpleDelegate, OnPauseClicked)
+		SLATE_EVENT(FSimpleDelegate, OnAbortClicked)
 
-void SReimportFeedback::Construct(const FArguments& InArgs, FText InMainText)
-{
-	MainText = InMainText;
+	SLATE_END_ARGS()
 
-	ChildSlot
-	[
-		SNew(SBorder)
-		.Padding(FMargin(10))
-		.BorderImage(FCoreStyle::Get().GetBrush("NotificationList.ItemBackground"))
-		[
-			SNew(SVerticalBox)
+	/** Construct this widget */
+	void Construct(const FArguments& InArgs)
+	{
+		ExpireDuration = InArgs._ExpireDuration;
+		OnExpired = InArgs._OnExpired;
 
-			+ SVerticalBox::Slot()
-			.Padding(FMargin(0))
-			.AutoHeight()
-			[
-				SNew(STextBlock)
-				.Text(this, &SReimportFeedback::GetMainText)
-				.Font(FCoreStyle::Get().GetFontStyle(TEXT("NotificationList.FontBold")))
-			]
+		MainText = InArgs._MainText;
+		bPaused = false;
+		bExpired = false;
 
-			+ SVerticalBox::Slot()
-			.Padding(FMargin(0, 5, 0, 0))
-			.AutoHeight()
-			[
-				SAssignNew(WidgetStack, SWidgetStack, 3)
-			]
+		auto OpenMessageLog = []{
+			FMessageLogModule& MessageLogModule = FModuleManager::LoadModuleChecked<FMessageLogModule>("MessageLog");
+			MessageLogModule.OpenMessageLog("AssetReimport");
+		};
 
-			+ SVerticalBox::Slot()
-			.Padding(FMargin(0, 5, 0, 0))
-			.AutoHeight()
-			.HAlign(HAlign_Right)
+		ChildSlot
+		[
+			SNew(SBorder)
+			.Padding(FMargin(10))
+			.BorderImage(FCoreStyle::Get().GetBrush("NotificationList.ItemBackground"))
 			[
-				SNew(SHyperlink)
-				.Visibility(this, &SReimportFeedback::GetHyperlinkVisibility)
-				.Text(NSLOCTEXT("ReimportContext", "OpenMessageLog", "Open message log"))
-				.TextStyle(FCoreStyle::Get(), "SmallText")
-				.OnNavigate_Lambda([]{
-					FMessageLogModule& MessageLogModule = FModuleManager::LoadModuleChecked<FMessageLogModule>("MessageLog");
-					MessageLogModule.OpenMessageLog("AssetReimport");
-				})
+				SNew(SVerticalBox)
+
+				+ SVerticalBox::Slot()
+				.Padding(FMargin(0))
+				.AutoHeight()
+				[
+					SNew(SHorizontalBox)
+
+					+ SHorizontalBox::Slot()
+					[
+						SNew(STextBlock)
+						.Text(this, &SReimportFeedback::GetMainText)
+						.Font(FCoreStyle::Get().GetFontStyle(TEXT("NotificationList.FontLight")))
+					]
+
+					+ SHorizontalBox::Slot()
+					.AutoWidth()
+					.VAlign(VAlign_Center)
+					.Padding(FMargin(0,0,4,0))
+					[
+						SAssignNew(PauseButton, SButton)
+						.ButtonStyle(FEditorStyle::Get(), "HoverHintOnly")
+						.ToolTipText(LOCTEXT("PauseTooltip", "Temporarily pause processing of these source content files"))
+						.OnClicked(this, &SReimportFeedback::OnPauseClicked, InArgs._OnPauseClicked)
+						[
+							SNew(SImage)
+							.ColorAndOpacity(FLinearColor(.8f,.8f,.8f,1.f))
+							.Image(this, &SReimportFeedback::GetPlayPauseBrush)
+						]
+					]
+
+					+ SHorizontalBox::Slot()
+					.AutoWidth()
+					.VAlign(VAlign_Center)
+					[
+						SAssignNew(AbortButton, SButton)
+						.ButtonStyle(FEditorStyle::Get(), "HoverHintOnly")
+						.ToolTipText(LOCTEXT("AbortTooltip", "Permanently abort processing of these source content files"))
+						.OnClicked(this, &SReimportFeedback::OnAbortClicked, InArgs._OnAbortClicked)
+						[
+							SNew(SImage)
+							.ColorAndOpacity(FLinearColor(0.8f,0.8f,0.8f,1.f))
+							.Image(FEditorStyle::GetBrush("GenericStop"))
+						]
+					]
+				]
+
+				+ SVerticalBox::Slot()
+				.Padding(FMargin(0, 5, 0, 0))
+				.AutoHeight()
+				[
+					SAssignNew(WidgetStack, SWidgetStack, 3)
+				]
+
+				+ SVerticalBox::Slot()
+				.Padding(FMargin(0, 5, 0, 0))
+				.AutoHeight()
+				.HAlign(HAlign_Right)
+				[
+					SNew(SHyperlink)
+					.Visibility(this, &SReimportFeedback::GetHyperlinkVisibility)
+					.Text(LOCTEXT("OpenMessageLog", "Open message log"))
+					.TextStyle(FCoreStyle::Get(), "SmallText")
+					.OnNavigate_Lambda(OpenMessageLog)
+				]
 			]
-		]
-	];
-}
+		];
+	}
+
+	/** Add a widget to this feedback's widget stack */
+	void Add(const TSharedRef<SWidget>& Widget)
+	{
+		WidgetStack->Add(Widget);
+	}
+
+	/** Disable input to this widget's dynamic content (except the message log hyperlink) */
+	void Disable()
+	{
+		// Get the final text to display
+		CachedMainText = MainText.Get(FText());
+
+		ExpireTimeout = FTimeLimit(ExpireDuration);
 
-FReimportFeedbackContext::FReimportFeedbackContext()
+		WidgetStack->SetVisibility(EVisibility::HitTestInvisible);
+		PauseButton->SetVisibility(EVisibility::Collapsed);
+		AbortButton->SetVisibility(EVisibility::Collapsed);
+	}
+
+	/** Enable, if previously disabled */
+	void Enable()
+	{
+		ExpireTimeout = FTimeLimit();
+
+		bPaused = false;
+		WidgetStack->SetVisibility(EVisibility::Visible);
+		PauseButton->SetVisibility(EVisibility::Visible);
+		AbortButton->SetVisibility(EVisibility::Visible);
+	}
+
+private:
+	
+	virtual void Tick(const FGeometry& AllottedGeometry, const double InCurrentTime, const float InDeltaTime) override
+	{
+		if (bExpired)
+		{
+			return;
+		}
+		else if(ExpireTimeout.IsValid())
+		{
+			if (ExpireTimeout.Exceeded())
+			{	
+				OnExpired.ExecuteIfBound();
+				bExpired = true;
+			}
+		}
+		else
+		{
+			CachedMainText = MainText.Get(FText());
+		}
+	}
+
+	/** Get the main text of this widget */
+	FText GetMainText() const
+	{
+		return CachedMainText;
+	}
+
+	/** Get the play/pause image */
+	const FSlateBrush* GetPlayPauseBrush() const
+	{
+		return bPaused ? FEditorStyle::GetBrush("GenericPlay") : FEditorStyle::GetBrush("GenericPause");
+	}
+
+	/** Called when pause is clicked */
+	FReply OnPauseClicked(FSimpleDelegate UserOnClicked)
+	{
+		bPaused = !bPaused;
+
+		UserOnClicked.ExecuteIfBound();
+		return FReply::Handled();
+	}
+	
+	/** Called when abort is clicked */
+	FReply OnAbortClicked(FSimpleDelegate UserOnClicked)
+	{
+		//Destroy();
+		UserOnClicked.ExecuteIfBound();
+		return FReply::Handled();
+	}
+
+	/** Get the visibility of the hyperlink to open the message log */
+	EVisibility GetHyperlinkVisibility() const
+	{
+		return WidgetStack->NumSlots() != 0 ? EVisibility::Visible : EVisibility::Collapsed;
+	}
+
+private:
+
+	/** The expire timeout used to fire OnExpired. Invalid when no timeout is set. */
+	FTimeLimit ExpireTimeout;
+
+	/** Amount of time to wait after this widget has been disabled before calling OnExpired */
+	float ExpireDuration;
+
+	/** Event that is called when this widget has been inactive and open for too long, and will fade out */
+	FSimpleDelegate OnExpired;
+
+	/** Cached main text for the notification */
+	FText CachedMainText;
+	TAttribute<FText> MainText;
+
+	/** Whether we are paused and/or expired */
+	bool bPaused, bExpired;
+
+	/** The widget stack, displaying contextural information about the current state of the process */
+	TSharedPtr<SWidgetStack> WidgetStack;
+	TSharedPtr<SButton> PauseButton, AbortButton;
+};
+
+FReimportFeedbackContext::FReimportFeedbackContext(const TAttribute<FText>& InMainText, const FSimpleDelegate& InOnPauseClicked, const FSimpleDelegate& InOnAbortClicked)
 	: bSuppressSlowTaskMessages(false)
+	, MainTextAttribute(InMainText)
+	, OnPauseClickedEvent(InOnPauseClicked)
+	, OnAbortClickedEvent(InOnAbortClicked)
 	, MessageLog("AssetReimport")
-{}
+{
+}
 
-void FReimportFeedbackContext::Initialize(TSharedRef<SReimportFeedback> Widget)
+void FReimportFeedbackContext::Show()
 {
-	NotificationContent = Widget;
+	if (NotificationContent.IsValid())
+	{
+		NotificationContent->Enable();
+	}
+	else
+	{
+		NotificationContent = SNew(SReimportFeedback)
+			.OnExpired(this, &FReimportFeedbackContext::OnNotificationExpired)
+			.MainText(MainTextAttribute)
+			.OnPauseClicked(OnPauseClickedEvent)
+			.OnAbortClicked(OnAbortClickedEvent);
 
-	FNotificationInfo Info(SharedThis(this));
-	Info.ExpireDuration = 3.f;
-	Info.bFireAndForget = false;
+		FNotificationInfo Info(AsShared());
+		Info.bFireAndForget = false;
 
-	Notification = FSlateNotificationManager::Get().AddNotification(Info);
+		Notification = FSlateNotificationManager::Get().AddNotification(Info);
 
-	MessageLog.NewPage(FText::Format(NSLOCTEXT("ReimportContext", "MessageLogPageLabel", "Outstanding source content changes {0}"), FText::AsTime(FDateTime::Now())));
+		MessageLog.NewPage(FText::Format(LOCTEXT("MessageLogPageLabel", "Outstanding source content changes {0}"), FText::AsTime(FDateTime::Now())));
+	}
 }
 
-void FReimportFeedbackContext::Destroy()
+void FReimportFeedbackContext::Hide()
 {
-	MessageLog.Notify(FText(), EMessageSeverity::Error);
-
 	if (Notification.IsValid())
 	{
 		NotificationContent->Disable();
 		Notification->SetCompletionState(SNotificationItem::CS_Success);
-		Notification->ExpireAndFadeout();
+	}
+}
+
+void FReimportFeedbackContext::OnNotificationExpired()
+{
+	if (Notification.IsValid())
+	{
+		MessageLog.Notify(FText(), EMessageSeverity::Error);
+		Notification->Fadeout();
+
+		NotificationContent = nullptr;
+		Notification = nullptr;
 	}
 }
 
@@ -325,15 +497,20 @@ void FReimportFeedbackContext::AddMessage(EMessageSeverity::Type Severity, const
 
 void FReimportFeedbackContext::AddWidget(const TSharedRef<SWidget>& Widget)
 {
-	NotificationContent->Add(Widget);
+	if (NotificationContent.IsValid())
+	{
+		NotificationContent->Add(Widget);
+	}
 }
 
 void FReimportFeedbackContext::StartSlowTask(const FText& Task, bool bShowCancelButton)
 {
 	FFeedbackContext::StartSlowTask(Task, bShowCancelButton);
 
-	if (!bSuppressSlowTaskMessages && !Task.IsEmpty())
+	if (NotificationContent.IsValid() && !bSuppressSlowTaskMessages && !Task.IsEmpty())
 	{
 		NotificationContent->Add(SNew(STextBlock).Text(Task));
 	}
 }
+
+#undef LOCTEXT_NAMESPACE
\ No newline at end of file
diff --git a/Engine/Source/Editor/UnrealEd/Private/AutoReimport/ReimportFeedbackContext.h b/Engine/Source/Editor/UnrealEd/Private/AutoReimport/ReimportFeedbackContext.h
--- a/Engine/Source/Editor/UnrealEd/Private/AutoReimport/ReimportFeedbackContext.h
+++ b/Engine/Source/Editor/UnrealEd/Private/AutoReimport/ReimportFeedbackContext.h
@@ -8,58 +8,18 @@
 #include "MessageLog.h"
 #include "AutoReimportUtilities.h"
 
-class SWidgetStack;
-
-/** Feedback context that overrides GWarn for import operations to prevent popup spam */
-class SReimportFeedback : public SCompoundWidget
-{
-public:
-	SLATE_BEGIN_ARGS(SReimportFeedback){}
-	SLATE_END_ARGS()
-
-	/** Construct this widget */
-	void Construct(const FArguments& InArgs, FText InMainText);
-
-	/** Add a widget to this feedback's widget stack */
-	void Add(const TSharedRef<SWidget>& Widget);
-
-	/** Disable input to this widget's dynamic content (except the message log hyperlink) */
-	void Disable();
-
-	/** Set the main text of this widget */
-	void SetMainText(FText InText);
-	
-	/** Get the main text of this widget */
-	FText GetMainText() const;
-
-private:
-	
-	/** Get the visibility of the hyperlink to open the message log */
-	EVisibility GetHyperlinkVisibility() const;
-
-private:
-
-	/** Cached main text for the notification */
-	FText MainText;
-
-	/** The widget stack, displaying contextural information about the current state of the process */
-	TSharedPtr<SWidgetStack> WidgetStack;
-};
-
 /** Feedback context that overrides GWarn for import operations to prevent popup spam */
 class FReimportFeedbackContext : public FFeedbackContext, public INotificationWidget, public TSharedFromThis<FReimportFeedbackContext>
 {
 public:
 	/** Constructor */
-	FReimportFeedbackContext();
-
-	/** Initialize this reimport context with the specified widget */
-	void Initialize(TSharedRef<SReimportFeedback> Widget);
+	FReimportFeedbackContext(const TAttribute<FText>& InMainText, const FSimpleDelegate& InOnPauseClicked, const FSimpleDelegate& InOnAbortClicked);
 
 	/** Destroy this reimport context */
-	void Destroy();
+	void Hide();
 
-	TSharedPtr<SReimportFeedback> GetContent() { return NotificationContent; }
+	/** Revive the notification if it was destroyed */
+	void Show();
 
 	/** Get the message log that this context is using */
 	FMessageLog& GetMessageLog() { return MessageLog; }
@@ -80,6 +40,13 @@ public:
 	bool bSuppressSlowTaskMessages;
 
 private:
+
+	/** Called when our notification has expired */
+	void OnNotificationExpired();
+
+	TAttribute<FText> MainTextAttribute;
+	FSimpleDelegate OnPauseClickedEvent, OnAbortClickedEvent;
+
 	/** Message log for output of errors and log messages */
 	FMessageLog MessageLog;
 
@@ -87,5 +54,5 @@ private:
 	TSharedPtr<class SNotificationItem> Notification;
 
 	/** The notification content */
-	TSharedPtr<SReimportFeedback> NotificationContent;
+	TSharedPtr<class SReimportFeedback> NotificationContent;
 };
\ No newline at end of file
