Fixed issue where custom set LoD materials were lost after reimporting a mesh.

diff --git a/Engine/Source/Editor/StaticMeshEditor/Private/StaticMeshEditor.cpp b/Engine/Source/Editor/StaticMeshEditor/Private/StaticMeshEditor.cpp
--- a/Engine/Source/Editor/StaticMeshEditor/Private/StaticMeshEditor.cpp
+++ b/Engine/Source/Editor/StaticMeshEditor/Private/StaticMeshEditor.cpp
@@ -1525,8 +1525,8 @@ void FStaticMeshEditor::SetEditorMesh(UStaticMesh* InStaticMesh)
 {
 	StaticMesh = InStaticMesh;
 
-	//Init stat arrays. A static mesh can have up to three level of details beyond the base mesh.
-	const int32 ArraySize = 4;
+	//Init stat arrays.
+	const int32 ArraySize = MAX_STATIC_MESH_LODS;
 	NumVertices.Empty(ArraySize);
 	NumVertices.AddZeroed(ArraySize);
 	NumTriangles.Empty(ArraySize);
diff --git a/Engine/Source/Editor/UnrealEd/Private/StaticMeshEdit.cpp b/Engine/Source/Editor/UnrealEd/Private/StaticMeshEdit.cpp
--- a/Engine/Source/Editor/UnrealEd/Private/StaticMeshEdit.cpp
+++ b/Engine/Source/Editor/UnrealEd/Private/StaticMeshEdit.cpp
@@ -924,20 +924,22 @@ void RestoreExistingMeshData(struct ExistingStaticMeshData* ExistingMeshDataPtr,
 {
 	if ( ExistingMeshDataPtr && NewMesh )
 	{
-		int32 NumCommonMaterials = FMath::Min(NewMesh->Materials.Num(), ExistingMeshDataPtr->ExistingLODData[0].ExistingMaterials.Num());
-		for (int32 MaterialIndex = 0; MaterialIndex < NumCommonMaterials; ++MaterialIndex)
-		{
-			NewMesh->Materials[MaterialIndex] = ExistingMeshDataPtr->ExistingLODData[0].ExistingMaterials[MaterialIndex];
-		}
-
+		NewMesh->SectionInfoMap.Clear();
 		NewMesh->SectionInfoMap.CopyFrom(ExistingMeshDataPtr->ExistingSectionInfoMap);
 
+		NewMesh->Materials.Reset();
+
 		int32 NumCommonLODs = FMath::Min<int32>(ExistingMeshDataPtr->ExistingLODData.Num(), NewMesh->SourceModels.Num());
 		for(int32 i=0; i<NumCommonLODs; i++)
 		{
 			NewMesh->SourceModels[i].BuildSettings = ExistingMeshDataPtr->ExistingLODData[i].ExistingBuildSettings;
 			NewMesh->SourceModels[i].ReductionSettings = ExistingMeshDataPtr->ExistingLODData[i].ExistingReductionSettings;
 			NewMesh->SourceModels[i].ScreenSize = ExistingMeshDataPtr->ExistingLODData[i].ExistingScreenSize;
+
+			if (ExistingMeshDataPtr->ExistingLODData[i].ExistingMaterials.Num() > 0)
+			{
+				NewMesh->Materials.Append(ExistingMeshDataPtr->ExistingLODData[i].ExistingMaterials);
+			}
 		}
 
 		for(int32 i=NumCommonLODs; i < ExistingMeshDataPtr->ExistingLODData.Num(); ++i)
