Fixed crash renaming a level which has not been saved yet

https://jira.ol.epicgames.net/browse/UE-14639

diff --git a/Engine/Source/Editor/UnrealEd/Private/ObjectTools.cpp b/Engine/Source/Editor/UnrealEd/Private/ObjectTools.cpp
--- a/Engine/Source/Editor/UnrealEd/Private/ObjectTools.cpp
+++ b/Engine/Source/Editor/UnrealEd/Private/ObjectTools.cpp
@@ -2449,8 +2449,7 @@ namespace ObjectTools
 				// Used in the IsValidObjectName checks below
 				FText Reason;
 
-				if( (ExistingPackage && InOutPackagesUserRefusedToFullyLoad.Contains(ExistingPackage)) ||
-					!PackageTools::HandleFullyLoadingPackages( TopLevelPackages, NSLOCTEXT("UnrealEd", "Rename", "Rename") ) )
+				if( ExistingPackage && ( InOutPackagesUserRefusedToFullyLoad.Contains(ExistingPackage) || !PackageTools::HandleFullyLoadingPackages( TopLevelPackages, NSLOCTEXT("UnrealEd", "Rename", "Rename") ) ) )
 				{
 					// HandleFullyLoadingPackages should never return false for empty input.
 					check( ExistingPackage );
