Material flattening functionality moved to a new module MaterialUtilities

diff --git a/Engine/Source/Developer/MaterialUtilities/MaterialUtilities.Build.cs b/Engine/Source/Developer/MaterialUtilities/MaterialUtilities.Build.cs
new file mode 100644
--- /dev/null
+++ b/Engine/Source/Developer/MaterialUtilities/MaterialUtilities.Build.cs
@@ -0,0 +1,33 @@
+// Copyright 1998-2015 Epic Games, Inc. All Rights Reserved.
+
+using UnrealBuildTool;
+using System.IO;
+
+public class MaterialUtilities : ModuleRules
+{
+	public MaterialUtilities(TargetInfo Target)
+	{
+        PrivateDependencyModuleNames.AddRange(
+			new string [] {
+				"Core",
+				"CoreUObject",
+				"Engine",
+				"RenderCore",
+                "RHI",
+                "Landscape"
+			}
+		);
+
+        PrivateIncludePathModuleNames.AddRange(
+            new string[] {
+                "Landscape"
+			}
+        );
+
+        CircularlyReferencedDependentModules.AddRange(
+            new string[] {
+                "Landscape"
+            }
+        );
+	}
+}
diff --git a/Engine/Source/Developer/MaterialUtilities/Private/MaterialUtilities.cpp b/Engine/Source/Developer/MaterialUtilities/Private/MaterialUtilities.cpp
new file mode 100644
--- /dev/null
+++ b/Engine/Source/Developer/MaterialUtilities/Private/MaterialUtilities.cpp
@@ -0,0 +1,1124 @@
+// Copyright 1998-2015 Epic Games, Inc. All Rights Reserved.
+#include "MaterialUtilitiesPrivatePCH.h"
+
+#include "Runtime/Engine/Classes/Engine/World.h"
+#include "Runtime/Engine/Classes/Materials/MaterialInterface.h"
+#include "Runtime/Engine/Classes/Materials/MaterialExpressionConstant.h"
+#include "Runtime/Engine/Classes/Engine/TextureRenderTarget2D.h"
+#include "Runtime/Engine/Classes/Engine/Texture2D.h"
+#include "Runtime/Engine/Public/TileRendering.h"
+#include "Runtime/Engine/Public/EngineModule.h"
+#include "Runtime/Engine/Public/ImageUtils.h"
+#include "Runtime/Engine/Public/MaterialCompiler.h"
+#include "RendererInterface.h"
+#include "LandscapeProxy.h"
+#include "LandscapeComponent.h"
+
+DEFINE_LOG_CATEGORY_STATIC(LogMeshUtilities, Log, All);
+
+/*------------------------------------------------------------------------------
+	Helper classes for render material to texture
+------------------------------------------------------------------------------*/
+class FExportMaterialProxy : public FMaterial, public FMaterialRenderProxy
+{
+public:
+	FExportMaterialProxy()
+		: FMaterial()
+	{
+		SetQualityLevelProperties(EMaterialQualityLevel::High, false, GMaxRHIFeatureLevel);
+	}
+
+	FExportMaterialProxy(UMaterialInterface* InMaterialInterface, EMaterialProperty InPropertyToCompile)
+		: FMaterial()
+		, MaterialInterface(InMaterialInterface)
+		, PropertyToCompile(InPropertyToCompile)
+	{
+		SetQualityLevelProperties(EMaterialQualityLevel::High, false, GMaxRHIFeatureLevel);
+		Material = InMaterialInterface->GetMaterial();
+		Material->AppendReferencedTextures(ReferencedTextures);
+		FPlatformMisc::CreateGuid(Id);
+
+		FMaterialResource* Resource = InMaterialInterface->GetMaterialResource(GMaxRHIFeatureLevel);
+
+		FMaterialShaderMapId ResourceId;
+		Resource->GetShaderMapId(GMaxRHIShaderPlatform, ResourceId);
+		// Override with a special usage so we won't re-use the shader map used by the material for rendering
+		switch (InPropertyToCompile)
+		{
+		case MP_BaseColor: ResourceId.Usage = EMaterialShaderMapUsage::MaterialExportBaseColor; break;
+		case MP_Specular: ResourceId.Usage = EMaterialShaderMapUsage::MaterialExportSpecular; break;
+		case MP_Normal: ResourceId.Usage = EMaterialShaderMapUsage::MaterialExportNormal; break;
+		case MP_Metallic: ResourceId.Usage = EMaterialShaderMapUsage::MaterialExportMetallic; break;
+		case MP_Roughness: ResourceId.Usage = EMaterialShaderMapUsage::MaterialExportRoughness; break;
+		};
+		
+		CacheShaders(ResourceId, GMaxRHIShaderPlatform, true);
+	}
+
+	/** This override is required otherwise the shaders aren't ready for use when the surface is rendered resulting in a blank image */
+	virtual bool RequiresSynchronousCompilation() const override { return true; };
+
+	/**
+	* Should the shader for this material with the given platform, shader type and vertex 
+	* factory type combination be compiled
+	*
+	* @param Platform		The platform currently being compiled for
+	* @param ShaderType	Which shader is being compiled
+	* @param VertexFactory	Which vertex factory is being compiled (can be NULL)
+	*
+	* @return true if the shader should be compiled
+	*/
+	virtual bool ShouldCache(EShaderPlatform Platform, const FShaderType* ShaderType, const FVertexFactoryType* VertexFactoryType) const override
+	{
+		// Always cache - decreases performance but avoids missing shaders during exports.
+		return true;
+	}
+
+	virtual const TArray<UTexture*>& GetReferencedTextures() const override
+	{
+		return ReferencedTextures;
+	}
+
+	////////////////
+	// FMaterialRenderProxy interface.
+	virtual const FMaterial* GetMaterial(ERHIFeatureLevel::Type FeatureLevel) const override
+	{
+		if(GetRenderingThreadShaderMap())
+		{
+			return this;
+		}
+		else
+		{
+			return UMaterial::GetDefaultMaterial(MD_Surface)->GetRenderProxy(false)->GetMaterial(FeatureLevel);
+		}
+	}
+
+	virtual bool GetVectorValue(const FName ParameterName, FLinearColor* OutValue, const FMaterialRenderContext& Context) const override
+	{
+		return MaterialInterface->GetRenderProxy(0)->GetVectorValue(ParameterName, OutValue, Context);
+	}
+
+	virtual bool GetScalarValue(const FName ParameterName, float* OutValue, const FMaterialRenderContext& Context) const override
+	{
+		return MaterialInterface->GetRenderProxy(0)->GetScalarValue(ParameterName, OutValue, Context);
+	}
+
+	virtual bool GetTextureValue(const FName ParameterName,const UTexture** OutValue, const FMaterialRenderContext& Context) const override
+	{
+		return MaterialInterface->GetRenderProxy(0)->GetTextureValue(ParameterName,OutValue,Context);
+	}
+
+	// Material properties.
+	/** Entry point for compiling a specific material property.  This must call SetMaterialProperty. */
+	virtual int32 CompilePropertyAndSetMaterialProperty(EMaterialProperty Property, FMaterialCompiler* Compiler, EShaderFrequency OverrideShaderFrequency, bool bUsePreviousFrameTime) const override
+	{
+		// needs to be called in this function!!
+		Compiler->SetMaterialProperty(Property, OverrideShaderFrequency, bUsePreviousFrameTime);
+
+		int32 Ret = CompilePropertyAndSetMaterialPropertyWithoutCast(Property, Compiler);
+
+		return Compiler->ForceCast(Ret, GetMaterialPropertyType(Property));
+	}
+
+	/** helper for CompilePropertyAndSetMaterialProperty() */
+	int32 CompilePropertyAndSetMaterialPropertyWithoutCast(EMaterialProperty Property, FMaterialCompiler* Compiler) const
+	{
+		if (Property == MP_EmissiveColor)
+		{
+			UMaterial* ProxyMaterial = MaterialInterface->GetMaterial();
+			EBlendMode BlendMode = MaterialInterface->GetBlendMode();
+			EMaterialShadingModel ShadingModel = MaterialInterface->GetShadingModel();
+			check(ProxyMaterial);
+			switch (PropertyToCompile)
+			{
+			case MP_EmissiveColor:
+				// Emissive is ALWAYS returned...
+				return Compiler->ForceCast(MaterialInterface->CompileProperty(Compiler,MP_EmissiveColor),MCT_Float3,true,true);
+			case MP_BaseColor:
+				// Only return for Opaque and Masked...
+				if (BlendMode == BLEND_Opaque || BlendMode == BLEND_Masked)
+				{
+					return Compiler->ForceCast(MaterialInterface->CompileProperty(Compiler, MP_BaseColor),MCT_Float3,true,true);
+				}
+				break;
+			case MP_Specular: 
+			case MP_Roughness:
+			case MP_Metallic:
+				// Only return for Opaque and Masked...
+				if (BlendMode == BLEND_Opaque || BlendMode == BLEND_Masked)
+				{
+					return Compiler->ForceCast(MaterialInterface->CompileProperty(Compiler, PropertyToCompile),MCT_Float,true,true);
+				}
+				break;
+			case MP_Normal:
+				// Only return for Opaque and Masked...
+				if (BlendMode == BLEND_Opaque || BlendMode == BLEND_Masked)
+				{
+					return Compiler->ForceCast( 
+						Compiler->Add( 
+							Compiler->Mul(MaterialInterface->CompileProperty(Compiler, MP_Normal), Compiler->Constant(0.5f)), // [-1,1] * 0.5
+							Compiler->Constant(0.5f)), // [-0.5,0.5] + 0.5
+						MCT_Float3, true, true );
+				}
+				break;
+			default:
+				return Compiler->Constant(1.0f);
+			}
+	
+			return Compiler->Constant(0.0f);
+		}
+		else if (Property == MP_WorldPositionOffset)
+		{
+			//This property MUST return 0 as a default or during the process of rendering textures out for lightmass to use, pixels will be off by 1.
+			return Compiler->Constant(0.0f);
+		}
+		else if (Property >= MP_CustomizedUVs0 && Property <= MP_CustomizedUVs7)
+		{
+			// Pass through customized UVs
+			return MaterialInterface->CompileProperty(Compiler, Property);
+		}
+		else
+		{
+			return Compiler->Constant(1.0f);
+		}
+	}
+
+	virtual FString GetMaterialUsageDescription() const override
+	{
+		return FString::Printf(TEXT("FExportMaterialRenderer %s"), MaterialInterface ? *MaterialInterface->GetName() : TEXT("NULL"));
+	}
+	virtual int32 GetMaterialDomain() const override
+	{
+		if (Material)
+		{
+			return Material->MaterialDomain;
+		}
+		return MD_Surface;
+	}
+	virtual bool IsTwoSided() const  override
+	{ 
+		if (MaterialInterface)
+		{
+			return MaterialInterface->IsTwoSided();
+		}
+		return false;
+	}
+	virtual bool IsDitheredLODTransition() const  override
+	{ 
+		if (MaterialInterface)
+		{
+			return MaterialInterface->IsDitheredLODTransition();
+		}
+		return false;
+	}
+	virtual bool IsLightFunction() const override
+	{
+		if (Material)
+		{
+			return (Material->MaterialDomain == MD_LightFunction);
+		}
+		return false;
+	}
+	virtual bool IsUsedWithDeferredDecal() const override
+	{
+		return	Material &&
+				Material->MaterialDomain == MD_DeferredDecal;
+	}
+	virtual bool IsSpecialEngineMaterial() const override
+	{
+		if (Material)
+		{
+			return (Material->bUsedAsSpecialEngineMaterial == 1);
+		}
+		return false;
+	}
+	virtual bool IsWireframe() const override
+	{
+		if (Material)
+		{
+			return (Material->Wireframe == 1);
+		}
+		return false;
+	}
+	virtual bool IsMasked() const override								{ return false; }
+	virtual enum EBlendMode GetBlendMode() const override				{ return BLEND_Opaque; }
+	virtual enum EMaterialShadingModel GetShadingModel() const override	{ return MSM_Unlit; }
+	virtual float GetOpacityMaskClipValue() const override				{ return 0.5f; }
+	virtual FString GetFriendlyName() const override { return FString::Printf(TEXT("FExportMaterialRenderer %s"), MaterialInterface ? *MaterialInterface->GetName() : TEXT("NULL")); }
+	/**
+	* Should shaders compiled for this material be saved to disk?
+	*/
+	virtual bool IsPersistent() const override { return false; }
+	virtual FGuid GetMaterialId() const override { return Id; }
+
+	const UMaterialInterface* GetMaterialInterface() const
+	{
+		return MaterialInterface;
+	}
+
+	friend FArchive& operator<< ( FArchive& Ar, FExportMaterialProxy& V )
+	{
+		return Ar << V.MaterialInterface;
+	}
+
+	bool IsMaterialInputConnected(UMaterial* InMaterial, EMaterialProperty MaterialInput)
+	{
+		bool bConnected = false;
+
+		switch (MaterialInput)
+		{
+		case MP_EmissiveColor:
+			bConnected = InMaterial->EmissiveColor.Expression != NULL;
+			break;
+		case MP_BaseColor:
+			bConnected = InMaterial->BaseColor.Expression != NULL;
+			break;
+		case MP_Specular:
+			bConnected = InMaterial->Specular.Expression != NULL;
+			break;
+		case MP_Normal:
+			bConnected = InMaterial->Normal.Expression != NULL;
+			break;
+		case MP_Metallic:
+			bConnected = InMaterial->Metallic.Expression != NULL;
+			break;
+		case MP_Roughness:
+			bConnected = InMaterial->Roughness.Expression != NULL;
+			break;
+		default:
+			break;
+		}
+
+		// Note: only checking to see whether the entire material attributes connection exists.  
+		// This means materials using the material attributes input will export more attributes than is necessary.
+		bConnected = InMaterial->bUseMaterialAttributes ? InMaterial->MaterialAttributes.Expression != NULL : bConnected;
+		return bConnected;
+	}
+
+	/**
+	 *	Checks if the configuration of the material proxy will generate a uniform
+	 *	value across the sampling... (Ie, nothing is hooked to the property)
+	 *
+	 *	@param	OutUniformValue		The value that will be returned.
+	 *
+	 *	@return	bool				true if a single value would be generated.
+	 *								false if not.
+	 */
+	bool WillGenerateUniformData(FColor& OutUniformValue)
+	{
+		// Pre-fill the value...
+		OutUniformValue.R = 0;
+		OutUniformValue.G = 0;
+		OutUniformValue.B = 0;
+		OutUniformValue.A = 0;
+
+		EBlendMode BlendMode = MaterialInterface->GetBlendMode();
+		EMaterialShadingModel ShadingModel = MaterialInterface->GetShadingModel();
+		
+		check(Material);
+		bool bExpressionIsNULL = false;
+		switch (PropertyToCompile)
+		{
+		case MP_EmissiveColor:
+			// Emissive is ALWAYS returned...
+			bExpressionIsNULL = !IsMaterialInputConnected(Material, PropertyToCompile);
+			break;
+		case MP_BaseColor:
+			// Only return for Opaque and Masked...
+			if (BlendMode == BLEND_Opaque || BlendMode == BLEND_Masked)
+			{
+				bExpressionIsNULL = !IsMaterialInputConnected(Material, PropertyToCompile);
+			}
+			break;
+		case MP_Specular: 
+		case MP_Metallic:
+		case MP_Roughness:
+			// Only return for Opaque and Masked...
+			if (BlendMode == BLEND_Opaque || BlendMode == BLEND_Masked)
+			{
+				bExpressionIsNULL = !IsMaterialInputConnected(Material, PropertyToCompile);
+				OutUniformValue.A = 255;
+			}
+			break;
+		case MP_Normal:
+			// Only return for Opaque and Masked...
+			if (BlendMode == BLEND_Opaque || BlendMode == BLEND_Masked)
+			{
+				bExpressionIsNULL = !IsMaterialInputConnected(Material, PropertyToCompile);
+				OutUniformValue.B = 255;	// Default normal is (0,0,1)
+			}
+			break;
+		}
+
+		return bExpressionIsNULL;
+	}
+
+	/**
+	* Iterate through all textures used by the material and return the maximum texture resolution used
+	* (ideally this could be made dependent of the material property)
+	*
+	* @param MaterialInterface The material to scan for texture size
+	*
+	* @return Size (width and height)
+	*/
+	FIntPoint FindMaxTextureSize(UMaterialInterface* InMaterialInterface, FIntPoint MinimumSize = FIntPoint(1, 1)) const
+	{
+		// static lod settings so that we only initialize them once
+		UTextureLODSettings* GameTextureLODSettings = UDeviceProfileManager::Get().GetActiveProfile()->GetTextureLODSettings();
+
+		TArray<UTexture*> MaterialTextures;
+
+		InMaterialInterface->GetUsedTextures(MaterialTextures, EMaterialQualityLevel::Num, false, GMaxRHIFeatureLevel, false);
+
+		// find the largest texture in the list (applying it's LOD bias)
+		FIntPoint MaxSize = MinimumSize;
+		for (int32 TexIndex = 0; TexIndex < MaterialTextures.Num(); TexIndex++)
+		{
+			UTexture* Texture = MaterialTextures[TexIndex];
+
+			if (Texture == NULL)
+			{
+				continue;
+			}
+
+			// get the max size of the texture
+			FIntPoint LocalSize(0, 0);
+			if (Texture->IsA(UTexture2D::StaticClass()))
+			{
+				UTexture2D* Tex2D = (UTexture2D*)Texture;
+				LocalSize = FIntPoint(Tex2D->GetSizeX(), Tex2D->GetSizeY());
+			}
+			else if (Texture->IsA(UTextureCube::StaticClass()))
+			{
+				UTextureCube* TexCube = (UTextureCube*)Texture;
+				LocalSize = FIntPoint(TexCube->GetSizeX(), TexCube->GetSizeY());
+			}
+
+			int32 LocalBias = GameTextureLODSettings->CalculateLODBias(Texture);
+
+			// bias the texture size based on LOD group
+			FIntPoint BiasedLocalSize(LocalSize.X >> LocalBias, LocalSize.Y >> LocalBias);
+
+			MaxSize.X = FMath::Max(BiasedLocalSize.X, MaxSize.X);
+			MaxSize.Y = FMath::Max(BiasedLocalSize.Y, MaxSize.Y);
+		}
+
+		return MaxSize;
+	}
+
+	static bool WillFillData(EBlendMode InBlendMode, EMaterialProperty InMaterialProperty)
+	{
+		if (InMaterialProperty == MP_EmissiveColor)
+		{
+			return true;
+		}
+
+		switch (InBlendMode)
+		{
+		case BLEND_Opaque:
+			{
+				switch (InMaterialProperty)
+				{
+				case MP_BaseColor:		return true;
+				case MP_Specular:		return true;
+				case MP_Normal:			return true;
+				case MP_Metallic:		return true;
+				case MP_Roughness:		return true;
+				}
+			}
+			break;
+		}
+		return false;
+	}
+
+private:
+	/** The material interface for this proxy */
+	UMaterialInterface* MaterialInterface;
+	UMaterial* Material;	
+	TArray<UTexture*> ReferencedTextures;
+	/** The property to compile for rendering the sample */
+	EMaterialProperty PropertyToCompile;
+	FGuid Id;
+};
+
+static void RenderMaterialTile(UWorld* InWorld, FMaterialRenderProxy* InMaterialProxy, UTextureRenderTarget2D* InRenderTarget, bool bFrontView)
+{
+	float CurrentRealTime = 0.f;
+	float CurrentWorldTime = 0.f;
+	float DeltaWorldTime = 0.f;
+	
+	if (InWorld)
+	{
+		CurrentRealTime = InWorld->GetRealTimeSeconds();
+		CurrentWorldTime = InWorld->GetTimeSeconds();
+		DeltaWorldTime = InWorld->GetDeltaSeconds();
+	}
+	else
+	{
+		CurrentRealTime = FApp::GetCurrentTime() - GStartTime;
+		CurrentWorldTime = FApp::GetDeltaTime();
+		DeltaWorldTime = FApp::GetCurrentTime() - GStartTime;
+	}
+
+	const FRenderTarget* RenderTargetResource = InRenderTarget->GameThread_GetRenderTargetResource();
+	FSceneViewFamily* ViewFamily = new FSceneViewFamily(FSceneViewFamily::ConstructionValues(
+		RenderTargetResource,
+		InWorld ? InWorld->Scene : nullptr,
+		FEngineShowFlags(ESFIM_Game))
+		.SetWorldTimes(CurrentWorldTime, DeltaWorldTime, CurrentRealTime)
+		.SetGammaCorrection(RenderTargetResource->GetDisplayGamma()));
+
+	FIntPoint ViewSize = RenderTargetResource->GetSizeXY();
+	FIntRect ViewRect(FIntPoint(0, 0), ViewSize);
+	// By default render tile in top view
+	FMatrix ViewMatrix = FMatrix(
+				FPlane(1, 0, 0, 0),
+				FPlane(0, -1, 0, 0),
+				FPlane(0, 0, -1, 0),
+				FPlane(0, 0, 0, 1));
+	FQuat Rotation = FQuat::Identity;
+
+	if (bFrontView)
+	{
+		ViewMatrix = FMatrix(
+				FPlane(1, 0, 0, 0),
+				FPlane(0, 0, -1, 0),
+				FPlane(0, 1, 0, 0),
+				FPlane(0, 0, 0, 1));
+		// Tile mesh was created on XY plane
+		Rotation = FRotator(0.0f, 0.0f, 90.0f).Quaternion();
+	}
+						
+	// make a temporary view
+	FSceneViewInitOptions ViewInitOptions;
+	ViewInitOptions.ViewFamily = ViewFamily;
+	ViewInitOptions.SetViewRectangle(ViewRect);
+	ViewInitOptions.ViewOrigin = FVector(ViewSize.X/2.f, ViewSize.Y/2.f, 0.0f);
+	ViewInitOptions.ViewRotationMatrix = ViewMatrix;
+	ViewInitOptions.ProjectionMatrix = FReversedZOrthoMatrix(ViewSize.X/2.f, ViewSize.Y/2.f, 0.5f/HALF_WORLD_MAX, HALF_WORLD_MAX);
+	ViewInitOptions.BackgroundColor = FLinearColor::Black;
+	ViewInitOptions.OverlayColor = FLinearColor::White;
+				
+	FSceneView* View = new FSceneView(ViewInitOptions);
+						
+	ENQUEUE_UNIQUE_RENDER_COMMAND_FIVEPARAMETER(
+		RenderMaterialTileCommand,
+		FSceneView*, InView, View,
+		FMaterialRenderProxy*, InProxy, InMaterialProxy,
+		const FRenderTarget*, InRenderTarget, RenderTargetResource,
+		FIntPoint, InSize, ViewSize,
+		FQuat, InRotation, Rotation,
+		{
+			::SetRenderTarget(RHICmdList, InRenderTarget->GetRenderTargetTexture(), FTextureRHIRef());
+			FIntRect ViewportRect = FIntRect(FIntPoint::ZeroValue, InSize);
+			RHICmdList.SetViewport(ViewportRect.Min.X, ViewportRect.Min.Y, 0.0f, ViewportRect.Max.X, ViewportRect.Max.Y, 1.0f);
+		
+			FTileRenderer::DrawRotatedTile(RHICmdList, *InView, InProxy, false, InRotation, 0.f, 0.f, InSize.X, InSize.Y,  0.f, 0.f, 1.f, 1.f);
+		});
+
+	FlushRenderingCommands();
+
+	delete View;
+}
+
+static void RenderSceneToTexture(
+		FSceneInterface* Scene,
+		const FName& VisualizationMode, 
+		const FVector& ViewOrigin,
+		const FMatrix& ViewRotationMatrix, 
+		const FMatrix& ProjectionMatrix,  
+		const TSet<FPrimitiveComponentId>& HiddenPrimitives, 
+		FIntPoint TargetSize,
+		float TargetGamma,
+		TArray<FColor>& OutSamples)
+{
+	auto RenderTargetTexture = NewObject<UTextureRenderTarget2D>();
+	check(RenderTargetTexture);
+	RenderTargetTexture->AddToRoot();
+	RenderTargetTexture->ClearColor = FLinearColor::Transparent;
+	RenderTargetTexture->TargetGamma = TargetGamma;
+	RenderTargetTexture->InitCustomFormat(TargetSize.X, TargetSize.Y, PF_FloatRGBA, false);
+	FTextureRenderTargetResource* RenderTargetResource = RenderTargetTexture->GameThread_GetRenderTargetResource();
+
+	FSceneViewFamilyContext ViewFamily(
+		FSceneViewFamily::ConstructionValues(RenderTargetResource, Scene, FEngineShowFlags(ESFIM_Game))
+			.SetWorldTimes(FApp::GetCurrentTime() - GStartTime, FApp::GetDeltaTime(), FApp::GetCurrentTime() - GStartTime)
+		);
+
+	// To enable visualization mode
+	ViewFamily.EngineShowFlags.SetPostProcessing(true);
+	ViewFamily.EngineShowFlags.SetVisualizeBuffer(true);
+	ViewFamily.EngineShowFlags.SetTonemapper(false);
+
+	FSceneViewInitOptions ViewInitOptions;
+	ViewInitOptions.SetViewRectangle(FIntRect(0, 0, TargetSize.X, TargetSize.Y));
+	ViewInitOptions.ViewFamily = &ViewFamily;
+	ViewInitOptions.HiddenPrimitives = HiddenPrimitives;
+	ViewInitOptions.ViewOrigin = ViewOrigin;
+	ViewInitOptions.ViewRotationMatrix = ViewRotationMatrix;
+	ViewInitOptions.ProjectionMatrix = ProjectionMatrix;
+		
+	FSceneView* NewView = new FSceneView(ViewInitOptions);
+	NewView->CurrentBufferVisualizationMode = VisualizationMode;
+	ViewFamily.Views.Add(NewView);
+					
+	FCanvas Canvas(RenderTargetResource, NULL, FApp::GetCurrentTime() - GStartTime, FApp::GetDeltaTime(), FApp::GetCurrentTime() - GStartTime, Scene->GetFeatureLevel());
+	Canvas.Clear(FLinearColor::Transparent);
+	GetRendererModule().BeginRenderingViewFamily(&Canvas, &ViewFamily);
+
+	// Copy the contents of the remote texture to system memory
+	OutSamples.SetNumUninitialized(TargetSize.X*TargetSize.Y);
+	FReadSurfaceDataFlags ReadSurfaceDataFlags;
+	ReadSurfaceDataFlags.SetLinearToGamma(false);
+	RenderTargetResource->ReadPixelsPtr(OutSamples.GetData(), ReadSurfaceDataFlags, FIntRect(0, 0, TargetSize.X, TargetSize.Y));
+	FlushRenderingCommands();
+					
+	RenderTargetTexture->RemoveFromRoot();
+	RenderTargetTexture = nullptr;
+}
+
+bool FMaterialUtilities::SupportsExport(EBlendMode InBlendMode, EMaterialProperty InMaterialProperty)
+{
+	return FExportMaterialProxy::WillFillData(InBlendMode, InMaterialProperty);
+}
+
+bool FMaterialUtilities::ExportMaterialProperty(UWorld* InWorld, UMaterialInterface* InMaterial, EMaterialProperty InMaterialProperty, UTextureRenderTarget2D* InRenderTarget, TArray<FColor>& OutBMP)
+{
+	TScopedPointer<FExportMaterialProxy> MaterialProxy(new FExportMaterialProxy(InMaterial, InMaterialProperty));
+	if (MaterialProxy == NULL)
+	{
+		return false;
+	}
+
+	FColor UniformValue;
+	if (MaterialProxy->WillGenerateUniformData(UniformValue))
+	{
+		// Single value... fill it in.
+		OutBMP.Empty();
+		OutBMP.Add(UniformValue);
+		return true;
+	}
+
+	bool bNormalmap = (InMaterialProperty == MP_Normal);
+
+	RenderMaterialTile(InWorld, MaterialProxy, InRenderTarget, !bNormalmap);
+					
+	FReadSurfaceDataFlags ReadPixelFlags(bNormalmap ? RCM_SNorm : RCM_UNorm);
+	ReadPixelFlags.SetLinearToGamma(false);
+
+	FTextureRenderTargetResource* RTResource = InRenderTarget->GameThread_GetRenderTargetResource();
+	return RTResource->ReadPixels(OutBMP, ReadPixelFlags);
+}
+
+bool FMaterialUtilities::ExportMaterialProperty(UWorld* InWorld, UMaterialInterface* InMaterial, EMaterialProperty InMaterialProperty, FIntPoint& OutSize, TArray<FColor>& OutBMP)
+{
+	TScopedPointer<FExportMaterialProxy> MaterialProxy(new FExportMaterialProxy(InMaterial, InMaterialProperty));
+	if (MaterialProxy == NULL)
+	{
+		return false;
+	}
+
+	FColor UniformValue;
+	if (MaterialProxy->WillGenerateUniformData(UniformValue))
+	{
+		// Single value... fill it in.
+		OutBMP.Empty();
+		OutBMP.Add(UniformValue);
+		return true;
+	}
+
+	bool bNormalmap = (InMaterialProperty == MP_Normal);
+	EPixelFormat Format = PF_FloatRGB;
+	OutSize = MaterialProxy->FindMaxTextureSize(InMaterial);
+	OutBMP.Reserve(OutSize.X*OutSize.Y);
+	OutBMP.SetNumUninitialized(OutSize.X*OutSize.Y);
+
+	UTextureRenderTarget2D* RenderTarget = NewObject<UTextureRenderTarget2D>();
+	check(RenderTarget);
+	RenderTarget->AddToRoot();
+	RenderTarget->ClearColor = FLinearColor(0.0f, 0.0f, 0.0f, 0.0f);
+	RenderTarget->InitCustomFormat(OutSize.X, OutSize.Y, Format, bNormalmap);
+		
+	RenderMaterialTile(InWorld, MaterialProxy, RenderTarget, !bNormalmap);
+					
+	FReadSurfaceDataFlags ReadPixelFlags(bNormalmap ? RCM_SNorm : RCM_UNorm);
+	ReadPixelFlags.SetLinearToGamma(false);
+
+	FTextureRenderTargetResource* RTResource = RenderTarget->GameThread_GetRenderTargetResource();
+	bool bResult = RTResource->ReadPixels(OutBMP, ReadPixelFlags);
+	RenderTarget->RemoveFromRoot();
+
+	return bResult;
+}
+
+bool FMaterialUtilities::ExportMaterial(UWorld* InWorld, UMaterialInterface* InMaterial, FFlattenMaterial& OutFlattenMaterial)
+{
+	// Render diffuse property
+	if (OutFlattenMaterial.DiffuseSize.X > 0 && 
+		OutFlattenMaterial.DiffuseSize.Y > 0)
+	{
+		// Create temporary render target
+		auto RenderTargetDiffuse = NewObject<UTextureRenderTarget2D>();
+		check(RenderTargetDiffuse);
+		RenderTargetDiffuse->AddToRoot();
+		RenderTargetDiffuse->ClearColor = FLinearColor(0.0f, 0.0f, 0.0f, 0.0f);
+		RenderTargetDiffuse->InitCustomFormat(
+			OutFlattenMaterial.DiffuseSize.X, 
+			OutFlattenMaterial.DiffuseSize.Y, PF_B8G8R8A8, false);
+			
+		//
+		OutFlattenMaterial.DiffuseSamples.Empty(OutFlattenMaterial.DiffuseSize.X * OutFlattenMaterial.DiffuseSize.Y);
+		bool bResult = ExportMaterialProperty(InWorld, InMaterial, MP_BaseColor, RenderTargetDiffuse, OutFlattenMaterial.DiffuseSamples);
+
+		// Uniform value
+		if (OutFlattenMaterial.DiffuseSamples.Num() == 1)
+		{
+			OutFlattenMaterial.DiffuseSize = FIntPoint(1,1);
+		}
+			
+		RenderTargetDiffuse->RemoveFromRoot();
+		RenderTargetDiffuse = nullptr;
+
+		if (!bResult)
+		{
+			return false;
+		}
+	}
+
+	// Render normal property
+	if (OutFlattenMaterial.NormalSize.X > 0 && 
+		OutFlattenMaterial.NormalSize.Y > 0 &&
+		InMaterial->GetMaterial()->HasNormalConnected())
+	{
+		// Create temporary render target
+		auto RenderTargetNormal = NewObject<UTextureRenderTarget2D>();
+		check(RenderTargetNormal);
+		RenderTargetNormal->AddToRoot();
+		RenderTargetNormal->ClearColor = FLinearColor(0.0f, 0.0f, 0.0f, 0.0f);
+		RenderTargetNormal->InitCustomFormat(
+			OutFlattenMaterial.NormalSize.X,
+			OutFlattenMaterial.NormalSize.Y, PF_FloatRGB, true);
+
+		//
+		OutFlattenMaterial.NormalSamples.Empty(OutFlattenMaterial.NormalSize.X * OutFlattenMaterial.NormalSize.Y);
+		bool bResult = ExportMaterialProperty(InWorld, InMaterial, MP_Normal, RenderTargetNormal, OutFlattenMaterial.NormalSamples);
+
+		// Uniform value
+		if (OutFlattenMaterial.NormalSamples.Num() == 1)
+		{
+			OutFlattenMaterial.NormalSize = FIntPoint(1,1);
+		}
+		
+		RenderTargetNormal->RemoveFromRoot();
+		RenderTargetNormal = nullptr;
+
+		if (!bResult)
+		{
+			return false;
+		}
+	}
+
+	// Render metallic property
+	if (OutFlattenMaterial.MetallicSize.X > 0 && 
+		OutFlattenMaterial.MetallicSize.Y > 0)
+	{
+		// Create temporary render target
+		auto RenderTargetMetallic = NewObject<UTextureRenderTarget2D>();
+		check(RenderTargetMetallic);
+		RenderTargetMetallic->AddToRoot();
+		RenderTargetMetallic->ClearColor = FLinearColor(0.0f, 0.0f, 0.0f, 0.0f);
+		RenderTargetMetallic->InitCustomFormat(
+			OutFlattenMaterial.MetallicSize.X,
+			OutFlattenMaterial.MetallicSize.Y, PF_B8G8R8A8, true);
+
+		//
+		OutFlattenMaterial.MetallicSamples.Empty(OutFlattenMaterial.MetallicSize.X * OutFlattenMaterial.MetallicSize.Y);
+		bool bResult = ExportMaterialProperty(InWorld, InMaterial, MP_Metallic, RenderTargetMetallic, OutFlattenMaterial.MetallicSamples);
+
+		// Uniform value
+		if (OutFlattenMaterial.MetallicSamples.Num() == 1)
+		{
+			OutFlattenMaterial.MetallicSize = FIntPoint(1,1);
+		}
+		
+		RenderTargetMetallic->RemoveFromRoot();
+		RenderTargetMetallic = nullptr;
+
+		if (!bResult)
+		{
+			return false;
+		}
+	}
+
+	// Render roughness property
+	if (OutFlattenMaterial.RoughnessSize.X > 0 && 
+		OutFlattenMaterial.RoughnessSize.Y > 0)
+	{
+		// Create temporary render target
+		auto RenderTargetRoughness = NewObject<UTextureRenderTarget2D>();
+		check(RenderTargetRoughness);
+		RenderTargetRoughness->AddToRoot();
+		RenderTargetRoughness->ClearColor = FLinearColor(0.0f, 0.0f, 0.0f, 0.0f);
+		RenderTargetRoughness->InitCustomFormat(
+			OutFlattenMaterial.RoughnessSize.X,
+			OutFlattenMaterial.RoughnessSize.Y, PF_B8G8R8A8, true);
+
+		//
+		OutFlattenMaterial.RoughnessSamples.Empty(OutFlattenMaterial.RoughnessSize.X * OutFlattenMaterial.RoughnessSize.Y);
+		bool bResult = ExportMaterialProperty(InWorld, InMaterial, MP_Roughness, RenderTargetRoughness, OutFlattenMaterial.RoughnessSamples);
+
+		// Uniform value
+		if (OutFlattenMaterial.RoughnessSamples.Num() == 1)
+		{
+			OutFlattenMaterial.RoughnessSize = FIntPoint(1,1);
+		}
+		
+		RenderTargetRoughness->RemoveFromRoot();
+		RenderTargetRoughness = nullptr;
+
+		if (!bResult)
+		{
+			return false;
+		}
+	}
+
+	// Render specular property
+	if (OutFlattenMaterial.SpecularSize.X > 0 && 
+		OutFlattenMaterial.SpecularSize.Y > 0)
+	{
+		// Create temporary render target
+		auto RenderTargetSpecular = NewObject<UTextureRenderTarget2D>();
+		check(RenderTargetSpecular);
+		RenderTargetSpecular->AddToRoot();
+		RenderTargetSpecular->ClearColor = FLinearColor(0.0f, 0.0f, 0.0f, 0.0f);
+		RenderTargetSpecular->InitCustomFormat(
+			OutFlattenMaterial.SpecularSize.X,
+			OutFlattenMaterial.SpecularSize.Y, PF_B8G8R8A8, true);
+
+		//
+		OutFlattenMaterial.SpecularSamples.Empty(OutFlattenMaterial.SpecularSize.X * OutFlattenMaterial.SpecularSize.Y);
+		bool bResult = ExportMaterialProperty(InWorld, InMaterial, MP_Specular, RenderTargetSpecular, OutFlattenMaterial.SpecularSamples);
+
+		// Uniform value
+		if (OutFlattenMaterial.SpecularSamples.Num() == 1)
+		{
+			OutFlattenMaterial.SpecularSize = FIntPoint(1,1);
+		}
+		
+		RenderTargetSpecular->RemoveFromRoot();
+		RenderTargetSpecular = nullptr;
+
+		if (!bResult)
+		{
+			return false;
+		}
+	}
+		
+	OutFlattenMaterial.MaterialId = InMaterial->GetLightingGuid();
+	return true;
+}
+
+bool FMaterialUtilities::ExportLandscapeMaterial(ALandscapeProxy* InLandscape, const TSet<FPrimitiveComponentId>& HiddenPrimitives, FFlattenMaterial& OutFlattenMaterial)
+{
+	check(InLandscape);
+
+	FIntRect LandscapeRect = InLandscape->GetBoundingRect();
+	FVector MidPoint = FVector(LandscapeRect.Min, 0.f) + FVector(LandscapeRect.Size(), 0.f)*0.5f;
+		
+	FVector LandscapeCenter = InLandscape->GetTransform().TransformPosition(MidPoint);
+	FVector LandscapeExtent = FVector(LandscapeRect.Size(), 0.f)*InLandscape->GetActorScale()*0.5f; 
+
+	FVector ViewOrigin = LandscapeCenter;
+	FMatrix ViewRotationMatrix = FInverseRotationMatrix(InLandscape->GetActorRotation());
+	ViewRotationMatrix*= FMatrix(FPlane(1,	0,	0,	0),
+							FPlane(0,	-1,	0,	0),
+							FPlane(0,	0,	-1,	0),
+							FPlane(0,	0,	0,	1));
+				
+	const float ZOffset = WORLD_MAX;
+	FMatrix ProjectionMatrix =  FReversedZOrthoMatrix(
+		LandscapeExtent.X,
+		LandscapeExtent.Y,
+		0.5f / ZOffset,
+		ZOffset);
+
+	FSceneInterface* Scene = InLandscape->GetWorld()->Scene;
+						
+	// Render diffuse texture using BufferVisualizationMode=BaseColor
+	if (OutFlattenMaterial.DiffuseSize.X > 0 && 
+		OutFlattenMaterial.DiffuseSize.Y > 0)
+	{
+		static const FName BaseColorName("BaseColor");
+		const float BaseColorGamma = 2.2f; // BaseColor to gamma space
+		RenderSceneToTexture(Scene, BaseColorName, ViewOrigin, ViewRotationMatrix, ProjectionMatrix, HiddenPrimitives, 
+			OutFlattenMaterial.DiffuseSize, BaseColorGamma, OutFlattenMaterial.DiffuseSamples);
+	}
+
+	// Render normal map using BufferVisualizationMode=WorldNormal
+	// Final material should use world space instead of tangent space for normals
+	if (OutFlattenMaterial.NormalSize.X > 0 && 
+		OutFlattenMaterial.NormalSize.Y > 0)
+	{
+		static const FName WorldNormalName("WorldNormal");
+		const float NormalColorGamma = 1.0f; // Dump normal texture in linear space
+		RenderSceneToTexture(Scene, WorldNormalName, ViewOrigin, ViewRotationMatrix, ProjectionMatrix, HiddenPrimitives, 
+			OutFlattenMaterial.NormalSize, NormalColorGamma, OutFlattenMaterial.NormalSamples);
+	}
+
+	// Render metallic map using BufferVisualizationMode=Metallic
+	if (OutFlattenMaterial.MetallicSize.X > 0 && 
+		OutFlattenMaterial.MetallicSize.Y > 0)
+	{
+		static const FName MetallicName("Metallic");
+		const float MetallicColorGamma = 1.0f; // Dump metallic texture in linear space
+		RenderSceneToTexture(Scene, MetallicName, ViewOrigin, ViewRotationMatrix, ProjectionMatrix, HiddenPrimitives, 
+			OutFlattenMaterial.MetallicSize, MetallicColorGamma, OutFlattenMaterial.MetallicSamples);
+	}
+
+	// Render roughness map using BufferVisualizationMode=Roughness
+	if (OutFlattenMaterial.RoughnessSize.X > 0 && 
+		OutFlattenMaterial.RoughnessSize.Y > 0)
+	{
+		static const FName RoughnessName("Roughness");
+		const float RoughnessColorGamma = 2.2f; // Roughness material powers color by 2.2, transform it back to linear
+		RenderSceneToTexture(Scene, RoughnessName, ViewOrigin, ViewRotationMatrix, ProjectionMatrix, HiddenPrimitives, 
+			OutFlattenMaterial.RoughnessSize, RoughnessColorGamma, OutFlattenMaterial.RoughnessSamples);
+	}
+
+	// Render specular map using BufferVisualizationMode=Specular
+	if (OutFlattenMaterial.SpecularSize.X > 0 && 
+		OutFlattenMaterial.SpecularSize.Y > 0)
+	{
+		static const FName SpecularName("Specular");
+		const float SpecularColorGamma = 1.0f; // Dump specular texture in linear space
+		RenderSceneToTexture(Scene, SpecularName, ViewOrigin, ViewRotationMatrix, ProjectionMatrix, HiddenPrimitives, 
+			OutFlattenMaterial.SpecularSize, SpecularColorGamma, OutFlattenMaterial.SpecularSamples);
+	}
+				
+	OutFlattenMaterial.MaterialId = InLandscape->GetLandscapeGuid();
+	return true;
+}
+
+UMaterial* FMaterialUtilities::CreateMaterial(const FFlattenMaterial& InFlattenMaterial, UPackage* InOuter, const FString& BaseName, EObjectFlags Flags, TArray<UObject*>& OutGeneratedAssets)
+{
+	// Base name for a new assets
+	// In case outer is null BaseName has to be long package name
+	if (InOuter == nullptr && FPackageName::IsShortPackageName(BaseName))
+	{
+		UE_LOG(LogMeshUtilities, Warning, TEXT("Invalid long package name: '%s'."), *BaseName);
+		return nullptr;
+	}
+
+	const FString AssetBaseName = FPackageName::GetShortName(BaseName);
+	const FString AssetBasePath = InOuter ? TEXT("") : FPackageName::GetLongPackagePath(BaseName) + TEXT("/");
+				
+	// Create material
+	const FString MaterialAssetName = TEXT("M_") + AssetBaseName;
+	UPackage* MaterialOuter = InOuter;
+	if (MaterialOuter == NULL)
+	{
+		MaterialOuter = CreatePackage(NULL, *(AssetBasePath + MaterialAssetName));
+		MaterialOuter->FullyLoad();
+		MaterialOuter->Modify();
+	}
+
+	UMaterial* Material = NewObject<UMaterial>(MaterialOuter, FName(*MaterialAssetName), Flags);
+	Material->TwoSided = false;
+	Material->SetShadingModel(MSM_DefaultLit);
+	OutGeneratedAssets.Add(Material);
+
+	int32 MaterialNodeY = -150;
+	int32 MaterialNodeStepY = 180;
+
+	// BaseColor
+	if (InFlattenMaterial.DiffuseSamples.Num() > 1)
+	{
+		const FString AssetName = TEXT("T_") + AssetBaseName + TEXT("_D");
+		const FString AssetLongName = AssetBasePath + AssetName;
+		const bool bSRGB = true;
+		UTexture2D* Texture = CreateTexture(InOuter, AssetLongName, InFlattenMaterial.DiffuseSize, InFlattenMaterial.DiffuseSamples, TC_Default, TEXTUREGROUP_World, Flags, bSRGB);
+		OutGeneratedAssets.Add(Texture);
+			
+		auto BasecolorExpression = NewObject<UMaterialExpressionTextureSample>(Material);
+		BasecolorExpression->Texture = Texture;
+		BasecolorExpression->SamplerType = EMaterialSamplerType::SAMPLERTYPE_Color;
+		BasecolorExpression->MaterialExpressionEditorX = -400;
+		BasecolorExpression->MaterialExpressionEditorY = MaterialNodeY;
+		Material->Expressions.Add(BasecolorExpression);
+		Material->BaseColor.Expression = BasecolorExpression;
+
+		MaterialNodeY+= MaterialNodeStepY;
+	}
+
+	// Metallic
+	if (InFlattenMaterial.MetallicSamples.Num() > 1)
+	{
+		const FString AssetName = TEXT("T_") + AssetBaseName + TEXT("_M");
+		const bool bSRGB = false;
+		UTexture2D* Texture = CreateTexture(InOuter, AssetBasePath + AssetName, InFlattenMaterial.MetallicSize, InFlattenMaterial.MetallicSamples, TC_Grayscale, TEXTUREGROUP_World, Flags, bSRGB);
+		OutGeneratedAssets.Add(Texture);
+			
+		auto MetallicExpression = NewObject<UMaterialExpressionTextureSample>(Material);
+		MetallicExpression->Texture = Texture;
+		MetallicExpression->SamplerType = EMaterialSamplerType::SAMPLERTYPE_LinearGrayscale;
+		MetallicExpression->MaterialExpressionEditorX = -400;
+		MetallicExpression->MaterialExpressionEditorY = MaterialNodeY;
+		Material->Expressions.Add(MetallicExpression);
+		Material->Metallic.Expression = MetallicExpression;
+
+		MaterialNodeY+= MaterialNodeStepY;
+	}
+	else
+	{
+		// Set Metallic to constant
+		float Metallic = InFlattenMaterial.MetallicSamples.Num() ? InFlattenMaterial.MetallicSamples[0].R : 0.0f;
+		auto MetallicExpression = NewObject<UMaterialExpressionConstant>(Material);
+		MetallicExpression->R = Metallic;
+		MetallicExpression->MaterialExpressionEditorX = -400;
+		MetallicExpression->MaterialExpressionEditorY = MaterialNodeY;
+		Material->Expressions.Add(MetallicExpression);
+		Material->Metallic.Expression = MetallicExpression;
+
+		MaterialNodeY+= MaterialNodeStepY;
+	}
+
+	// Specular
+	if (InFlattenMaterial.SpecularSamples.Num() > 1)
+	{
+		const FString AssetName = TEXT("T_") + AssetBaseName + TEXT("_S");
+		const bool bSRGB = false;
+		UTexture2D* Texture = CreateTexture(InOuter, AssetBasePath + AssetName, InFlattenMaterial.SpecularSize, InFlattenMaterial.SpecularSamples, TC_Grayscale, TEXTUREGROUP_World, Flags, bSRGB);
+		OutGeneratedAssets.Add(Texture);
+			
+		auto SpecularExpression = NewObject<UMaterialExpressionTextureSample>(Material);
+		SpecularExpression->Texture = Texture;
+		SpecularExpression->SamplerType = EMaterialSamplerType::SAMPLERTYPE_LinearGrayscale;
+		SpecularExpression->MaterialExpressionEditorX = -400;
+		SpecularExpression->MaterialExpressionEditorY = MaterialNodeY;
+		Material->Expressions.Add(SpecularExpression);
+		Material->Specular.Expression = SpecularExpression;
+
+		MaterialNodeY+= MaterialNodeStepY;
+	}
+
+	// Roughness
+	if (InFlattenMaterial.RoughnessSamples.Num() > 1)
+	{
+		const FString AssetName = TEXT("T_") + AssetBaseName + TEXT("_R");
+		const bool bSRGB = false;
+		UTexture2D* Texture = CreateTexture(InOuter, AssetBasePath + AssetName, InFlattenMaterial.RoughnessSize, InFlattenMaterial.RoughnessSamples, TC_Grayscale, TEXTUREGROUP_World, Flags, bSRGB);
+		OutGeneratedAssets.Add(Texture);
+			
+		auto RoughnessExpression = NewObject<UMaterialExpressionTextureSample>(Material);
+		RoughnessExpression->Texture = Texture;
+		RoughnessExpression->SamplerType = EMaterialSamplerType::SAMPLERTYPE_LinearGrayscale;
+		RoughnessExpression->MaterialExpressionEditorX = -400;
+		RoughnessExpression->MaterialExpressionEditorY = MaterialNodeY;
+		Material->Expressions.Add(RoughnessExpression);
+		Material->Roughness.Expression = RoughnessExpression;
+
+		MaterialNodeY+= MaterialNodeStepY;
+	}
+	else
+	{
+		// Set Roughness to constant
+		float Roughness = InFlattenMaterial.RoughnessSamples.Num() ? InFlattenMaterial.RoughnessSamples[0].R : 0.8f;
+		auto RoughnessExpression = NewObject<UMaterialExpressionConstant>(Material);
+		RoughnessExpression->R = Roughness;
+		RoughnessExpression->MaterialExpressionEditorX = -400;
+		RoughnessExpression->MaterialExpressionEditorY = MaterialNodeY;
+		Material->Expressions.Add(RoughnessExpression);
+		Material->Roughness.Expression = RoughnessExpression;
+
+		MaterialNodeY+= MaterialNodeStepY;
+	}
+
+	// Normal
+	if (InFlattenMaterial.NormalSamples.Num() > 1)
+	{
+		const FString AssetName = TEXT("T_") + AssetBaseName + TEXT("_N");
+		const bool bSRGB = false;
+		UTexture2D* Texture = CreateTexture(InOuter, AssetBasePath + AssetName, InFlattenMaterial.NormalSize, InFlattenMaterial.NormalSamples, TC_Normalmap, TEXTUREGROUP_WorldNormalMap, Flags, bSRGB);
+		OutGeneratedAssets.Add(Texture);
+			
+		auto NormalExpression = NewObject<UMaterialExpressionTextureSample>(Material);
+		NormalExpression->Texture = Texture;
+		NormalExpression->SamplerType = EMaterialSamplerType::SAMPLERTYPE_Normal;
+		NormalExpression->MaterialExpressionEditorX = -400;
+		NormalExpression->MaterialExpressionEditorY = MaterialNodeY;
+		Material->Expressions.Add(NormalExpression);
+		Material->Normal.Expression = NormalExpression;
+
+		MaterialNodeY+= MaterialNodeStepY;
+	}
+							
+	Material->PostEditChange();
+	return Material;
+}
+
+UTexture2D* FMaterialUtilities::CreateTexture(UPackage* Outer, const FString& AssetLongName, FIntPoint Size, const TArray<FColor>& Samples, TextureCompressionSettings CompressionSettings, TextureGroup LODGroup, EObjectFlags Flags, bool bSRGB, const FGuid& SourceGuidHash)
+{
+	FCreateTexture2DParameters TexParams;
+	TexParams.bUseAlpha = false;
+	TexParams.CompressionSettings = CompressionSettings;
+	TexParams.bDeferCompression = true;
+	TexParams.bSRGB = bSRGB;
+	TexParams.SourceGuidHash = SourceGuidHash;
+
+	if (Outer == nullptr)
+	{
+		Outer = CreatePackage(NULL, *AssetLongName);
+		Outer->FullyLoad();
+		Outer->Modify();
+	}
+
+	UTexture2D* Texture = FImageUtils::CreateTexture2D(Size.X, Size.Y, Samples, Outer, FPackageName::GetShortName(AssetLongName), Flags, TexParams);
+	Texture->LODGroup = LODGroup;
+	Texture->PostEditChange();
+		
+	return Texture;
+}
+
+bool FMaterialUtilities::ExportBaseColor(ULandscapeComponent* LandscapeComponent, int32 TextureSize, TArray<FColor>& OutSamples)
+{
+	ALandscapeProxy* LandscapeProxy = LandscapeComponent->GetLandscapeProxy();
+
+	FIntPoint ComponentOrigin = LandscapeComponent->GetSectionBase() - LandscapeProxy->LandscapeSectionOffset;
+	FIntPoint ComponentSize(LandscapeComponent->ComponentSizeQuads, LandscapeComponent->ComponentSizeQuads);
+	FVector MidPoint = FVector(ComponentOrigin, 0.f) + FVector(ComponentSize, 0.f)*0.5f;
+
+	FVector LandscapeCenter = LandscapeProxy->GetTransform().TransformPosition(MidPoint);
+	FVector LandscapeExtent = FVector(ComponentSize, 0.f)*LandscapeProxy->GetActorScale()*0.5f;
+
+	FVector ViewOrigin = LandscapeCenter;
+	FMatrix ViewRotationMatrix = FInverseRotationMatrix(LandscapeProxy->GetActorRotation());
+	ViewRotationMatrix *= FMatrix(FPlane(1, 0, 0, 0),
+		FPlane(0, -1, 0, 0),
+		FPlane(0, 0, -1, 0),
+		FPlane(0, 0, 0, 1));
+
+	const float ZOffset = WORLD_MAX;
+	FMatrix ProjectionMatrix = FReversedZOrthoMatrix(
+		LandscapeExtent.X,
+		LandscapeExtent.Y,
+		0.5f / ZOffset,
+		ZOffset);
+
+	FSceneInterface* Scene = LandscapeProxy->GetWorld()->Scene;
+
+	// Hide all but the component
+	TSet<FPrimitiveComponentId> HiddenPrimitives;
+	for (auto PrimitiveComponentId : Scene->GetScenePrimitiveComponentIds())
+	{
+		HiddenPrimitives.Add(PrimitiveComponentId);
+	}
+	HiddenPrimitives.Remove(LandscapeComponent->SceneProxy->GetPrimitiveComponentId());
+				
+	FIntPoint TargetSize(TextureSize, TextureSize);
+
+	// Render diffuse texture using BufferVisualizationMode=BaseColor
+	static const FName BaseColorName("BaseColor");
+	const float BaseColorGamma = 2.2f;
+	RenderSceneToTexture(Scene, BaseColorName, ViewOrigin, ViewRotationMatrix, ProjectionMatrix, HiddenPrimitives, TargetSize, BaseColorGamma, OutSamples);
+	return true;
+}
\ No newline at end of file
diff --git a/Engine/Source/Developer/MaterialUtilities/Public/MaterialUtilities.h b/Engine/Source/Developer/MaterialUtilities/Public/MaterialUtilities.h
new file mode 100644
--- /dev/null
+++ b/Engine/Source/Developer/MaterialUtilities/Public/MaterialUtilities.h
@@ -0,0 +1,130 @@
+// Copyright 1998-2015 Epic Games, Inc. All Rights Reserved.
+
+#pragma once
+
+#include "ModuleInterface.h"
+
+struct FFlattenMaterial
+{
+	FFlattenMaterial()
+		: DiffuseSize(512, 512)
+		, NormalSize(512, 512)
+		, MetallicSize(0, 0)
+		, RoughnessSize(0, 0)
+		, SpecularSize(0, 0)
+	{}
+		
+	FGuid			MaterialId;
+		
+	FIntPoint		DiffuseSize;
+	FIntPoint		NormalSize;
+	FIntPoint		MetallicSize;	
+	FIntPoint		RoughnessSize;	
+	FIntPoint		SpecularSize;	
+			
+	TArray<FColor>	DiffuseSamples;
+	TArray<FColor>	NormalSamples;
+	TArray<FColor>	MetallicSamples;
+	TArray<FColor>	RoughnessSamples;
+	TArray<FColor>	SpecularSamples;
+};
+
+
+class UMaterialInterface;
+class UMaterial;
+class UTexture2D;
+class UTextureRenderTarget2D;
+class UWorld;
+class ALandscapeProxy;
+class ULandscapeComponent;
+
+/**
+ * Material utilities
+ */
+class FMaterialUtilities
+{
+public:
+	/**
+	 * Whether material utilities support exporting specified material blend mode and property 
+	 */
+	MATERIALUTILITIES_API static bool SupportsExport(EBlendMode InBlendMode, EMaterialProperty InMaterialProperty);
+
+	/**
+	 * Renders specified material property into texture
+	 *
+	 * @param InWorld				World object to use for material property rendering
+	 * @param InMaterial			Target material
+	 * @param InMaterialProperty	Material property to render
+	 * @param InRenderTarget		Render target to render to
+	 * @param OutBMP				Output array of rendered samples 
+	 * @return						Whether operation was successful
+	 */
+	MATERIALUTILITIES_API static bool ExportMaterialProperty(UWorld* InWorld, UMaterialInterface* InMaterial, EMaterialProperty InMaterialProperty, UTextureRenderTarget2D* InRenderTarget, TArray<FColor>& OutBMP);
+
+	/**
+	 * Renders specified material property into texture
+	 *
+	 * @param InWorld				World object to use for material property rendering
+	 * @param InMaterial			Target material
+	 * @param InMaterialProperty	Material property to render
+	 * @param OutBMP				Output array of rendered samples 
+	 * @return						Whether operation was successful
+	 */
+
+	MATERIALUTILITIES_API static bool ExportMaterialProperty(UWorld* InWorld, UMaterialInterface* InMaterial, EMaterialProperty InMaterialProperty,  FIntPoint& OutSize, TArray<FColor>& OutBMP);
+
+	/**
+	 * Flattens specified material
+	 *
+	 * @param InWorld				World object to use for material rendering
+	 * @param InMaterial			Target material
+	 * @param OutFlattenMaterial	Output flattened material
+	 * @return						Whether operation was successful
+	 */
+	MATERIALUTILITIES_API static bool ExportMaterial(UWorld* InWorld, UMaterialInterface* InMaterial, FFlattenMaterial& OutFlattenMaterial);
+
+	/**
+	 * Flattens specified landscape material
+	 *
+	 * @param InLandscape			Target landscape
+	 * @param HiddenPrimitives		Primitives to hide while rendering scene to texture
+	 * @param OutFlattenMaterial	Output flattened material
+	 * @return						Whether operation was successful
+	 */
+	MATERIALUTILITIES_API static bool ExportLandscapeMaterial(ALandscapeProxy* InLandscape, const TSet<FPrimitiveComponentId>& HiddenPrimitives, FFlattenMaterial& OutFlattenMaterial);
+	
+	/**
+ 	 * Generates a texture from an array of samples 
+	 *
+	 * @param Outer					Outer for the material and texture objects, if NULL new packages will be created for each asset
+	 * @param AssetLongName			Long asset path for the new texture
+	 * @param Size					Resolution of the texture to generate (must match the number of samples)
+	 * @param Samples				Color data for the texture
+	 * @param CompressionSettings	Compression settings for the new texture
+	 * @param LODGroup				LODGroup for the new texture
+	 * @param Flags					ObjectFlags for the new texture
+	 * @param bSRGB					Whether to set the bSRGB flag on the new texture
+	 * @param SourceGuidHash		(optional) Hash (stored as Guid) to use part of the texture source's DDC key.
+	 * @return						The new texture.
+	 */
+	MATERIALUTILITIES_API static UTexture2D* CreateTexture(UPackage* Outer, const FString& AssetLongName, FIntPoint Size, const TArray<FColor>& Samples, TextureCompressionSettings CompressionSettings, TextureGroup LODGroup, EObjectFlags Flags, bool bSRGB, const FGuid& SourceGuidHash = FGuid());
+
+	/**
+	 * Creates UMaterial object from a flatten material
+	 *
+	 * @param Outer					Outer for the material and texture objects, if NULL new packages will be created for each asset
+	 * @param BaseName				BaseName for the material and texture objects, should be a long package name in case Outer is not specified
+	 * @param Flags					Object flags for the material and texture objects.
+	 * @param OutGeneratedAssets	List of generated assets - material, textures
+	 * @return						Returns a pointer to the constructed UMaterial object.
+	 */
+	MATERIALUTILITIES_API static UMaterial* CreateMaterial(const FFlattenMaterial& InFlattenMaterial, UPackage* InOuter, const FString& BaseName, EObjectFlags Flags, TArray<UObject*>& OutGeneratedAssets);
+
+	/**
+	* Creates bakes textures for a ULandscapeComponent
+	*
+	* @param LandscapeComponent		The component to bake textures for
+	* @return						Whether operation was successful
+	*/
+	MATERIALUTILITIES_API static bool ExportBaseColor(ULandscapeComponent* LandscapeComponent, int32 TextureSize, TArray<FColor>& OutSamples);
+};
diff --git a/Engine/Source/Developer/MeshUtilities/MeshUtilities.Build.cs b/Engine/Source/Developer/MeshUtilities/MeshUtilities.Build.cs
--- a/Engine/Source/Developer/MeshUtilities/MeshUtilities.Build.cs
+++ b/Engine/Source/Developer/MeshUtilities/MeshUtilities.Build.cs
@@ -15,7 +15,7 @@ public class MeshUtilities : ModuleRules
 				"RawMesh",
 				"RenderCore", // For FPackedNormal
 				"SlateCore",
-				"UnrealEd",
+				"MaterialUtilities"
 			}
 		);
 
diff --git a/Engine/Source/Developer/MeshUtilities/Private/MeshUtilities.cpp b/Engine/Source/Developer/MeshUtilities/Private/MeshUtilities.cpp
--- a/Engine/Source/Developer/MeshUtilities/Private/MeshUtilities.cpp
+++ b/Engine/Source/Developer/MeshUtilities/Private/MeshUtilities.cpp
@@ -10,13 +10,13 @@
 #include "ThirdParty/nvtesslib/inc/nvtess.h"
 #include "SkeletalMeshTools.h"
 #include "ImageUtils.h"
-#include "MaterialExportUtils.h"
 #include "Textures/TextureAtlas.h"
 #include "LayoutUV.h"
 #include "mikktspace.h"
 #include "DistanceFieldAtlas.h"
 #include "FbxErrors.h"
 #include "Components/SplineMeshComponent.h"
+#include "MaterialUtilities.h"
 
 //@todo - implement required vector intrinsics for other implementations
 #if PLATFORM_ENABLE_VECTORINTRINSICS
@@ -3687,7 +3687,7 @@ void FMeshUtilities::CreateProxyMesh(
 	
 	// Convert collected static mesh components and landscapes into raw meshes and flatten materials
 	TArray<FRawMesh>								RawMeshes;
-	TArray<MaterialExportUtils::FFlattenMaterial>	UniqueMaterials;
+	TArray<FFlattenMaterial>						UniqueMaterials;
 	TMap<int32, TArray<int32>>						MaterialMap;
 	FBox											ProxyBounds(0);
 	
@@ -3720,7 +3720,7 @@ void FMeshUtilities::CreateProxyMesh(
 
 	// Convert materials into flatten materials
 	{
-		MaterialExportUtils::FFlattenMaterial FlattenMaterial;
+		FFlattenMaterial FlattenMaterial;
 		FIntPoint TargetTextureSize = FIntPoint(InProxySettings.TextureWidth, InProxySettings.TextureHeight);
 
 		FlattenMaterial.DiffuseSize = TargetTextureSize;
@@ -3732,7 +3732,7 @@ void FMeshUtilities::CreateProxyMesh(
 		for (UMaterialInterface* Material : StaticMeshMaterials)
 		{
 			UniqueMaterials.Add(FlattenMaterial);
-			MaterialExportUtils::ExportMaterial(InWorld, Material, UniqueMaterials.Last());
+			FMaterialUtilities::ExportMaterial(InWorld, Material, UniqueMaterials.Last());
 		}
 	}
 		
@@ -3757,7 +3757,7 @@ void FMeshUtilities::CreateProxyMesh(
 	// Build proxy mesh
 	//
 	FRawMesh								ProxyRawMesh;
-	MaterialExportUtils::FFlattenMaterial	ProxyFlattenMaterial;
+	FFlattenMaterial						ProxyFlattenMaterial;
 	
 	MeshMerging->BuildProxy(RawMeshes, UniqueMaterials, InProxySettings, ProxyRawMesh, ProxyFlattenMaterial);
 
@@ -3769,7 +3769,7 @@ void FMeshUtilities::CreateProxyMesh(
 	}
 	
 	// Construct proxy material
-	UMaterial* ProxyMaterial = MaterialExportUtils::CreateMaterial(ProxyFlattenMaterial, InOuter, ProxyBasePackageName, RF_Public|RF_Standalone, OutAssetsToSync);
+	UMaterial* ProxyMaterial = FMaterialUtilities::CreateMaterial(ProxyFlattenMaterial, InOuter, ProxyBasePackageName, RF_Public|RF_Standalone, OutAssetsToSync);
 	
 	// Construct proxy static mesh
 	UPackage* MeshPackage = InOuter;
@@ -4191,10 +4191,8 @@ static FVector2D GetValidUV(const FVector2D& UV)
 	return NewUV;
 }
 
-static void MergeMaterials(UWorld* InWorld, const TArray<UMaterialInterface*>& InMaterialList, MaterialExportUtils::FFlattenMaterial& OutMergedMaterial, TArray<FRawMeshUVTransform>& OutUVTransforms)
+static void MergeMaterials(UWorld* InWorld, const TArray<UMaterialInterface*>& InMaterialList, FFlattenMaterial& OutMergedMaterial, TArray<FRawMeshUVTransform>& OutUVTransforms)
 {
-	using namespace MaterialExportUtils;
-	
 	OutUVTransforms.Reserve(InMaterialList.Num());
 
 	// We support merging only for opaque materials
@@ -4278,7 +4276,7 @@ static void MergeMaterials(UWorld* InWorld, const TArray<UMaterialInterface*>& I
 		FlatMaterial.RoughnessSize	= bExportRoughness ? ExportTextureSize : FIntPoint::ZeroValue;
 		FlatMaterial.SpecularSize	= bExportSpecular ? ExportTextureSize : FIntPoint::ZeroValue;
 
-		ExportMaterial(InWorld, Material, FlatMaterial);
+		FMaterialUtilities::ExportMaterial(InWorld, Material, FlatMaterial);
 
 		if (FlatMaterial.DiffuseSamples.Num() > 0)
 		{
@@ -4455,7 +4453,7 @@ void FMeshUtilities::MergeActors(
 	if (InSettings.bMergeMaterials)
 	{
 		FIntPoint AtlasTextureSize			= FIntPoint(InSettings.MergedMaterialAtlasResolution, InSettings.MergedMaterialAtlasResolution);
-		MaterialExportUtils::FFlattenMaterial MergedFlatMaterial;
+		FFlattenMaterial MergedFlatMaterial;
 		MergedFlatMaterial.DiffuseSize		= AtlasTextureSize;
 		MergedFlatMaterial.NormalSize		= InSettings.bExportNormalMap ? AtlasTextureSize : FIntPoint::ZeroValue;
 		MergedFlatMaterial.MetallicSize		= InSettings.bExportMetallicMap ? AtlasTextureSize : FIntPoint::ZeroValue;
@@ -4545,7 +4543,7 @@ void FMeshUtilities::MergeActors(
 			MaterialPackage->Modify();
 		}
 
-		UMaterial* MergedMaterial = CreateMaterial(MergedFlatMaterial, MaterialPackage, MaterialAssetName, RF_Public|RF_Standalone, OutAssetsToSync);
+		UMaterial* MergedMaterial = FMaterialUtilities::CreateMaterial(MergedFlatMaterial, MaterialPackage, MaterialAssetName, RF_Public|RF_Standalone, OutAssetsToSync);
 		UniqueMaterials.Add(MergedMaterial);
 	}
 		
diff --git a/Engine/Source/Developer/MeshUtilities/Public/MeshUtilities.h b/Engine/Source/Developer/MeshUtilities/Public/MeshUtilities.h
--- a/Engine/Source/Developer/MeshUtilities/Public/MeshUtilities.h
+++ b/Engine/Source/Developer/MeshUtilities/Public/MeshUtilities.h
@@ -47,11 +47,7 @@ public:
 	virtual bool IsSupported() const = 0;
 };
 
-//
-namespace MaterialExportUtils
-{
-	struct FFlattenMaterial;
-};
+struct FFlattenMaterial;
 
 /**
  * Mesh merging interface.
@@ -61,10 +57,10 @@ class IMeshMerging
 public:
 	virtual void BuildProxy(
 		const TArray<FRawMesh>& InputMeshes,
-		const TArray<MaterialExportUtils::FFlattenMaterial>& InputMaterials,
+		const TArray<FFlattenMaterial>& InputMaterials,
 		const struct FMeshProxySettings& InProxySettings,
 		FRawMesh& OutProxyMesh,
-		MaterialExportUtils::FFlattenMaterial& OutMaterial
+		FFlattenMaterial& OutMaterial
 		) = 0;
 };
 
diff --git a/Engine/Source/Developer/SimplygonMeshReduction/Private/SimplygonMeshReduction.cpp b/Engine/Source/Developer/SimplygonMeshReduction/Private/SimplygonMeshReduction.cpp
--- a/Engine/Source/Developer/SimplygonMeshReduction/Private/SimplygonMeshReduction.cpp
+++ b/Engine/Source/Developer/SimplygonMeshReduction/Private/SimplygonMeshReduction.cpp
@@ -3,6 +3,7 @@
 #include "UnrealEd.h"
 #include "RawMesh.h"
 #include "MeshUtilities.h"
+#include "MaterialUtilities.h"
 
 // Standard Simplygon channels have some issues with extracting color data back from simplification, 
 // so we use this workaround with user channels
@@ -24,7 +25,6 @@ static const char* USER_MATERIAL_CHANNEL_SPECULAR = "UserSpecular";
 #endif
 
 #include "MeshBoneReduction.h"
-#include "MaterialExportUtils.h"
 #include "ComponentReregisterContext.h"
 
 #define LOCTEXT_NAMESPACE "SimplygonMeshReduction"
@@ -446,10 +446,10 @@ public:
 	// IMeshMerging interface
 	virtual void BuildProxy(
 		const TArray<FRawMesh>& InputMeshes,
-		const TArray<MaterialExportUtils::FFlattenMaterial>& InputMaterials,
+		const TArray<FFlattenMaterial>& InputMaterials,
 		const struct FMeshProxySettings& InProxySettings,
 		FRawMesh& OutProxyMesh,
-		MaterialExportUtils::FFlattenMaterial& OutMaterial) override
+		FFlattenMaterial& OutMaterial) override
 	{
 		if (InputMeshes.Num() == 0)
 		{
@@ -1569,7 +1569,7 @@ private:
 
 	//Material conversions
 	bool CreateSGMaterialFromFlattenMaterial(
-		const TArray<MaterialExportUtils::FFlattenMaterial>& InputMaterials,
+		const TArray<FFlattenMaterial>& InputMaterials,
 		SimplygonSDK::spMaterialTable& OutSGMaterialTable, 
 		FMaterialCastingProperties& OutCastProperties)
 	{
@@ -1583,7 +1583,7 @@ private:
 		for (int32 MaterialIndex = 0; MaterialIndex < InputMaterials.Num(); MaterialIndex++)
 		{
 			SimplygonSDK::spMaterial SGMaterial = SDK->CreateMaterial();
-			const MaterialExportUtils::FFlattenMaterial& FlattenMaterial = InputMaterials[MaterialIndex];
+			const FFlattenMaterial& FlattenMaterial = InputMaterials[MaterialIndex];
 
 			//Create UE4 Channels
 			//Simplygon 5.5 has three new channels already present called base color metallic roughness
@@ -1679,7 +1679,7 @@ private:
 
 	void CreateFlattenMaterialFromSGMaterial(
 		SimplygonSDK::spMaterialTable& InSGMaterialTable, 
-		MaterialExportUtils::FFlattenMaterial& OutMaterial)
+		FFlattenMaterial& OutMaterial)
 	{
 		SimplygonSDK::spMaterial SGMaterial = InSGMaterialTable->GetMaterial(0);
 		
diff --git a/Engine/Source/Developer/SimplygonMeshReduction/SimplygonMeshReduction.Build.cs b/Engine/Source/Developer/SimplygonMeshReduction/SimplygonMeshReduction.Build.cs
--- a/Engine/Source/Developer/SimplygonMeshReduction/SimplygonMeshReduction.Build.cs
+++ b/Engine/Source/Developer/SimplygonMeshReduction/SimplygonMeshReduction.Build.cs
@@ -17,6 +17,7 @@ public class SimplygonMeshReduction : ModuleRules
 		AddThirdPartyPrivateStaticDependencies(Target, "Simplygon");
 		PrivateDependencyModuleNames.Add("RawMesh");
 		PrivateIncludePathModuleNames.Add("MeshUtilities");
+        PrivateIncludePathModuleNames.Add("MaterialUtilities");
         PrivateDependencyModuleNames.Add("MeshBoneReduction");
 	}
 }
diff --git a/Engine/Source/Editor/UnrealEd/Private/EditorExporters.cpp b/Engine/Source/Editor/UnrealEd/Private/EditorExporters.cpp
--- a/Engine/Source/Editor/UnrealEd/Private/EditorExporters.cpp
+++ b/Engine/Source/Editor/UnrealEd/Private/EditorExporters.cpp
@@ -6,16 +6,14 @@
 
 #include "UnrealEd.h"
 #include "StaticMeshResources.h"
-#include "Materials/MaterialExpressionConstant.h"
 #include "SoundDefinitions.h"
 #include "Kismet2/DebuggerCommands.h"
-#include "MaterialCompiler.h"
 #include "SurfaceIterators.h"
 #include "FbxExporter.h"
 #include "AudioDerivedData.h"
 #include "RawMesh.h"
 #include "MaterialExportUtils.h"
-#include "ImageUtils.h"
+#include "MaterialUtilities.h"
 #include "InstancedFoliageActor.h"
 #include "Landscape.h"
 #include "LandscapeComponent.h"
@@ -31,13 +29,11 @@
 #include "GameFramework/PhysicsVolume.h"
 #include "Engine/StaticMeshActor.h"
 #include "Engine/TextureRenderTarget2D.h"
-#include "CanvasTypes.h"
 #include "Engine/Selection.h"
 #include "AI/Navigation/NavigationSystem.h"
 #include "Engine/TextureCube.h"
 #include "EngineUtils.h"
 #include "DeviceProfiles/DeviceProfileManager.h"
-#include "TileRendering.h"
 #include "Engine/TextureLODSettings.h"
 
 DEFINE_LOG_CATEGORY_STATIC(LogEditorExporters, Log, All);
@@ -707,534 +703,6 @@ bool ULevelExporterSTL::ExportText( const FExportObjectInnerContext* Context, UO
 	return 1;
 }
 
-
-/*------------------------------------------------------------------------------
-	Helper classes for render material to texture
-	@todo UE4: Replace this with shared Lightmass version when we merge
-------------------------------------------------------------------------------*/
-
-
-/**
-* Class for rendering previews of material expressions in the material editor's linked object viewport.
-*/
-class FExportMaterialProxy : public FMaterial, public FMaterialRenderProxy
-{
-public:
-	FExportMaterialProxy()
-		: FMaterial()
-	{
-		SetQualityLevelProperties(EMaterialQualityLevel::High, false, GMaxRHIFeatureLevel);
-	}
-
-	FExportMaterialProxy(UMaterialInterface* InMaterialInterface, EMaterialProperty InPropertyToCompile)
-		: FMaterial()
-		, MaterialInterface(InMaterialInterface)
-		, PropertyToCompile(InPropertyToCompile)
-	{
-		SetQualityLevelProperties(EMaterialQualityLevel::High, false, GMaxRHIFeatureLevel);
-		Material = InMaterialInterface->GetMaterial();
-		Material->AppendReferencedTextures(ReferencedTextures);
-		FPlatformMisc::CreateGuid(Id);
-
-		FMaterialResource* Resource = InMaterialInterface->GetMaterialResource(GMaxRHIFeatureLevel);
-
-		FMaterialShaderMapId ResourceId;
-		Resource->GetShaderMapId(GMaxRHIShaderPlatform, ResourceId);
-		// Override with a special usage so we won't re-use the shader map used by the material for rendering
-		switch (InPropertyToCompile)
-		{
-		case MP_BaseColor: ResourceId.Usage = EMaterialShaderMapUsage::MaterialExportBaseColor; break;
-		case MP_Specular: ResourceId.Usage = EMaterialShaderMapUsage::MaterialExportSpecular; break;
-		case MP_Normal: ResourceId.Usage = EMaterialShaderMapUsage::MaterialExportNormal; break;
-		case MP_Metallic: ResourceId.Usage = EMaterialShaderMapUsage::MaterialExportMetallic; break;
-		case MP_Roughness: ResourceId.Usage = EMaterialShaderMapUsage::MaterialExportRoughness; break;
-		};
-		
-		CacheShaders(ResourceId, GMaxRHIShaderPlatform, true);
-	}
-
-	/** This override is required otherwise the shaders aren't ready for use when the surface is rendered resulting in a blank image */
-	virtual bool RequiresSynchronousCompilation() const override { return true; };
-
-	/**
-	* Should the shader for this material with the given platform, shader type and vertex 
-	* factory type combination be compiled
-	*
-	* @param Platform		The platform currently being compiled for
-	* @param ShaderType	Which shader is being compiled
-	* @param VertexFactory	Which vertex factory is being compiled (can be NULL)
-	*
-	* @return true if the shader should be compiled
-	*/
-	virtual bool ShouldCache(EShaderPlatform Platform, const FShaderType* ShaderType, const FVertexFactoryType* VertexFactoryType) const override
-	{
-		// Always cache - decreases performance but avoids missing shaders during exports.
-		return true;
-	}
-
-	virtual const TArray<UTexture*>& GetReferencedTextures() const override
-	{
-		return ReferencedTextures;
-	}
-
-	////////////////
-	// FMaterialRenderProxy interface.
-	virtual const FMaterial* GetMaterial(ERHIFeatureLevel::Type FeatureLevel) const override
-	{
-		if(GetRenderingThreadShaderMap())
-		{
-			return this;
-		}
-		else
-		{
-			return UMaterial::GetDefaultMaterial(MD_Surface)->GetRenderProxy(false)->GetMaterial(FeatureLevel);
-		}
-	}
-
-	virtual bool GetVectorValue(const FName ParameterName, FLinearColor* OutValue, const FMaterialRenderContext& Context) const override
-	{
-		return MaterialInterface->GetRenderProxy(0)->GetVectorValue(ParameterName, OutValue, Context);
-	}
-
-	virtual bool GetScalarValue(const FName ParameterName, float* OutValue, const FMaterialRenderContext& Context) const override
-	{
-		return MaterialInterface->GetRenderProxy(0)->GetScalarValue(ParameterName, OutValue, Context);
-	}
-
-	virtual bool GetTextureValue(const FName ParameterName,const UTexture** OutValue, const FMaterialRenderContext& Context) const override
-	{
-		return MaterialInterface->GetRenderProxy(0)->GetTextureValue(ParameterName,OutValue,Context);
-	}
-
-	// Material properties.
-	/** Entry point for compiling a specific material property.  This must call SetMaterialProperty. */
-	virtual int32 CompilePropertyAndSetMaterialProperty(EMaterialProperty Property, FMaterialCompiler* Compiler, EShaderFrequency OverrideShaderFrequency, bool bUsePreviousFrameTime) const override
-	{
-		// needs to be called in this function!!
-		Compiler->SetMaterialProperty(Property, OverrideShaderFrequency, bUsePreviousFrameTime);
-
-		int32 Ret = CompilePropertyAndSetMaterialPropertyWithoutCast(Property, Compiler);
-
-		return Compiler->ForceCast(Ret, GetMaterialPropertyType(Property));
-	}
-
-	/** helper for CompilePropertyAndSetMaterialProperty() */
-	int32 CompilePropertyAndSetMaterialPropertyWithoutCast(EMaterialProperty Property, FMaterialCompiler* Compiler) const
-	{
-		if (Property == MP_EmissiveColor)
-		{
-			UMaterial* ProxyMaterial = MaterialInterface->GetMaterial();
-			EBlendMode BlendMode = MaterialInterface->GetBlendMode();
-			EMaterialShadingModel ShadingModel = MaterialInterface->GetShadingModel();
-			check(ProxyMaterial);
-			switch (PropertyToCompile)
-			{
-			case MP_EmissiveColor:
-				// Emissive is ALWAYS returned...
-				return Compiler->ForceCast(MaterialInterface->CompileProperty(Compiler,MP_EmissiveColor),MCT_Float3,true,true);
-			case MP_BaseColor:
-				// Only return for Opaque and Masked...
-				if (BlendMode == BLEND_Opaque || BlendMode == BLEND_Masked)
-				{
-					return Compiler->ForceCast(MaterialInterface->CompileProperty(Compiler, MP_BaseColor),MCT_Float3,true,true);
-				}
-				break;
-			case MP_Specular: 
-			case MP_Roughness:
-			case MP_Metallic:
-				// Only return for Opaque and Masked...
-				if (BlendMode == BLEND_Opaque || BlendMode == BLEND_Masked)
-				{
-					return Compiler->ForceCast(MaterialInterface->CompileProperty(Compiler, PropertyToCompile),MCT_Float,true,true);
-				}
-				break;
-			case MP_Normal:
-				// Only return for Opaque and Masked...
-				if (BlendMode == BLEND_Opaque || BlendMode == BLEND_Masked)
-				{
-					return Compiler->ForceCast( 
-						Compiler->Add( 
-							Compiler->Mul(MaterialInterface->CompileProperty(Compiler, MP_Normal), Compiler->Constant(0.5f)), // [-1,1] * 0.5
-							Compiler->Constant(0.5f)), // [-0.5,0.5] + 0.5
-						MCT_Float3, true, true );
-				}
-				break;
-			default:
-				return Compiler->Constant(1.0f);
-			}
-	
-			return Compiler->Constant(0.0f);
-		}
-		else if (Property == MP_WorldPositionOffset)
-		{
-			//This property MUST return 0 as a default or during the process of rendering textures out for lightmass to use, pixels will be off by 1.
-			return Compiler->Constant(0.0f);
-		}
-		else if (Property >= MP_CustomizedUVs0 && Property <= MP_CustomizedUVs7)
-		{
-			// Pass through customized UVs
-			return MaterialInterface->CompileProperty(Compiler, Property);
-		}
-		else
-		{
-			return Compiler->Constant(1.0f);
-		}
-	}
-
-	virtual FString GetMaterialUsageDescription() const override
-	{
-		return FString::Printf(TEXT("FExportMaterialRenderer %s"), MaterialInterface ? *MaterialInterface->GetName() : TEXT("NULL"));
-	}
-	virtual int32 GetMaterialDomain() const override
-	{
-		if (Material)
-		{
-			return Material->MaterialDomain;
-		}
-		return MD_Surface;
-	}
-	virtual bool IsTwoSided() const  override
-	{ 
-		if (MaterialInterface)
-		{
-			return MaterialInterface->IsTwoSided();
-		}
-		return false;
-	}
-	virtual bool IsDitheredLODTransition() const  override
-	{ 
-		if (MaterialInterface)
-		{
-			return MaterialInterface->IsDitheredLODTransition();
-		}
-		return false;
-	}
-	virtual bool IsLightFunction() const override
-	{
-		if (Material)
-		{
-			return (Material->MaterialDomain == MD_LightFunction);
-		}
-		return false;
-	}
-	virtual bool IsUsedWithDeferredDecal() const override
-	{
-		return	Material &&
-				Material->MaterialDomain == MD_DeferredDecal;
-	}
-	virtual bool IsSpecialEngineMaterial() const override
-	{
-		if (Material)
-		{
-			return (Material->bUsedAsSpecialEngineMaterial == 1);
-		}
-		return false;
-	}
-	virtual bool IsWireframe() const override
-	{
-		if (Material)
-		{
-			return (Material->Wireframe == 1);
-		}
-		return false;
-	}
-	virtual bool IsMasked() const override								{ return false; }
-	virtual enum EBlendMode GetBlendMode() const override				{ return BLEND_Opaque; }
-	virtual enum EMaterialShadingModel GetShadingModel() const override	{ return MSM_Unlit; }
-	virtual float GetOpacityMaskClipValue() const override				{ return 0.5f; }
-	virtual FString GetFriendlyName() const override { return FString::Printf(TEXT("FExportMaterialRenderer %s"), MaterialInterface ? *MaterialInterface->GetName() : TEXT("NULL")); }
-	/**
-	* Should shaders compiled for this material be saved to disk?
-	*/
-	virtual bool IsPersistent() const override { return false; }
-	virtual FGuid GetMaterialId() const override { return Id; }
-
-	const UMaterialInterface* GetMaterialInterface() const
-	{
-		return MaterialInterface;
-	}
-
-	friend FArchive& operator<< ( FArchive& Ar, FExportMaterialProxy& V )
-	{
-		return Ar << V.MaterialInterface;
-	}
-
-	bool IsMaterialInputConnected(UMaterial* InMaterial, EMaterialProperty MaterialInput)
-	{
-		bool bConnected = false;
-
-		switch (MaterialInput)
-		{
-		case MP_EmissiveColor:
-			bConnected = InMaterial->EmissiveColor.Expression != NULL;
-			break;
-		case MP_BaseColor:
-			bConnected = InMaterial->BaseColor.Expression != NULL;
-			break;
-		case MP_Specular:
-			bConnected = InMaterial->Specular.Expression != NULL;
-			break;
-		case MP_Normal:
-			bConnected = InMaterial->Normal.Expression != NULL;
-			break;
-		case MP_Metallic:
-			bConnected = InMaterial->Metallic.Expression != NULL;
-			break;
-		case MP_Roughness:
-			bConnected = InMaterial->Roughness.Expression != NULL;
-			break;
-		default:
-			break;
-		}
-
-		// Note: only checking to see whether the entire material attributes connection exists.  
-		// This means materials using the material attributes input will export more attributes than is necessary.
-		bConnected = InMaterial->bUseMaterialAttributes ? InMaterial->MaterialAttributes.Expression != NULL : bConnected;
-		return bConnected;
-	}
-
-	/**
-	 *	Checks if the configuration of the material proxy will generate a uniform
-	 *	value across the sampling... (Ie, nothing is hooked to the property)
-	 *
-	 *	@param	OutUniformValue		The value that will be returned.
-	 *
-	 *	@return	bool				true if a single value would be generated.
-	 *								false if not.
-	 */
-	bool WillGenerateUniformData(FColor& OutUniformValue)
-	{
-		// Pre-fill the value...
-		OutUniformValue.R = 0;
-		OutUniformValue.G = 0;
-		OutUniformValue.B = 0;
-		OutUniformValue.A = 0;
-
-		EBlendMode BlendMode = MaterialInterface->GetBlendMode();
-		EMaterialShadingModel ShadingModel = MaterialInterface->GetShadingModel();
-		
-		check(Material);
-		bool bExpressionIsNULL = false;
-		switch (PropertyToCompile)
-		{
-		case MP_EmissiveColor:
-			// Emissive is ALWAYS returned...
-			bExpressionIsNULL = !IsMaterialInputConnected(Material, PropertyToCompile);
-			break;
-		case MP_BaseColor:
-			// Only return for Opaque and Masked...
-			if (BlendMode == BLEND_Opaque || BlendMode == BLEND_Masked)
-			{
-				bExpressionIsNULL = !IsMaterialInputConnected(Material, PropertyToCompile);
-			}
-			break;
-		case MP_Specular: 
-		case MP_Metallic:
-		case MP_Roughness:
-			// Only return for Opaque and Masked...
-			if (BlendMode == BLEND_Opaque || BlendMode == BLEND_Masked)
-			{
-				bExpressionIsNULL = !IsMaterialInputConnected(Material, PropertyToCompile);
-				OutUniformValue.A = 255;
-			}
-			break;
-		case MP_Normal:
-			// Only return for Opaque and Masked...
-			if (BlendMode == BLEND_Opaque || BlendMode == BLEND_Masked)
-			{
-				bExpressionIsNULL = !IsMaterialInputConnected(Material, PropertyToCompile);
-				OutUniformValue.B = 255;	// Default normal is (0,0,1)
-			}
-			break;
-		}
-
-		return bExpressionIsNULL;
-	}
-
-	/**
-	* Iterate through all textures used by the material and return the maximum texture resolution used
-	* (ideally this could be made dependent of the material property)
-	*
-	* @param MaterialInterface The material to scan for texture size
-	*
-	* @return Size (width and height)
-	*/
-	FIntPoint FindMaxTextureSize(UMaterialInterface* InMaterialInterface, FIntPoint MinimumSize = FIntPoint(1, 1)) const
-	{
-		// static lod settings so that we only initialize them once
-		UTextureLODSettings* GameTextureLODSettings = UDeviceProfileManager::Get().GetActiveProfile()->GetTextureLODSettings();
-
-		TArray<UTexture*> MaterialTextures;
-
-		InMaterialInterface->GetUsedTextures(MaterialTextures, EMaterialQualityLevel::Num, false, GMaxRHIFeatureLevel, false);
-
-		// find the largest texture in the list (applying it's LOD bias)
-		FIntPoint MaxSize = MinimumSize;
-		for (int32 TexIndex = 0; TexIndex < MaterialTextures.Num(); TexIndex++)
-		{
-			UTexture* Texture = MaterialTextures[TexIndex];
-
-			if (Texture == NULL)
-			{
-				continue;
-			}
-
-			// get the max size of the texture
-			FIntPoint LocalSize(0, 0);
-			if (Texture->IsA(UTexture2D::StaticClass()))
-			{
-				UTexture2D* Tex2D = (UTexture2D*)Texture;
-				LocalSize = FIntPoint(Tex2D->GetSizeX(), Tex2D->GetSizeY());
-			}
-			else if (Texture->IsA(UTextureCube::StaticClass()))
-			{
-				UTextureCube* TexCube = (UTextureCube*)Texture;
-				LocalSize = FIntPoint(TexCube->GetSizeX(), TexCube->GetSizeY());
-			}
-
-			int32 LocalBias = GameTextureLODSettings->CalculateLODBias(Texture);
-
-			// bias the texture size based on LOD group
-			FIntPoint BiasedLocalSize(LocalSize.X >> LocalBias, LocalSize.Y >> LocalBias);
-
-			MaxSize.X = FMath::Max(BiasedLocalSize.X, MaxSize.X);
-			MaxSize.Y = FMath::Max(BiasedLocalSize.Y, MaxSize.Y);
-		}
-
-		return MaxSize;
-	}
-
-	static bool WillFillData(EBlendMode InBlendMode, EMaterialProperty InMaterialProperty)
-	{
-		if (InMaterialProperty == MP_EmissiveColor)
-		{
-			return true;
-		}
-
-		switch (InBlendMode)
-		{
-		case BLEND_Opaque:
-			{
-				switch (InMaterialProperty)
-				{
-				case MP_BaseColor:		return true;
-				case MP_Specular:		return true;
-				case MP_Normal:			return true;
-				case MP_Metallic:		return true;
-				case MP_Roughness:		return true;
-				}
-			}
-			break;
-		}
-		return false;
-	}
-
-private:
-	/** The material interface for this proxy */
-	UMaterialInterface* MaterialInterface;
-	UMaterial* Material;	
-	TArray<UTexture*> ReferencedTextures;
-	/** The property to compile for rendering the sample */
-	EMaterialProperty PropertyToCompile;
-	FGuid Id;
-};
-
-/**
- * Render a material to a render target
- *
- */
-bool GenerateExportMaterialPropertyData(
-		UMaterialInterface* InMaterial, EMaterialProperty InMaterialProperty, 
-		int32& InOutSizeX, int32& InOutSizeY, 
-		TArray<FColor>& OutBMP,
-		UTextureRenderTarget2D*& RenderTarget,
-		FCanvas*& Canvas)
-{
-	FExportMaterialProxy* MaterialProxy = new FExportMaterialProxy(InMaterial, InMaterialProperty);
-	if (MaterialProxy == NULL)
-	{
-		UE_LOG(LogEditorExporters, Warning, TEXT("Failed to create FExportMaterialProxy!"));
-		return false;
-	}
-
-	bool bNormalmap = (InMaterialProperty == MP_Normal);
-
-	FReadSurfaceDataFlags ReadPixelFlags(bNormalmap ? RCM_SNorm : RCM_UNorm);
-
-	bool bResult = true;
-
-	FColor UniformValue;
-	if (MaterialProxy->WillGenerateUniformData(UniformValue) == false)
-	{
-		//@todo UE4. The format may be determined by the material property...
-		// For example, if Diffuse doesn't need to be F16 it can create a standard RGBA8 target.
-		EPixelFormat Format = PF_FloatRGB;
-
-		FIntPoint MaxTextureSize = MaterialProxy->FindMaxTextureSize(InMaterial);
-		InOutSizeX = MaxTextureSize.X;
-		InOutSizeY = MaxTextureSize.Y;
-
-		bool bForceLinear = bNormalmap;
-
-		if(RenderTarget && 
-			((RenderTarget->bForceLinearGamma != bForceLinear) || (RenderTarget->OverrideFormat != Format) || (RenderTarget->SizeX != InOutSizeX) || (RenderTarget->SizeY != InOutSizeY))
-			)
-		{
-			RenderTarget->RemoveFromRoot();
-			RenderTarget = NULL;
-			delete Canvas;
-			Canvas = NULL;
-		}
-
-		if(!RenderTarget)
-		{
-			RenderTarget = NewObject<UTextureRenderTarget2D>();
-			check(RenderTarget);
-			RenderTarget->AddToRoot();
-			RenderTarget->ClearColor = FLinearColor(0.0f, 0.0f, 0.0f, 0.0f);
-			RenderTarget->InitCustomFormat(InOutSizeX, InOutSizeY, Format, bForceLinear);
-
-			Canvas = new FCanvas(RenderTarget->GetRenderTargetResource(), NULL, 0, 0, 0, GMaxRHIFeatureLevel);
-			check(Canvas);
-		}
-
-		// start drawing to the render target
-		Canvas->SetRenderTarget_GameThread(RenderTarget->GetRenderTargetResource());
-		// Freeze time while capturing the material's inputs
-		Canvas->Clear(FLinearColor(0,0,0,0));
-		FCanvasTileItem TileItem( FVector2D( 0.0f, 0.0f ), MaterialProxy, FVector2D( InOutSizeX, InOutSizeY ) );
-		TileItem.bFreezeTime = true;
-		Canvas->DrawItem( TileItem );
-		Canvas->Flush_GameThread();
-		FlushRenderingCommands();
-		Canvas->SetRenderTarget_GameThread(NULL);
-		FlushRenderingCommands();
-
-		// if PF_FloatRGB was used as render target format, the gamma conversion during rendering is deactivated
-		// if we want it not in linear space we need to convert
-		ReadPixelFlags.SetLinearToGamma(Format == PF_FloatRGB && !bForceLinear); //-V560 //remove after implementing todo (near to Format variable declaration)
-
-		if(!RenderTarget->GetRenderTargetResource()->ReadPixels(OutBMP, ReadPixelFlags))
-		{
-			bResult = false;
-		}
-	}
-	else
-	{
-		// Single value... fill it in.
-		InOutSizeX = 1;
-		InOutSizeY = 1; 
-		OutBMP.Empty(1);
-		OutBMP.AddZeroed(1);
-		OutBMP[0] = UniformValue;
-	}
-
-	return bResult;
-}
-
-
-
-
 /*------------------------------------------------------------------------------
 	Helper classes for the OBJ exporters.
 ------------------------------------------------------------------------------*/
@@ -1488,34 +956,23 @@ static void AddActorToOBJs(AActor* Actor, TArray<FOBJGeom*>& Objects, TSet<UMate
 	}
 }
 
-
 // @param Material must not be 0
 // @param MatProp e.g. MP_DiffuseColor
-static void ExportMaterialPropertyTexture(
-	const FString &BMPFilename, 
-	UMaterialInterface* Material, 
-	const EMaterialProperty MatProp,
-	UTextureRenderTarget2D*& RenderTarget,
-	FCanvas*& Canvas
-	)
+static void ExportMaterialPropertyTexture(const FString &BMPFilename, UMaterialInterface* Material, const EMaterialProperty MatProp)
 {
 	check(Material);
 
 	// make the BMP for the diffuse channel
 	TArray<FColor> OutputBMP;
-	int32 SizeX = 1024;
-	int32 SizeY = 1024;
-
+	FIntPoint Size;
+	
 	TEnumAsByte<EBlendMode> BlendMode = Material->GetBlendMode();
-	bool bIsValidMaterial = FExportMaterialProxy::WillFillData((EBlendMode)(BlendMode), MatProp);
+	bool bIsValidMaterial = FMaterialUtilities::SupportsExport((EBlendMode)(BlendMode), MatProp);
 
 	if (bIsValidMaterial)
 	{
-		// make space for the bmp
-		OutputBMP.AddUninitialized(SizeX * SizeY);
-
 		// render the material to a texture to export as a bmp
-		if (!GenerateExportMaterialPropertyData(Material, MatProp, SizeX, SizeY, OutputBMP, RenderTarget, Canvas))
+		if (!FMaterialUtilities::ExportMaterialProperty(nullptr, Material, MatProp, Size, OutputBMP))
 		{
 			bIsValidMaterial = false;
 		}
@@ -1524,13 +981,13 @@ static void ExportMaterialPropertyTexture(
 	// make invalid textures a solid red
 	if (!bIsValidMaterial)
 	{
-		SizeX = SizeY = 1;
+		Size = FIntPoint(1,1);
 		OutputBMP.Empty();
 		OutputBMP.Add(FColor(255, 0, 0, 255));
 	}
 
 	// export the diffuse channel bmp
-	FFileHelper::CreateBitmap(*BMPFilename, SizeX, SizeY, OutputBMP.GetData());
+	FFileHelper::CreateBitmap(*BMPFilename, Size.X, Size.Y, OutputBMP.GetData());
 }
 
 /**
@@ -1562,9 +1019,6 @@ void ExportOBJs(FOutputDevice& FileAr, FStringOutputDevice* MemAr, FFeedbackCont
 		MaterialLib->SetSuppressEventTag(true);
 		MaterialLib->SetAutoEmitLineTerminator(false);
 
-		UTextureRenderTarget2D* RenderTarget = NULL;
-		FCanvas* Canvas = NULL;
-
 		// export the material set to a mtllib
 		int32 MaterialIndex = 0;
 		for (TSet<UMaterialInterface*>::TConstIterator It(*Materials); It; ++It, ++MaterialIndex)
@@ -1577,33 +1031,25 @@ void ExportOBJs(FOutputDevice& FileAr, FStringOutputDevice* MemAr, FFeedbackCont
 
 			{
 				FString BMPFilename = FPaths::GetPath(MaterialLibFilename) / MaterialName + TEXT("_D.bmp");
-				ExportMaterialPropertyTexture(BMPFilename, Material, MP_BaseColor, RenderTarget, Canvas);
+				ExportMaterialPropertyTexture(BMPFilename, Material, MP_BaseColor);
 				MaterialLib->Logf(TEXT("\tmap_Kd %s\r\n"), *FPaths::GetCleanFilename(BMPFilename));
 			}
 
 			{
 				FString BMPFilename = FPaths::GetPath(MaterialLibFilename) / MaterialName + TEXT("_S.bmp");
-				ExportMaterialPropertyTexture(BMPFilename, Material, MP_Specular, RenderTarget, Canvas);
+				ExportMaterialPropertyTexture(BMPFilename, Material, MP_Specular);
 				MaterialLib->Logf(TEXT("\tmap_Ks %s\r\n"), *FPaths::GetCleanFilename(BMPFilename));
 			}
 
 			{
 				FString BMPFilename = FPaths::GetPath(MaterialLibFilename) / MaterialName + TEXT("_N.bmp");
-				ExportMaterialPropertyTexture(BMPFilename, Material, MP_Normal, RenderTarget, Canvas);
+				ExportMaterialPropertyTexture(BMPFilename, Material, MP_Normal);
 				MaterialLib->Logf(TEXT("\tbump %s\r\n"), *FPaths::GetCleanFilename(BMPFilename));
 			}
 
 			MaterialLib->Logf(TEXT("\r\n"));
 		}
 		
-		if(RenderTarget)
-		{
-			RenderTarget->RemoveFromRoot();
-			RenderTarget = NULL;
-			delete Canvas;
-			Canvas = NULL;
-		}
-
 		MaterialLib->TearDown();
 		delete MaterialLib;
 
@@ -2464,651 +1910,43 @@ UExportTextContainer::UExportTextContainer(const FObjectInitializer& ObjectIniti
 -----------------------------------------------------------------------------*/
 namespace MaterialExportUtils
 {
-	void RenderMaterialTile(UWorld* InWorld, FMaterialRenderProxy* InMaterialProxy, UTextureRenderTarget2D* InRenderTarget, bool bFrontView)
-	{
-		float CurrentRealTime = 0.f;
-		float CurrentWorldTime = 0.f;
-		float DeltaWorldTime = 0.f;
-	
-		if (InWorld)
-		{
-			CurrentRealTime = InWorld->GetRealTimeSeconds();
-			CurrentWorldTime = InWorld->GetTimeSeconds();
-			DeltaWorldTime = InWorld->GetDeltaSeconds();
-		}
-		else
-		{
-			CurrentRealTime = FApp::GetCurrentTime() - GStartTime;
-			CurrentWorldTime = FApp::GetDeltaTime();
-			DeltaWorldTime = FApp::GetCurrentTime() - GStartTime;
-		}
-
-		const FRenderTarget* RenderTargetResource = InRenderTarget->GameThread_GetRenderTargetResource();
-		FSceneViewFamily* ViewFamily = new FSceneViewFamily(FSceneViewFamily::ConstructionValues(
-			RenderTargetResource,
-			InWorld ? InWorld->Scene : nullptr,
-			FEngineShowFlags(ESFIM_Game))
-			.SetWorldTimes(CurrentWorldTime, DeltaWorldTime, CurrentRealTime)
-			.SetGammaCorrection(RenderTargetResource->GetDisplayGamma()));
-
-		FIntPoint ViewSize = RenderTargetResource->GetSizeXY();
-		FIntRect ViewRect(FIntPoint(0, 0), ViewSize);
-		// By default render tile in top view
-		FMatrix ViewMatrix = FMatrix(
-					FPlane(1, 0, 0, 0),
-					FPlane(0, -1, 0, 0),
-					FPlane(0, 0, -1, 0),
-					FPlane(0, 0, 0, 1));
-		FQuat Rotation = FQuat::Identity;
-
-		if (bFrontView)
-		{
-			ViewMatrix = FMatrix(
-					FPlane(1, 0, 0, 0),
-					FPlane(0, 0, -1, 0),
-					FPlane(0, 1, 0, 0),
-					FPlane(0, 0, 0, 1));
-			// Tile mesh was created on XY plane
-			Rotation = FRotator(0.0f, 0.0f, 90.0f).Quaternion();
-		}
-						
-		// make a temporary view
-		FSceneViewInitOptions ViewInitOptions;
-		ViewInitOptions.ViewFamily = ViewFamily;
-		ViewInitOptions.SetViewRectangle(ViewRect);
-		ViewInitOptions.ViewOrigin = FVector(ViewSize.X/2.f, ViewSize.Y/2.f, 0.0f);
-		ViewInitOptions.ViewRotationMatrix = ViewMatrix;
-		ViewInitOptions.ProjectionMatrix = FReversedZOrthoMatrix(ViewSize.X/2.f, ViewSize.Y/2.f, 0.5f/HALF_WORLD_MAX, HALF_WORLD_MAX);
-		ViewInitOptions.BackgroundColor = FLinearColor::Black;
-		ViewInitOptions.OverlayColor = FLinearColor::White;
-				
-		FSceneView* View = new FSceneView(ViewInitOptions);
-						
-		ENQUEUE_UNIQUE_RENDER_COMMAND_FIVEPARAMETER(
-			RenderMaterialTileCommand,
-			FSceneView*, InView, View,
-			FMaterialRenderProxy*, InProxy, InMaterialProxy,
-			const FRenderTarget*, InRenderTarget, RenderTargetResource,
-			FIntPoint, InSize, ViewSize,
-			FQuat, InRotation, Rotation,
-			{
-				::SetRenderTarget(RHICmdList, InRenderTarget->GetRenderTargetTexture(), FTextureRHIRef());
-				FIntRect ViewportRect = FIntRect(FIntPoint::ZeroValue, InSize);
-				RHICmdList.SetViewport(ViewportRect.Min.X, ViewportRect.Min.Y, 0.0f, ViewportRect.Max.X, ViewportRect.Max.Y, 1.0f);
-		
-				FTileRenderer::DrawRotatedTile(RHICmdList, *InView, InProxy, false, InRotation, 0.f, 0.f, InSize.X, InSize.Y,  0.f, 0.f, 1.f, 1.f);
-			});
-
-		FlushRenderingCommands();
-
-		delete View;
-	}
-	
 	bool ExportMaterialProperty(UMaterialInterface* InMaterial,	EMaterialProperty InMaterialProperty, UTextureRenderTarget2D* InRenderTarget, TArray<FColor>& OutBMP)
 	{
-		return ExportMaterialProperty(nullptr, InMaterial, InMaterialProperty, InRenderTarget, OutBMP);
+		return FMaterialUtilities::ExportMaterialProperty(nullptr, InMaterial, InMaterialProperty, InRenderTarget, OutBMP);
 	}
 		
 	bool ExportMaterialProperty(UWorld* InWorld, UMaterialInterface* InMaterial, EMaterialProperty InMaterialProperty, UTextureRenderTarget2D* InRenderTarget, TArray<FColor>& OutBMP)
 	{
-		TScopedPointer<FExportMaterialProxy> MaterialProxy(new FExportMaterialProxy(InMaterial, InMaterialProperty));
-		if (MaterialProxy == NULL)
-		{
-			return false;
-		}
-
-		FColor UniformValue;
-		if (MaterialProxy->WillGenerateUniformData(UniformValue))
-		{
-			// Single value... fill it in.
-			OutBMP.Empty();
-			OutBMP.Add(UniformValue);
-			return true;
-		}
-	
-		check(InRenderTarget);
-
-		bool bNormalmap = (InMaterialProperty == MP_Normal);
-
-		RenderMaterialTile(InWorld, MaterialProxy, InRenderTarget, !bNormalmap);
-					
-		FReadSurfaceDataFlags ReadPixelFlags(bNormalmap ? RCM_SNorm : RCM_UNorm);
-		ReadPixelFlags.SetLinearToGamma(false);
-
-		FTextureRenderTargetResource* RTResource = InRenderTarget->GameThread_GetRenderTargetResource();
-		return RTResource->ReadPixels(OutBMP, ReadPixelFlags);
+		return FMaterialUtilities::ExportMaterialProperty(InWorld, InMaterial, InMaterialProperty, InRenderTarget, OutBMP);
 	}
 	
 	bool ExportMaterial(UMaterialInterface* InMaterial, FFlattenMaterial& OutFlattenMaterial)
 	{
-		return ExportMaterial(nullptr, InMaterial, OutFlattenMaterial);
+		return FMaterialUtilities::ExportMaterial(nullptr, InMaterial, OutFlattenMaterial);
 	}
 	
 	bool ExportMaterial(UWorld* InWorld, UMaterialInterface* InMaterial, FFlattenMaterial& OutFlattenMaterial)
 	{
-		// Render diffuse property
-		if (OutFlattenMaterial.DiffuseSize.X > 0 && 
-			OutFlattenMaterial.DiffuseSize.Y > 0)
-		{
-			// Create temporary render target
-			auto RenderTargetDiffuse = NewObject<UTextureRenderTarget2D>();
-			check(RenderTargetDiffuse);
-			RenderTargetDiffuse->AddToRoot();
-			RenderTargetDiffuse->ClearColor = FLinearColor(0.0f, 0.0f, 0.0f, 0.0f);
-			RenderTargetDiffuse->InitCustomFormat(
-				OutFlattenMaterial.DiffuseSize.X, 
-				OutFlattenMaterial.DiffuseSize.Y, PF_B8G8R8A8, false);
-			
-			//
-			OutFlattenMaterial.DiffuseSamples.Empty(OutFlattenMaterial.DiffuseSize.X * OutFlattenMaterial.DiffuseSize.Y);
-			bool bResult = ExportMaterialProperty(InMaterial, MP_BaseColor, RenderTargetDiffuse, OutFlattenMaterial.DiffuseSamples);
-
-			// Uniform value
-			if (OutFlattenMaterial.DiffuseSamples.Num() == 1)
-			{
-				OutFlattenMaterial.DiffuseSize = FIntPoint(1,1);
-			}
-			
-			RenderTargetDiffuse->RemoveFromRoot();
-			RenderTargetDiffuse = nullptr;
-
-			if (!bResult)
-			{
-				return false;
-			}
-		}
-
-		// Render normal property
-		if (OutFlattenMaterial.NormalSize.X > 0 && 
-			OutFlattenMaterial.NormalSize.Y > 0 &&
-			InMaterial->GetMaterial()->HasNormalConnected())
-		{
-			// Create temporary render target
-			auto RenderTargetNormal = NewObject<UTextureRenderTarget2D>();
-			check(RenderTargetNormal);
-			RenderTargetNormal->AddToRoot();
-			RenderTargetNormal->ClearColor = FLinearColor(0.0f, 0.0f, 0.0f, 0.0f);
-			RenderTargetNormal->InitCustomFormat(
-				OutFlattenMaterial.NormalSize.X,
-				OutFlattenMaterial.NormalSize.Y, PF_FloatRGB, true);
-
-			//
-			OutFlattenMaterial.NormalSamples.Empty(OutFlattenMaterial.NormalSize.X * OutFlattenMaterial.NormalSize.Y);
-			bool bResult = ExportMaterialProperty(InMaterial, MP_Normal, RenderTargetNormal, OutFlattenMaterial.NormalSamples);
-
-			// Uniform value
-			if (OutFlattenMaterial.NormalSamples.Num() == 1)
-			{
-				OutFlattenMaterial.NormalSize = FIntPoint(1,1);
-			}
-		
-			RenderTargetNormal->RemoveFromRoot();
-			RenderTargetNormal = nullptr;
-
-			if (!bResult)
-			{
-				return false;
-			}
-		}
-
-		// Render metallic property
-		if (OutFlattenMaterial.MetallicSize.X > 0 && 
-			OutFlattenMaterial.MetallicSize.Y > 0)
-		{
-			// Create temporary render target
-			auto RenderTargetMetallic = NewObject<UTextureRenderTarget2D>();
-			check(RenderTargetMetallic);
-			RenderTargetMetallic->AddToRoot();
-			RenderTargetMetallic->ClearColor = FLinearColor(0.0f, 0.0f, 0.0f, 0.0f);
-			RenderTargetMetallic->InitCustomFormat(
-				OutFlattenMaterial.MetallicSize.X,
-				OutFlattenMaterial.MetallicSize.Y, PF_B8G8R8A8, true);
-
-			//
-			OutFlattenMaterial.MetallicSamples.Empty(OutFlattenMaterial.MetallicSize.X * OutFlattenMaterial.MetallicSize.Y);
-			bool bResult = ExportMaterialProperty(InMaterial, MP_Metallic, RenderTargetMetallic, OutFlattenMaterial.MetallicSamples);
-
-			// Uniform value
-			if (OutFlattenMaterial.MetallicSamples.Num() == 1)
-			{
-				OutFlattenMaterial.MetallicSize = FIntPoint(1,1);
-			}
-		
-			RenderTargetMetallic->RemoveFromRoot();
-			RenderTargetMetallic = nullptr;
-
-			if (!bResult)
-			{
-				return false;
-			}
-		}
-
-		// Render roughness property
-		if (OutFlattenMaterial.RoughnessSize.X > 0 && 
-			OutFlattenMaterial.RoughnessSize.Y > 0)
-		{
-			// Create temporary render target
-			auto RenderTargetRoughness = NewObject<UTextureRenderTarget2D>();
-			check(RenderTargetRoughness);
-			RenderTargetRoughness->AddToRoot();
-			RenderTargetRoughness->ClearColor = FLinearColor(0.0f, 0.0f, 0.0f, 0.0f);
-			RenderTargetRoughness->InitCustomFormat(
-				OutFlattenMaterial.RoughnessSize.X,
-				OutFlattenMaterial.RoughnessSize.Y, PF_B8G8R8A8, true);
-
-			//
-			OutFlattenMaterial.RoughnessSamples.Empty(OutFlattenMaterial.RoughnessSize.X * OutFlattenMaterial.RoughnessSize.Y);
-			bool bResult = ExportMaterialProperty(InMaterial, MP_Roughness, RenderTargetRoughness, OutFlattenMaterial.RoughnessSamples);
-
-			// Uniform value
-			if (OutFlattenMaterial.RoughnessSamples.Num() == 1)
-			{
-				OutFlattenMaterial.RoughnessSize = FIntPoint(1,1);
-			}
-		
-			RenderTargetRoughness->RemoveFromRoot();
-			RenderTargetRoughness = nullptr;
-
-			if (!bResult)
-			{
-				return false;
-			}
-		}
-
-		// Render specular property
-		if (OutFlattenMaterial.SpecularSize.X > 0 && 
-			OutFlattenMaterial.SpecularSize.Y > 0)
-		{
-			// Create temporary render target
-			auto RenderTargetSpecular = NewObject<UTextureRenderTarget2D>();
-			check(RenderTargetSpecular);
-			RenderTargetSpecular->AddToRoot();
-			RenderTargetSpecular->ClearColor = FLinearColor(0.0f, 0.0f, 0.0f, 0.0f);
-			RenderTargetSpecular->InitCustomFormat(
-				OutFlattenMaterial.SpecularSize.X,
-				OutFlattenMaterial.SpecularSize.Y, PF_B8G8R8A8, true);
-
-			//
-			OutFlattenMaterial.SpecularSamples.Empty(OutFlattenMaterial.SpecularSize.X * OutFlattenMaterial.SpecularSize.Y);
-			bool bResult = ExportMaterialProperty(InMaterial, MP_Specular, RenderTargetSpecular, OutFlattenMaterial.SpecularSamples);
-
-			// Uniform value
-			if (OutFlattenMaterial.SpecularSamples.Num() == 1)
-			{
-				OutFlattenMaterial.SpecularSize = FIntPoint(1,1);
-			}
-		
-			RenderTargetSpecular->RemoveFromRoot();
-			RenderTargetSpecular = nullptr;
-
-			if (!bResult)
-			{
-				return false;
-			}
-		}
-		
-		OutFlattenMaterial.MaterialId = InMaterial->GetLightingGuid();
-		return true;
+		return FMaterialUtilities::ExportMaterial(InWorld, InMaterial, OutFlattenMaterial);
 	}
 
-	void RenderSceneToTexture(
-		FSceneInterface* Scene,
-		const FName& VisualizationMode, 
-		const FVector& ViewOrigin,
-		const FMatrix& ViewRotationMatrix, 
-		const FMatrix& ProjectionMatrix,  
-		const TSet<FPrimitiveComponentId>& HiddenPrimitives, 
-		FIntPoint TargetSize,
-		float TargetGamma,
-		TArray<FColor>& OutSamples)
-	{
-		auto RenderTargetTexture = NewObject<UTextureRenderTarget2D>();
-		check(RenderTargetTexture);
-		RenderTargetTexture->AddToRoot();
-		RenderTargetTexture->ClearColor = FLinearColor::Transparent;
-		RenderTargetTexture->TargetGamma = TargetGamma;
-		RenderTargetTexture->InitCustomFormat(TargetSize.X, TargetSize.Y, PF_FloatRGBA, false);
-		FTextureRenderTargetResource* RenderTargetResource = RenderTargetTexture->GameThread_GetRenderTargetResource();
-
-		FSceneViewFamilyContext ViewFamily(
-			FSceneViewFamily::ConstructionValues(RenderTargetResource, Scene, FEngineShowFlags(ESFIM_Game))
-				.SetWorldTimes(FApp::GetCurrentTime() - GStartTime, FApp::GetDeltaTime(), FApp::GetCurrentTime() - GStartTime)
-			);
-
-		// To enable visualization mode
-		ViewFamily.EngineShowFlags.SetPostProcessing(true);
-		ViewFamily.EngineShowFlags.SetVisualizeBuffer(true);
-		ViewFamily.EngineShowFlags.SetTonemapper(false);
-
-		FSceneViewInitOptions ViewInitOptions;
-		ViewInitOptions.SetViewRectangle(FIntRect(0, 0, TargetSize.X, TargetSize.Y));
-		ViewInitOptions.ViewFamily = &ViewFamily;
-		ViewInitOptions.HiddenPrimitives = HiddenPrimitives;
-		ViewInitOptions.ViewOrigin = ViewOrigin;
-		ViewInitOptions.ViewRotationMatrix = ViewRotationMatrix;
-		ViewInitOptions.ProjectionMatrix = ProjectionMatrix;
-		
-		FSceneView* NewView = new FSceneView(ViewInitOptions);
-		NewView->CurrentBufferVisualizationMode = VisualizationMode;
-		ViewFamily.Views.Add(NewView);
-					
-		FCanvas Canvas(RenderTargetResource, NULL, FApp::GetCurrentTime() - GStartTime, FApp::GetDeltaTime(), FApp::GetCurrentTime() - GStartTime, Scene->GetFeatureLevel());
-		Canvas.Clear(FLinearColor::Transparent);
-		GetRendererModule().BeginRenderingViewFamily(&Canvas, &ViewFamily);
-
-		// Copy the contents of the remote texture to system memory
-		OutSamples.SetNumUninitialized(TargetSize.X*TargetSize.Y);
-		FReadSurfaceDataFlags ReadSurfaceDataFlags;
-		ReadSurfaceDataFlags.SetLinearToGamma(false);
-		RenderTargetResource->ReadPixelsPtr(OutSamples.GetData(), ReadSurfaceDataFlags, FIntRect(0, 0, TargetSize.X, TargetSize.Y));
-		FlushRenderingCommands();
-					
-		RenderTargetTexture->RemoveFromRoot();
-		RenderTargetTexture = nullptr;
-	}
-		
 	bool ExportMaterial(ALandscapeProxy* InLandscape, const TSet<FPrimitiveComponentId>& HiddenPrimitives, FFlattenMaterial& OutFlattenMaterial)
 	{
-		check(InLandscape);
-
-		FIntRect LandscapeRect = InLandscape->GetBoundingRect();
-		FVector MidPoint = FVector(LandscapeRect.Min, 0.f) + FVector(LandscapeRect.Size(), 0.f)*0.5f;
-		
-		FVector LandscapeCenter = InLandscape->GetTransform().TransformPosition(MidPoint);
-		FVector LandscapeExtent = FVector(LandscapeRect.Size(), 0.f)*InLandscape->GetActorScale()*0.5f; 
-
-		FVector ViewOrigin = LandscapeCenter;
-		FMatrix ViewRotationMatrix = FInverseRotationMatrix(InLandscape->GetActorRotation());
-		ViewRotationMatrix*= FMatrix(FPlane(1,	0,	0,	0),
-							 FPlane(0,	-1,	0,	0),
-							 FPlane(0,	0,	-1,	0),
-							 FPlane(0,	0,	0,	1));
-				
-		const float ZOffset = WORLD_MAX;
-		FMatrix ProjectionMatrix =  FReversedZOrthoMatrix(
-			LandscapeExtent.X,
-			LandscapeExtent.Y,
-			0.5f / ZOffset,
-			ZOffset);
-
-		FSceneInterface* Scene = InLandscape->GetWorld()->Scene;
-						
-		// Render diffuse texture using BufferVisualizationMode=BaseColor
-		if (OutFlattenMaterial.DiffuseSize.X > 0 && 
-			OutFlattenMaterial.DiffuseSize.Y > 0)
-		{
-			static const FName BaseColorName("BaseColor");
-			const float BaseColorGamma = 2.2f; // BaseColor to gamma space
-			RenderSceneToTexture(Scene, BaseColorName, ViewOrigin, ViewRotationMatrix, ProjectionMatrix, HiddenPrimitives, 
-				OutFlattenMaterial.DiffuseSize, BaseColorGamma, OutFlattenMaterial.DiffuseSamples);
-		}
-
-		// Render normal map using BufferVisualizationMode=WorldNormal
-		// Final material should use world space instead of tangent space for normals
-		if (OutFlattenMaterial.NormalSize.X > 0 && 
-			OutFlattenMaterial.NormalSize.Y > 0)
-		{
-			static const FName WorldNormalName("WorldNormal");
-			const float NormalColorGamma = 1.0f; // Dump normal texture in linear space
-			RenderSceneToTexture(Scene, WorldNormalName, ViewOrigin, ViewRotationMatrix, ProjectionMatrix, HiddenPrimitives, 
-				OutFlattenMaterial.NormalSize, NormalColorGamma, OutFlattenMaterial.NormalSamples);
-		}
-
-		// Render metallic map using BufferVisualizationMode=Metallic
-		if (OutFlattenMaterial.MetallicSize.X > 0 && 
-			OutFlattenMaterial.MetallicSize.Y > 0)
-		{
-			static const FName MetallicName("Metallic");
-			const float MetallicColorGamma = 1.0f; // Dump metallic texture in linear space
-			RenderSceneToTexture(Scene, MetallicName, ViewOrigin, ViewRotationMatrix, ProjectionMatrix, HiddenPrimitives, 
-				OutFlattenMaterial.MetallicSize, MetallicColorGamma, OutFlattenMaterial.MetallicSamples);
-		}
-
-		// Render roughness map using BufferVisualizationMode=Roughness
-		if (OutFlattenMaterial.RoughnessSize.X > 0 && 
-			OutFlattenMaterial.RoughnessSize.Y > 0)
-		{
-			static const FName RoughnessName("Roughness");
-			const float RoughnessColorGamma = 2.2f; // Roughness material powers color by 2.2, transform it back to linear
-			RenderSceneToTexture(Scene, RoughnessName, ViewOrigin, ViewRotationMatrix, ProjectionMatrix, HiddenPrimitives, 
-				OutFlattenMaterial.RoughnessSize, RoughnessColorGamma, OutFlattenMaterial.RoughnessSamples);
-		}
-
-		// Render specular map using BufferVisualizationMode=Specular
-		if (OutFlattenMaterial.SpecularSize.X > 0 && 
-			OutFlattenMaterial.SpecularSize.Y > 0)
-		{
-			static const FName SpecularName("Specular");
-			const float SpecularColorGamma = 1.0f; // Dump specular texture in linear space
-			RenderSceneToTexture(Scene, SpecularName, ViewOrigin, ViewRotationMatrix, ProjectionMatrix, HiddenPrimitives, 
-				OutFlattenMaterial.SpecularSize, SpecularColorGamma, OutFlattenMaterial.SpecularSamples);
-		}
-				
-		OutFlattenMaterial.MaterialId = InLandscape->GetLandscapeGuid();
-		return true;
+		return FMaterialUtilities::ExportLandscapeMaterial(InLandscape, HiddenPrimitives, OutFlattenMaterial);
 	}
 
 	UTexture2D* CreateTexture(UPackage* Outer, const FString& AssetLongName, FIntPoint Size, const TArray<FColor>& Samples, TextureCompressionSettings CompressionSettings, TextureGroup LODGroup, EObjectFlags Flags, bool bSRGB, const FGuid& SourceGuidHash)
 	{
-		FCreateTexture2DParameters TexParams;
-		TexParams.bUseAlpha = false;
-		TexParams.CompressionSettings = CompressionSettings;
-		TexParams.bDeferCompression = true;
-		TexParams.bSRGB = bSRGB;
-		TexParams.SourceGuidHash = SourceGuidHash;
-
-		if (Outer == nullptr)
-		{
-			Outer = CreatePackage(NULL, *AssetLongName);
-			Outer->FullyLoad();
-			Outer->Modify();
-		}
-
-		UTexture2D* Texture = FImageUtils::CreateTexture2D(Size.X, Size.Y, Samples, Outer, FPackageName::GetShortName(AssetLongName), Flags, TexParams);
-		Texture->LODGroup = LODGroup;
-		Texture->PostEditChange();
-		
-		return Texture;
+		return FMaterialUtilities::CreateTexture(Outer, AssetLongName, Size, Samples, CompressionSettings, LODGroup, Flags, bSRGB, SourceGuidHash);
 	}
 
 	UMaterial* CreateMaterial(const FFlattenMaterial& InFlattenMaterial, UPackage* Outer, const FString& BaseName, EObjectFlags Flags, TArray<UObject*>& OutGeneratedAssets)
 	{
-		// Base name for a new assets
-		// In case outer is null BaseName has to be long package name
-		if (Outer == nullptr && FPackageName::IsShortPackageName(BaseName))
-		{
-			UE_LOG(LogEditorExporters, Warning, TEXT("Invalid long package name: '%s'."), *BaseName);
-			return nullptr;
-		}
-
-		const FString AssetBaseName = FPackageName::GetShortName(BaseName);
-		const FString AssetBasePath = Outer ? TEXT("") : FPackageName::GetLongPackagePath(BaseName) + TEXT("/");
-				
-		// Create material
-		const FString MaterialAssetName = TEXT("M_") + AssetBaseName;
-		UPackage* MaterialOuter = Outer;
-		if (MaterialOuter == NULL)
-		{
-			MaterialOuter = CreatePackage(NULL, *(AssetBasePath + MaterialAssetName));
-			MaterialOuter->FullyLoad();
-			MaterialOuter->Modify();
-		}
-
-		UMaterial* Material = NewObject<UMaterial>(MaterialOuter, FName(*MaterialAssetName), Flags);
-		Material->TwoSided = false;
-		Material->SetShadingModel(MSM_DefaultLit);
-		OutGeneratedAssets.Add(Material);
-
-		int32 MaterialNodeY = -150;
-		int32 MaterialNodeStepY = 180;
-
-		// BaseColor
-		if (InFlattenMaterial.DiffuseSamples.Num() > 1)
-		{
-			const FString AssetName = TEXT("T_") + AssetBaseName + TEXT("_D");
-			const FString AssetLongName = AssetBasePath + AssetName;
-			const bool bSRGB = true;
-			UTexture2D* Texture = CreateTexture(Outer, AssetLongName, InFlattenMaterial.DiffuseSize, InFlattenMaterial.DiffuseSamples, TC_Default, TEXTUREGROUP_World, Flags, bSRGB);
-			OutGeneratedAssets.Add(Texture);
-			
-			auto BasecolorExpression = NewObject<UMaterialExpressionTextureSample>(Material);
-			BasecolorExpression->Texture = Texture;
-			BasecolorExpression->SamplerType = EMaterialSamplerType::SAMPLERTYPE_Color;
-			BasecolorExpression->MaterialExpressionEditorX = -400;
-			BasecolorExpression->MaterialExpressionEditorY = MaterialNodeY;
-			Material->Expressions.Add(BasecolorExpression);
-			Material->BaseColor.Expression = BasecolorExpression;
-
-			MaterialNodeY+= MaterialNodeStepY;
-		}
-
-		// Metallic
-		if (InFlattenMaterial.MetallicSamples.Num() > 1)
-		{
-			const FString AssetName = TEXT("T_") + AssetBaseName + TEXT("_M");
-			const bool bSRGB = false;
-			UTexture2D* Texture = CreateTexture(Outer, AssetBasePath + AssetName, InFlattenMaterial.MetallicSize, InFlattenMaterial.MetallicSamples, TC_Grayscale, TEXTUREGROUP_World, Flags, bSRGB);
-			OutGeneratedAssets.Add(Texture);
-			
-			auto MetallicExpression = NewObject<UMaterialExpressionTextureSample>(Material);
-			MetallicExpression->Texture = Texture;
-			MetallicExpression->SamplerType = EMaterialSamplerType::SAMPLERTYPE_LinearGrayscale;
-			MetallicExpression->MaterialExpressionEditorX = -400;
-			MetallicExpression->MaterialExpressionEditorY = MaterialNodeY;
-			Material->Expressions.Add(MetallicExpression);
-			Material->Metallic.Expression = MetallicExpression;
-
-			MaterialNodeY+= MaterialNodeStepY;
-		}
-		else
-		{
-			// Set Metallic to constant
-			float Metallic = InFlattenMaterial.MetallicSamples.Num() ? InFlattenMaterial.MetallicSamples[0].R : 0.0f;
-			auto MetallicExpression = NewObject<UMaterialExpressionConstant>(Material);
-			MetallicExpression->R = Metallic;
-			MetallicExpression->MaterialExpressionEditorX = -400;
-			MetallicExpression->MaterialExpressionEditorY = MaterialNodeY;
-			Material->Expressions.Add(MetallicExpression);
-			Material->Metallic.Expression = MetallicExpression;
-
-			MaterialNodeY+= MaterialNodeStepY;
-		}
-
-		// Specular
-		if (InFlattenMaterial.SpecularSamples.Num() > 1)
-		{
-			const FString AssetName = TEXT("T_") + AssetBaseName + TEXT("_S");
-			const bool bSRGB = false;
-			UTexture2D* Texture = CreateTexture(Outer, AssetBasePath + AssetName, InFlattenMaterial.SpecularSize, InFlattenMaterial.SpecularSamples, TC_Grayscale, TEXTUREGROUP_World, Flags, bSRGB);
-			OutGeneratedAssets.Add(Texture);
-			
-			auto SpecularExpression = NewObject<UMaterialExpressionTextureSample>(Material);
-			SpecularExpression->Texture = Texture;
-			SpecularExpression->SamplerType = EMaterialSamplerType::SAMPLERTYPE_LinearGrayscale;
-			SpecularExpression->MaterialExpressionEditorX = -400;
-			SpecularExpression->MaterialExpressionEditorY = MaterialNodeY;
-			Material->Expressions.Add(SpecularExpression);
-			Material->Specular.Expression = SpecularExpression;
-
-			MaterialNodeY+= MaterialNodeStepY;
-		}
-
-		// Roughness
-		if (InFlattenMaterial.RoughnessSamples.Num() > 1)
-		{
-			const FString AssetName = TEXT("T_") + AssetBaseName + TEXT("_R");
-			const bool bSRGB = false;
-			UTexture2D* Texture = CreateTexture(Outer, AssetBasePath + AssetName, InFlattenMaterial.RoughnessSize, InFlattenMaterial.RoughnessSamples, TC_Grayscale, TEXTUREGROUP_World, Flags, bSRGB);
-			OutGeneratedAssets.Add(Texture);
-			
-			auto RoughnessExpression = NewObject<UMaterialExpressionTextureSample>(Material);
-			RoughnessExpression->Texture = Texture;
-			RoughnessExpression->SamplerType = EMaterialSamplerType::SAMPLERTYPE_LinearGrayscale;
-			RoughnessExpression->MaterialExpressionEditorX = -400;
-			RoughnessExpression->MaterialExpressionEditorY = MaterialNodeY;
-			Material->Expressions.Add(RoughnessExpression);
-			Material->Roughness.Expression = RoughnessExpression;
-
-			MaterialNodeY+= MaterialNodeStepY;
-		}
-		else
-		{
-			// Set Roughness to constant
-			float Roughness = InFlattenMaterial.RoughnessSamples.Num() ? InFlattenMaterial.RoughnessSamples[0].R : 0.8f;
-			auto RoughnessExpression = NewObject<UMaterialExpressionConstant>(Material);
-			RoughnessExpression->R = Roughness;
-			RoughnessExpression->MaterialExpressionEditorX = -400;
-			RoughnessExpression->MaterialExpressionEditorY = MaterialNodeY;
-			Material->Expressions.Add(RoughnessExpression);
-			Material->Roughness.Expression = RoughnessExpression;
-
-			MaterialNodeY+= MaterialNodeStepY;
-		}
-
-		// Normal
-		if (InFlattenMaterial.NormalSamples.Num() > 1)
-		{
-			const FString AssetName = TEXT("T_") + AssetBaseName + TEXT("_N");
-			const bool bSRGB = false;
-			UTexture2D* Texture = CreateTexture(Outer, AssetBasePath + AssetName, InFlattenMaterial.NormalSize, InFlattenMaterial.NormalSamples, TC_Normalmap, TEXTUREGROUP_WorldNormalMap, Flags, bSRGB);
-			OutGeneratedAssets.Add(Texture);
-			
-			auto NormalExpression = NewObject<UMaterialExpressionTextureSample>(Material);
-			NormalExpression->Texture = Texture;
-			NormalExpression->SamplerType = EMaterialSamplerType::SAMPLERTYPE_Normal;
-			NormalExpression->MaterialExpressionEditorX = -400;
-			NormalExpression->MaterialExpressionEditorY = MaterialNodeY;
-			Material->Expressions.Add(NormalExpression);
-			Material->Normal.Expression = NormalExpression;
-
-			MaterialNodeY+= MaterialNodeStepY;
-		}
-							
-		Material->PostEditChange();
-		return Material;
+		return FMaterialUtilities::CreateMaterial(InFlattenMaterial, Outer, BaseName, Flags, OutGeneratedAssets);
 	}
 
 	bool ExportBaseColor(ULandscapeComponent* LandscapeComponent, int32 TextureSize, TArray<FColor>& OutSamples)
 	{
-		ALandscapeProxy* LandscapeProxy = LandscapeComponent->GetLandscapeProxy();
-
-		FIntPoint ComponentOrigin = LandscapeComponent->GetSectionBase() - LandscapeProxy->LandscapeSectionOffset;
-		FIntPoint ComponentSize(LandscapeComponent->ComponentSizeQuads, LandscapeComponent->ComponentSizeQuads);
-		FVector MidPoint = FVector(ComponentOrigin, 0.f) + FVector(ComponentSize, 0.f)*0.5f;
-
-		FVector LandscapeCenter = LandscapeProxy->GetTransform().TransformPosition(MidPoint);
-		FVector LandscapeExtent = FVector(ComponentSize, 0.f)*LandscapeProxy->GetActorScale()*0.5f;
-
-		FVector ViewOrigin = LandscapeCenter;
-		FMatrix ViewRotationMatrix = FInverseRotationMatrix(LandscapeProxy->GetActorRotation());
-		ViewRotationMatrix *= FMatrix(FPlane(1, 0, 0, 0),
-			FPlane(0, -1, 0, 0),
-			FPlane(0, 0, -1, 0),
-			FPlane(0, 0, 0, 1));
-
-		const float ZOffset = WORLD_MAX;
-		FMatrix ProjectionMatrix = FReversedZOrthoMatrix(
-			LandscapeExtent.X,
-			LandscapeExtent.Y,
-			0.5f / ZOffset,
-			ZOffset);
-
-		FSceneInterface* Scene = LandscapeProxy->GetWorld()->Scene;
-
-		// Hide all but the component
-		TSet<FPrimitiveComponentId> HiddenPrimitives;
-		for (auto PrimitiveComponentId : Scene->GetScenePrimitiveComponentIds())
-		{
-			HiddenPrimitives.Add(PrimitiveComponentId);
-		}
-		HiddenPrimitives.Remove(LandscapeComponent->SceneProxy->GetPrimitiveComponentId());
-				
-		FIntPoint TargetSize(TextureSize, TextureSize);
-
-		// Render diffuse texture using BufferVisualizationMode=BaseColor
-		static const FName BaseColorName("BaseColor");
-		const float BaseColorGamma = 2.2f;
-		RenderSceneToTexture(Scene, BaseColorName, ViewOrigin, ViewRotationMatrix, ProjectionMatrix, HiddenPrimitives, TargetSize, BaseColorGamma, OutSamples);
-		return true;
+		return FMaterialUtilities::ExportBaseColor(LandscapeComponent, TextureSize, OutSamples);
 	}
 }
diff --git a/Engine/Source/Editor/UnrealEd/Public/MaterialExportUtils.h b/Engine/Source/Editor/UnrealEd/Public/MaterialExportUtils.h
--- a/Engine/Source/Editor/UnrealEd/Public/MaterialExportUtils.h
+++ b/Engine/Source/Editor/UnrealEd/Public/MaterialExportUtils.h
@@ -4,34 +4,10 @@
 // Forward declarations
 class ALandscapeProxy;
 class ULandscapeComponent;
+struct FFlattenMaterial;
 
 namespace MaterialExportUtils
 {
-	struct UNREALED_API FFlattenMaterial
-	{
-		FFlattenMaterial()
-			: DiffuseSize(512, 512)
-			, NormalSize(512, 512)
-			, MetallicSize(0, 0)
-			, RoughnessSize(0, 0)
-			, SpecularSize(0, 0)
-		{}
-		
-		FGuid			MaterialId;
-		
-		FIntPoint		DiffuseSize;
-		FIntPoint		NormalSize;
-		FIntPoint		MetallicSize;	
-		FIntPoint		RoughnessSize;	
-		FIntPoint		SpecularSize;	
-			
-		TArray<FColor>	DiffuseSamples;
-		TArray<FColor>	NormalSamples;
-		TArray<FColor>	MetallicSamples;
-		TArray<FColor>	RoughnessSamples;
-		TArray<FColor>	SpecularSamples;
-	};
-	
 	/**
 	 * Renders specified material property into texture
 	 *
@@ -42,7 +18,10 @@ namespace MaterialExportUtils
 	 * @param OutBMP				Output array of rendered samples 
 	 * @return						Whether operation was successful
 	 */
+	DEPRECATED(4.9, "MaterialExportUtils::ExportMaterialProperty() will be removed, use FMaterialUtilities::ExportMaterialProperty().")
 	UNREALED_API bool ExportMaterialProperty(UWorld* InWorld, UMaterialInterface* InMaterial, EMaterialProperty InMaterialProperty, UTextureRenderTarget2D* InRenderTarget, TArray<FColor>& OutBMP);
+	
+	DEPRECATED(4.9, "MaterialExportUtils::ExportMaterialProperty() will be removed, use FMaterialUtilities::ExportMaterialProperty().")
 	UNREALED_API bool ExportMaterialProperty(UMaterialInterface* InMaterial, EMaterialProperty InMaterialProperty, UTextureRenderTarget2D* InRenderTarget, TArray<FColor>& OutBMP);
 
 	/**
@@ -53,7 +32,10 @@ namespace MaterialExportUtils
 	 * @param OutFlattenMaterial	Output flattened material
 	 * @return						Whether operation was successful
 	 */
+	DEPRECATED(4.9, "MaterialExportUtils::ExportMaterial() will be removed, use FMaterialUtilities::ExportMaterial().")
 	UNREALED_API bool ExportMaterial(UWorld* InWorld, UMaterialInterface* InMaterial, FFlattenMaterial& OutFlattenMaterial);
+	
+	DEPRECATED(4.9, "MaterialExportUtils::ExportMaterial() will be removed, use FMaterialUtilities::ExportMaterial().")
 	UNREALED_API bool ExportMaterial(UMaterialInterface* InMaterial, FFlattenMaterial& OutFlattenMaterial);
 	
 	/**
@@ -64,6 +46,7 @@ namespace MaterialExportUtils
 	 * @param OutFlattenMaterial	Output flattened material
 	 * @return						Whether operation was successful
 	 */
+	DEPRECATED(4.9, "MaterialExportUtils::ExportMaterial() will be removed, use FMaterialUtilities::ExportLandscapeMaterial().")
 	UNREALED_API bool ExportMaterial(ALandscapeProxy* InLandscape, const TSet<FPrimitiveComponentId>& HiddenPrimitives, FFlattenMaterial& OutFlattenMaterial);
 
 	/**
@@ -75,6 +58,7 @@ namespace MaterialExportUtils
 	 * @param OutGeneratedAssets	List of generated assets - material, textures
 	 * @return						Returns a pointer to the constructed UMaterial object.
 	 */
+	DEPRECATED(4.9, "MaterialExportUtils::CreateMaterial() will be removed, use FMaterialUtilities::CreateMaterial().")
 	UNREALED_API UMaterial* CreateMaterial(const FFlattenMaterial& InFlattenMaterial, UPackage* InOuter, const FString& BaseName, EObjectFlags Flags, TArray<UObject*>& OutGeneratedAssets);
 
 
@@ -92,6 +76,7 @@ namespace MaterialExportUtils
 	 * @param SourceGuidHash		(optional) Hash (stored as Guid) to use part of the texture source's DDC key.
 	 * @return						The new texture.
 	 */
+	DEPRECATED(4.9, "MaterialExportUtils::CreateTexture() will be removed, use FMaterialUtilities::CreateTexture().")
 	UNREALED_API UTexture2D* CreateTexture(UPackage* Outer, const FString& AssetLongName, FIntPoint Size, const TArray<FColor>& Samples, TextureCompressionSettings CompressionSettings, TextureGroup LODGroup, EObjectFlags Flags, bool bSRGB, const FGuid& SourceGuidHash = FGuid());
 
 	/**
@@ -100,6 +85,7 @@ namespace MaterialExportUtils
 	* @param LandscapeComponent		The component to bake textures for
 	* @return						Whether operation was successful
 	*/
+	DEPRECATED(4.9, "MaterialExportUtils::ExportBaseColor() will be removed, use FMaterialUtilities::ExportBaseColor().")
 	UNREALED_API bool ExportBaseColor(ULandscapeComponent* LandscapeComponent, int32 TextureSize, TArray<FColor>& OutSamples);
 }
 
diff --git a/Engine/Source/Editor/UnrealEd/UnrealEd.Build.cs b/Engine/Source/Editor/UnrealEd/UnrealEd.Build.cs
--- a/Engine/Source/Editor/UnrealEd/UnrealEd.Build.cs
+++ b/Engine/Source/Editor/UnrealEd/UnrealEd.Build.cs
@@ -129,6 +129,7 @@ public class UnrealEd : ModuleRules
                 "Foliage",
                 "VectorVM",
 				"TreeMap",
+                "MaterialUtilities",
 			}
 		);
 
diff --git a/Engine/Source/Editor/WorldBrowser/Private/Tiles/WorldTileCollectionModel.cpp b/Engine/Source/Editor/WorldBrowser/Private/Tiles/WorldTileCollectionModel.cpp
--- a/Engine/Source/Editor/WorldBrowser/Private/Tiles/WorldTileCollectionModel.cpp
+++ b/Engine/Source/Editor/WorldBrowser/Private/Tiles/WorldTileCollectionModel.cpp
@@ -7,8 +7,8 @@
 #include "FileHelpers.h"
 #include "ContentBrowserModule.h"
 #include "AssetRegistryModule.h"
-#include "MaterialExportUtils.h"
 #include "MeshUtilities.h"
+#include "MaterialUtilities.h"
 #include "RawMesh.h"
 #include "LandscapeEdMode.h"
 #include "ImageWrapper.h"
@@ -1979,8 +1979,6 @@ bool FWorldTileCollectionModel::GenerateLODLevels(FLevelModelList InLevelList, i
 			}
 		}
 
-		using namespace MaterialExportUtils;
-
 		// Convert landscape actors into static meshes
 		int32 LandscapeActorIndex = 0;
 		for (ALandscapeProxy* Landscape : LandscapeActors)
@@ -2037,7 +2035,7 @@ bool FWorldTileCollectionModel::GenerateLODLevels(FLevelModelList InLevelList, i
 			LandscapeFlattenMaterial.RoughnessSize	= SimplificationDetails.bGenerateLandscapeRoughnessMap ? LandscapeTextureSize : FIntPoint::ZeroValue;
 			LandscapeFlattenMaterial.SpecularSize	= SimplificationDetails.bGenerateLandscapeSpecularMap ? LandscapeTextureSize : FIntPoint::ZeroValue;
 			
-			ExportMaterial(Landscape, PrimitivesToHide, LandscapeFlattenMaterial);
+			FMaterialUtilities::ExportLandscapeMaterial(Landscape, PrimitivesToHide, LandscapeFlattenMaterial);
 		
 			if (SimplificationDetails.bBakeGrassToLandscape)
 			{
@@ -2045,7 +2043,7 @@ bool FWorldTileCollectionModel::GenerateLODLevels(FLevelModelList InLevelList, i
 			}
 			FString LandscapeBaseAssetName = FString::Printf(TEXT("%s_LOD%d"), *Landscape->GetName(), TargetLODIndex + 1);
 			// Construct landscape material
-			UMaterial* StaticLandscapeMaterial = MaterialExportUtils::CreateMaterial(
+			UMaterial* StaticLandscapeMaterial = FMaterialUtilities::CreateMaterial(
 				LandscapeFlattenMaterial, AssetsOuter, *(AssetsPath + LandscapeBaseAssetName), RF_Public|RF_Standalone, GeneratedAssets);
 			// Currently landscape exports world space normal map
 			StaticLandscapeMaterial->bTangentSpaceNormal = false;
diff --git a/Engine/Source/Editor/WorldBrowser/WorldBrowser.Build.cs b/Engine/Source/Editor/WorldBrowser/WorldBrowser.Build.cs
--- a/Engine/Source/Editor/WorldBrowser/WorldBrowser.Build.cs
+++ b/Engine/Source/Editor/WorldBrowser/WorldBrowser.Build.cs
@@ -12,9 +12,10 @@ public class WorldBrowser : ModuleRules
         new string[] {
                 "AssetRegistry",
 				"AssetTools",
-                "MeshUtilities",
                 "ContentBrowser",
 				"Landscape",
+                "MeshUtilities",
+                "MaterialUtilities",
 			}
         );
      
@@ -39,10 +40,10 @@ public class WorldBrowser : ModuleRules
                 "SourceControl",
 				"SourceControlWindows",
                 "RawMesh",
-                "MeshUtilities",
                 "LandscapeEditor",
                 "ImageWrapper",
                 "Foliage",
+                "MaterialUtilities",
             }
 		);
 
diff --git a/Engine/Source/Runtime/Landscape/Landscape.Build.cs b/Engine/Source/Runtime/Landscape/Landscape.Build.cs
--- a/Engine/Source/Runtime/Landscape/Landscape.Build.cs
+++ b/Engine/Source/Runtime/Landscape/Landscape.Build.cs
@@ -55,13 +55,15 @@ public class Landscape : ModuleRules
 		{
 			PrivateDependencyModuleNames.AddRange(
 				new string[] {
-                    "UnrealEd", 
+                    "UnrealEd",
+                    "MaterialUtilities", 
                 }
 			);
 
 			CircularlyReferencedDependentModules.AddRange(
 				new string[] {
                     "UnrealEd",
+                    "MaterialUtilities",
                 }
 			);
 		}
diff --git a/Engine/Source/Runtime/Landscape/Private/Landscape.cpp b/Engine/Source/Runtime/Landscape/Private/Landscape.cpp
--- a/Engine/Source/Runtime/Landscape/Private/Landscape.cpp
+++ b/Engine/Source/Runtime/Landscape/Private/Landscape.cpp
@@ -38,7 +38,7 @@ Landscape.cpp: Terrain rendering
 #include "Materials/MaterialExpressionLandscapeVisibilityMask.h"
 
 #if WITH_EDITOR
-#include "MaterialExportUtils.h"
+#include "MaterialUtilities.h"
 #endif
 
 /** Landscape stats */
@@ -2436,7 +2436,7 @@ void ALandscapeProxy::UpdateBakedTextures()
 
 						int32 BakeSize = ComponentSamples >> 3;
 						TArray<FColor> Samples;
-						if (MaterialExportUtils::ExportBaseColor(Component, BakeSize, Samples))
+						if (FMaterialUtilities::ExportBaseColor(Component, BakeSize, Samples))
 						{
 							int32 AtlasOffsetX = FMath::RoundToInt(Component->HeightmapScaleBias.Z * (float)HeightmapTexture->GetSizeX()) >> 3;
 							int32 AtlasOffsetY = FMath::RoundToInt(Component->HeightmapScaleBias.W * (float)HeightmapTexture->GetSizeY()) >> 3;
@@ -2447,7 +2447,7 @@ void ALandscapeProxy::UpdateBakedTextures()
 							NumGenerated++;
 						}
 					}
-					UTexture2D* AtlasTexture = MaterialExportUtils::CreateTexture(GetOutermost(), HeightmapTexture->GetName() + TEXT("_BaseColor"), AtlasSize, AtlasSamples, TC_Default, TEXTUREGROUP_World, RF_Public, true, CombinedStateId);
+					UTexture2D* AtlasTexture = FMaterialUtilities::CreateTexture(GetOutermost(), HeightmapTexture->GetName() + TEXT("_BaseColor"), AtlasSize, AtlasSamples, TC_Default, TEXTUREGROUP_World, RF_Public, true, CombinedStateId);
 					AtlasTexture->MarkPackageDirty();
 
 					for (ULandscapeComponent* Component : Info.Components)
