Non-unity build fix

Cannot convert from undefined type to SWidget

diff --git a/Engine/Source/Editor/UnrealEd/Private/AutoReimport/ReimportFeedbackContext.cpp b/Engine/Source/Editor/UnrealEd/Private/AutoReimport/ReimportFeedbackContext.cpp
--- a/Engine/Source/Editor/UnrealEd/Private/AutoReimport/ReimportFeedbackContext.cpp
+++ b/Engine/Source/Editor/UnrealEd/Private/AutoReimport/ReimportFeedbackContext.cpp
@@ -503,6 +503,11 @@ void FReimportFeedbackContext::AddWidget(const TSharedRef<SWidget>& Widget)
 	}
 }
 
+TSharedRef<SWidget> FReimportFeedbackContext::AsWidget()
+{
+	return NotificationContent.ToSharedRef();
+}
+
 void FReimportFeedbackContext::StartSlowTask(const FText& Task, bool bShowCancelButton)
 {
 	FFeedbackContext::StartSlowTask(Task, bShowCancelButton);
diff --git a/Engine/Source/Editor/UnrealEd/Private/AutoReimport/ReimportFeedbackContext.h b/Engine/Source/Editor/UnrealEd/Private/AutoReimport/ReimportFeedbackContext.h
--- a/Engine/Source/Editor/UnrealEd/Private/AutoReimport/ReimportFeedbackContext.h
+++ b/Engine/Source/Editor/UnrealEd/Private/AutoReimport/ReimportFeedbackContext.h
@@ -33,7 +33,7 @@ public:
 	/** INotificationWidget and FFeedbackContext overrides */
 	virtual void Serialize( const TCHAR* V, ELogVerbosity::Type Verbosity, const class FName& Category ) override {}
 	virtual void OnSetCompletionState(SNotificationItem::ECompletionState State) override {}
-	virtual TSharedRef<SWidget> AsWidget() override { return NotificationContent.ToSharedRef(); };
+	virtual TSharedRef<SWidget> AsWidget() override;
 	virtual void StartSlowTask( const FText& Task, bool bShowCancelButton=false) override;
 
 	/** True if we will suppress slow task messages from being shown on the UI, false if otherwise */
