Fix crash attempting to reimport a render target.  Render targets cannot be reimported since they arent generated from source art data

diff --git a/Engine/Source/Editor/UnrealEd/Private/Factories/EditorFactories.cpp b/Engine/Source/Editor/UnrealEd/Private/Factories/EditorFactories.cpp
--- a/Engine/Source/Editor/UnrealEd/Private/Factories/EditorFactories.cpp
+++ b/Engine/Source/Editor/UnrealEd/Private/Factories/EditorFactories.cpp
@@ -5510,7 +5510,7 @@ UTextureCube* UReimportTextureFactory::CreateTextureCube( UObject* InParent, FNa
 bool UReimportTextureFactory::CanReimport( UObject* Obj, TArray<FString>& OutFilenames )
 {	
 	UTexture* pTex = Cast<UTexture>(Obj);
-	if(pTex)
+	if( pTex && !pTex->IsA<UTextureRenderTarget>() )
 	{
 		OutFilenames.Add(FReimportManager::ResolveImportFilename(pTex->SourceFilePath, pTex));
 		return true;
