Fixed crash when using -iterate and doesn't find the previously created asset registry (because of clean cook).

diff --git a/Engine/Source/Editor/UnrealEd/Private/Commandlets/ChunkManifestGenerator.cpp b/Engine/Source/Editor/UnrealEd/Private/Commandlets/ChunkManifestGenerator.cpp
--- a/Engine/Source/Editor/UnrealEd/Private/Commandlets/ChunkManifestGenerator.cpp
+++ b/Engine/Source/Editor/UnrealEd/Private/Commandlets/ChunkManifestGenerator.cpp
@@ -418,10 +418,12 @@ bool FChunkManifestGenerator::LoadAssetRegistry(const FString& SandboxPath, cons
 	{
 		FString PlatformSandboxPath = SandboxPath.Replace(TEXT("[Platform]"), *Platform->PlatformName());
 		FArchive* AssetRegistryReader = IFileManager::Get().CreateFileReader(*PlatformSandboxPath);
-
 		TMap<FName, FAssetData*> SavedAssetRegistryData;
-		AssetRegistry.LoadRegistryData(*AssetRegistryReader, SavedAssetRegistryData);
-
+		if (AssetRegistryReader)
+		{
+			
+			AssetRegistry.LoadRegistryData(*AssetRegistryReader, SavedAssetRegistryData);
+		}
 		for (FAssetData LoadedAssetData : AssetRegistryData)
 		{
 			if (PackagesToKeep &&
