#HLOD
- Make sure you can change distance
- Turned off simplygon by default
- Fixed memory issues and so on.

diff --git a/Engine/Source/Editor/UnrealEd/Private/HierarchicalLOD.cpp b/Engine/Source/Editor/UnrealEd/Private/HierarchicalLOD.cpp
--- a/Engine/Source/Editor/UnrealEd/Private/HierarchicalLOD.cpp
+++ b/Engine/Source/Editor/UnrealEd/Private/HierarchicalLOD.cpp
@@ -434,6 +434,7 @@ void FHierarchicalLODBuilder::InitializeClusters(class ULevel* InLevel, const in
 		else // at this point we only care for LODActors
 		{
 			Clusters.Empty();
+
 			// we filter the LOD index first
 			TArray<AActor*> Actors;
 			for(int32 ActorId=0; ActorId<InLevel->Actors.Num(); ++ActorId)
@@ -467,6 +468,11 @@ void FHierarchicalLODBuilder::InitializeClusters(class ULevel* InLevel, const in
 					Clusters.Add(NewClusterCandidate);
 				}
 			}
+
+			// shrink after adding actors
+			// LOD 0 has lots of actors, and subsequence LODs tend to have a lot less actors
+			// so this should save a lot more. 
+			Clusters.Shrink();
 		}
 	}
 }
@@ -626,8 +632,6 @@ void FLODCluster::BuildActor(class ULevel* InLevel, const int32 LODIdx)
 		////////////////////////////////////////////////////////////////////////////////////
 		// create asset using Actors
 		const FHierarchicalSimplification& LODSetup = InLevel->GetWorld()->GetWorldSettings()->HierarchicalLODSetup[LODIdx];
-		const FMeshProxySettings& ProxySetting = LODSetup.ProxySetting;
-		const FMeshMergingSettings& MergeSetting = LODSetup.MergeSetting;
 
 		// Where generated assets will be stored
 		UPackage* AssetsOuter = InLevel->GetOutermost(); // this asset is going to save with map, this means, I'll have to delete with it
@@ -660,11 +664,11 @@ void FLODCluster::BuildActor(class ULevel* InLevel, const int32 LODIdx)
 				const FString PackageName = FString::Printf(TEXT("LOD_%s"), *FirstActor->GetName());
 				if (MeshUtilities.GetMeshMergingInterface() && LODSetup.bSimplifyMesh)
 				{
-					MeshUtilities.CreateProxyMesh(Actors, ProxySetting, AssetsOuter, PackageName, OutAssets, OutProxyLocation);
+					MeshUtilities.CreateProxyMesh(Actors, LODSetup.ProxySetting, AssetsOuter, PackageName, OutAssets, OutProxyLocation);
 				}
 				else
 				{
-					MeshUtilities.MergeActors(Actors, MergeSetting, AssetsOuter, PackageName, LODIdx+1, OutAssets, OutProxyLocation, true );
+					MeshUtilities.MergeActors(Actors, LODSetup.MergeSetting, AssetsOuter, PackageName, LODIdx+1, OutAssets, OutProxyLocation, true );
 				}
 
 				// we make it private, so it can't be used by outside of map since it's useless, and then remove standalone
@@ -698,6 +702,7 @@ void FLODCluster::BuildActor(class ULevel* InLevel, const int32 LODIdx)
 					NewActor->SubActors = Actors;
 					NewActor->LODLevel = LODIdx+1; 
 					float DrawDistance = LODSetup.DrawDistance;
+					NewActor->LODDrawDistance = DrawDistance;
 
 					NewActor->GetStaticMeshComponent()->StaticMesh = MainMesh;
 
diff --git a/Engine/Source/Runtime/Engine/Classes/Engine/LODActor.h b/Engine/Source/Runtime/Engine/Classes/Engine/LODActor.h
--- a/Engine/Source/Runtime/Engine/Classes/Engine/LODActor.h
+++ b/Engine/Source/Runtime/Engine/Classes/Engine/LODActor.h
@@ -25,7 +25,7 @@ public:
 	TArray<class AActor*> SubActors;
 
 	/** what distance do you want this to show up instead of SubActors */
-	UPROPERTY(Category=LODActor, VisibleAnywhere)
+	UPROPERTY(Category=LODActor, EditAnywhere, meta=(ClampMin = "0", UIMin = "0"))
 	float LODDrawDistance;
 
 	/** what distance do you want this to show up instead of SubActors */
@@ -48,7 +48,7 @@ protected:
 	virtual FString GetDetailedInfoInternal() const override;
 	virtual FBox GetComponentsBoundingBox(bool bNonColliding = false) const override;	
 #if WITH_EDITOR
-	virtual void PostEditChangeChainProperty(FPropertyChangedChainEvent& PropertyChangedEvent) override;
+	virtual void PostEditChangeProperty(FPropertyChangedEvent& PropertyChangedEvent) override;
 #endif // WITH_EDITOR	
 	virtual void PostRegisterAllComponents() override;
 	// End UObject interface.
diff --git a/Engine/Source/Runtime/Engine/Classes/GameFramework/WorldSettings.h b/Engine/Source/Runtime/Engine/Classes/GameFramework/WorldSettings.h
--- a/Engine/Source/Runtime/Engine/Classes/GameFramework/WorldSettings.h
+++ b/Engine/Source/Runtime/Engine/Classes/GameFramework/WorldSettings.h
@@ -230,13 +230,12 @@ struct ENGINE_API FHierarchicalSimplification
 	FMeshMergingSettings MergeSetting;
 
 	FHierarchicalSimplification()
-		: bSimplifyMesh(true)
+		: bSimplifyMesh(false)
 		, DrawDistance(3000)
 		, DesiredBoundRadius(3000)
 		, DesiredFillingPercentage(50)
 		, MinNumberOfActorsToBuild(2)
 	{
-		ProxySetting.ScreenSize = 200;
 		MergeSetting.bMergeMaterials = true;
 	}
 };
diff --git a/Engine/Source/Runtime/Engine/Private/LODActor.cpp b/Engine/Source/Runtime/Engine/Private/LODActor.cpp
--- a/Engine/Source/Runtime/Engine/Private/LODActor.cpp
+++ b/Engine/Source/Runtime/Engine/Private/LODActor.cpp
@@ -24,6 +24,7 @@ ALODActor::ALODActor(const FObjectInitializer& ObjectInitializer)
 	StaticMeshComponent->bGenerateOverlapEvents = false;
 	StaticMeshComponent->bCastDynamicShadow = false;
 	StaticMeshComponent->bCastStaticShadow = false;
+	StaticMeshComponent->CastShadow = false;
 
 	RootComponent = StaticMeshComponent;
 }
@@ -46,9 +47,33 @@ void ALODActor::PostRegisterAllComponents()
 }
 #if WITH_EDITOR
 
-void ALODActor::PostEditChangeChainProperty(FPropertyChangedChainEvent& PropertyChangedEvent)
+void ALODActor::PostEditChangeProperty(FPropertyChangedEvent& PropertyChangedEvent)
 {
-	Super::PostEditChangeChainProperty(PropertyChangedEvent);
+	UProperty* PropertyThatChanged = PropertyChangedEvent.Property;
+	FName PropertyName = PropertyThatChanged != NULL ? PropertyThatChanged->GetFName() : NAME_None;
+	
+	if (PropertyName == GET_MEMBER_NAME_CHECKED(ALODActor, LODDrawDistance))
+	{
+		for (auto& Actor: SubActors)
+		{
+			if (Actor)
+			{
+				TArray<UPrimitiveComponent*> InnerComponents;
+				Actor->GetComponents<UPrimitiveComponent>(InnerComponents);
+
+				for(auto& Component: InnerComponents)
+				{
+					UPrimitiveComponent* ParentComp = Component->GetLODParentPrimitive();
+					if (ParentComp)
+					{
+						ParentComp->MinDrawDistance = LODDrawDistance;
+						ParentComp->MarkRenderStateDirty();
+					}
+				}
+			}
+		}
+	}
+	Super::PostEditChangeProperty(PropertyChangedEvent);
 }
 
 bool ALODActor::GetReferencedContentObjects( TArray<UObject*>& Objects ) const
