Moved exec commands for reloading localization data out of editor execs into engine execs so they can be used in standalone games. This is essential as the editor itself never loads game specific localization data.

diff --git a/Engine/Source/Editor/UnrealEd/Private/UnrealEdSrv.cpp b/Engine/Source/Editor/UnrealEd/Private/UnrealEdSrv.cpp
--- a/Engine/Source/Editor/UnrealEd/Private/UnrealEdSrv.cpp
+++ b/Engine/Source/Editor/UnrealEd/Private/UnrealEdSrv.cpp
@@ -784,27 +784,6 @@ bool UUnrealEdEngine::Exec( UWorld* InWorld, const TCHAR* Stream, FOutputDevice&
 		HandleDumpSelectionCommand( Str, Ar );
 	}
 
-#if ENABLE_LOC_TESTING
-	{
-		FString CultureName;
-		if( FParse::Value(Str,TEXT("CULTURE="), CultureName) )
-		{
-			FInternationalization::Get().SetCurrentCulture( CultureName );
-		}
-	}
-
-	{
-		FString ConfigFilePath;
-		if(FParse::Value(Str,TEXT("REGENLOC="), ConfigFilePath))
-		{
-			FJsonInternationalizationArchiveSerializer ArchiveSerializer;
-			FJsonInternationalizationManifestSerializer ManifestSerializer;
-
-			FTextLocalizationManager::Get().LoadFromManifestAndArchives(ConfigFilePath, ArchiveSerializer, ManifestSerializer);
-		}
-	}
-#endif
-
 	//----------------------------------------------------------------------------------
 	// EDIT
 	//
@@ -3014,23 +2993,6 @@ bool UUnrealEdEngine::Exec_Mode( const TCHAR* Str, FOutputDevice& Ar )
 		Word1 = MAX_uint16;
 	}
 
-#if ENABLE_LOC_TESTING
-	FString CultureName;
-	if( FParse::Value(Str,TEXT("CULTURE="), CultureName) )
-	{
-		FInternationalization::Get().SetCurrentCulture( CultureName );
-	}
-
-	FString ConfigFilePath;
-	if( FParse::Value(Str,TEXT("REGENLOC="), ConfigFilePath) )
-	{
-		FJsonInternationalizationArchiveSerializer ArchiveSerializer;
-		FJsonInternationalizationManifestSerializer ManifestSerializer;
-
-		FTextLocalizationManager::Get().LoadFromManifestAndArchives(ConfigFilePath, ArchiveSerializer, ManifestSerializer);
-	}
-#endif
-
 	if( FParse::Value(Str,TEXT("USESIZINGBOX="), DWord1) )
 	{
 		FinishAllSnaps();
diff --git a/Engine/Source/Runtime/Engine/Engine.Build.cs b/Engine/Source/Runtime/Engine/Engine.Build.cs
--- a/Engine/Source/Runtime/Engine/Engine.Build.cs
+++ b/Engine/Source/Runtime/Engine/Engine.Build.cs
@@ -68,7 +68,8 @@ public class Engine : ModuleRules
 				"Landscape",
                 "UMG",
 				"Projects",
-				"Niagara"
+				"Niagara",
+                "Internationalization"
 			}
         );
 
diff --git a/Engine/Source/Runtime/Engine/Private/UnrealEngine.cpp b/Engine/Source/Runtime/Engine/Private/UnrealEngine.cpp
--- a/Engine/Source/Runtime/Engine/Private/UnrealEngine.cpp
+++ b/Engine/Source/Runtime/Engine/Private/UnrealEngine.cpp
@@ -102,6 +102,8 @@
 #include "SNotificationList.h"
 #include "Engine/UserInterfaceSettings.h"
 #include "ComponentRecreateRenderStateContext.h"
+#include "JsonInternationalizationArchiveSerializer.h"
+#include "JsonInternationalizationManifestSerializer.h"
 
 DEFINE_LOG_CATEGORY(LogEngine);
 
@@ -2243,6 +2245,27 @@ bool UEngine::Exec( UWorld* InWorld, const TCHAR* Cmd, FOutputDevice& Ar )
 		return true;
 	}
 
+#if ENABLE_LOC_TESTING
+	{
+		FString CultureName;
+		if (FParse::Value(Cmd, TEXT("CULTURE="), CultureName))
+		{
+			FInternationalization::Get().SetCurrentCulture(CultureName);
+		}
+	}
+
+	{
+		FString ConfigFilePath;
+		if (FParse::Value(Cmd, TEXT("REGENLOC="), ConfigFilePath))
+		{
+			FJsonInternationalizationArchiveSerializer ArchiveSerializer;
+			FJsonInternationalizationManifestSerializer ManifestSerializer;
+
+			FTextLocalizationManager::Get().LoadFromManifestAndArchives(ConfigFilePath, ArchiveSerializer, ManifestSerializer);
+		}
+	}
+#endif
+
 	// Handle engine command line.
 	if ( FParse::Command(&Cmd,TEXT("FLUSHLOG")) )
 	{
