Fixed a crash on broadcasting information regarding a level being removed from the editor world #UE4

It turned out the broadcasting was happening just after EditorDestroyLevel call, and the boardcasted delegate does use a pointer to that destroyed level.
Note: this change has been merged over from dev branch, and it's bringing back functionality cut out by DmitriyD in CL#2510067

diff --git a/Engine/Source/Editor/UnrealEd/Private/EditorLevelUtils.cpp b/Engine/Source/Editor/UnrealEd/Private/EditorLevelUtils.cpp
--- a/Engine/Source/Editor/UnrealEd/Private/EditorLevelUtils.cpp
+++ b/Engine/Source/Editor/UnrealEd/Private/EditorLevelUtils.cpp
@@ -405,6 +405,11 @@ namespace EditorLevelUtils
 				MakeLevelCurrent(OwningWorld->PersistentLevel);
 			}
 
+			// Scrub any levelscript actions so GC can take effect
+			// note that this needs to be done here, before EditorDestroyLevel call
+			// since it doesn't make much sense to pass the InLevel around after 
+			// calling "destruction" kind of functions on it
+			FWorldDelegates::LevelRemovedFromWorld.Broadcast(InLevel, OwningWorld);
 
 			EditorDestroyLevel(InLevel);
 
