Fixed crash when default maps directory in the EditorPerProjectUserSettings.ini lies outside of the current project content directory.

diff --git a/Engine/Source/Editor/UnrealEd/Private/FileHelpers.cpp b/Engine/Source/Editor/UnrealEd/Private/FileHelpers.cpp
--- a/Engine/Source/Editor/UnrealEd/Private/FileHelpers.cpp
+++ b/Engine/Source/Editor/UnrealEd/Private/FileHelpers.cpp
@@ -636,7 +636,12 @@ static bool SaveAsImplementation( UWorld* InWorld, const FString& DefaultFilenam
 	{
 		const FString DefaultName = TEXT("NewMap");
 		FString PackageName;
-		FPackageName::TryConvertFilenameToLongPackageName(DefaultDirectory / DefaultName, PackageName);
+		if (!FPackageName::TryConvertFilenameToLongPackageName(DefaultDirectory / DefaultName, PackageName))
+		{
+			// Initial location is invalid (e.g. lies outside of the project): set location to /Game/Maps instead
+			DefaultDirectory = FPaths::GameContentDir() / TEXT("Maps");
+			ensure(FPackageName::TryConvertFilenameToLongPackageName(DefaultDirectory / DefaultName, PackageName));
+		}
 		FString Name;
 		FAssetToolsModule& AssetToolsModule = FModuleManager::LoadModuleChecked<FAssetToolsModule>("AssetTools");
 		AssetToolsModule.Get().CreateUniqueAssetName(PackageName, TEXT(""), PackageName, Name);
@@ -2352,7 +2357,12 @@ static int32 InternalSavePackage( UPackage* PackageToSave, bool& bOutPackageLoca
 			if( UEditorEngine::IsUsingWorldAssets() )
 			{
 				FString DefaultPackagePath;
-				FPackageName::TryConvertFilenameToLongPackageName(DefaultLocation / FinalPackageFilename, DefaultPackagePath);
+				if (!FPackageName::TryConvertFilenameToLongPackageName(DefaultLocation / FinalPackageFilename, DefaultPackagePath))
+				{
+					// Original location is invalid; set default location to /Game/Maps
+					DefaultLocation = FPaths::GameContentDir() / TEXT("Maps");
+					ensure(FPackageName::TryConvertFilenameToLongPackageName(DefaultLocation / FinalPackageFilename, DefaultPackagePath));
+				}
 
 				FString SaveAsPackageName;
 				bSaveFile = OpenLevelSaveAsDialog(
