Back out changelist 2532989

diff --git a/Engine/Source/Editor/UnrealEd/Private/EditorLevelUtils.cpp b/Engine/Source/Editor/UnrealEd/Private/EditorLevelUtils.cpp
--- a/Engine/Source/Editor/UnrealEd/Private/EditorLevelUtils.cpp
+++ b/Engine/Source/Editor/UnrealEd/Private/EditorLevelUtils.cpp
@@ -405,11 +405,6 @@ namespace EditorLevelUtils
 				MakeLevelCurrent(OwningWorld->PersistentLevel);
 			}
 
-			// Scrub any levelscript actions so GC can take effect
-			// note that this needs to be done here, before EditorDestroyLevel call
-			// since it doesn't make much sense to pass the InLevel around after 
-			// calling "destruction" kind of functions on it
-			FWorldDelegates::LevelRemovedFromWorld.Broadcast(InLevel, OwningWorld);
 
 			EditorDestroyLevel(InLevel);
 
