#HLOD: added virtual funcion so that they can override the behavior
- right now mobile staticmesh isn't included

https://jira.ol.epicgames.net/browse/UE-14841

diff --git a/Engine/Source/Editor/UnrealEd/Private/HierarchicalLOD.cpp b/Engine/Source/Editor/UnrealEd/Private/HierarchicalLOD.cpp
--- a/Engine/Source/Editor/UnrealEd/Private/HierarchicalLOD.cpp
+++ b/Engine/Source/Editor/UnrealEd/Private/HierarchicalLOD.cpp
@@ -379,6 +379,7 @@ bool ShouldGenerateCluster(class AActor* Actor)
 	// TODO: support instanced static meshes
 	Components.RemoveAll([](UStaticMeshComponent* Val){ return Val->IsA(UInstancedStaticMeshComponent::StaticClass()); });
 
+	int32 ValidComponentCount=0;
 	// now make sure you check parent primitive, so that we don't build for the actor that already has built. 
 	if (Components.Num() > 0)
 	{
@@ -388,12 +389,16 @@ bool ShouldGenerateCluster(class AActor* Actor)
 			{
 				return false;
 			}
-		}
 
-		return true;
+			// see if we should generate it
+			if (ComponentIter->ShouldGenerateAutoLOD())
+			{
+				++ValidComponentCount;
+			}
+		}
 	}
 
-	return false;
+	return (ValidComponentCount>0);
 }
 
 void FHierarchicalLODBuilder::InitializeClusters(class ULevel* InLevel, const int32 LODIdx, float CullCost)
diff --git a/Engine/Source/Runtime/Engine/Classes/Components/PrimitiveComponent.h b/Engine/Source/Runtime/Engine/Classes/Components/PrimitiveComponent.h
--- a/Engine/Source/Runtime/Engine/Classes/Components/PrimitiveComponent.h
+++ b/Engine/Source/Runtime/Engine/Classes/Components/PrimitiveComponent.h
@@ -998,6 +998,8 @@ public:
 
 #if WITH_EDITOR
 	virtual const int32 GetNumUncachedStaticLightingInteractions() const override; // recursive function
+	/** This function is used to create hierarchical LOD for the level. You can decide to opt out if you don't want. */
+	virtual const bool ShouldGenerateAutoLOD() const;
 #endif
 
 	// Begin UActorComponent Interface
diff --git a/Engine/Source/Runtime/Engine/Classes/GameFramework/WorldSettings.h b/Engine/Source/Runtime/Engine/Classes/GameFramework/WorldSettings.h
--- a/Engine/Source/Runtime/Engine/Classes/GameFramework/WorldSettings.h
+++ b/Engine/Source/Runtime/Engine/Classes/GameFramework/WorldSettings.h
@@ -237,6 +237,7 @@ struct ENGINE_API FHierarchicalSimplification
 		, MinNumberOfActorsToBuild(2)
 	{
 		MergeSetting.bMergeMaterials = true;
+		MergeSetting.bGenerateLightMapUV = true;
 	}
 };
 
diff --git a/Engine/Source/Runtime/Engine/Private/Components/PrimitiveComponent.cpp b/Engine/Source/Runtime/Engine/Private/Components/PrimitiveComponent.cpp
--- a/Engine/Source/Runtime/Engine/Private/Components/PrimitiveComponent.cpp
+++ b/Engine/Source/Runtime/Engine/Private/Components/PrimitiveComponent.cpp
@@ -2606,5 +2606,9 @@ void UPrimitiveComponent::SetCustomNavigableGeometry(const EHasCustomNavigableGe
 	bHasCustomNavigableGeometry = InType;
 }
 
+const bool UPrimitiveComponent::ShouldGenerateAutoLOD() const
+{
+	return (Mobility != EComponentMobility::Movable);
+}
 #undef LOCTEXT_NAMESPACE
 
