Paper2D: Editor Layer settings takes you to the correct settings page [UE-14708]

diff --git a/Engine/Source/Editor/UnrealEd/Private/STransformViewportToolbar.cpp b/Engine/Source/Editor/UnrealEd/Private/STransformViewportToolbar.cpp
--- a/Engine/Source/Editor/UnrealEd/Private/STransformViewportToolbar.cpp
+++ b/Engine/Source/Editor/UnrealEd/Private/STransformViewportToolbar.cpp
@@ -650,7 +650,7 @@ TSharedRef<SWidget> STransformViewportToolBar::FillLayer2DSnapMenu()
 		{
 			if (ISettingsModule* SettingsModule = FModuleManager::GetModulePtr<ISettingsModule>("Settings"))
 			{
-				SettingsModule->ShowViewer("Project", "Editor", "2D");
+				SettingsModule->ShowViewer("Project", "Editor", "LevelEditor2DSettings");
 			}
 		}
 
diff --git a/Engine/Source/Editor/UnrealEd/Public/Settings/EditorProjectSettings.h b/Engine/Source/Editor/UnrealEd/Public/Settings/EditorProjectSettings.h
--- a/Engine/Source/Editor/UnrealEd/Public/Settings/EditorProjectSettings.h
+++ b/Engine/Source/Editor/UnrealEd/Public/Settings/EditorProjectSettings.h
@@ -102,7 +102,7 @@ enum class ELevelEditor2DAxis : uint8
 };
 
 /**
- * Configure the settings for the 2D Level Editor
+ * Configure settings for the 2D Level Editor
  */
 UCLASS(config=Editor, meta=(DisplayName="2D"))
 class UNREALED_API ULevelEditor2DSettings : public UDeveloperSettings
