Paper2D: Combined 2D widget visuals improved [UE-14719]

diff --git a/Engine/Source/Editor/UnrealEd/Private/UnrealWidget.cpp b/Engine/Source/Editor/UnrealEd/Private/UnrealWidget.cpp
--- a/Engine/Source/Editor/UnrealEd/Private/UnrealWidget.cpp
+++ b/Engine/Source/Editor/UnrealEd/Private/UnrealWidget.cpp
@@ -902,10 +902,14 @@ void FWidget::Render_2D(const FSceneView* View, FPrimitiveDrawInterface* PDI, FE
 		// but in the ortho viewports it will prevent scaling in the direction of the camera and thus intersecting the near plane.
 		FVector FlattenScale = FVector(Scale.Component(0) == 1.0f ? 1.0f / UniformScale : 1.0f, Scale.Component(1) == 1.0f ? 1.0f / UniformScale : 1.0f, Scale.Component(2) == 1.0f ? 1.0f / UniformScale : 1.0f);
 		float ArrowRadius = ScaledRadius * 2.0f;
+		float ArrowStartRadius = ScaledRadius * 1.3f;
+
+		uint8 HoverAlpha = 0xff;
+		uint8 NormalAlpha = 0x2f;
 
 		if (((DrawAxis&EAxisList::XZ) == EAxisList::XZ) && (bIsPerspective || bIsLocalSpace || bIsOrthoXZ)) // Front
 		{
-			uint8 Alpha = ((CurrentAxis&EAxisList::XZ) == EAxisList::XZ ? 0x3f : 0x0f);
+			uint8 Alpha = ((CurrentAxis&EAxisList::XZ) == EAxisList::XZ ? HoverAlpha : NormalAlpha);
 			PDI->SetHitProxy(new HWidgetAxis(EAxisList::XZ, bDisabled));
 			{
 				FColor Color = YColor;
@@ -917,10 +921,13 @@ void FWidget::Render_2D(const FSceneView* View, FPrimitiveDrawInterface* PDI, FE
 
 			PDI->SetHitProxy(new HWidgetAxis(EAxisList::Rotate2D, bDisabled));
 			{
+				FColor Color = YColor;
+				Color.A = ((CurrentAxis&EAxisList::Rotate2D) == EAxisList::Rotate2D ? HoverAlpha : NormalAlpha);
+
 				FVector XAxis = CustomCoordSystem.TransformPosition(FVector(1, 0, 0).RotateAngleAxis((EditorModeTools ? EditorModeTools->TranslateRotate2DAngle : 0), FVector(0, -1, 0)));
 				FVector YAxis = CustomCoordSystem.TransformPosition(FVector(0, 0, 1).RotateAngleAxis((EditorModeTools ? EditorModeTools->TranslateRotate2DAngle : 0), FVector(0, -1, 0)));
-				FVector BaseArrowPoint = InLocation + XAxis * ScaledRadius;
-				DrawFlatArrow(PDI, BaseArrowPoint, XAxis, YAxis, YColor, ArrowRadius, ArrowRadius*.5f, TransparentPlaneMaterialXY->GetRenderProxy(false), SDPG_Foreground);
+				FVector BaseArrowPoint = InLocation + XAxis * ArrowStartRadius;
+				DrawFlatArrow(PDI, BaseArrowPoint, XAxis, YAxis, Color, ArrowRadius, ArrowRadius*.5f, TransparentPlaneMaterialXY->GetRenderProxy(false), SDPG_Foreground);
 			}
 			PDI->SetHitProxy(NULL);
 		}
