Paper2D: Actors snapped to layer weren't being updated immediately [UE-14711]

diff --git a/Engine/Source/Editor/UnrealEd/Private/STransformViewportToolbar.cpp b/Engine/Source/Editor/UnrealEd/Private/STransformViewportToolbar.cpp
--- a/Engine/Source/Editor/UnrealEd/Private/STransformViewportToolbar.cpp
+++ b/Engine/Source/Editor/UnrealEd/Private/STransformViewportToolbar.cpp
@@ -686,7 +686,7 @@ TSharedRef<SWidget> STransformViewportToolBar::FillLayer2DSnapMenu()
 					}
 				}
 
-				GEditor->RedrawViewports();
+				GEditor->RedrawLevelEditingViewports(true);
 			}
 
 			return FReply::Handled();
