Paper2D: 2D Widget rotation indicator shows initial rotation on selection [UE-14716]

diff --git a/Engine/Source/Editor/UnrealEd/Private/EditorSelectUtils.cpp b/Engine/Source/Editor/UnrealEd/Private/EditorSelectUtils.cpp
--- a/Engine/Source/Editor/UnrealEd/Private/EditorSelectUtils.cpp
+++ b/Engine/Source/Editor/UnrealEd/Private/EditorSelectUtils.cpp
@@ -184,6 +184,8 @@ void UUnrealEdEngine::SetPivot( FVector NewPivot, bool bSnapPivotToGrid, bool bI
 	EditorModeTools.TranslateRotate2DAngle = 0.0f;
 	FVector TranslateRotateWidgetWorldXAxis;
 
+	FVector Widget2DWorldXAxis;
+
 	AActor* LastSelectedActor = NULL;
 	for ( FSelectionIterator It( GetSelectedActorIterator() ) ; It ; ++It )
 	{
@@ -199,6 +201,13 @@ void UUnrealEdEngine::SetPivot( FVector NewPivot, bool bSnapPivotToGrid, bool bI
 			{
 				TranslateRotateWidgetWorldXAxis = FVector(1.0f, 0.0f, 0.0f);
 			}
+
+			Widget2DWorldXAxis = Actor->ActorToWorld().TransformVector(FVector(1, 0, 0));
+			Widget2DWorldXAxis.Y = 0;
+			if (!Widget2DWorldXAxis.Normalize())
+			{
+				Widget2DWorldXAxis = FVector(1, 0, 0);
+			}
 		}
 
 		LastSelectedActor = Actor;
@@ -216,10 +225,11 @@ void UUnrealEdEngine::SetPivot( FVector NewPivot, bool bSnapPivotToGrid, bool bI
 		}
 	}
 
-	//if there are multiple actors selected, just use the x-axis for the "translate/rotate" widget
+	//if there are multiple actors selected, just use the x-axis for the "translate/rotate" or 2D widgets
 	if (Count == 1)
 	{
 		EditorModeTools.TranslateRotateXAxisAngle = TranslateRotateWidgetWorldXAxis.Rotation().Yaw;
+		EditorModeTools.TranslateRotate2DAngle = FMath::RadiansToDegrees(FMath::Atan2(Widget2DWorldXAxis.Z, Widget2DWorldXAxis.X));
 	}
 
 	// Update showing.
