Call Function nodes are properly refreshed after a function signature was changed.
Unnecessary "MarkBlueprintAsStructurallyModified" are removed.

UE-14593 Updating the output of a function does not update the instance placed in the EventGraph

diff --git a/Engine/Source/Editor/BlueprintGraph/Private/EdGraphSchema_K2.cpp b/Engine/Source/Editor/BlueprintGraph/Private/EdGraphSchema_K2.cpp
--- a/Engine/Source/Editor/BlueprintGraph/Private/EdGraphSchema_K2.cpp
+++ b/Engine/Source/Editor/BlueprintGraph/Private/EdGraphSchema_K2.cpp
@@ -5737,7 +5737,6 @@ UEdGraphPin* UEdGraphSchema_K2::DropPinOnNode(UEdGraphNode* InTargetNode, const
 		{
 			if(UBlueprint* Blueprint = *BlueprintIt)
 			{
-				FBlueprintEditorUtils::MarkBlueprintAsModified(Blueprint);
 				Blueprint->BroadcastChanged();
 			}
 		}
diff --git a/Engine/Source/Editor/Kismet/Private/BlueprintDetailsCustomization.cpp b/Engine/Source/Editor/Kismet/Private/BlueprintDetailsCustomization.cpp
--- a/Engine/Source/Editor/Kismet/Private/BlueprintDetailsCustomization.cpp
+++ b/Engine/Source/Editor/Kismet/Private/BlueprintDetailsCustomization.cpp
@@ -3164,7 +3164,6 @@ void FBaseBlueprintGraphActionDetails::OnParamsChanged(UK2Node_EditablePinBase*
 		{
 			if(UBlueprint* Blueprint = *BlueprintIt)
 			{
-				FBlueprintEditorUtils::MarkBlueprintAsModified(Blueprint);
 				Blueprint->BroadcastChanged();
 			}
 		}
@@ -3262,11 +3261,6 @@ bool FBaseBlueprintGraphActionDetails::OnPinRenamed(UK2Node_EditablePinBase* Tar
 		{
 			(*NodeIter)->RenameUserDefinedPin(OldName, NewName, false);
 		}
-
-		for (auto BlueprintIt = PinRenamedHelper.ModifiedBlueprints.CreateIterator(); BlueprintIt; ++BlueprintIt)
-		{
-			FBlueprintEditorUtils::MarkBlueprintAsStructurallyModified(*BlueprintIt);
-		}
 	}
 	return true;
 }
diff --git a/Engine/Source/Editor/UnrealEd/Private/Kismet2/BlueprintEditorUtils.cpp b/Engine/Source/Editor/UnrealEd/Private/Kismet2/BlueprintEditorUtils.cpp
--- a/Engine/Source/Editor/UnrealEd/Private/Kismet2/BlueprintEditorUtils.cpp
+++ b/Engine/Source/Editor/UnrealEd/Private/Kismet2/BlueprintEditorUtils.cpp
@@ -282,7 +282,10 @@ static void RenameVariableReferences(UBlueprint* Blueprint, UClass* VariableClas
 
 void FBasePinChangeHelper::Broadcast(UBlueprint* InBlueprint, UK2Node_EditablePinBase* InTargetNode, UEdGraph* Graph)
 {
-	if (UK2Node_Tunnel* TunnelNode = Cast<UK2Node_Tunnel>(InTargetNode))
+	UK2Node_Tunnel* TunnelNode = Cast<UK2Node_Tunnel>(InTargetNode);
+	const UK2Node_FunctionTerminator* FunctionDefNode = Cast<const UK2Node_FunctionTerminator>(InTargetNode);
+	const UK2Node_Event* EventNode = Cast<const UK2Node_Event>(InTargetNode);
+	if (TunnelNode)
 	{
 		UBlueprint* Blueprint = FBlueprintEditorUtils::FindBlueprintForGraphChecked(Graph);
 
@@ -306,8 +309,15 @@ void FBasePinChangeHelper::Broadcast(UBlueprint* InBlueprint, UK2Node_EditablePi
 			EditCompositeTunnelNode(TunnelNode);
 		}
 	}
-	else if (UK2Node_FunctionTerminator* FunctionDefNode = Cast<UK2Node_FunctionTerminator>(InTargetNode))
+	else if (FunctionDefNode || EventNode)
 	{
+		auto Func = FFunctionFromNodeHelper::FunctionFromNode(FunctionDefNode ? static_cast<const UK2Node*>(FunctionDefNode) : static_cast<const UK2Node*>(EventNode));
+		const FName FuncName = Func 
+			? Func->GetFName() 
+			: (FunctionDefNode ? FunctionDefNode->SignatureName : EventNode->GetFunctionName());
+		const UClass* SignatureClass = Func
+			? Func->GetOwnerClass()
+			: (FunctionDefNode ? FunctionDefNode->SignatureClass : nullptr);
 		// Reconstruct all function call sites that call this function (in open blueprints)
 		for (TObjectIterator<UK2Node_CallFunction> It(RF_Transient); It; ++It)
 		{
@@ -316,10 +326,11 @@ void FBasePinChangeHelper::Broadcast(UBlueprint* InBlueprint, UK2Node_EditablePi
 			{
 				UBlueprint* CallSiteBlueprint = FBlueprintEditorUtils::FindBlueprintForNode(CallSite);
 
-				const bool bNameMatches = (CallSite->FunctionReference.GetMemberName() == FunctionDefNode->SignatureName);
+				const bool bNameMatches = (CallSite->FunctionReference.GetMemberName() == FuncName);
 				const UClass* MemberParentClass = CallSite->FunctionReference.GetMemberParentClass(CallSite->GetBlueprintClassFromNode());
-				const bool bClassMatchesEasy = (MemberParentClass != NULL) && (MemberParentClass->IsChildOf(FunctionDefNode->SignatureClass) || MemberParentClass->IsChildOf(InBlueprint->GeneratedClass));
-				const bool bClassMatchesHard = (CallSiteBlueprint != NULL) && (CallSite->FunctionReference.IsSelfContext()) && (FunctionDefNode->SignatureClass == NULL) && (CallSiteBlueprint == InBlueprint || CallSiteBlueprint->SkeletonGeneratedClass->IsChildOf(InBlueprint->SkeletonGeneratedClass));
+				const bool bClassMatchesEasy = (MemberParentClass != NULL) && (MemberParentClass->IsChildOf(SignatureClass) || MemberParentClass->IsChildOf(InBlueprint->GeneratedClass));
+				const bool bClassMatchesHard = (CallSiteBlueprint != NULL) && (CallSite->FunctionReference.IsSelfContext()) && (SignatureClass == NULL) 
+					&& (CallSiteBlueprint == InBlueprint || CallSiteBlueprint->SkeletonGeneratedClass->IsChildOf(InBlueprint->SkeletonGeneratedClass));
 				const bool bValidSchema = CallSite->GetSchema() != NULL;
 
 				if (bNameMatches && bValidSchema && (bClassMatchesEasy || bClassMatchesHard))
