Exclude NoRedist and NotForLicensees packages from generating DDC by default; we don't (AFAIK) ever need to generate DDC purely for internal use.

diff --git a/Engine/Source/Editor/UnrealEd/Private/Commandlets/DerivedDataCacheCommandlet.cpp b/Engine/Source/Editor/UnrealEd/Private/Commandlets/DerivedDataCacheCommandlet.cpp
--- a/Engine/Source/Editor/UnrealEd/Private/Commandlets/DerivedDataCacheCommandlet.cpp
+++ b/Engine/Source/Editor/UnrealEd/Private/Commandlets/DerivedDataCacheCommandlet.cpp
@@ -91,6 +91,11 @@ int32 UDerivedDataCacheCommandlet::Main( const FString& Params )
 			PackageFilter |= NORMALIZE_ExcludeDeveloperPackages;
 		}
 
+		if( !Switches.Contains(TEXT("NOREDIST")) )
+		{
+			PackageFilter |= NORMALIZE_ExcludeNoRedistPackages;
+		}
+
 		// assume the first token is the map wildcard/pathname
 		TArray<FString> Unused;
 		for ( int32 TokenIndex = 0; TokenIndex < Tokens.Num(); TokenIndex++ )
diff --git a/Engine/Source/Editor/UnrealEd/Private/Commandlets/PackageUtilities.cpp b/Engine/Source/Editor/UnrealEd/Private/Commandlets/PackageUtilities.cpp
--- a/Engine/Source/Editor/UnrealEd/Private/Commandlets/PackageUtilities.cpp
+++ b/Engine/Source/Editor/UnrealEd/Private/Commandlets/PackageUtilities.cpp
@@ -164,6 +164,15 @@ bool NormalizePackageNames( TArray<FString> PackageNames, TArray<FString>& Packa
 					continue;
 				}
 			}
+
+			if ( (PackageFilter&NORMALIZE_ExcludeNoRedistPackages) != 0 )
+			{
+				if(PackagePathNames.Contains("/NoRedist/") || PackagePathNames.Contains("/NotForLicensees/") || PackagePathNames.Contains("/EpicInternal/"))
+				{
+					PackagePathNames.RemoveAt(PackageIndex);
+					continue;
+				}
+			}
 		}
 	}
 
diff --git a/Engine/Source/Editor/UnrealEd/Public/PackageHelperFunctions.h b/Engine/Source/Editor/UnrealEd/Public/PackageHelperFunctions.h
--- a/Engine/Source/Editor/UnrealEd/Public/PackageHelperFunctions.h
+++ b/Engine/Source/Editor/UnrealEd/Public/PackageHelperFunctions.h
@@ -29,6 +29,8 @@ enum EPackageNormalizationFlags
 	NORMALIZE_ExcludeNonDeveloperPackages  = 0x10,
 	/** do not include packages inside the Engine/Content folders in the result array; only relevant if the input array is empty */
 	NORMALIZE_ExcludeEnginePackages		= 0x20,
+	/** do not include packages inside NoRedist or NotForLicensees folders */
+	NORMALIZE_ExcludeNoRedistPackages	= 0x40,
 	/** Combo flags */
 	NORMALIZE_DefaultFlags				= NORMALIZE_ResetExistingLoaders,
 };
