Fix a build issue on Mac caused by CL 2535363
(conditional expression is ambiguous; 'UClass *' can be converted to 'const TSubclassOf<class UObject>' and vice versa)

diff --git a/Engine/Source/Editor/UnrealEd/Private/Kismet2/BlueprintEditorUtils.cpp b/Engine/Source/Editor/UnrealEd/Private/Kismet2/BlueprintEditorUtils.cpp
--- a/Engine/Source/Editor/UnrealEd/Private/Kismet2/BlueprintEditorUtils.cpp
+++ b/Engine/Source/Editor/UnrealEd/Private/Kismet2/BlueprintEditorUtils.cpp
@@ -317,7 +317,7 @@ void FBasePinChangeHelper::Broadcast(UBlueprint* InBlueprint, UK2Node_EditablePi
 			: (FunctionDefNode ? FunctionDefNode->SignatureName : EventNode->GetFunctionName());
 		const UClass* SignatureClass = Func
 			? Func->GetOwnerClass()
-			: (FunctionDefNode ? FunctionDefNode->SignatureClass : nullptr);
+			: (UClass*)(FunctionDefNode ? FunctionDefNode->SignatureClass : nullptr);
 		// Reconstruct all function call sites that call this function (in open blueprints)
 		for (TObjectIterator<UK2Node_CallFunction> It(RF_Transient); It; ++It)
 		{
