Cook commandlet now saves other packages while waiting for the previous package to load it's cached cooked platform data (building shaders / textures).
Optimization for cooking when caching is happening.
Added a heirarchical timer for the cooking stats which gives summary after cooking is finished.

diff --git a/Engine/Source/Editor/UnrealEd/Private/CookOnTheFlyServer.cpp b/Engine/Source/Editor/UnrealEd/Private/CookOnTheFlyServer.cpp
--- a/Engine/Source/Editor/UnrealEd/Private/CookOnTheFlyServer.cpp
+++ b/Engine/Source/Editor/UnrealEd/Private/CookOnTheFlyServer.cpp
@@ -42,14 +42,14 @@ DEFINE_LOG_CATEGORY_STATIC(LogCook, Log, All);
 
 
 #define DEBUG_COOKONTHEFLY 0
-#define OUTPUT_TIMING 0
+#define OUTPUT_TIMING 1
 
 #define USEASSETREGISTRYFORDEPENDENTPACKAGES 1
 #define VERIFY_GETDEPENDENTPACKAGES 0 // verify has false hits because old serialization method for generating dependencies had errors (included transient objects which shouldn't be in asset registry), but you can still use verify to build a list then cross check against transient objects.  
 
 #if OUTPUT_TIMING
 
-
+#define HEIRARCHICAL_TIMER 1
 
 struct FTimerInfo
 {
@@ -73,22 +73,76 @@ public:
 	double Length;
 };
 
+#if HEIRARCHICAL_TIMER
+struct FHierarchicalTimerInfo : public FTimerInfo
+{
+public:
+	FHierarchicalTimerInfo *Parent;
+	TMap<FString, FHierarchicalTimerInfo*> Children;
+
+	FHierarchicalTimerInfo(const FHierarchicalTimerInfo& InTimerInfo) : FTimerInfo(InTimerInfo)
+	{
+		Children = InTimerInfo.Children;
+		for (auto& Child : Children)
+		{
+			Child.Value->Parent = this;
+		}
+	}
+
+	FHierarchicalTimerInfo(FHierarchicalTimerInfo&& InTimerInfo) : FTimerInfo(InTimerInfo)
+	{
+		Swap(Children, InTimerInfo.Children);
+
+		for (auto& Child : Children)
+		{
+			Child.Value->Parent = this;
+		}
+	}
+
+	FHierarchicalTimerInfo(FString &&Name) : FTimerInfo(MoveTemp(Name), 0.0)
+	{
+	}
+
+	FHierarchicalTimerInfo* FindChild(const FString& Name)
+	{
+		FHierarchicalTimerInfo* Child = (FHierarchicalTimerInfo*)(Children.FindRef(Name));
+		if ( !Child )
+		{
+			FString Temp = Name;
+			Child = Children.Add(Name, new FHierarchicalTimerInfo(MoveTemp(Temp)));
+			Child->Parent = this;
+		}
+
+		check(Child);
+		return Child;
+	}
+};
+
+static FHierarchicalTimerInfo RootTimerInfo(MoveTemp(FString(TEXT("Root"))));
+static FHierarchicalTimerInfo* CurrentTimerInfo = &RootTimerInfo;
+#endif
+
 
 static TArray<FTimerInfo> GTimerInfo;
 
+
+
 struct FScopeTimer
 {
 private:
 	bool Started;
 	bool DecrementScope;
 	static int GScopeDepth;
+#if HEIRARCHICAL_TIMER
+	FHierarchicalTimerInfo* HeirarchyTimerInfo;
+#endif
 public:
 
 	FScopeTimer( const FScopeTimer &outer )
 	{
 		Index = outer.Index;
 		DecrementScope = false;
-		Started = false;	
+		Started = false;
 	}
 
 	FScopeTimer( const FString &InName, bool IncrementScope = false )
@@ -106,6 +160,10 @@ public:
 		}
 		Index = GTimerInfo.Emplace(MoveTemp(Name), 0.0);
 		Started = false;
+#if HEIRARCHICAL_TIMER
+		HeirarchyTimerInfo = CurrentTimerInfo->FindChild(Name);
+		CurrentTimerInfo = HeirarchyTimerInfo;
+#endif
 	}
 
 	void Start()
@@ -114,6 +172,9 @@ public:
 		{
 			GTimerInfo[Index].Length -= FPlatformTime::Seconds();
 			Started = true;
+#if HEIRARCHICAL_TIMER
+			HeirarchyTimerInfo->Length -= FPlatformTime::Seconds();
+#endif
 		}
 	}
 
@@ -121,6 +182,9 @@ public:
 	{
 		if ( Started )
 		{
+#if HEIRARCHICAL_TIMER
+			HeirarchyTimerInfo->Length += FPlatformTime::Seconds();
+#endif
 			GTimerInfo[Index].Length += FPlatformTime::Seconds();
 			Started = false;
 		}
@@ -129,6 +193,10 @@ public:
 	~FScopeTimer()
 	{
 		Stop();
+#if HEIRARCHICAL_TIMER
+		check(CurrentTimerInfo == HeirarchyTimerInfo);
+		CurrentTimerInfo = HeirarchyTimerInfo->Parent;
+#endif
 		if ( DecrementScope )
 		{
 			--GScopeDepth;
@@ -336,19 +404,19 @@ FName UCookOnTheFlyServer::GetCachedStandardPackageFileFName( const FName& Packa
 
 FString UCookOnTheFlyServer::GetCachedPackageFilename( const UPackage* Package ) const 
 {
-	check( Package->GetName() == Package->GetFName().ToString() );
+	// check( Package->GetName() == Package->GetFName().ToString() );
 	return Cache( Package->GetFName() ).PackageFilename;
 }
 
 FString UCookOnTheFlyServer::GetCachedStandardPackageFilename( const UPackage* Package ) const 
 {
-	check( Package->GetName() == Package->GetFName().ToString() );
+	// check( Package->GetName() == Package->GetFName().ToString() );
 	return Cache( Package->GetFName() ).StandardFilename;
 }
 
 FName UCookOnTheFlyServer::GetCachedStandardPackageFileFName( const UPackage* Package ) const 
 {
-	check( Package->GetName() == Package->GetFName().ToString() );
+	// check( Package->GetName() == Package->GetFName().ToString() );
 	return Cache( Package->GetFName() ).StandardFileFName;
 }
 
@@ -1449,9 +1517,11 @@ uint32 UCookOnTheFlyServer::TickCookOnTheSide( const float TimeSlice, uint32 &Co
 
 		if ( PackagesToSave.Num() )
 		{
+			int32 OriginalPackagesToSaveCount = PackagesToSave.Num();
 			SCOPE_TIMER(SavingPackages);
 			for ( int32 I = 0; I < PackagesToSave.Num(); ++I )
 			{
+
 				UPackage *Package = PackagesToSave[I];
 
 				// if we are processing unsolicited packages we can optionally not save these right now
@@ -1459,6 +1529,8 @@ uint32 UCookOnTheFlyServer::TickCookOnTheSide( const float TimeSlice, uint32 &Co
 				// we want to do this in cook on the fly also, if there is a new network package request instead of saving unsolicited packages we can process the requested package
 				if ( (IsRealtimeMode() || IsCookOnTheFlyMode()) && (I >= FirstUnsolicitedPackage) )
 				{
+					SCOPE_TIMER(WaitingForCachedCookedPlatformData);
+
 					bool bShouldFinishTick = false;
 
 					if ( CookRequests.HasItems() )
@@ -1538,17 +1610,36 @@ uint32 UCookOnTheFlyServer::TickCookOnTheSide( const float TimeSlice, uint32 &Co
 				UPackage *NextNextPackage = PackagesToSave[FMath::Min( PackagesToSave.Num()-1, I + 2 )];
 				if ( NextPackage != Package )
 				{
+					SCOPE_TIMER(PrecachePlatformDataForNextPackage);
 					TArray<UObject*> ObjectsInPackage;
 					GetObjectsWithOuter( NextPackage, ObjectsInPackage );
 					BeginPackageCacheForCookedPlatformData(ObjectsInPackage);
 				}
 				if ( NextNextPackage != NextPackage )
 				{
+					SCOPE_TIMER(PrecachePlatformDataForNextNextPackage);
 					TArray<UObject*> ObjectsInPackage;
 					GetObjectsWithOuter( NextNextPackage, ObjectsInPackage );
 					BeginPackageCacheForCookedPlatformData(ObjectsInPackage);
 				}
 
+
+				// if we are running the cook commandlet
+				// if we already went through the entire package list then don't keep requeuing requests
+				if (!IsCookOnTheFlyMode() && !IsRealtimeMode() && (I < OriginalPackagesToSaveCount))
+				{
+					// we don't want to break out of this function without saving all the packages
+					// we can skip this package and do it later though
+
+					TArray<UObject*> ObjectsInPackage;
+					GetObjectsWithOuter(Package, ObjectsInPackage);
+					if (FinishPackageCacheForCookedPlatformData(ObjectsInPackage) == false)
+					{
+						PackagesToSave.Add(Package);
+					}
+				}
+
+
 				bool bShouldSaveAsync = true;
 				FString Temp;
 				if ( FParse::Value( FCommandLine::Get(), TEXT("-diffagainstcookdirectory="), Temp ) || FParse::Value(FCommandLine::Get(), TEXT("-breakonfile="), Temp))
@@ -1557,21 +1648,23 @@ uint32 UCookOnTheFlyServer::TickCookOnTheSide( const float TimeSlice, uint32 &Co
 					bShouldSaveAsync = false;
 				}
 
-				SCOPE_TIMER(SaveCookedPackage);
-				if( SaveCookedPackage(Package, SAVE_KeepGUID | (bShouldSaveAsync ? SAVE_Async : SAVE_None) | (IsCookFlagSet(ECookInitializationFlags::Unversioned) ? SAVE_Unversioned : 0), bWasUpToDate, AllTargetPlatformNames ) )
 				{
-					// Update flags used to determine garbage collection.
-					if (Package->ContainsMap())
+					SCOPE_TIMER(SaveCookedPackage);
+					if (SaveCookedPackage(Package, SAVE_KeepGUID | (bShouldSaveAsync ? SAVE_Async : SAVE_None) | (IsCookFlagSet(ECookInitializationFlags::Unversioned) ? SAVE_Unversioned : 0), bWasUpToDate, AllTargetPlatformNames))
 					{
-						Result |= COSR_CookedMap;
-					}
-					else
-					{
-						++CookedPackageCount;
-						Result |= COSR_CookedPackage;
+						// Update flags used to determine garbage collection.
+						if (Package->ContainsMap())
+						{
+							Result |= COSR_CookedMap;
+						}
+						else
+						{
+							++CookedPackageCount;
+							Result |= COSR_CookedPackage;
+						}
 					}
+					Timer.SavedPackage();
 				}
-				Timer.SavedPackage();
 
 				{
 					SCOPE_TIMER(ClearAllCachedCookedPlatformData);
