Iterative cooking respects custom versions changing and clears out cooked content.
Disabled timing info for cooker.

diff --git a/Engine/Source/Editor/UnrealEd/Private/CookOnTheFlyServer.cpp b/Engine/Source/Editor/UnrealEd/Private/CookOnTheFlyServer.cpp
--- a/Engine/Source/Editor/UnrealEd/Private/CookOnTheFlyServer.cpp
+++ b/Engine/Source/Editor/UnrealEd/Private/CookOnTheFlyServer.cpp
@@ -42,7 +42,7 @@ DEFINE_LOG_CATEGORY_STATIC(LogCook, Log, All);
 
 
 #define DEBUG_COOKONTHEFLY 0
-#define OUTPUT_TIMING 1
+#define OUTPUT_TIMING 0
 
 #define USEASSETREGISTRYFORDEPENDENTPACKAGES 1
 #define VERIFY_GETDEPENDENTPACKAGES 0 // verify has false hits because old serialization method for generating dependencies had errors (included transient objects which shouldn't be in asset registry), but you can still use verify to build a list then cross check against transient objects.  
@@ -2373,6 +2373,14 @@ bool UCookOnTheFlyServer::GetCurrentIniVersionStrings( const ITargetPlatform* Ta
 	GetVersionFormatNumbersForIniVersionStrings( IniVersionStrings, TEXT("TextureFormat"), TPM->GetTextureFormats() );
 	GetVersionFormatNumbersForIniVersionStrings( IniVersionStrings, TEXT("ShaderFormat"), TPM->GetShaderFormats() );
 
+
+	static const FCustomVersionContainer& CustomVersionContainer = FCustomVersionContainer::GetRegistered();
+	for (const auto& CustomVersion : CustomVersionContainer.GetAllVersions())
+	{
+		FString CustomVersionString = FString::Printf(TEXT("%s:%s:%d"), *CustomVersion.FriendlyName, *CustomVersion.Key.ToString(), CustomVersion.Version);
+		IniVersionStrings.Emplace(MoveTemp(CustomVersionString));
+	}
+
 	FString MaterialShaderMapDDCVersion = FString::Printf(TEXT("MaterialShaderMapDDCVersion:%s"), *GetMaterialShaderMapDDCKey());
 	IniVersionStrings.Emplace( MoveTemp(MaterialShaderMapDDCVersion) );
 	FString GlobalDDCVersion = FString::Printf(TEXT("GlobalDDCVersion:%s"), *GetGlobalShaderMapDDCKey());
