UE-14903 Crash occurs when compiling the ViewportsTutorial blueprint.

When a BPGC from a DataOnly BP has UberGraphFrame, it will fully recompiled while regeneration.

diff --git a/Engine/Source/Editor/UnrealEd/Private/Kismet2/BlueprintEditorUtils.cpp b/Engine/Source/Editor/UnrealEd/Private/Kismet2/BlueprintEditorUtils.cpp
--- a/Engine/Source/Editor/UnrealEd/Private/Kismet2/BlueprintEditorUtils.cpp
+++ b/Engine/Source/Editor/UnrealEd/Private/Kismet2/BlueprintEditorUtils.cpp
@@ -1422,7 +1422,13 @@ UClass* FBlueprintEditorUtils::RegenerateBlueprintClass(UBlueprint* Blueprint, U
 		const bool bDataOnlyClassThatMustBeRecompiled = !bHasCode && !bIsMacro
 			&& (!ClassToRegenerate || (Blueprint->ParentClass != ClassToRegenerate->GetSuperClass()));
 
-		if (bHasCode || bDataOnlyClassThatMustBeRecompiled)
+		auto BPGClassToRegenerate = Cast<UBlueprintGeneratedClass>(ClassToRegenerate);
+		const bool bHasPendingUberGraphFrame = BPGClassToRegenerate
+			&& (BPGClassToRegenerate->UberGraphFramePointerProperty || BPGClassToRegenerate->UberGraphFunction);
+
+		const bool bShouldBeRecompiled = bHasCode || bDataOnlyClassThatMustBeRecompiled || bHasPendingUberGraphFrame;
+
+		if (bShouldBeRecompiled)
 		{
 			// Make sure parent function calls are up to date
 			FBlueprintEditorUtils::ConformCallsToParentFunctions(Blueprint);
@@ -1518,7 +1524,7 @@ UClass* FBlueprintEditorUtils::RegenerateBlueprintClass(UBlueprint* Blueprint, U
 			Package->FindExportsInMemoryFirst(true);
 		}
 
-		bRegenerated = bHasCode || bDataOnlyClassThatMustBeRecompiled;
+		bRegenerated = bShouldBeRecompiled;
 
 		if (!FKismetEditorUtilities::IsClassABlueprintSkeleton(ClassToRegenerate))
 		{
