Editor: Added analytics for # of times an asset editor was opened during a session

diff --git a/Engine/Source/Editor/UnrealEd/Private/Toolkits/AssetEditorManager.cpp b/Engine/Source/Editor/UnrealEd/Private/Toolkits/AssetEditorManager.cpp
--- a/Engine/Source/Editor/UnrealEd/Private/Toolkits/AssetEditorManager.cpp
+++ b/Engine/Source/Editor/UnrealEd/Private/Toolkits/AssetEditorManager.cpp
@@ -57,15 +57,20 @@ void FAssetEditorManager::OnExit()
 	CloseAllAssetEditors();
 
 	// Don't attempt to report usage stats if analytics isn't available
-	if(FEngineAnalytics::IsAvailable())
+	if (FEngineAnalytics::IsAvailable())
 	{
 		TArray<FAnalyticsEventAttribute> EditorUsageAttribs;
-		for (auto Iter = EditorDurations.CreateConstIterator(); Iter; ++Iter)
+		EditorUsageAttribs.Empty(2);
+		for (auto Iter = EditorUsageAnalytics.CreateConstIterator(); Iter; ++Iter)
 		{
-			FString EventName = FString::Printf(TEXT("Editor.Usage.%s"), *Iter.Key().ToString());
+			const FAssetEditorAnalyticInfo& Data = Iter.Value();
+			EditorUsageAttribs.Reset();
+			EditorUsageAttribs.Emplace(TEXT("TotalDuration.Seconds"), FString::Printf(TEXT("%.1f"), Data.SumDuration.GetTotalSeconds()));
+			EditorUsageAttribs.Emplace(TEXT("OpenedAssets.Count"), FString::Printf(TEXT("%d"), Data.NumTimesOpened));
 
-			FEngineAnalytics::GetProvider().RecordEvent(EventName, TEXT("Duration.Seconds"), FString::Printf( TEXT( "%.1f" ), Iter.Value().GetTotalSeconds()));
-		}	
+			const FString EventName = FString::Printf(TEXT("Editor.Usage.%s"), *Iter.Key().ToString());
+			FEngineAnalytics::GetProvider().RecordEvent(EventName, EditorUsageAttribs);
+		}
 	}
 }
 
@@ -206,12 +211,9 @@ void FAssetEditorManager::NotifyEditorClosed(IAssetEditorInstance* InInstance)
 	FOpenedEditorTime EditorTime = OpenedEditorTimes.FindAndRemoveChecked(InInstance);
 
 	// Record the editor open-close duration
-	FTimespan* Duration = EditorDurations.Find(EditorTime.EditorName);
-	if (Duration == NULL)
-	{
-		Duration = &EditorDurations.Add(EditorTime.EditorName, FTimespan::Zero());
-	}
-	*Duration += FDateTime::UtcNow() - EditorTime.OpenedTime;
+	FAssetEditorAnalyticInfo& AnalyticsForThisAsset = EditorUsageAnalytics.FindOrAdd(EditorTime.EditorName);
+	AnalyticsForThisAsset.SumDuration += FDateTime::UtcNow() - EditorTime.OpenedTime;
+	AnalyticsForThisAsset.NumTimesOpened++;
 
 	SaveOpenAssetEditors(false);
 }
diff --git a/Engine/Source/Editor/UnrealEd/Public/Toolkits/AssetEditorManager.h b/Engine/Source/Editor/UnrealEd/Public/Toolkits/AssetEditorManager.h
--- a/Engine/Source/Editor/UnrealEd/Public/Toolkits/AssetEditorManager.h
+++ b/Engine/Source/Editor/UnrealEd/Public/Toolkits/AssetEditorManager.h
@@ -140,6 +140,19 @@ private:
 		FDateTime OpenedTime;
 	};
 
+	/** struct used to track total time and # of invocations during an overall UnrealEd session */
+	struct FAssetEditorAnalyticInfo
+	{
+		FTimespan SumDuration;
+		int32 NumTimesOpened;
+
+		FAssetEditorAnalyticInfo()
+			: SumDuration(0)
+			, NumTimesOpened(0)
+		{
+		}
+	};
+
 	/**
 	 * Holds the opened assets.
 	 */
@@ -158,7 +171,7 @@ private:
 	/**
 	 * Holds the cumulative time editors have been open by type.
 	 */
-	TMap<FName, FTimespan> EditorDurations;
+	TMap<FName, FAssetEditorAnalyticInfo> EditorUsageAnalytics;
 
 private:
 
