Editor: Renamed OpenedAssets.Count to OpenedInstances.Count to be more precise

diff --git a/Engine/Source/Editor/UnrealEd/Private/Toolkits/AssetEditorManager.cpp b/Engine/Source/Editor/UnrealEd/Private/Toolkits/AssetEditorManager.cpp
--- a/Engine/Source/Editor/UnrealEd/Private/Toolkits/AssetEditorManager.cpp
+++ b/Engine/Source/Editor/UnrealEd/Private/Toolkits/AssetEditorManager.cpp
@@ -66,7 +66,7 @@ void FAssetEditorManager::OnExit()
 			const FAssetEditorAnalyticInfo& Data = Iter.Value();
 			EditorUsageAttribs.Reset();
 			EditorUsageAttribs.Emplace(TEXT("TotalDuration.Seconds"), FString::Printf(TEXT("%.1f"), Data.SumDuration.GetTotalSeconds()));
-			EditorUsageAttribs.Emplace(TEXT("OpenedAssets.Count"), FString::Printf(TEXT("%d"), Data.NumTimesOpened));
+			EditorUsageAttribs.Emplace(TEXT("OpenedInstances.Count"), FString::Printf(TEXT("%d"), Data.NumTimesOpened));
 
 			const FString EventName = FString::Printf(TEXT("Editor.Usage.%s"), *Iter.Key().ToString());
 			FEngineAnalytics::GetProvider().RecordEvent(EventName, EditorUsageAttribs);
