Fixed issue where cancelling unsaved packages dialog when launching resulted in the Launch button permanently disabled.

diff --git a/Engine/Source/Editor/UnrealEd/Private/PlayLevel.cpp b/Engine/Source/Editor/UnrealEd/Private/PlayLevel.cpp
--- a/Engine/Source/Editor/UnrealEd/Private/PlayLevel.cpp
+++ b/Engine/Source/Editor/UnrealEd/Private/PlayLevel.cpp
@@ -1541,6 +1541,7 @@ void UEditorEngine::PlayUsingLauncher()
 			bool bSaveContentPackages = true;
 			if (!FEditorFileUtils::SaveDirtyPackages(bPromptUserToSave, bSaveMapPackages, bSaveContentPackages))
 			{
+				CancelRequestPlaySession();
 				return;
 			}
 		}
@@ -1550,7 +1551,7 @@ void UEditorEngine::PlayUsingLauncher()
 
 			if (MapNames.Num() == 0)
 			{
-				GEditor->CancelRequestPlaySession();
+				CancelRequestPlaySession();
 				return;
 			}
 		}
