Fixed crash in gathering text from assets, specifically dialogue waves, if the dialogue wave has no *non-null* targets.

diff --git a/Engine/Source/Editor/UnrealEd/Private/Commandlets/GatherTextFromAssetsCommandlet.cpp b/Engine/Source/Editor/UnrealEd/Private/Commandlets/GatherTextFromAssetsCommandlet.cpp
--- a/Engine/Source/Editor/UnrealEd/Private/Commandlets/GatherTextFromAssetsCommandlet.cpp
+++ b/Engine/Source/Editor/UnrealEd/Private/Commandlets/GatherTextFromAssetsCommandlet.cpp
@@ -82,6 +82,7 @@ bool UGatherTextFromAssetsCommandlet::FDialogueHelper::ProcessDialogueWave( cons
 		TArray<FString> TargetDisplayNameList;
 
 		// Collect info on all the targets
+		bool HasTarget = false;
 		for( const UDialogueVoice* TargetDialogueVoice : DialogueContext.Targets )
 		{
 			if( TargetDialogueVoice )
@@ -111,11 +112,13 @@ bool UGatherTextFromAssetsCommandlet::FDialogueHelper::ProcessDialogueWave( cons
 				{
 					AccumulatedTargetPlurality = EGrammaticalNumber::Plural;
 				}
+
+				HasTarget = true;
 			}
 		}
 
-		FString FinalTargetGender = DialogueContext.Targets.Num() ? GetGrammaticalGenderString( AccumulatedTargetGender ) : TEXT("");
-		FString FinalTargetPlurality = DialogueContext.Targets.Num() ? GetGrammaticalNumberString( AccumulatedTargetPlurality ) : TEXT("");
+		FString FinalTargetGender = HasTarget ? GetGrammaticalGenderString( AccumulatedTargetGender ) : TEXT("");
+		FString FinalTargetPlurality = HasTarget ? GetGrammaticalNumberString( AccumulatedTargetPlurality ) : TEXT("");
 
 		// Add the context specific variation
 		{
