Landscape grass output node is no longer removed when cleaning unused expressions in the material editor

diff --git a/Engine/Source/Editor/UnrealEd/Private/MaterialGraph.cpp b/Engine/Source/Editor/UnrealEd/Private/MaterialGraph.cpp
--- a/Engine/Source/Editor/UnrealEd/Private/MaterialGraph.cpp
+++ b/Engine/Source/Editor/UnrealEd/Private/MaterialGraph.cpp
@@ -327,6 +327,19 @@ void UMaterialGraph::GetUnusedExpressions(TArray<UEdGraphNode*>& UnusedNodes) co
 				NodesToCheck.Push(InputPins[Index]->LinkedTo[0]->GetOwningNode());
 			}
 		}
+
+		for (int32 Index = 0; Index < Nodes.Num(); Index++)
+		{
+			UMaterialGraphNode* GraphNode = Cast<UMaterialGraphNode>(Nodes[Index]);
+			if (GraphNode)
+			{
+				UMaterialExpressionLandscapeGrassOutput* Output = Cast<UMaterialExpressionLandscapeGrassOutput>(GraphNode->MaterialExpression);
+				if (Output)
+				{
+					NodesToCheck.Push(GraphNode);
+				}
+			}
+		}
 	}
 	else if (MaterialFunction)
 	{
