PR #1120: ParticleSystemAudit CrashFix (Contributed by [REDACTED_PERSON])

#github https://github.com/EpicGames/UnrealEngine/pull/1120

diff --git a/Engine/Source/Editor/UnrealEd/Private/Commandlets/ParticleSystemAuditCommandlet.cpp b/Engine/Source/Editor/UnrealEd/Private/Commandlets/ParticleSystemAuditCommandlet.cpp
--- a/Engine/Source/Editor/UnrealEd/Private/Commandlets/ParticleSystemAuditCommandlet.cpp
+++ b/Engine/Source/Editor/UnrealEd/Private/Commandlets/ParticleSystemAuditCommandlet.cpp
@@ -40,7 +40,7 @@ int32 UParticleSystemAuditCommandlet::Main(const FString& Params)
 
 bool UParticleSystemAuditCommandlet::ProcessParticleSystems()
 {
-	FAssetRegistryModule& AssetRegistryModule = FModuleManager::GetModuleChecked<FAssetRegistryModule>(TEXT("AssetRegistry"));
+	FAssetRegistryModule& AssetRegistryModule = FModuleManager::LoadModuleChecked<FAssetRegistryModule>(TEXT("AssetRegistry"));
 	IAssetRegistry& AssetRegistry = AssetRegistryModule.Get();
 	AssetRegistry.SearchAllAssets(true);
 
