Build fix (missing include)

diff --git a/Engine/Source/Editor/UnrealEd/Private/MaterialGraph.cpp b/Engine/Source/Editor/UnrealEd/Private/MaterialGraph.cpp
--- a/Engine/Source/Editor/UnrealEd/Private/MaterialGraph.cpp
+++ b/Engine/Source/Editor/UnrealEd/Private/MaterialGraph.cpp
@@ -7,6 +7,7 @@
 
 #include "Materials/MaterialExpressionComment.h"
 #include "Materials/MaterialExpressionFunctionOutput.h"
+#include "Materials/MaterialExpressionLandscapeGrassOutput.h"
 
 #include "GraphEditor.h"
 #include "MaterialShared.h"
