Changed the normal map identification heuristics so that nearly transparent pixels aren't considered.  This should allow normal maps with large amounts of transparency to still be considered as such.

diff --git a/Engine/Source/Editor/UnrealEd/Private/NormalMapIdentification.cpp b/Engine/Source/Editor/UnrealEd/Private/NormalMapIdentification.cpp
--- a/Engine/Source/Editor/UnrealEd/Private/NormalMapIdentification.cpp
+++ b/Engine/Source/Editor/UnrealEd/Private/NormalMapIdentification.cpp
@@ -191,10 +191,9 @@ public:
 				FLinearColor ColorSample = Sampler.DoSampleColor( X, Y );
 				if ( !ColorSample.IsAlmostBlack() )
 				{
+					// Nearly transparent pixels don't contribute to the calculation
 					if (FMath::IsNearlyZero(ColorSample.A, AlphaComponentNearlyZeroThreshold))
 					{
-						AverageColor += FLinearColor::Transparent;
-						NumSamplesTaken++;
 						continue;
 					}
 
