Added ue4 version and ue4 licensee version for consideration when using -iterate.

diff --git a/Engine/Source/Editor/UnrealEd/Private/CookOnTheFlyServer.cpp b/Engine/Source/Editor/UnrealEd/Private/CookOnTheFlyServer.cpp
--- a/Engine/Source/Editor/UnrealEd/Private/CookOnTheFlyServer.cpp
+++ b/Engine/Source/Editor/UnrealEd/Private/CookOnTheFlyServer.cpp
@@ -2381,6 +2381,11 @@ bool UCookOnTheFlyServer::GetCurrentIniVersionStrings( const ITargetPlatform* Ta
 		IniVersionStrings.Emplace(MoveTemp(CustomVersionString));
 	}
 
+	FString UE4Ver = FString::Printf(TEXT("PackageFileVersions:%d:%d"), GPackageFileUE4Version, GPackageFileLicenseeUE4Version);
+	IniVersionStrings.Emplace(MoveTemp(UE4Ver));
+	FString UE4NetVer = FString::Printf(TEXT("NetFileVersions:%d:%d"), GEngineNetVersion, GEngineNegotiationVersion);
+	IniVersionStrings.Emplace(MoveTemp(UE4NetVer));
+
 	FString MaterialShaderMapDDCVersion = FString::Printf(TEXT("MaterialShaderMapDDCVersion:%s"), *GetMaterialShaderMapDDCKey());
 	IniVersionStrings.Emplace( MoveTemp(MaterialShaderMapDDCVersion) );
 	FString GlobalDDCVersion = FString::Printf(TEXT("GlobalDDCVersion:%s"), *GetGlobalShaderMapDDCKey());
