Promoted cook on the side out of experimental

diff --git a/Engine/Source/Editor/UnrealEd/Classes/Settings/EditorExperimentalSettings.h b/Engine/Source/Editor/UnrealEd/Classes/Settings/EditorExperimentalSettings.h
--- a/Engine/Source/Editor/UnrealEd/Classes/Settings/EditorExperimentalSettings.h
+++ b/Engine/Source/Editor/UnrealEd/Classes/Settings/EditorExperimentalSettings.h
@@ -81,10 +81,6 @@ public:
 	UPROPERTY(EditAnywhere, config, Category = Cooking, meta = (DisplayName = "Disable Cook In The Editor feature (cooks from launch on will be run in a separate process if disabled)", ConfigRestartRequired=true))
 	bool bDisableCookInEditor;
 
-	/** Enable cook on the side */
-	UPROPERTY(EditAnywhere, config, Category = Cooking, meta = (DisplayName = "Cook On The Side, requires editor restart (Run a cook on the fly server in the background of the editor)"))
-	bool bCookOnTheSide;
-
 	/** Enable -iterate for launch on */
 	UPROPERTY(EditAnywhere, config, Category = Cooking, meta = (DisplayName = "Iterative cooking for builds launched form the editor (launch on)"))
 	bool bIterativeCookingForLaunchOn;
diff --git a/Engine/Source/Editor/UnrealEd/Private/UnrealEdEngine.cpp b/Engine/Source/Editor/UnrealEd/Private/UnrealEdEngine.cpp
--- a/Engine/Source/Editor/UnrealEd/Private/UnrealEdEngine.cpp
+++ b/Engine/Source/Editor/UnrealEd/Private/UnrealEdEngine.cpp
@@ -128,7 +128,11 @@ void UUnrealEdEngine::Init(IEngineLoop* InEngineLoop)
 	UEditorExperimentalSettings const* ExperimentalSettings =  GetDefault<UEditorExperimentalSettings>();
 	ECookInitializationFlags BaseCookingFlags = ECookInitializationFlags::AutoTick | ECookInitializationFlags::AsyncSave | ECookInitializationFlags::Compressed;
 	BaseCookingFlags |= ExperimentalSettings->bIterativeCookingForLaunchOn ? ECookInitializationFlags::Iterative : ECookInitializationFlags::None;
-	if ( ExperimentalSettings->bCookOnTheSide )
+
+	bool bEnableCookOnTheSide = false;
+	GConfig->GetBool(TEXT("/Script/UnrealEd.CookerSettings"), TEXT("bEnableCookOnTheSide"), bEnableCookOnTheSide, GEngineIni);
+
+	if ( bEnableCookOnTheSide )
 	{
 		CookServer = NewObject<UCookOnTheFlyServer>();
 		CookServer->Initialize( ECookMode::CookOnTheFlyFromTheEditor, BaseCookingFlags );
diff --git a/Engine/Source/Editor/UnrealEd/Public/CookerSettings.h b/Engine/Source/Editor/UnrealEd/Public/CookerSettings.h
--- a/Engine/Source/Editor/UnrealEd/Public/CookerSettings.h
+++ b/Engine/Source/Editor/UnrealEd/Public/CookerSettings.h
@@ -18,6 +18,10 @@ class UNREALED_API UCookerSettings : public UDeveloperSettings
 
 public:
 
+	UPROPERTY(GlobalConfig, EditAnywhere, Category = Cooker, meta = (DisplayName = "Enable cooking via network in the background of the editor, launch on uses this setting, requires device to have network access to editor"))
+	bool bEnableCookOnTheSide;
+
+
 	/** Quality of 0 means fastest, 4 means best quality */
 	UPROPERTY(GlobalConfig, EditAnywhere, Category = Textures, meta = (DisplayName = "PVRTC Compression Quality (0-4, 0 is fastest)"))
 	int32 DefaultPVRTCQuality;
