Fixed replace references not being able to pick up references in linkers

The object reference replacement archives can't serialize linkers which may hold references to the objects in question. Resetting loaders before we run this code ensures there are no unreachable references.

This addresses UE-13884 and UE-13877

diff --git a/Engine/Source/Editor/UnrealEd/Private/ObjectTools.cpp b/Engine/Source/Editor/UnrealEd/Private/ObjectTools.cpp
--- a/Engine/Source/Editor/UnrealEd/Private/ObjectTools.cpp
+++ b/Engine/Source/Editor/UnrealEd/Private/ObjectTools.cpp
@@ -580,6 +580,9 @@ namespace ObjectTools
 			}
 		}
 
+		// Reset linker loaders to remove the possibility that any references to 'ObjectsToReplace' exist in the loaders (these can't get picked up by the replace archives)
+		ResetLoaders(nullptr);
+
 		TMap<UObject*, int32> ObjToNumRefsMap;
 		if( ObjectToReplaceWith != NULL )
 		{
