Improved heuristics for detecting normal map textures.

diff --git a/Engine/Source/Editor/UnrealEd/Private/NormalMapIdentification.cpp b/Engine/Source/Editor/UnrealEd/Private/NormalMapIdentification.cpp
--- a/Engine/Source/Editor/UnrealEd/Private/NormalMapIdentification.cpp
+++ b/Engine/Source/Editor/UnrealEd/Private/NormalMapIdentification.cpp
@@ -30,7 +30,7 @@ namespace
 
 	// We sample up to this many tiles in each axis. Sampling more tiles
 	// will likely be more accurate, but will take longer.
-	const int32 MaxTilesPerAxis = 8;
+	const int32 MaxTilesPerAxis = 16;
 
 	// This is used in the comparison with "mid-gray"
 	const float ColorComponentNearlyZeroThreshold = (2.0f / 255.0f);
@@ -43,6 +43,12 @@ namespace
 	// the assumption that these are likely invalid values for a general normal map
 	const float ColorComponentMinVectorThreshold = (2.0f / 255.0f) * 2.0f - 1.0f;
 	const float ColorComponentMaxVectorThreshold = (253.0f/255.0f) * 2.0f - 1.0f;
+
+	// This is the threshold delta length for a vector to be considered as a unit vector
+	const float NormalVectorUnitLengthDeltaThreshold = 0.45f;
+
+	// Rejected to taken sample ratio threshold.
+	const float RejectedToTakenRatioThreshold = 0.33f;
 }
 
 ////////////////////////////////////////////////////////////////////////////////
@@ -171,6 +177,7 @@ template<class SamplerClass> class TNormalMapAnalyzer
 public:
 	TNormalMapAnalyzer()
 	: NumSamplesTaken(0)
+	, NumSamplesRejected(0)
 	, NumSamplesThreshold(0)
 	, AverageColor(0.0f,0.0f,0.0f,0.0f)
 	{
@@ -189,38 +196,42 @@ public:
 			for ( int32 X=Left; X != (Left+Width); X++ )
 			{
 				FLinearColor ColorSample = Sampler.DoSampleColor( X, Y );
-				if ( !ColorSample.IsAlmostBlack() )
+
+				// Nearly black or transparent pixels don't contribute to the calculation
+				if (FMath::IsNearlyZero(ColorSample.A, AlphaComponentNearlyZeroThreshold) ||
+					ColorSample.IsAlmostBlack())
 				{
-					// Nearly transparent pixels don't contribute to the calculation
-					if (FMath::IsNearlyZero(ColorSample.A, AlphaComponentNearlyZeroThreshold))
-					{
-						continue;
-					}
+					continue;
+				}
 
-					// Scale and bias, if required, to get a signed vector
-					float Vx = Sampler.ScaleAndBiasComponent( ColorSample.R );
-					float Vy = Sampler.ScaleAndBiasComponent( ColorSample.G );
-					float Vz = Sampler.ScaleAndBiasComponent( ColorSample.B );
+				// Scale and bias, if required, to get a signed vector
+				float Vx = Sampler.ScaleAndBiasComponent( ColorSample.R );
+				float Vy = Sampler.ScaleAndBiasComponent( ColorSample.G );
+				float Vz = Sampler.ScaleAndBiasComponent( ColorSample.B );
 
-					// If the vector is close to zero (mid-gray) then ignore it as invalid
-					if ( FMath::IsNearlyZero(Vx, ColorComponentNearlyZeroThreshold) &&
-						FMath::IsNearlyZero(Vy, ColorComponentNearlyZeroThreshold) &&
-						FMath::IsNearlyZero(Vz, ColorComponentNearlyZeroThreshold) )
-					{
-						continue;
-					}
+				const float Length = FMath::Sqrt(Vx * Vx + Vy * Vy + Vz * Vz);
+				if (Length < ColorComponentNearlyZeroThreshold)
+				{
+					// mid-grey pixels representing (0,0,0) are also not considered as they may be used to denote unused areas
+					continue;
+				}
 
-					// Assume that if X or Y are very close to +1 or -1 then it is an invalid sample.
-					// If this were to happen in a real normal map, it would imply an impossible gradient
-					if ( !FMath::IsWithinInclusive( Vx, ColorComponentMinVectorThreshold, ColorComponentMaxVectorThreshold ) ||
-						!FMath::IsWithinInclusive( Vy, ColorComponentMinVectorThreshold, ColorComponentMaxVectorThreshold ) )
-					{
-						continue;
-					}
+				// If the vector is sufficiently different in length from a unit vector, consider it invalid.
+				if (FMath::Abs(Length - 1.0f) > NormalVectorUnitLengthDeltaThreshold)
+				{
+					NumSamplesRejected++;
+					continue;
+				}
 
-					AverageColor += ColorSample;
-					NumSamplesTaken++;
+				// If the vector is pointing backwards then it is an invalid sample, so consider it invalid
+				if (Vz < 0.0f)
+				{
+					NumSamplesRejected++;
+					continue;
 				}
+
+				AverageColor += ColorSample;
+				NumSamplesTaken++;
 			}
 		}
 	}
@@ -291,6 +302,13 @@ public:
 		// if we managed to take a reasonable number of samples then we can evaluate the result
 		if ( NumSamplesTaken >= NumSamplesThreshold )
 		{
+			const float RejectedToTakenRatio = static_cast<float>(NumSamplesRejected) / static_cast<float>(NumSamplesTaken);
+			if ( RejectedToTakenRatio >= RejectedToTakenRatioThreshold )
+			{
+				// Too many invalid samples, probably not a normal map
+				return false;
+			}
+
 			AverageColor /= (float)NumSamplesTaken;
 
 			// See if the resulting vector lies anywhere near the {0,0,1} vector
@@ -316,6 +334,7 @@ public:
 	}
 
 	int32 NumSamplesTaken;
+	int32 NumSamplesRejected;
 	int32 NumSamplesThreshold;
 	FLinearColor AverageColor;
 
