Fixed misspelling in persona preferences

diff --git a/Engine/Source/Editor/Persona/Private/SAnimTimingPanel.cpp b/Engine/Source/Editor/Persona/Private/SAnimTimingPanel.cpp
--- a/Engine/Source/Editor/Persona/Private/SAnimTimingPanel.cpp
+++ b/Engine/Source/Editor/Persona/Private/SAnimTimingPanel.cpp
@@ -31,17 +31,17 @@ void SAnimTimingNode::Construct(const FArguments& InArgs)
 		case ETimingElementType::NotifyStateBegin:
 		case ETimingElementType::NotifyStateEnd:
 		{
-			NodeColour = EditorOptions->NotifyTimingNodeColour;
+			NodeColour = EditorOptions->NotifyTimingNodeColor;
 			break;
 		}
 		case ETimingElementType::BranchPointNotify:
 		{
-			NodeColour = EditorOptions->BranchingPointTimingNodeColour;
+			NodeColour = EditorOptions->BranchingPointTimingNodeColor;
 			break;
 		}
 		case ETimingElementType::Section:
 		{
-			NodeColour = EditorOptions->SectionTimingNodeColour;
+			NodeColour = EditorOptions->SectionTimingNodeColor;
 			break;
 		}
 		default:
diff --git a/Engine/Source/Editor/UnrealEd/Classes/Preferences/PersonaOptions.h b/Engine/Source/Editor/UnrealEd/Classes/Preferences/PersonaOptions.h
--- a/Engine/Source/Editor/UnrealEd/Classes/Preferences/PersonaOptions.h
+++ b/Engine/Source/Editor/UnrealEd/Classes/Preferences/PersonaOptions.h
@@ -50,13 +50,13 @@ class UNREALED_API UPersonaOptions : public UObject
 	uint32 DefaultLocalAxesSelection;
 
 	UPROPERTY(EditAnywhere, config, Category = Options)
-	FLinearColor SectionTimingNodeColour;
+	FLinearColor SectionTimingNodeColor;
 
 	UPROPERTY(EditAnywhere, config, Category = Options)
-	FLinearColor NotifyTimingNodeColour;
+	FLinearColor NotifyTimingNodeColor;
 
 	UPROPERTY(EditAnywhere, config, Category = Options)
-	FLinearColor BranchingPointTimingNodeColour;
+	FLinearColor BranchingPointTimingNodeColor;
 
 public:
 	void SetViewportBackgroundColor( const FLinearColor& InViewportBackgroundColor);
diff --git a/Engine/Source/Editor/UnrealEd/Private/PreferenceStubs.cpp b/Engine/Source/Editor/UnrealEd/Private/PreferenceStubs.cpp
--- a/Engine/Source/Editor/UnrealEd/Private/PreferenceStubs.cpp
+++ b/Engine/Source/Editor/UnrealEd/Private/PreferenceStubs.cpp
@@ -41,9 +41,9 @@ UPersonaOptions::UPersonaOptions(const FObjectInitializer& ObjectInitializer)
 {
 	ViewModeIndex = VMI_Lit;
 
-	SectionTimingNodeColour = FLinearColor(0.0f, 1.0f, 0.0f);
-	NotifyTimingNodeColour = FLinearColor(1.0f, 0.0f, 0.0f);
-	BranchingPointTimingNodeColour = FLinearColor(0.5f, 1.0f, 1.0f);
+	SectionTimingNodeColor = FLinearColor(0.0f, 1.0f, 0.0f);
+	NotifyTimingNodeColor = FLinearColor(1.0f, 0.0f, 0.0f);
+	BranchingPointTimingNodeColor = FLinearColor(0.5f, 1.0f, 1.0f);
 }
 
 void UPersonaOptions::SetViewportBackgroundColor( const FLinearColor& InViewportBackgroundColor)
@@ -115,18 +115,18 @@ void UPersonaOptions::SetShowMeshStats( int32 InShowMeshStats )
 
 void UPersonaOptions::SetSectionTimingNodeColor(const FLinearColor& InColor)
 {
-	SectionTimingNodeColour = InColor;
+	SectionTimingNodeColor = InColor;
 	SaveConfig();
 }
 
 void UPersonaOptions::SetNotifyTimingNodeColor(const FLinearColor& InColor)
 {
-	NotifyTimingNodeColour = InColor;
+	NotifyTimingNodeColor = InColor;
 	SaveConfig();
 }
 
 void UPersonaOptions::SetBranchingPointTimingNodeColor(const FLinearColor& InColor)
 {
-	BranchingPointTimingNodeColour = InColor;
+	BranchingPointTimingNodeColor = InColor;
 	SaveConfig();
 }
