[UE-14592] On load, destroying components that were inherited from a native super-class (and have since been removed from that class).

diff --git a/Engine/Source/Editor/UnrealEd/Private/Kismet2/Kismet2.cpp b/Engine/Source/Editor/UnrealEd/Private/Kismet2/Kismet2.cpp
--- a/Engine/Source/Editor/UnrealEd/Private/Kismet2/Kismet2.cpp
+++ b/Engine/Source/Editor/UnrealEd/Private/Kismet2/Kismet2.cpp
@@ -970,6 +970,116 @@ void FKismetEditorUtilities::GenerateCppCode(UBlueprint* InBlueprintObj, TShared
 	}
 }
 
+
+namespace ConformComponentsUtils
+{
+	static void ConformRemovedNativeComponents(UObject* BpCdo);
+	static UObject* FindeNativeArchetype(UActorComponent* Component);
+};
+
+static void ConformComponentsUtils::ConformRemovedNativeComponents(UObject* BpCdo)
+{
+	UClass* BlueprintClass = BpCdo->GetClass();
+	check(BpCdo->HasAnyFlags(RF_ClassDefaultObject) && BlueprintClass->HasAnyClassFlags(CLASS_CompiledFromBlueprint));
+
+	AActor* ActorCDO = Cast<AActor>(BpCdo);
+	if (ActorCDO == nullptr)
+	{
+		return;
+	}
+
+	UClass* const NativeSuperClass = FBlueprintEditorUtils::FindFirstNativeClass(BlueprintClass);
+	const AActor* NativeCDO = GetDefault<AActor>(NativeSuperClass);
+
+	TInlineComponentArray<UActorComponent*> OldNativeComponents;
+	ActorCDO->GetComponents(OldNativeComponents);
+	TInlineComponentArray<UActorComponent*> NewNativeComponents;
+	NativeCDO->GetComponents(NewNativeComponents);
+
+	TSet<UObject*> DestroyedComponents;
+	for (UActorComponent* Component : OldNativeComponents)
+	{
+		UObject* NativeArchetype = FindeNativeArchetype(Component);
+		if ((NativeArchetype == nullptr) || !NativeArchetype->HasAnyFlags(RF_ClassDefaultObject))
+		{
+			continue;
+		}
+		// else, the component has been removed from our native super class
+
+		Component->DestroyComponent(/*bPromoteChildren =*/false);
+		DestroyedComponents.Add(Component);
+
+		UClass* ComponentClass = Component->GetClass();
+		for (TFieldIterator<UArrayProperty> ArrayPropIt(NativeSuperClass); ArrayPropIt; ++ArrayPropIt)
+		{
+			UArrayProperty* ArrayProp = *ArrayPropIt;
+
+			UObjectProperty* ObjInnerProp = Cast<UObjectProperty>(ArrayProp->Inner);
+			if ((ObjInnerProp == nullptr) || !ComponentClass->IsChildOf(ObjInnerProp->PropertyClass))
+			{
+				continue;
+			}
+
+			uint8* BpArrayPtr = ArrayProp->ContainerPtrToValuePtr<uint8>(ActorCDO);
+			FScriptArrayHelper BpArrayHelper(ArrayProp, BpArrayPtr);
+			// iterate backwards so we can remove as we go
+			for (int32 ArrayIndex = BpArrayHelper.Num()-1; ArrayIndex >= 0; --ArrayIndex)
+			{
+				uint8* BpEntryPtr = BpArrayHelper.GetRawPtr(ArrayIndex);
+				UObject* ObjEntryValue = ObjInnerProp->GetObjectPropertyValue(BpEntryPtr);
+
+				if (ObjEntryValue == Component)
+				{
+					// NOTE: until we fixup UE-15224, then this may be undesirably diverging from the natively defined 
+					//       array (think delta serialization); however, I think from Blueprint creation on we treat 
+					//       instanced sub-object arrays as differing (just may be confusing to the user)
+					BpArrayHelper.RemoveValues(ArrayIndex);
+				}
+			}
+		}
+
+		// @TODO: have to also remove from map properties now that they're available
+	}
+
+	// 
+	for (TFieldIterator<UObjectProperty> ObjPropIt(NativeSuperClass); ObjPropIt; ++ObjPropIt)
+	{
+		UObjectProperty* ObjectProp = *ObjPropIt;
+		UObject* PropObjValue = ObjectProp->GetObjectPropertyValue_InContainer(ActorCDO);
+
+		if (DestroyedComponents.Contains(PropObjValue))
+		{
+			UObject* SuperObjValue = ObjectProp->GetObjectPropertyValue_InContainer(NativeCDO);
+			ObjectProp->SetObjectPropertyValue_InContainer(ActorCDO, SuperObjValue);
+		}
+	}
+}
+
+static UObject* ConformComponentsUtils::FindeNativeArchetype(UActorComponent* Component)
+{
+	check(Component->HasAnyFlags(RF_DefaultSubObject));
+
+	UActorComponent* Archetype = Cast<UActorComponent>(Component->GetArchetype());
+	if (Archetype == nullptr)
+	{
+		return nullptr;
+	}
+
+	UObject* ArchetypeOwner = Archetype->GetOuter();
+	UClass* OwnerClass = ArchetypeOwner->GetClass();
+
+	const bool bOwnerIsNative = OwnerClass->HasAnyClassFlags(CLASS_Native);
+	if (bOwnerIsNative)
+	{
+		return Archetype;
+	}
+	if (Archetype == Component)
+	{
+		return nullptr;
+	}
+	return FindeNativeArchetype(Archetype);
+}
+
 /** Tries to make sure that a blueprint is conformed to its native parent, in case any native class flags have changed */
 void FKismetEditorUtilities::ConformBlueprintFlagsAndComponents(UBlueprint* BlueprintObj)
 {
@@ -980,6 +1090,9 @@ void FKismetEditorUtilities::ConformBlueprintFlagsAndComponents(UBlueprint* Blue
 	{
 		SkelClass->ClassFlags |= (ParentClass->ClassFlags & CLASS_ScriptInherit);
 		UObject* SkelCDO = SkelClass->GetDefaultObject();
+		// NOTE: we don't need to call ConformRemovedNativeComponents() for skel
+		//       classes, as they're generated on load (and not saved with stale 
+		//       components)
 		SkelCDO->InstanceSubobjectTemplates();
 	}
 
@@ -987,6 +1100,7 @@ void FKismetEditorUtilities::ConformBlueprintFlagsAndComponents(UBlueprint* Blue
 	{
 		GenClass->ClassFlags |= (ParentClass->ClassFlags & CLASS_ScriptInherit);
 		UObject* GenCDO = GenClass->GetDefaultObject();
+		ConformComponentsUtils::ConformRemovedNativeComponents(GenCDO);
 		GenCDO->InstanceSubobjectTemplates();
 	}
 }
