[UE-15137] Guarding against a null function pointer, when an implemented event-function is deleted from a blueprint's super class.

diff --git a/Engine/Source/Editor/UnrealEd/Private/Kismet2/BlueprintEditorUtils.cpp b/Engine/Source/Editor/UnrealEd/Private/Kismet2/BlueprintEditorUtils.cpp
--- a/Engine/Source/Editor/UnrealEd/Private/Kismet2/BlueprintEditorUtils.cpp
+++ b/Engine/Source/Editor/UnrealEd/Private/Kismet2/BlueprintEditorUtils.cpp
@@ -1989,6 +1989,11 @@ void FBlueprintEditorUtils::MarkBlueprintAsStructurallyModified(UBlueprint* Blue
 		{
 			for (auto SkelClass : SkelClassesToRecompile)
 			{
+				if (SkelClass->HasAnyClassFlags(CLASS_NewerVersionExists))
+				{
+					continue;
+				}
+
 				auto SkelBlueprint = Cast<UBlueprint>(SkelClass->ClassGeneratedBy);
 				if (SkelBlueprint
 					&& SkelBlueprint->Status != BS_BeingCreated
@@ -5413,6 +5418,14 @@ void FBlueprintEditorUtils::ConformImplementedEvents(UBlueprint* Blueprint)
 	check(NULL != Blueprint);
 
 	// Collect all event graph names
+	// @TODO: currently, in the compile process, Blueprint->EventGraphs gets 
+	//        filled out (in FKismetCompilerContext::CreateFunctionStubForEvent)
+	//        after this function is called... normally this would cause a 
+	//        problem here (as we're relying on a stale set from the last time  
+	//        this Blueprint was compiled), but because both the skeleton class 
+	//        and generated class execute this and CreateFunctionStubForEvent(), 
+	//        then the second time through (for the generated class) EventGraphs 
+	//        should be accurate
 	TArray<FName> EventGraphNames;
 	for(int EventGraphIndex = 0; EventGraphIndex < Blueprint->EventGraphs.Num(); ++EventGraphIndex)
 	{
@@ -5452,12 +5465,16 @@ void FBlueprintEditorUtils::ConformImplementedEvents(UBlueprint* Blueprint)
 					const bool bEventNodeUsedByInterface = ImplementedInterfaceClasses.Contains(EventNode->EventReference.GetMemberParentClass(EventNode->GetBlueprintClassFromNode()));
 					if (Blueprint->GeneratedClass && !bEventNodeUsedByInterface)
 					{
+						FMemberReference& FuncRef = EventNode->EventReference;
+						const FName EventFuncName = FuncRef.GetMemberName();
+
 						// See if the generated class implements an event with the given function signature
 						const UFunction* TargetFunction = EventNode->EventReference.ResolveMember<UFunction>(Blueprint->GeneratedClass);
-						if (TargetFunction || EventGraphNames.Contains(EventNode->EventReference.GetMemberName()))
+						if (TargetFunction || EventGraphNames.Contains(EventFuncName))
 						{
+							UClass* FuncOwnerClass = FuncRef.GetMemberParentClass(EventNode->GetBlueprintClassFromNode());
 							// The generated class implements the event but the function signature is not up-to-date
-							if (!Blueprint->GeneratedClass->IsChildOf(EventNode->EventReference.GetMemberParentClass(EventNode->GetBlueprintClassFromNode())))
+							if (!Blueprint->GeneratedClass->IsChildOf(FuncOwnerClass))
 							{
 								FFormatNamedArguments Args;
 								Args.Add(TEXT("NodeTitle"), EventNode->GetNodeTitle(ENodeTitleType::ListView));
@@ -5467,7 +5484,14 @@ void FBlueprintEditorUtils::ConformImplementedEvents(UBlueprint* Blueprint)
 								Blueprint->Message_Note( FText::Format(LOCTEXT("EventSignatureFixed_Note", "{NodeTitle} ({EventNodeName}) had an invalid function signature - it has now been fixed."), Args).ToString() );
 
 								// Fix up the event signature
-								EventNode->EventReference.SetFromField<UFunction>(TargetFunction, false);
+								if (TargetFunction != nullptr)
+								{
+									EventNode->EventReference.SetFromField<UFunction>(TargetFunction, false);
+								}
+								else
+								{
+									EventNode->EventReference.SetExternalMember(EventFuncName, Blueprint->GeneratedClass);
+								}
 							}
 						}
 						else
