UE-15206 - Fix for crash

diff --git a/Engine/Source/Editor/UnrealEd/Private/Kismet2/BlueprintEditorUtils.cpp b/Engine/Source/Editor/UnrealEd/Private/Kismet2/BlueprintEditorUtils.cpp
--- a/Engine/Source/Editor/UnrealEd/Private/Kismet2/BlueprintEditorUtils.cpp
+++ b/Engine/Source/Editor/UnrealEd/Private/Kismet2/BlueprintEditorUtils.cpp
@@ -2049,7 +2049,10 @@ void FBlueprintEditorUtils::MarkBlueprintAsStructurallyModified(UBlueprint* Blue
 		TArray<UClass*> ChildrenOfClass;
 		if (UClass* SkelClass = Blueprint->SkeletonGeneratedClass)
 		{
-			GetDerivedClasses(SkelClass, ChildrenOfClass, false);
+			if (!Blueprint->bIsRegeneratingOnLoad)
+			{
+				GetDerivedClasses(SkelClass, ChildrenOfClass, false);
+			}
 		}
 
 		{
