[UE-13320] Transactin buffer losing track of native components when blueprint is reinstanced

diff --git a/Engine/Source/Editor/UnrealEd/Classes/Editor/Transactor.h b/Engine/Source/Editor/UnrealEd/Classes/Editor/Transactor.h
--- a/Engine/Source/Editor/UnrealEd/Classes/Editor/Transactor.h
+++ b/Engine/Source/Editor/UnrealEd/Classes/Editor/Transactor.h
@@ -32,7 +32,22 @@ struct FUndoSessionContext
 };
 
 
-
+inline bool OuterIsCDO(const UObject* Obj, UObject*& OutCDO)
+{
+	bool bOuterIsCDO = false;
+	UObject* Iter = Obj->GetOuter();
+	while (Iter)
+	{
+		if (Iter->HasAllFlags(RF_ClassDefaultObject))
+		{
+			bOuterIsCDO = true;
+			OutCDO = Iter;
+			break;
+		}
+		Iter = Iter->GetOuter();
+	}
+	return bOuterIsCDO;
+}
 
 /*-----------------------------------------------------------------------------
 	FTransaction.
@@ -65,28 +80,76 @@ protected:
 	{
 	public:
 
+		/** 
+			This type is necessary because the blueprint system is destroying and creating 
+			CDOs at edit time (usually on compile, but also on load), but also stores user 
+			entered data in the CDO. We "need"  changes to a CDO to persist across instances 
+			because as we undo and redo we  need to apply changes to different instances of 
+			the CDO - alternatively we could destroy and create the CDO as part of a transaction 
+			(this alternative is the reason for the bunny ears around need).
+
+			DanO: My long term preference is for the editor to use a dynamic, mutable type
+			(rather than the CDO) to store editor data. The CDO can then be re-instanced (or not)
+			as runtime code requires.
+		*/
 		struct FPersistentObjectRef
 		{
 		private:
 			UObject* Object;
 			UClass* SourceCDO;
+			FName SubObjectId;
+
 		public:
 			FPersistentObjectRef()
-				: Object(NULL)
-				, SourceCDO(NULL)
+				: Object(nullptr)
+				, SourceCDO(nullptr)
+				, SubObjectId()
 			{}
 
 			FPersistentObjectRef(UObject* InObject)
 			{
 				const bool bIsCDO = InObject && InObject->HasAllFlags(RF_ClassDefaultObject);
-				Object = bIsCDO ? NULL : InObject;
-				SourceCDO = bIsCDO ? InObject->GetClass() : NULL;
+				UObject* CDO = nullptr;
+				const bool bIsSubobjectOfCDO = OuterIsCDO(InObject, CDO);
+				if (bIsCDO )
+				{
+					Object = nullptr;
+					SourceCDO = InObject->GetClass();
+				}
+				else if (bIsSubobjectOfCDO)
+				{
+					Object = nullptr;
+					SubObjectId = InObject->GetFName();
+					SourceCDO = CDO->GetClass();
+				}
+				else
+				{
+					Object = InObject;
+					SourceCDO = nullptr;
+				}
 			}
 
 			UObject* Get() const
 			{
 				checkSlow(!SourceCDO || !Object);
-				return SourceCDO ? SourceCDO->GetDefaultObject(false) : Object;
+				if (SourceCDO)
+				{
+					if (SubObjectId != FName())
+					{
+						// find the subobject:
+						return SourceCDO->GetDefaultSubobjectByName(SubObjectId);
+					}
+					else
+					{
+						return SourceCDO->GetDefaultObject(false);
+					}
+				}
+				return Object;
+			}
+
+			bool ShouldAddReference() const
+			{
+				return !(SourceCDO != nullptr && SubObjectId != FName());
 			}
 
 			UObject* operator->() const
diff --git a/Engine/Source/Editor/UnrealEd/Private/EditorTransaction.cpp b/Engine/Source/Editor/UnrealEd/Private/EditorTransaction.cpp
--- a/Engine/Source/Editor/UnrealEd/Private/EditorTransaction.cpp
+++ b/Engine/Source/Editor/UnrealEd/Private/EditorTransaction.cpp
@@ -167,7 +167,13 @@ FArchive& operator<<( FArchive& Ar, FTransaction::FObjectRecord& R )
 FTransaction::FObjectRecord::FReferencedObject::FReferencedObject(UObject* InObject)
 {
 	UActorComponent* Component = Cast<UActorComponent>(InObject);
-	if (Component && Component->IsCreatedByConstructionScript())
+	UObject* CDO = nullptr;
+	if (Component && OuterIsCDO(Component, CDO))
+	{
+		Object = CDO;
+		ComponentName = Component->GetFName();
+	}
+	else if (Component && Component->IsCreatedByConstructionScript())
 	{
 		Object = Component->GetOuter();
 		ComponentName = Component->GetFName();
@@ -190,9 +196,12 @@ void FTransaction::FObjectRecord::FReferencedObject::AddReferencedObjects(FRefer
 
 void FTransaction::FObjectRecord::AddReferencedObjects( FReferenceCollector& Collector )
 {
-	UObject* Obj = Object.Get();
-	Collector.AddReferencedObject(Obj);
-	Object = Obj;
+	if (Object.ShouldAddReference())
+	{
+		UObject* Obj = Object.Get();
+		Collector.AddReferencedObject(Obj);
+		Object = Obj;
+	}
 	for( FReferencedObject& ReferencedObject : ReferencedObjects )
 	{
 		ReferencedObject.AddReferencedObjects(Collector);
