UE-7584 Blueprint Break on Exceptions does not work with Accessed None errors

config bCanSuppressAccessViolation added in Engine.ini, "false" by default.

diff --git a/Engine/Config/BaseEngine.ini b/Engine/Config/BaseEngine.ini
--- a/Engine/Config/BaseEngine.ini
+++ b/Engine/Config/BaseEngine.ini
@@ -1564,6 +1564,7 @@ bEnableInheritableComponents=true
 bDeferDependencyLoads=true
 bForceDisableCookedDependencyDeferring=false
 bExecutionAfterReturn=false
+bCanSuppressAccessViolation=false
 
 [/Script/Engine.Blueprint]
 bRecompileOnLoad=true
diff --git a/Engine/Source/Editor/KismetCompiler/Private/KismetCompilerVMBackend.cpp b/Engine/Source/Editor/KismetCompiler/Private/KismetCompilerVMBackend.cpp
--- a/Engine/Source/Editor/KismetCompiler/Private/KismetCompilerVMBackend.cpp
+++ b/Engine/Source/Editor/KismetCompiler/Private/KismetCompilerVMBackend.cpp
@@ -253,7 +253,8 @@ protected:
 		{
 			bInContext = true;
 
-			if (bUnsafeToSkip)
+			static const FBoolConfigValueHelper CanSuppressAccessViolation(TEXT("Kismet"), TEXT("bCanSuppressAccessViolation"), GEngineIni);
+			if (bUnsafeToSkip || !CanSuppressAccessViolation)
 			{
 				Writer << EX_Context;
 			}
diff --git a/Engine/Source/Editor/UnrealEd/Private/Kismet2/KismetDebugUtilities.cpp b/Engine/Source/Editor/UnrealEd/Private/Kismet2/KismetDebugUtilities.cpp
--- a/Engine/Source/Editor/UnrealEd/Private/Kismet2/KismetDebugUtilities.cpp
+++ b/Engine/Source/Editor/UnrealEd/Private/Kismet2/KismetDebugUtilities.cpp
@@ -126,6 +126,8 @@ void FKismetDebugUtilities::OnScriptException(const UObject* ActiveObject, const
 					->AddToken(FTextToken::Create(FText::FromString(MsgInBlueprintStr)))
 					->AddToken(FUObjectToken::Create(BlueprintObj, FText::FromString(BlueprintObj->GetName()))->OnMessageTokenActivated(FOnMessageTokenActivated::CreateStatic(&Local::OnMessageLogLinkActivated)));
 			}
+			bForceToCurrentObject = true;
+			bShouldBreakExecution = GetDefault<UEditorExperimentalSettings>()->bBreakOnExceptions;
 			break;
 		case EBlueprintExceptionType::InfiniteLoop:
 			bForceToCurrentObject = true;
