Made fix for material editor Cleanup button removing Landscape Grass Output nodes support other custom output node types (UE-14694)

diff --git a/Engine/Source/Editor/UnrealEd/Private/MaterialGraph.cpp b/Engine/Source/Editor/UnrealEd/Private/MaterialGraph.cpp
--- a/Engine/Source/Editor/UnrealEd/Private/MaterialGraph.cpp
+++ b/Engine/Source/Editor/UnrealEd/Private/MaterialGraph.cpp
@@ -7,7 +7,7 @@
 
 #include "Materials/MaterialExpressionComment.h"
 #include "Materials/MaterialExpressionFunctionOutput.h"
-#include "Materials/MaterialExpressionLandscapeGrassOutput.h"
+#include "Materials/MaterialExpressionCustomOutput.h"
 
 #include "GraphEditor.h"
 #include "MaterialShared.h"
@@ -334,8 +334,8 @@ void UMaterialGraph::GetUnusedExpressions(TArray<UEdGraphNode*>& UnusedNodes) co
 			UMaterialGraphNode* GraphNode = Cast<UMaterialGraphNode>(Nodes[Index]);
 			if (GraphNode)
 			{
-				UMaterialExpressionLandscapeGrassOutput* Output = Cast<UMaterialExpressionLandscapeGrassOutput>(GraphNode->MaterialExpression);
-				if (Output)
+				UMaterialExpressionCustomOutput* CustomOutput = Cast<UMaterialExpressionCustomOutput>(GraphNode->MaterialExpression);
+				if (CustomOutput)
 				{
 					NodesToCheck.Push(GraphNode);
 				}
