Undo deprecating Translation Editor experimental setting, just causes build warnings in our CIS

diff --git a/Engine/Source/Editor/UnrealEd/Classes/Settings/EditorExperimentalSettings.h b/Engine/Source/Editor/UnrealEd/Classes/Settings/EditorExperimentalSettings.h
--- a/Engine/Source/Editor/UnrealEd/Classes/Settings/EditorExperimentalSettings.h
+++ b/Engine/Source/Editor/UnrealEd/Classes/Settings/EditorExperimentalSettings.h
@@ -30,7 +30,6 @@ public:
 	bool bEnableLocalizationDashboard;
 
 	/** (DEPRECATED) Allows usage of the Translation Editor */
-	DEPRECATED(4.8, "The Translation Editor menu will be removed. The Translation Editor is now accessed from the Localization Dashboard.")
 	UPROPERTY(EditAnywhere, config, Category = Tools, meta = (DisplayName = "Translation Editor (Menu deprecated, please access from Localization Dashboard)"))
 	bool bEnableTranslationEditor;
 
