Moved MaterialEditor promotion tests into that module.
Removed MinimalAPI from UMaterialGraph and exposed the whole class.

diff --git a/Engine/Source/Editor/MaterialEditor/Private/Tests/MaterialEditorTests.cpp b/Engine/Source/Editor/MaterialEditor/Private/Tests/MaterialEditorTests.cpp
new file mode 100644
--- /dev/null
+++ b/Engine/Source/Editor/MaterialEditor/Private/Tests/MaterialEditorTests.cpp
@@ -0,0 +1,244 @@
+// Copyright 1998-2015 Epic Games, Inc. All Rights Reserved.
+
+#include "MaterialEditorModule.h"
+
+//Automation
+#include "Tests/AutomationTestSettings.h"
+#include "AutomationEditorCommon.h"
+#include "AutomationEditorPromotionCommon.h"
+
+//Assets
+#include "AssetSelection.h"
+
+//Materials
+#include "AssetEditorManager.h"
+#include "AssetEditorToolkit.h"
+#include "IMaterialEditor.h"
+#include "Private/MaterialEditor.h"
+#include "Materials/MaterialExpressionTextureSample.h"
+#include "Materials/MaterialExpressionConstant3Vector.h"
+
+#define LOCTEXT_NAMESPACE "EditorMaterialEditorPromotionTests"
+
+DEFINE_LOG_CATEGORY_STATIC(LogEditorMaterialEditorPromotionTests, Log, All);
+
+/**
+* Helper functions used by the build promotion automation test
+*/
+namespace MaterialEditorPromotionTestUtils
+{
+	/** 
+	* Gets saved settings for Material Editor promotion 
+	*/
+	static FMaterialEditorPromotionSettings TestSettings()
+	{
+		UAutomationTestSettings const* AutomationTestSettings = GetDefault<UAutomationTestSettings>();
+		check(AutomationTestSettings);
+
+		return AutomationTestSettings->MaterialEditorPromotionTest;
+	}
+
+
+	/**
+	* Assigns a normal map to a material
+	*
+	* @param InNormalTexture - The texture to use as the normal map
+	* @param InMaterial - The material to modify
+	*/
+	static bool AssignNormalToMaterial(UTexture* InNormalTexture, UMaterial* InMaterial)
+	{
+		IAssetEditorInstance* OpenEditor = FAssetEditorManager::Get().FindEditorForAsset(InMaterial, true);
+		IMaterialEditor* CurrentMaterialEditor = (IMaterialEditor*)OpenEditor;
+
+		//Create the texture sample and auto assign the normal texture
+		UMaterialExpressionTextureSample* NewTextureSampleExpression = Cast<UMaterialExpressionTextureSample>(CurrentMaterialEditor->CreateNewMaterialExpression(UMaterialExpressionTextureSample::StaticClass(), FVector2D(100.f, 200.f), true, true));
+		if (NewTextureSampleExpression)
+		{
+			UE_LOG(LogEditorMaterialEditorPromotionTests, Display, TEXT("Created a new texture sample expression"));
+			NewTextureSampleExpression->Texture = InNormalTexture;
+			NewTextureSampleExpression->SamplerType = EMaterialSamplerType::SAMPLERTYPE_Normal;
+			UE_LOG(LogEditorMaterialEditorPromotionTests, Display, TEXT("Assigned the normal map texture to the new node"));
+			UMaterial* EditorMaterial = Cast<UMaterial>(CurrentMaterialEditor->GetMaterialInterface());
+			UMaterialGraph* MaterialGraph = EditorMaterial->MaterialGraph;
+			EditorMaterial->Normal.Connect(0, NewTextureSampleExpression);
+			UE_LOG(LogEditorMaterialEditorPromotionTests, Display, TEXT("Connected the new node to the normal pin"));
+			MaterialGraph->LinkGraphNodesFromMaterial();
+
+			return true;
+		}
+		else
+		{
+			UE_LOG(LogEditorMaterialEditorPromotionTests, Error, TEXT("Could not add a texture sample to %s"), *InMaterial->GetName());
+			return false;
+		}
+	}
+
+	template<class AssetType> 
+	static AssetType* GetAssetFromPackagePath(const FString& PackagePath)
+	{
+		FAssetData AssetData = FEditorAutomationTestUtilities::GetAssetDataFromPackagePath(PackagePath);
+		return Cast<AssetType>(AssetData.GetAsset());
+	}
+
+}
+
+
+/**
+* Container for items related to the create asset test
+*/
+namespace MaterialEditorPromotionTestHelper
+{
+	struct FMaterialEditorPromotionTest
+	{
+		/** Pointer to the execution info of this test */
+		FAutomationTestExecutionInfo* TestExecutionInfo;
+
+		/** Function definition for the test stage functions */
+		typedef bool(MaterialEditorPromotionTestHelper::FMaterialEditorPromotionTest::*TestStageFn)();
+
+		/** List of test stage functions */
+		TArray<TestStageFn> TestStages;
+
+		/** The index of the test stage we are on */
+		int32 CurrentStage;
+
+		// Test assets
+		UTexture* DiffuseTexture;
+		UTexture* NormalTexture;
+
+		/** Material created from the "Creating a Material" stage */
+		UMaterial* CreatedMaterial;
+
+		/**
+		* Constructor
+		*/
+		FMaterialEditorPromotionTest(FAutomationTestExecutionInfo* InExecutionInfo) :
+			CurrentStage(0)
+		{
+			FMemory::Memzero(this, sizeof(FMaterialEditorPromotionTest));
+
+			TestExecutionInfo = InExecutionInfo;
+
+			//Add test stage functions to array in order of execution
+			TestStages.Add(&MaterialEditorPromotionTestHelper::FMaterialEditorPromotionTest::ContentBrowser_CreateAMaterial_Part1);
+			TestStages.Add(&MaterialEditorPromotionTestHelper::FMaterialEditorPromotionTest::ContentBrowser_CreateAMaterial_Part2);
+		}
+		/**
+		* Runs the current test stage
+		*/
+		bool Update()
+		{
+			bool bStageComplete = (this->*TestStages[CurrentStage])();
+			if (bStageComplete)
+			{
+				CurrentStage++;
+			}
+			return CurrentStage >= TestStages.Num();
+		}
+
+	private:
+
+		/**
+		* ContentBrowser Test Stage: Creating a material (Part 1)
+		*    Creates a material from the diffuse provided in the AutomationTestSettings and opens the material editor
+		*/
+		bool ContentBrowser_CreateAMaterial_Part1()
+		{
+			FString DiffuseTexturePackagePath = MaterialEditorPromotionTestUtils::TestSettings().DefaultDiffuseTexture.FilePath;
+			if (!(DiffuseTexturePackagePath.Len() > 0))
+			{
+				UE_LOG(LogEditorMaterialEditorPromotionTests, Error, TEXT("Failed to load diffuse texture."));
+				return false;
+			}
+			DiffuseTexture = MaterialEditorPromotionTestUtils::GetAssetFromPackagePath<UTexture>(DiffuseTexturePackagePath);
+
+			if (DiffuseTexture)
+			{
+				// @TODO: Break out TEST 1: Create material from texture
+				CreatedMaterial = FEditorPromotionTestUtilities::CreateMaterialFromTexture(DiffuseTexture);
+				if (CreatedMaterial)
+				{
+					UE_LOG(LogEditorMaterialEditorPromotionTests, Display, TEXT("Created new material (%s) from diffuse (%s)"), *CreatedMaterial->GetName(), *DiffuseTexture->GetName());
+					// @TODO: Break out TEST 2: Open a material
+					FAssetEditorManager::Get().OpenEditorForAsset(CreatedMaterial);
+					UE_LOG(LogEditorMaterialEditorPromotionTests, Display, TEXT("Opened the material editor"));
+				}
+				else
+				{
+					UE_LOG(LogEditorMaterialEditorPromotionTests, Error, TEXT("Failed to create material fom diffuse"));
+				}
+			}
+			else
+			{
+				UE_LOG(LogEditorMaterialEditorPromotionTests, Warning, TEXT("Skipping Create Material: DiffuseTexture is missing"));
+			}
+
+			return true;
+		}
+
+		/**
+		* ContentBrowser Test Stage: Creating a material (Part 2)
+		*    Adds the normal map texture to the material and updates the shader
+		*/
+		bool ContentBrowser_CreateAMaterial_Part2()
+		{ 
+			if (CreatedMaterial)
+			{
+				FString NormalTexturePackagePath = MaterialEditorPromotionTestUtils::TestSettings().DefaultNormalTexture.FilePath;
+				if (!(NormalTexturePackagePath.Len() > 0))
+				{
+					UE_LOG(LogEditorMaterialEditorPromotionTests, Error, TEXT("Failed to load normal texture."));
+					return false;
+				}
+				NormalTexture = MaterialEditorPromotionTestUtils::GetAssetFromPackagePath<UTexture>(NormalTexturePackagePath);
+
+				if (NormalTexture)
+				{
+					IAssetEditorInstance* AssetEditor = FAssetEditorManager::Get().FindEditorForAsset(CreatedMaterial, true);
+					FMaterialEditor* MaterialEditor = (FMaterialEditor*)AssetEditor;
+
+					// @TODO: Break out TEST 3: Add normal map to a material
+					if (MaterialEditorPromotionTestUtils::AssignNormalToMaterial(NormalTexture, CreatedMaterial))
+					{
+						FAssetEditorManager::Get().FindEditorForAsset(CreatedMaterial, true);
+						// @TODO: Break out TEST 4: Compile a material
+						MaterialEditor->UpdateMaterialAfterGraphChange();
+						MaterialEditor->bMaterialDirty = false; // Remove dirty flag so window closes without prompt
+						UE_LOG(LogEditorMaterialEditorPromotionTests, Display, TEXT("Compiled the new material"));
+						MaterialEditor->CloseWindow();
+
+					}
+				}
+				else
+				{
+					UE_LOG(LogEditorMaterialEditorPromotionTests, Warning, TEXT("No normal texture to apply to new material"));
+				}
+			}
+
+			return true;
+		}
+	};
+}
+
+
+/**
+* Latent command to run the main build promotion test
+*/
+DEFINE_LATENT_AUTOMATION_COMMAND_ONE_PARAMETER(FRunMaterialEditorPromotionTestCommand, TSharedPtr<MaterialEditorPromotionTestHelper::FMaterialEditorPromotionTest>, MaterialEditorPromotionTest);
+bool FRunMaterialEditorPromotionTestCommand::Update()
+{
+	return MaterialEditorPromotionTest->Update();
+}
+
+
+/**
+* Automation test that handles the build promotion process
+*/
+IMPLEMENT_SIMPLE_AUTOMATION_TEST(FMaterialEditorPromotionTest, "System.Promotion.Editor.Material Editor", EAutomationTestFlags::ATF_Editor);
+
+bool FMaterialEditorPromotionTest::RunTest(const FString& Parameters)
+{
+	TSharedPtr<MaterialEditorPromotionTestHelper::FMaterialEditorPromotionTest> MaterialEditorPromotionTest = MakeShareable(new MaterialEditorPromotionTestHelper::FMaterialEditorPromotionTest(&ExecutionInfo));
+	ADD_LATENT_AUTOMATION_COMMAND(FRunMaterialEditorPromotionTestCommand(MaterialEditorPromotionTest));
+	return true;
+}
\ No newline at end of file
diff --git a/Engine/Source/Editor/UnrealEd/Classes/MaterialGraph/MaterialGraph.h b/Engine/Source/Editor/UnrealEd/Classes/MaterialGraph/MaterialGraph.h
--- a/Engine/Source/Editor/UnrealEd/Classes/MaterialGraph/MaterialGraph.h
+++ b/Engine/Source/Editor/UnrealEd/Classes/MaterialGraph/MaterialGraph.h
@@ -66,8 +66,8 @@ struct FMaterialInputInfo
 	}
 };
 
-UCLASS(MinimalAPI)
-class UMaterialGraph : public UEdGraph
+UCLASS()
+class UNREALED_API UMaterialGraph : public UEdGraph
 {
 	GENERATED_UCLASS_BODY()
 
@@ -103,7 +103,7 @@ public:
 	/**
 	 * Completely rebuild the graph from the material, removing all old nodes
 	 */
-	UNREALED_API void RebuildGraph();
+	void RebuildGraph();
 
 	/**
 	 * Add an Expression to the Graph
@@ -112,7 +112,7 @@ public:
 	 *
 	 * @return	UMaterialGraphNode*	Newly created Graph node to represent expression
 	 */
-	UNREALED_API class UMaterialGraphNode*			AddExpression(UMaterialExpression* Expression);
+	class UMaterialGraphNode*			AddExpression(UMaterialExpression* Expression);
 
 	/**
 	 * Add a Comment to the Graph
@@ -121,27 +121,27 @@ public:
 	 *
 	 * @return	UMaterialGraphNode_Comment*	Newly created Graph node to represent comment
 	 */
-	UNREALED_API class UMaterialGraphNode_Comment*	AddComment(UMaterialExpressionComment* Comment);
+	class UMaterialGraphNode_Comment*	AddComment(UMaterialExpressionComment* Comment);
 
 	/** Link all of the Graph nodes using the Material's connections */
 	void LinkGraphNodesFromMaterial();
 
 	/** Link the Material using the Graph node's connections */
-	UNREALED_API void LinkMaterialExpressionsFromGraph() const;
+	void LinkMaterialExpressionsFromGraph() const;
 
 	/**
 	 * Check whether a material input should be marked as active
 	 *
 	 * @param	GraphPin	Pin representing the material input
 	 */
-	UNREALED_API bool IsInputActive(class UEdGraphPin* GraphPin) const;
+	bool IsInputActive(class UEdGraphPin* GraphPin) const;
 
 	/**
 	 * Get a list of nodes representing expressions that are not used in the Material
 	 *
 	 * @param	UnusedNodes	Array to contain nodes representing unused expressions
 	 */
-	UNREALED_API void GetUnusedExpressions(TArray<class UEdGraphNode*>& UnusedNodes) const;
+	void GetUnusedExpressions(TArray<class UEdGraphNode*>& UnusedNodes) const;
 
 private:
 	/**
diff --git a/Engine/Source/Editor/UnrealEd/Private/Tests/AutomationEditorCommon.cpp b/Engine/Source/Editor/UnrealEd/Private/Tests/AutomationEditorCommon.cpp
--- a/Engine/Source/Editor/UnrealEd/Private/Tests/AutomationEditorCommon.cpp
+++ b/Engine/Source/Editor/UnrealEd/Private/Tests/AutomationEditorCommon.cpp
@@ -68,36 +68,6 @@ DEFINE_LOG_CATEGORY_STATIC(LogAutomationEditorCommon, Log, All);
 
 namespace AutomationEditorCommonUtils
 {
-	/**
-	* Converts a package path to an asset path
-	*
-	* @param PackagePath - The package path to convert
-	*/
-	FString ConvertPackagePathToAssetPath(const FString& PackagePath)
-	{
-		const FString Filename = FPaths::ConvertRelativePathToFull(PackagePath);
-		FString EngineFileName = Filename;
-		FString GameFileName = Filename;
-		if (FPaths::MakePathRelativeTo(EngineFileName, *FPaths::EngineContentDir()) && !EngineFileName.Contains(TEXT("../")))
-		{
-			const FString ShortName = FPaths::GetBaseFilename(EngineFileName);
-			const FString PathName = FPaths::GetPath(EngineFileName);
-			const FString AssetName = FString::Printf(TEXT("/Engine/%s/%s.%s"), *PathName, *ShortName, *ShortName);
-			return AssetName;
-		}
-		else if (FPaths::MakePathRelativeTo(GameFileName, *FPaths::GameContentDir()) && !GameFileName.Contains(TEXT("../")))
-		{
-			const FString ShortName = FPaths::GetBaseFilename(GameFileName);
-			const FString PathName = FPaths::GetPath(GameFileName);
-			const FString AssetName = FString::Printf(TEXT("/Game/%s/%s.%s"), *PathName, *ShortName, *ShortName);
-			return AssetName;
-		}
-		else
-		{
-			UE_LOG(LogAutomationEditorCommon, Error, TEXT("PackagePath (%s) is invalid for the current project"), *PackagePath);
-			return TEXT("");
-		}
-	}
 
 	/**
 	* Imports an object using a given factory
@@ -784,6 +754,57 @@ bool FWaitToFinishBuildDeployCommand::Update()
 	return false;
 }
 
+//////////////////////////////////////////////////////////////////////
+//Asset Path Commands
+
+/**
+* Converts a package path to an asset path
+*
+* @param PackagePath - The package path to convert
+*/
+FString FEditorAutomationTestUtilities::ConvertPackagePathToAssetPath(const FString& PackagePath)
+{
+	const FString Filename = FPaths::ConvertRelativePathToFull(PackagePath);
+	FString EngineFileName = Filename;
+	FString GameFileName = Filename;
+	if (FPaths::MakePathRelativeTo(EngineFileName, *FPaths::EngineContentDir()) && !EngineFileName.Contains(TEXT("../")))
+	{
+		const FString ShortName = FPaths::GetBaseFilename(EngineFileName);
+		const FString PathName = FPaths::GetPath(EngineFileName);
+		const FString AssetName = FString::Printf(TEXT("/Engine/%s/%s.%s"), *PathName, *ShortName, *ShortName);
+		return AssetName;
+	}
+	else if (FPaths::MakePathRelativeTo(GameFileName, *FPaths::GameContentDir()) && !GameFileName.Contains(TEXT("../")))
+	{
+		const FString ShortName = FPaths::GetBaseFilename(GameFileName);
+		const FString PathName = FPaths::GetPath(GameFileName);
+		const FString AssetName = FString::Printf(TEXT("/Game/%s/%s.%s"), *PathName, *ShortName, *ShortName);
+		return AssetName;
+	}
+	else
+	{
+		UE_LOG(LogAutomationEditorCommon, Error, TEXT("PackagePath (%s) is invalid for the current project"), *PackagePath);
+		return TEXT("");
+	}
+}
+
+/**
+* Gets the asset data from a package path
+*
+* @param PackagePath - The package path used to look up the asset data
+*/
+FAssetData FEditorAutomationTestUtilities::GetAssetDataFromPackagePath(const FString& PackagePath)
+{
+	FString AssetPath = FEditorAutomationTestUtilities::ConvertPackagePathToAssetPath(PackagePath);
+	if (AssetPath.Len() > 0)
+	{
+		IAssetRegistry& AssetRegistry = FModuleManager::LoadModuleChecked<FAssetRegistryModule>(TEXT("AssetRegistry")).Get();
+		return AssetRegistry.GetAssetByObjectPath(*AssetPath);
+	}
+
+	return FAssetData();
+}
+
 //////////////////////////////////////////////////////////////////////
 //Find Asset Commands
 
diff --git a/Engine/Source/Editor/UnrealEd/Private/Tests/AutomationEditorPromotionCommon.cpp b/Engine/Source/Editor/UnrealEd/Private/Tests/AutomationEditorPromotionCommon.cpp
new file mode 100644
--- /dev/null
+++ b/Engine/Source/Editor/UnrealEd/Private/Tests/AutomationEditorPromotionCommon.cpp
@@ -0,0 +1,43 @@
+// Copyright 1998-2015 Epic Games, Inc. All Rights Reserved.
+
+#include "UnrealEd.h"
+#include "AutomationEditorPromotionCommon.h"
+
+/**
+* Gets the base path for this asset
+*/
+FString FEditorPromotionTestUtilities::GetGamePath()
+{
+	return TEXT("/Game/BuildPromotionTest");
+}
+
+/**
+* Creates a material from an existing texture
+*
+* @param InTexture - The texture to use as the diffuse for the new material
+*/
+UMaterial* FEditorPromotionTestUtilities::CreateMaterialFromTexture(UTexture* InTexture)
+{
+	// Create the factory used to generate the asset
+	UMaterialFactoryNew* Factory = NewObject<UMaterialFactoryNew>();
+	Factory->InitialTexture = InTexture;
+
+
+	const FString AssetName = FString::Printf(TEXT("%s_Mat"), *InTexture->GetName());
+	const FString PackageName = FEditorPromotionTestUtilities::GetGamePath() + TEXT("/") + AssetName;
+	UPackage* AssetPackage = CreatePackage(NULL, *PackageName);
+	EObjectFlags Flags = RF_Public | RF_Standalone;
+
+	UObject* CreatedAsset = Factory->FactoryCreateNew(UMaterial::StaticClass(), AssetPackage, FName(*AssetName), Flags, NULL, GWarn);
+
+	if (CreatedAsset)
+	{
+		// Notify the asset registry
+		FAssetRegistryModule::AssetCreated(CreatedAsset);
+
+		// Mark the package dirty...
+		AssetPackage->MarkPackageDirty();
+	}
+
+	return Cast<UMaterial>(CreatedAsset);
+}
\ No newline at end of file
diff --git a/Engine/Source/Editor/UnrealEd/Private/Tests/EditorBuildPromotionTests.cpp b/Engine/Source/Editor/UnrealEd/Private/Tests/EditorBuildPromotionTests.cpp
--- a/Engine/Source/Editor/UnrealEd/Private/Tests/EditorBuildPromotionTests.cpp
+++ b/Engine/Source/Editor/UnrealEd/Private/Tests/EditorBuildPromotionTests.cpp
@@ -6,6 +6,7 @@
 #include "Tests/AutomationTestSettings.h"
 #include "AutomationCommon.h"
 #include "AutomationEditorCommon.h"
+#include "AutomationEditorPromotionCommon.h"
 
 //Assets
 #include "AssetRegistryModule.h"
@@ -83,14 +84,6 @@ namespace EditorBuildPromotionTestUtils
 	static const FString BlueprintNameString = TEXT("BuildPromotionBlueprint");
 	static const FName BlueprintStringVariableName(TEXT("MyStringVariable"));
 
-	/**
-	* Gets the base path for this asset
-	*/
-	static FString GetGamePath()
-	{
-		return TEXT("/Game/BuildPromotionTest");
-	}
-
 	/**
 	* Gets the full path to the folder on disk
 	*/
@@ -260,71 +253,6 @@ namespace EditorBuildPromotionTestUtils
 		GConfig->LoadFile(GEditorPerProjectIni);
 	}
 
-	/**
-	* Creates a material from an existing texture
-	*
-	* @param InTexture - The texture to use as the diffuse for the new material
-	*/
-	static UMaterial* CreateMaterialFromTexture(UTexture* InTexture)
-	{
-		// Create the factory used to generate the asset
-		UMaterialFactoryNew* Factory = NewObject<UMaterialFactoryNew>();
-		Factory->InitialTexture = InTexture;
-
-
-		const FString AssetName = FString::Printf(TEXT("%s_Mat"), *InTexture->GetName());
-		const FString PackageName = GetGamePath() + TEXT("/") + AssetName;
-		UPackage* AssetPackage = CreatePackage(NULL, *PackageName);
-		EObjectFlags Flags = RF_Public | RF_Standalone;
-
-		UObject* CreatedAsset = Factory->FactoryCreateNew(UMaterial::StaticClass(), AssetPackage, FName(*AssetName), Flags, NULL, GWarn);
-
-		if (CreatedAsset)
-		{
-			// Notify the asset registry
-			FAssetRegistryModule::AssetCreated(CreatedAsset);
-
-			// Mark the package dirty...
-			AssetPackage->MarkPackageDirty();
-		}
-
-		return Cast<UMaterial>(CreatedAsset);
-	}
-
-	/**
-	* Assigns a normal map to a material
-	*
-	* @param InNormalTexture - The texture to use as the normal map
-	* @param InMaterial - The material to modify
-	*/
-	static bool AssignNormalToMaterial(UTexture* InNormalTexture, UMaterial* InMaterial)
-	{
-		IAssetEditorInstance* OpenEditor = FAssetEditorManager::Get().FindEditorForAsset(InMaterial, true);
-		IMaterialEditor* CurrentMaterialEditor = (IMaterialEditor*)OpenEditor;
-
-		//Create the texture sample and auto assign the normal texture
-		UMaterialExpressionTextureSample* NewTextureSampleExpression = Cast<UMaterialExpressionTextureSample>(CurrentMaterialEditor->CreateNewMaterialExpression(UMaterialExpressionTextureSample::StaticClass(), FVector2D(100.f, 200.f), true, true));
-		if (NewTextureSampleExpression)
-		{
-			UE_LOG(LogEditorBuildPromotionTests, Display, TEXT("Created a new texture sample expression"));
-			NewTextureSampleExpression->Texture = InNormalTexture;
-			NewTextureSampleExpression->SamplerType = EMaterialSamplerType::SAMPLERTYPE_Normal;
-			UE_LOG(LogEditorBuildPromotionTests, Display, TEXT("Assigned the normal map texture to the new node"));
-			UMaterial* EditorMaterial = Cast<UMaterial>(CurrentMaterialEditor->GetMaterialInterface());
-			UMaterialGraph* MaterialGraph = EditorMaterial->MaterialGraph;
-			EditorMaterial->Normal.Connect(0, NewTextureSampleExpression);
-			UE_LOG(LogEditorBuildPromotionTests, Display, TEXT("Connected the new node to the normal pin"));
-			MaterialGraph->LinkGraphNodesFromMaterial();
-
-			return true;
-		}
-		else
-		{
-			UE_LOG(LogEditorBuildPromotionTests, Error, TEXT("Could not add a texture sample to %s"), *InMaterial->GetName());
-			return false;
-		}
-	}
-
 	/**
 	* Finds a visible widget by type.     SLOW!!!!!
 	*
@@ -1050,29 +978,11 @@ namespace EditorBuildPromotionTestUtils
 	}
 
 	/**
-	* Gets the asset data from a package path
-	*
-	* @param PackagePath - The package path used to look up the asset data
-	*/
-	static FAssetData GetAssetDataFromPackagePath(const FString& PackagePath)
-	{
-		FString AssetPath = AutomationEditorCommonUtils::ConvertPackagePathToAssetPath(PackagePath);
-		if (AssetPath.Len() > 0)
-		{
-			IAssetRegistry& AssetRegistry = FModuleManager::LoadModuleChecked<FAssetRegistryModule>(TEXT("AssetRegistry")).Get();
-			return AssetRegistry.GetAssetByObjectPath(*AssetPath);
-		}
-
-		return FAssetData();
-	}
-
-	/**
-	* Adds a default mesh to the level and applies a material
+	* Adds a default mesh to the level
 	*
 	* @param Location - The location to place the actor
-	* @param Material - The material to apply
 	*/
-	static AActor* AddDefaultMeshToLevelWithMaterial(const FVector& Location, UMaterialInterface* Material)
+	static AStaticMeshActor* AddDefaultMeshToLevel(const FVector& Location)
 	{
 		UAutomationTestSettings const* AutomationTestSettings = GetDefault<UAutomationTestSettings>();
 		check(AutomationTestSettings);
@@ -1081,30 +991,48 @@ namespace EditorBuildPromotionTestUtils
 		FString AssetPackagePath = AutomationTestSettings->BuildPromotionTest.DefaultStaticMeshAsset.FilePath;
 		if (AssetPackagePath.Len() > 0)
 		{
-			FAssetData AssetData = EditorBuildPromotionTestUtils::GetAssetDataFromPackagePath(AssetPackagePath);
+			FAssetData AssetData = FEditorAutomationTestUtilities::GetAssetDataFromPackagePath(AssetPackagePath);
 			UStaticMesh* DefaultMesh = Cast<UStaticMesh>(AssetData.GetAsset());
 			if (DefaultMesh)
 			{
 				AStaticMeshActor* PlacedMesh = Cast<AStaticMeshActor>(FActorFactoryAssetProxy::AddActorForAsset(DefaultMesh));
 				PlacedMesh->SetActorLocation(Location);
 
-				PlacedMesh->GetStaticMeshComponent()->SetMaterial(0, Material);
-
 				return PlacedMesh;
 			}
 			else
 			{
-				UE_LOG(LogEditorBuildPromotionTests, Error, TEXT("DefaultStaticMeshAsset is invalid.  Unable to add material to level"));
+				UE_LOG(LogEditorBuildPromotionTests, Error, TEXT("DefaultStaticMeshAsset is invalid."));
 			}
 		}
 		else
 		{
-			UE_LOG(LogEditorBuildPromotionTests, Warning, TEXT("Can't add material to level because no DefaultMeshAsset is assigned"));
+			UE_LOG(LogEditorBuildPromotionTests, Warning, TEXT("Can't add Static Mesh to level because no DefaultMeshAsset is defined."));
 		}
 
 		return NULL;
 	}
 
+	/*
+	* Applies a material to a static mesh. Triggers a test failure if StaticMesh is not valid.
+	*
+	* @param StaticMesh - the static mesh to apply the material to
+	* @param Material - the material to apply
+	*/
+	static bool ApplyMaterialToStaticMesh(AStaticMeshActor* StaticMesh, UMaterialInterface* Material)
+	{
+		if (StaticMesh)
+		{
+			StaticMesh->GetStaticMeshComponent()->SetMaterial(0, Material);
+			return true;
+		}
+		else
+		{
+			UE_LOG(LogEditorBuildPromotionTests, Error, TEXT("Failed to apply material to static mesh because mesh does not exist"));
+			return false;
+		}
+	}
+
 	/**
 	* Imports an asset using the supplied factory and file
 	*
@@ -1114,7 +1042,7 @@ namespace EditorBuildPromotionTestUtils
 	static UObject* ImportAsset(UFactory* ImportFactory, const FString& ImportPath)
 	{
 		const FString Name = ObjectTools::SanitizeObjectName(FPaths::GetBaseFilename(ImportPath));
-		const FString PackageName = FString::Printf(TEXT("%s/%s"), *GetGamePath(), *Name);
+		const FString PackageName = FString::Printf(TEXT("%s/%s"), *FEditorPromotionTestUtilities::GetGamePath(), *Name);
 
 		UObject* ImportedAsset = AutomationEditorCommonUtils::ImportAssetUsingFactory(ImportFactory,Name,PackageName,ImportPath);
 
@@ -1136,8 +1064,9 @@ namespace EditorBuildPromotionTestUtils
 				}
 				else
 				{
-					UMaterial* NewMaterial = EditorBuildPromotionTestUtils::CreateMaterialFromTexture(TextureObject);
-					PlacedActor = EditorBuildPromotionTestUtils::AddDefaultMeshToLevelWithMaterial(PlaceLocation, NewMaterial);
+					UMaterial* NewMaterial = FEditorPromotionTestUtilities::CreateMaterialFromTexture(TextureObject);
+					PlacedActor = EditorBuildPromotionTestUtils::AddDefaultMeshToLevel(PlaceLocation);
+					EditorBuildPromotionTestUtils::ApplyMaterialToStaticMesh(Cast<AStaticMeshActor>(PlacedActor), NewMaterial);
 				}
 			}
 			else
@@ -1167,7 +1096,7 @@ namespace EditorBuildPromotionTestUtils
 	*/
 	static UObject* CreateAsset(UFactory* CreateFactory, UClass* AssetClass, const FString& AssetName)
 	{
-		FString PackageName = FString::Printf(TEXT("%s/%s"), *GetGamePath(), *AssetName);
+		FString PackageName = FString::Printf(TEXT("%s/%s"), *FEditorPromotionTestUtilities::GetGamePath(), *AssetName);
 		UPackage* AssetPackage = CreatePackage(NULL, *PackageName);
 		EObjectFlags Flags = RF_Public | RF_Standalone;
 
@@ -1311,7 +1240,7 @@ namespace EditorBuildPromotionTestUtils
 		// Form a filter from the paths
 		FARFilter Filter;
 		Filter.bRecursivePaths = true;
-		new (Filter.PackagePaths) FName(*EditorBuildPromotionTestUtils::GetGamePath());
+		new (Filter.PackagePaths) FName(*FEditorPromotionTestUtilities::GetGamePath());
 
 		// Query for a list of assets in the selected paths
 		TArray<FAssetData> AssetList;
@@ -1339,8 +1268,8 @@ namespace EditorBuildPromotionTestUtils
 			}
 		}
 
-		UE_LOG(LogEditorBuildPromotionTests, Display, TEXT("Clearing Path: %s"), *EditorBuildPromotionTestUtils::GetGamePath());
-		AssetRegistry.RemovePath(EditorBuildPromotionTestUtils::GetGamePath());
+		UE_LOG(LogEditorBuildPromotionTests, Display, TEXT("Clearing Path: %s"), *FEditorPromotionTestUtilities::GetGamePath());
+		AssetRegistry.RemovePath(FEditorPromotionTestUtilities::GetGamePath());
 
 		//Remove the directory
 		bool bEnsureExists = false;
@@ -1597,12 +1526,14 @@ namespace BuildPromotionTestHelper
 			UTexture* TextureAsset = Cast<UTexture>(CurrentAsset);
 			if (MaterialAsset)
 			{
-				AActor* PlacedActor = EditorBuildPromotionTestUtils::AddDefaultMeshToLevelWithMaterial(PlaceLocation, MaterialAsset);
+				AStaticMeshActor* PlacedActor = EditorBuildPromotionTestUtils::AddDefaultMeshToLevel(PlaceLocation); // , MaterialAsset);
+				EditorBuildPromotionTestUtils::ApplyMaterialToStaticMesh(PlacedActor, MaterialAsset);
 			}
 			else if (TextureAsset)
 			{
-				UMaterial* NewMaterial = EditorBuildPromotionTestUtils::CreateMaterialFromTexture(TextureAsset);
-				AActor* PlacedActor = EditorBuildPromotionTestUtils::AddDefaultMeshToLevelWithMaterial(PlaceLocation, NewMaterial);
+				UMaterial* NewMaterial = FEditorPromotionTestUtilities::CreateMaterialFromTexture(TextureAsset);
+				AStaticMeshActor* PlacedActor = EditorBuildPromotionTestUtils::AddDefaultMeshToLevel(PlaceLocation); //, NewMaterial);
+				EditorBuildPromotionTestUtils::ApplyMaterialToStaticMesh(PlacedActor, NewMaterial);
 			}
 			else
 			{
@@ -1712,9 +1643,6 @@ namespace BuildPromotionTestHelper
 		UMaterial* SCTestMat;
 		FString ChosenMaterialColor;
 
-		/** Material created from the "Creating a Material" stage */
-		UMaterial* CreatedMaterial;
-
 		/** Particle System created from the "Creating a Particle System" stage */
 		UParticleSystem* CreatedPS;
 
@@ -1806,9 +1734,7 @@ namespace BuildPromotionTestHelper
 			ADD_TEST_STAGE(ContentBrowser_OpenAssets_Part1,				TEXT("Open Asset Types"));
 			ADD_TEST_STAGE(ContentBrowser_OpenAssets_Part2,				TEXT("Open Asset Types"));
 			ADD_TEST_STAGE(ContentBrowser_ReimportAsset,				TEXT("Re-import Assets"));
-			ADD_TEST_STAGE(ContentBrowser_CreateAMaterial_Part1,		TEXT("Creating a Material"));
-			ADD_TEST_STAGE(ContentBrowser_CreateAMaterial_Part2,		TEXT("Creating a Material"));
-			ADD_TEST_STAGE(ContentBrowser_CreateAMaterial_Part3,		TEXT("Creating a Material"));
+			ADD_TEST_STAGE(ContentBrowser_AssignAMaterial,				TEXT("Assigning a Material"));
 			ADD_TEST_STAGE(ContentBrowser_CreateAParticleSystem_Part1,	TEXT("Creating a Particle System"));
 			ADD_TEST_STAGE(ContentBrowser_CreateAParticleSystem_Part2,	TEXT("Creating a Particle System"));
 			ADD_TEST_STAGE(EndSection, TEXT("Content Browser"));
@@ -2355,6 +2281,7 @@ namespace BuildPromotionTestHelper
 			return false;
 		}
 
+//@TODO: Rewrite this without macro if possible
 #define IMPORT_ASSET_WITH_FACTORY(TFactoryClass,TObjectClass,ImportSetting,ClassVariable,ExtraSettings) \
 { \
 	const FString FilePath = AutomationTestSettings->BuildPromotionTest.ImportWorkflow.ImportSetting.ImportFilePath.FilePath; \
@@ -2503,7 +2430,7 @@ namespace BuildPromotionTestHelper
 			SectionSuccessCount--;
 
 			//Save all the new assets
-			EditorBuildPromotionTestUtils::SaveAllAssetsInFolder(EditorBuildPromotionTestUtils::GetGamePath());
+			EditorBuildPromotionTestUtils::SaveAllAssetsInFolder(FEditorPromotionTestUtilities::GetGamePath());
 
 			return true;
 		}
@@ -2522,7 +2449,7 @@ namespace BuildPromotionTestHelper
 			const FString SourceControlMaterialPath = AutomationTestSettings->BuildPromotionTest.SourceControlMaterial.FilePath;
 			if (SourceControlMaterialPath.Len() > 0)
 			{
-				FAssetData MaterialData = EditorBuildPromotionTestUtils::GetAssetDataFromPackagePath(SourceControlMaterialPath);
+				FAssetData MaterialData = FEditorAutomationTestUtilities::GetAssetDataFromPackagePath(SourceControlMaterialPath);
 				SCTestMat = Cast<UMaterial>(MaterialData.GetAsset());
 
 				if (SCTestMat)
@@ -2598,6 +2525,7 @@ namespace BuildPromotionTestHelper
 			return true;
 		}
 
+		// @TODO: Rewrite this to use a lighter-weight asset type
 		/**
 		* ContentBrowser Test Stage: Source Control (part 3)
 		*    Changes the source control material's color
@@ -2643,7 +2571,7 @@ namespace BuildPromotionTestHelper
 					}
 				}
 
-				FAssetData MaterialData = EditorBuildPromotionTestUtils::GetAssetDataFromPackagePath(AutomationTestSettings->BuildPromotionTest.SourceControlMaterial.FilePath);
+				FAssetData MaterialData = FEditorAutomationTestUtilities::GetAssetDataFromPackagePath(AutomationTestSettings->BuildPromotionTest.SourceControlMaterial.FilePath);
 				FString PackageFileName = FPackageName::LongPackageNameToFilename(MaterialData.PackageName.ToString(), FPackageName::GetAssetPackageExtension());
 				FString MaterialFilePath = FPaths::ConvertRelativePathToFull(PackageFileName);
 				AsyncHelper = SourceControlAutomationCommon::FAsyncCommandHelper(MaterialFilePath);
@@ -2734,7 +2662,7 @@ namespace BuildPromotionTestHelper
 			AssetPackagePath = AutomationTestSettings->BuildPromotionTest.OpenAssets.BlueprintAsset.FilePath;
 			if (AssetPackagePath.Len() > 0)
 			{
-				AssetData = EditorBuildPromotionTestUtils::GetAssetDataFromPackagePath(AssetPackagePath);
+				AssetData = FEditorAutomationTestUtilities::GetAssetDataFromPackagePath(AssetPackagePath);
 				Asset = AssetData.GetAsset();
 				if (Asset)
 				{
@@ -2756,7 +2684,7 @@ namespace BuildPromotionTestHelper
 			AssetPackagePath = AutomationTestSettings->BuildPromotionTest.OpenAssets.MaterialAsset.FilePath;
 			if (AssetPackagePath.Len() > 0)
 			{
-				AssetData = EditorBuildPromotionTestUtils::GetAssetDataFromPackagePath(AssetPackagePath);
+				AssetData = FEditorAutomationTestUtilities::GetAssetDataFromPackagePath(AssetPackagePath);
 				Asset = AssetData.GetAsset();
 				if (Asset)
 				{
@@ -2778,7 +2706,7 @@ namespace BuildPromotionTestHelper
 			AssetPackagePath = AutomationTestSettings->BuildPromotionTest.OpenAssets.ParticleSystemAsset.FilePath;
 			if (AssetPackagePath.Len() > 0)
 			{
-				AssetData = EditorBuildPromotionTestUtils::GetAssetDataFromPackagePath(AssetPackagePath);
+				AssetData = FEditorAutomationTestUtilities::GetAssetDataFromPackagePath(AssetPackagePath);
 				Asset = AssetData.GetAsset();
 				if (Asset)
 				{
@@ -2800,7 +2728,7 @@ namespace BuildPromotionTestHelper
 			AssetPackagePath = AutomationTestSettings->BuildPromotionTest.OpenAssets.SkeletalMeshAsset.FilePath;
 			if (AssetPackagePath.Len() > 0)
 			{
-				AssetData = EditorBuildPromotionTestUtils::GetAssetDataFromPackagePath(AssetPackagePath);
+				AssetData = FEditorAutomationTestUtilities::GetAssetDataFromPackagePath(AssetPackagePath);
 				Asset = AssetData.GetAsset();
 				if (Asset)
 				{
@@ -2822,7 +2750,7 @@ namespace BuildPromotionTestHelper
 			AssetPackagePath = AutomationTestSettings->BuildPromotionTest.OpenAssets.StaticMeshAsset.FilePath;
 			if (AssetPackagePath.Len() > 0)
 			{
-				AssetData = EditorBuildPromotionTestUtils::GetAssetDataFromPackagePath(AssetPackagePath);
+				AssetData = FEditorAutomationTestUtilities::GetAssetDataFromPackagePath(AssetPackagePath);
 				Asset = AssetData.GetAsset();
 				if (Asset)
 				{
@@ -2844,7 +2772,7 @@ namespace BuildPromotionTestHelper
 			AssetPackagePath = AutomationTestSettings->BuildPromotionTest.OpenAssets.TextureAsset.FilePath;
 			if (AssetPackagePath.Len() > 0)
 			{
-				AssetData = EditorBuildPromotionTestUtils::GetAssetDataFromPackagePath(AssetPackagePath);
+				AssetData = FEditorAutomationTestUtilities::GetAssetDataFromPackagePath(AssetPackagePath);
 				Asset = AssetData.GetAsset();
 				if (Asset)
 				{
@@ -2915,78 +2843,44 @@ namespace BuildPromotionTestHelper
 		}
 
 		/**
-		* ContentBrowser Test Stage: Creating a material (Part 1)
-		*    Creates a material from the diffuse provided in the AutomationTestSettings and opens the material editor
+		* ContentBrowser Test Stage: Creating a material (Part 3)
+		*    Closes all assets editor and adds the material to a default object in the map
 		*/
-		bool ContentBrowser_CreateAMaterial_Part1()
+		bool ContentBrowser_AssignAMaterial()
 		{
-			if (DiffuseTexture)
-			{
-				CreatedMaterial = EditorBuildPromotionTestUtils::CreateMaterialFromTexture(DiffuseTexture);
-				if (CreatedMaterial)
-				{
-					UE_LOG(LogEditorBuildPromotionTests, Display, TEXT("Created new material (%s) from diffuse (%s)"),*CreatedMaterial->GetName(), *DiffuseTexture->GetName());
-					FAssetEditorManager::Get().OpenEditorForAsset(CreatedMaterial);
-					UE_LOG(LogEditorBuildPromotionTests, Display, TEXT("Opened the material editor"));
-				}
-				else
-				{
-					UE_LOG(LogEditorBuildPromotionTests, Error, TEXT("Failed to create material fom diffuse"));
-				}
-			}
-			else
+			// SETUP
+			FAssetEditorManager::Get().CloseAllAssetEditors();
+			UAutomationTestSettings const* AutomationTestSettings = GetDefault<UAutomationTestSettings>();
+			check(AutomationTestSettings);
+
+			// Load default material asset
+			FString MaterialPackagePath = AutomationTestSettings->MaterialEditorPromotionTest.DefaultMaterialAsset.FilePath;
+			if (!(MaterialPackagePath.Len() > 0))
 			{
-				SkippedTests.Add(TEXT("ContentBrowser: Creating a material. (No diffuse)"));
-				UE_LOG(LogEditorBuildPromotionTests, Warning, TEXT("Skipping Create Material: DiffuseTexture is missing"));
+				UE_LOG(LogEditorBuildPromotionTests, Warning, TEXT("Skipping material assignment test: no default material defined."));
+				return true;
 			}
 
-			return true;
-		}
-
-		/**
-		* ContentBrowser Test Stage: Creating a material (Part 2)
-		*    Adds the normalmap texture to the material and updates the shader
-		*/
-		bool ContentBrowser_CreateAMaterial_Part2()
-		{
-			if (CreatedMaterial)
+			FAssetData MaterialAssetData = FEditorAutomationTestUtilities::GetAssetDataFromPackagePath(MaterialPackagePath);
+			UMaterial* DefaultMaterial = Cast<UMaterial>(MaterialAssetData.GetAsset());
+			if (!DefaultMaterial)
 			{
-				if (NormalTexture)
-				{
-					IAssetEditorInstance* AssetEditor = FAssetEditorManager::Get().FindEditorForAsset(CreatedMaterial, true);
-					FMaterialEditor* MaterialEditor = (FMaterialEditor*)AssetEditor;
-
-					if (EditorBuildPromotionTestUtils::AssignNormalToMaterial(NormalTexture, CreatedMaterial))
-					{
-						FAssetEditorManager::Get().FindEditorForAsset(CreatedMaterial, true);
-						EditorBuildPromotionTestUtils::SendUpdateMaterialCommand();
-						UE_LOG(LogEditorBuildPromotionTests, Display, TEXT("Compiled the new material"));
-					}
-				}
-				else
-				{
-					SkippedTests.Add(TEXT("ContentBrowser: Creating a material. (No normalmap)"));
-					UE_LOG(LogEditorBuildPromotionTests, Warning, TEXT("No normal texture to apply to new material"));
-				}
+				UE_LOG(LogEditorBuildPromotionTests, Error, TEXT("Failed to load default material asset."));
+				return false;
 			}
 
-			return true;
-		}
-
-		/**
-		* ContentBrowser Test Stage: Creating a material (Part 3)
-		*    Closes the material editor and adds the material to a default object in the map
-		*/
-		bool ContentBrowser_CreateAMaterial_Part3()
-		{
-			FAssetEditorManager::Get().CloseAllAssetEditors();
-
+			// Add static mesh to world as material assignment target
 			const FVector PlaceLocation(0, 2240, 166);
-			AActor* PlacedActor = EditorBuildPromotionTestUtils::AddDefaultMeshToLevelWithMaterial(PlaceLocation, CreatedMaterial);
-			if ( PlacedActor )
+			AStaticMeshActor* PlacedStaticMesh = EditorBuildPromotionTestUtils::AddDefaultMeshToLevel(PlaceLocation);
+
+			// RUN TEST
+			// @TODO: Put in a check to verify that the mesh's Material[0] == DefaultMaterial
+			if (EditorBuildPromotionTestUtils::ApplyMaterialToStaticMesh(PlacedStaticMesh, DefaultMaterial))
 			{
-				UE_LOG(LogEditorBuildPromotionTests, Display, TEXT("Placed new material in the level"));
-			}
+				UE_LOG(LogEditorBuildPromotionTests, Display, TEXT("Successfully assigned material to static mesh."));
+			}  // No need to error on false, since ApplyMaterialToStaticMesh triggers its own error if it fails
+
+			// @TODO: TEARDOWN
 
 			return true;
 		}
@@ -3064,7 +2958,7 @@ namespace BuildPromotionTestHelper
 			{
 				BlueprintPackage->SetDirtyFlag(true);
 				BlueprintPackage->FullyLoad();
-				const FString PackagePath = EditorBuildPromotionTestUtils::GetGamePath() + TEXT("/") + EditorBuildPromotionTestUtils::BlueprintNameString;
+				const FString PackagePath = FEditorPromotionTestUtilities::GetGamePath() + TEXT("/") + EditorBuildPromotionTestUtils::BlueprintNameString;
 				if (UPackage::SavePackage(BlueprintPackage, NULL, RF_Standalone, *FPackageName::LongPackageNameToFilename(PackagePath, FPackageName::GetAssetPackageExtension()), GLog, nullptr, false, true, SAVE_None))
 				{
 					UE_LOG(LogEditorBuildPromotionTests, Display, TEXT("Saved blueprint (%s)"), *BlueprintObject->GetName());
@@ -3090,14 +2984,14 @@ namespace BuildPromotionTestHelper
 			const FString FirstMeshPath = AutomationTestSettings->BuildPromotionTest.BlueprintSettings.FirstMeshPath.FilePath;
 			if (FirstMeshPath.Len() > 0)
 			{
-				AssetData = EditorBuildPromotionTestUtils::GetAssetDataFromPackagePath(FirstMeshPath);
+				AssetData = FEditorAutomationTestUtilities::GetAssetDataFromPackagePath(FirstMeshPath);
 				FirstBlueprintMesh = Cast<UStaticMesh>(AssetData.GetAsset());
 			}
 
 			const FString SecondMeshPath = AutomationTestSettings->BuildPromotionTest.BlueprintSettings.SecondMeshPath.FilePath;
 			if (SecondMeshPath.Len() > 0)
 			{
-				AssetData = EditorBuildPromotionTestUtils::GetAssetDataFromPackagePath(SecondMeshPath);
+				AssetData = FEditorAutomationTestUtilities::GetAssetDataFromPackagePath(SecondMeshPath);
 				SecondBlueprintMesh = Cast<UStaticMesh>(AssetData.GetAsset());
 			}
 
@@ -3106,7 +3000,7 @@ namespace BuildPromotionTestHelper
 				UBlueprintFactory* Factory = NewObject<UBlueprintFactory>();
 				Factory->ParentClass = AActor::StaticClass();
 
-				const FString PackageName = EditorBuildPromotionTestUtils::GetGamePath() + TEXT("/") + EditorBuildPromotionTestUtils::BlueprintNameString;
+				const FString PackageName = FEditorPromotionTestUtilities::GetGamePath() + TEXT("/") + EditorBuildPromotionTestUtils::BlueprintNameString;
 				BlueprintPackage = CreatePackage(NULL, *PackageName);
 				EObjectFlags Flags = RF_Public | RF_Standalone;
 
@@ -4189,7 +4083,7 @@ namespace BuildPromotionTestHelper
 			BuildStartErrorCount = TestExecutionInfo->Errors.Num();
 
 			//Save all the new assets
-			EditorBuildPromotionTestUtils::SaveAllAssetsInFolder(EditorBuildPromotionTestUtils::GetGamePath());
+			EditorBuildPromotionTestUtils::SaveAllAssetsInFolder(FEditorPromotionTestUtilities::GetGamePath());
 
 			UE_LOG(LogEditorBuildPromotionTests, Display, TEXT("Turning level visibility off"));
 			bool bShouldBeVisible = false;
@@ -4276,7 +4170,7 @@ namespace BuildPromotionTestHelper
 				FEditorFileUtils::SaveLevel(Level);
 
 				//Save all the new assets again because material usage flags may have changed.
-				EditorBuildPromotionTestUtils::SaveAllAssetsInFolder(EditorBuildPromotionTestUtils::GetGamePath());
+				EditorBuildPromotionTestUtils::SaveAllAssetsInFolder(FEditorPromotionTestUtilities::GetGamePath());
 
 				//Copy the errors over to the LogItems
 				for (int32 i = BuildStartErrorCount; i < TestExecutionInfo->Errors.Num(); ++i)
diff --git a/Engine/Source/Editor/UnrealEd/Public/Tests/AutomationEditorCommon.h b/Engine/Source/Editor/UnrealEd/Public/Tests/AutomationEditorCommon.h
--- a/Engine/Source/Editor/UnrealEd/Public/Tests/AutomationEditorCommon.h
+++ b/Engine/Source/Editor/UnrealEd/Public/Tests/AutomationEditorCommon.h
@@ -46,13 +46,6 @@ namespace AutomationEditorCommonUtils
 		return GEditor->NewMap();
 	}
 
-	/**
-	* Converts a package path to an asset path
-	*
-	* @param PackagePath - The package path to convert
-	*/
-	FString ConvertPackagePathToAssetPath(const FString& PackagePath);
-
 	/**
 	* Imports an object using a given factory
 	*
@@ -242,6 +235,21 @@ DEFINE_LATENT_AUTOMATION_COMMAND_ONE_PARAMETER(FDeleteDirCommand, FString, InFol
 class FEditorAutomationTestUtilities
 {
 public:
+	
+	/**
+	* Converts a package path to an asset path
+	*
+	* @param PackagePath - The package path to convert
+	*/
+	UNREALED_API static FString ConvertPackagePathToAssetPath(const FString& PackagePath);
+
+	/**
+	* Gets the asset data from a package path
+	*
+	* @param PackagePath - The package path used to look up the asset data
+	*/
+	UNREALED_API static FAssetData GetAssetDataFromPackagePath(const FString& PackagePath);
+
 	/**
 	* Loads the map specified by an automation test
 	*
diff --git a/Engine/Source/Editor/UnrealEd/Public/Tests/AutomationEditorPromotionCommon.h b/Engine/Source/Editor/UnrealEd/Public/Tests/AutomationEditorPromotionCommon.h
new file mode 100644
--- /dev/null
+++ b/Engine/Source/Editor/UnrealEd/Public/Tests/AutomationEditorPromotionCommon.h
@@ -0,0 +1,27 @@
+// Copyright 1998-2015 Epic Games, Inc. All Rights Reserved.
+
+#pragma once
+// Automation
+#include "AutomationCommon.h"
+#include "Tests/AutomationTestSettings.h"
+
+//Materials
+
+class FEditorPromotionTestUtilities
+{
+
+public:
+
+	/**
+	* Gets the base path for this asset
+	*/
+	UNREALED_API static FString GetGamePath();
+
+
+	/**
+	* Creates a material from an existing texture
+	*
+	* @param InTexture - The texture to use as the diffuse for the new material
+	*/
+	UNREALED_API static UMaterial* CreateMaterialFromTexture(UTexture* InTexture);
+};
\ No newline at end of file
diff --git a/Engine/Source/Runtime/Engine/Classes/Tests/AutomationTestSettings.h b/Engine/Source/Runtime/Engine/Classes/Tests/AutomationTestSettings.h
--- a/Engine/Source/Runtime/Engine/Classes/Tests/AutomationTestSettings.h
+++ b/Engine/Source/Runtime/Engine/Classes/Tests/AutomationTestSettings.h
@@ -241,6 +241,28 @@ struct FBuildPromotionNewProjectSettings
 	FString NewProjectNameOverride;
 };
 
+/**
+* Holds settings for the material editor build promotion tests
+*/
+USTRUCT()
+struct FMaterialEditorPromotionSettings
+{
+	GENERATED_USTRUCT_BODY()
+
+	/** Default material asset to apply to static meshes **/
+	UPROPERTY(EditAnywhere, Category = Automation)
+	FFilePath DefaultMaterialAsset;
+
+	/** Default material asset to apply to static meshes **/
+	UPROPERTY(EditAnywhere, Category = Automation)
+	FFilePath DefaultDiffuseTexture;
+
+	/** Default material asset to apply to static meshes **/
+	UPROPERTY(EditAnywhere, Category = Automation)
+	FFilePath DefaultNormalTexture;
+
+};
+
 /**
 * Holds settings for the editor build promotion test
 */
@@ -336,6 +358,12 @@ public:
 	UPROPERTY(EditAnywhere, config, Category = Automation)
 	FBuildPromotionTestSettings BuildPromotionTest;
 
+	/**
+	* Material editor promotion test settings
+	*/
+	UPROPERTY(EditAnywhere, config, Category = Automation)
+	FMaterialEditorPromotionSettings MaterialEditorPromotionTest;
+
 	/**
 	* Modules to load that have engine tests
 	*/
