Fixed landscape holes not letting light through
Fixed landscape sometimes reflecting the wrong light colour
- both were due to lightmass assuming only one UV channel and 0-1 UVs (Landscape uses UV#3 for weightmap coords and it isn't 0-1) and incorrectly handling the expanded components (the landscape light mesh is expanded to fix seams in the light/shadow maps of adjacent components)
- the fix for holes doesn't use the material because of the expanded components, so would require rendering 9 materials to one output texture, which was too horrible to make work. Ideally we wouldn't expand the components at all :(

diff --git a/Engine/Source/Editor/UnrealEd/Private/Lightmass/Lightmass.cpp b/Engine/Source/Editor/UnrealEd/Private/Lightmass/Lightmass.cpp
--- a/Engine/Source/Editor/UnrealEd/Private/Lightmass/Lightmass.cpp
+++ b/Engine/Source/Editor/UnrealEd/Private/Lightmass/Lightmass.cpp
@@ -447,10 +447,18 @@ FLightmassExporter::~FLightmassExporter()
 	}
 }
 
-void FLightmassExporter::AddMaterial(UMaterialInterface* InMaterialInterface)
+void FLightmassExporter::AddMaterial(UMaterialInterface* InMaterialInterface, const FStaticLightingMesh* InStaticLightingMesh /*= nullptr*/)
 {
 	if (InMaterialInterface)
 	{
+		FLightmassMaterialExportSettings ExportSettings = { InStaticLightingMesh };
+
+		if (auto* ExistingExportSettings = MaterialExportSettings.Find(InMaterialInterface))
+		{
+			checkf(ExportSettings == *ExistingExportSettings, TEXT("Attempting to add the same material twice with different export settings, this is not (currently) supported"));
+			return;
+		}
+
 		// Check for material texture changes...
 		//@TODO: Add package to warning list if it needs to be resaved (perf warning)
 		InMaterialInterface->UpdateLightmassTextureTracking();
@@ -470,7 +478,8 @@ void FLightmassExporter::AddMaterial(UMaterialInterface* InMaterialInterface)
 			}
 		}
 
-		Materials.AddUnique(InMaterialInterface);
+		Materials.Add(InMaterialInterface);
+		MaterialExportSettings.Add(InMaterialInterface, ExportSettings);
 	}
 }
 
@@ -663,7 +672,8 @@ bool FLightmassExporter::WriteToMaterialChannel(FLightmassStatistics& Stats)
 					}
 					else
 					{
-						ExportMaterial(Materials[CurrentAmortizationIndex]);
+						auto* CurrentMaterial = Materials[CurrentAmortizationIndex];
+						ExportMaterial(CurrentMaterial, MaterialExportSettings.FindChecked(CurrentMaterial));
 						++CurrentAmortizationIndex;
 					}
 				}
@@ -1117,7 +1127,7 @@ void FLightmassExporter::WriteStaticMeshes()
 
 void FLightmassExporter::BuildMaterialMap(UMaterialInterface* Material)
 {
-	if (Material)
+	if (ensure(Material))
 	{
 		FGuid LightingGuid = Material->GetLightingGuid();
 
@@ -1162,12 +1172,12 @@ void FLightmassExporter::BlockOnShaderCompilation()
 	GShaderCompilingManager->FinishAllCompilation();
 }
 
-void FLightmassExporter::ExportMaterial(UMaterialInterface* Material)
+void FLightmassExporter::ExportMaterial(UMaterialInterface* Material, const FLightmassMaterialExportSettings& ExportSettings)
 {
 	FMaterialExportDataEntry* ExportEntry = MaterialExportData.Find(Material);
 
 	// Only create the Swarm channel if there is something to export
-	if (Material && ExportEntry)
+	if (ensure(Material) && ExportEntry)
 	{
 		Lightmass::FBaseMaterialData BaseMaterialData;
 		BaseMaterialData.Guid = Material->GetLightingGuid();
@@ -1177,8 +1187,8 @@ void FLightmassExporter::ExportMaterial(UMaterialInterface* Material)
 		FMemory::Memzero(&MaterialData,sizeof(MaterialData));
 		UMaterial* BaseMaterial = Material->GetMaterial();
 		MaterialData.bTwoSided = (uint32)Material->IsTwoSided();
-		MaterialData.EmissiveBoost =  Material->GetEmissiveBoost();
-		MaterialData.DiffuseBoost =  Material->GetDiffuseBoost();
+		MaterialData.EmissiveBoost = Material->GetEmissiveBoost();
+		MaterialData.DiffuseBoost = Material->GetDiffuseBoost();
 
 		TArray<FFloat16Color> MaterialEmissive;
 		TArray<FFloat16Color> MaterialDiffuse;
@@ -1186,12 +1196,13 @@ void FLightmassExporter::ExportMaterial(UMaterialInterface* Material)
 		TArray<FFloat16Color> MaterialNormal;
 
 		if (MaterialRenderer.GenerateMaterialData(
-			*Material, 
-			MaterialData, 
-			*ExportEntry, 
+			*Material,
+			ExportSettings,
+			MaterialData,
+			*ExportEntry,
 			MaterialDiffuse,
 			MaterialEmissive,
-			MaterialTransmission, 
+			MaterialTransmission,
 			MaterialNormal))
 		{
 			// open the channel
@@ -1208,7 +1219,7 @@ void FLightmassExporter::ExportMaterial(UMaterialInterface* Material)
 				uint8* OutData;
 				int32 OutSize;
 
-				OutSize = FMath::Square(MaterialData.EmissiveSize) * sizeof(FFloat16Color);  
+				OutSize = FMath::Square(MaterialData.EmissiveSize) * sizeof(FFloat16Color);
 				if (OutSize > 0)
 				{
 					OutData = (uint8*)(MaterialEmissive.GetData());
diff --git a/Engine/Source/Editor/UnrealEd/Private/Lightmass/Lightmass.h b/Engine/Source/Editor/UnrealEd/Private/Lightmass/Lightmass.h
--- a/Engine/Source/Editor/UnrealEd/Private/Lightmass/Lightmass.h
+++ b/Engine/Source/Editor/UnrealEd/Private/Lightmass/Lightmass.h
@@ -21,6 +21,9 @@
 #include "Lightmass/LightmassCharacterIndirectDetailVolume.h"
 #include "GameFramework/WorldSettings.h"
 
+// Forward declarations
+class FLandscapeStaticLightingMesh;
+
 /** Forward declarations of Lightmass types */
 namespace Lightmass
 {
@@ -30,6 +33,16 @@ namespace Lightmass
 	struct FMaterialElementData;
 }
 
+struct FLightmassMaterialExportSettings
+{
+	const FStaticLightingMesh* UnwrapMesh;
+
+	friend bool operator ==(const FLightmassMaterialExportSettings& lhs, const FLightmassMaterialExportSettings& rhs)
+	{
+		return lhs.UnwrapMesh == rhs.UnwrapMesh;
+	}
+};
+
 /** Lightmass Exporter class */
 class FLightmassExporter
 {
@@ -83,7 +96,9 @@ public:
 		CharacterIndirectDetailVolumes.Add(InDetailVolume->GetComponentsBoundingBox(true));
 	}
 
-	void AddMaterial(UMaterialInterface* InMaterialInterface);
+	// if provided, InStaticLightingMesh is used to UV unwrap the material into the static lighting textures
+	void AddMaterial(UMaterialInterface* InMaterialInterface, const FStaticLightingMesh* InStaticLightingMesh = nullptr);
+
 	void AddLight(ULightComponentBase* Light);
 
 	const FStaticLightingMapping* FindMappingByGuid(FGuid FindGuid) const;
@@ -94,7 +109,6 @@ public:
 	TArray<FGuid> VisibilityBucketGuids;
 
 private:
-
 	void WriteToChannel( FLightmassStatistics& Stats, FGuid& DebugMappingGuid );
 	bool WriteToMaterialChannel(FLightmassStatistics& Stats);
 
@@ -120,7 +134,7 @@ private:
 	 */
 	void BuildMaterialMap(UMaterialInterface* Material);
 	void BlockOnShaderCompilation();
-	void ExportMaterial(UMaterialInterface* Material);
+	void ExportMaterial(UMaterialInterface* Material, const FLightmassMaterialExportSettings& ExportSettings);
 
 	void WriteMeshInstances( int32 Channel );
 	void WriteLandscapeInstances( int32 Channel );
@@ -197,11 +211,12 @@ private:
 	TArray<const class UStaticMesh*> StaticMeshes;
 
 	// Landscape
-	TArray<const class FLandscapeStaticLightingMesh*> LandscapeLightingMeshes;
-	TArray<class FLandscapeStaticLightingTextureMapping*> LandscapeTextureMappings;
+	TArray<const FLandscapeStaticLightingMesh*> LandscapeLightingMeshes;
+	TArray<FLandscapeStaticLightingTextureMapping*> LandscapeTextureMappings;
 
 	// materials
-	TArray<class UMaterialInterface*> Materials;
+	TArray<UMaterialInterface*> Materials;
+	TMap<UMaterialInterface*, FLightmassMaterialExportSettings> MaterialExportSettings;
 	TMap<UMaterialInterface*, FMaterialExportDataEntry> MaterialExportData;
 
 	/** Exporting progress bar maximum value. */
diff --git a/Engine/Source/Editor/UnrealEd/Private/Lightmass/LightmassLandscapeRender.cpp b/Engine/Source/Editor/UnrealEd/Private/Lightmass/LightmassLandscapeRender.cpp
new file mode 100644
--- /dev/null
+++ b/Engine/Source/Editor/UnrealEd/Private/Lightmass/LightmassLandscapeRender.cpp
@@ -0,0 +1,309 @@
+// Copyright 1998-2015 Epic Games, Inc. All Rights Reserved.
+
+#include "UnrealEd.h"
+#include "LightmassLandscapeRender.h"
+#include "LandscapeLight.h"
+
+#include "LandscapeComponent.h"
+#include "LandscapeProxy.h"
+#include "LandscapeMaterialInstanceConstant.h"
+#include "EngineModule.h"
+#include "RHIStaticStates.h"
+#include "LandscapeEdit.h"
+
+struct FLightmassLandscapeVertex
+{
+	FVector			Position;
+	FPackedNormal	TangentX;
+	FPackedNormal	TangentZ;
+	FColor			Color;
+
+	// 0: Layer texcoord (XY)
+	// 1: Layer texcoord (XZ)
+	// 2: Layer texcoord (YZ)
+	// 3: Weightmap texcoord
+	// 4: Lightmap texcoord (ignored)
+	// 5: Heightmap texcoord (ignored)
+	FVector2DHalf UVs[4];
+
+	FLightmassLandscapeVertex(FVector InPosition, FVector LayerTexcoords, FVector2D WeightmapTexcoords)
+		: Position(InPosition)
+		, TangentX(FVector(1, 0, 0))
+		, TangentZ(FVector(0, 0, 1))
+		, Color(FColor::White)
+	{
+		// TangentZ.w contains the sign of the tangent basis determinant. Assume +1
+		TangentZ.Vector.W = 255;
+
+		UVs[0] = FVector2D(LayerTexcoords.X, LayerTexcoords.Y);
+		UVs[1] = FVector2D(LayerTexcoords.X, LayerTexcoords.Y); // Z not currently set, so use Y
+		UVs[2] = FVector2D(LayerTexcoords.Y, LayerTexcoords.X); // Z not currently set, so use X
+		UVs[3] = WeightmapTexcoords;
+	};
+};
+
+TGlobalResource<FVertexBuffer> LightmassLandscapeVertexBuffer;
+
+class FLightmassLandscapeVertexFactory : public FLocalVertexFactory
+{
+public:
+
+	/** Default constructor. */
+	FLightmassLandscapeVertexFactory()
+	{
+		FLocalVertexFactory::DataType Data;
+		// Position
+		Data.PositionComponent = STRUCTMEMBER_VERTEXSTREAMCOMPONENT(&LightmassLandscapeVertexBuffer, FLightmassLandscapeVertex, Position, VET_Float3);
+		// Tangents
+		Data.TangentBasisComponents[0] = STRUCTMEMBER_VERTEXSTREAMCOMPONENT(&LightmassLandscapeVertexBuffer, FLightmassLandscapeVertex, TangentX, VET_PackedNormal);
+		Data.TangentBasisComponents[1] = STRUCTMEMBER_VERTEXSTREAMCOMPONENT(&LightmassLandscapeVertexBuffer, FLightmassLandscapeVertex, TangentZ, VET_PackedNormal);
+		// Color
+		Data.ColorComponent = STRUCTMEMBER_VERTEXSTREAMCOMPONENT(&LightmassLandscapeVertexBuffer, FLightmassLandscapeVertex, Color, VET_Color);
+		// UVs (packed two to a stream component)
+		Data.TextureCoordinates.Add(STRUCTMEMBER_VERTEXSTREAMCOMPONENT(&LightmassLandscapeVertexBuffer, FLightmassLandscapeVertex, UVs[0], VET_Half4));
+		Data.TextureCoordinates.Add(STRUCTMEMBER_VERTEXSTREAMCOMPONENT(&LightmassLandscapeVertexBuffer, FLightmassLandscapeVertex, UVs[2], VET_Half4));
+
+		// update the data
+		SetData(Data);
+	}
+};
+TGlobalResource<FLightmassLandscapeVertexFactory> LightmassLandscapeVertexFactory;
+
+TGlobalResource<FIdentityPrimitiveUniformBuffer> LightmassLandscapeUniformBuffer;
+
+void RenderLandscapeMaterialForLightmass(const FLandscapeStaticLightingMesh* LandscapeMesh, FMaterialRenderProxy* MaterialProxy, const FRenderTarget* RenderTarget)
+{
+	const ULandscapeComponent* LandscapeComponent = CastChecked<ULandscapeComponent>(LandscapeMesh->Component);
+
+	const int32 SubsectionSizeQuads = LandscapeComponent->SubsectionSizeQuads;
+	const int32 NumSubsections = LandscapeComponent->NumSubsections;
+	const int32 ComponentSizeQuads = LandscapeComponent->ComponentSizeQuads;
+
+	int32 PatchExpandCountX = 0;
+	int32 PatchExpandCountY = 0;
+	int32 DesiredSize = 1;
+	const float LightMapOffset = 0.0f;
+	const float LightMapRes = LandscapeComponent->StaticLightingResolution > 0.f ? LandscapeComponent->StaticLightingResolution : LandscapeComponent->GetLandscapeProxy()->StaticLightingResolution;
+	const int32 LightingLOD = LandscapeComponent->GetLandscapeProxy()->StaticLightingLOD;
+	const float LightMapRatio = ::GetTerrainExpandPatchCount(LightMapRes, PatchExpandCountX, PatchExpandCountY, ComponentSizeQuads, (NumSubsections * (SubsectionSizeQuads + 1)), DesiredSize, LightingLOD);
+
+	const FVector2D PatchExpandOffset = FVector2D((float)PatchExpandCountX / (ComponentSizeQuads + 2 * PatchExpandCountX), (float)PatchExpandCountY / (ComponentSizeQuads + 2 * PatchExpandCountY)) * FVector2D(RenderTarget->GetSizeXY());
+	const FVector2D PatchExpandScale = FVector2D((float)ComponentSizeQuads / (ComponentSizeQuads + 2 * PatchExpandCountX), (float)ComponentSizeQuads / (ComponentSizeQuads + 2 * PatchExpandCountY));
+
+	TArray<FLightmassLandscapeVertex> Vertices;
+	TArray<uint16> Indices;
+	Vertices.Reserve(FMath::Square(NumSubsections) * 4);
+	Indices.Reserve(FMath::Square(NumSubsections) * 6);
+
+	const float fraction = 1.0f / NumSubsections;
+	const FVector2D PositionScale = FVector2D(RenderTarget->GetSizeXY()) * fraction * PatchExpandScale;
+	const float LayerScale = SubsectionSizeQuads;
+	const float WeightmapSubsection = LandscapeComponent->WeightmapSubsectionOffset;
+	const FVector2D WeightmapBias = FVector2D(LandscapeComponent->WeightmapScaleBias.Z, LandscapeComponent->WeightmapScaleBias.W);
+	const FVector2D WeightmapScale = FVector2D(LandscapeComponent->WeightmapScaleBias.X, LandscapeComponent->WeightmapScaleBias.Y) * SubsectionSizeQuads;
+
+	const int32 SubsectionX_Start = PatchExpandCountX > 0 ? -1 : 0;
+	const int32 SubsectionX_End = NumSubsections + (PatchExpandCountX > 0 ? 1 : 0);
+	const int32 SubsectionY_Start = PatchExpandCountY > 0 ? -1 : 0;
+	const int32 SubsectionY_End = NumSubsections + (PatchExpandCountY > 0 ? 1 : 0);
+	for (int32 SubsectionY = SubsectionY_Start; SubsectionY < SubsectionY_End; ++SubsectionY)
+	{
+		for (int32 SubsectionX = SubsectionX_Start; SubsectionX < SubsectionX_End; ++SubsectionX)
+		{
+			const FIntPoint UVSubsection = FIntPoint((SubsectionX >= 0 ? SubsectionX : 0),
+			                                         (SubsectionY >= 0 ? SubsectionY : 0));
+			const FVector2D UVScale = FVector2D((SubsectionX >= 0 && SubsectionX < NumSubsections ? 1 : 0),
+			                                    (SubsectionY >= 0 && SubsectionY < NumSubsections ? 1 : 0));
+
+			const FVector2D BasePosition = PatchExpandOffset + FVector2D(SubsectionX, SubsectionY) * PositionScale;
+			const FVector2D BaseLayerCoords = FVector2D(UVSubsection) * LayerScale;
+			const FVector2D BaseWeightmapCoords = WeightmapBias + FVector2D(UVSubsection) * WeightmapSubsection;
+
+			int32 Index = Vertices.Add(FLightmassLandscapeVertex(FVector(BasePosition /*FVector2D(0, 0) * PositionScale*/, 0), FVector(BaseLayerCoords /*FVector2D(0, 0) * UVScale * LayerScale*/, 0), BaseWeightmapCoords /*FVector2D(0, 0) * UVScale * WeightmapScale*/));
+			verifySlow(   Vertices.Add(FLightmassLandscapeVertex(FVector(BasePosition + FVector2D(1, 0) * PositionScale,   0), FVector(BaseLayerCoords + FVector2D(1, 0) * UVScale * LayerScale,   0), BaseWeightmapCoords + FVector2D(1, 0) * UVScale * WeightmapScale  )) == Index + 1);
+			verifySlow(   Vertices.Add(FLightmassLandscapeVertex(FVector(BasePosition + FVector2D(0, 1) * PositionScale,   0), FVector(BaseLayerCoords + FVector2D(0, 1) * UVScale * LayerScale,   0), BaseWeightmapCoords + FVector2D(0, 1) * UVScale * WeightmapScale  )) == Index + 2);
+			verifySlow(   Vertices.Add(FLightmassLandscapeVertex(FVector(BasePosition + FVector2D(1, 1) * PositionScale,   0), FVector(BaseLayerCoords + FVector2D(1, 1) * UVScale * LayerScale,   0), BaseWeightmapCoords + FVector2D(1, 1) * UVScale * WeightmapScale  )) == Index + 3);
+			checkSlow(Index + 3 <= MAX_uint16);
+			Indices.Add(Index);
+			Indices.Add(Index + 3);
+			Indices.Add(Index + 1);
+			Indices.Add(Index);
+			Indices.Add(Index + 2);
+			Indices.Add(Index + 3);
+		}
+	}
+
+	FMeshBatch MeshElement;
+	MeshElement.DynamicVertexStride = sizeof(FLightmassLandscapeVertex);
+	MeshElement.UseDynamicData = true;
+	MeshElement.bDisableBackfaceCulling = true;
+	MeshElement.CastShadow = false;
+	MeshElement.bWireframe = false;
+	MeshElement.Type = PT_TriangleList;
+	MeshElement.DepthPriorityGroup = SDPG_Foreground;
+	MeshElement.bUseAsOccluder = false;
+	MeshElement.bSelectable = false;
+	MeshElement.DynamicVertexData = Vertices.GetData();
+	MeshElement.VertexFactory = &LightmassLandscapeVertexFactory;
+	MeshElement.MaterialRenderProxy = MaterialProxy;
+
+	FMeshBatchElement& BatchElement = MeshElement.Elements[0];
+	BatchElement.PrimitiveUniformBufferResource = &LightmassLandscapeUniformBuffer;
+	BatchElement.DynamicIndexData = Indices.GetData();
+	BatchElement.FirstIndex = 0;
+	BatchElement.NumPrimitives = Indices.Num() / 3;
+	BatchElement.MinVertexIndex = 0;
+	BatchElement.MaxVertexIndex = Vertices.Num() - 1;
+	BatchElement.DynamicIndexStride = sizeof(uint16);
+
+	FSceneViewFamily ViewFamily(FSceneViewFamily::ConstructionValues(
+		RenderTarget,
+		NULL,
+		FEngineShowFlags(ESFIM_Game))
+		.SetWorldTimes(0, 0, 0)
+		.SetGammaCorrection(RenderTarget->GetDisplayGamma()));
+
+	const FIntRect RTViewRect(FIntPoint(0, 0), RenderTarget->GetSizeXY());
+
+	// make a temporary view
+	FSceneViewInitOptions ViewInitOptions;
+	ViewInitOptions.ViewFamily = &ViewFamily;
+	ViewInitOptions.SetViewRectangle(RTViewRect);
+	ViewInitOptions.ViewOrigin = FVector::ZeroVector;
+	ViewInitOptions.ViewRotationMatrix = FMatrix::Identity;
+	ViewInitOptions.ProjectionMatrix = FCanvas::CalcBaseTransform2D(RenderTarget->GetSizeXY().X, RenderTarget->GetSizeXY().Y);
+	ViewInitOptions.BackgroundColor = FLinearColor::Black;
+	ViewInitOptions.OverlayColor = FLinearColor::White;
+
+	FSceneView View(ViewInitOptions);
+
+	ENQUEUE_UNIQUE_RENDER_COMMAND_ONEPARAMETER(
+		CanvasFlushSetupCommand,
+		const FRenderTarget*, RenderTarget, RenderTarget,
+		{
+			//SCOPED_DRAW_EVENT(RHICmdList, CanvasFlush);
+
+			// Set the RHI render target.
+			::SetRenderTarget(RHICmdList, RenderTarget->GetRenderTargetTexture(), FTextureRHIRef());
+			// disable depth test & writes
+			RHICmdList.SetDepthStencilState(TStaticDepthStencilState<false,CF_Always>::GetRHI());
+
+			const FIntRect RTViewRect = FIntRect(0, 0, RenderTarget->GetRenderTargetTexture()->GetSizeX(), RenderTarget->GetRenderTargetTexture()->GetSizeY());
+
+			// set viewport to RT size
+			RHICmdList.SetViewport(RTViewRect.Min.X, RTViewRect.Min.Y, 0.0f, RTViewRect.Max.X, RTViewRect.Max.Y, 1.0f);
+		});
+
+	ENQUEUE_UNIQUE_RENDER_COMMAND_TWOPARAMETER(
+		RenderLandscapeMaterialForLightmass,
+		const FSceneView&, View, View,
+		const FMeshBatch&, Mesh, MeshElement,
+		{
+			//SCOPED_DRAW_EVENT(RHICmdList, RenderLandscapeMaterialToTexture);
+			GetRendererModule().DrawTileMesh(RHICmdList, View, Mesh, false, FHitProxyId());
+		});
+	FlushRenderingCommands();
+}
+
+void GetLandscapeOpacityData(const FLandscapeStaticLightingMesh* LandscapeMesh, int32& InOutSizeX, int32& InOutSizeY, TArray<FFloat16Color>& OutMaterialSamples)
+{
+	const ULandscapeComponent* LandscapeComponent = CastChecked<ULandscapeComponent>(LandscapeMesh->Component);
+
+	const int32 SubsectionSizeQuads = LandscapeComponent->SubsectionSizeQuads;
+	const int32 NumSubsections = LandscapeComponent->NumSubsections;
+	const int32 ComponentSizeQuads = LandscapeComponent->ComponentSizeQuads;
+
+	int32 PatchExpandCountX = 0;
+	int32 PatchExpandCountY = 0;
+	int32 DesiredSize = 1;
+	const float LightMapOffset = 0.0f;
+	const float LightMapRes = LandscapeComponent->StaticLightingResolution > 0.f ? LandscapeComponent->StaticLightingResolution : LandscapeComponent->GetLandscapeProxy()->StaticLightingResolution;
+	const int32 LightingLOD = LandscapeComponent->GetLandscapeProxy()->StaticLightingLOD;
+	const float LightMapRatio = ::GetTerrainExpandPatchCount(LightMapRes, PatchExpandCountX, PatchExpandCountY, ComponentSizeQuads, (NumSubsections * (SubsectionSizeQuads + 1)), DesiredSize, LightingLOD);
+
+	ULandscapeInfo* const LandscapeInfo = LandscapeComponent->GetLandscapeInfo(false);
+	check(LandscapeInfo);
+
+	FLandscapeEditDataInterface DataInterface = FLandscapeEditDataInterface(LandscapeInfo);
+	const int32 X1 = LandscapeComponent->SectionBaseX - PatchExpandCountX;
+	const int32 X2 = LandscapeComponent->SectionBaseX + ComponentSizeQuads + PatchExpandCountX + 1;
+	const int32 Y1 = LandscapeComponent->SectionBaseY - PatchExpandCountY;
+	const int32 Y2 = LandscapeComponent->SectionBaseY + ComponentSizeQuads + PatchExpandCountY + 1;
+	const int32 XSize = X2 - X1 + 1;
+	const int32 YSize = Y2 - Y1 + 1;
+
+	TArray<uint8> Data;
+	Data.AddUninitialized(YSize * XSize);
+	FMemory::Memset(Data.GetData(), 255, Data.Num());
+
+	DataInterface.GetWeightDataFast(ALandscapeProxy::VisibilityLayer, X1, Y1, X2, Y2, Data.GetData(), 0);
+
+	const int32 BaseComponentX = LandscapeComponent->SectionBaseX / ComponentSizeQuads;
+	const int32 BaseComponentY = LandscapeComponent->SectionBaseY / ComponentSizeQuads;
+	for (int32 ComponentY = BaseComponentY - 1; ComponentY <= BaseComponentY + 1; ++ComponentY)
+	{
+		for (int32 ComponentX = BaseComponentX - 1; ComponentX <= BaseComponentX + 1; ++ComponentX)
+		{
+			if (ComponentX == 0 && ComponentY == 0)
+			{
+				continue;
+			}
+
+			ULandscapeComponent* Component = LandscapeInfo->XYtoComponentMap.FindRef(FIntPoint(ComponentX, ComponentY));
+			if (Component)
+			{
+				if (Component->WeightmapLayerAllocations.ContainsByPredicate([](const FWeightmapLayerAllocationInfo& Allocation) { return Allocation.LayerInfo == ALandscapeProxy::VisibilityLayer; }))
+				{
+					// filled by GetWeightDataFast()
+				}
+				else
+				{
+					// fill with 0 (fully opaque)
+					// we don't worry about the border between components, we assume the value is identical on both components (it should be anyway)
+					const int32 X_Start = FMath::Clamp(ComponentX  * ComponentSizeQuads + PatchExpandCountX, 0, XSize);
+					const int32 X_End   = FMath::Clamp(ComponentX  * ComponentSizeQuads + PatchExpandCountX + ComponentSizeQuads + 1, 0, XSize);
+					const int32 Y_Start = FMath::Clamp(ComponentY  * ComponentSizeQuads + PatchExpandCountX, 0, YSize);
+					const int32 Y_End   = FMath::Clamp(ComponentY  * ComponentSizeQuads + PatchExpandCountX + ComponentSizeQuads + 1, 0, YSize);
+					for (int32 Y = Y_Start; Y < Y_End; ++Y)
+					{
+						uint8* DataPtr = &Data[Y * XSize + X_Start];
+						FMemory::Memset(DataPtr, 0, X_End - X_Start);
+					}
+				}
+			}
+			else
+			{
+				// filled with 255 (hole) by the Memset above.
+				// not done here due to the complexity of handling the shared border between existing / non-existing components
+				// it would be better if we could remove the triangles for non-existent components from the expanded landscape lighting mesh, but we don't yet do that
+			}
+		}
+	}
+
+	// Scale up the hole map to compensate for lightmass using point-sampling
+	static const float ScaleFactor = 3;
+	InOutSizeX = (XSize - 1) * ScaleFactor;
+	InOutSizeY = (YSize - 1) * ScaleFactor;
+
+	OutMaterialSamples.Empty(InOutSizeX * InOutSizeY);
+	for (int32 Y = 0; Y < InOutSizeY; ++Y)
+	{
+		for (int32 X = 0; X < InOutSizeX; ++X)
+		{
+			// Adding a half-pixel offset to compensate for lightmass using point-sampling
+			const FVector2D SampleUV = FVector2D(X + 0.5f, Y + 0.5f) / ScaleFactor;
+			const int32 X0 = FMath::FloorToInt(SampleUV.X);
+			const float Xf = SampleUV.X - X0;
+			const int32 Y0 = FMath::FloorToInt(SampleUV.Y);
+			const float Yf = SampleUV.Y - Y0;
+			const uint8 Value = FMath::BiLerp(
+				Data[FMath::Clamp(Y0    , 0, YSize - 1) * XSize + FMath::Clamp(X0    , 0, XSize - 1)],
+				Data[FMath::Clamp(Y0    , 0, YSize - 1) * XSize + FMath::Clamp(X0 + 1, 0, XSize - 1)],
+				Data[FMath::Clamp(Y0 + 1, 0, YSize - 1) * XSize + FMath::Clamp(X0    , 0, XSize - 1)],
+				Data[FMath::Clamp(Y0 + 1, 0, YSize - 1) * XSize + FMath::Clamp(X0 + 1, 0, XSize - 1)],
+				Xf, Yf);
+			const float Opacity = (255 - Value) / 255.0f;
+			OutMaterialSamples.Add(FLinearColor(Opacity, Opacity, Opacity));
+		}
+	}
+}
diff --git a/Engine/Source/Editor/UnrealEd/Private/Lightmass/LightmassLandscapeRender.h b/Engine/Source/Editor/UnrealEd/Private/Lightmass/LightmassLandscapeRender.h
new file mode 100644
--- /dev/null
+++ b/Engine/Source/Editor/UnrealEd/Private/Lightmass/LightmassLandscapeRender.h
@@ -0,0 +1,11 @@
+// Copyright 1998-2015 Epic Games, Inc. All Rights Reserved.
+
+#pragma once
+
+#include "LightmassRender.h"
+
+class FMaterialRenderProxy;
+class FLandscapeStaticLightingMesh;
+
+extern void RenderLandscapeMaterialForLightmass(const FLandscapeStaticLightingMesh* LandscapeMesh, FMaterialRenderProxy* MaterialProxy, const FRenderTarget* RenderTarget);
+extern void GetLandscapeOpacityData(const FLandscapeStaticLightingMesh* LandscapeMesh, int32& InOutSizeX, int32& InOutSizeY, TArray<FFloat16Color>& OutMaterialSamples);
diff --git a/Engine/Source/Editor/UnrealEd/Private/Lightmass/LightmassRender.cpp b/Engine/Source/Editor/UnrealEd/Private/Lightmass/LightmassRender.cpp
--- a/Engine/Source/Editor/UnrealEd/Private/Lightmass/LightmassRender.cpp
+++ b/Engine/Source/Editor/UnrealEd/Private/Lightmass/LightmassRender.cpp
@@ -14,6 +14,7 @@
 #include "Engine/TextureRenderTarget2D.h"
 #include "Materials/MaterialInstance.h"
 #include "CanvasTypes.h"
+#include "LightmassLandscapeRender.h"
 
 DEFINE_LOG_CATEGORY_STATIC(LogLightmassRender, Error, All);
 
@@ -774,33 +775,39 @@ void FLightmassMaterialRenderer::BeginGenerateMaterialData(
 {
 	UMaterial* BaseMaterial = InMaterial->GetMaterial();
 
-	TEnumAsByte<EBlendMode> BlendMode = InMaterial->GetBlendMode();
+	EBlendMode BlendMode = InMaterial->GetBlendMode();
+
+	const bool bIsLandscapeMaterial = InMaterial->IsA<ULandscapeMaterialInstanceConstant>();
 
 	if (BaseMaterial)
 	{
-		check(MaterialExportData.Find(InMaterial) == NULL);
+		check(!MaterialExportData.Contains(InMaterial));
 
 		FMaterialExportDataEntry& MaterialData = MaterialExportData.Add(InMaterial, FMaterialExportDataEntry(ChannelName));
 
-		if (FLightmassMaterialProxy::WillFillData((EBlendMode)(BlendMode), MP_DiffuseColor))
+		if (FLightmassMaterialProxy::WillFillData(BlendMode, MP_DiffuseColor))
 		{
 			MaterialData.DiffuseMaterialProxy = new FLightmassMaterialProxy();
 			MaterialData.DiffuseMaterialProxy->BeginCompiling(InMaterial, MP_DiffuseColor, EMaterialShaderMapUsage::LightmassExportDiffuse);
 		}
 
-		if (FLightmassMaterialProxy::WillFillData((EBlendMode)(BlendMode), MP_EmissiveColor))
+		if (FLightmassMaterialProxy::WillFillData(BlendMode, MP_EmissiveColor))
 		{
 			MaterialData.EmissiveMaterialProxy = new FLightmassMaterialProxy();
 			MaterialData.EmissiveMaterialProxy->BeginCompiling(InMaterial, MP_EmissiveColor, EMaterialShaderMapUsage::LightmassExportEmissive);
 		}
 
-		if (FLightmassMaterialProxy::WillFillData((EBlendMode)(BlendMode), MP_Opacity))
+		if (FLightmassMaterialProxy::WillFillData(BlendMode, MP_Opacity))
 		{
-			MaterialData.OpacityMaterialProxy = new FLightmassMaterialProxy();
-			MaterialData.OpacityMaterialProxy->BeginCompiling(InMaterial, MP_Opacity, EMaterialShaderMapUsage::LightmassExportOpacity);
+			// Landscape opacity is generated from the hole mask, not the material
+			if (!bIsLandscapeMaterial)
+			{
+				MaterialData.OpacityMaterialProxy = new FLightmassMaterialProxy();
+				MaterialData.OpacityMaterialProxy->BeginCompiling(InMaterial, MP_Opacity, EMaterialShaderMapUsage::LightmassExportOpacity);
+			}
 		}
 
-		if (bInWantNormals && FLightmassMaterialProxy::WillFillData((EBlendMode)(BlendMode), MP_Normal))
+		if (bInWantNormals && FLightmassMaterialProxy::WillFillData(BlendMode, MP_Normal))
 		{
 			MaterialData.NormalMaterialProxy = new FLightmassMaterialProxy();
 			MaterialData.NormalMaterialProxy->BeginCompiling(InMaterial, MP_Normal, EMaterialShaderMapUsage::LightmassExportNormal);
@@ -820,12 +827,13 @@ void FLightmassMaterialRenderer::BeginGenerateMaterialData(
  *	@return	bool					true if successful, false if not.
  */
 bool FLightmassMaterialRenderer::GenerateMaterialData(
-	UMaterialInterface& InMaterial, 
-	Lightmass::FMaterialData& OutMaterialData, 
+	UMaterialInterface& InMaterial,
+	const FLightmassMaterialExportSettings& InExportSettings,
+	Lightmass::FMaterialData& OutMaterialData,
 	FMaterialExportDataEntry& MaterialExportEntry,
-	TArray<FFloat16Color>& OutMaterialDiffuse, 
-	TArray<FFloat16Color>& OutMaterialEmissive, 
-	TArray<FFloat16Color>& OutMaterialTransmission, 
+	TArray<FFloat16Color>& OutMaterialDiffuse,
+	TArray<FFloat16Color>& OutMaterialEmissive,
+	TArray<FFloat16Color>& OutMaterialTransmission,
 	TArray<FFloat16Color>& OutMaterialNormal)
 {
 	bool bResult = true;
@@ -853,10 +861,18 @@ bool FLightmassMaterialRenderer::GenerateMaterialData(
 	OutMaterialData.bCastShadowAsMasked = BaseMaterial->GetCastShadowAsMasked() && IsTranslucentBlendMode((EBlendMode)BlendMode);
 	OutMaterialData.OpacityMaskClipValue = InMaterial.GetOpacityMaskClipValue();
 
+	const bool bIsLandscapeMaterial = InMaterial.IsA<ULandscapeMaterialInstanceConstant>();
+
+	// due to landscape using an expanded mesh, we have to mask out the edge data even on opaque components (sigh)
+	if (bIsLandscapeMaterial && OutMaterialData.BlendMode == Lightmass::BLEND_Opaque)
+	{
+		OutMaterialData.BlendMode = Lightmass::BLEND_Masked;
+	}
+
 	// Diffuse
 	if (MaterialExportEntry.DiffuseMaterialProxy)
 	{
-		if (!GenerateMaterialPropertyData(InMaterial, *MaterialExportEntry.DiffuseMaterialProxy, MP_DiffuseColor, OutMaterialData.DiffuseSize, OutMaterialData.DiffuseSize, OutMaterialDiffuse))
+		if (!GenerateMaterialPropertyData(InMaterial, InExportSettings, MaterialExportEntry.DiffuseMaterialProxy, MP_DiffuseColor, OutMaterialData.DiffuseSize, OutMaterialData.DiffuseSize, OutMaterialDiffuse))
 		{
 			UE_LOG(LogLightmassRender, Warning, TEXT("Failed to generate diffuse material samples for %s: %s"),
 				*(InMaterial.GetLightingGuid().ToString()), *(InMaterial.GetPathName()));
@@ -868,7 +884,7 @@ bool FLightmassMaterialRenderer::GenerateMaterialData(
 	// Emissive
 	if (MaterialExportEntry.EmissiveMaterialProxy)
 	{
-		if (!GenerateMaterialPropertyData(InMaterial, *MaterialExportEntry.EmissiveMaterialProxy, MP_EmissiveColor, OutMaterialData.EmissiveSize, OutMaterialData.EmissiveSize, OutMaterialEmissive))
+		if (!GenerateMaterialPropertyData(InMaterial, InExportSettings, MaterialExportEntry.EmissiveMaterialProxy, MP_EmissiveColor, OutMaterialData.EmissiveSize, OutMaterialData.EmissiveSize, OutMaterialEmissive))
 		{
 			UE_LOG(LogLightmassRender, Warning, TEXT("Failed to generate emissive material samples for %s: %s"),
 				*(InMaterial.GetLightingGuid().ToString()), *(InMaterial.GetPathName()));
@@ -878,9 +894,10 @@ bool FLightmassMaterialRenderer::GenerateMaterialData(
 	}
 
 	// Transmission
-	if (MaterialExportEntry.OpacityMaterialProxy)
+	// Landscape opacity is generated from the hole mask, not the material
+	if (MaterialExportEntry.OpacityMaterialProxy || bIsLandscapeMaterial)
 	{
-		if (!GenerateMaterialPropertyData(InMaterial, *MaterialExportEntry.OpacityMaterialProxy, MP_Opacity, OutMaterialData.TransmissionSize, OutMaterialData.TransmissionSize, OutMaterialTransmission))
+		if (!GenerateMaterialPropertyData(InMaterial, InExportSettings, MaterialExportEntry.OpacityMaterialProxy, MP_Opacity, OutMaterialData.TransmissionSize, OutMaterialData.TransmissionSize, OutMaterialTransmission))
 		{
 			UE_LOG(LogLightmassRender, Warning, TEXT("Failed to generate transmissive material samples for %s: %s"),
 				*(InMaterial.GetLightingGuid().ToString()), *(InMaterial.GetPathName()));
@@ -892,7 +909,7 @@ bool FLightmassMaterialRenderer::GenerateMaterialData(
 	// Normal
 	if (MaterialExportEntry.NormalMaterialProxy)
 	{
-		if (!GenerateMaterialPropertyData(InMaterial, *MaterialExportEntry.NormalMaterialProxy, MP_Normal, OutMaterialData.NormalSize, OutMaterialData.NormalSize, OutMaterialNormal))
+		if (!GenerateMaterialPropertyData(InMaterial, InExportSettings, MaterialExportEntry.NormalMaterialProxy, MP_Normal, OutMaterialData.NormalSize, OutMaterialData.NormalSize, OutMaterialNormal))
 		{
 			UE_LOG(LogLightmassRender, Warning, TEXT("Failed to generate normal material samples for %s: %s"),
 				*(InMaterial.GetLightingGuid().ToString()), *(InMaterial.GetPathName()));
@@ -904,6 +921,39 @@ bool FLightmassMaterialRenderer::GenerateMaterialData(
 	return bResult;
 }
 
+void LightmassDebugExportMaterial(UMaterialInterface& InMaterial, EMaterialProperty InMaterialProperty, FFloat16Color* InMaterialSamples, int32 InSizeX, int32 InSizeY)
+{
+	TArray<FColor> OutputBuffer;
+	OutputBuffer.Empty(InSizeX * InSizeY);
+	bool bSRGB = InMaterialProperty != MP_Normal;
+	for (int32 i = 0; i < InSizeX * InSizeY; ++i)
+	{
+		FLinearColor LinearColor(InMaterialSamples[i]);
+		OutputBuffer.Add(LinearColor.ToFColor(bSRGB));
+	}
+
+	// Create screenshot folder if not already present.
+	// Save the contents of the array to a bitmap file.
+	FString TempPath = FPaths::ScreenShotDir();
+	TempPath += TEXT("/Materials");
+	IFileManager::Get().MakeDirectory(*TempPath, true);
+	FString TempName = InMaterial.GetPathName();
+	TempName = TempName.Replace(TEXT("."), TEXT("_"));
+	TempName = TempName.Replace(TEXT(":"), TEXT("_"));
+	FString OutputName = TempPath / TempName;
+	OutputName += TEXT("_");
+	switch (InMaterialProperty)
+	{
+	case MP_DiffuseColor:	OutputName += TEXT("Diffuse");			break;
+	case MP_EmissiveColor:	OutputName += TEXT("Emissive");			break;
+	case MP_SpecularColor:	OutputName += TEXT("Specular");			break;
+	case MP_Normal:			OutputName += TEXT("Normal");			break;
+	case MP_Opacity:		OutputName += TEXT("Transmissive");		break;
+	}
+	OutputName += TEXT(".BMP");
+	FFileHelper::CreateBitmap(*OutputName, InSizeX, InSizeY, OutputBuffer.GetData());
+}
+
 /**
  *	Generate the material data for the given material and it's given property.
  *
@@ -916,26 +966,51 @@ bool FLightmassMaterialRenderer::GenerateMaterialData(
  *	@return	bool					true if successful, false if not.
  */
 bool FLightmassMaterialRenderer::GenerateMaterialPropertyData(
-	UMaterialInterface& InMaterial, 
-	FLightmassMaterialProxy& MaterialProxy,
-	EMaterialProperty InMaterialProperty, 
-	int32& InOutSizeX, 
-	int32& InOutSizeY, 
+	UMaterialInterface& InMaterial,
+	const FLightmassMaterialExportSettings& InExportSettings,
+	FLightmassMaterialProxy* MaterialProxy,
+	EMaterialProperty InMaterialProperty,
+	int32& InOutSizeX,
+	int32& InOutSizeY,
 	TArray<FFloat16Color>& OutMaterialSamples)
 {
 	bool bResult = true;
 
 	FFloat16Color UniformValue;
-	// Verify that async compiling has completed for this material
-	// If the ShaderMap is NULL that's because it failed to compile, which is ok as the default material will be used for exporting
-	check(!MaterialProxy.GetGameThreadShaderMap() || MaterialProxy.GetGameThreadShaderMap()->IsValidForRendering());
 
-	if (MaterialProxy.WillGenerateUniformData(UniformValue) == false)
+	const bool bIsLandscapeMaterial = InMaterial.IsA<ULandscapeMaterialInstanceConstant>();
+
+	// Landscape opacity needs to be handled specially because it needs to look at the neighbor components
+	// trying to actually use the neighbor materials is all but impossible so we read the data from the hole mask ourself
+	if (bIsLandscapeMaterial && InMaterialProperty == MP_Opacity)
+	{
+		auto* LandscapeMesh = static_cast<const FLandscapeStaticLightingMesh*>(InExportSettings.UnwrapMesh);
+		GetLandscapeOpacityData(LandscapeMesh, InOutSizeX, InOutSizeY, OutMaterialSamples);
+
+		if (GLightmassDebugOptions.bDebugMaterials == true)
+		{
+			LightmassDebugExportMaterial(InMaterial, InMaterialProperty, OutMaterialSamples.GetData(), InOutSizeX, InOutSizeY);
+		}
+	}
+	else if (MaterialProxy->WillGenerateUniformData(UniformValue))
 	{
+		// Single value... fill it in.
+		InOutSizeX = 1;
+		InOutSizeY = 1;
+		OutMaterialSamples.Empty(1);
+		OutMaterialSamples.AddZeroed(1);
+		OutMaterialSamples[0] = UniformValue;
+	}
+	else
+	{
+		// Verify that async compiling has completed for this material
+		// If the ShaderMap is NULL that's because it failed to compile, which is ok as the default material will be used for exporting
+		check(!MaterialProxy->GetGameThreadShaderMap() || MaterialProxy->GetGameThreadShaderMap()->IsValidForRendering());
+
 		//@todo UE4. The format may be determined by the material property...
 		// For example, if Diffuse doesn't need to be F16 it can create a standard RGBA8 target.
 		EPixelFormat Format = PF_FloatRGBA;
-		if (MaterialProxy.GetRenderTargetFormatAndSize(InMaterialProperty, Format, InMaterial.GetExportResolutionScale(), InOutSizeX, InOutSizeY))
+		if (MaterialProxy->GetRenderTargetFormatAndSize(InMaterialProperty, Format, InMaterial.GetExportResolutionScale(), InOutSizeX, InOutSizeY))
 		{
 			if (CreateRenderTarget(Format, InOutSizeX, InOutSizeY) == false)
 			{
@@ -944,80 +1019,53 @@ bool FLightmassMaterialRenderer::GenerateMaterialPropertyData(
 			}
 			else
 			{
-				// At this point, we can't just return false at failure since we have some clean-up to do...
-				Canvas->SetRenderTarget_GameThread(RenderTarget->GameThread_GetRenderTargetResource());
-
-				// Clear the render target to black
-				// This is necessary because the below DrawTile doesn't write to the first column and first row
-				//@todo - figure out and fix DrawTile issues when rendering a full render target quad
-				Canvas->Clear(FLinearColor(0,0,0,0));
-				FCanvasTileItem TileItem( FVector2D( 0.0f, 0.0f ), &MaterialProxy, FVector2D( InOutSizeX, InOutSizeY ) );
-				TileItem.bFreezeTime = true;
-				Canvas->DrawItem( TileItem );
-				Canvas->Flush_GameThread();
-				FlushRenderingCommands();
-				Canvas->SetRenderTarget_GameThread(NULL);
-				FlushRenderingCommands();
+				if (bIsLandscapeMaterial)
+				{
+					// Landscape needs special handling because it uses multiple UVs, which isn't yet supported by lightmass's regular pipeline
+					auto* LandscapeMesh = static_cast<const FLandscapeStaticLightingMesh*>(InExportSettings.UnwrapMesh);
+					RenderLandscapeMaterialForLightmass(LandscapeMesh, MaterialProxy, RenderTarget->GameThread_GetRenderTargetResource());
+				}
+				else
+				{
+					// At this point, we can't just return false at failure since we have some clean-up to do...
+					Canvas->SetRenderTarget_GameThread(RenderTarget->GameThread_GetRenderTargetResource());
+
+					// Clear the render target to black
+					// This is necessary because the below DrawTile doesn't write to the first column and first row
+					//@todo - figure out and fix DrawTile issues when rendering a full render target quad
+					Canvas->Clear(FLinearColor(0, 0, 0, 0));
+					FCanvasTileItem TileItem(FVector2D(0.0f, 0.0f), MaterialProxy, FVector2D(InOutSizeX, InOutSizeY));
+					TileItem.bFreezeTime = true;
+					Canvas->DrawItem(TileItem);
+					Canvas->Flush_GameThread();
+					FlushRenderingCommands();
+					Canvas->SetRenderTarget_GameThread(NULL);
+					FlushRenderingCommands();
+				}
 
 				// Read in the data
 				//@todo UE4. Check the format! RenderTarget->Format
 				// If we are going to allow non-F16 formats, then the storage will have to be aware of it!
 				if (RenderTarget->GameThread_GetRenderTargetResource()->ReadFloat16Pixels(OutMaterialSamples) == false)
 				{
-					UE_LOG(LogLightmassRender, Warning, TEXT("Failed to read Float16Pixels for 0x%08x property of %s: %s"), 
+					UE_LOG(LogLightmassRender, Warning, TEXT("Failed to read Float16Pixels for 0x%08x property of %s: %s"),
 						(uint32)InMaterialProperty, *(InMaterial.GetLightingGuid().ToString()), *(InMaterial.GetPathName()));
 					bResult = false;
 				}
-				else if (GLightmassDebugOptions.bDebugMaterials == true)
-				{
-					TArray<FColor> OutputBuffer;
-					OutputBuffer.Empty(OutMaterialSamples.Num());
-					bool bSRGB = InMaterialProperty != MP_Normal;
-					for (int32 i = 0; i < OutMaterialSamples.Num(); ++i)
-					{
-						FLinearColor LinearColor(OutMaterialSamples[i]);
-						OutputBuffer.Add(LinearColor.ToFColor(bSRGB));
-					}
 
-					// Create screenshot folder if not already present.
-					// Save the contents of the array to a bitmap file.
-					FString TempPath = FPaths::ScreenShotDir();
-					TempPath += TEXT("/Materials");
-					IFileManager::Get().MakeDirectory(*TempPath, true);
-					FString TempName = InMaterial.GetPathName();
-					TempName = TempName.Replace(TEXT("."), TEXT("_"));
-					TempName = TempName.Replace(TEXT(":"), TEXT("_"));
-					FString OutputName = TempPath / TempName;
-					OutputName += TEXT("_");
-					switch (InMaterialProperty)
-					{
-					case MP_DiffuseColor:	OutputName += TEXT("Diffuse");			break;
-					case MP_EmissiveColor:	OutputName += TEXT("Emissive");			break;
-					case MP_SpecularColor:	OutputName += TEXT("Specular");			break;
-					case MP_Normal:			OutputName += TEXT("Normal");			break;
-					case MP_Opacity:		OutputName += TEXT("Transmissive");		break;
-					}
-					OutputName += TEXT(".BMP");
-					FFileHelper::CreateBitmap(*OutputName,InOutSizeX,InOutSizeY,OutputBuffer.GetData());
+				if (GLightmassDebugOptions.bDebugMaterials == true)
+				{
+					LightmassDebugExportMaterial(InMaterial, InMaterialProperty, OutMaterialSamples.GetData(), InOutSizeX, InOutSizeY);
 				}
 			}
 		}
 		else
 		{
-			UE_LOG(LogLightmassRender, Warning, TEXT("Failed to get render target format and size for 0x%08x property of %s: %s"), 
+			UE_LOG(LogLightmassRender, Warning, TEXT("Failed to get render target format and size for 0x%08x property of %s: %s"),
 				(uint32)InMaterialProperty, *(InMaterial.GetLightingGuid().ToString()), *(InMaterial.GetPathName()));
 			bResult = false;
 		}
 	}
-	else
-	{
-		// Single value... fill it in.
-		InOutSizeX = 1;
-		InOutSizeY = 1; 
-		OutMaterialSamples.Empty(1);
-		OutMaterialSamples.AddZeroed(1);
-		OutMaterialSamples[0] = UniformValue;
-	}
 
 	return bResult;
 }
diff --git a/Engine/Source/Editor/UnrealEd/Private/Lightmass/LightmassRender.h b/Engine/Source/Editor/UnrealEd/Private/Lightmass/LightmassRender.h
--- a/Engine/Source/Editor/UnrealEd/Private/Lightmass/LightmassRender.h
+++ b/Engine/Source/Editor/UnrealEd/Private/Lightmass/LightmassRender.h
@@ -4,8 +4,7 @@
 	LightmassRender.h: lightmass rendering-related definitions.
 =============================================================================*/
 
-#ifndef __LIGHTMASSRENDER_H__
-#define __LIGHTMASSRENDER_H__
+#pragma once
 
 /** Forward declarations of Lightmass types */
 namespace Lightmass
@@ -14,6 +13,7 @@ namespace Lightmass
 }
 
 class FLightmassMaterialProxy;
+struct FLightmassMaterialExportSettings;
 
 /** Information about a material to be exported to lightmass, which is generated by a first pass over the materials. */
 class FMaterialExportDataEntry
@@ -55,8 +55,8 @@ public:
 
 	/** Initiates generation of material export data by kicking off async shader compiles. */
 	void BeginGenerateMaterialData(
-		UMaterialInterface* InMaterial, 
-		bool bInWantNormals, 
+		UMaterialInterface* InMaterial,
+		bool bInWantNormals,
 		const FString& ChannelName,
 		TMap<UMaterialInterface*, FMaterialExportDataEntry>& MaterialExportData);
 
@@ -72,12 +72,13 @@ public:
 	 *	@return	bool					true if successful, false if not.
 	 */
 	bool GenerateMaterialData(
-		UMaterialInterface& InMaterial, 
-		Lightmass::FMaterialData& OutMaterialData, 
+		UMaterialInterface& InMaterial,
+		const FLightmassMaterialExportSettings& InExportSettings,
+		Lightmass::FMaterialData& OutMaterialData,
 		FMaterialExportDataEntry& MaterialExportEntry,
-		TArray<FFloat16Color>& OutMaterialDiffuse, 
-		TArray<FFloat16Color>& OutMaterialEmissive, 
-		TArray<FFloat16Color>& OutMaterialTransmission, 
+		TArray<FFloat16Color>& OutMaterialDiffuse,
+		TArray<FFloat16Color>& OutMaterialEmissive,
+		TArray<FFloat16Color>& OutMaterialTransmission,
 		TArray<FFloat16Color>& OutMaterialNormal);
 
 protected:
@@ -104,8 +105,9 @@ protected:
 	 *	@return	bool					true if successful, false if not.
 	 */
 	bool GenerateMaterialPropertyData(
-		UMaterialInterface& InMaterial,
-		FLightmassMaterialProxy& MaterialProxy,
+		UMaterialInterface& InMaterial, 
+		const FLightmassMaterialExportSettings& InExportSettings,
+		FLightmassMaterialProxy* MaterialProxy,
 		EMaterialProperty InMaterialProperty, 
 		int32& InOutSizeX, 
 		int32& InOutSizeY, 
@@ -115,4 +117,3 @@ protected:
 	FCanvas* Canvas;
 };
 
-#endif	//__LIGHTMASSRENDER
diff --git a/Engine/Source/Editor/UnrealEd/Private/StaticLightingSystem/StaticLightingExport.cpp b/Engine/Source/Editor/UnrealEd/Private/StaticLightingSystem/StaticLightingExport.cpp
--- a/Engine/Source/Editor/UnrealEd/Private/StaticLightingSystem/StaticLightingExport.cpp
+++ b/Engine/Source/Editor/UnrealEd/Private/StaticLightingSystem/StaticLightingExport.cpp
@@ -133,7 +133,7 @@ void FLandscapeStaticLightingMesh::ExportMeshInstance(class FLightmassExporter*
 
 	if (LandscapeComponent && LandscapeComponent->MaterialInstance)
 	{
-		Exporter->AddMaterial(LandscapeComponent->MaterialInstance);
+		Exporter->AddMaterial(LandscapeComponent->MaterialInstance, this);
 	}
 
 	for( int32 LightIdx=0; LightIdx < RelevantLights.Num(); LightIdx++ )
diff --git a/Engine/Source/Programs/UnrealLightmass/Public/ImportExport.h b/Engine/Source/Programs/UnrealLightmass/Public/ImportExport.h
--- a/Engine/Source/Programs/UnrealLightmass/Public/ImportExport.h
+++ b/Engine/Source/Programs/UnrealLightmass/Public/ImportExport.h
@@ -49,7 +49,7 @@ namespace Lightmass
 	static const FGuid LM_DEBUGOUTPUT_VERSION			= FGuid(0xfa520216, 0xd5aa49eb, 0xa9bbc1f5, 0x8ab715fc);
 	static const FGuid LM_SCENE_VERSION					= FGuid(0x8ad6acb1, 0xe15e4fec, 0x93c72936, 0xf2b67090);
 	static const FGuid LM_STATICMESH_VERSION			= FGuid(0x9b3b1bc6, 0x9270497a, 0xa07522fe, 0xc7055e91);
-	static const FGuid LM_MATERIAL_VERSION				= FGuid(0xa6640d72, 0x33bf47b4, 0xbbcf76bb, 0x4350a2b5);
+	static const FGuid LM_MATERIAL_VERSION				= FGuid(0x8487ad2e, 0x5dba413c, 0x81373d29, 0xb2915895);
 
 
 	/** Alert source object type identifiers... */
diff --git a/Engine/Source/Runtime/Landscape/Private/LandscapeEdit.cpp b/Engine/Source/Runtime/Landscape/Private/LandscapeEdit.cpp
--- a/Engine/Source/Runtime/Landscape/Private/LandscapeEdit.cpp
+++ b/Engine/Source/Runtime/Landscape/Private/LandscapeEdit.cpp
@@ -2020,7 +2020,7 @@ void ALandscapeProxy::Import(FGuid Guid, int32 VertsX, int32 VertsY,
 			// Weightmap is sized the same as the component
 			int32 WeightmapSize = (SubsectionSizeQuads + 1) * NumSubsections;
 			// Should be power of two
-			check(((WeightmapSize - 1) & WeightmapSize) == 0);
+			check(FMath::IsPowerOfTwo(WeightmapSize));
 
 			LandscapeComponent->WeightmapScaleBias = FVector4(1.0f / (float)WeightmapSize, 1.0f / (float)WeightmapSize, 0.5f / (float)WeightmapSize, 0.5f / (float)WeightmapSize);
 			LandscapeComponent->WeightmapSubsectionOffset = (float)(SubsectionSizeQuads + 1) / (float)WeightmapSize;
diff --git a/Engine/Source/Runtime/Landscape/Private/LandscapeLight.cpp b/Engine/Source/Runtime/Landscape/Private/LandscapeLight.cpp
--- a/Engine/Source/Runtime/Landscape/Private/LandscapeLight.cpp
+++ b/Engine/Source/Runtime/Landscape/Private/LandscapeLight.cpp
@@ -610,13 +610,13 @@ void ULandscapeComponent::GetStaticLightingInfo(FStaticLightingPrimitiveInfo& Ou
 {
 	if( HasStaticLighting() )
 	{
-		float LightMapRes = StaticLightingResolution > 0.f ? StaticLightingResolution : GetLandscapeProxy()->StaticLightingResolution;
-		int32 PatchExpandCountX = 1;
-		int32 PatchExpandCountY = 1;
-		int32 DesiredSize = 1;
-		int32 LightingLOD = GetLandscapeProxy()->StaticLightingLOD;
+		const float LightMapRes = StaticLightingResolution > 0.f ? StaticLightingResolution : GetLandscapeProxy()->StaticLightingResolution;
+		const int32 LightingLOD = GetLandscapeProxy()->StaticLightingLOD;
 
-		float LightMapRatio = ::GetTerrainExpandPatchCount(LightMapRes, PatchExpandCountX, PatchExpandCountY, ComponentSizeQuads, (NumSubsections * (SubsectionSizeQuads+1)), DesiredSize, LightingLOD);
+		int32 PatchExpandCountX = 0;
+		int32 PatchExpandCountY = 0;
+		int32 DesiredSize = 1;
+		const float LightMapRatio = ::GetTerrainExpandPatchCount(LightMapRes, PatchExpandCountX, PatchExpandCountY, ComponentSizeQuads, (NumSubsections * (SubsectionSizeQuads+1)), DesiredSize, LightingLOD);
 
 		int32 SizeX = DesiredSize;
 		int32 SizeY = DesiredSize;
diff --git a/Engine/Source/Runtime/Landscape/Private/LandscapeRender.cpp b/Engine/Source/Runtime/Landscape/Private/LandscapeRender.cpp
--- a/Engine/Source/Runtime/Landscape/Private/LandscapeRender.cpp
+++ b/Engine/Source/Runtime/Landscape/Private/LandscapeRender.cpp
@@ -945,9 +945,9 @@ namespace DebugColorMask
 void FLandscapeComponentSceneProxy::OnTransformChanged()
 {
 	// Set Lightmap ScaleBias
-	int32 PatchExpandCountX = 1;
-	int32 PatchExpandCountY = 1;
-	int32 DesiredSize = 1;
+	int32 PatchExpandCountX = 0;
+	int32 PatchExpandCountY = 0;
+	int32 DesiredSize = 1; // output by GetTerrainExpandPatchCount but not used below
 	const float LightMapRatio = ::GetTerrainExpandPatchCount(StaticLightingResolution, PatchExpandCountX, PatchExpandCountY, ComponentSizeQuads, (NumSubsections * (SubsectionSizeQuads + 1)), DesiredSize, StaticLightingLOD);
 	const float LightmapLODScaleX = LightMapRatio / ((ComponentSizeVerts >> StaticLightingLOD) + 2 * PatchExpandCountX);
 	const float LightmapLODScaleY = LightMapRatio / ((ComponentSizeVerts >> StaticLightingLOD) + 2 * PatchExpandCountY);
diff --git a/Engine/Source/Runtime/Landscape/Public/LandscapeLight.h b/Engine/Source/Runtime/Landscape/Public/LandscapeLight.h
--- a/Engine/Source/Runtime/Landscape/Public/LandscapeLight.h
+++ b/Engine/Source/Runtime/Landscape/Public/LandscapeLight.h
@@ -1,11 +1,10 @@
 // Copyright 1998-2015 Epic Games, Inc. All Rights Reserved.
 
 /*=============================================================================
-LandscapeLight.h: Static lighting for LandscapeComponents
-=============================================================================*/
+  LandscapeLight.h: Static lighting for LandscapeComponents
+  =============================================================================*/
 
-#ifndef _LANDSCAPELIGHT_H
-#define _LANDSCAPELIGHT_H
+#pragma once
 
 #include "StaticLighting.h"
 
@@ -94,4 +93,51 @@ public:
 #endif
 };
 
-#endif // _LANDSCAPELIGHT_H
+namespace
+{
+	// LightmapRes: Multiplier of lightmap size relative to landscape size
+	// X: (Output) PatchExpandCountX (at Lighting LOD)
+	// Y: (Output) PatchExpandCountY (at Lighting LOD)
+	// ComponentSize: Component size in patches (at LOD 0)
+	// LigtmapSize: Size desired for lightmap (texels)
+	// DesiredSize: (Output) Recommended lightmap size (texels)
+	// return: LightMapRatio
+	static float GetTerrainExpandPatchCount(float LightMapRes, int32& X, int32& Y, int32 ComponentSize, int32 LightmapSize, int32& DesiredSize, uint32 LightingLOD)
+	{
+		if (LightMapRes <= 0) return 0.f;
+
+		// Assuming DXT_1 compression at the moment...
+		int32 PixelPaddingX = GPixelFormats[PF_DXT1].BlockSizeX; // "/2" ?
+		int32 PixelPaddingY = GPixelFormats[PF_DXT1].BlockSizeY;
+		int32 PatchExpandCountX = (LightMapRes >= 1.f) ? (PixelPaddingX) / LightMapRes : (PixelPaddingX);
+		int32 PatchExpandCountY = (LightMapRes >= 1.f) ? (PixelPaddingY) / LightMapRes : (PixelPaddingY);
+
+		X = FMath::Max<int32>(1, PatchExpandCountX >> LightingLOD);
+		Y = FMath::Max<int32>(1, PatchExpandCountY >> LightingLOD);
+
+		DesiredSize = (LightMapRes >= 1.f) ? FMath::Min<int32>((int32)((ComponentSize + 1) * LightMapRes), 4096) : FMath::Min<int32>((int32)((LightmapSize)* LightMapRes), 4096);
+		int32 CurrentSize = (LightMapRes >= 1.f) ? FMath::Min<int32>((int32)((2 * (X << LightingLOD) + ComponentSize + 1) * LightMapRes), 4096) : FMath::Min<int32>((int32)((2 * (X << LightingLOD) + LightmapSize) * LightMapRes), 4096);
+
+		// Find proper Lightmap Size
+		if (CurrentSize > DesiredSize)
+		{
+			// Find maximum bit
+			int32 PriorSize = DesiredSize;
+			while (DesiredSize > 0)
+			{
+				PriorSize = DesiredSize;
+				DesiredSize = DesiredSize & ~(DesiredSize & ~(DesiredSize - 1));
+			}
+
+			DesiredSize = PriorSize << 1; // next bigger size
+			if (CurrentSize * CurrentSize <= ((PriorSize * PriorSize) << 1))
+			{
+				DesiredSize = PriorSize;
+			}
+		}
+
+		int32 DestSize = (float)DesiredSize / CurrentSize * (ComponentSize*LightMapRes);
+		float LightMapRatio = (float)DestSize / (ComponentSize*LightMapRes) * CurrentSize / DesiredSize;
+		return LightMapRatio;
+	}
+}
diff --git a/Engine/Source/Runtime/Landscape/Public/LandscapeRender.h b/Engine/Source/Runtime/Landscape/Public/LandscapeRender.h
--- a/Engine/Source/Runtime/Landscape/Public/LandscapeRender.h
+++ b/Engine/Source/Runtime/Landscape/Public/LandscapeRender.h
@@ -767,52 +767,3 @@ public:
 		}
 	}
 };
-
-namespace
-{
-	// LightmapRes: Multiplier of lightmap size relative to landscape size
-	// X: (Output) PatchExpandCountX (at Lighting LOD)
-	// Y: (Output) PatchExpandCountY (at Lighting LOD)
-	// ComponentSize: Component size in patches (at LOD 0)
-	// LigtmapSize: Size desired for lightmap (texels)
-	// DesiredSize: (Output) Recommended lightmap size (texels)
-	// return: LightMapRatio
-	static float GetTerrainExpandPatchCount(float LightMapRes, int32& X, int32& Y, int32 ComponentSize, int32 LightmapSize, int32& DesiredSize, uint32 LightingLOD)
-	{
-		if (LightMapRes <= 0) return 0.f;
-
-		// Assuming DXT_1 compression at the moment...
-		int32 PixelPaddingX = GPixelFormats[PF_DXT1].BlockSizeX;
-		int32 PixelPaddingY = GPixelFormats[PF_DXT1].BlockSizeY;
-		int32 PatchExpandCountX = (LightMapRes >= 1.f) ? (PixelPaddingX) / LightMapRes : (PixelPaddingX);
-		int32 PatchExpandCountY = (LightMapRes >= 1.f) ? (PixelPaddingY) / LightMapRes : (PixelPaddingY);
-
-		X = FMath::Max<int32>(1, PatchExpandCountX >> LightingLOD);
-		Y = FMath::Max<int32>(1, PatchExpandCountY >> LightingLOD);
-
-		DesiredSize = (LightMapRes >= 1.f) ? FMath::Min<int32>((int32)((ComponentSize + 1) * LightMapRes), 4096) : FMath::Min<int32>((int32)((LightmapSize)* LightMapRes), 4096);
-		int32 CurrentSize = (LightMapRes >= 1.f) ? FMath::Min<int32>((int32)((2 * (X << LightingLOD) + ComponentSize + 1) * LightMapRes), 4096) : FMath::Min<int32>((int32)((2 * (X << LightingLOD) + LightmapSize) * LightMapRes), 4096);
-
-		// Find proper Lightmap Size
-		if (CurrentSize > DesiredSize)
-		{
-			// Find maximum bit
-			int32 PriorSize = DesiredSize;
-			while (DesiredSize > 0)
-			{
-				PriorSize = DesiredSize;
-				DesiredSize = DesiredSize & ~(DesiredSize & ~(DesiredSize - 1));
-			}
-
-			DesiredSize = PriorSize << 1; // next bigger size
-			if (CurrentSize * CurrentSize <= ((PriorSize * PriorSize) << 1))
-			{
-				DesiredSize = PriorSize;
-			}
-		}
-
-		int32 DestSize = (float)DesiredSize / CurrentSize * (ComponentSize*LightMapRes);
-		float LightMapRatio = (float)DestSize / (ComponentSize*LightMapRes) * CurrentSize / DesiredSize;
-		return LightMapRatio;
-	}
-}
