Fix crash in new landscape lightmass opacity export

diff --git a/Engine/Source/Editor/UnrealEd/Private/Lightmass/LightmassLandscapeRender.cpp b/Engine/Source/Editor/UnrealEd/Private/Lightmass/LightmassLandscapeRender.cpp
--- a/Engine/Source/Editor/UnrealEd/Private/Lightmass/LightmassLandscapeRender.cpp
+++ b/Engine/Source/Editor/UnrealEd/Private/Lightmass/LightmassLandscapeRender.cpp
@@ -260,10 +260,10 @@ void GetLandscapeOpacityData(const FLandscapeStaticLightingMesh* LandscapeMesh,
 				{
 					// fill with 0 (fully opaque)
 					// we don't worry about the border between components, we assume the value is identical on both components (it should be anyway)
-					const int32 X_Start = FMath::Clamp(ComponentX  * ComponentSizeQuads + PatchExpandCountX, 0, XSize);
-					const int32 X_End   = FMath::Clamp(ComponentX  * ComponentSizeQuads + PatchExpandCountX + ComponentSizeQuads + 1, 0, XSize);
-					const int32 Y_Start = FMath::Clamp(ComponentY  * ComponentSizeQuads + PatchExpandCountX, 0, YSize);
-					const int32 Y_End   = FMath::Clamp(ComponentY  * ComponentSizeQuads + PatchExpandCountX + ComponentSizeQuads + 1, 0, YSize);
+					const int32 X_Start = FMath::Clamp(ComponentX  * ComponentSizeQuads + PatchExpandCountX                         , 0, XSize - 1);
+					const int32 X_End   = FMath::Clamp(ComponentX  * ComponentSizeQuads + PatchExpandCountX + ComponentSizeQuads + 1, 0, XSize - 1);
+					const int32 Y_Start = FMath::Clamp(ComponentY  * ComponentSizeQuads + PatchExpandCountX                         , 0, YSize - 1);
+					const int32 Y_End   = FMath::Clamp(ComponentY  * ComponentSizeQuads + PatchExpandCountX + ComponentSizeQuads + 1, 0, YSize - 1);
 					for (int32 Y = Y_Start; Y < Y_End; ++Y)
 					{
 						uint8* DataPtr = &Data[Y * XSize + X_Start];
