Fixed the game mode class always showing as "None" in the project settings

UE-15153 - Unable to change Default Game Mode through Project Settings window

CL# 2483623 (to fix UE-11234) made GetCurrentGameModeClass use GetValue to get the value as an object, rather than get the value as a string and then use FindObject to locate the class, however the property handle used by FStringClassReference doesn't implement a GetValue function that returns an object, so the function was always failing and returning null.

I've restored the code back to its original state, and added a guard against calling FindObject while GIsSavingPackage is true. I also added a cached version of the game mode class so that we can try and return something valid in the case where GetCurrentGameModeClass is called while GIsSavingPackage is true.

diff --git a/Engine/Source/Editor/UnrealEd/Public/GameModeInfoCustomizer.h b/Engine/Source/Editor/UnrealEd/Public/GameModeInfoCustomizer.h
--- a/Engine/Source/Editor/UnrealEd/Public/GameModeInfoCustomizer.h
+++ b/Engine/Source/Editor/UnrealEd/Public/GameModeInfoCustomizer.h
@@ -19,6 +19,7 @@ public:
 	{
 		OwningObject = InOwningObject;
 		GameModePropertyName = InGameModePropertyName;
+		CachedGameModeClass = nullptr;
 	}
 
 	/** Create widget for the name of a default class property */
@@ -164,15 +165,34 @@ public:
 	/** Get the currently set GameMode class */
 	const UClass* GetCurrentGameModeClass() const
 	{
-		UObject* ClassValue = nullptr;
-		DefaultGameModeClassHandle->GetValue(ClassValue);
-		const UClass* GameModeClass = Cast<UClass>(ClassValue);
-		return GameModeClass;
+		FString ClassName;
+		DefaultGameModeClassHandle->GetValueAsFormattedString(ClassName);
+
+		if (ClassName.IsEmpty() || ClassName == "None")
+		{
+			CachedGameModeClass = nullptr;
+			return nullptr;
+		}
+
+		FString StrippedClassName = ClassName;
+		ConstructorHelpers::StripObjectClass(StrippedClassName);
+
+		// Just returned the cached value if it's already up-to-date, or if we're saving packages since we can't search for the object while a save is happening
+		if ((CachedGameModeClass && CachedGameModeClass->GetPathName() == StrippedClassName) || GIsSavingPackage)
+		{
+			return CachedGameModeClass;
+		}
+
+		CachedGameModeClass = FindObject<UClass>(ANY_PACKAGE, *StrippedClassName);
+		return CachedGameModeClass;
 	}
 
 	void SetCurrentGameModeClass(const UClass* NewGameModeClass)
 	{
-		DefaultGameModeClassHandle->SetValueFromFormattedString((NewGameModeClass) ? NewGameModeClass->GetPathName() : "None");
+		if (DefaultGameModeClassHandle->SetValueFromFormattedString((NewGameModeClass) ? NewGameModeClass->GetPathName() : TEXT("None")) == FPropertyAccess::Success)
+		{
+			CachedGameModeClass = NewGameModeClass;
+		}
 	}
 
 	/** Get the CDO from the currently set GameMode class */
@@ -325,6 +345,8 @@ private:
 	FName GameModePropertyName;
 	/** Handle to the DefaultGameMode property */
 	TSharedPtr<IPropertyHandle> DefaultGameModeClassHandle;
+	/** Cached class pointer from the DefaultGameModeClassHandle */
+	mutable const UClass* CachedGameModeClass;
 };
 
 #undef LOCTEXT_NAMESPACE
