Fix variable shadowing

diff --git a/Engine/Source/Editor/UnrealEd/Private/Lightmass/LightmassLandscapeRender.cpp b/Engine/Source/Editor/UnrealEd/Private/Lightmass/LightmassLandscapeRender.cpp
--- a/Engine/Source/Editor/UnrealEd/Private/Lightmass/LightmassLandscapeRender.cpp
+++ b/Engine/Source/Editor/UnrealEd/Private/Lightmass/LightmassLandscapeRender.cpp
@@ -163,12 +163,12 @@ void RenderLandscapeMaterialForLightmass(const FLandscapeStaticLightingMesh* Lan
 		.SetWorldTimes(0, 0, 0)
 		.SetGammaCorrection(RenderTarget->GetDisplayGamma()));
 
-	const FIntRect RTViewRect(FIntPoint(0, 0), RenderTarget->GetSizeXY());
+	const FIntRect ViewRect(FIntPoint(0, 0), RenderTarget->GetSizeXY());
 
 	// make a temporary view
 	FSceneViewInitOptions ViewInitOptions;
 	ViewInitOptions.ViewFamily = &ViewFamily;
-	ViewInitOptions.SetViewRectangle(RTViewRect);
+	ViewInitOptions.SetViewRectangle(ViewRect);
 	ViewInitOptions.ViewOrigin = FVector::ZeroVector;
 	ViewInitOptions.ViewRotationMatrix = FMatrix::Identity;
 	ViewInitOptions.ProjectionMatrix = FCanvas::CalcBaseTransform2D(RenderTarget->GetSizeXY().X, RenderTarget->GetSizeXY().Y);
