Added more options to the build before Play On
UE-15378

diff --git a/Engine/Source/Editor/UnrealEd/Classes/Settings/LevelEditorPlaySettings.h b/Engine/Source/Editor/UnrealEd/Classes/Settings/LevelEditorPlaySettings.h
--- a/Engine/Source/Editor/UnrealEd/Classes/Settings/LevelEditorPlaySettings.h
+++ b/Engine/Source/Editor/UnrealEd/Classes/Settings/LevelEditorPlaySettings.h
@@ -78,6 +78,20 @@ enum EPlayNetMode
 };
 
 
+UENUM()
+enum EPlayOnBuildMode
+{
+	/** Always build */
+	PlayOnBuild_Always UMETA(DisplayName="Always Build"),
+
+	/** Never build. */
+	PlayOnBuild_Never UMETA(DisplayName="Never Build"),
+
+	/** Build based on project type */
+	PlayOnBuild_Default UMETA(DisplayName="Only Build Code Projects"),
+};
+
+
 /**
  * Holds information about a screen resolution to be used for playing.
  */
@@ -194,7 +208,7 @@ public:
 
 	/** The width of the new view port window in pixels (0 = use the desktop's screen resolution). */
 	UPROPERTY(config, EditAnywhere, Category = PlayOnDevice)
-	bool BuildGameBeforeLaunch;
+	TEnumAsByte<EPlayOnBuildMode> BuildGameBeforeLaunch;
 
 private:
 
@@ -344,4 +358,11 @@ public:
 	/** Collection of common screen resolutions on television screens. */
 	UPROPERTY(config)
 	TArray<FPlayScreenResolution> TelevisionScreenResolutions;
+
+protected:
+
+	// UObject overrides
+
+	virtual void PostEditChangeProperty(struct FPropertyChangedEvent& PropertyChangedEvent) override;
+
 };
diff --git a/Engine/Source/Editor/UnrealEd/Private/PlayLevel.cpp b/Engine/Source/Editor/UnrealEd/Private/PlayLevel.cpp
--- a/Engine/Source/Editor/UnrealEd/Private/PlayLevel.cpp
+++ b/Engine/Source/Editor/UnrealEd/Private/PlayLevel.cpp
@@ -1303,13 +1303,14 @@ void UEditorEngine::HandleStageStarted(const FString& InStage, TWeakPtr<SNotific
 	}
 	else if (InStage.Contains(TEXT("Build Task")))
 	{
+		EPlayOnBuildMode bBuildType = GetDefault<ULevelEditorPlaySettings>()->BuildGameBeforeLaunch;
 		FString PlatformName = PlayUsingLauncherDeviceId.Left(PlayUsingLauncherDeviceId.Find(TEXT("@")));
 		if (PlatformName.Contains(TEXT("NoEditor")))
 		{
 			PlatformName = PlatformName.Left(PlatformName.Find(TEXT("NoEditor")));
 		}
 		Arguments.Add(TEXT("PlatformName"), FText::FromString(PlatformName));
-		if (FRocketSupport::IsRocket() || !bPlayUsingLauncherHasCode || !bPlayUsingLauncherHasCompiler)
+		if (FRocketSupport::IsRocket() || !bPlayUsingLauncherHasCode || !bPlayUsingLauncherHasCompiler || bBuildType == EPlayOnBuildMode::PlayOnBuild_Never)
 		{
 			NotificationText = FText::Format(LOCTEXT("LauncherTaskValidateNotification", "Validating Executable for {PlatformName}..."), Arguments);
 		}
@@ -1457,15 +1458,15 @@ void UEditorEngine::PlayUsingLauncher()
 
 		// does the project have any code?
 		FGameProjectGenerationModule& GameProjectModule = FModuleManager::LoadModuleChecked<FGameProjectGenerationModule>(TEXT("GameProjectGeneration"));
-		bPlayUsingLauncherHasCode = GameProjectModule.Get().ProjectHasCodeFiles();
+		bPlayUsingLauncherHasCode = GameProjectModule.Get().ProjectRequiresBuild(FName(*PlayUsingLauncherDeviceId.Left(PlayUsingLauncherDeviceId.Find(TEXT("@")))));
 		bPlayUsingLauncherHasCompiler = FSourceCodeNavigation::IsCompilerAvailable();
 
 		// Setup launch profile, keep the setting here to a minimum.
 		ILauncherProfileRef LauncherProfile = LauncherServicesModule.CreateProfile(TEXT("Play On Device"));
-		bool bAlwaysBuild = GetDefault<ULevelEditorPlaySettings>()->BuildGameBeforeLaunch;
-		if ((bPlayUsingLauncherHasCode || bAlwaysBuild) && bPlayUsingLauncherHasCompiler)
+		EPlayOnBuildMode bBuildType = GetDefault<ULevelEditorPlaySettings>()->BuildGameBeforeLaunch;
+		if ((bBuildType == EPlayOnBuildMode::PlayOnBuild_Always) || (bBuildType == PlayOnBuild_Default && (bPlayUsingLauncherHasCode) && bPlayUsingLauncherHasCompiler))
 		{
-			LauncherProfile->SetBuildGame((bPlayUsingLauncherHasCode || bAlwaysBuild) && bPlayUsingLauncherHasCompiler);
+			LauncherProfile->SetBuildGame(true);
 
 			// set the build configuration to be the same as the running editor
 			FString ExeName = FUnrealEdMisc::Get().GetExecutableForCommandlets();
diff --git a/Engine/Source/Editor/UnrealEd/Private/Settings/SettingsClasses.cpp b/Engine/Source/Editor/UnrealEd/Private/Settings/SettingsClasses.cpp
--- a/Engine/Source/Editor/UnrealEd/Private/Settings/SettingsClasses.cpp
+++ b/Engine/Source/Editor/UnrealEd/Private/Settings/SettingsClasses.cpp
@@ -11,6 +11,8 @@
 #include "NotificationManager.h"
 #include "ISettingsModule.h"
 
+#include "SourceCodeNavigation.h"
+
 #define LOCTEXT_NAMESPACE "SettingsClasses"
 
 /* UContentBrowserSettings interface
@@ -346,8 +348,16 @@ ULevelEditorPlaySettings::ULevelEditorPlaySettings( const FObjectInitializer& Ob
 	PlayNetDedicated = false;
 	RunUnderOneProcess = true;
 	RouteGamepadToSecondWindow = false;
+	BuildGameBeforeLaunch = EPlayOnBuildMode::PlayOnBuild_Default;
 }
 
+void ULevelEditorPlaySettings::PostEditChangeProperty(struct FPropertyChangedEvent& PropertyChangedEvent)
+{
+	if (BuildGameBeforeLaunch != EPlayOnBuildMode::PlayOnBuild_Always && !FSourceCodeNavigation::IsCompilerAvailable())
+	{
+		BuildGameBeforeLaunch == EPlayOnBuildMode::PlayOnBuild_Never;
+	}
+}
 /* ULevelEditorViewportSettings interface
  *****************************************************************************/
 
