Editor: Fixing non-unity build

diff --git a/Engine/Source/Editor/UnrealEd/Private/Lightmass/LightmassRender.cpp b/Engine/Source/Editor/UnrealEd/Private/Lightmass/LightmassRender.cpp
--- a/Engine/Source/Editor/UnrealEd/Private/Lightmass/LightmassRender.cpp
+++ b/Engine/Source/Editor/UnrealEd/Private/Lightmass/LightmassRender.cpp
@@ -15,6 +15,8 @@
 #include "Materials/MaterialInstance.h"
 #include "CanvasTypes.h"
 #include "LightmassLandscapeRender.h"
+#include "LandscapeMaterialInstanceConstant.h"
+#include "LandscapeLight.h"
 
 DEFINE_LOG_CATEGORY_STATIC(LogLightmassRender, Error, All);
 
diff --git a/Engine/Source/Editor/UnrealEd/Private/Tests/AutomationEditorPromotionCommon.cpp b/Engine/Source/Editor/UnrealEd/Private/Tests/AutomationEditorPromotionCommon.cpp
--- a/Engine/Source/Editor/UnrealEd/Private/Tests/AutomationEditorPromotionCommon.cpp
+++ b/Engine/Source/Editor/UnrealEd/Private/Tests/AutomationEditorPromotionCommon.cpp
@@ -2,6 +2,7 @@
 
 #include "UnrealEd.h"
 #include "AutomationEditorPromotionCommon.h"
+#include "AssetRegistryModule.h"
 
 /**
 * Gets the base path for this asset
