CIS fix

diff --git a/Engine/Source/Editor/UnrealEd/Private/Settings/SettingsClasses.cpp b/Engine/Source/Editor/UnrealEd/Private/Settings/SettingsClasses.cpp
--- a/Engine/Source/Editor/UnrealEd/Private/Settings/SettingsClasses.cpp
+++ b/Engine/Source/Editor/UnrealEd/Private/Settings/SettingsClasses.cpp
@@ -355,7 +355,7 @@ void ULevelEditorPlaySettings::PostEditChangeProperty(struct FPropertyChangedEve
 {
 	if (BuildGameBeforeLaunch != EPlayOnBuildMode::PlayOnBuild_Always && !FSourceCodeNavigation::IsCompilerAvailable())
 	{
-		BuildGameBeforeLaunch == EPlayOnBuildMode::PlayOnBuild_Never;
+		BuildGameBeforeLaunch = EPlayOnBuildMode::PlayOnBuild_Never;
 	}
 }
 /* ULevelEditorViewportSettings interface
