Fixed issue with iterate not correctly filling the list of previous cooked files resulting in more then nessisary cooked files.

diff --git a/Engine/Source/Editor/UnrealEd/Private/CookOnTheFlyServer.cpp b/Engine/Source/Editor/UnrealEd/Private/CookOnTheFlyServer.cpp
--- a/Engine/Source/Editor/UnrealEd/Private/CookOnTheFlyServer.cpp
+++ b/Engine/Source/Editor/UnrealEd/Private/CookOnTheFlyServer.cpp
@@ -246,7 +246,7 @@ void OutputTimers()
 
 		OutputValues[Index] = FString::Printf(TEXT("%f"),TimerInfo.Length);
 	}
-	static FString NewLine = FString(TEXT("\n"));
+	static FString NewLine = FString(TEXT("\r\n"));
 
 	if (OutputTimerIndexMap)
 	{
@@ -2496,20 +2496,52 @@ void UCookOnTheFlyServer::PopulateCookedPackagesFromDisk( const TArray<ITargetPl
 	for (int32 Index = 0; Index < Platforms.Num(); Index++)
 	{
 		ITargetPlatform* Target = Platforms[Index];
-		FString SandboxDirectory = GetSandboxDirectory(Target->PlatformName());
+		FString SandboxPath = GetSandboxDirectory(Target->PlatformName());
 		FName PlatformFName( *Target->PlatformName() );
+
+		FString EngineSandboxPath= SandboxFile->ConvertToSandboxPath(*FPaths::EngineDir()) + TEXT("/");
+		EngineSandboxPath.ReplaceInline(TEXT("[Platform]"), *Target->PlatformName());
+
+		FString GameSandboxPath = SandboxFile->ConvertToSandboxPath(*(FPaths::GameDir() + TEXT("a.txt")));
+		GameSandboxPath.ReplaceInline(TEXT("a.txt"), TEXT(""));
+		GameSandboxPath.ReplaceInline(TEXT("[Platform]"), *Target->PlatformName());
+
+		FString LocalGamePath = FPaths::GameDir();
+		if (FPaths::IsProjectFilePathSet())
+		{
+			LocalGamePath = FPaths::GetPath(FPaths::GetProjectFilePath()) + TEXT("/");
+		}
+
+		FString LocalEnginePath = FPaths::EngineDir();
+
 		// use the timestamp grabbing visitor
 		IPlatformFile& PlatformFile = FPlatformFileManager::Get().GetPlatformFile();
 		FLocalTimestampDirectoryVisitor Visitor(PlatformFile, DirectoriesToSkip, DirectoriesToNotRecurse, false);
 
-		PlatformFile.IterateDirectory(*SandboxDirectory, Visitor);
+		PlatformFile.IterateDirectory(*SandboxPath, Visitor);
+
+		/*FString SandboxEngine = Sandbox->ConvertToSandboxPath(*LocalEngineDir) + TEXT("/");
+		// we need to add an extra bit to the game path to make the sandbox convert it correctly (investigate?)
+		// @todo: double check this
+		FString SandboxGame = Sandbox->ConvertToSandboxPath(*(LocalGameDir + TEXT("a.txt"))).Replace(TEXT("a.txt"), TEXT(""));*/
+		
 
 		for (TMap<FString, FDateTime>::TIterator TimestampIt(Visitor.FileTimes); TimestampIt; ++TimestampIt)
 		{
 			FString CookedFilename = TimestampIt.Key();
+			if ((CookedFilename.EndsWith(TEXT(".uasset")) == false) &&
+				(CookedFilename.EndsWith(TEXT(".umap")) == false))
+			{
+				// don't care if it's not a map or an asset
+				continue;
+			}
 			FDateTime CookedTimestamp = TimestampIt.Value();
 			CookedFilename = FPaths::ConvertRelativePathToFull(CookedFilename);
-			FString StandardCookedFilename = CookedFilename.Replace(*SandboxDirectory, *(FPaths::GetRelativePathToRoot()));
+			FString StandardCookedFilename = CookedFilename;
+
+			StandardCookedFilename.ReplaceInline(*EngineSandboxPath, *LocalEnginePath);
+			StandardCookedFilename.ReplaceInline(*GameSandboxPath, *LocalGamePath);
+			StandardCookedFilename.ReplaceInline(*SandboxPath, *(FPaths::GetRelativePathToRoot()));
 			FDateTime DependentTimestamp;
 
 			if (PDInfoModule.DeterminePackageDependentTimeStamp(*(FPaths::GetBaseFilename(StandardCookedFilename, false)), DependentTimestamp) == true)
@@ -2593,19 +2625,50 @@ void UCookOnTheFlyServer::CleanSandbox( const bool bIterative )
 			for (int32 Index = 0; Index < Platforms.Num(); Index++)
 			{
 				ITargetPlatform* Target = Platforms[Index];
-				FString SandboxDirectory = GetSandboxDirectory(Target->PlatformName());
+				
+				FString EngineSandboxPath = SandboxFile->ConvertToSandboxPath(*FPaths::EngineDir()) + TEXT("/");
+				EngineSandboxPath.ReplaceInline(TEXT("[Platform]"), *Target->PlatformName());
+
+				FString GameSandboxPath = SandboxFile->ConvertToSandboxPath(*(FPaths::GameDir() + TEXT("a.txt")));
+				GameSandboxPath.ReplaceInline(TEXT("a.txt"), TEXT(""));
+				GameSandboxPath.ReplaceInline(TEXT("[Platform]"), *Target->PlatformName());
+
+				FString LocalGamePath = FPaths::GameDir();
+				if (FPaths::IsProjectFilePathSet())
+				{
+					LocalGamePath = FPaths::GetPath(FPaths::GetProjectFilePath()) + TEXT("/");
+				}
+
+				FString LocalEnginePath = FPaths::EngineDir();
+
+
 				FName PlatformFName( *Target->PlatformName() );
 				// use the timestamp grabbing visitor
 				IPlatformFile& PlatformFile = FPlatformFileManager::Get().GetPlatformFile();
 				FLocalTimestampDirectoryVisitor Visitor(PlatformFile, DirectoriesToSkip, DirectoriesToNotRecurse, false);
 			
-				PlatformFile.IterateDirectory(*SandboxDirectory, Visitor);
+				FString SandboxPath = GetSandboxDirectory(Target->PlatformName());
+
+				PlatformFile.IterateDirectory(*SandboxPath, Visitor);
 
 				for (TMap<FString, FDateTime>::TIterator TimestampIt(Visitor.FileTimes); TimestampIt; ++TimestampIt)
 				{
 					FString CookedFilename = TimestampIt.Key();
+					if ((CookedFilename.EndsWith(TEXT(".uasset")) == false) &&
+						(CookedFilename.EndsWith(TEXT(".umap")) == false))
+					{
+						// don't care if it's not a map or an asset
+						continue;
+					}
+
 					FDateTime CookedTimestamp = TimestampIt.Value();
-					FString StandardCookedFilename = CookedFilename.Replace(*SandboxDirectory, *(FPaths::GetRelativePathToRoot()));
+					CookedFilename = FPaths::ConvertRelativePathToFull(CookedFilename);
+					FString StandardCookedFilename = CookedFilename;
+
+					StandardCookedFilename.ReplaceInline(*EngineSandboxPath, *LocalEnginePath);
+					StandardCookedFilename.ReplaceInline(*GameSandboxPath, *LocalGamePath);
+					StandardCookedFilename.ReplaceInline(*SandboxPath, *(FPaths::GetRelativePathToRoot()));
+
 					FDateTime DependentTimestamp;
 					FName StandardCookedFileFName = FName(*StandardCookedFilename);
 
@@ -3362,6 +3425,7 @@ void UCookOnTheFlyServer::StartCookByTheBook( const FCookByTheBookStartupOptions
 			if (IsCookFlagSet(ECookInitializationFlags::Iterative) && (PackagesKeptFromPreviousCook.Num() > 0) )
 			{
 				const FString& SandboxRegistryFilename = GetSandboxAssetRegistryFilename();
+
 				ManifestGenerator->LoadAssetRegistry(GetCookedAssetRegistryFilename(PlatformName.ToString()), &PackagesKeptFromPreviousCook);
 			}
 
