Fixed issue with asset registry loading using incorrect archiver in the cooker in iterative cooking mode, causes crash when using iterative cooking.

diff --git a/Engine/Source/Editor/UnrealEd/Private/Commandlets/ChunkManifestGenerator.cpp b/Engine/Source/Editor/UnrealEd/Private/Commandlets/ChunkManifestGenerator.cpp
--- a/Engine/Source/Editor/UnrealEd/Private/Commandlets/ChunkManifestGenerator.cpp
+++ b/Engine/Source/Editor/UnrealEd/Private/Commandlets/ChunkManifestGenerator.cpp
@@ -416,12 +416,20 @@ bool FChunkManifestGenerator::LoadAssetRegistry(const FString& SandboxPath, cons
 
 	for (auto Platform : Platforms)
 	{
+		/*FString PlatformSandboxPath = SandboxPath.Replace(TEXT("[Platform]"), *Platform->PlatformName());
+		FArchive* AssetRegistryReader = IFileManager::Get().CreateFileReader(*PlatformSandboxPath);*/
+
 		FString PlatformSandboxPath = SandboxPath.Replace(TEXT("[Platform]"), *Platform->PlatformName());
-		FArchive* AssetRegistryReader = IFileManager::Get().CreateFileReader(*PlatformSandboxPath);
+		FArrayReader FileContents;
+		if (FFileHelper::LoadFileToArray(FileContents, *PlatformSandboxPath) == false)
+		{
+			continue;
+		}
+		FArchive* AssetRegistryReader = &FileContents;
+
 		TMap<FName, FAssetData*> SavedAssetRegistryData;
 		if (AssetRegistryReader)
 		{
-			
 			AssetRegistry.LoadRegistryData(*AssetRegistryReader, SavedAssetRegistryData);
 		}
 		for (FAssetData LoadedAssetData : AssetRegistryData)
