messaging for when instruments launch fails
UE-12395, UE-9979
#ios

diff --git a/Engine/Source/Editor/UnrealEd/Private/Analytics/EditorAnalytics.cpp b/Engine/Source/Editor/UnrealEd/Private/Analytics/EditorAnalytics.cpp
--- a/Engine/Source/Editor/UnrealEd/Private/Analytics/EditorAnalytics.cpp
+++ b/Engine/Source/Editor/UnrealEd/Private/Analytics/EditorAnalytics.cpp
@@ -168,6 +168,12 @@ FString FEditorAnalytics::TranslateErrorCode(int32 ErrorCode)
 		return TEXT("Failed to delete staging directory.  This could be because something is currently using the staging directory (ps4/xbox/etc)");
 	case EAnalyticsErrorCodes::MissingExecutable:
 		return LOCTEXT("UATErrorMissingExecutable", "Missing UE4Game binary.\nYou may have to build the UE4 project with your IDE. Alternatively, build using UnrealBuildTool with the commandline:\nUE4Game <Platform> <Configuration>").ToString();
+	case EAnalyticsErrorCodes::FilesInstallFailed:
+		return TEXT("Failed to deploy files to device.  Check to make sure your device is connected.");
+	case EAnalyticsErrorCodes::DeviceNotSetupForDevelopment:
+		return TEXT("Failed to launch on device.  Make sure your device has been enabled for development from within the Xcode Devices window.");
+	case EAnalyticsErrorCodes::DeviceOSNewerThanSDK:
+		return TEXT("Failed to launch on device.  Make sure your install of Xcode matches or is newer than the OS on your device.");
 	}
 	return TEXT("Unknown Error");
 }
diff --git a/Engine/Source/Editor/UnrealEd/Public/EditorAnalytics.h b/Engine/Source/Editor/UnrealEd/Public/EditorAnalytics.h
--- a/Engine/Source/Editor/UnrealEd/Public/EditorAnalytics.h
+++ b/Engine/Source/Editor/UnrealEd/Public/EditorAnalytics.h
@@ -53,10 +53,13 @@ namespace EAnalyticsErrorCodes
 		ObbNotFound = 53,
 		AndroidBuildToolsPathNotFound = 54,
 		NoApkSuitableForArchitecture = 55,
+		FilesInstallFailed = 56,
 		LauncherFailed = 100,
 		UATLaunchFailure = 101,
 		FailedToDeleteStagingDirectory = 102,
 		MissingExecutable = 103,
+		DeviceNotSetupForDevelopment = 150,
+		DeviceOSNewerThanSDK = 151,
 	};
 };
 
diff --git a/Engine/Source/Programs/AutomationTool/IOS/IOSPlatform.Automation.cs b/Engine/Source/Programs/AutomationTool/IOS/IOSPlatform.Automation.cs
--- a/Engine/Source/Programs/AutomationTool/IOS/IOSPlatform.Automation.cs
+++ b/Engine/Source/Programs/AutomationTool/IOS/IOSPlatform.Automation.cs
@@ -820,6 +820,17 @@ public class IOSPlatform : Platform
 			{
 				Directory.Delete (Params.BaseStageDirectory + "/IOS/launch.trace", true);
 			}
+
+			switch (Result.ExitCode)
+			{
+				case 253:
+					ErrorReporter.Error("Launch Failure", (int)ErrorCodes.Error_DeviceNotSetupForDevelopment);
+					break;
+
+				case 255:
+					ErrorReporter.Error("Launch Failure", (int)ErrorCodes.Error_DeviceOSNewerThanSDK);
+					break;
+			}
 		}
 	}
 
diff --git a/Engine/Source/Programs/AutomationTool/Program.cs b/Engine/Source/Programs/AutomationTool/Program.cs
--- a/Engine/Source/Programs/AutomationTool/Program.cs
+++ b/Engine/Source/Programs/AutomationTool/Program.cs
@@ -54,10 +54,13 @@ namespace AutomationTool
         Error_ObbNotFound = 53,
         Error_AndroidBuildToolsPathNotFound = 54,
         Error_NoApkSuitableForArchitecture = 55,
-        Error_LauncherFailed = 100,
+		Error_FilesInstallFailed = 56,
+		Error_LauncherFailed = 100,
 		Error_UATLaunchFailure = 101,
 		Error_FailedToDeleteStagingDirectory = 102,
 		Error_MissingExecutable = 103,
+		Error_DeviceNotSetupForDevelopment = 150,
+		Error_DeviceOSNewerThanSDK = 151,
 	};
 
 	public class ErrorReporter
diff --git a/Engine/Source/Programs/IOS/iPhonePackager/iPhonePackager.cs b/Engine/Source/Programs/IOS/iPhonePackager/iPhonePackager.cs
--- a/Engine/Source/Programs/IOS/iPhonePackager/iPhonePackager.cs
+++ b/Engine/Source/Programs/IOS/iPhonePackager/iPhonePackager.cs
@@ -65,6 +65,9 @@ namespace iPhonePackager
 		Error_LauncherFailed = 100,
         Error_UATLaunchFailure = 101,
         Error_FailedToDeleteStagingDirectory = 102,
+		Error_MissingExecutable = 103,
+		Error_DeviceNotSetupForDevelopment = 150,
+		Error_DeviceOSNewerThanSDK = 151,
 	};
 
 	public partial class Program
