Fix error messaging to match the check being performed.

diff --git a/Engine/Source/Editor/UnrealEd/Private/Commandlets/CookCommandlet.cpp b/Engine/Source/Editor/UnrealEd/Private/Commandlets/CookCommandlet.cpp
--- a/Engine/Source/Editor/UnrealEd/Private/Commandlets/CookCommandlet.cpp
+++ b/Engine/Source/Editor/UnrealEd/Private/Commandlets/CookCommandlet.cpp
@@ -455,7 +455,7 @@ bool UCookCommandlet::SaveCookedPackage( UPackage* Package, uint32 SaveFlags, bo
 				const FString FullFilename = FPaths::ConvertRelativePathToFull( PlatFilename );
 				if( FullFilename.Len() >= PLATFORM_MAX_FILEPATH_LENGTH )
 				{
-					UE_LOG( LogCookCommandlet, Error, TEXT( "Couldn't save package, filename is too long :%s" ), *PlatFilename );
+					UE_LOG( LogCookCommandlet, Error, TEXT( "Couldn't save package, filename is too long :%s" ), *FullFilename );
 					bSavedCorrectly = false;
 				}
 				else
