Paper2D: Fixed the 2D snap grid settings not being saved to a persistent ini (now saved to DefaultEditor.ini) [UE-15412]
- Moved transient settings (whether the snap planes are currently enabled over to per-user ULevelEditorViewportSettings
- Fixed several string literals in ULevelEditorViewportSettings to use GET_MEMBER_NAME_CHECKED instead (some TBD)
- Added three default layers at -100, 0, and +100

diff --git a/Engine/Source/Editor/UnrealEd/Classes/Settings/LevelEditorViewportSettings.h b/Engine/Source/Editor/UnrealEd/Classes/Settings/LevelEditorViewportSettings.h
--- a/Engine/Source/Editor/UnrealEd/Classes/Settings/LevelEditorViewportSettings.h
+++ b/Engine/Source/Editor/UnrealEd/Classes/Settings/LevelEditorViewportSettings.h
@@ -303,6 +303,14 @@ private:
 
 public:
 
+	/** If enabled, actor rotations will snap to the grid. */
+	UPROPERTY(EditAnywhere, config, Category=GridSnapping, meta=(DisplayName = "Enable 2D Layer Snapping"))
+	uint32 bEnableLayerSnap:1;
+
+	/** The index of the snap plane to use when bEnableLayerSnap is true (from the project SnapLayers array) */
+	UPROPERTY(config)
+	int32 ActiveSnapLayerIndex;
+
 	/** If true actor snap will be enabled in the editor **/
 	UPROPERTY(config, Category=GridSnapping, VisibleDefaultsOnly,AdvancedDisplay)
 	uint32 bEnableActorSnap:1;
diff --git a/Engine/Source/Editor/UnrealEd/Private/LevelEditorViewport.cpp b/Engine/Source/Editor/UnrealEd/Private/LevelEditorViewport.cpp
--- a/Engine/Source/Editor/UnrealEd/Private/LevelEditorViewport.cpp
+++ b/Engine/Source/Editor/UnrealEd/Private/LevelEditorViewport.cpp
@@ -990,16 +990,18 @@ bool FLevelEditorViewportClient::HasDropPreviewActors() const
 
 static bool IsDroppingOn2DLayer()
 {
+	const ULevelEditorViewportSettings* ViewportSettings = GetDefault<ULevelEditorViewportSettings>();
 	const ULevelEditor2DSettings* Settings2D = GetDefault<ULevelEditor2DSettings>();
-	return Settings2D->bEnableLayerSnap && Settings2D->SnapLayers.IsValidIndex(Settings2D->ActiveSnapLayerIndex);
+	return ViewportSettings->bEnableLayerSnap && Settings2D->SnapLayers.IsValidIndex(ViewportSettings->ActiveSnapLayerIndex);
 }
 
 static FActorPositionTraceResult TraceForPositionOn2DLayer(const FViewportCursorLocation& Cursor)
 {
+	const ULevelEditorViewportSettings* ViewportSettings = GetDefault<ULevelEditorViewportSettings>();
 	const ULevelEditor2DSettings* Settings2D = GetDefault<ULevelEditor2DSettings>();
-	check(Settings2D->SnapLayers.IsValidIndex(Settings2D->ActiveSnapLayerIndex));
+	check(Settings2D->SnapLayers.IsValidIndex(ViewportSettings->ActiveSnapLayerIndex));
 
-	float Offset = Settings2D->SnapLayers[Settings2D->ActiveSnapLayerIndex].Depth;
+	const float Offset = Settings2D->SnapLayers[ViewportSettings->ActiveSnapLayerIndex].Depth;
 	FVector PlaneCenter(0, 0, 0);
 	FVector PlaneNormal(0, 0, 0);
 
@@ -1011,8 +1013,8 @@ static FActorPositionTraceResult TraceForPositionOn2DLayer(const FViewportCursor
 	}
 
 	FActorPositionTraceResult Result;
-	float Numerator = FVector::DotProduct(PlaneCenter - Cursor.GetOrigin(), PlaneNormal);
-	float Denominator = FVector::DotProduct(PlaneNormal, Cursor.GetDirection());
+	const float Numerator = FVector::DotProduct(PlaneCenter - Cursor.GetOrigin(), PlaneNormal);
+	const float Denominator = FVector::DotProduct(PlaneNormal, Cursor.GetDirection());
 	if (FMath::Abs(Denominator) < SMALL_NUMBER)
 	{
 		Result.State = FActorPositionTraceResult::Failed;
diff --git a/Engine/Source/Editor/UnrealEd/Private/STransformViewportToolbar.cpp b/Engine/Source/Editor/UnrealEd/Private/STransformViewportToolbar.cpp
--- a/Engine/Source/Editor/UnrealEd/Private/STransformViewportToolbar.cpp
+++ b/Engine/Source/Editor/UnrealEd/Private/STransformViewportToolbar.cpp
@@ -417,10 +417,11 @@ FText STransformViewportToolBar::GetRotationGridLabel() const
 
 FText STransformViewportToolBar::GetLayer2DLabel() const
 {
+	const ULevelEditorViewportSettings* ViewportSettings = GetDefault<ULevelEditorViewportSettings>();
 	const ULevelEditor2DSettings* Settings2D = GetDefault<ULevelEditor2DSettings>();
-	if (Settings2D->SnapLayers.IsValidIndex(Settings2D->ActiveSnapLayerIndex))
+	if (Settings2D->SnapLayers.IsValidIndex(ViewportSettings->ActiveSnapLayerIndex))
 	{
-		return FText::FromString(Settings2D->SnapLayers[Settings2D->ActiveSnapLayerIndex].Name);
+		return FText::FromString(Settings2D->SnapLayers[ViewportSettings->ActiveSnapLayerIndex].Name);
 	}
 	
 	return FText();
@@ -505,16 +506,16 @@ void STransformViewportToolBar::SetScaleGridSize( int32 InIndex )
 }
 
 /**
-	* Sets the scale grid size
-	*
-	* @param	InIndex	The new index of the scale grid size to use
-	*/
+ * Sets the active 2d snap layer 
+ *
+ * @param	InIndex	The new index of the 2d layer to use
+ */
 void STransformViewportToolBar::SetLayer2D( int32 Layer2DIndex )
 {
-	ULevelEditor2DSettings* Settings2D = GetMutableDefault<ULevelEditor2DSettings>();
-	Settings2D->bEnableLayerSnap = true;
-	Settings2D->ActiveSnapLayerIndex = Layer2DIndex;
-	Settings2D->PostEditChange();
+	ULevelEditorViewportSettings* ViewportSettings = GetMutableDefault<ULevelEditorViewportSettings>();
+	ViewportSettings->bEnableLayerSnap = true;
+	ViewportSettings->ActiveSnapLayerIndex = Layer2DIndex;
+	ViewportSettings->PostEditChange();
 }
 
 		
@@ -562,7 +563,8 @@ bool STransformViewportToolBar::IsScaleGridSizeChecked(int32 GridSizeIndex)
 
 bool STransformViewportToolBar::IsLayer2DSelected(int32 LayerIndex)
 {
-	return (GetDefault<ULevelEditor2DSettings>()->ActiveSnapLayerIndex == LayerIndex);
+	const ULevelEditorViewportSettings* ViewportSettings = GetDefault<ULevelEditorViewportSettings>();
+	return (ViewportSettings->ActiveSnapLayerIndex == LayerIndex);
 }
 
 void STransformViewportToolBar::TogglePreserveNonUniformScale()
@@ -629,6 +631,7 @@ TSharedRef<SWidget> STransformViewportToolBar::FillRotationGridSnapMenu()
 
 TSharedRef<SWidget> STransformViewportToolBar::FillLayer2DSnapMenu()
 {
+	const ULevelEditorViewportSettings* ViewportSettings = GetDefault<ULevelEditorViewportSettings>();
 	const ULevelEditor2DSettings* Settings2D = GetDefault<ULevelEditor2DSettings>();
 	int32 LayerCount = Settings2D->SnapLayers.Num();
 	const bool bInShouldCloseWindowAfterMenuSelection = true;
@@ -654,14 +657,15 @@ TSharedRef<SWidget> STransformViewportToolBar::FillLayer2DSnapMenu()
 			}
 		}
 
-		static FReply SnapSelectionToSelectedOffset(int32 SnapLayerIndex)
+		static FReply SnapSelectionToSelectedOffset()
 		{
-			auto Settings = GetDefault<ULevelEditor2DSettings>();
-			if (Settings->SnapLayers.IsValidIndex(SnapLayerIndex))
+			const ULevelEditorViewportSettings* ViewportSettings = GetDefault<ULevelEditorViewportSettings>();
+			const ULevelEditor2DSettings* Settings2D = GetDefault<ULevelEditor2DSettings>();
+			if (Settings2D->SnapLayers.IsValidIndex(ViewportSettings->ActiveSnapLayerIndex))
 			{
 				const FScopedTransaction Transaction(NSLOCTEXT("UnrealEd", "SnapSelection2D", "Snap Selection to 2D Layer"));
 
-				float SnapDepth = Settings->SnapLayers[Settings->ActiveSnapLayerIndex].Depth;
+				float SnapDepth = Settings2D->SnapLayers[ViewportSettings->ActiveSnapLayerIndex].Depth;
 				USelection* SelectedActors = GEditor->GetSelectedActors();
 				for (FSelectionIterator Iter(*SelectedActors); Iter; ++Iter)
 				{
@@ -673,7 +677,7 @@ TSharedRef<SWidget> STransformViewportToolBar::FillLayer2DSnapMenu()
 						FTransform Transform = Actor->GetTransform();
 						FVector CurrentLocation = Transform.GetLocation();
 
-						switch (Settings->SnapAxis)
+						switch (Settings2D->SnapAxis)
 						{
 						case ELevelEditor2DAxis::X: CurrentLocation.X = SnapDepth; break;
 						case ELevelEditor2DAxis::Y: CurrentLocation.Y = SnapDepth; break;
@@ -694,8 +698,8 @@ TSharedRef<SWidget> STransformViewportToolBar::FillLayer2DSnapMenu()
 
 		static bool IsSnapSelectedActorsEnabled()
 		{
-			auto Settings = GetDefault<ULevelEditor2DSettings>();
-			return Settings->SnapLayers.IsValidIndex(Settings->ActiveSnapLayerIndex) && GEditor->GetSelectedActorCount() > 0;
+			const ULevelEditor2DSettings* Settings = GetDefault<ULevelEditor2DSettings>();
+			return Settings->SnapLayers.IsValidIndex(GetDefault<ULevelEditorViewportSettings>()->ActiveSnapLayerIndex) && (GEditor->GetSelectedActorCount() > 0);
 		}
 	};
 
@@ -709,7 +713,7 @@ TSharedRef<SWidget> STransformViewportToolBar::FillLayer2DSnapMenu()
 		SNew(SButton)
 		.Text(LOCTEXT("2DSnap_SnapSelection", "Snap Selected Actors"))
 		.IsEnabled(TAttribute<bool>::Create(TAttribute<bool>::FGetter::CreateStatic(FLocalFunctions::IsSnapSelectedActorsEnabled)))
-		.OnClicked(FOnClicked::CreateStatic(FLocalFunctions::SnapSelectionToSelectedOffset, Settings2D->ActiveSnapLayerIndex))
+		.OnClicked(FOnClicked::CreateStatic(FLocalFunctions::SnapSelectionToSelectedOffset))
 		.VAlign(VAlign_Center)
 		.HAlign(HAlign_Fill),
 
@@ -818,8 +822,9 @@ ECheckBoxState STransformViewportToolBar::IsRotationGridSnapChecked() const
 
 ECheckBoxState STransformViewportToolBar::IsLayer2DSnapChecked() const
 {
+	const ULevelEditorViewportSettings* ViewportSettings = GetDefault<ULevelEditorViewportSettings>();
 	const ULevelEditor2DSettings* Settings2D = GetDefault<ULevelEditor2DSettings>();
-	bool bChecked = Settings2D->bEnableLayerSnap && Settings2D->SnapLayers.IsValidIndex(Settings2D->ActiveSnapLayerIndex);
+	const bool bChecked = ViewportSettings->bEnableLayerSnap && Settings2D->SnapLayers.IsValidIndex(ViewportSettings->ActiveSnapLayerIndex);
 	return bChecked ? ECheckBoxState::Checked : ECheckBoxState::Unchecked;
 }
 
@@ -846,17 +851,18 @@ void STransformViewportToolBar::HandleToggleRotationGridSnap(ECheckBoxState InSt
 
 void STransformViewportToolBar::HandleToggleLayer2DSnap(ECheckBoxState InState)
 {
-	ULevelEditor2DSettings* Settings2D = GetMutableDefault<ULevelEditor2DSettings>();
-	if (!Settings2D->bEnableLayerSnap && Settings2D->SnapLayers.Num() > 0)
+	ULevelEditorViewportSettings* ViewportSettings = GetMutableDefault<ULevelEditorViewportSettings>();
+	const ULevelEditor2DSettings* Settings2D = GetDefault<ULevelEditor2DSettings>();
+	if (!ViewportSettings->bEnableLayerSnap && (Settings2D->SnapLayers.Num() > 0))
 	{
-		Settings2D->bEnableLayerSnap = true;
-		Settings2D->ActiveSnapLayerIndex = FMath::Clamp(Settings2D->ActiveSnapLayerIndex, 0, Settings2D->SnapLayers.Num() - 1);
+		ViewportSettings->bEnableLayerSnap = true;
+		ViewportSettings->ActiveSnapLayerIndex = FMath::Clamp(ViewportSettings->ActiveSnapLayerIndex, 0, Settings2D->SnapLayers.Num() - 1);
 	}
 	else
 	{
-		Settings2D->bEnableLayerSnap = false;
+		ViewportSettings->bEnableLayerSnap = false;
 	}
-	Settings2D->PostEditChange();
+	ViewportSettings->PostEditChange();
 }
 
 void STransformViewportToolBar::HandleToggleScaleGridSnap(ECheckBoxState InState)
diff --git a/Engine/Source/Editor/UnrealEd/Private/Settings/EditorProjectSettings.cpp b/Engine/Source/Editor/UnrealEd/Private/Settings/EditorProjectSettings.cpp
--- a/Engine/Source/Editor/UnrealEd/Private/Settings/EditorProjectSettings.cpp
+++ b/Engine/Source/Editor/UnrealEd/Private/Settings/EditorProjectSettings.cpp
@@ -142,8 +142,10 @@ void UEditorProjectAppearanceSettings::PostInitProperties()
 *****************************************************************************/
 
 ULevelEditor2DSettings::ULevelEditor2DSettings(const FObjectInitializer& ObjectInitializer)
-	: Super(ObjectInitializer),
-	  SnapAxis(ELevelEditor2DAxis::Y)
+	: Super(ObjectInitializer)
+	, SnapAxis(ELevelEditor2DAxis::Y)
 {
+	SnapLayers.Emplace(TEXT("Foreground"), 100.0f);
+	SnapLayers.Emplace(TEXT("Default"), 0.0f);
+	SnapLayers.Emplace(TEXT("Background"), -100.0f);
 }
-
diff --git a/Engine/Source/Editor/UnrealEd/Private/Settings/SettingsClasses.cpp b/Engine/Source/Editor/UnrealEd/Private/Settings/SettingsClasses.cpp
--- a/Engine/Source/Editor/UnrealEd/Private/Settings/SettingsClasses.cpp
+++ b/Engine/Source/Editor/UnrealEd/Private/Settings/SettingsClasses.cpp
@@ -388,7 +388,7 @@ void ULevelEditorViewportSettings::PostEditChangeProperty( struct FPropertyChang
 		? PropertyChangedEvent.Property->GetFName()
 		: NAME_None;
 
-	if (Name == FName(TEXT("bAllowTranslateRotateZWidget")))
+	if (Name == GET_MEMBER_NAME_CHECKED(ULevelEditorViewportSettings, bAllowTranslateRotateZWidget))
 	{
 		if (bAllowTranslateRotateZWidget)
 		{
@@ -399,31 +399,31 @@ void ULevelEditorViewportSettings::PostEditChangeProperty( struct FPropertyChang
 			GLevelEditorModeTools().SetWidgetMode(FWidget::WM_Translate);
 		}
 	}
-	else if (Name == FName(TEXT("bHighlightWithBrackets")))
+	else if (Name == GET_MEMBER_NAME_CHECKED(ULevelEditorViewportSettings, bHighlightWithBrackets))
 	{
 		GEngine->SetSelectedMaterialColor(bHighlightWithBrackets
 			? FLinearColor::Black
 			: GetDefault<UEditorStyleSettings>()->SelectionColor);
 	}
-	else if (Name == FName(TEXT("HoverHighlightIntensity")))
+	else if (Name == GET_MEMBER_NAME_CHECKED(ULevelEditorViewportSettings, HoverHighlightIntensity))
 	{
 		GEngine->HoverHighlightIntensity = HoverHighlightIntensity;
 	}
-	else if (Name == FName(TEXT("SelectionHighlightIntensity")))
+	else if (Name == GET_MEMBER_NAME_CHECKED(ULevelEditorViewportSettings, SelectionHighlightIntensity))
 	{
 		GEngine->SelectionHighlightIntensity = SelectionHighlightIntensity;
 	}
-	else if (Name == FName(TEXT("BSPSelectionHighlightIntensity")))
+	else if (Name == GET_MEMBER_NAME_CHECKED(ULevelEditorViewportSettings, BSPSelectionHighlightIntensity))
 	{
 		GEngine->BSPSelectionHighlightIntensity = BSPSelectionHighlightIntensity;
 	}
-	else if ((Name == FName(TEXT("UserDefinedPosGridSizes"))) || (Name == FName(TEXT("UserDefinedRotGridSizes"))) || (Name == FName(TEXT("ScalingGridSizes"))) || (Name == FName(TEXT("GridIntervals"))))
+	else if ((Name == FName(TEXT("UserDefinedPosGridSizes"))) || (Name == FName(TEXT("UserDefinedRotGridSizes"))) || (Name == FName(TEXT("ScalingGridSizes"))) || (Name == FName(TEXT("GridIntervals")))) //@TODO: This should use GET_MEMBER_NAME_CHECKED
 	{
-		const float MinGridSize = (Name == FName(TEXT("GridIntervals"))) ? 4.0f : 0.0001f;
+		const float MinGridSize = (Name == FName(TEXT("GridIntervals"))) ? 4.0f : 0.0001f; //@TODO: This should use GET_MEMBER_NAME_CHECKED
 		TArray<float>* ArrayRef = nullptr;
 		int32* IndexRef = nullptr;
 
-		if (Name == FName(TEXT("ScalingGridSizes")))
+		if (Name == GET_MEMBER_NAME_CHECKED(ULevelEditorViewportSettings, ScalingGridSizes))
 		{
 			ArrayRef = &(ScalingGridSizes);
 			IndexRef = &(CurrentScalingGridSize);
@@ -444,16 +444,24 @@ void ULevelEditorViewportSettings::PostEditChangeProperty( struct FPropertyChang
 			}
 		}
 	}
-	else if (Name == FName(TEXT("bUsePowerOf2SnapSize")))
+	else if (Name == GET_MEMBER_NAME_CHECKED(ULevelEditorViewportSettings, bUsePowerOf2SnapSize))
 	{
 		const float BSPSnapSize = bUsePowerOf2SnapSize ? 128.0f : 100.0f;
 		UModel::SetGlobalBSPTexelScale(BSPSnapSize);
 	}
-	else if (Name == FName(TEXT("BillboardScale")))
+	else if (Name == GET_MEMBER_NAME_CHECKED(ULevelEditorViewportSettings, BillboardScale))
 	{
 		UBillboardComponent::SetEditorScale(BillboardScale);
 		UArrowComponent::SetEditorScale(BillboardScale);
 	}
+	else if (Name == GET_MEMBER_NAME_CHECKED(ULevelEditorViewportSettings, bEnableLayerSnap))
+	{
+		ULevelEditor2DSettings* Settings2D = GetMutableDefault<ULevelEditor2DSettings>();
+		if (bEnableLayerSnap && !Settings2D->bMode2DEnabled)
+		{
+			Settings2D->bMode2DEnabled = true;
+		}
+	}
 
 	if (!FUnrealEdMisc::Get().IsDeletePreferences())
 	{
diff --git a/Engine/Source/Editor/UnrealEd/Public/Settings/EditorProjectSettings.h b/Engine/Source/Editor/UnrealEd/Public/Settings/EditorProjectSettings.h
--- a/Engine/Source/Editor/UnrealEd/Public/Settings/EditorProjectSettings.h
+++ b/Engine/Source/Editor/UnrealEd/Public/Settings/EditorProjectSettings.h
@@ -84,6 +84,11 @@ struct UNREALED_API FMode2DLayer
 	, Depth(0)
 	{ }
 
+	FMode2DLayer(FString InName, float InDepth)
+		: Name(InName)
+		, Depth(InDepth)
+	{ }
+
 	/** Whether snapping to surfaces in the world is enabled */
 	UPROPERTY(EditAnywhere, config, Category = Layer)
 	FString Name;
@@ -104,7 +109,7 @@ enum class ELevelEditor2DAxis : uint8
 /**
  * Configure settings for the 2D Level Editor
  */
-UCLASS(config=Editor, meta=(DisplayName="2D"))
+UCLASS(config=Editor, meta=(DisplayName="2D"), defaultconfig)
 class UNREALED_API ULevelEditor2DSettings : public UDeveloperSettings
 {
 	GENERATED_UCLASS_BODY()
@@ -114,14 +119,6 @@ public:
 	UPROPERTY(EditAnywhere, config, Category = General, meta = (DisplayName = "Enable 2D Mode"))
 	bool bMode2DEnabled;
 
-	/** When enabled and active snap layer is valid, items dropped in the viewport will be dropped on the appropriate Y distance away */
-	UPROPERTY(EditAnywhere, config, Category = LayerSnapping, AdvancedDisplay)
-	bool bEnableLayerSnap;
-
-	/** Will try to use snap settings from SnapLayers */
-	UPROPERTY(EditAnywhere, config, Category = LayerSnapping, AdvancedDisplay)
-	int32 ActiveSnapLayerIndex;
-
 	/** Snap axis */
 	UPROPERTY(EditAnywhere, config, Category = LayerSnapping)
 	ELevelEditor2DAxis SnapAxis;
