Fixed: Build button does not generate HLOD assets
UE-15399

diff --git a/Engine/Source/Editor/LevelEditor/Private/LevelEditorActions.cpp b/Engine/Source/Editor/LevelEditor/Private/LevelEditorActions.cpp
--- a/Engine/Source/Editor/LevelEditor/Private/LevelEditorActions.cpp
+++ b/Engine/Source/Editor/LevelEditor/Private/LevelEditorActions.cpp
@@ -624,13 +624,8 @@ void FLevelEditorActionCallbacks::BuildPathsOnly_Execute()
 
 void FLevelEditorActionCallbacks::BuildLODsOnly_Execute()
 {
-	// Build paths!
+	// Build HLOD
 	FEditorBuildUtils::EditorBuild(GetWorld(), EBuildOptions::BuildHierarchicalLOD);
-
-	if (BuildLighting_CanExecute())
-	{
-		BuildLightingOnly_Execute();
-	}
 }
 
 bool FLevelEditorActionCallbacks::IsLightingQualityChecked( ELightingBuildQuality TestQuality )
diff --git a/Engine/Source/Editor/UnrealEd/Private/EditorBuildUtils.cpp b/Engine/Source/Editor/UnrealEd/Private/EditorBuildUtils.cpp
--- a/Engine/Source/Editor/UnrealEd/Private/EditorBuildUtils.cpp
+++ b/Engine/Source/Editor/UnrealEd/Private/EditorBuildUtils.cpp
@@ -329,12 +329,10 @@ bool FEditorBuildUtils::EditorBuild( UWorld* InWorld, EBuildOptions::Type Id, co
 
 				GUnrealEd->Exec( InWorld, TEXT("MAP REBUILD ALLVISIBLE") );
 
-				// disable this in build all path, it's very slow, it can be painful
-				// they don't have any collision, but make sure it's added first before doing anything
-// 				{
-// 					BuildProgressWidget.Pin()->SetBuildType(SBuildProgressWidget::BUILDTYPE_LODs);
-//					TriggerHierarchicalLODBuilder(InWorld, Id);
-// 				}
+ 				{
+ 					BuildProgressWidget.Pin()->SetBuildType(SBuildProgressWidget::BUILDTYPE_LODs);
+					TriggerHierarchicalLODBuilder(InWorld, Id);
+ 				}
 
 				{
 					BuildProgressWidget.Pin()->SetBuildType(SBuildProgressWidget::BUILDTYPE_Paths);
@@ -803,7 +801,7 @@ void FEditorBuildUtils::TriggerNavigationBuilder(UWorld* InWorld, EBuildOptions:
 
 void FEditorBuildUtils::TriggerHierarchicalLODBuilder(UWorld* InWorld, EBuildOptions::Type Id)
 {
-	// Invoke navmesh generator
+	// Invoke HLOD generator
 	InWorld->HierarchicalLODBuilder.Build();
 }
 #undef LOCTEXT_NAMESPACE
\ No newline at end of file
diff --git a/Engine/Source/Editor/UnrealEd/Private/HierarchicalLOD.cpp b/Engine/Source/Editor/UnrealEd/Private/HierarchicalLOD.cpp
--- a/Engine/Source/Editor/UnrealEd/Private/HierarchicalLOD.cpp
+++ b/Engine/Source/Editor/UnrealEd/Private/HierarchicalLOD.cpp
@@ -124,18 +124,6 @@ void FHierarchicalLODBuilder::Build()
 {
 	check (World);
 
-	if (World->GetWorldSettings()->bEnableHierarchicalLODSystem == false)
-	{
-		if (EAppReturnType::No == 
-			FMessageDialog::Open( EAppMsgType::YesNo, LOCTEXT("HierarchicalLOD_Warning", 
-			"Enable Hierarchical LOD System is disabled in World Setting. This process will delete previously generated LODActors if exists.\n\nWould you like to continue?")) )
-		{
-			// has been canceled, return
-			GEditor->SetMapBuildCancelled(true);
-			return;
-		}
-	}
-	
 	const TArray<class ULevel*>& Levels = World->GetLevels();
 
 	for (const auto& LevelIter : Levels)
