PR #1136: Fix BSP texture align tools all being initialized to default (Contributed by [REDACTED_PERSON])

diff --git a/Engine/Source/Editor/UnrealEd/Private/TexAlignTools.cpp b/Engine/Source/Editor/UnrealEd/Private/TexAlignTools.cpp
--- a/Engine/Source/Editor/UnrealEd/Private/TexAlignTools.cpp
+++ b/Engine/Source/Editor/UnrealEd/Private/TexAlignTools.cpp
@@ -441,10 +441,10 @@ void FTexAlignTools::Init()
 {
 	// Create the list of aligners.
 	Aligners.Empty();
-	Aligners.Add(NewObject<UTexAligner>(GetTransientPackage(), NAME_None, RF_Public | RF_RootSet | RF_Standalone));
-	Aligners.Add(NewObject<UTexAligner>(GetTransientPackage(), NAME_None, RF_Public | RF_RootSet | RF_Standalone));
-	Aligners.Add(NewObject<UTexAligner>(GetTransientPackage(), NAME_None, RF_Public | RF_RootSet | RF_Standalone));
-	Aligners.Add(NewObject<UTexAligner>(GetTransientPackage(), NAME_None, RF_Public | RF_RootSet | RF_Standalone));
+	Aligners.Add(NewObject<UTexAlignerDefault>(GetTransientPackage(), NAME_None, RF_Public | RF_RootSet | RF_Standalone));
+	Aligners.Add(NewObject<UTexAlignerPlanar>(GetTransientPackage(), NAME_None, RF_Public | RF_RootSet | RF_Standalone));
+	Aligners.Add(NewObject<UTexAlignerBox>(GetTransientPackage(), NAME_None, RF_Public | RF_RootSet | RF_Standalone));
+	Aligners.Add(NewObject<UTexAlignerFit>(GetTransientPackage(), NAME_None, RF_Public | RF_RootSet | RF_Standalone));
 	
 	FEditorDelegates::FitTextureToSurface.AddRaw(this, &FTexAlignTools::OnEditorFitTextureToSurface);
 }
