Fix for SettingsClasses missing an include for EditorProjectSettings.h
Only seemed to affect non-unity builds.

diff --git a/Engine/Source/Editor/UnrealEd/Private/Settings/SettingsClasses.cpp b/Engine/Source/Editor/UnrealEd/Private/Settings/SettingsClasses.cpp
--- a/Engine/Source/Editor/UnrealEd/Private/Settings/SettingsClasses.cpp
+++ b/Engine/Source/Editor/UnrealEd/Private/Settings/SettingsClasses.cpp
@@ -10,6 +10,7 @@
 #include "SNotificationList.h"
 #include "NotificationManager.h"
 #include "ISettingsModule.h"
+#include "EditorProjectSettings.h"
 
 #include "SourceCodeNavigation.h"
 
