Fixed non-ANSI clothing assets failing to load.
#UE-8658

diff --git a/Engine/Source/Editor/UnrealEd/Private/ApexClothingUtils.cpp b/Engine/Source/Editor/UnrealEd/Private/ApexClothingUtils.cpp
--- a/Engine/Source/Editor/UnrealEd/Private/ApexClothingUtils.cpp
+++ b/Engine/Source/Editor/UnrealEd/Private/ApexClothingUtils.cpp
@@ -186,14 +186,21 @@ NxClothingAsset* CreateApexClothingAssetFromFile(FString& Filename)
 {
 	NxClothingAsset* ApexClothingAsset = NULL;
 
-	// Create a stream to read the file
-	physx::PxFileBuf* Stream = GApexSDK->createStream(TCHAR_TO_ANSI(*Filename), PxFileBuf::OPEN_READ_ONLY);
+	TArray<uint8> FileBuffer;
+	if(FFileHelper::LoadFileToArray(FileBuffer, *Filename, FILEREAD_Silent))
+	{
+		physx::PxFileBuf* Stream = GApexSDK->createMemoryReadStream((void*)FileBuffer.GetData(), FileBuffer.Num());
 
-	if (Stream != NULL)
+		if(Stream != NULL)
+		{
+			ApexClothingAsset = CreateApexClothingAssetFromPxStream(*Stream);
+			// Release our stream
+			Stream->release();
+		}
+	}
+	else
 	{
-		ApexClothingAsset = CreateApexClothingAssetFromPxStream(*Stream);
-		// Release our stream
-		Stream->release();
+		UE_LOG(LogApexClothingUtils, Warning, TEXT("Could not open APEX clothing file: %s"), *Filename);
 	}
 
 	return ApexClothingAsset;
