Editor ends a multi-process PIE session if it receives a network failure. Fixes UE-6775.

diff --git a/Engine/Source/Editor/UnrealEd/Classes/Editor/EditorEngine.h b/Engine/Source/Editor/UnrealEd/Classes/Editor/EditorEngine.h
--- a/Engine/Source/Editor/UnrealEd/Classes/Editor/EditorEngine.h
+++ b/Engine/Source/Editor/UnrealEd/Classes/Editor/EditorEngine.h
@@ -689,6 +689,7 @@ public:
 	virtual UWorld* CreatePIEWorldByDuplication(FWorldContext &WorldContext, UWorld* InWorld, FString &PlayWorldMapName) override;
 	virtual bool GetMapBuildCancelled() const override { return false; }
 	virtual void SetMapBuildCancelled(bool InCancelled) override { /* Intentionally empty. */ }
+	virtual void HandleNetworkFailure(UWorld *World, UNetDriver *NetDriver, ENetworkFailure::Type FailureType, const FString& ErrorString) override;
 protected:
 	virtual void InitializeObjectReferences() override;
 	virtual void ProcessToggleFreezeCommand(UWorld* InWorld) override;
diff --git a/Engine/Source/Editor/UnrealEd/Private/EditorEngine.cpp b/Engine/Source/Editor/UnrealEd/Private/EditorEngine.cpp
--- a/Engine/Source/Editor/UnrealEd/Private/EditorEngine.cpp
+++ b/Engine/Source/Editor/UnrealEd/Private/EditorEngine.cpp
@@ -6257,6 +6257,26 @@ UWorld::InitializationValues UEditorEngine::GetEditorWorldInitializationValues()
 		.EnableTraceCollision(true);
 }
 
+void UEditorEngine::HandleNetworkFailure(UWorld *World, UNetDriver *NetDriver, ENetworkFailure::Type FailureType, const FString& ErrorString)
+{
+	// If the failure occurred during PIE while connected to another process, simply end the PIE session before
+	// trying to travel anywhere.
+	if (PlayOnLocalPCSessions.Num() > 0)
+	{
+		for (const FWorldContext& WorldContext : WorldList)
+		{
+			if (WorldContext.WorldType == EWorldType::PIE && WorldContext.World() == World)
+			{
+				RequestEndPlayMap();
+				return;
+			}
+		}
+	}
+
+	// Otherwise, perform normal engine failure handling.
+	Super::HandleNetworkFailure(World, NetDriver, FailureType, ErrorString);
+}
+
 //////////////////////////////////////////////////////////////////////////
 // FActorLabelUtilities
 
