Fixed camel casing on FindDependencyBy* methods in FManifestInfo.

diff --git a/Engine/Source/Editor/UnrealEd/Classes/Commandlets/GatherTextCommandletBase.h b/Engine/Source/Editor/UnrealEd/Classes/Commandlets/GatherTextCommandletBase.h
--- a/Engine/Source/Editor/UnrealEd/Classes/Commandlets/GatherTextCommandletBase.h
+++ b/Engine/Source/Editor/UnrealEd/Classes/Commandlets/GatherTextCommandletBase.h
@@ -196,10 +196,10 @@ public:
 	bool AddManifestDependency( const FString& InManifestFile );
 
 	// Return the first entry in the dependencies that matches the passed in namespace and context.
-	TSharedPtr< FManifestEntry > FindDependencyEntrybyContext(  const FString& Namespace, const FContext& Context, FString& OutFileName );
+	TSharedPtr< FManifestEntry > FindDependencyEntryByContext(  const FString& Namespace, const FContext& Context, FString& OutFileName );
 
 	// Return the first entry in the dependencies that matches the passed in namespace and source.
-	TSharedPtr< FManifestEntry > FindDependencyEntrybySource( const FString& Namespace, const FLocItem& Source, FString& OutFileName );
+	TSharedPtr< FManifestEntry > FindDependencyEntryBySource( const FString& Namespace, const FLocItem& Source, FString& OutFileName );
 
 	// Returns the number of manifest dependencies
 	int32 NumDependencies() { return ManifestDependencies.Num(); }
diff --git a/Engine/Source/Editor/UnrealEd/Private/Commandlets/GatherTextCommandletBase.cpp b/Engine/Source/Editor/UnrealEd/Private/Commandlets/GatherTextCommandletBase.cpp
--- a/Engine/Source/Editor/UnrealEd/Private/Commandlets/GatherTextCommandletBase.cpp
+++ b/Engine/Source/Editor/UnrealEd/Private/Commandlets/GatherTextCommandletBase.cpp
@@ -275,7 +275,7 @@ bool FManifestInfo::AddManifestDependencies( const TArray< FString >& InManifest
 }
 
 
-TSharedPtr< FManifestEntry > FManifestInfo::FindDependencyEntrybyContext( const FString& Namespace, const FContext& Context, FString& OutFileName )
+TSharedPtr< FManifestEntry > FManifestInfo::FindDependencyEntryByContext( const FString& Namespace, const FContext& Context, FString& OutFileName )
 {
 	TSharedPtr<FManifestEntry> DependencyEntry = NULL;
 	OutFileName = TEXT("");
@@ -292,7 +292,7 @@ TSharedPtr< FManifestEntry > FManifestInfo::FindDependencyEntrybyContext( const
 	return DependencyEntry;
 }
 
-TSharedPtr< FManifestEntry > FManifestInfo::FindDependencyEntrybySource( const FString& Namespace, const FLocItem& Source, FString& OutFileName )
+TSharedPtr< FManifestEntry > FManifestInfo::FindDependencyEntryBySource( const FString& Namespace, const FLocItem& Source, FString& OutFileName )
 {
 	TSharedPtr<FManifestEntry> DependencyEntry = NULL;
 	OutFileName = TEXT("");
@@ -323,7 +323,7 @@ void FManifestInfo::ApplyManifestDependencies()
 			{
 				FString DependencyFileName;
 
-				const TSharedPtr<FManifestEntry> DependencyEntry = FindDependencyEntrybyContext( ManifestEntry->Namespace, *ContextIt, DependencyFileName );
+				const TSharedPtr<FManifestEntry> DependencyEntry = FindDependencyEntryByContext( ManifestEntry->Namespace, *ContextIt, DependencyFileName );
 				
 				if( DependencyEntry.IsValid() )
 				{
@@ -378,7 +378,7 @@ bool FManifestInfo::AddEntry( const FString& EntryDescription, const FString& Na
 	TSharedPtr< FManifestEntry > ExistingEntry = Manifest->FindEntryByContext( Namespace, Context );
 	if( !ExistingEntry.IsValid() )
 	{
-		ExistingEntry = FindDependencyEntrybyContext( Namespace, Context, ExistingEntryFileName );
+		ExistingEntry = FindDependencyEntryByContext( Namespace, Context, ExistingEntryFileName );
 	}
 
 	if( ExistingEntry.IsValid() )
diff --git a/Engine/Source/Editor/UnrealEd/Private/Commandlets/GatherTextFromAssetsCommandlet.cpp b/Engine/Source/Editor/UnrealEd/Private/Commandlets/GatherTextFromAssetsCommandlet.cpp
--- a/Engine/Source/Editor/UnrealEd/Private/Commandlets/GatherTextFromAssetsCommandlet.cpp
+++ b/Engine/Source/Editor/UnrealEd/Private/Commandlets/GatherTextFromAssetsCommandlet.cpp
@@ -1,4 +1,4 @@
-// Copyright 1998-2015 Epic Games, Inc. All Rights Reserved.
+﻿// Copyright 1998-2015 Epic Games, Inc. All Rights Reserved.
 
 #include "UnrealEd.h"
 #include "PackageTools.h"
@@ -607,7 +607,7 @@ void UGatherTextFromAssetsCommandlet::ProcessTextProperty(UTextProperty* InTextP
 		if( !ExistingEntry.IsValid() )
 		{
 			FString FileInfo;
-			ExistingEntry = ManifestInfo->FindDependencyEntrybyContext( Namespace, SearchContext, FileInfo );
+			ExistingEntry = ManifestInfo->FindDependencyEntryByContext( Namespace, SearchContext, FileInfo );
 		}
 
 		// Entry already exists, check for conflict.
