Fixed titles on Material editor nodes to update when the pin connections change.
[UE-15352] - Math Nodes are not updating on input

#change NodeConnectionListChanged() method in UMaterialGraphNode now clears the visualization cache on the schema triggering all FText attributes on the node's widgets to update.

diff --git a/Engine/Source/Editor/AIGraph/Private/AIGraphNode.cpp b/Engine/Source/Editor/AIGraph/Private/AIGraphNode.cpp
--- a/Engine/Source/Editor/AIGraph/Private/AIGraphNode.cpp
+++ b/Engine/Source/Editor/AIGraph/Private/AIGraphNode.cpp
@@ -221,6 +221,8 @@ bool UAIGraphNode::IsSubNode() const
 
 void UAIGraphNode::NodeConnectionListChanged()
 {
+	Super::NodeConnectionListChanged();
+
 	GetAIGraph()->UpdateAsset();
 }
 
diff --git a/Engine/Source/Editor/BehaviorTreeEditor/Private/BehaviorTreeDecoratorGraphNode.cpp b/Engine/Source/Editor/BehaviorTreeEditor/Private/BehaviorTreeDecoratorGraphNode.cpp
--- a/Engine/Source/Editor/BehaviorTreeEditor/Private/BehaviorTreeDecoratorGraphNode.cpp
+++ b/Engine/Source/Editor/BehaviorTreeEditor/Private/BehaviorTreeDecoratorGraphNode.cpp
@@ -82,6 +82,8 @@ void UBehaviorTreeDecoratorGraphNode::AutowireNewNode(UEdGraphPin* FromPin)
 
 void UBehaviorTreeDecoratorGraphNode::NodeConnectionListChanged()
 {
+	Super::NodeConnectionListChanged();
+
 	const UBehaviorTreeDecoratorGraph* MyGraph = CastChecked<UBehaviorTreeDecoratorGraph>(GetGraph());
 	
 	UBehaviorTreeGraphNode_CompositeDecorator* CompositeDecorator = Cast<UBehaviorTreeGraphNode_CompositeDecorator>(MyGraph->GetOuter());
diff --git a/Engine/Source/Editor/BlueprintGraph/Private/K2Node_MacroInstance.cpp b/Engine/Source/Editor/BlueprintGraph/Private/K2Node_MacroInstance.cpp
--- a/Engine/Source/Editor/BlueprintGraph/Private/K2Node_MacroInstance.cpp
+++ b/Engine/Source/Editor/BlueprintGraph/Private/K2Node_MacroInstance.cpp
@@ -282,6 +282,8 @@ void UK2Node_MacroInstance::NotifyPinConnectionListChanged(UEdGraphPin* ChangedP
 
 void UK2Node_MacroInstance::NodeConnectionListChanged()
 {
+	Super::NodeConnectionListChanged();
+
 	if (bReconstructNode)
 	{
 		ReconstructNode();
diff --git a/Engine/Source/Editor/BlueprintGraph/Private/K2Node_Select.cpp b/Engine/Source/Editor/BlueprintGraph/Private/K2Node_Select.cpp
--- a/Engine/Source/Editor/BlueprintGraph/Private/K2Node_Select.cpp
+++ b/Engine/Source/Editor/BlueprintGraph/Private/K2Node_Select.cpp
@@ -652,6 +652,8 @@ void UK2Node_Select::SetEnum(UEnum* InEnum, bool bForceRegenerate)
 
 void UK2Node_Select::NodeConnectionListChanged()
 {
+	Super::NodeConnectionListChanged();
+
 	if (bReconstructNode)
 	{
 		ReconstructNode();
diff --git a/Engine/Source/Editor/UnrealEd/Classes/MaterialGraph/MaterialGraphNode.h b/Engine/Source/Editor/UnrealEd/Classes/MaterialGraph/MaterialGraphNode.h
--- a/Engine/Source/Editor/UnrealEd/Classes/MaterialGraph/MaterialGraphNode.h
+++ b/Engine/Source/Editor/UnrealEd/Classes/MaterialGraph/MaterialGraphNode.h
@@ -57,6 +57,7 @@ public:
 	virtual void PrepareForCopying() override;
 	virtual void GetContextMenuActions(const FGraphNodeContextMenuBuilder& Context) const override;
 	virtual void PostPlacedNewNode() override;
+	virtual void NodeConnectionListChanged() override;
 	virtual void OnRenameNode(const FString& NewName) override;
 	virtual void OnUpdateCommentText( const FString& NewComment ) override;
 	virtual void GetPinHoverText(const UEdGraphPin& Pin, FString& HoverTextOut) const override;
diff --git a/Engine/Source/Editor/UnrealEd/Private/MaterialGraphNode.cpp b/Engine/Source/Editor/UnrealEd/Private/MaterialGraphNode.cpp
--- a/Engine/Source/Editor/UnrealEd/Private/MaterialGraphNode.cpp
+++ b/Engine/Source/Editor/UnrealEd/Private/MaterialGraphNode.cpp
@@ -598,6 +598,17 @@ void UMaterialGraphNode::PostPlacedNewNode()
 	}
 }
 
+void UMaterialGraphNode::NodeConnectionListChanged()
+{
+	Super::NodeConnectionListChanged();
+
+	const UEdGraphSchema* Schema = GetSchema();
+	if (Schema != nullptr)
+	{
+		Schema->ForceVisualizationCacheClear();
+	}
+}
+
 void UMaterialGraphNode::OnRenameNode(const FString& NewName)
 {
 	MaterialExpression->Modify();
