PR #1143: Fix small but annoying typos in ContentWrangleCommandlet (Contributed by [REDACTED_PERSON])

diff --git a/Engine/Source/Editor/UnrealEd/Private/Commandlets/ContentCommandlets.cpp b/Engine/Source/Editor/UnrealEd/Private/Commandlets/ContentCommandlets.cpp
--- a/Engine/Source/Editor/UnrealEd/Private/Commandlets/ContentCommandlets.cpp
+++ b/Engine/Source/Editor/UnrealEd/Private/Commandlets/ContentCommandlets.cpp
@@ -1006,10 +1006,10 @@ int32 UWrangleContentCommandlet::Main( const FString& Params )
 
 		// gather any per map packages for cooking
 		TArray<FString> PerMapPackagesToLoad;
-		for (FConfigSectionMap::TIterator PacakgeIt(PackagesToFullyLoad); PacakgeIt; ++PacakgeIt)
+		for (FConfigSectionMap::TIterator PackageIt(PackagesToFullyLoad); PackageIt; ++PackageIt)
 		{
 			// add dependencies for the per-map packages for this map (if any)
-			TArray<FString>* Packages = PerMapCookPackages.Find(PacakgeIt.Value());
+			TArray<FString>* Packages = PerMapCookPackages.Find(PackageIt.Value());
 			if (Packages != NULL)
 			{
 				for (int32 PackageIndex = 0; PackageIndex < Packages->Num(); PackageIndex++)
@@ -1038,12 +1038,12 @@ int32 UWrangleContentCommandlet::Main( const FString& Params )
 		}
 
 		// go over all the packages that we want to fully load
-		for (FConfigSectionMap::TIterator PacakgeIt(PackagesToFullyLoad); PacakgeIt; ++PacakgeIt)
+		for (FConfigSectionMap::TIterator PackageIt(PackagesToFullyLoad); PackageIt; ++PackageIt)
 		{
 			// there may be multiple sublevels to load if this package is a persistent level with sublevels
 			TArray<FString> PackagesToLoad;
 			// start off just loading this package (more may be added in the loop)
-			PackagesToLoad.Add(*PacakgeIt.Value());
+			PackagesToLoad.Add(*PackageIt.Value());
 
 			for (int32 PackageIndex = 0; PackageIndex < PackagesToLoad.Num(); PackageIndex++)
 			{
