Removed a check that was asserting when compiling a Blueprint with invalid CallFunction or CallParentFunction nodes.

diff --git a/Engine/Source/Editor/UnrealEd/Private/Kismet2/BlueprintEditorUtils.cpp b/Engine/Source/Editor/UnrealEd/Private/Kismet2/BlueprintEditorUtils.cpp
--- a/Engine/Source/Editor/UnrealEd/Private/Kismet2/BlueprintEditorUtils.cpp
+++ b/Engine/Source/Editor/UnrealEd/Private/Kismet2/BlueprintEditorUtils.cpp
@@ -5367,7 +5367,6 @@ void FBlueprintEditorUtils::ConformCallsToParentFunctions(UBlueprint* Blueprint)
 					{
 						// Cache a reference to the output exec pin
 						UEdGraphPin* OutputPin = CallFunctionNode->GetThenPin();
-						check(NULL != OutputPin);
 
 						// We're going to destroy the existing parent function call node, but first we need to persist any existing connections
 						for(int PinIndex = 0; PinIndex < CallFunctionNode->Pins.Num(); ++PinIndex)
