Fix crash summoning component visualizer context menu in immersive mode

diff --git a/Engine/Source/Editor/LevelEditor/Private/SLevelEditor.cpp b/Engine/Source/Editor/LevelEditor/Private/SLevelEditor.cpp
--- a/Engine/Source/Editor/LevelEditor/Private/SLevelEditor.cpp
+++ b/Engine/Source/Editor/LevelEditor/Private/SLevelEditor.cpp
@@ -1372,6 +1372,11 @@ TArray< TSharedPtr< ILevelViewport > > SLevelEditor::GetViewports() const
 	return OutViewports;
 }
  
+TSharedPtr<ILevelViewport> SLevelEditor::GetActiveViewportInterface()
+{
+	return GetActiveViewport();
+}
+
 TSharedPtr< class FAssetThumbnailPool > SLevelEditor::GetThumbnailPool() const
 {
 	return ThumbnailPool;
diff --git a/Engine/Source/Editor/LevelEditor/Private/SLevelEditor.h b/Engine/Source/Editor/LevelEditor/Private/SLevelEditor.h
--- a/Engine/Source/Editor/LevelEditor/Private/SLevelEditor.h
+++ b/Engine/Source/Editor/LevelEditor/Private/SLevelEditor.h
@@ -64,6 +64,7 @@ public:
 	virtual void SummonLevelViewportContextMenu() override;
 	virtual const TArray< TSharedPtr< class IToolkit > >& GetHostedToolkits() const override;
 	virtual TArray< TSharedPtr< ILevelViewport > > GetViewports() const override;
+	virtual TSharedPtr<ILevelViewport> GetActiveViewportInterface() override;
 	virtual TSharedPtr< class FAssetThumbnailPool > GetThumbnailPool() const override;
 	virtual void AppendCommands( const TSharedRef<FUICommandList>& InCommandsToAppend ) override;
 
diff --git a/Engine/Source/Editor/LevelEditor/Public/ILevelEditor.h b/Engine/Source/Editor/LevelEditor/Public/ILevelEditor.h
--- a/Engine/Source/Editor/LevelEditor/Public/ILevelEditor.h
+++ b/Engine/Source/Editor/LevelEditor/Public/ILevelEditor.h
@@ -21,6 +21,9 @@ public:
 
 	/** Gets an array of all viewports in this level editor */
 	virtual TArray< TSharedPtr< ILevelViewport > > GetViewports() const = 0;
+	
+	/** Gets the active level viewport for this level editor */
+	virtual TSharedPtr<ILevelViewport> GetActiveViewportInterface() = 0;
 
 	/** Get the thumbnail pool used by this level editor */
 	virtual TSharedPtr< class FAssetThumbnailPool > GetThumbnailPool() const = 0;
diff --git a/Engine/Source/Editor/UnrealEd/Private/ComponentVisualizerManager.cpp b/Engine/Source/Editor/UnrealEd/Private/ComponentVisualizerManager.cpp
--- a/Engine/Source/Editor/UnrealEd/Private/ComponentVisualizerManager.cpp
+++ b/Engine/Source/Editor/UnrealEd/Private/ComponentVisualizerManager.cpp
@@ -3,6 +3,7 @@
 #include "UnrealEd.h"
 #include "ComponentVisualizerManager.h"
 #include "ILevelEditor.h"
+#include "ILevelViewport.h"
 
 /** Handle a click on the specified editor viewport client */
 bool FComponentVisualizerManager::HandleClick(FLevelEditorViewportClient* InViewportClient, HHitProxy* HitProxy, const FViewportClick& Click)
@@ -14,10 +15,11 @@ bool FComponentVisualizerManager::HandleClick(FLevelEditorViewportClient* InView
 		if (MenuWidget.IsValid())
 		{
 			auto ParentLevelEditor = InViewportClient->ParentLevelEditor.Pin();
-			if (ParentLevelEditor.IsValid())
+			if (ParentLevelEditor.IsValid() && ParentLevelEditor->GetActiveViewportInterface().IsValid() )
 			{
+				TSharedPtr<SWidget> MenuParent = ParentLevelEditor->GetActiveViewportInterface()->GetViewportWidget().Pin();
 				FSlateApplication::Get().PushMenu(
-					ParentLevelEditor.ToSharedRef(),
+					MenuParent.ToSharedRef(),
 					MenuWidget.ToSharedRef(),
 					FSlateApplication::Get().GetCursorPos(),
 					FPopupTransitionEffect(FPopupTransitionEffect::ContextMenu));
