Allow project to define Audio Quality Levels.
* Each quality level allows the definition of the max audio channels to be used
* A sound cue node exists to allow branching behavior based on quality level.
* Wave assets from Wave Player sound nodes will not be loaded if connected only to a branch for a quality level that is not currently selected
* Quality Level to use for stand alone is specified via Game User Settings. Changes to this value will not take effect until the next run.
* Quality Level to use for PIE is specified from the Level Editor Play Settings and can be changed dynamically causing newly played sounds to use the alternate path
* Add SoundNodeAssetReferencer abstract class to operate as an interface for sound nodes that need to have assets loaded dependingo n the sound quality node

diff --git a/Engine/Config/BaseEngine.ini b/Engine/Config/BaseEngine.ini
--- a/Engine/Config/BaseEngine.ini
+++ b/Engine/Config/BaseEngine.ini
@@ -1223,8 +1223,9 @@ bUseStreamingPause=false
 +K2FieldRedirects=(OldFieldName="KismetSystemLibrary.EXPERIMENTAL_HideAdBanner",NewFieldName="KismetSystemLibrary.HideAdBanner")
 +K2FieldRedirects=(OldFieldName="KismetSystemLibrary.EXPERIMENTAL_CloseAdBanner",NewFieldName="KismetSystemLibrary.ForceCloseAdBanner")
 
-;SoundNodeWavePlayer using asset pointers - DEPRECATED(4.9)
+;SoundNodeWavePlayer/ModPlayer using asset pointers - DEPRECATED(4.9)
 +TaggedPropertyRedirects=(ClassName="SoundNodeWavePlayer",OldPropertyName="SoundWave",NewPropertyName="SoundWaveAssetPtr")
++TaggedPropertyRedirects=(ClassName="SoundNodeModPlayer",OldPropertyName="SoundMod",NewPropertyName="SoundModAssetPtr")
 
 [CoreUObject.Metadata]
 ; Note: UnrealHeaderTool should be rerun after making changes to MetadataRedirects to catch any keys specified in class headers
diff --git a/Engine/Plugins/Runtime/SoundMod/Source/SoundMod/Private/SoundNodeModPlayer.cpp b/Engine/Plugins/Runtime/SoundMod/Source/SoundMod/Private/SoundNodeModPlayer.cpp
--- a/Engine/Plugins/Runtime/SoundMod/Source/SoundMod/Private/SoundNodeModPlayer.cpp
+++ b/Engine/Plugins/Runtime/SoundMod/Source/SoundMod/Private/SoundNodeModPlayer.cpp
@@ -6,11 +6,35 @@
 
 #define LOCTEXT_NAMESPACE "SoundNodeModPlayer"
 
-USoundNodeModPlayer::USoundNodeModPlayer(const FObjectInitializer& ObjectInitializer)
-	: Super(ObjectInitializer)
+void USoundNodeModPlayer::AddReferencedObjects(UObject* InThis, FReferenceCollector& Collector)
 {
+	USoundNodeModPlayer* This = CastChecked<USoundNodeModPlayer>(InThis);
+	Super::AddReferencedObjects(This, Collector);
+
+	Collector.AddReferencedObject(This->SoundMod);
+}
+
+void USoundNodeModPlayer::LoadAsset()
+{
+	SoundMod = SoundModAssetPtr.LoadSynchronous();
+}
+
+void USoundNodeModPlayer::SetSoundMod(USoundMod* InSoundMod)
+{
+	SoundMod = InSoundMod;
+	SoundModAssetPtr = InSoundMod;
 }
 
+#if WITH_EDITOR
+void USoundNodeModPlayer::PostEditChangeProperty(FPropertyChangedEvent& PropertyChangedEvent)
+{
+	if (PropertyChangedEvent.Property && PropertyChangedEvent.Property->GetFName() == GET_MEMBER_NAME_CHECKED(USoundNodeModPlayer, SoundModAssetPtr))
+	{
+		LoadAsset();
+	}
+}
+#endif
+
 void USoundNodeModPlayer::ParseNodes( FAudioDevice* AudioDevice, const UPTRINT NodeWaveInstanceHash, FActiveSound& ActiveSound, const FSoundParseParameters& ParseParams, TArray<FWaveInstance*>& WaveInstances )
 {
 	if (SoundMod)
diff --git a/Engine/Source/Editor/UnrealEd/Classes/Settings/LevelEditorPlaySettings.h b/Engine/Source/Editor/UnrealEd/Classes/Settings/LevelEditorPlaySettings.h
--- a/Engine/Source/Editor/UnrealEd/Classes/Settings/LevelEditorPlaySettings.h
+++ b/Engine/Source/Editor/UnrealEd/Classes/Settings/LevelEditorPlaySettings.h
@@ -150,7 +150,7 @@ public:
 	TEnumAsByte<ELabelAnchorMode> MouseControlLabelPosition;
 
 	/** Should Play-in-Viewport respect HMD orientations (default = false) */
-	UPROPERTY(config, EditAnywhere, Category = PlayInEditor, meta = (ToolTip = "Whether or not HMD orientation should be used when playing in viewport"))
+	UPROPERTY(config, EditAnywhere, Category=PlayInEditor, meta=(ToolTip="Whether or not HMD orientation should be used when playing in viewport"))
 	bool ViewportGetsHMDControl;
 
 	/** Whether to automatically recompile blueprints on PIE */
@@ -158,9 +158,13 @@ public:
 	bool AutoRecompileBlueprints;
 
 	/** Whether to play sounds during PIE */
-	UPROPERTY(config, EditAnywhere, Category = PlayInEditor, meta = (ToolTip = "Whether to play sounds when in a Play In Editor session"))
+	UPROPERTY(config, EditAnywhere, Category=PlayInEditor, meta=(ToolTip="Whether to play sounds when in a Play In Editor session"))
 	bool EnableSound;
 
+	/** Which quality level to use when playing in editor */
+	UPROPERTY(config, EditAnywhere, Category=PlayInEditor)
+	int32 PlayInEditorSoundQuality;
+
 	/** True if Play In Editor should only load currently-visible levels in PIE. */
 	UPROPERTY(config)
 	uint32 bOnlyLoadVisibleLevelsInPIE:1;
diff --git a/Engine/Source/Editor/UnrealEd/Private/Settings/LevelEditorPlaySettingsCustomization.h b/Engine/Source/Editor/UnrealEd/Private/Settings/LevelEditorPlaySettingsCustomization.h
--- a/Engine/Source/Editor/UnrealEd/Private/Settings/LevelEditorPlaySettingsCustomization.h
+++ b/Engine/Source/Editor/UnrealEd/Private/Settings/LevelEditorPlaySettingsCustomization.h
@@ -4,6 +4,7 @@
 
 #include "PropertyEditing.h"
 #include "PropertyCustomizationHelpers.h"
+#include "Runtime/Engine/Classes/Sound/AudioSettings.h"
 
 
 #define LOCTEXT_NAMESPACE "FLevelEditorPlaySettingsCustomization"
@@ -296,6 +297,49 @@ public:
 	{
 		const float MaxPropertyWidth = 400.0f;
 
+		// play in editor settings
+		IDetailCategoryBuilder& PlayInEditorCategory = LayoutBuilder.EditCategory("PlayInEditor");
+		{
+			TArray<TSharedRef<IPropertyHandle>> PIECategoryProperties;
+			PlayInEditorCategory.GetDefaultProperties(PIECategoryProperties, true, false);
+
+			TSharedPtr<IPropertyHandle> PIEEnableSoundHandle = LayoutBuilder.GetProperty(GET_MEMBER_NAME_CHECKED(ULevelEditorPlaySettings, EnableSound));
+			PIESoundQualityLevelHandle = LayoutBuilder.GetProperty(GET_MEMBER_NAME_CHECKED(ULevelEditorPlaySettings, PlayInEditorSoundQuality));
+			PIESoundQualityLevelHandle->MarkHiddenByCustomization();
+
+			for (TSharedRef<IPropertyHandle>& PropertyHandle : PIECategoryProperties)
+			{
+				if (PropertyHandle->GetProperty() != PIESoundQualityLevelHandle->GetProperty())
+				{
+					PlayInEditorCategory.AddProperty(PropertyHandle);
+				}
+
+				if (PropertyHandle->GetProperty() == PIEEnableSoundHandle->GetProperty())
+				{
+					PlayInEditorCategory.AddCustomRow(PIESoundQualityLevelHandle->GetPropertyDisplayName(), false)
+						.NameContent()
+						[
+							PIESoundQualityLevelHandle->CreatePropertyNameWidget()
+						]
+						.ValueContent()
+						.MaxDesiredWidth(MaxPropertyWidth)
+						[
+							SAssignNew(QualityLevelComboBox, SComboBox<TSharedPtr<FString>>)
+							.OptionsSource(&AvailableQualityLevels)
+							.OnComboBoxOpening(this, &FLevelEditorPlaySettingsCustomization::HandleQualityLevelComboBoxOpening)
+							.OnGenerateWidget(this, &FLevelEditorPlaySettingsCustomization::HandleQualityLevelComboBoxGenerateWidget)
+							.OnSelectionChanged(this, &FLevelEditorPlaySettingsCustomization::HandleQualityLevelSelectionChanged)
+							[
+								SNew(STextBlock)
+								.Text(this, &FLevelEditorPlaySettingsCustomization::GetSelectedQualityLevelName)
+							]
+						];
+				}
+			}
+
+
+		}
+
 		// play in new window settings
 		IDetailCategoryBuilder& PlayInNewWindowCategory = LayoutBuilder.EditCategory("PlayInNewWindow");
 		{
@@ -634,6 +678,55 @@ private:
 	{
 		return GetDefault<ULevelEditorPlaySettings>()->GetAdditionalLaunchOptionsVisibility();
 	}
+
+	void HandleQualityLevelComboBoxOpening()
+	{
+		const UAudioSettings* AudioSettings = GetDefault<UAudioSettings>();
+		AvailableQualityLevels.Empty(AudioSettings->QualityLevels.Num());
+		for (const FAudioQualitySettings& AQSettings : AudioSettings->QualityLevels)
+		{
+			AvailableQualityLevels.Add(MakeShareable(new FString(AQSettings.DisplayName.ToString())));
+		}
+		QualityLevelComboBox->RefreshOptions();
+	}
+
+	TSharedRef<SWidget> HandleQualityLevelComboBoxGenerateWidget(TSharedPtr<FString> InItem)
+	{
+		return SNew(STextBlock)
+				.Text(FText::FromString(*InItem));
+	}
+
+	void HandleQualityLevelSelectionChanged(TSharedPtr<FString> InSelection, ESelectInfo::Type SelectInfo)
+	{
+		if (InSelection.IsValid())
+		{
+			const UAudioSettings* AudioSettings = GetDefault<UAudioSettings>();
+			for (int32 QualityLevel = 0; QualityLevel < AudioSettings->QualityLevels.Num(); ++QualityLevel)
+			{
+				if (AudioSettings->QualityLevels[QualityLevel].DisplayName.ToString() == *InSelection)
+				{
+					PIESoundQualityLevelHandle->SetValue(QualityLevel);
+					break;
+				}
+			}
+		}
+	}
+
+	FText GetSelectedQualityLevelName() const
+	{
+		int32 QualityLevel = 0;
+		PIESoundQualityLevelHandle->GetValue(QualityLevel);
+		const UAudioSettings* AudioSettings = GetDefault<UAudioSettings>();
+		return (QualityLevel >= 0 && QualityLevel < AudioSettings->QualityLevels.Num() ? AudioSettings->QualityLevels[QualityLevel].DisplayName : FText::GetEmpty());
+	}
+
+private:
+
+	/** Collection of possible quality levels we can use as a parent for this profile */
+	TArray<TSharedPtr<FString>> AvailableQualityLevels;
+	TSharedPtr<IPropertyHandle> PIESoundQualityLevelHandle;
+	TSharedPtr<SComboBox<TSharedPtr<FString>>> QualityLevelComboBox;
+
 };
 
 
diff --git a/Engine/Source/Runtime/Engine/Classes/GameFramework/GameUserSettings.h b/Engine/Source/Runtime/Engine/Classes/GameFramework/GameUserSettings.h
--- a/Engine/Source/Runtime/Engine/Classes/GameFramework/GameUserSettings.h
+++ b/Engine/Source/Runtime/Engine/Classes/GameFramework/GameUserSettings.h
@@ -67,6 +67,12 @@ public:
 	/** Set scalability settings to sensible fallback values, for use when the benchmark fails or potentially causes a crash */
 	void SetBenchmarkFallbackValues();
 
+	/** Sets the user's audio quality level setting */
+	void SetAudioQualityLevel(int32 QualityLevel);
+
+	/** Returns the user's audio quality level setting */
+	int32 GetAudioQualityLevel() const { return AudioQualityLevel; }
+
 	/** Checks if any user settings is different from current */
 	virtual bool IsDirty() const;
 
@@ -163,6 +169,9 @@ protected:
 	UPROPERTY(config)
 	uint32 Version;
 
+	UPROPERTY(config)
+	int32 AudioQualityLevel;
+
 	/**
 	 * Check if the current version of the game user settings is valid. Sub-classes can override this to provide game-specific versioning as necessary.
 	 * @return True if the current version is valid, false if it is not
diff --git a/Engine/Source/Runtime/Engine/Classes/Sound/AudioSettings.h b/Engine/Source/Runtime/Engine/Classes/Sound/AudioSettings.h
--- a/Engine/Source/Runtime/Engine/Classes/Sound/AudioSettings.h
+++ b/Engine/Source/Runtime/Engine/Classes/Sound/AudioSettings.h
@@ -5,6 +5,25 @@
 #include "AudioSettings.generated.h"
 
 
+USTRUCT()
+struct ENGINE_API FAudioQualitySettings
+{
+	GENERATED_BODY()
+
+	UPROPERTY(EditAnywhere, Category="Quality")
+	FText DisplayName;
+
+	// The number of audio channels that can be used at once
+	// NOTE: Some platforms may cap this value to a lower setting regardless of what the settings request
+	UPROPERTY(EditAnywhere, Category="Quality", meta=(ClampMin="1"))
+	int32 MaxChannels;
+
+	FAudioQualitySettings()
+		: MaxChannels(32)
+	{
+	}
+};
+
 /**
  * Audio settings.
  */
@@ -13,6 +32,11 @@ class ENGINE_API UAudioSettings : public UDeveloperSettings
 {
 	GENERATED_UCLASS_BODY()
 
+#if WITH_EDITOR
+	virtual void PreEditChange(UProperty* PropertyAboutToChange) override;
+	virtual void PostEditChangeChainProperty( struct FPropertyChangedChainEvent& PropertyChangedEvent) override;
+#endif
+
 	/** The SoundClass assigned to newly created sounds */
 	UPROPERTY(config, EditAnywhere, Category="Audio", meta=(AllowedClasses="SoundClass", DisplayName="Default Sound Class"))
 	FStringAssetReference DefaultSoundClassName;
@@ -31,4 +55,17 @@ class ENGINE_API UAudioSettings : public UDeveloperSettings
 	/** How many streaming sounds can be played at the same time (if more are played they will be sorted by priority) */
 	UPROPERTY(config, EditAnywhere, Category="Audio", meta=(ClampMin=0))
 	int32 MaximumConcurrentStreams;
+
+	UPROPERTY(config, EditAnywhere, Category="Quality")
+	TArray<FAudioQualitySettings> QualityLevels;
+
+	const FAudioQualitySettings& GetQualityLevelSettings(int32 QualityLevel) const;
+
+private:
+
+#if WITH_EDITOR
+	TArray<FAudioQualitySettings> CachedQualityLevels;
+#endif
+
+	void AddDefaultSettings();
 };
diff --git a/Engine/Source/Runtime/Engine/Classes/Sound/SoundNodeAssetReferencer.h b/Engine/Source/Runtime/Engine/Classes/Sound/SoundNodeAssetReferencer.h
new file mode 100644
--- /dev/null
+++ b/Engine/Source/Runtime/Engine/Classes/Sound/SoundNodeAssetReferencer.h
@@ -0,0 +1,19 @@
+// Copyright 1998-2015 Epic Games, Inc. All Rights Reserved.
+
+
+#pragma once
+#include "Sound/SoundNode.h"
+#include "SoundNodeAssetReferencer.generated.h"
+
+/** 
+ * Sound node that contains a reference to the raw wave file to be played
+ */
+UCLASS(abstract)
+class ENGINE_API USoundNodeAssetReferencer : public USoundNode
+{
+	GENERATED_BODY()
+
+public:
+	virtual void LoadAsset() PURE_VIRTUAL(USoundNodesAssetReferencer::LoadAsset,);
+};
+
diff --git a/Engine/Source/Runtime/Engine/Classes/Sound/SoundNodeQualityLevel.h b/Engine/Source/Runtime/Engine/Classes/Sound/SoundNodeQualityLevel.h
new file mode 100644
--- /dev/null
+++ b/Engine/Source/Runtime/Engine/Classes/Sound/SoundNodeQualityLevel.h
@@ -0,0 +1,38 @@
+// Copyright 1998-2015 Epic Games, Inc. All Rights Reserved.
+
+ 
+#pragma once
+#include "Sound/SoundNode.h"
+#include "SoundNodeQualityLevel.generated.h"
+
+/**
+ * This SoundNode uses GameUserSettings AudioQualityLevel (or the editor override) to choose which branch to play
+ * and at runtime will only load in to memory sound waves connected to the branch that will be selected
+ */
+UCLASS(hidecategories=Object, editinlinenew, MinimalAPI, meta=( DisplayName="Quality Level" ))
+class USoundNodeQualityLevel : public USoundNode
+{
+	GENERATED_BODY()
+
+public:
+
+#if WITH_EDITOR
+	// Begin UObject Interface
+	virtual void PostLoad() override;
+	// End UObject Interface
+#endif
+
+	// Begin USoundNode interface.
+	virtual void ParseNodes( FAudioDevice* AudioDevice, const UPTRINT NodeWaveInstanceHash, FActiveSound& ActiveSound, const FSoundParseParameters& ParseParams, TArray<FWaveInstance*>& WaveInstances ) override;
+	virtual int32 GetMaxChildNodes() const override;
+	virtual int32 GetMinChildNodes() const override;
+#if WITH_EDITOR
+	virtual FText GetInputPinName(int32 PinIndex) const override;
+#endif
+	// End USoundNode interface.
+
+#if WITH_EDITOR
+	void ReconcileNode(bool bReconstructNode);
+#endif
+};
+
diff --git a/Engine/Source/Runtime/Engine/Classes/Sound/SoundNodeWavePlayer.h b/Engine/Source/Runtime/Engine/Classes/Sound/SoundNodeWavePlayer.h
--- a/Engine/Source/Runtime/Engine/Classes/Sound/SoundNodeWavePlayer.h
+++ b/Engine/Source/Runtime/Engine/Classes/Sound/SoundNodeWavePlayer.h
@@ -2,16 +2,16 @@
 
 
 #pragma once
-#include "Sound/SoundNode.h"
+#include "Sound/SoundNodeAssetReferencer.h"
 #include "SoundNodeWavePlayer.generated.h"
 
 /** 
  * Sound node that contains a reference to the raw wave file to be played
  */
 UCLASS(hidecategories=Object, editinlinenew, MinimalAPI, meta=( DisplayName="Wave Player" ))
-class USoundNodeWavePlayer : public USoundNode
+class USoundNodeWavePlayer : public USoundNodeAssetReferencer
 {
-	GENERATED_UCLASS_BODY()
+	GENERATED_BODY()
 
 private:
 	UPROPERTY(EditAnywhere, Category=WavePlayer, meta=(DisplayName="Sound Wave"))
@@ -26,8 +26,6 @@ public:
 
 	ENGINE_API USoundWave* GetSoundWave() const { return SoundWave; }
 	ENGINE_API void SetSoundWave(USoundWave* SoundWave);
-	void LoadSoundWave();
-
 
 	// Begin UObject Interface
 	static void AddReferencedObjects(UObject* InThis, FReferenceCollector& Collector);
@@ -46,5 +44,9 @@ public:
 #endif
 	// End USoundNode Interface
 
+	// Begin USoundNodeAssetReferencer Interface
+	virtual void LoadAsset() override;
+	// End USoundNode Interface
+
 };
 
diff --git a/Engine/Source/Runtime/Engine/Private/AudioDevice.cpp b/Engine/Source/Runtime/Engine/Private/AudioDevice.cpp
--- a/Engine/Source/Runtime/Engine/Private/AudioDevice.cpp
+++ b/Engine/Source/Runtime/Engine/Private/AudioDevice.cpp
@@ -48,7 +48,7 @@ FAudioEffectsManager* FAudioDevice::CreateEffectsManager()
 	return new FAudioEffectsManager(this);
 }
 
-bool FAudioDevice::Init()
+bool FAudioDevice::Init(int32 InMaxChannels)
 {
 	if (bIsInitialized)
 	{
@@ -57,7 +57,7 @@ bool FAudioDevice::Init()
 
 	bool bDeferStartupPrecache = false;
 	// initialize config variables
-	verify(GConfig->GetInt(TEXT("Audio"), TEXT("MaxChannels"), MaxChannels, GEngineIni));
+	MaxChannels = InMaxChannels;
 	verify(GConfig->GetInt(TEXT("Audio"), TEXT("CommonAudioPoolSize"), CommonAudioPoolSize, GEngineIni));
 	
 	// If this is true, skip the initial startup precache so we can do it later in the flow
diff --git a/Engine/Source/Runtime/Engine/Private/AudioDeviceManager.cpp b/Engine/Source/Runtime/Engine/Private/AudioDeviceManager.cpp
--- a/Engine/Source/Runtime/Engine/Private/AudioDeviceManager.cpp
+++ b/Engine/Source/Runtime/Engine/Private/AudioDeviceManager.cpp
@@ -138,7 +138,7 @@ FAudioDevice* FAudioDeviceManager::CreateAudioDevice(uint32& HandleOut, bool bCr
 	++NumActiveAudioDevices;
 
 	const UAudioSettings* AudioSettings = GetDefault<UAudioSettings>();
-	if (!NewAudioDevice->Init())
+	if (!NewAudioDevice->Init(AudioSettings->GetQualityLevelSettings(GEngine->GetGameUserSettings()->GetAudioQualityLevel()).MaxChannels))
 	{
 		ShutdownAudioDevice(HandleOut);
 		HandleOut = AUDIO_DEVICE_HANDLE_INVALID;
diff --git a/Engine/Source/Runtime/Engine/Private/AudioSettings.cpp b/Engine/Source/Runtime/Engine/Private/AudioSettings.cpp
--- a/Engine/Source/Runtime/Engine/Private/AudioSettings.cpp
+++ b/Engine/Source/Runtime/Engine/Private/AudioSettings.cpp
@@ -6,10 +6,94 @@
 
 #include "EnginePrivate.h"
 #include "Sound/AudioSettings.h"
+#include "Sound/SoundNodeQualityLevel.h"
 
+#define LOCTEXT_NAMESPACE "AudioSettings"
 
 UAudioSettings::UAudioSettings(const FObjectInitializer& ObjectInitializer)
 : Super(ObjectInitializer)
 {
 	SectionName = TEXT("Audio");
+
+	AddDefaultSettings();
 }
+
+void UAudioSettings::AddDefaultSettings()
+{
+	FAudioQualitySettings DefaultSettings;
+	DefaultSettings.DisplayName = LOCTEXT("DefaultSettingsName", "Default");
+	GConfig->GetInt(TEXT("Audio"), TEXT("MaxChannels"), DefaultSettings.MaxChannels, GEngineIni); // for backwards compatibility
+	QualityLevels.Add(DefaultSettings);
+}
+
+#if WITH_EDITOR
+void UAudioSettings::PreEditChange(UProperty* PropertyAboutToChange)
+{
+	// Cache at least the first entry in case someone tries to clear the array
+	CachedQualityLevels = QualityLevels;
+}
+
+void UAudioSettings::PostEditChangeChainProperty(FPropertyChangedChainEvent& PropertyChangedEvent)
+{
+	if (PropertyChangedEvent.Property)
+	{
+		bool bReconcileNodes = false;
+
+		if (PropertyChangedEvent.Property->GetFName() == GET_MEMBER_NAME_CHECKED(UAudioSettings, QualityLevels))
+		{
+			if (QualityLevels.Num() == 0)
+			{
+				QualityLevels.Add(CachedQualityLevels[0]);
+			}
+			else if (QualityLevels.Num() > CachedQualityLevels.Num())
+			{
+				for (FAudioQualitySettings& AQSettings : QualityLevels)
+				{
+					if (AQSettings.DisplayName.IsEmpty())
+					{
+						bool bFoundDuplicate;
+						int32 NewQualityLevelIndex = 0;
+						FText NewLevelName;
+						do 
+						{
+							bFoundDuplicate = false;
+							NewLevelName = FText::Format(LOCTEXT("NewQualityLevelName","New Level{0}"), (NewQualityLevelIndex > 0 ? FText::FromString(FString::Printf(TEXT(" %d"),NewQualityLevelIndex)) : FText::GetEmpty()));
+							for (const FAudioQualitySettings& AQSettings : QualityLevels)
+							{
+								if (AQSettings.DisplayName.EqualTo(NewLevelName))
+								{
+									bFoundDuplicate = true;
+									break;
+								}
+							}
+							NewQualityLevelIndex++;
+						} while (bFoundDuplicate);
+						AQSettings.DisplayName = NewLevelName;
+					}
+				}
+			}
+
+			bReconcileNodes = true;
+		}
+		else if (PropertyChangedEvent.Property->GetFName() == GET_MEMBER_NAME_CHECKED(FAudioQualitySettings, DisplayName))
+		{
+			bReconcileNodes = true;
+		}
+
+		if (bReconcileNodes)
+		{
+			for (TObjectIterator<USoundNodeQualityLevel> It; It; ++It)
+			{
+				It->ReconcileNode(true);
+			}
+		}
+	}
+}
+#endif
+
+const FAudioQualitySettings& UAudioSettings::GetQualityLevelSettings(int32 QualityLevel) const
+{
+	return QualityLevels[FMath::Clamp(0, QualityLevel, QualityLevels.Num() - 1)];
+}
+
+#undef LOCTEXT_NAMESPACE
\ No newline at end of file
diff --git a/Engine/Source/Runtime/Engine/Private/GameUserSettings.cpp b/Engine/Source/Runtime/Engine/Private/GameUserSettings.cpp
--- a/Engine/Source/Runtime/Engine/Private/GameUserSettings.cpp
+++ b/Engine/Source/Runtime/Engine/Private/GameUserSettings.cpp
@@ -383,3 +383,7 @@ void UGameUserSettings::SetBenchmarkFallbackValues()
 	ScalabilityQuality.SetBenchmarkFallback();
 }
 
+void UGameUserSettings::SetAudioQualityLevel(int32 QualityLevel)
+{
+	AudioQualityLevel = QualityLevel;
+}
\ No newline at end of file
diff --git a/Engine/Source/Runtime/Engine/Private/SoundCue.cpp b/Engine/Source/Runtime/Engine/Private/SoundCue.cpp
--- a/Engine/Source/Runtime/Engine/Private/SoundCue.cpp
+++ b/Engine/Source/Runtime/Engine/Private/SoundCue.cpp
@@ -5,6 +5,7 @@
 #include "Sound/SoundNodeMixer.h"
 #include "Sound/SoundNodeWavePlayer.h"
 #include "Sound/SoundNodeAttenuation.h"
+#include "Sound/SoundNodeQualityLevel.h"
 #include "Sound/SoundWave.h"
 #if WITH_EDITOR
 #include "Kismet2/BlueprintEditorUtils.h"
@@ -79,9 +80,9 @@ void USoundCue::PostLoad()
 		// Always load all sound waves in the editor
 		for (USoundNode* SoundNode : AllNodes)
 		{
-			if (USoundNodeWavePlayer* WavePlayerNode = Cast<USoundNodeWavePlayer>(SoundNode))
+			if (USoundNodeAssetReferencer* AssetReferencerNode = Cast<USoundNodeAssetReferencer>(SoundNode))
 			{
-				WavePlayerNode->LoadSoundWave();
+				AssetReferencerNode->LoadAsset();
 			}
 		}
 	}
@@ -95,15 +96,19 @@ void USoundCue::PostLoad()
 		{
 			if (USoundNode* SoundNode = NodesToEvaluate.Pop(false))
 			{
-				if (USoundNodeWavePlayer* WavePlayerNode = Cast<USoundNodeWavePlayer>(SoundNode))
+				if (USoundNodeAssetReferencer* AssetReferencerNode = Cast<USoundNodeAssetReferencer>(SoundNode))
 				{
-					WavePlayerNode->LoadSoundWave();
+					AssetReferencerNode->LoadAsset();
 				}
-				/* else if USoundNodeQuality
+				else if (USoundNodeQualityLevel* QualityLevelNode = Cast<USoundNodeQualityLevel>(SoundNode))
 				{
-					// Only pick the node connected for current quality
+					// Only pick the node connected for current quality, currently don't support changing audio quality on the fly
+					static const int32 CachedQualityLevel = GEngine->GetGameUserSettings()->GetAudioQualityLevel();
+					if (CachedQualityLevel < QualityLevelNode->ChildNodes.Num())
+					{
+						NodesToEvaluate.Add(QualityLevelNode->ChildNodes[CachedQualityLevel]);
+					}
 				}
-				*/
 				else
 				{
 					NodesToEvaluate.Append(SoundNode->ChildNodes);
diff --git a/Engine/Source/Runtime/Engine/Private/SoundNodeQualityLevel.cpp b/Engine/Source/Runtime/Engine/Private/SoundNodeQualityLevel.cpp
new file mode 100644
--- /dev/null
+++ b/Engine/Source/Runtime/Engine/Private/SoundNodeQualityLevel.cpp
@@ -0,0 +1,83 @@
+// Copyright 1998-2015 Epic Games, Inc. All Rights Reserved.
+
+
+#include "EnginePrivate.h"
+#include "ActiveSound.h"
+#include "Sound/SoundNodeQualityLevel.h"
+
+#if WITH_EDITOR
+void USoundNodeQualityLevel::PostLoad()
+{
+	Super::PostLoad();
+
+	ReconcileNode(false);
+}
+
+void USoundNodeQualityLevel::ReconcileNode(bool bReconstructNode)
+{
+	while (ChildNodes.Num() > GetMinChildNodes())
+	{
+		RemoveChildNode(ChildNodes.Num()-1);
+	}
+	while (ChildNodes.Num() < GetMinChildNodes())
+	{
+		InsertChildNode(ChildNodes.Num());
+	}
+#if WITH_EDITORONLY_DATA
+	if (GIsEditor && bReconstructNode && GraphNode)
+	{
+		GraphNode->ReconstructNode();
+		GraphNode->GetGraph()->NotifyGraphChanged();
+	}
+#endif
+}
+
+FText USoundNodeQualityLevel::GetInputPinName(int32 PinIndex) const
+{
+	return GetDefault<UAudioSettings>()->GetQualityLevelSettings(PinIndex).DisplayName;
+}
+#endif
+
+int32 USoundNodeQualityLevel::GetMaxChildNodes() const
+{
+	return GetDefault<UAudioSettings>()->QualityLevels.Num();
+}
+
+int32 USoundNodeQualityLevel::GetMinChildNodes() const
+{
+	return GetDefault<UAudioSettings>()->QualityLevels.Num();
+}
+
+void USoundNodeQualityLevel::ParseNodes( FAudioDevice* AudioDevice, const UPTRINT NodeWaveInstanceHash, FActiveSound& ActiveSound, const FSoundParseParameters& ParseParams, TArray<FWaveInstance*>& WaveInstances )
+{
+#if WITH_EDITOR
+	int32 QualityLevel = 0;
+
+	if (GIsEditor)
+	{
+		RETRIEVE_SOUNDNODE_PAYLOAD( sizeof( int32 ) );
+		DECLARE_SOUNDNODE_ELEMENT( int32, CachedQualityLevel );
+
+		if (*RequiresInitialization)
+		{
+			const bool bIsPIESound = ((GEditor->bIsSimulatingInEditor || GEditor->PlayWorld != NULL) && ActiveSound.World != NULL);
+			if (bIsPIESound)
+			{
+				QualityLevel = GetDefault<ULevelEditorPlaySettings>()->PlayInEditorSoundQuality;
+			}
+		}
+	}
+	else
+	{
+		static const int32 CachedQualityLevel = GEngine->GetGameUserSettings()->GetAudioQualityLevel();
+		QualityLevel = CachedQualityLevel;
+	}
+#else
+	static const int32 QualityLevel = GEngine->GetGameUserSettings()->GetAudioQualityLevel();
+#endif
+
+	if (QualityLevel < ChildNodes.Num() && ChildNodes[QualityLevel])
+	{
+		ChildNodes[QualityLevel]->ParseNodes( AudioDevice, GetNodeWaveInstanceHash(NodeWaveInstanceHash, ChildNodes[QualityLevel], QualityLevel), ActiveSound, ParseParams, WaveInstances );
+	}
+}
\ No newline at end of file
diff --git a/Engine/Source/Runtime/Engine/Private/SoundNodeWavePlayer.cpp b/Engine/Source/Runtime/Engine/Private/SoundNodeWavePlayer.cpp
--- a/Engine/Source/Runtime/Engine/Private/SoundNodeWavePlayer.cpp
+++ b/Engine/Source/Runtime/Engine/Private/SoundNodeWavePlayer.cpp
@@ -7,11 +7,6 @@
 
 #define LOCTEXT_NAMESPACE "SoundNodeWavePlayer"
 
-USoundNodeWavePlayer::USoundNodeWavePlayer(const FObjectInitializer& ObjectInitializer)
-	: Super(ObjectInitializer)
-{
-}
-
 void USoundNodeWavePlayer::AddReferencedObjects(UObject* InThis, FReferenceCollector& Collector)
 {
 	USoundNodeWavePlayer* This = CastChecked<USoundNodeWavePlayer>(InThis);
@@ -20,7 +15,7 @@ void USoundNodeWavePlayer::AddReferencedObjects(UObject* InThis, FReferenceColle
 	Collector.AddReferencedObject(This->SoundWave);
 }
 
-void USoundNodeWavePlayer::LoadSoundWave()
+void USoundNodeWavePlayer::LoadAsset()
 {
 	SoundWave = SoundWaveAssetPtr.LoadSynchronous();
 }
@@ -36,7 +31,7 @@ void USoundNodeWavePlayer::PostEditChangeProperty(FPropertyChangedEvent& Propert
 {
 	if (PropertyChangedEvent.Property && PropertyChangedEvent.Property->GetFName() == GET_MEMBER_NAME_CHECKED(USoundNodeWavePlayer, SoundWaveAssetPtr))
 	{
-		LoadSoundWave();
+		LoadAsset();
 	}
 }
 #endif
diff --git a/Engine/Source/Runtime/Engine/Public/AudioDevice.h b/Engine/Source/Runtime/Engine/Public/AudioDevice.h
--- a/Engine/Source/Runtime/Engine/Public/AudioDevice.h
+++ b/Engine/Source/Runtime/Engine/Public/AudioDevice.h
@@ -226,7 +226,7 @@ public:
 	/**
 	 * Basic initialisation of the platform agnostic layer of the audio system
 	 */
-	bool Init( void );
+	bool Init(int32 InMaxChannels);
 
 	/**
 	 * Tears down the audio device
