[UE-15537] When doing a skelton recompile of children blueprints, making sure we pass along the bIsRegeneratingOnLoad flag (to keep from resetting loaders in the midst of loading).

diff --git a/Engine/Source/Editor/UnrealEd/Private/Kismet2/BlueprintEditorUtils.cpp b/Engine/Source/Editor/UnrealEd/Private/Kismet2/BlueprintEditorUtils.cpp
--- a/Engine/Source/Editor/UnrealEd/Private/Kismet2/BlueprintEditorUtils.cpp
+++ b/Engine/Source/Editor/UnrealEd/Private/Kismet2/BlueprintEditorUtils.cpp
@@ -1990,7 +1990,7 @@ void FBlueprintEditorUtils::MarkBlueprintAsStructurallyModified(UBlueprint* Blue
 {
 	struct FRefreshHelper
 	{
-		static void SkeletalRecompileChildren(TArray<UClass*> SkelClassesToRecompile)
+		static void SkeletalRecompileChildren(TArray<UClass*> SkelClassesToRecompile, bool bIsCompilingOnLoad)
 		{
 			for (auto SkelClass : SkelClassesToRecompile)
 			{
@@ -1999,7 +1999,7 @@ void FBlueprintEditorUtils::MarkBlueprintAsStructurallyModified(UBlueprint* Blue
 					continue;
 				}
 
-				auto SkelBlueprint = Cast<UBlueprint>(SkelClass->ClassGeneratedBy);
+				UBlueprint* SkelBlueprint = Cast<UBlueprint>(SkelClass->ClassGeneratedBy);
 				if (SkelBlueprint
 					&& SkelBlueprint->Status != BS_BeingCreated
 					&& !SkelBlueprint->bBeingCompiled
@@ -2015,13 +2015,20 @@ void FBlueprintEditorUtils::MarkBlueprintAsStructurallyModified(UBlueprint* Blue
 					Results.bSilentMode = true;
 					Results.bLogInfoOnly = true;
 
-					FKismetCompilerOptions CompileOptions;
-					CompileOptions.CompileType = EKismetCompileType::SkeletonOnly;
-					Compiler.CompileBlueprint(SkelBlueprint, CompileOptions, Results);
+					//if (!SkelBlueprint->bHasBeenRegenerated)
+					{
+						bool const bWasRegenerating = SkelBlueprint->bIsRegeneratingOnLoad;
+						SkelBlueprint->bIsRegeneratingOnLoad = bIsCompilingOnLoad;
+
+						FKismetCompilerOptions CompileOptions;
+						CompileOptions.CompileType = EKismetCompileType::SkeletonOnly;
+						Compiler.CompileBlueprint(SkelBlueprint, CompileOptions, Results);
 
-					SkelBlueprint->BroadcastCompiled();
+						SkelBlueprint->BroadcastCompiled();
 
-					SkeletalRecompileChildren(ChildrenOfClass);
+						SkeletalRecompileChildren(ChildrenOfClass, bIsCompilingOnLoad);
+						SkelBlueprint->bIsRegeneratingOnLoad = bWasRegenerating;
+					}
 				}
 			}
 		}
@@ -2073,7 +2080,7 @@ void FBlueprintEditorUtils::MarkBlueprintAsStructurallyModified(UBlueprint* Blue
 		}
 		UpdateDelegatesInBlueprint(Blueprint);
 
-		FRefreshHelper::SkeletalRecompileChildren(ChildrenOfClass);
+		FRefreshHelper::SkeletalRecompileChildren(ChildrenOfClass, Blueprint->bIsRegeneratingOnLoad);
 
 		{
 			BP_SCOPED_COMPILER_EVENT_STAT(EKismetCompilerStats_NotifyBlueprintChanged);
