UE-14675: [CrashReport] Crash occurring during hot reload when construction scripts are being re-run.

diff --git a/Engine/Source/Editor/UnrealEd/Private/Kismet2/KismetReinstanceUtilities.cpp b/Engine/Source/Editor/UnrealEd/Private/Kismet2/KismetReinstanceUtilities.cpp
--- a/Engine/Source/Editor/UnrealEd/Private/Kismet2/KismetReinstanceUtilities.cpp
+++ b/Engine/Source/Editor/UnrealEd/Private/Kismet2/KismetReinstanceUtilities.cpp
@@ -923,7 +923,7 @@ void FBlueprintCompileReinstancer::ReplaceInstancesOfClass(UClass* OldClass, UCl
 	TArray<UObject*> PotentialEditorsForRefreshing;
 
 	// A list of component owners that need their construction scripts re-ran (because a component of theirs has been reinstanced)
-	TSet<AActor*> OwnersToReconstruct;
+	TSet<AActor*> OwnersToRerunConstructionScript;
 
 	// Set global flag to let system know we are reconstructing blueprint instances
 	TGuardValue<bool> GuardTemplateNameFlag(GIsReconstructingBlueprintInstances, true);
@@ -1100,7 +1100,12 @@ void FBlueprintCompileReinstancer::ReplaceInstancesOfClass(UClass* OldClass, UCl
 						// we need to keep track of actor instances that need 
 						// their construction scripts re-ran (since we've just 
 						// replaced a component they own)
-						OwnersToReconstruct.Add(OwningActor);
+						
+						// Skipping CDOs as CSs are not allowed for them.
+						if (!OwningActor->HasAnyFlags(RF_ClassDefaultObject))
+						{
+							OwnersToRerunConstructionScript.Add(OwningActor);
+						}
 					}
 				}
 			}
@@ -1179,7 +1184,7 @@ void FBlueprintCompileReinstancer::ReplaceInstancesOfClass(UClass* OldClass, UCl
 
 	// in the case where we're replacing component instances, we need to make 
 	// sure to re-run their owner's construction scripts
-	for (AActor* ActorInstance : OwnersToReconstruct)
+	for (AActor* ActorInstance : OwnersToRerunConstructionScript)
 	{
 		ActorInstance->RerunConstructionScripts();
 	}
