Cooker now checks project settings cook map list then all maps, still handles mapsinisection if provided on commandline.  Stopped UAT from passing in mapsinisection=AllMaps.

diff --git a/Engine/Source/Editor/MainFrame/Private/Frame/MainFrameActions.cpp b/Engine/Source/Editor/MainFrame/Private/Frame/MainFrameActions.cpp
--- a/Engine/Source/Editor/MainFrame/Private/Frame/MainFrameActions.cpp
+++ b/Engine/Source/Editor/MainFrame/Private/Frame/MainFrameActions.cpp
@@ -506,17 +506,6 @@ void FMainFrameActionCallbacks::PackageProject( const FName InPlatformInfoName )
 			OptionalParams += TEXT(" -CookMapsOnly");
 		}
 	}
-	else if ( PackagingSettings->MapsToCook.Num() )
-	{
-		OptionalParams += TEXT(" -mapstocook=");
-		for ( const auto& Map : PackagingSettings->MapsToCook )
-		{
-			OptionalParams += FString::Printf(TEXT("%s+"), *Map.FilePath);
-		}
-		OptionalParams.RemoveFromEnd("+");
-	}
-	
-	
 
 	if (PackagingSettings->UsePakFile)
 	{
diff --git a/Engine/Source/Editor/UnrealEd/Private/Commandlets/CookCommandlet.cpp b/Engine/Source/Editor/UnrealEd/Private/Commandlets/CookCommandlet.cpp
--- a/Engine/Source/Editor/UnrealEd/Private/Commandlets/CookCommandlet.cpp
+++ b/Engine/Source/Editor/UnrealEd/Private/Commandlets/CookCommandlet.cpp
@@ -1059,35 +1059,39 @@ bool UCookCommandlet::NewCook( const TArray<ITargetPlatform*>& Platforms, TArray
 	FString CreateReleaseVersion;
 	FParse::Value( *Params, TEXT("CreateReleaseVersion="), CreateReleaseVersion);
 
-	/*FString AssetRegistry;
-	if (FParse::Value(*Params, TEXT("SHIPPEDASSETREGISTRY="), AssetRegistry))
+	// Add any map sections specified on command line
+	TArray<FString> AlwaysCookMapList;
+
+	// Add the default map section
+	GEditor->LoadMapListFromIni(TEXT("AlwaysCookMaps"), AlwaysCookMapList);
+
+	TArray<FString> MapList;
+	// Add any map sections specified on command line
+	GEditor->ParseMapSectionIni(*Params, MapList);
+
+	if (MapList.Num() == 0)
 	{
-		TArray<FName> TargetPlatformNames;
-		for (const auto &Platform : Platforms)
+		// if we didn't find any maps look in the project settings for maps
+
+		UProjectPackagingSettings* PackagingSettings = Cast<UProjectPackagingSettings>(UProjectPackagingSettings::StaticClass()->GetDefaultObject());
+
+		for (const auto& MapToCook : PackagingSettings->MapsToCook)
 		{
-			FName PlatformName = FName(*Platform->PlatformName());
-			TargetPlatformNames.Add(PlatformName); // build list of all target platform names
+			MapList.Add(MapToCook.FilePath);
 		}
-		CookOnTheFlyServer->WarmCookedPackages(FPaths::GameContentDir() / AssetRegistry, TargetPlatformNames);
-	}*/
+	}
 
-	TArray<FString> CmdLineIniSections;
-	FString SectionStr;
-	if (FParse::Value(*Params, TEXT("MAPINISECTION="), SectionStr))
+	// if we still don't have any mapsList check if the allmaps ini section is filled out
+	// this is for backwards compatibility
+	if (MapList.Num() == 0)
 	{
-		if (SectionStr.Contains(TEXT("+")))
-		{
-			SectionStr.ParseIntoArray(CmdLineIniSections,TEXT("+"),true);
-		}
-		else
-		{
-			CmdLineIniSections.Add(SectionStr);
-		}
+		GEditor->ParseMapSectionIni(TEXT("-MAPINISECTION=AllMaps"), MapList);
 	}
 
-	// Add any map sections specified on command line
-	TArray<FString> MapList;
-	GEditor->ParseMapSectionIni(*Params, MapList);
+
+	// put the always cook map list at the front of the map list
+	AlwaysCookMapList.Append(MapList);
+	Swap(MapList, AlwaysCookMapList);
 
 	TArray<FString> CmdLineMapEntries;
 	TArray<FString> CmdLineDirEntries;
diff --git a/Engine/Source/Programs/AutomationTool/CommandletUtils.cs b/Engine/Source/Programs/AutomationTool/CommandletUtils.cs
--- a/Engine/Source/Programs/AutomationTool/CommandletUtils.cs
+++ b/Engine/Source/Programs/AutomationTool/CommandletUtils.cs
@@ -26,7 +26,7 @@ namespace AutomationTool
 			string MapsToCook = "";
 			if (IsNullOrEmpty(Maps))
 			{
-				MapsToCook = "-MapIniSection=AllMaps";
+				// MapsToCook = "-MapIniSection=AllMaps";
 			}
 			else
 			{
