Registry temporary directory as mount point when generating asset registry as we could be cooking from the temporary saved directory incase of editor.

diff --git a/Engine/Source/Editor/UnrealEd/Private/CookOnTheFlyServer.cpp b/Engine/Source/Editor/UnrealEd/Private/CookOnTheFlyServer.cpp
--- a/Engine/Source/Editor/UnrealEd/Private/CookOnTheFlyServer.cpp
+++ b/Engine/Source/Editor/UnrealEd/Private/CookOnTheFlyServer.cpp
@@ -2710,6 +2710,8 @@ void UCookOnTheFlyServer::CleanSandbox( const bool bIterative )
 
 void UCookOnTheFlyServer::GenerateAssetRegistry(const TArray<ITargetPlatform*>& Platforms)
 {
+	FPackageName::RegisterMountPoint(TEXT("/Temp/"), FPaths::GameSavedDir());
+
 	// load the interface
 	FAssetRegistryModule& AssetRegistryModule = FModuleManager::LoadModuleChecked<FAssetRegistryModule>(TEXT("AssetRegistry"));
 	IAssetRegistry& AssetRegistry = AssetRegistryModule.Get();
