Fix for the Split niagara op

diff --git a/Engine/Source/Editor/UnrealEd/Private/Niagara/NiagaraEditorCommon.cpp b/Engine/Source/Editor/UnrealEd/Private/Niagara/NiagaraEditorCommon.cpp
--- a/Engine/Source/Editor/UnrealEd/Private/Niagara/NiagaraEditorCommon.cpp
+++ b/Engine/Source/Editor/UnrealEd/Private/Niagara/NiagaraEditorCommon.cpp
@@ -450,7 +450,7 @@ void FNiagaraOpInfo::Init()
 
 	Op = &OpInfos.Add(Split);
 	Op->Name = Split;
-	Op->FriendlyName = NSLOCTEXT("NiagaraOpInfo", "Split Name", "Compose");
+	Op->FriendlyName = NSLOCTEXT("NiagaraOpInfo", "Split Name", "Split");
 	Op->Description = NSLOCTEXT("NiagaraOpInfo", "Split Desc", "X = A.xxxx, Y = A.yyyy, Z = A.zzzz, W = A.wwww");
 	Op->Inputs.Add(FNiagaraOpInOutInfo(A, ENiagaraDataType::Vector, AText, AText, DefaultStr_VecOne));
 	Op->Outputs.Add(FNiagaraOpInOutInfo(X, ENiagaraDataType::Vector, XText, XText, DefaultStr_VecOne));
