Resolved an issue where missing keys and conflicting identities in text within assets were resolved but not gathered properly after resolution.

diff --git a/Engine/Source/Editor/UnrealEd/Private/Commandlets/GatherTextFromAssetsCommandlet.cpp b/Engine/Source/Editor/UnrealEd/Private/Commandlets/GatherTextFromAssetsCommandlet.cpp
--- a/Engine/Source/Editor/UnrealEd/Private/Commandlets/GatherTextFromAssetsCommandlet.cpp
+++ b/Engine/Source/Editor/UnrealEd/Private/Commandlets/GatherTextFromAssetsCommandlet.cpp
@@ -571,7 +571,7 @@ void UGatherTextFromAssetsCommandlet::ProcessTextProperty(UTextProperty* InTextP
 
 	FString Namespace;
 	FString Key;
-	const bool FoundNamespaceAndKey = FTextLocalizationManager::Get().FindNamespaceAndKeyFromDisplayString(Data->DisplayString, Namespace, Key);
+	bool FoundNamespaceAndKey = FTextLocalizationManager::Get().FindNamespaceAndKeyFromDisplayString(Data->DisplayString, Namespace, Key);
 
 	// Check if text is localizable and check for missing key. Unlocalizable texts don't need a key.
 	if( ( !FoundNamespaceAndKey || Key.IsEmpty() ) && Data->ShouldGatherForLocalization() )
@@ -592,6 +592,7 @@ void UGatherTextFromAssetsCommandlet::ProcessTextProperty(UTextProperty* InTextP
 			}
 
 			// Fixed.
+			FoundNamespaceAndKey = FTextLocalizationManager::Get().FindNamespaceAndKeyFromDisplayString(Data->DisplayString, Namespace, Key);
 			NewEntry.Status = EAssetTextGatherStatus::MissingKey_Resolved;
 		}
 	}
@@ -622,6 +623,7 @@ void UGatherTextFromAssetsCommandlet::ProcessTextProperty(UTextProperty* InTextP
 				FTextLocalizationManager::Get().UpdateDisplayString(Data->DisplayString, *Data->DisplayString, Namespace, FGuid::NewGuid().ToString());
 
 				// Fixed.
+				FoundNamespaceAndKey = FTextLocalizationManager::Get().FindNamespaceAndKeyFromDisplayString(Data->DisplayString, Namespace, Key);
 				NewEntry.Status = EAssetTextGatherStatus::IdentityConflict_Resolved;
 
 				// Conflict resolved, no existing entry.
