Implemented gathering of text localization data at serialization-saving time for packages, serializing the localization data into the package header, and loading it from the package header without fully loading the package when gathering text from assets.

diff --git a/Engine/Source/Editor/UnrealEd/Classes/Commandlets/GatherTextFromAssetsCommandlet.h b/Engine/Source/Editor/UnrealEd/Classes/Commandlets/GatherTextFromAssetsCommandlet.h
--- a/Engine/Source/Editor/UnrealEd/Classes/Commandlets/GatherTextFromAssetsCommandlet.h
+++ b/Engine/Source/Editor/UnrealEd/Classes/Commandlets/GatherTextFromAssetsCommandlet.h
@@ -23,7 +23,7 @@ namespace EAssetTextGatherStatus
 }
 
 class FLocMetadataValue;
-class UDialogueWave;
+
 /**
  *	UGatherTextFromAssetsCommandlet: Localization commandlet that collects all text to be localized from the game assets.
  */
@@ -32,9 +32,8 @@ class UGatherTextFromAssetsCommandlet : public UGatherTextCommandletBase
 {
 	GENERATED_UCLASS_BODY()
 
+	void ProcessGatherableTextDataArray(const FString& PackageFilePath, const TArray<FGatherableTextData>& GatherableTextDataArray);
 	void ProcessPackages( const TArray< UPackage* >& PackagesToProcess );
-	void ProcessDialogueWave( const UDialogueWave* DialogueWave );
-	void ProcessTextProperty( UTextProperty* TextProp, FText* Data, UObject* Object, bool& OutRepaired );
 
 public:
 	// Begin UCommandlet Interface
@@ -46,7 +45,7 @@ private:
 	{
 		struct FEntry
 		{
-			FString ObjectPath;
+			FString PathToSourceString;
 			TSharedPtr<FString, ESPMode::ThreadSafe> SourceString;
 			EAssetTextGatherStatus::Type Status;
 		};
@@ -58,51 +57,6 @@ private:
 		FNamespaceTable Namespaces;
 	};
 
-	// Class that takes DialogueWaves and prepares internationalization manifest entries 
-	class FDialogueHelper
-	{
-	public:
-		bool ProcessDialogueWave( const UDialogueWave* DialogueWave );
-		const FString& GetSpokenSource() const { return SpokenSource; }
-		
-		TSharedPtr< const FContext > GetBaseContext() const { return Base; };
-		const TArray< TSharedPtr< FContext > >& GetContextSpecificVariations() const { return ContextSpecificVariations; }
-
-	private:
-		TSharedPtr< FLocMetadataValue > GenSourceTargetMetadata( const FString& SpeakerName, const TArray< FString >& TargetNames, bool bCompact = true );
-		FString ArrayMetaDataToString( const TArray< TSharedPtr< FLocMetadataValue > >& MetadataArray );
-		FString GetDialogueVoiceName( const UDialogueVoice* DialogueVoice );
-		FString GetGrammaticalGenderString( TEnumAsByte<EGrammaticalGender::Type> Gender );
-		FString GetGrammaticalNumberString( TEnumAsByte<EGrammaticalNumber::Type> Plurality );
-
-	public:
-		static const FString DialogueNamespace;
-		static const FString PropertyName_VoiceActorDirection;
-		static const FString PropertyName_Speaker;
-		static const FString PropertyName_Speakers;
-		static const FString PropertyName_Targets;
-		static const FString PropertyName_GrammaticalGender;
-		static const FString PropertyName_GrammaticalPlurality;
-		static const FString PropertyName_TargetGrammaticalGender;
-		static const FString PropertyName_TargetGrammaticalNumber;
-		static const FString PropertyName_Optional;
-		static const FString PropertyName_DialogueVariations;
-		static const FString PropertyName_IsMature;
-
-	private:
-		FString DialogueKey;
-		FString SourceLocation;
-		FString SpokenSource;
-		FString VoiceActorDirection;
-		bool bIsMature;
-
-		// The non-optional entry
-		TSharedPtr< FContext > Base;
-		// Optional entries that are context specific
-		TArray< TSharedPtr< FContext > > ContextSpecificVariations;
-		
-	};
-
 private:
 	static const FString UsageText;
 
diff --git a/Engine/Source/Editor/UnrealEd/Private/Commandlets/GatherTextFromAssetsCommandlet.cpp b/Engine/Source/Editor/UnrealEd/Private/Commandlets/GatherTextFromAssetsCommandlet.cpp
--- a/Engine/Source/Editor/UnrealEd/Private/Commandlets/GatherTextFromAssetsCommandlet.cpp
+++ b/Engine/Source/Editor/UnrealEd/Private/Commandlets/GatherTextFromAssetsCommandlet.cpp
@@ -27,660 +27,98 @@ const FString UGatherTextFromAssetsCommandlet::UsageText
 	TEXT("    <paths to exclude> Paths to exclude. Delimited with ';'. Accepts wildcards. eg \"*Content/Developers/*;*/TestMaps/*\" OPTIONAL: If not present, nothing will be excluded.\r\n")
 	);
 
-const FString UGatherTextFromAssetsCommandlet::FDialogueHelper::DialogueNamespace						= TEXT("Dialogue");
-const FString UGatherTextFromAssetsCommandlet::FDialogueHelper::PropertyName_VoiceActorDirection		= TEXT("Voice Actor Direction");
-const FString UGatherTextFromAssetsCommandlet::FDialogueHelper::PropertyName_Speaker					= TEXT("Speaker");
-const FString UGatherTextFromAssetsCommandlet::FDialogueHelper::PropertyName_Speakers					= TEXT("Speakers");
-const FString UGatherTextFromAssetsCommandlet::FDialogueHelper::PropertyName_Targets					= TEXT("Targets");
-const FString UGatherTextFromAssetsCommandlet::FDialogueHelper::PropertyName_GrammaticalGender			= TEXT("Gender");
-const FString UGatherTextFromAssetsCommandlet::FDialogueHelper::PropertyName_GrammaticalPlurality		= TEXT("Plurality");
-const FString UGatherTextFromAssetsCommandlet::FDialogueHelper::PropertyName_TargetGrammaticalGender	= TEXT("TargetGender");
-const FString UGatherTextFromAssetsCommandlet::FDialogueHelper::PropertyName_TargetGrammaticalNumber	= TEXT("TargetPlurality");
-const FString UGatherTextFromAssetsCommandlet::FDialogueHelper::PropertyName_Optional					= TEXT("Optional");
-const FString UGatherTextFromAssetsCommandlet::FDialogueHelper::PropertyName_DialogueVariations			= TEXT("Variations");
-const FString UGatherTextFromAssetsCommandlet::FDialogueHelper::PropertyName_IsMature					= TEXT("*IsMature");
-
-bool UGatherTextFromAssetsCommandlet::FDialogueHelper::ProcessDialogueWave( const UDialogueWave* DialogueWave )
-{
-	if( !DialogueWave )
-	{
-		return false;
-	}
-
-	DialogueKey = DialogueWave->LocalizationGUID.ToString();
-	SourceLocation = DialogueWave->GetPathName();
-	SpokenSource = DialogueWave->SpokenText;
-#if WITH_EDITORONLY_DATA
-	VoiceActorDirection = DialogueWave->VoiceActorDirection;
-#endif
-	bIsMature = DialogueWave->bMature;
-	
-	// Stores the human readable info describing source and targets for each context of this DialogueWave
-	TArray< TSharedPtr< FLocMetadataValue > > VariationsDisplayInfoList;
-
-	for( const FDialogueContextMapping& ContextMapping : DialogueWave->ContextMappings )
-	{
-		const FDialogueContext& DialogueContext = ContextMapping.Context;
-		const UDialogueVoice* SpeakerDialogueVoice = DialogueContext.Speaker;
-
-		// Skip over entries with invalid speaker
-		if( !SpeakerDialogueVoice )
-		{
-			continue;
-		}
-
-		// Collect speaker info
-		FString SpeakerDisplayName = GetDialogueVoiceName( SpeakerDialogueVoice );
-		FString SpeakerGender = GetGrammaticalGenderString( SpeakerDialogueVoice->Gender );
-		FString SpeakerPlurality = GetGrammaticalNumberString( SpeakerDialogueVoice->Plurality );
-		FString SpeakerGuid = SpeakerDialogueVoice->LocalizationGUID.ToString();
-
-		EGrammaticalGender::Type AccumulatedTargetGender = (EGrammaticalGender::Type)-1;
-		EGrammaticalNumber::Type AccumulatedTargetPlurality = (EGrammaticalNumber::Type)-1;
-
-		TArray<FString> TargetGuidsList;
-		TArray<FString> TargetDisplayNameList;
-
-		// Collect info on all the targets
-		bool HasTarget = false;
-		for( const UDialogueVoice* TargetDialogueVoice : DialogueContext.Targets )
-		{
-			if( TargetDialogueVoice )
-			{
-				FString TargetDisplayName = GetDialogueVoiceName( TargetDialogueVoice );
-				TargetDisplayNameList.AddUnique( TargetDisplayName );
-				FString TargetGender = GetGrammaticalGenderString( TargetDialogueVoice->Gender );
-				FString TargetPlurality = GetGrammaticalNumberString( TargetDialogueVoice->Plurality );
-				FString TargetGuid = TargetDialogueVoice->LocalizationGUID.ToString();
-
-				TargetGuidsList.AddUnique( TargetGuid );
-
-				if( AccumulatedTargetGender == -1)
-				{
-					AccumulatedTargetGender = TargetDialogueVoice->Gender;
-				}
-				else if( AccumulatedTargetGender != TargetDialogueVoice->Gender )
-				{
-					AccumulatedTargetGender = EGrammaticalGender::Mixed;
-				}
-
-				if( AccumulatedTargetPlurality == -1 )
-				{
-					AccumulatedTargetPlurality = TargetDialogueVoice->Plurality;
-				}
-				else if( AccumulatedTargetPlurality == EGrammaticalNumber::Singular )
-				{
-					AccumulatedTargetPlurality = EGrammaticalNumber::Plural;
-				}
-
-				HasTarget = true;
-			}
-		}
-
-		FString FinalTargetGender = HasTarget ? GetGrammaticalGenderString( AccumulatedTargetGender ) : TEXT("");
-		FString FinalTargetPlurality = HasTarget ? GetGrammaticalNumberString( AccumulatedTargetPlurality ) : TEXT("");
-
-		// Add the context specific variation
-		{
-			TSharedPtr< FLocMetadataObject > KeyMetaDataObject = MakeShareable( new FLocMetadataObject() );
-
-			// Setup a loc metadata object with all the context specific keys.
-			{
-				KeyMetaDataObject->SetStringField( PropertyName_GrammaticalGender, SpeakerGender );
-				KeyMetaDataObject->SetStringField( PropertyName_GrammaticalPlurality, SpeakerPlurality );
-				KeyMetaDataObject->SetStringField( PropertyName_Speaker, SpeakerGuid );
-				KeyMetaDataObject->SetStringField( PropertyName_TargetGrammaticalGender, FinalTargetGender );
-				KeyMetaDataObject->SetStringField( PropertyName_TargetGrammaticalNumber, FinalTargetPlurality );
-
-				TArray< TSharedPtr< FLocMetadataValue > > TargetGuidsMetadata;
-				for( FString& TargetGuid : TargetGuidsList )
-				{
-					TargetGuidsMetadata.Add( MakeShareable( new FLocMetadataValueString( TargetGuid ) ) );
-				}
-
-				KeyMetaDataObject->SetArrayField( PropertyName_Targets, TargetGuidsMetadata );
-			}
-
-			// Create the human readable info that describes the source and target of this dialogue
-			TSharedPtr< FLocMetadataValue > SourceTargetInfo = GenSourceTargetMetadata( SpeakerDisplayName, TargetDisplayNameList );
-
-			TSharedPtr< FLocMetadataObject > InfoMetaDataObject = MakeShareable( new FLocMetadataObject() );
-
-			// Setup a loc metadata object with all the context specific info.  This usually includes human readable descriptions of the dialogue
-			{
-				if( SourceTargetInfo.IsValid() )
-				{
-					InfoMetaDataObject->SetField( PropertyName_DialogueVariations, SourceTargetInfo );
-				}
-
-				if( !VoiceActorDirection.IsEmpty() )
-				{
-					InfoMetaDataObject->SetStringField( PropertyName_VoiceActorDirection, VoiceActorDirection );
-				}
-			}
-
-
-			TSharedRef< FContext > Context = MakeShareable( new FContext );
-
-			// Setup the context
-			{
-				Context->Key = DialogueWave->GetContextLocalizationKey( DialogueContext );
-				Context->SourceLocation = SourceLocation;
-				Context->bIsOptional = true;
-				Context->KeyMetadataObj = KeyMetaDataObject->Values.Num() > 0 ? KeyMetaDataObject : NULL;
-				Context->InfoMetadataObj = InfoMetaDataObject->Values.Num() > 0 ? InfoMetaDataObject : NULL;
-			}
-
-			// Add this context specific variation to
-			ContextSpecificVariations.Add( Context );
-
-			{
-				// Add human readable info describing the source and targets of this dialogue to the non-optional manifest entry if it does not exist already
-				bool bAddContextDisplayInfoToBase = true;
-				for( TSharedPtr< FLocMetadataValue > VariationInfo : VariationsDisplayInfoList )
-				{
-					if( *SourceTargetInfo == *VariationInfo)
-					{
-						bAddContextDisplayInfoToBase = false;
-						break;
-					}
-				}
-				if( bAddContextDisplayInfoToBase )
-				{
-					VariationsDisplayInfoList.Add( SourceTargetInfo );
-				}
-			}
-		}
-	}
-
-	// Create the base, non-optional entry
-	{
-		TSharedPtr< FLocMetadataObject > InfoMetaDataObject = MakeShareable( new FLocMetadataObject() );
-
-		// Setup a loc metadata object with all the context specific info.  This usually includes human readable descriptions of the dialogue
-		{
-			if( VariationsDisplayInfoList.Num() > 0 )
-			{
-				InfoMetaDataObject->SetArrayField( PropertyName_DialogueVariations, VariationsDisplayInfoList );
-			}
-
-			if( !VoiceActorDirection.IsEmpty() )
-			{
-				InfoMetaDataObject->SetStringField( PropertyName_VoiceActorDirection, VoiceActorDirection );
-			}
-		}
-
-		TSharedRef< FContext > Context = MakeShareable( new FContext );
-
-		// Setup the context
-		{
-			Context->Key = DialogueKey;
-			Context->SourceLocation = SourceLocation;
-			Context->bIsOptional = false;
-			Context->InfoMetadataObj = InfoMetaDataObject->Values.Num() > 0 ? InfoMetaDataObject : NULL;
-		}
-
-		Base = Context;
-	}
-
-	return true;
-}
-
-FString UGatherTextFromAssetsCommandlet::FDialogueHelper::GetGrammaticalNumberString(TEnumAsByte<EGrammaticalNumber::Type> Plurality)
-{
-	const static UEnum* Enum = FindObject<UEnum>(ANY_PACKAGE, TEXT("EGrammaticalNumber"), true);
-	check (Enum);
-	check (Enum->NumEnums() > Plurality);
-
-	const FString KeyName = TEXT("DisplayName");
-	const FString PluralityString = Enum->GetMetaData(*KeyName, Plurality);
-
-	return PluralityString;
-}
-
-
-
-FString UGatherTextFromAssetsCommandlet::FDialogueHelper::GetGrammaticalGenderString(TEnumAsByte<EGrammaticalGender::Type> Gender)
-{
-	const static UEnum* Enum = FindObject<UEnum>(ANY_PACKAGE, TEXT("EGrammaticalGender"), true);
-	check (Enum);
-	check (Enum->NumEnums() > Gender);
-
-	const FString KeyName = TEXT("DisplayName");
-	const FString GenderString = Enum->GetMetaData(*KeyName, Gender);
-
-	return GenderString;
-}
-
-
-FString UGatherTextFromAssetsCommandlet::FDialogueHelper::GetDialogueVoiceName( const UDialogueVoice* DialogueVoice )
-{
-	FString Name = DialogueVoice->GetName();
-	return Name;
-}
-
-FString UGatherTextFromAssetsCommandlet::FDialogueHelper::ArrayMetaDataToString( const TArray< TSharedPtr< FLocMetadataValue > >& MetadataArray )
-{
-	FString FinalString;
-	if( MetadataArray.Num() > 0 )
-	{
-		TArray< FString > MetadataStrings;
-		for( TSharedPtr< FLocMetadataValue > DataEntry : MetadataArray )
-		{
-			if( DataEntry->Type == ELocMetadataType::String )
-			{
-				MetadataStrings.Add(DataEntry->AsString());
-			}
-		}
-		MetadataStrings.Sort();
-		for( const FString& StrEntry : MetadataStrings )
-		{
-			if(FinalString.Len())
-			{
-				FinalString += TEXT(",");
-			}
-			FinalString += StrEntry;
-		}
-	}
-	return FinalString;
-}
-
-TSharedPtr< FLocMetadataValue > UGatherTextFromAssetsCommandlet::FDialogueHelper::GenSourceTargetMetadata( const FString& SpeakerName, const TArray< FString >& TargetNames, bool bCompact )
-{
-	/*
-	This function can support two different formats.
-	
-	The first format is compact and results in string entries that will later be combined into something like this
-	"Variations": [
-		"Jenny -> Audience",
-		"Zak -> Audience"
-	]
-
-	The second format is verbose and results in object entries that will later be combined into something like this
-	"VariationsTest": [
-		{
-			"Speaker": "Jenny",
-			"Targets": [
-				"Audience"
-			]
-		},
-		{
-			"Speaker": "Zak",
-			"Targets": [
-				"Audience"
-			]
-		}
-	]
-	*/
-
-	TSharedPtr< FLocMetadataValue > Result;
-	if( bCompact )
-	{
-		TArray<FString> SortedTargetNames = TargetNames;
-		SortedTargetNames.Sort();
-		FString TargetNamesString;
-		for( const FString& StrEntry : SortedTargetNames )
-		{
-			if(TargetNamesString.Len())
-			{
-				TargetNamesString += TEXT(",");
-			}
-			TargetNamesString += StrEntry;
-		}
-		Result = MakeShareable( new FLocMetadataValueString( FString::Printf( TEXT("%s -> %s" ), *SpeakerName, *TargetNamesString ) ) );
-	}
-	else
-	{
-		TArray< TSharedPtr< FLocMetadataValue > > TargetNamesMetadataList;
-		for( const FString& StrEntry: TargetNames )
-		{
-			TargetNamesMetadataList.Add( MakeShareable( new FLocMetadataValueString( StrEntry ) ) );
-		}
-
-		TSharedPtr< FLocMetadataObject > MetadataObj = MakeShareable( new FLocMetadataObject() );
-		MetadataObj->SetStringField( PropertyName_Speaker, SpeakerName );
-		MetadataObj->SetArrayField( PropertyName_Targets, TargetNamesMetadataList );
-
-		Result = MakeShareable( new FLocMetadataValueObject( MetadataObj.ToSharedRef() ) );
-	}
-	return Result;
-}
-
-
 UGatherTextFromAssetsCommandlet::UGatherTextFromAssetsCommandlet(const FObjectInitializer& ObjectInitializer)
 	:Super(ObjectInitializer)
 {
 
 }
 
-namespace
+void UGatherTextFromAssetsCommandlet::ProcessGatherableTextDataArray(const FString& PackageFilePath, const TArray<FGatherableTextData>& GatherableTextDataArray)
 {
-	template<typename TConcreteChildClass>
-	class FGatherTextFromObject_Base
+	for (const FGatherableTextData& GatherableTextData : GatherableTextDataArray)
 	{
-	public:
-		static inline void Execute(UGatherTextFromAssetsCommandlet* const Commandlet, UObject* const Object, const UPackage* const ObjectPackage, const bool ShouldFixBroken)
+		for (const FTextSourceSiteContext& TextSourceSiteContext : GatherableTextData.SourceSiteContexts)
 		{
-			if( !Object->HasAnyFlags( RF_Transient | RF_PendingKill ) )
-			{
-				TConcreteChildClass(Commandlet, Object, ObjectPackage, ShouldFixBroken).Execute_Internal();
-			}
-		}
+			const FString KeyName = TextSourceSiteContext.KeyName;
 
-	protected:
-		FGatherTextFromObject_Base(UGatherTextFromAssetsCommandlet* const InCommandlet, UObject* const InObject, const UPackage* const InObjectPackage, const bool InShouldFixBroken)
-			: Commandlet(InCommandlet)
-			, Object(InObject)
-			, ObjectPackage(InObjectPackage)
-			, ShouldFixBroken(InShouldFixBroken)
-		{
-		}
+			// Find existing entry from manifest or manifest dependencies.
+			FContext Context;
+			Context.Key = TextSourceSiteContext.KeyName;
+			const FLocMetadataObject DefaultMetadataObject;
+			Context.KeyMetadataObj = !(FLocMetadataObject::IsMetadataExactMatch(&TextSourceSiteContext.KeyMetaData, &DefaultMetadataObject)) ? MakeShareable(new FLocMetadataObject(TextSourceSiteContext.KeyMetaData)) : nullptr;
+			Context.InfoMetadataObj = !(FLocMetadataObject::IsMetadataExactMatch(&TextSourceSiteContext.InfoMetaData, &DefaultMetadataObject)) ? MakeShareable(new FLocMetadataObject(TextSourceSiteContext.InfoMetaData)) : nullptr;
+			Context.bIsOptional = TextSourceSiteContext.IsOptional;
+			Context.SourceLocation = TextSourceSiteContext.SiteDescription;
 
-		void ProcessProperty(UProperty* const Property, void* const ValueAddress)
-		{
-			UTextProperty* const TextProperty = Cast<UTextProperty>(Property);
-			UArrayProperty* const ArrayProperty = Cast<UArrayProperty>(Property);
-			UStructProperty* const StructProperty = Cast<UStructProperty>(Property);
+			TSharedPtr< FManifestEntry > ExistingEntry = ManifestInfo->GetManifest()->FindEntryByContext( GatherableTextData.NamespaceName, Context );
 
-			// Handles both native, fixed-size arrays and plain old non-array properties.
-			for(int32 i = 0; i < Property->ArrayDim; ++i)
+			if (!ExistingEntry.IsValid())
 			{
-				void* const ElementValueAddress = reinterpret_cast<uint8*>(ValueAddress) + Property->ElementSize * i;
-
-				// Property is a text property.
-				if (TextProperty)
-				{
-					FText* Text = reinterpret_cast<FText*>(ElementValueAddress);
-
-					if( FTextInspector::GetFlags(*Text) & ETextFlag::ConvertedProperty )
-					{
-						ObjectPackage->MarkPackageDirty();
-					}
-
-					bool Fixed = false;
-					Commandlet->ProcessTextProperty(TextProperty, Text, Object, /*OUT*/Fixed );
-					if( Fixed )
-					{
-						ObjectPackage->MarkPackageDirty();
-					}
-				}
-				// Property is a DYNAMIC array property.
-				else if (ArrayProperty)
-				{
-					// Iterate over all objects of the array.
-					FScriptArrayHelper ScriptArrayHelper(ArrayProperty, ElementValueAddress);
-					const int32 ElementCount = ScriptArrayHelper.Num();
-					for(int32 j = 0; j < ElementCount; ++j)
-					{
-						ProcessProperty( ArrayProperty->Inner, ScriptArrayHelper.GetRawPtr(j) );
-					}
-				}
-				// Property is a struct property.
-				else if (StructProperty)
-				{
-					// Iterate over all properties of the struct's class.
-					for (TFieldIterator<UProperty>PropIt(StructProperty->Struct, EFieldIteratorFlags::IncludeSuper, EFieldIteratorFlags::ExcludeDeprecated, EFieldIteratorFlags::IncludeInterfaces); PropIt; ++PropIt)
-					{
-						ProcessProperty( *PropIt, PropIt->ContainerPtrToValuePtr<void>(ElementValueAddress) );
-					}
-				}
+				FString FileInfo;
+				ExistingEntry = ManifestInfo->FindDependencyEntryByContext( GatherableTextData.NamespaceName, Context, FileInfo );
 			}
-		}
 
-	protected:
-		UGatherTextFromAssetsCommandlet* const Commandlet;
-		UObject* const Object;
-		const UPackage* const ObjectPackage;
-		const bool ShouldFixBroken;
-	};
-
-	class FGatherTextFromObject : public FGatherTextFromObject_Base<FGatherTextFromObject>
-	{
-	public:
-		FGatherTextFromObject(UGatherTextFromAssetsCommandlet* const InCommandlet, UObject* const InObject, const UPackage* const InObjectPackage, const bool InShouldFixBroken)
-			: FGatherTextFromObject_Base(InCommandlet, InObject, InObjectPackage, InShouldFixBroken)
-		{
-		}
+			FConflictTracker::FEntry NewEntry;
+			NewEntry.PathToSourceString = TextSourceSiteContext.SiteDescription;
+			NewEntry.SourceString = MakeShareable(new FString(GatherableTextData.SourceData.SourceString));
+			NewEntry.Status = EAssetTextGatherStatus::None;
 
-		void Execute_Internal()
-		{
-			// Iterate over all fields of the object's class.
-			for (TFieldIterator<UProperty>PropIt(Object->GetClass(), EFieldIteratorFlags::IncludeSuper, EFieldIteratorFlags::ExcludeDeprecated, EFieldIteratorFlags::IncludeInterfaces); PropIt; ++PropIt)
+			if (TextSourceSiteContext.KeyName.IsEmpty())
 			{
-				ProcessProperty( *PropIt, PropIt->ContainerPtrToValuePtr<void>(Object) );
+				NewEntry.Status = EAssetTextGatherStatus::MissingKey;
 			}
-		}
-	};
-
-	class FGatherTextFromDataTable : public FGatherTextFromObject_Base<FGatherTextFromDataTable>
-	{
-	public:
-		FGatherTextFromDataTable(UGatherTextFromAssetsCommandlet* const InCommandlet, UObject* const InDataTable, const UPackage* const InObjectPackage, const bool InShouldFixBroken)
-			: FGatherTextFromObject_Base(InCommandlet, InDataTable, InObjectPackage, InShouldFixBroken)
-		{
-		}
-
-		void Execute_Internal()
-		{
-			UDataTable* const DataTable = Cast<UDataTable>(Object);
 
-			// Iterate over all properties of the struct's class.
-			for (TFieldIterator<UProperty>PropIt(DataTable->RowStruct, EFieldIteratorFlags::IncludeSuper, EFieldIteratorFlags::ExcludeDeprecated, EFieldIteratorFlags::IncludeInterfaces); PropIt; ++PropIt)
+			// Entry already exists, check for conflict.
+			if (ExistingEntry.IsValid() && ExistingEntry->Source.Text != ( NewEntry.SourceString.IsValid() ? **(NewEntry.SourceString) : TEXT("") ))
 			{
-				for (const auto& Pair : DataTable->RowMap)
-				{
-					ProcessProperty( *PropIt, PropIt->ContainerPtrToValuePtr<void>(Pair.Value) );
-				}
+				NewEntry.Status = EAssetTextGatherStatus::IdentityConflict;
 			}
-		}
-	};
-}
-
-void UGatherTextFromAssetsCommandlet::ProcessPackages( const TArray< UPackage* >& PackagesToProcess )
-{
-	for(UPackage* const Package : PackagesToProcess)
-	{
-		TArray<UObject*> Objects;
-		GetObjectsWithOuter(Package, Objects);
-
-		for(UObject* const Object : Objects)
-		{
-			if ( Object->IsA( UBlueprint::StaticClass() ) )
-			{
-				UBlueprint* Blueprint = Cast<UBlueprint>( Object );
 
-				if( Blueprint->GeneratedClass != NULL )
-				{
-					FGatherTextFromObject::Execute(this, Blueprint->GeneratedClass->GetDefaultObject(), Package, bFixBroken);
-				}
-				else
-				{
-					UE_LOG(LogGatherTextFromAssetsCommandlet, Warning, TEXT("%s - Invalid generated class!"), *Blueprint->GetFullName());
-				}
-			}
-			else if ( Object->IsA( UDataTable::StaticClass() ) )
-			{
-				UDataTable* DataTable = Cast<UDataTable>(Object);
-				FGatherTextFromDataTable::Execute(this, DataTable, Package, bFixBroken);
-			}
-			else if( Object->IsA( UDialogueWave::StaticClass() ) )
+			// Gather if it requires gathering and isn't in a bad state.
+			if ((NewEntry.Status & EAssetTextGatherStatus::BadBitMask) == false)
 			{
-				UDialogueWave* DialogueWave = Cast<UDialogueWave>(Object);
-				ProcessDialogueWave(DialogueWave);
+				ManifestInfo->AddEntry(TextSourceSiteContext.SiteDescription, GatherableTextData.NamespaceName, NewEntry.SourceString.Get() ? *(NewEntry.SourceString) : TEXT(""), Context );
 			}
 
-			FGatherTextFromObject::Execute(this, Object, Package, bFixBroken);
-		}
-	}
-}
-
-void UGatherTextFromAssetsCommandlet::ProcessDialogueWave( const UDialogueWave* DialogueWave )
-{
-	if( !DialogueWave || DialogueWave->HasAnyFlags( RF_Transient | RF_PendingKill ) )
-	{
-		return;
-	}
-
-	FString DialogueName = DialogueWave->GetName();
-	// Use a helper class to extract the dialogue info and prepare it for the manifest
-	FDialogueHelper DialogueHelper;
-	if( DialogueHelper.ProcessDialogueWave( DialogueWave ) )
-	{
-		const FString& SpokenSource = DialogueHelper.GetSpokenSource();
-
-		if( !SpokenSource.IsEmpty() )
-		{
-			{
-				TSharedPtr< const FContext > Base = DialogueHelper.GetBaseContext();
-
-				FString EntryDescription = FString::Printf( TEXT("In non-optional variation of %s"), *DialogueName);
-				ManifestInfo->AddEntry(EntryDescription, DialogueHelper.DialogueNamespace, SpokenSource, *Base);
-			}
-			
-			{
-				const TArray< TSharedPtr< FContext > > Variations = DialogueHelper.GetContextSpecificVariations();
-				for( TSharedPtr< FContext > Variation : Variations )
-				{
-					FString EntryDescription = FString::Printf( TEXT("In context specific variation of %s"), *DialogueName);
-					ManifestInfo->AddEntry(EntryDescription, DialogueHelper.DialogueNamespace, SpokenSource, *Variation);
-				}
-			}
+			// Add to conflict tracker.
+			FConflictTracker::FKeyTable& KeyTable = ConflictTracker.Namespaces.FindOrAdd(GatherableTextData.NamespaceName);
+			FConflictTracker::FEntryArray& EntryArray = KeyTable.FindOrAdd(TextSourceSiteContext.KeyName);
+			EntryArray.Add(NewEntry);
 		}
 	}
 }
 
-void UGatherTextFromAssetsCommandlet::ProcessTextProperty(UTextProperty* InTextProp, FText* Data, UObject* Object, bool& OutFixed)
+void UGatherTextFromAssetsCommandlet::ProcessPackages( const TArray< UPackage* >& PackagesToProcess )
 {
-	OutFixed = false;
-
-	// Early out of gathering pin friendly names unless requested to gather them.
-	if (Object->IsA(UEdGraphPin::StaticClass()))
-	{
-		UProperty* const PinFriendlyNameProperty = Cast<UTextProperty>(UEdGraphPin::StaticClass()->FindPropertyByName(GET_MEMBER_NAME_CHECKED(UEdGraphPin, PinFriendlyName)));
-		const bool IsPinNameProperty = PinFriendlyNameProperty == InTextProp;
-		if (!ShouldGatherBlueprintPinNames && IsPinNameProperty)
-		{
-			return;
-		}
-	}
-
-	FConflictTracker::FEntry NewEntry;
-	NewEntry.ObjectPath = Object->GetPathName();
-	NewEntry.SourceString = Data->GetSourceString();
-	NewEntry.Status = EAssetTextGatherStatus::None;
-
-	FString Namespace;
-	FString Key;
-	bool FoundNamespaceAndKey = FTextLocalizationManager::Get().FindNamespaceAndKeyFromDisplayString(Data->DisplayString, Namespace, Key);
-
-	// Check if text is localizable and check for missing key. Unlocalizable texts don't need a key.
-	if( ( !FoundNamespaceAndKey || Key.IsEmpty() ) && Data->ShouldGatherForLocalization() )
-	{
-		NewEntry.Status = EAssetTextGatherStatus::MissingKey;
-
-		// Fix missing key if allowed.
-		if (bFixBroken)
-		{
-			// Generate new GUID for key.
-			if (FoundNamespaceAndKey)
-			{
-				FTextLocalizationManager::Get().UpdateDisplayString(Data->DisplayString, *Data->DisplayString, Namespace, FGuid::NewGuid().ToString());
-			}
-			else
-			{
-				FTextLocalizationManager::Get().AddDisplayString(Data->DisplayString, TEXT(""), FGuid::NewGuid().ToString());
-			}
-
-			// Fixed.
-			FoundNamespaceAndKey = FTextLocalizationManager::Get().FindNamespaceAndKeyFromDisplayString(Data->DisplayString, Namespace, Key);
-			NewEntry.Status = EAssetTextGatherStatus::MissingKey_Resolved;
-		}
-	}
-
-	// Must have valid key to test for identity conflicts. Even if a text doesn't require localization, if it conflicts, it should be warned about.
-	if( FoundNamespaceAndKey && !Key.IsEmpty() )
+	for(UPackage* const Package : PackagesToProcess)
 	{
-		// Find existing entry from manifest or manifest dependencies.
-		FContext SearchContext;
-		SearchContext.Key = Key;
-		TSharedPtr< FManifestEntry > ExistingEntry = ManifestInfo->GetManifest()->FindEntryByContext( Namespace, SearchContext );
-
-		if( !ExistingEntry.IsValid() )
+		TArray<UObject*> ObjectsInPackage;
+		GetObjectsWithOuter(Package, ObjectsInPackage);
+		for (UObject* const Object : ObjectsInPackage)
 		{
-			FString FileInfo;
-			ExistingEntry = ManifestInfo->FindDependencyEntryByContext( Namespace, SearchContext, FileInfo );
-		}
-
-		// Entry already exists, check for conflict.
-		if( ExistingEntry.IsValid() && ExistingEntry->Source.Text != ( NewEntry.SourceString.IsValid() ? **(NewEntry.SourceString) : TEXT("") ))
-		{
-			NewEntry.Status = EAssetTextGatherStatus::IdentityConflict;
+			TArray<FGatherableTextData> GatherableTextDataArray;
 
-			// Fix conflict if allowed.
-			if (bFixBroken)
+			if( !Object->HasAnyFlags( RF_Transient | RF_PendingKill ) )
 			{
-				// Generate new GUID for key.
-				FTextLocalizationManager::Get().UpdateDisplayString(Data->DisplayString, *Data->DisplayString, Namespace, FGuid::NewGuid().ToString());
-
-				// Fixed.
-				FoundNamespaceAndKey = FTextLocalizationManager::Get().FindNamespaceAndKeyFromDisplayString(Data->DisplayString, Namespace, Key);
-				NewEntry.Status = EAssetTextGatherStatus::IdentityConflict_Resolved;
-
-				// Conflict resolved, no existing entry.
-				ExistingEntry.Reset();
-			}
-		}
-	}
-
-	// Gather if it requires gathering and isn't in a bad state.
-	if( Data->ShouldGatherForLocalization() && (NewEntry.Status & EAssetTextGatherStatus::BadBitMask) == false)
-	{
-		FString SrcLocation = Object->GetPathName() + TEXT(".") + InTextProp->GetName();
-
-		// Adjust the source location if needed.
-		{
-			UClass* Class = Object->GetClass();
-			UObject* CDO = Class ? Class->GetDefaultObject() : NULL;
+				GatherLocalizationDataFromPropertiesOfDataStructure(Object->GetClass(), Object, GatherableTextDataArray);
 
-			if( CDO && CDO != Object )
-			{
-				for( TFieldIterator<UTextProperty> PropIt(CDO->GetClass(), EFieldIteratorFlags::IncludeSuper); PropIt; ++PropIt )
+				for(UClass* Class = Object->GetClass(); Class != nullptr; Class = Class->GetSuperClass())
 				{
-					UTextProperty* TextProp	=  Cast<UTextProperty>( *(PropIt) );
-					FText* DataCDO = TextProp->ContainerPtrToValuePtr<FText>( CDO );
-					FString NamespaceCDO;
-					FString KeyCDO;
-					const bool FoundCDONamespaceAndKey = FTextLocalizationManager::Get().FindNamespaceAndKeyFromDisplayString(DataCDO->DisplayString, NamespaceCDO, KeyCDO);
-
-					if( (!FoundCDONamespaceAndKey && !FoundNamespaceAndKey) || (FoundCDONamespaceAndKey && FoundNamespaceAndKey && ( KeyCDO == Key ) ) )
+					UPackage::FLocalizationDataGatheringCallback* const CustomCallback = UPackage::GetTypeSpecificLocalizationDataGatheringCallbacks().Find(Class);
+					if (CustomCallback)
 					{
-						SrcLocation = CDO->GetPathName() + TEXT(".") + TextProp->GetName();
-						break;
+						(*CustomCallback)(Object, GatherableTextDataArray);
 					}
 				}
+
+				ProcessGatherableTextDataArray(Package->FileName.ToString(), GatherableTextDataArray);
 			}
 		}
-
-		FContext Context;
-		Context.Key = Key;
-		Context.SourceLocation = SrcLocation;
-
-		FString EntryDescription = FString::Printf( TEXT("In %s"), *Object->GetFullName());
-		ManifestInfo->AddEntry(EntryDescription, Namespace, NewEntry.SourceString.Get() ? *(NewEntry.SourceString) : TEXT(""), Context );
 	}
-
-	// Add to conflict tracker.
-	FConflictTracker::FKeyTable& KeyTable = ConflictTracker.Namespaces.FindOrAdd(Namespace);
-	FConflictTracker::FEntryArray& EntryArray = KeyTable.FindOrAdd(Key);
-	EntryArray.Add(NewEntry);
-
-	OutFixed = (NewEntry.Status == EAssetTextGatherStatus::MissingKey_Resolved || NewEntry.Status == EAssetTextGatherStatus::IdentityConflict_Resolved);
 }
 
 int32 UGatherTextFromAssetsCommandlet::Main(const FString& Params)
 {
-	TGuardValue<bool> DisableCompileOnLoad(GForceDisableBlueprintCompileOnLoad, true);
-
 	// Parse command line.
 	TArray<FString> Tokens;
 	TArray<FString> Switches;
@@ -907,13 +345,29 @@ int32 UGatherTextFromAssetsCommandlet::Main(const FString& Params)
 				if( FileReader )
 				{
 					// Read package file summary from the file
-					FPackageFileSummary PackageSummary;
-					(*FileReader) << PackageSummary;
+					FPackageFileSummary PackageFileSummary;
+					(*FileReader) << PackageFileSummary;
+
+					if (PackageFileSummary.GetFileVersionUE4() >= VER_UE4_SERIALIZE_TEXT_IN_PACKAGES)
+					{
+						TArray<FGatherableTextData> GatherableTextDataArray;
+
+						if (PackageFileSummary.GatherableTextDataOffset > 0)
+						{
+							FileReader->Seek(PackageFileSummary.GatherableTextDataOffset);
+
+							for(int32 i = 0; i < PackageFileSummary.GatherableTextDataCount; ++i)
+							{
+								FGatherableTextData* GatherableTextData = new(GatherableTextDataArray) FGatherableTextData;
+								(*FileReader) << *GatherableTextData;
+							}
 
-					// Early out check if the package has been flagged as needing localization gathering
-					if( PackageSummary.PackageFlags & PKG_RequiresLocalizationGather || PackageSummary.GetFileVersionUE4() < VER_UE4_PACKAGE_REQUIRES_LOCALIZATION_GATHER_FLAGGING )
+							ProcessGatherableTextDataArray(PackageFile, GatherableTextDataArray);
+						}
+					}
+					else if (PackageFileSummary.PackageFlags & PKG_RequiresLocalizationGather || PackageFileSummary.GetFileVersionUE4() < VER_UE4_PACKAGE_REQUIRES_LOCALIZATION_GATHER_FLAGGING)
 					{
-						PackageFileNamesToLoad.Add( PackageFile );
+						PackageFileNamesToLoad.Add(PackageFile);
 					}
 				}
 			}
@@ -1030,22 +484,22 @@ int32 UGatherTextFromAssetsCommandlet::Main(const FString& Params)
 				{
 				case EAssetTextGatherStatus::MissingKey:
 					{
-						UE_LOG(LogGatherTextFromAssetsCommandlet, Warning, TEXT("Detected missing key on asset \"%s\"."), *Entry.ObjectPath);
+						UE_LOG(LogGatherTextFromAssetsCommandlet, Warning, TEXT("Detected missing key on asset \"%s\"."), *Entry.PathToSourceString);
 					}
 					break;
 				case EAssetTextGatherStatus::MissingKey_Resolved:
 					{
-						UE_LOG(LogGatherTextFromAssetsCommandlet, Warning, TEXT("Fixed missing key on asset \"%s\"."), *Entry.ObjectPath);
+						UE_LOG(LogGatherTextFromAssetsCommandlet, Warning, TEXT("Fixed missing key on asset \"%s\"."), *Entry.PathToSourceString);
 					}
 					break;
 				case EAssetTextGatherStatus::IdentityConflict:
 					{
-						UE_LOG(LogGatherTextFromAssetsCommandlet, Warning, TEXT("Detected duplicate identity with differing source on asset \"%s\"."), *Entry.ObjectPath);
+						UE_LOG(LogGatherTextFromAssetsCommandlet, Warning, TEXT("Detected duplicate identity with differing source on asset \"%s\"."), *Entry.PathToSourceString);
 					}
 					break;
 				case EAssetTextGatherStatus::IdentityConflict_Resolved:
 					{
-						UE_LOG(LogGatherTextFromAssetsCommandlet, Warning, TEXT("Fixed duplicate identity with differing source on asset \"%s\"."), *Entry.ObjectPath);
+						UE_LOG(LogGatherTextFromAssetsCommandlet, Warning, TEXT("Fixed duplicate identity with differing source on asset \"%s\"."), *Entry.PathToSourceString);
 					}
 					break;
 				}
diff --git a/Engine/Source/Editor/UnrealEd/Private/Commandlets/PackageUtilities.cpp b/Engine/Source/Editor/UnrealEd/Private/Commandlets/PackageUtilities.cpp
--- a/Engine/Source/Editor/UnrealEd/Private/Commandlets/PackageUtilities.cpp
+++ b/Engine/Source/Editor/UnrealEd/Private/Commandlets/PackageUtilities.cpp
@@ -1167,6 +1167,38 @@ void FPkgInfoReporter_Log::GeneratePackageReport( FLinkerLoad* InLinker/*=NULL*/
 		}
 	}
 
+	if( (InfoFlags&PKGINFO_Text) != 0 )
+	{
+		UE_LOG(LogPackageUtilities, Warning, TEXT("--------------------------------------------") );
+		GWarn->Log ( TEXT("Gatherable Text Data Map"));
+		GWarn->Log ( TEXT("=========="));
+
+		if (Linker->SerializeGatherableTextDataMap(true))
+		{
+			UE_LOG(LogPackageUtilities, Display, TEXT("Number of Text Data Entries: %d"), Linker->GatherableTextDataMap.Num());
+
+			for (int32 i = 0; i < Linker->GatherableTextDataMap.Num(); ++i)
+			{
+				const FGatherableTextData& GatherableTextData = Linker->GatherableTextDataMap[i];
+				UE_LOG(LogPackageUtilities, Display, TEXT("Entry %d:"), 1 + i);
+				UE_LOG(LogPackageUtilities, Display, TEXT("\t   String: %s"), *GatherableTextData.SourceData.SourceString.ReplaceCharWithEscapedChar());
+				UE_LOG(LogPackageUtilities, Display, TEXT("\tNamespace: %s"), *GatherableTextData.NamespaceName);
+				UE_LOG(LogPackageUtilities, Display, TEXT("\t   Key(s): %d"), GatherableTextData.SourceSiteContexts.Num());
+				for (const FTextSourceSiteContext& TextSourceSiteContext : GatherableTextData.SourceSiteContexts)
+				{
+					UE_LOG(LogPackageUtilities, Display, TEXT("\t\t%s from %s"), *TextSourceSiteContext.KeyName, *TextSourceSiteContext.SiteDescription);
+				}
+			}
+		}
+		else
+		{
+			if ( Linker->Summary.GatherableTextDataOffset > 0 )
+			{
+				UE_LOG(LogPackageUtilities, Warning,TEXT("Failed to load gatherable text data for package %s!"), *LinkerName.ToString());
+			}
+		}
+	}
+
 
 	if( (InfoFlags&PKGINFO_Thumbs) != 0 )
 	{
diff --git a/Engine/Source/Editor/UnrealEd/Public/PackageUtilityWorkers.h b/Engine/Source/Editor/UnrealEd/Public/PackageUtilityWorkers.h
--- a/Engine/Source/Editor/UnrealEd/Public/PackageUtilityWorkers.h
+++ b/Engine/Source/Editor/UnrealEd/Public/PackageUtilityWorkers.h
@@ -23,8 +23,9 @@ enum EPackageInfoFlags
 	PKGINFO_Thumbs			=0x80,
 	PKGINFO_Lazy			=0x100,
 	PKGINFO_AssetRegistry	=0x200,
+	PKGINFO_Text			=0x400,
 
-	PKGINFO_All			= PKGINFO_Names|PKGINFO_Imports|PKGINFO_Exports|PKGINFO_Chunks|PKGINFO_Depends|PKGINFO_Paths|PKGINFO_Thumbs|PKGINFO_Lazy|PKGINFO_AssetRegistry,
+	PKGINFO_All			= PKGINFO_Names|PKGINFO_Imports|PKGINFO_Exports|PKGINFO_Chunks|PKGINFO_Depends|PKGINFO_Paths|PKGINFO_Thumbs|PKGINFO_Lazy|PKGINFO_AssetRegistry|PKGINFO_Text,
 };
 
 /**
diff --git a/Engine/Source/Runtime/Core/Private/Internationalization/GatherableTextData.cpp b/Engine/Source/Runtime/Core/Private/Internationalization/GatherableTextData.cpp
new file mode 100644
--- /dev/null
+++ b/Engine/Source/Runtime/Core/Private/Internationalization/GatherableTextData.cpp
@@ -0,0 +1,33 @@
+// Copyright 1998-2015 Epic Games, Inc. All Rights Reserved.
+
+#include "CorePrivatePch.h"
+#include "Internationalization/GatherableTextData.h"
+
+FArchive& operator<<(FArchive& Archive, FTextSourceSiteContext& This)
+{
+	Archive << This.KeyName;
+	Archive << This.SiteDescription;
+	Archive << This.IsEditorOnly;
+	Archive << This.IsOptional;
+	Archive << This.InfoMetaData;
+	Archive << This.KeyMetaData;
+
+	return Archive;
+}
+
+FArchive& operator<<(FArchive& Archive, FTextSourceData& This)
+{
+	Archive << This.SourceString;
+	Archive << This.SourceStringMetaData;
+
+	return Archive;
+}
+
+FArchive& operator<<(FArchive& Archive, FGatherableTextData& This)
+{
+	Archive << This.NamespaceName;
+	Archive << This.SourceData;
+	Archive << This.SourceSiteContexts;
+
+	return Archive;
+}
\ No newline at end of file
diff --git a/Engine/Source/Runtime/Core/Private/Internationalization/InternationalizationManifest.cpp b/Engine/Source/Runtime/Core/Private/Internationalization/InternationalizationManifest.cpp
--- a/Engine/Source/Runtime/Core/Private/Internationalization/InternationalizationManifest.cpp
+++ b/Engine/Source/Runtime/Core/Private/Internationalization/InternationalizationManifest.cpp
@@ -192,7 +192,7 @@ bool FLocItem::IsExactMatch( const FLocItem& Other ) const
 {
 	if( Text.Equals( Other.Text, ESearchCase::CaseSensitive ) )
 	{
-		return FLocMetadataObject::IsMetadataExactMatch(MetadataObj, Other.MetadataObj );
+		return FLocMetadataObject::IsMetadataExactMatch(MetadataObj.Get(), Other.MetadataObj.Get() );
 	}
 	return false;
 }
diff --git a/Engine/Source/Runtime/Core/Private/Internationalization/InternationalizationMetadata.cpp b/Engine/Source/Runtime/Core/Private/Internationalization/InternationalizationMetadata.cpp
--- a/Engine/Source/Runtime/Core/Private/Internationalization/InternationalizationMetadata.cpp
+++ b/Engine/Source/Runtime/Core/Private/Internationalization/InternationalizationMetadata.cpp
@@ -7,9 +7,11 @@ DEFINE_LOG_CATEGORY_STATIC(LogInternationalizationMetadata, Log, All);
 
 const FString FLocMetadataObject::COMPARISON_MODIFIER_PREFIX = TEXT("*");
 
+FLocMetadataValue::~FLocMetadataValue() = default;
+
 void FLocMetadataValue::ErrorMessage( const FString& InType )
 {
-	UE_LOG(LogInternationalizationMetadata, Fatal, TEXT("LocMetadata Value of type '%s' used as a '%s'."), *GetType(), *InType);
+	UE_LOG(LogInternationalizationMetadata, Fatal, TEXT("LocMetadata Value of type '%s' used as a '%s'."), *GetTypeString(), *InType);
 }
 
 FLocMetadataObject::FLocMetadataObject( const FLocMetadataObject& Other )
@@ -229,18 +231,18 @@ bool FLocMetadataObject::operator<( const FLocMetadataObject& Other ) const
 	return false;
 }
 
-bool FLocMetadataObject::IsMetadataExactMatch( TSharedPtr<FLocMetadataObject> MetadataA, TSharedPtr<FLocMetadataObject> MetadataB )
+bool FLocMetadataObject::IsMetadataExactMatch( const FLocMetadataObject* const MetadataA, const FLocMetadataObject* const MetadataB )
 {
-	if( !MetadataA.IsValid() && !MetadataB.IsValid() )
+	if( !MetadataA && !MetadataB )
 	{
 		return true;
 	}
-	else if( MetadataA.IsValid() != MetadataB.IsValid() )
+	else if( (MetadataA != nullptr) != (MetadataB != nullptr) )
 	{
 		// If we are in here, we know that one of the metadata entries is null, if the other
 		//  contains zero entries we will still consider them equivalent.
-		if( (MetadataA.IsValid() && MetadataA->Values.Num() == 0) ||
-			(MetadataB.IsValid() && MetadataB->Values.Num() == 0) )
+		if( (MetadataA && MetadataA->Values.Num() == 0) ||
+			(MetadataB && MetadataB->Values.Num() == 0) )
 		{
 			return true;
 		}
@@ -254,6 +256,123 @@ bool FLocMetadataObject::IsMetadataExactMatch( TSharedPtr<FLocMetadataObject> Me
 	return MetadataA->IsExactMatch( *(MetadataB) );
 }
 
+namespace
+{
+	void SerializeLocMetadataValue(FArchive& Archive, FLocMetadataValue*& Value)
+	{
+		if (Archive.IsSaving())
+		{
+			check(Value != nullptr);
+		}
+
+		if(Archive.IsLoading())
+		{
+			check(Value == nullptr);
+		}
+
+		ELocMetadataType MetaDataType = ELocMetadataType::None;
+		if (Archive.IsSaving())
+		{
+			MetaDataType = Value->Type;
+		}
+
+		int32 MetaDataTypeAsInt = static_cast<int32>(MetaDataType);
+		Archive << MetaDataTypeAsInt;
+		MetaDataType = static_cast<ELocMetadataType>(MetaDataTypeAsInt);
+
+		switch(MetaDataType)
+		{
+		case ELocMetadataType::Array:
+			if (Archive.IsSaving())
+			{
+				FLocMetadataValueArray::Serialize(static_cast<const FLocMetadataValueArray&>(*Value), Archive);
+			}
+			else if (Archive.IsLoading())
+			{
+				Value = new FLocMetadataValueArray(Archive);
+			}
+			break;
+		case ELocMetadataType::Boolean:
+			if (Archive.IsSaving())
+			{
+				FLocMetadataValueBoolean::Serialize(static_cast<const FLocMetadataValueBoolean&>(*Value), Archive);
+			}
+			else if (Archive.IsLoading())
+			{
+				Value = new FLocMetadataValueBoolean(Archive);
+			}
+			break;
+		case ELocMetadataType::Object:
+			if (Archive.IsSaving())
+			{
+				FLocMetadataValueObject::Serialize(static_cast<const FLocMetadataValueObject&>(*Value), Archive);
+			}
+			else if (Archive.IsLoading())
+			{
+				Value = new FLocMetadataValueObject(Archive);
+			}
+			break;
+		case ELocMetadataType::String:
+			if (Archive.IsSaving())
+			{
+				FLocMetadataValueString::Serialize(static_cast<const FLocMetadataValueString&>(*Value), Archive);
+			}
+			else if (Archive.IsLoading())
+			{
+				Value = new FLocMetadataValueString(Archive);
+			}
+			break;
+		default:
+			checkNoEntry();
+		}		
+	}
+}
+
+FArchive& operator<<(FArchive& Archive, FLocMetadataObject& Object)
+{
+	int32 ValueCount = Object.Values.Num();
+	Archive << ValueCount;
+
+	if (Archive.IsLoading())
+	{
+		Object.Values.Reserve(ValueCount);
+	}
+
+	TArray<FString> MapKeys;
+	Object.Values.GetKeys(MapKeys);
+	for (int32 i = 0; i < ValueCount; ++i)
+	{
+		FString Key;
+		if (Archive.IsSaving())
+		{
+			Key = MapKeys[i];
+		}
+		
+		Archive << Key;
+		
+		if(Archive.IsLoading())
+		{
+			Object.Values.Add(Key);
+		}
+
+		TSharedPtr<FLocMetadataValue> Value;
+		if (Archive.IsSaving())
+		{
+			Value = Object.Values[Key];
+		}
+
+		FLocMetadataValue* ValueRawPointer = Value.Get();
+		SerializeLocMetadataValue(Archive, ValueRawPointer);
+		Value = MakeShareable(ValueRawPointer);
+
+		if (Archive.IsLoading())
+		{
+			Object.Values[Key] = Value;
+		}
+	}
+
+	return Archive;
+}
 
 bool FLocMetadataValueString::EqualTo( const FLocMetadataValue& Other ) const
 {
@@ -272,6 +391,21 @@ TSharedRef<FLocMetadataValue> FLocMetadataValueString::Clone() const
 	return MakeShareable( new FLocMetadataValueString( Value ) );
 }
 
+FLocMetadataValueString::FLocMetadataValueString( FArchive& Archive )
+{
+	check(Archive.IsLoading());
+
+	Archive << Value;
+}
+
+void FLocMetadataValueString::Serialize( const FLocMetadataValueString& Value, FArchive& Archive )
+{
+	check(Archive.IsSaving());
+
+	FString StringValue = Value.Value;
+	Archive << StringValue;
+}
+
 bool FLocMetadataValueBoolean::EqualTo( const FLocMetadataValue& Other ) const
 {
 	const FLocMetadataValueBoolean* OtherObj = (FLocMetadataValueBoolean*) &Other;
@@ -289,6 +423,21 @@ TSharedRef<FLocMetadataValue> FLocMetadataValueBoolean::Clone() const
 	return MakeShareable( new FLocMetadataValueBoolean( Value ) );
 }
 
+FLocMetadataValueBoolean::FLocMetadataValueBoolean( FArchive& Archive )
+{
+	check(Archive.IsLoading());
+
+	Archive << Value;
+}
+
+void FLocMetadataValueBoolean::Serialize( const FLocMetadataValueBoolean& Value, FArchive& Archive )
+{
+	check(Archive.IsSaving());
+
+	bool BoolValue = Value.Value;
+	Archive << BoolValue;
+}
+
 struct FCompareMetadataValue
 {
 	FORCEINLINE bool operator()( TSharedPtr<FLocMetadataValue> A, TSharedPtr<FLocMetadataValue> B ) const
@@ -371,6 +520,37 @@ TSharedRef<FLocMetadataValue> FLocMetadataValueArray::Clone() const
 	return MakeShareable( new FLocMetadataValueArray( NewValue ) );
 }
 
+FLocMetadataValueArray::FLocMetadataValueArray( FArchive& Archive )
+{
+	check(Archive.IsLoading());
+
+	int32 ElementCount;
+	Archive << ElementCount;
+
+	Value.SetNum(ElementCount);
+
+	for (TSharedPtr<FLocMetadataValue>& Element : Value)
+	{
+		FLocMetadataValue* ElementRawPointer = Element.Get();
+		SerializeLocMetadataValue(Archive, ElementRawPointer);
+		Element = MakeShareable(ElementRawPointer);
+	}
+}
+
+void FLocMetadataValueArray::Serialize( const FLocMetadataValueArray& Value, FArchive& Archive )
+{
+	check(Archive.IsSaving());
+
+	int32 ElementCount = Value.Value.Num();
+	Archive << ElementCount;
+
+	for (const TSharedPtr<FLocMetadataValue>& Element : Value.Value)
+	{
+		FLocMetadataValue* ElementRawPointer = Element.Get();
+		SerializeLocMetadataValue(Archive, ElementRawPointer);
+	}
+}
+
 bool FLocMetadataValueObject::EqualTo( const FLocMetadataValue& Other ) const
 {
 	const FLocMetadataValueObject* OtherObj = (FLocMetadataValueObject*) &Other;  
@@ -404,3 +584,18 @@ TSharedRef<FLocMetadataValue> FLocMetadataValueObject::Clone() const
 	TSharedRef<FLocMetadataObject> NewLocMetadataObject = MakeShareable( new FLocMetadataObject( *(this->Value) ) );
 	return MakeShareable( new FLocMetadataValueObject( NewLocMetadataObject ) );
 }
+
+FLocMetadataValueObject::FLocMetadataValueObject( FArchive& Archive )
+	: Value(new FLocMetadataObject)
+{
+	check(Archive.IsLoading());
+
+	Archive << *Value;
+}
+
+void FLocMetadataValueObject::Serialize( const FLocMetadataValueObject& Value, FArchive& Archive )
+{
+	check(Archive.IsSaving());
+
+	Archive << *(Value.Value);
+}
\ No newline at end of file
diff --git a/Engine/Source/Runtime/Core/Public/Internationalization/GatherableTextData.h b/Engine/Source/Runtime/Core/Public/Internationalization/GatherableTextData.h
new file mode 100644
--- /dev/null
+++ b/Engine/Source/Runtime/Core/Public/Internationalization/GatherableTextData.h
@@ -0,0 +1,40 @@
+// Copyright 1998-2015 Epic Games, Inc. All Rights Reserved.
+#pragma once
+ 
+#include "InternationalizationMetadata.h"
+
+struct FTextSourceSiteContext
+{
+	FTextSourceSiteContext()
+		: IsEditorOnly(false)
+		, IsOptional(false)
+	{
+	}
+
+	FString KeyName;
+	FString SiteDescription;
+	bool IsEditorOnly;
+	bool IsOptional;
+	FLocMetadataObject InfoMetaData;
+	FLocMetadataObject KeyMetaData;
+};
+
+CORE_API FArchive& operator<<(FArchive& Archive, FTextSourceSiteContext& This);
+
+struct FTextSourceData
+{
+	FString SourceString;
+	FLocMetadataObject SourceStringMetaData;
+};
+
+CORE_API FArchive& operator<<(FArchive& Archive, FTextSourceData& This);
+
+struct FGatherableTextData
+{
+	FString NamespaceName;
+	FTextSourceData SourceData;
+
+	TArray<FTextSourceSiteContext> SourceSiteContexts;
+};
+
+CORE_API FArchive& operator<<(FArchive& Archive, FGatherableTextData& This);
\ No newline at end of file
diff --git a/Engine/Source/Runtime/Core/Public/Internationalization/InternationalizationMetadata.h b/Engine/Source/Runtime/Core/Public/Internationalization/InternationalizationMetadata.h
--- a/Engine/Source/Runtime/Core/Public/Internationalization/InternationalizationMetadata.h
+++ b/Engine/Source/Runtime/Core/Public/Internationalization/InternationalizationMetadata.h
@@ -6,17 +6,14 @@ DECLARE_LOG_CATEGORY_EXTERN(LogInternationalization, Log, All);
 /**
  * Represents all the types a LocMetadata Value can be.
  */
-namespace ELocMetadataType
+enum class ELocMetadataType : int32
 {
-	enum Type
-	{
-		None,
-		Boolean,
-		String,
-		Array,
-		Object,
-	};
-}
+	None,
+	Boolean,
+	String,
+	Array,
+	Object,
+};
 
 class FLocMetadataObject;
 
@@ -26,6 +23,7 @@ class FLocMetadataObject;
 class CORE_API FLocMetadataValue
 {
 public:
+	virtual ~FLocMetadataValue() = 0;
 
 	/** Returns this value as a number, throwing an error if this is not a Metadata String */
 	virtual FString AsString() {ErrorMessage(TEXT("String")); return FString();}
@@ -44,25 +42,19 @@ public:
 	bool operator==( const FLocMetadataValue& Other ) { return ( (Type == Other.Type) && EqualTo( Other ) ); }
 	bool operator<( const FLocMetadataValue& Other ) { return (Type == Other.Type) ? LessThan( Other ) : (Type < Other.Type); }
 
-	ELocMetadataType::Type Type;
+	ELocMetadataType Type;
 
 protected:
-
 	FLocMetadataValue() : Type(ELocMetadataType::None) {}
-	virtual ~FLocMetadataValue() {}
+	FLocMetadataValue( const FLocMetadataValue& ) = delete;  
+	FLocMetadataValue& operator=( const FLocMetadataValue& ) = delete;
 
 	virtual bool EqualTo( const FLocMetadataValue& Other ) const = 0;
 	virtual bool LessThan( const FLocMetadataValue& Other ) const = 0;
 
-	virtual FString GetType() const = 0;
+	virtual FString GetTypeString() const = 0;
 
 	void ErrorMessage(const FString& InType);
-
-private:
-	
-	/** Copying is only supported using Clone() function */
-	FLocMetadataValue( const FLocMetadataValue& ); 
-	FLocMetadataValue& operator=( const FLocMetadataValue& );
 };
 
 
@@ -80,7 +72,7 @@ public:
 	/** Copy ctor */
 	FLocMetadataObject( const FLocMetadataObject& Other );
 	
-	template<ELocMetadataType::Type LocMetadataType>
+	template<ELocMetadataType LocMetadataType>
 	TSharedPtr<FLocMetadataValue> GetField( const FString& FieldName )
 	{
 		const TSharedPtr<FLocMetadataValue>* Field = Values.Find(FieldName);
@@ -108,7 +100,7 @@ public:
 	}
 	
 	/** Checks to see if the FieldName exists in the object, and has the specified type. */
-	template<ELocMetadataType::Type LocMetadataType>
+	template<ELocMetadataType LocMetadataType>
 	bool HasTypedField(const FString& FieldName)
 	{
 		const TSharedPtr<FLocMetadataValue>* Field = Values.Find(FieldName);
@@ -156,7 +148,7 @@ public:
 	/** Similar functionality to == operator but ensures everything matches(ignores COMPARISON_MODIFIER_PREFIX). */
 	bool IsExactMatch( const FLocMetadataObject& Other ) const;
 
-	static bool IsMetadataExactMatch( TSharedPtr<FLocMetadataObject> MetadataA, TSharedPtr<FLocMetadataObject> MetadataB );
+	static bool IsMetadataExactMatch( const FLocMetadataObject* const MetadataA, const FLocMetadataObject* const MetadataB );
 
 public:
 	/** Stores the name/value pairs for the metadata object */
@@ -166,24 +158,29 @@ public:
 	static const FString COMPARISON_MODIFIER_PREFIX;
 };
 
+CORE_API FArchive& operator<<(FArchive& Archive, FLocMetadataObject& Object);
 
 /** A LocMetadata String Value. */
 class CORE_API FLocMetadataValueString : public FLocMetadataValue
 {
 public:
 	FLocMetadataValueString(const FString& InString) : Value(InString) {Type = ELocMetadataType::String;}
+	FLocMetadataValueString(FArchive& Archive);
+	static void Serialize(const FLocMetadataValueString& Value, FArchive& Archive);
+
 	virtual FString AsString() override {return Value;}
 
 	virtual TSharedRef<FLocMetadataValue> Clone() const override;
 	void SetString( const FString& InString ) { Value = InString; }
 
 protected:
+
 	FString Value;
 
 	virtual bool EqualTo( const FLocMetadataValue& Other ) const override;
 	virtual bool LessThan( const FLocMetadataValue& Other ) const override;
 
-	virtual FString GetType() const override {return TEXT("String");}
+	virtual FString GetTypeString() const override {return TEXT("String");}
 };
 
 
@@ -192,6 +189,9 @@ class CORE_API FLocMetadataValueBoolean : public FLocMetadataValue
 {
 public:
 	FLocMetadataValueBoolean(bool InBool) : Value(InBool) {Type = ELocMetadataType::Boolean;}
+	FLocMetadataValueBoolean(FArchive& Archive);
+	static void Serialize(const FLocMetadataValueBoolean& Value, FArchive& Archive);
+
 	virtual bool AsBool() override {return Value;}
 
 	virtual TSharedRef<FLocMetadataValue> Clone() const override;
@@ -202,7 +202,7 @@ protected:
 	virtual bool EqualTo( const FLocMetadataValue& Other ) const override;
 	virtual bool LessThan( const FLocMetadataValue& Other ) const override;
 
-	virtual FString GetType() const override {return TEXT("Boolean");}
+	virtual FString GetTypeString() const override {return TEXT("Boolean");}
 };
 
 /** A LocMetadata Array Value. */
@@ -210,6 +210,9 @@ class CORE_API FLocMetadataValueArray : public FLocMetadataValue
 {
 public:
 	FLocMetadataValueArray(const TArray< TSharedPtr<FLocMetadataValue> >& InArray) : Value(InArray) {Type = ELocMetadataType::Array;}
+	FLocMetadataValueArray(FArchive& Archive);
+	static void Serialize(const FLocMetadataValueArray& Value, FArchive& Archive);
+
 	virtual TArray< TSharedPtr<FLocMetadataValue> > AsArray() override {return Value;}
 
 	virtual TSharedRef<FLocMetadataValue> Clone() const override;
@@ -220,7 +223,7 @@ protected:
 	virtual bool EqualTo( const FLocMetadataValue& Other ) const override;
 	virtual bool LessThan( const FLocMetadataValue& Other ) const override;
 
-	virtual FString GetType() const override {return TEXT("Array");}
+	virtual FString GetTypeString() const override {return TEXT("Array");}
 };
 
 /** A LocMetadata Object Value. */
@@ -228,6 +231,9 @@ class CORE_API FLocMetadataValueObject : public FLocMetadataValue
 {
 public:
 	FLocMetadataValueObject(TSharedPtr<FLocMetadataObject> InObject) : Value(InObject) {Type = ELocMetadataType::Object;}
+	FLocMetadataValueObject(FArchive& Archive);
+	static void Serialize(const FLocMetadataValueObject& Value, FArchive& Archive);
+
 	virtual TSharedPtr<FLocMetadataObject> AsObject() override {return Value;}
 
 	virtual TSharedRef<FLocMetadataValue> Clone() const override;
@@ -238,10 +244,5 @@ protected:
 	virtual bool EqualTo( const FLocMetadataValue& Other ) const override;
 	virtual bool LessThan( const FLocMetadataValue& Other ) const override;
 
-	virtual FString GetType() const override {return TEXT("Object");}
-
-private:
-	// Hidden default ctor
-	FLocMetadataValueObject()
-	{}
+	virtual FString GetTypeString() const override {return TEXT("Object");}
 };
diff --git a/Engine/Source/Runtime/Core/Public/UObject/ObjectVersion.h b/Engine/Source/Runtime/Core/Public/UObject/ObjectVersion.h
--- a/Engine/Source/Runtime/Core/Public/UObject/ObjectVersion.h
+++ b/Engine/Source/Runtime/Core/Public/UObject/ObjectVersion.h
@@ -1,4 +1,4 @@
-// Copyright 1998-2015 Epic Games, Inc. All Rights Reserved.
+﻿// Copyright 1998-2015 Epic Games, Inc. All Rights Reserved.
 
 #pragma once
 
@@ -512,6 +512,8 @@ enum EUnrealEngineObjectUE4Version
 	VER_UE4_LIBRARY_CATEGORIES_AS_FTEXT,
 	// Check for duplicate exports while saving packages.
 	VER_UE4_SKIP_DUPLICATE_EXPORTS_ON_SAVE_PACKAGE,
+	// Pre-gathering of gatherable, localizable text in packages to optimize text gathering operation times
+	VER_UE4_SERIALIZE_TEXT_IN_PACKAGES,
 
 	// -----<new versions can be added before this line>-------------------------------------------------
 	// - this needs to be the last line (see note below)
diff --git a/Engine/Source/Runtime/CoreUObject/Private/Serialization/PropertyLocalizationDataGathering.cpp b/Engine/Source/Runtime/CoreUObject/Private/Serialization/PropertyLocalizationDataGathering.cpp
new file mode 100644
--- /dev/null
+++ b/Engine/Source/Runtime/CoreUObject/Private/Serialization/PropertyLocalizationDataGathering.cpp
@@ -0,0 +1,144 @@
+// Copyright 1998-2015 Epic Games, Inc. All Rights Reserved.
+
+#include "CoreUObjectPrivate.h"
+#include "Serialization/PropertyLocalizationDataGathering.h"
+#include "UTextProperty.h"
+
+void GatherLocalizationDataFromPropertiesOfDataStructure(UStruct* const Structure, void* const Data, TArray<FGatherableTextData>& GatherableTextDataArray)
+{
+	FString Path;
+	{
+		UClass* const Class = Cast<UClass>(Structure);
+		if (Class)
+		{
+			UObject* const Object = reinterpret_cast<UObject* const>(Data);
+			Path = Object->GetPathName();
+		}
+	}
+
+	// Iterate over all fields of the object's class.
+	for (TFieldIterator<UProperty> PropIt(Structure, EFieldIteratorFlags::IncludeSuper, EFieldIteratorFlags::ExcludeDeprecated, EFieldIteratorFlags::IncludeInterfaces); PropIt; ++PropIt)
+	{
+		GatherLocalizationDataFromChildTextProperies(Path, *PropIt, PropIt->ContainerPtrToValuePtr<void>(Data), GatherableTextDataArray, false);
+	}
+}
+
+void GatherLocalizationDataFromChildTextProperies(const FString& PathToParent, UProperty* const Property, void* const ValueAddress, TArray<FGatherableTextData>& GatherableTextDataArray, const bool ForceIsEditorOnly)
+{
+	UTextProperty* const TextProperty = Cast<UTextProperty>(Property);
+	UArrayProperty* const ArrayProperty = Cast<UArrayProperty>(Property);
+	UStructProperty* const StructProperty = Cast<UStructProperty>(Property);
+
+	// Property is a text property.
+	if (TextProperty)
+	{
+		GatherLocalizationDataFromTextProperty(PathToParent, TextProperty, ValueAddress, GatherableTextDataArray, ForceIsEditorOnly);
+	}
+	else
+	{
+		// Handles both native, fixed-size arrays and plain old non-array properties.
+		const bool IsFixedSizeArray = Property->ArrayDim > 1;
+		for(int32 i = 0; i < Property->ArrayDim; ++i)
+		{
+			const FString PathToElement = FString(PathToParent.IsEmpty() ? TEXT("") : PathToParent + TEXT(".")) + (IsFixedSizeArray ? Property->GetName() + FString::Printf(TEXT("[%d]"), i) : Property->GetName());
+			void* const ElementValueAddress = reinterpret_cast<uint8*>(ValueAddress) + Property->ElementSize * i;
+
+			// Property is a DYNAMIC array property.
+			if (ArrayProperty)
+			{
+				// Iterate over all elements of the array.
+				FScriptArrayHelper ScriptArrayHelper(ArrayProperty, ElementValueAddress);
+				const int32 ElementCount = ScriptArrayHelper.Num();
+				for(int32 j = 0; j < ElementCount; ++j)
+				{
+					GatherLocalizationDataFromChildTextProperies(PathToElement + FString::Printf(TEXT("(%d)"), j), ArrayProperty->Inner, ScriptArrayHelper.GetRawPtr(j), GatherableTextDataArray, ForceIsEditorOnly || (Property->PropertyFlags & CPF_EditorOnly) != 0);
+				}
+			}
+			// Property is a struct property.
+			else if (StructProperty)
+			{
+				// Iterate over all properties of the struct's class.
+				for (TFieldIterator<UProperty>PropIt(StructProperty->Struct, EFieldIteratorFlags::IncludeSuper, EFieldIteratorFlags::ExcludeDeprecated, EFieldIteratorFlags::IncludeInterfaces); PropIt; ++PropIt)
+				{
+					const FString Path = PathToElement;
+
+					GatherLocalizationDataFromChildTextProperies(PathToElement, *PropIt, PropIt->ContainerPtrToValuePtr<void>(ElementValueAddress), GatherableTextDataArray, ForceIsEditorOnly || (Property->PropertyFlags & CPF_EditorOnly) != 0);
+				}
+			}
+		}
+	}
+}
+
+void GatherLocalizationDataFromTextProperty(const FString& PathToParent, UTextProperty* const TextProperty, void* const ValueAddress, TArray<FGatherableTextData>& GatherableTextDataArray, const bool ForceIsEditorOnly)
+{
+	const bool IsFixedSizeArray = TextProperty->ArrayDim > 1;
+	for(int32 i = 0; i < TextProperty->ArrayDim; ++i)
+	{
+		const FString PathToElement = FString(PathToParent.IsEmpty() ? TEXT("") : PathToParent + TEXT(".")) + (IsFixedSizeArray ? TextProperty->GetName() + FString::Printf(TEXT("[%d]"), i) : TextProperty->GetName());
+		FText* const ElementValueAddress = reinterpret_cast<FText*>(ValueAddress) + TextProperty->ElementSize * i;
+
+		UPackage* const Package = TextProperty->GetOutermost();
+		if( FTextInspector::GetFlags(*ElementValueAddress) & ETextFlag::ConvertedProperty )
+		{
+			Package->MarkPackageDirty();
+		}
+
+		FString Namespace;
+		FString Key;
+		const FTextDisplayStringRef DisplayString = FTextInspector::GetSharedDisplayString(*ElementValueAddress);
+		const bool FoundNamespaceAndKey = FTextLocalizationManager::Get().FindNamespaceAndKeyFromDisplayString(DisplayString, Namespace, Key);
+
+		if (!ElementValueAddress->ShouldGatherForLocalization())
+		{
+			return;
+		}
+
+		FTextSourceData SourceData;
+		{
+			const FString* SourceString = FTextInspector::GetSourceString(*ElementValueAddress);
+			SourceData.SourceString = SourceString ? *SourceString : TEXT("");
+		}
+
+		FGatherableTextData* GatherableTextData = GatherableTextDataArray.FindByPredicate([&](const FGatherableTextData& Candidate)
+		{
+			return	Candidate.NamespaceName == Namespace && 
+				Candidate.SourceData.SourceString == SourceData.SourceString &&
+				Candidate.SourceData.SourceStringMetaData == SourceData.SourceStringMetaData;
+		});
+		if(!GatherableTextData)
+		{
+			GatherableTextData = new(GatherableTextDataArray) FGatherableTextData;
+			GatherableTextData->NamespaceName = Namespace;
+			GatherableTextData->SourceData = SourceData;
+		}
+
+		// If this text has no key, fix it by generating a new GUID for the key.
+		if (!FoundNamespaceAndKey || Key.IsEmpty())
+		{
+			const FString NewKey = FGuid::NewGuid().ToString();
+			bool HasSuccessfullySetKey = false;
+			if (FoundNamespaceAndKey)
+			{
+				HasSuccessfullySetKey = FTextLocalizationManager::Get().UpdateDisplayString(DisplayString, *DisplayString, Namespace, NewKey);
+			}
+			else
+			{
+				HasSuccessfullySetKey = FTextLocalizationManager::Get().AddDisplayString(DisplayString, TEXT(""), NewKey);
+			}
+
+			if (HasSuccessfullySetKey)
+			{
+				Key = NewKey;
+			}
+		}
+
+		// Add to map.
+		{
+			FTextSourceSiteContext& SourceSiteContext = *(new(GatherableTextData->SourceSiteContexts) FTextSourceSiteContext);
+			SourceSiteContext.KeyName = Key;
+			SourceSiteContext.SiteDescription = PathToElement;
+			SourceSiteContext.IsEditorOnly = ForceIsEditorOnly || (TextProperty->PropertyFlags & CPF_EditorOnly) != 0;
+			SourceSiteContext.IsOptional = false;
+		}
+	}
+}
\ No newline at end of file
diff --git a/Engine/Source/Runtime/CoreUObject/Private/UObject/Linker.cpp b/Engine/Source/Runtime/CoreUObject/Private/UObject/Linker.cpp
--- a/Engine/Source/Runtime/CoreUObject/Private/UObject/Linker.cpp
+++ b/Engine/Source/Runtime/CoreUObject/Private/UObject/Linker.cpp
@@ -128,6 +128,11 @@ void FLinker::Serialize( FArchive& Ar )
 		}
 	}
 
+	if (Ar.IsSaving() || Ar.UE4Ver() >= VER_UE4_SERIALIZE_TEXT_IN_PACKAGES)
+	{
+		Ar << GatherableTextDataMap;
+	}
+
 	// Prevent garbage collecting of linker's names and package.
 	Ar << NameMap << LinkerRoot;
 	{
@@ -144,6 +149,7 @@ void FLinker::Serialize( FArchive& Ar )
 			Ar << I.ClassPackage << I.ClassName;
 		}
 	}
+
 }
 
 void FLinker::AddReferencedObjects(FReferenceCollector& Collector)
diff --git a/Engine/Source/Runtime/CoreUObject/Private/UObject/LinkerLoad.cpp b/Engine/Source/Runtime/CoreUObject/Private/UObject/LinkerLoad.cpp
--- a/Engine/Source/Runtime/CoreUObject/Private/UObject/LinkerLoad.cpp
+++ b/Engine/Source/Runtime/CoreUObject/Private/UObject/LinkerLoad.cpp
@@ -11,6 +11,7 @@
 #include "LinkerManager.h"
 #include "Serialization/DeferredMessageLog.h"
 #include "UObject/UObjectThreadContext.h"
+#include "GatherableTextData.h"
 
 #define LOCTEXT_NAMESPACE "LinkerLoad"
 
@@ -573,6 +574,12 @@ FLinkerLoad::ELinkerStatus FLinkerLoad::Tick( float InTimeLimit, bool bInUseTime
 				Status = SerializeNameMap();
 			}
 
+			// Serialize the gatherable text data map.
+			if( Status == LINKER_Loaded )
+			{
+				Status = SerializeGatherableTextDataMap();
+			}
+
 			// Serialize the import map.
 			if( Status == LINKER_Loaded )
 			{
@@ -667,6 +674,7 @@ FLinkerLoad::FLinkerLoad(UPackage* InParent, const TCHAR* InFilename, uint32 InL
 , bHaveImportsBeenVerified(false)
 , Loader(nullptr)
 , NameMapIndex(0)
+, GatherableTextDataMapIndex(0)
 , ImportMapIndex(0)
 , ExportMapIndex(0)
 , DependsMapIndex(0)
@@ -1099,9 +1107,10 @@ FLinkerLoad::ELinkerStatus FLinkerLoad::SerializePackageFileSummary()
 		}
 
 		// Slack everything according to summary.
-		ImportMap   .Empty( Summary.ImportCount   );
-		ExportMap   .Empty( Summary.ExportCount   );
-		NameMap		.Empty( Summary.NameCount     );
+		ImportMap					.Empty( Summary.ImportCount				);
+		ExportMap					.Empty( Summary.ExportCount				);
+		GatherableTextDataMap		.Empty( Summary.GatherableTextDataCount );
+		NameMap						.Empty( Summary.NameCount				);
 		// Depends map gets pre-sized in SerializeDependsMap if used.
 
 		// Avoid serializing it again.
@@ -1163,6 +1172,38 @@ FLinkerLoad::ELinkerStatus FLinkerLoad::SerializeNameMap()
 	return ((NameMapIndex == Summary.NameCount) && !IsTimeLimitExceeded( TEXT("serializing name map") )) ? LINKER_Loaded : LINKER_TimedOut;
 }
 
+/**
+ * Serializes the gatherable text data container.
+ */
+FLinkerLoad::ELinkerStatus FLinkerLoad::SerializeGatherableTextDataMap(bool bForceEnableForCommandlet)
+{
+#if WITH_EDITORONLY_DATA
+	DECLARE_SCOPE_CYCLE_COUNTER( TEXT( "FLinkerLoad::SerializeGatherableTextDataMap" ), STAT_LinkerLoad_SerializeGatherableTextDataMap, STATGROUP_LinkerLoad );
+
+	// Skip serializing gatherable text data if we are using seekfree loading
+	if( !bForceEnableForCommandlet && !GIsEditor )
+	{
+		return LINKER_Loaded;
+	}
+
+	if( GatherableTextDataMapIndex == 0 && Summary.GatherableTextDataCount > 0 )
+	{
+		Seek( Summary.GatherableTextDataOffset );
+	}
+
+	while( GatherableTextDataMapIndex < Summary.GatherableTextDataCount && !IsTimeLimitExceeded(TEXT("serializing gatherable text data map"),100) )
+	{
+		FGatherableTextData* GatherableTextData = new(GatherableTextDataMap)FGatherableTextData;
+		*this << *GatherableTextData;
+		GatherableTextDataMapIndex++;
+	}
+
+	return ((GatherableTextDataMapIndex == Summary.GatherableTextDataCount) && !IsTimeLimitExceeded( TEXT("serializing gatherable text data map") )) ? LINKER_Loaded : LINKER_TimedOut;
+#endif
+
+	return LINKER_Loaded;
+}
+
 /**
  * Serializes the import map.
  */
@@ -3956,6 +3997,7 @@ void FLinkerLoad::Detach()
 
 	// Empty out no longer used arrays.
 	NameMap.Empty();
+	GatherableTextDataMap.Empty();
 	ImportMap.Empty();
 	ExportMap.Empty();
 
diff --git a/Engine/Source/Runtime/CoreUObject/Private/UObject/LinkerManager.cpp b/Engine/Source/Runtime/CoreUObject/Private/UObject/LinkerManager.cpp
--- a/Engine/Source/Runtime/CoreUObject/Private/UObject/LinkerManager.cpp
+++ b/Engine/Source/Runtime/CoreUObject/Private/UObject/LinkerManager.cpp
@@ -44,12 +44,14 @@ bool FLinkerManager::Exec(class UWorld* InWorld, const TCHAR* Cmd, FOutputDevice
 			}
 			Ar.Logf
 				(
-				TEXT("%s (%s): Names=%i (%iK/%iK) Imports=%i (%iK) Exports=%i (%iK) Gen=%i Bulk=%i"),
+				TEXT("%s (%s): Names=%i (%iK/%iK) Text=%i (%iK) Imports=%i (%iK) Exports=%i (%iK) Gen=%i Bulk=%i"),
 				*Linker->Filename,
 				*Linker->LinkerRoot->GetFullName(),
 				Linker->NameMap.Num(),
 				Linker->NameMap.Num() * sizeof(FName) / 1024,
 				NameSize / 1024,
+				Linker->GatherableTextDataMap.Num(),
+				Linker->GatherableTextDataMap.Num() * sizeof(FGatherableTextData) / 1024,
 				Linker->ImportMap.Num(),
 				Linker->ImportMap.Num() * sizeof(FObjectImport) / 1024,
 				Linker->ExportMap.Num(),
diff --git a/Engine/Source/Runtime/CoreUObject/Private/UObject/PackageFileSummary.cpp b/Engine/Source/Runtime/CoreUObject/Private/UObject/PackageFileSummary.cpp
--- a/Engine/Source/Runtime/CoreUObject/Private/UObject/PackageFileSummary.cpp
+++ b/Engine/Source/Runtime/CoreUObject/Private/UObject/PackageFileSummary.cpp
@@ -118,9 +118,13 @@ FArchive& operator<<( FArchive& Ar, FPackageFileSummary& Sum )
 		{
 			Ar.SetFilterEditorOnly(true);
 		}
-		Ar << Sum.NameCount     << Sum.NameOffset;
-		Ar << Sum.ExportCount   << Sum.ExportOffset;
-		Ar << Sum.ImportCount   << Sum.ImportOffset;
+		Ar << Sum.NameCount					<< Sum.NameOffset;
+		if (Sum.FileVersionUE4 >= VER_UE4_SERIALIZE_TEXT_IN_PACKAGES)
+		{
+			Ar << Sum.GatherableTextDataCount	<< Sum.GatherableTextDataOffset;
+		}
+		Ar << Sum.ExportCount				<< Sum.ExportOffset;
+		Ar << Sum.ImportCount				<< Sum.ImportOffset;
 		Ar << Sum.DependsOffset;
 
 		if (Ar.IsLoading() && (Sum.FileVersionUE4 < VER_UE4_OLDEST_LOADABLE_PACKAGE || Sum.FileVersionUE4 > GPackageFileUE4Version))
diff --git a/Engine/Source/Runtime/CoreUObject/Private/UObject/SavePackage.cpp b/Engine/Source/Runtime/CoreUObject/Private/UObject/SavePackage.cpp
--- a/Engine/Source/Runtime/CoreUObject/Private/UObject/SavePackage.cpp
+++ b/Engine/Source/Runtime/CoreUObject/Private/UObject/SavePackage.cpp
@@ -685,7 +685,6 @@ FArchive& FArchiveSaveTagImports::operator<<( FAssetPtr& AssetPtr)
 	return *this << ID;
 }
 
-
 /**
  * Helper class for package compression.
  */
@@ -2890,7 +2889,7 @@ bool UPackage::SavePackage( UPackage* InOuter, UObject* Base, EObjectFlags TopLe
 
 		FText StatusMessage = FText::Format( NSLOCTEXT("Core", "SavingFile", "Saving file: {CleanFilename}..."), Args );
 
-		const int32 TotalSaveSteps = 31;
+		const int32 TotalSaveSteps = 33;
 		FScopedSlowTask SlowTask(TotalSaveSteps, StatusMessage, bSlowTask);
 		SlowTask.MakeDialog(SaveFlags & SAVE_FromAutosave ? true : false);
 
@@ -3127,6 +3126,33 @@ bool UPackage::SavePackage( UPackage* InOuter, UObject* Base, EObjectFlags TopLe
 				if ( EndSavingIfCancelled( Linker, TempFilename ) ) { return false; }
 				SlowTask.EnterProgressFrame();
 
+				if ( !(Linker->Summary.PackageFlags & PKG_FilterEditorOnly) )
+				{
+					TArray<UObject*> TagExpObjects;
+					GetObjectsWithAnyMarks(TagExpObjects, OBJECTMARK_TagExp);
+					for (UObject* const Object : TagExpObjects)
+					{
+						if( !Object->HasAnyFlags( RF_Transient | RF_PendingKill ) )
+						{
+							GatherLocalizationDataFromPropertiesOfDataStructure(Object->GetClass(), Object, Linker->GatherableTextDataMap);
+
+							for(UClass* Class = Object->GetClass(); Class != nullptr; Class = Class->GetSuperClass())
+							{
+								FLocalizationDataGatheringCallback* const CustomCallback = GetTypeSpecificLocalizationDataGatheringCallbacks().Find(Class);
+								if (CustomCallback)
+								{
+									(*CustomCallback)(Object, Linker->GatherableTextDataMap);
+								}
+							}
+						}
+					}
+				}
+
+				UE_LOG_COOK_TIME(TEXT("Gather Localizable Text Data"));
+				
+				if ( EndSavingIfCancelled( Linker, TempFilename ) ) { return false; }
+				SlowTask.EnterProgressFrame();
+
 				TArray<UObject*> PrivateObjects;
 				TArray<UObject*> ObjectsInOtherMaps;
 				TArray<UObject*> LevelObjects;
@@ -3407,12 +3433,30 @@ bool UPackage::SavePackage( UPackage* InOuter, UObject* Base, EObjectFlags TopLe
 					Linker->NameIndices.Add(Linker->NameMap[i], i);
 				}
 
-				
 				UE_LOG_COOK_TIME(TEXT("Serialize Names"));
 
 				if ( EndSavingIfCancelled( Linker, TempFilename ) ) { return false; }
 				SlowTask.EnterProgressFrame();
 
+				Linker->Summary.GatherableTextDataOffset = 0;
+				Linker->Summary.GatherableTextDataCount = 0;
+				if ( !(Linker->Summary.PackageFlags & PKG_FilterEditorOnly) )
+				{
+					Linker->Summary.GatherableTextDataOffset = Linker->Tell();
+
+					// Save gatherable text data.
+					Linker->Summary.GatherableTextDataCount = Linker->GatherableTextDataMap.Num();
+					for (FGatherableTextData& GatherableTextData : Linker->GatherableTextDataMap)
+					{
+						*Linker << GatherableTextData;
+					}
+				}
+
+				UE_LOG_COOK_TIME(TEXT("Serialize Gatherable Text Data"));
+
+				if ( EndSavingIfCancelled( Linker, TempFilename ) ) { return false; }
+				SlowTask.EnterProgressFrame();
+
 				// Build ImportMap.
 				{
 					TArray<UObject*> TagImpObjects;
diff --git a/Engine/Source/Runtime/CoreUObject/Public/Serialization/PropertyLocalizationDataGathering.h b/Engine/Source/Runtime/CoreUObject/Public/Serialization/PropertyLocalizationDataGathering.h
new file mode 100644
--- /dev/null
+++ b/Engine/Source/Runtime/CoreUObject/Public/Serialization/PropertyLocalizationDataGathering.h
@@ -0,0 +1,9 @@
+// Copyright 1998-2015 Epic Games, Inc. All Rights Reserved.
+
+#pragma once
+
+#include "GatherableTextData.h"
+
+COREUOBJECT_API extern void GatherLocalizationDataFromPropertiesOfDataStructure(UStruct* const Structure, void* const Data, TArray<FGatherableTextData>& GatherableTextDataArray);
+COREUOBJECT_API extern void GatherLocalizationDataFromChildTextProperies(const FString& PathToParent, UProperty* const Property, void* const ValueAddress, TArray<FGatherableTextData>& GatherableTextDataArray, const bool ForceIsEditorOnly);
+COREUOBJECT_API extern void GatherLocalizationDataFromTextProperty(const FString& PathToParent, UTextProperty* const TextProperty, void* const ValueAddress, TArray<FGatherableTextData>& GatherableTextDataArray, const bool ForceIsEditorOnly);
\ No newline at end of file
diff --git a/Engine/Source/Runtime/CoreUObject/Public/UObject/CoreObject.h b/Engine/Source/Runtime/CoreUObject/Public/UObject/CoreObject.h
--- a/Engine/Source/Runtime/CoreUObject/Public/UObject/CoreObject.h
+++ b/Engine/Source/Runtime/CoreUObject/Public/UObject/CoreObject.h
@@ -4,7 +4,8 @@
 
 #include "ObjectBase.h"
 #include "WorldCompositionUtility.h"
-
+#include "GatherableTextData.h"
+#include "PropertyLocalizationDataGathering.h"
 
 /**
  * Structure to hold information about an external packages objects used in cross-level references
@@ -264,6 +265,13 @@ public:
 		return (PackageFlags & PKG_RequiresLocalizationGather) ? true : false;
 	}
 
+	typedef TFunction<void (const UObject* const, TArray<FGatherableTextData>&)> FLocalizationDataGatheringCallback;
+	static TMap<UField*, FLocalizationDataGatheringCallback>& GetTypeSpecificLocalizationDataGatheringCallbacks()
+	{
+		static TMap<UField*, FLocalizationDataGatheringCallback> TypeSpecificLocalizationDataGatheringCallbacks;
+		return TypeSpecificLocalizationDataGatheringCallbacks;
+	}
+
 	/** Returns true if this package has a thumbnail map */
 	bool HasThumbnailMap() const
 	{
diff --git a/Engine/Source/Runtime/CoreUObject/Public/UObject/Linker.h b/Engine/Source/Runtime/CoreUObject/Public/UObject/Linker.h
--- a/Engine/Source/Runtime/CoreUObject/Public/UObject/Linker.h
+++ b/Engine/Source/Runtime/CoreUObject/Public/UObject/Linker.h
@@ -4,6 +4,7 @@
 
 #include "ObjectBase.h"
 #include "EngineVersion.h"
+#include "GatherableTextData.h"
 
 
 DECLARE_LOG_CATEGORY_EXTERN(LogLinker, Log, All);
@@ -589,6 +590,16 @@ public:
 	 */
 	int32 	NameOffset;
 
+	/**
+	 * Number of gatherable text data items in this package
+	 */
+	int32		GatherableTextDataCount;
+
+	/**
+	 * Location into the file on disk for the gatherable text data items
+	 */
+	int32 	GatherableTextDataOffset;
+	
 	/**
 	 * Number of exports contained in this package
 	 */
@@ -965,6 +976,9 @@ public:
 	/** Names used by objects contained within this package */
 	TArray<FName>			NameMap;
 
+	/** Gatherable text data contained within this package */
+	TArray<FGatherableTextData> GatherableTextDataMap;
+
 	/** The name of the file for this package */
 	FString					Filename;
 
@@ -1340,15 +1354,17 @@ public:
 private:
 	// Variables used during async linker creation.
 
-	/** Current index into name map, used by async linker creation for spreading out serializing name entries.				*/
+	/** Current index into name map, used by async linker creation for spreading out serializing name entries.					*/
 	int32						NameMapIndex;
-	/** Current index into import map, used by async linker creation for spreading out serializing importmap entries.		*/	
+	/** Current index into gatherable text data map, used by async linker creation for spreading out serializing text entries.	*/
+	int32						GatherableTextDataMapIndex;
+	/** Current index into import map, used by async linker creation for spreading out serializing importmap entries.			*/	
 	int32						ImportMapIndex;
-	/** Current index into export map, used by async linker creation for spreading out serializing exportmap entries.		*/
+	/** Current index into export map, used by async linker creation for spreading out serializing exportmap entries.			*/
 	int32						ExportMapIndex;
-	/** Current index into depends map, used by async linker creation for spreading out serializing dependsmap entries.		*/
+	/** Current index into depends map, used by async linker creation for spreading out serializing dependsmap entries.			*/
 	int32						DependsMapIndex;
-	/** Current index into export hash map, used by async linker creation for spreading out hashing exports.				*/
+	/** Current index into export hash map, used by async linker creation for spreading out hashing exports.					*/
 	int32						ExportHashIndex;
 
 
@@ -1899,6 +1915,11 @@ private:
 	ELinkerStatus SerializeDependsMap();
 
 public:
+	/**
+	 * Serializes the gatherable text data container.
+	 */
+	COREUOBJECT_API ELinkerStatus SerializeGatherableTextDataMap(bool bForceEnableForCommandlet = false);
+
 	/**
 	 * Serializes thumbnails
 	 */
diff --git a/Engine/Source/Runtime/Engine/Private/DataTable.cpp b/Engine/Source/Runtime/Engine/Private/DataTable.cpp
--- a/Engine/Source/Runtime/Engine/Private/DataTable.cpp
+++ b/Engine/Source/Runtime/Engine/Private/DataTable.cpp
@@ -11,9 +11,38 @@ DEFINE_LOG_CATEGORY(LogDataTable);
 ENGINE_API const FString FDataTableRowHandle::Unknown(TEXT("UNKNOWN"));
 ENGINE_API const FString FDataTableCategoryHandle::Unknown(TEXT("UNKNOWN"));
 
+#if WITH_EDITORONLY_DATA
+namespace
+{
+	void GatherDataTableForLocalization(const UObject* const Object, TArray<FGatherableTextData>& GatherableTextDataArray)
+	{
+		const UDataTable* const DataTable = CastChecked<UDataTable>(Object);
+
+		const FString PathToObject = DataTable->GetPathName();
+		for (const auto& Pair : DataTable->RowMap)
+		{
+			const FString PathToRow = PathToObject + TEXT(".") + Pair.Key.ToString();
+			for (TFieldIterator<UProperty> PropIt(DataTable->RowStruct, EFieldIteratorFlags::IncludeSuper, EFieldIteratorFlags::ExcludeDeprecated, EFieldIteratorFlags::IncludeInterfaces); PropIt; ++PropIt)
+			{
+				GatherLocalizationDataFromChildTextProperies(PathToRow, *PropIt, PropIt->ContainerPtrToValuePtr<void>(Pair.Value), GatherableTextDataArray, false);
+			}
+		}
+	}
+}
+#endif
+
 UDataTable::UDataTable(const FObjectInitializer& ObjectInitializer)
 	: Super(ObjectInitializer)
 {
+#if WITH_EDITORONLY_DATA
+	struct FAutomaticRegistrationOfLocalizationGatherer
+	{
+		FAutomaticRegistrationOfLocalizationGatherer()
+		{
+			UPackage::GetTypeSpecificLocalizationDataGatheringCallbacks().Add(UDataTable::StaticClass(), &GatherDataTableForLocalization);
+		}
+	} AutomaticRegistrationOfLocalizationGatherer;
+#endif
 }
 
 void UDataTable::LoadStructData(FArchive& Ar)
diff --git a/Engine/Source/Runtime/Engine/Private/DialogueWave.cpp b/Engine/Source/Runtime/Engine/Private/DialogueWave.cpp
--- a/Engine/Source/Runtime/Engine/Private/DialogueWave.cpp
+++ b/Engine/Source/Runtime/Engine/Private/DialogueWave.cpp
@@ -6,6 +6,424 @@
 #include "Sound/DialogueSoundWaveProxy.h"
 #include "Sound/SoundWave.h"
 #include "SubtitleManager.h"
+#include "GatherableTextData.h"
+#include "InternationalizationMetadata.h"
+
+#if WITH_EDITORONLY_DATA
+namespace
+{
+	// Class that takes DialogueWaves and prepares internationalization manifest entries 
+	class FDialogueHelper
+	{
+	public:
+		bool ProcessDialogueWave( const UDialogueWave* DialogueWave );
+		const FString& GetSpokenSource() const { return SpokenSource; }
+
+		const FTextSourceSiteContext& GetBaseContext() const { return Base; };
+		const TArray<FTextSourceSiteContext>& GetContextSpecificVariations() const { return ContextSpecificVariations; }
+
+	private:
+		TSharedPtr< FLocMetadataValue > GenSourceTargetMetadata( const FString& SpeakerName, const TArray< FString >& TargetNames, bool bCompact = true );
+		FString ArrayMetaDataToString( const TArray< TSharedPtr< FLocMetadataValue > >& MetadataArray );
+		FString GetDialogueVoiceName( const UDialogueVoice* DialogueVoice );
+		FString GetGrammaticalGenderString( TEnumAsByte<EGrammaticalGender::Type> Gender );
+		FString GetGrammaticalNumberString( TEnumAsByte<EGrammaticalNumber::Type> Plurality );
+
+	public:
+		static const FString DialogueNamespace;
+		static const FString PropertyName_VoiceActorDirection;
+		static const FString PropertyName_Speaker;
+		static const FString PropertyName_Speakers;
+		static const FString PropertyName_Targets;
+		static const FString PropertyName_GrammaticalGender;
+		static const FString PropertyName_GrammaticalPlurality;
+		static const FString PropertyName_TargetGrammaticalGender;
+		static const FString PropertyName_TargetGrammaticalNumber;
+		static const FString PropertyName_Optional;
+		static const FString PropertyName_DialogueVariations;
+		static const FString PropertyName_IsMature;
+
+	private:
+		FString DialogueKey;
+		FString SourceLocation;
+		FString SpokenSource;
+		FString VoiceActorDirection;
+		bool bIsMature;
+
+		// The non-optional entry
+		FTextSourceSiteContext Base;
+		// Optional entries that are context specific
+		TArray< FTextSourceSiteContext > ContextSpecificVariations;
+	};
+
+	const FString FDialogueHelper::DialogueNamespace						= TEXT("Dialogue");
+	const FString FDialogueHelper::PropertyName_VoiceActorDirection		= TEXT("Voice Actor Direction");
+	const FString FDialogueHelper::PropertyName_Speaker					= TEXT("Speaker");
+	const FString FDialogueHelper::PropertyName_Speakers					= TEXT("Speakers");
+	const FString FDialogueHelper::PropertyName_Targets					= TEXT("Targets");
+	const FString FDialogueHelper::PropertyName_GrammaticalGender			= TEXT("Gender");
+	const FString FDialogueHelper::PropertyName_GrammaticalPlurality		= TEXT("Plurality");
+	const FString FDialogueHelper::PropertyName_TargetGrammaticalGender	= TEXT("TargetGender");
+	const FString FDialogueHelper::PropertyName_TargetGrammaticalNumber	= TEXT("TargetPlurality");
+	const FString FDialogueHelper::PropertyName_Optional					= TEXT("Optional");
+	const FString FDialogueHelper::PropertyName_DialogueVariations			= TEXT("Variations");
+	const FString FDialogueHelper::PropertyName_IsMature					= TEXT("*IsMature");
+
+	bool FDialogueHelper::ProcessDialogueWave( const UDialogueWave* DialogueWave )
+	{
+		if( !DialogueWave )
+		{
+			return false;
+		}
+
+		DialogueKey = DialogueWave->LocalizationGUID.ToString();
+		SourceLocation = DialogueWave->GetPathName();
+		SpokenSource = DialogueWave->SpokenText;
+		VoiceActorDirection = DialogueWave->VoiceActorDirection;
+		bIsMature = DialogueWave->bMature;
+
+		// Stores the human readable info describing source and targets for each context of this DialogueWave
+		TArray< TSharedPtr< FLocMetadataValue > > VariationsDisplayInfoList;
+
+		for( const FDialogueContextMapping& ContextMapping : DialogueWave->ContextMappings )
+		{
+			const FDialogueContext& DialogueContext = ContextMapping.Context;
+			const UDialogueVoice* SpeakerDialogueVoice = DialogueContext.Speaker;
+
+			// Skip over entries with invalid speaker
+			if( !SpeakerDialogueVoice )
+			{
+				continue;
+			}
+
+			// Collect speaker info
+			FString SpeakerDisplayName = GetDialogueVoiceName( SpeakerDialogueVoice );
+			FString SpeakerGender = GetGrammaticalGenderString( SpeakerDialogueVoice->Gender );
+			FString SpeakerPlurality = GetGrammaticalNumberString( SpeakerDialogueVoice->Plurality );
+			FString SpeakerGuid = SpeakerDialogueVoice->LocalizationGUID.ToString();
+
+			EGrammaticalGender::Type AccumulatedTargetGender = (EGrammaticalGender::Type)-1;
+			EGrammaticalNumber::Type AccumulatedTargetPlurality = (EGrammaticalNumber::Type)-1;
+
+			TArray<FString> TargetGuidsList;
+			TArray<FString> TargetDisplayNameList;
+
+			// Collect info on all the targets
+			bool HasTarget = false;
+			for( const UDialogueVoice* TargetDialogueVoice : DialogueContext.Targets )
+			{
+				if( TargetDialogueVoice )
+				{
+					FString TargetDisplayName = GetDialogueVoiceName( TargetDialogueVoice );
+					TargetDisplayNameList.AddUnique( TargetDisplayName );
+					FString TargetGender = GetGrammaticalGenderString( TargetDialogueVoice->Gender );
+					FString TargetPlurality = GetGrammaticalNumberString( TargetDialogueVoice->Plurality );
+					FString TargetGuid = TargetDialogueVoice->LocalizationGUID.ToString();
+
+					TargetGuidsList.AddUnique( TargetGuid );
+
+					if( AccumulatedTargetGender == -1)
+					{
+						AccumulatedTargetGender = TargetDialogueVoice->Gender;
+					}
+					else if( AccumulatedTargetGender != TargetDialogueVoice->Gender )
+					{
+						AccumulatedTargetGender = EGrammaticalGender::Mixed;
+					}
+
+					if( AccumulatedTargetPlurality == -1 )
+					{
+						AccumulatedTargetPlurality = TargetDialogueVoice->Plurality;
+					}
+					else if( AccumulatedTargetPlurality == EGrammaticalNumber::Singular )
+					{
+						AccumulatedTargetPlurality = EGrammaticalNumber::Plural;
+					}
+
+					HasTarget = true;
+				}
+			}
+
+			FString FinalTargetGender = HasTarget ? GetGrammaticalGenderString( AccumulatedTargetGender ) : TEXT("");
+			FString FinalTargetPlurality = HasTarget ? GetGrammaticalNumberString( AccumulatedTargetPlurality ) : TEXT("");
+
+			// Add the context specific variation
+			{
+				TSharedPtr< FLocMetadataObject > KeyMetaDataObject = MakeShareable( new FLocMetadataObject() );
+
+				// Setup a loc metadata object with all the context specific keys.
+				{
+					KeyMetaDataObject->SetStringField( PropertyName_GrammaticalGender, SpeakerGender );
+					KeyMetaDataObject->SetStringField( PropertyName_GrammaticalPlurality, SpeakerPlurality );
+					KeyMetaDataObject->SetStringField( PropertyName_Speaker, SpeakerGuid );
+					KeyMetaDataObject->SetStringField( PropertyName_TargetGrammaticalGender, FinalTargetGender );
+					KeyMetaDataObject->SetStringField( PropertyName_TargetGrammaticalNumber, FinalTargetPlurality );
+
+					TArray< TSharedPtr< FLocMetadataValue > > TargetGuidsMetadata;
+					for( FString& TargetGuid : TargetGuidsList )
+					{
+						TargetGuidsMetadata.Add( MakeShareable( new FLocMetadataValueString( TargetGuid ) ) );
+					}
+
+					KeyMetaDataObject->SetArrayField( PropertyName_Targets, TargetGuidsMetadata );
+				}
+
+				// Create the human readable info that describes the source and target of this dialogue
+				TSharedPtr< FLocMetadataValue > SourceTargetInfo = GenSourceTargetMetadata( SpeakerDisplayName, TargetDisplayNameList );
+
+				TSharedPtr< FLocMetadataObject > InfoMetaDataObject = MakeShareable( new FLocMetadataObject() );
+
+				// Setup a loc metadata object with all the context specific info.  This usually includes human readable descriptions of the dialogue
+				{
+					if( SourceTargetInfo.IsValid() )
+					{
+						InfoMetaDataObject->SetField( PropertyName_DialogueVariations, SourceTargetInfo );
+					}
+
+					if( !VoiceActorDirection.IsEmpty() )
+					{
+						InfoMetaDataObject->SetStringField( PropertyName_VoiceActorDirection, VoiceActorDirection );
+					}
+				}
+
+
+				FTextSourceSiteContext& Context = *(new(ContextSpecificVariations) FTextSourceSiteContext);
+
+				// Setup the context
+				{
+					Context.KeyName = DialogueWave->GetContextLocalizationKey( DialogueContext );
+					Context.SiteDescription = SourceLocation;
+					Context.IsOptional = true;
+					Context.KeyMetaData = KeyMetaDataObject->Values.Num() > 0 ? *KeyMetaDataObject : FLocMetadataObject();
+					Context.InfoMetaData = InfoMetaDataObject->Values.Num() > 0 ? *InfoMetaDataObject : FLocMetadataObject();
+				}
+
+				{
+					// Add human readable info describing the source and targets of this dialogue to the non-optional manifest entry if it does not exist already
+					bool bAddContextDisplayInfoToBase = true;
+					for( TSharedPtr< FLocMetadataValue > VariationInfo : VariationsDisplayInfoList )
+					{
+						if( *SourceTargetInfo == *VariationInfo)
+						{
+							bAddContextDisplayInfoToBase = false;
+							break;
+						}
+					}
+					if( bAddContextDisplayInfoToBase )
+					{
+						VariationsDisplayInfoList.Add( SourceTargetInfo );
+					}
+				}
+			}
+		}
+
+		// Create the base, non-optional entry
+		{
+			TSharedPtr< FLocMetadataObject > InfoMetaDataObject = MakeShareable( new FLocMetadataObject() );
+
+			// Setup a loc metadata object with all the context specific info.  This usually includes human readable descriptions of the dialogue
+			{
+				if( VariationsDisplayInfoList.Num() > 0 )
+				{
+					InfoMetaDataObject->SetArrayField( PropertyName_DialogueVariations, VariationsDisplayInfoList );
+				}
+
+				if( !VoiceActorDirection.IsEmpty() )
+				{
+					InfoMetaDataObject->SetStringField( PropertyName_VoiceActorDirection, VoiceActorDirection );
+				}
+			}
+
+			
+			// Setup the context
+			{
+				Base.KeyName = DialogueKey;
+				Base.SiteDescription = SourceLocation;
+				Base.IsOptional = false;
+				Base.InfoMetaData = InfoMetaDataObject->Values.Num() > 0 ? *InfoMetaDataObject : FLocMetadataObject();
+			}
+
+		}
+
+		return true;
+	}
+
+	FString FDialogueHelper::GetGrammaticalNumberString(TEnumAsByte<EGrammaticalNumber::Type> Plurality)
+	{
+		switch(Plurality)
+		{
+		case EGrammaticalNumber::Singular:
+			return TEXT("Singular");
+			break;
+		case EGrammaticalNumber::Plural:
+			return TEXT("Plural");
+			break;
+		default:
+			return FString();
+		}
+	}
+
+	FString FDialogueHelper::GetGrammaticalGenderString(TEnumAsByte<EGrammaticalGender::Type> Gender)
+	{
+		switch(Gender)
+		{
+		case EGrammaticalGender::Neuter:
+			return TEXT("Neuter");
+			break;
+		case EGrammaticalGender::Masculine:
+			return TEXT("Masculine");
+			break;
+		case EGrammaticalGender::Feminine:
+			return TEXT("Feminine");
+			break;
+		case EGrammaticalGender::Mixed:
+			return TEXT("Mixed");
+			break;
+		default:
+			return TEXT("");
+		}
+	}
+
+	FString FDialogueHelper::GetDialogueVoiceName( const UDialogueVoice* DialogueVoice )
+	{
+		FString Name = DialogueVoice->GetName();
+		return Name;
+	}
+
+	FString FDialogueHelper::ArrayMetaDataToString( const TArray< TSharedPtr< FLocMetadataValue > >& MetadataArray )
+	{
+		FString FinalString;
+		if( MetadataArray.Num() > 0 )
+		{
+			TArray< FString > MetadataStrings;
+			for( TSharedPtr< FLocMetadataValue > DataEntry : MetadataArray )
+			{
+				if( DataEntry->Type == ELocMetadataType::String )
+				{
+					MetadataStrings.Add(DataEntry->AsString());
+				}
+			}
+			MetadataStrings.Sort();
+			for( const FString& StrEntry : MetadataStrings )
+			{
+				if(FinalString.Len())
+				{
+					FinalString += TEXT(",");
+				}
+				FinalString += StrEntry;
+			}
+		}
+		return FinalString;
+	}
+
+	TSharedPtr< FLocMetadataValue > FDialogueHelper::GenSourceTargetMetadata( const FString& SpeakerName, const TArray< FString >& TargetNames, bool bCompact )
+	{
+		/*
+		This function can support two different formats.
+
+		The first format is compact and results in string entries that will later be combined into something like this
+		"Variations": [
+		"Jenny -> Audience",
+		"Zak -> Audience"
+		]
+
+		The second format is verbose and results in object entries that will later be combined into something like this
+		"VariationsTest": [
+		{
+		"Speaker": "Jenny",
+		"Targets": [
+		"Audience"
+		]
+		},
+		{
+		"Speaker": "Zak",
+		"Targets": [
+		"Audience"
+		]
+		}
+		]
+		*/
+
+		TSharedPtr< FLocMetadataValue > Result;
+		if( bCompact )
+		{
+			TArray<FString> SortedTargetNames = TargetNames;
+			SortedTargetNames.Sort();
+			FString TargetNamesString;
+			for( const FString& StrEntry : SortedTargetNames )
+			{
+				if(TargetNamesString.Len())
+				{
+					TargetNamesString += TEXT(",");
+				}
+				TargetNamesString += StrEntry;
+			}
+			Result = MakeShareable( new FLocMetadataValueString( FString::Printf( TEXT("%s -> %s" ), *SpeakerName, *TargetNamesString ) ) );
+		}
+		else
+		{
+			TArray< TSharedPtr< FLocMetadataValue > > TargetNamesMetadataList;
+			for( const FString& StrEntry: TargetNames )
+			{
+				TargetNamesMetadataList.Add( MakeShareable( new FLocMetadataValueString( StrEntry ) ) );
+			}
+
+			TSharedPtr< FLocMetadataObject > MetadataObj = MakeShareable( new FLocMetadataObject() );
+			MetadataObj->SetStringField( PropertyName_Speaker, SpeakerName );
+			MetadataObj->SetArrayField( PropertyName_Targets, TargetNamesMetadataList );
+
+			Result = MakeShareable( new FLocMetadataValueObject( MetadataObj.ToSharedRef() ) );
+		}
+		return Result;
+	}
+
+	void GatherDialogueWaveForLocalization(const UObject* const Object, TArray<FGatherableTextData>& GatherableTextDataArray)
+	{
+		const UDialogueWave* const DialogueWave = CastChecked<UDialogueWave>(Object);
+
+		FDialogueHelper DialogueHelper;
+		if( DialogueHelper.ProcessDialogueWave( DialogueWave ) )
+		{
+			const FString& SpokenSource = DialogueHelper.GetSpokenSource();
+
+			if( !SpokenSource.IsEmpty() )
+			{
+				FTextSourceData SourceData;
+				{
+					SourceData.SourceString = SpokenSource;
+				}
+
+				FGatherableTextData* GatherableTextData = GatherableTextDataArray.FindByPredicate([&](const FGatherableTextData& Candidate)
+				{
+					return	Candidate.NamespaceName == DialogueHelper.DialogueNamespace && 
+						Candidate.SourceData.SourceString == SourceData.SourceString &&
+						Candidate.SourceData.SourceStringMetaData == SourceData.SourceStringMetaData;
+				});
+				if(!GatherableTextData)
+				{
+					GatherableTextData = new(GatherableTextDataArray) FGatherableTextData;
+					GatherableTextData->NamespaceName = DialogueHelper.DialogueNamespace;
+					GatherableTextData->SourceData = SourceData;
+				}
+
+				{
+					const FTextSourceSiteContext& Base = DialogueHelper.GetBaseContext();
+
+					FTextSourceSiteContext& SourceSiteContext = *(new(GatherableTextData->SourceSiteContexts) FTextSourceSiteContext(Base));
+				}
+
+				{
+					const TArray<FTextSourceSiteContext>& Variations = DialogueHelper.GetContextSpecificVariations();
+					for( const FTextSourceSiteContext& Variation : Variations )
+					{
+						FTextSourceSiteContext& SourceSiteContext = *(new(GatherableTextData->SourceSiteContexts) FTextSourceSiteContext(Variation));
+					}
+				}
+			}
+		}
+	}
+}
+#endif
 
 bool operator==(const FDialogueContextMapping& LHS, const FDialogueContextMapping& RHS)
 {
@@ -101,6 +519,16 @@ UDialogueWave::UDialogueWave(const FObjectInitializer& ObjectInitializer)
 	: Super(ObjectInitializer)
 	, LocalizationGUID( FGuid::NewGuid() )
 {
+#if WITH_EDITORONLY_DATA
+	struct FAutomaticRegistrationOfLocalizationGatherer
+	{
+		FAutomaticRegistrationOfLocalizationGatherer()
+		{
+			UPackage::GetTypeSpecificLocalizationDataGatheringCallbacks().Add(UDialogueWave::StaticClass(), &GatherDialogueWaveForLocalization);
+		}
+	} AutomaticRegistrationOfLocalizationGatherer;
+#endif
+
 	ContextMappings.Add( FDialogueContextMapping() );
 }
 
@@ -210,7 +638,7 @@ USoundBase* UDialogueWave::GetWaveFromContext(const FDialogueContext& Context) c
 			break;
 		}
 	}
-	
+
 	return Proxy;
 }
 
diff --git a/Engine/Source/Runtime/Engine/Private/UserDefinedStruct.cpp b/Engine/Source/Runtime/Engine/Private/UserDefinedStruct.cpp
--- a/Engine/Source/Runtime/Engine/Private/UserDefinedStruct.cpp
+++ b/Engine/Source/Runtime/Engine/Private/UserDefinedStruct.cpp
@@ -7,9 +7,39 @@
 #endif //WITH_EDITOR
 #include "Engine/UserDefinedStruct.h"
 
+#if WITH_EDITORONLY_DATA
+namespace
+{
+	void GatherUserDefinedStructForLocalization(const UObject* const Object, TArray<FGatherableTextData>& GatherableTextDataArray)
+	{
+		const UUserDefinedStruct* const UserDefinedStruct = CastChecked<UUserDefinedStruct>(Object);
+
+		const FString PathToObject = UserDefinedStruct->GetPathName();
+
+		FStructOnScope StructData(UserDefinedStruct);
+		FStructureEditorUtils::Fill_MakeStructureDefaultValue(UserDefinedStruct, StructData.GetStructMemory());
+
+		// Iterate over all fields of the object's class.
+		for (TFieldIterator<UProperty> PropIt(StructData.GetStruct(), EFieldIteratorFlags::IncludeSuper, EFieldIteratorFlags::ExcludeDeprecated, EFieldIteratorFlags::IncludeInterfaces); PropIt; ++PropIt)
+		{
+			GatherLocalizationDataFromChildTextProperies(PathToObject, *PropIt, PropIt->ContainerPtrToValuePtr<void>(StructData.GetStructMemory()), GatherableTextDataArray, false);
+		}
+	}
+}
+#endif
+
 UUserDefinedStruct::UUserDefinedStruct(const FObjectInitializer& ObjectInitializer)
 	: Super(ObjectInitializer)
 {
+#if WITH_EDITORONLY_DATA
+	struct FAutomaticRegistrationOfLocalizationGatherer
+	{
+		FAutomaticRegistrationOfLocalizationGatherer()
+		{
+			UPackage::GetTypeSpecificLocalizationDataGatheringCallbacks().Add(UUserDefinedStruct::StaticClass(), &GatherUserDefinedStructForLocalization);
+		}
+	} AutomaticRegistrationOfLocalizationGatherer;
+#endif
 }
 
 #if WITH_EDITOR
