Packages kept from previous cook now get translated into short package names for comparision when loading asset registry.

diff --git a/Engine/Source/Editor/UnrealEd/Classes/CookOnTheSide/CookOnTheFlyServer.h b/Engine/Source/Editor/UnrealEd/Classes/CookOnTheSide/CookOnTheFlyServer.h
--- a/Engine/Source/Editor/UnrealEd/Classes/CookOnTheSide/CookOnTheFlyServer.h
+++ b/Engine/Source/Editor/UnrealEd/Classes/CookOnTheSide/CookOnTheFlyServer.h
@@ -679,6 +679,7 @@ private:
 	FString GetCachedStandardPackageFilename( const UPackage* Package ) const;
 	FName GetCachedStandardPackageFileFName( const UPackage* Package ) const;
 	const FString& GetCachedSandboxFilename( const UPackage* Package, TAutoPtr<class FSandboxPlatformFile>& SandboxFile ) const;
+	const FName* GetCachedPackageFilenameToPackageFName(const FName& StandardPackageFilename) const;
 	const FCachedPackageFilename& Cache(const FName& PackageName) const;
 	void ClearPackageFilenameCache() const;
 	bool ClearPackageFilenameCacheForPackage( const UPackage* Package ) const;
@@ -687,6 +688,7 @@ private:
 	// declared mutable as it's used to cache package filename strings and I don't want to declare all functions using it as non const
 	// used by GetCached * Filename functions
 	mutable TMap<FName, FCachedPackageFilename> PackageFilenameCache; // filename cache (only process the string operations once)
+	mutable TMap<FName, FName> PackageFilenameToPackageFNameCache;
 
 	// declared mutable as it's used purely as a cache and don't want to have to declare all the functions as non const just because of this cache
 	// used by IniSettingsOutOfDate and GetCurrentIniStrings 
