Slate: Added the ability to specify a minimum desired width for SViewportToolBarComboMenu

diff --git a/Engine/Source/Editor/UnrealEd/Private/SViewportToolBarComboMenu.cpp b/Engine/Source/Editor/UnrealEd/Private/SViewportToolBarComboMenu.cpp
--- a/Engine/Source/Editor/UnrealEd/Private/SViewportToolBarComboMenu.cpp
+++ b/Engine/Source/Editor/UnrealEd/Private/SViewportToolBarComboMenu.cpp
@@ -36,9 +36,7 @@ void SViewportToolBarComboMenu::Construct( const FArguments& InArgs )
 	}
 
 	{
-		SAssignNew( MenuAnchor, SMenuAnchor )
-		.Placement( MenuPlacement_BelowAnchor )
-		[
+		TSharedRef<SWidget> ButtonContents =
 			SNew(SButton)
 			.ButtonStyle( FEditorStyle::Get(), EMultiBlockLocation::ToName(ButtonStyle,EMultiBlockLocation::End) )
 			.ContentPadding( FMargin( 5.0f, 0.0f ) )
@@ -79,7 +77,22 @@ void SViewportToolBarComboMenu::Construct( const FArguments& InArgs )
 					+SHorizontalBox::Slot()
 						.FillWidth(1.0f)
 				]
-			]
+			];
+		
+		if (InArgs._MinDesiredButtonWidth > 0.0f)
+		{
+			ButtonContents =
+				SNew(SBox)
+				.MinDesiredWidth(InArgs._MinDesiredButtonWidth)
+				[
+					ButtonContents
+				];
+		}
+
+		MenuAnchor = SNew(SMenuAnchor)
+		.Placement( MenuPlacement_BelowAnchor )
+		[
+			ButtonContents
 		]
 		.OnGetMenuContent( InArgs._OnGetMenuContent );
 	}
diff --git a/Engine/Source/Editor/UnrealEd/Public/SViewportToolBarComboMenu.h b/Engine/Source/Editor/UnrealEd/Public/SViewportToolBarComboMenu.h
--- a/Engine/Source/Editor/UnrealEd/Public/SViewportToolBarComboMenu.h
+++ b/Engine/Source/Editor/UnrealEd/Public/SViewportToolBarComboMenu.h
@@ -14,7 +14,7 @@
 class UNREALED_API SViewportToolBarComboMenu : public SCompoundWidget 
 {
 public:
-	SLATE_BEGIN_ARGS(SViewportToolBarComboMenu) : _BlockLocation(EMultiBlockLocation::Start) {}
+	SLATE_BEGIN_ARGS(SViewportToolBarComboMenu) : _BlockLocation(EMultiBlockLocation::Start), _MinDesiredButtonWidth(-1.0f) {}
 	
 		/** We need to know about the toolbar we are in */
 		SLATE_ARGUMENT( TSharedPtr<class SViewportToolBar>, ParentToolBar );
@@ -46,6 +46,9 @@ public:
 		/** The button location */
 		SLATE_ARGUMENT( EMultiBlockLocation::Type, BlockLocation )
 
+		/** The minimum desired width of the menu button contents */
+		SLATE_ARGUMENT( float, MinDesiredButtonWidth )
+
 	SLATE_END_ARGS( )
 
 	/**
