Editor: Gated editor flight controls to only take effect while modifier keys (ctrl, alt, shift, cmd) are not held down, to allow things like [modifier]+PageUp can be bound to other actions

diff --git a/Engine/Source/Editor/UnrealEd/Private/EditorViewportClient.cpp b/Engine/Source/Editor/UnrealEd/Private/EditorViewportClient.cpp
--- a/Engine/Source/Editor/UnrealEd/Private/EditorViewportClient.cpp
+++ b/Engine/Source/Editor/UnrealEd/Private/EditorViewportClient.cpp
@@ -1291,14 +1291,16 @@ void FEditorViewportClient::UpdateCameraMovement( float DeltaTime )
 		// Certain keys are only available while the flight camera input mode is active
 		const bool bUsingFlightInput = IsFlightCameraInputModeActive() || bIsUsingTrackpad;
 
+		// Is the current press unmodified?
+		const bool bUnmodifiedPress = !IsAltPressed() && !IsShiftPressed() && !IsCtrlPressed() && !IsCmdPressed();
+
 		// Do we want to use the regular arrow keys for flight input?
-		// Because the arrow keys are used for nudging actors, we'll only do this while the ALT key is up
-		const bool bAltDown = IsAltPressed();
-		const bool bRemapArrowKeys = !bAltDown;
+		// Because the arrow keys are used for things like nudging actors, we'll only do this while the press is unmodified
+		const bool bRemapArrowKeys = bUnmodifiedPress;
 
 		// Do we want to remap the various WASD keys for flight input?
 		const bool bRemapWASDKeys =
-			(!bAltDown && !IsCtrlPressed() && !IsShiftPressed()) &&
+			(bUnmodifiedPress) &&
 			(GetDefault<ULevelEditorViewportSettings>()->FlightCameraControlType == WASD_Always ||
 			( bUsingFlightInput &&
 			( GetDefault<ULevelEditorViewportSettings>()->FlightCameraControlType == WASD_RMBOnly && (Viewport->KeyState(EKeys::RightMouseButton ) ||Viewport->KeyState(EKeys::MiddleMouseButton) || Viewport->KeyState(EKeys::LeftMouseButton) || bIsUsingTrackpad ) ) ) ) &&
@@ -1318,13 +1320,13 @@ void FEditorViewportClient::UpdateCameraMovement( float DeltaTime )
 		// Forward/back
 		if( ( bRemapWASDKeys && Viewport->KeyState( FViewportNavigationCommands::Get().Forward->GetActiveChord()->Key ) ) ||
 			( bRemapArrowKeys && Viewport->KeyState( EKeys::Up ) ) ||
-			Viewport->KeyState( EKeys::NumPadEight ) )
+			( bUnmodifiedPress && Viewport->KeyState(EKeys::NumPadEight) ) )
 		{
 			CameraUserImpulseData->MoveForwardBackwardImpulse += 1.0f;
 		}
 		if( ( bRemapWASDKeys && Viewport->KeyState( FViewportNavigationCommands::Get().Backward->GetActiveChord()->Key ) ) ||
 			( bRemapArrowKeys && Viewport->KeyState( EKeys::Down ) ) ||
-			Viewport->KeyState( EKeys::NumPadTwo ) )
+			( bUnmodifiedPress && Viewport->KeyState( EKeys::NumPadTwo ) ) )
 		{
 			CameraUserImpulseData->MoveForwardBackwardImpulse -= 1.0f;
 		}
@@ -1332,37 +1334,37 @@ void FEditorViewportClient::UpdateCameraMovement( float DeltaTime )
 		// Right/left
 		if ( ( bRemapWASDKeys && Viewport->KeyState( FViewportNavigationCommands::Get().Right->GetActiveChord()->Key) ) ||
 			( bRemapArrowKeys && Viewport->KeyState( EKeys::Right ) ) ||
-			Viewport->KeyState( EKeys::NumPadSix ) )
+			( bUnmodifiedPress && Viewport->KeyState( EKeys::NumPadSix ) ) )
 		{
 			CameraUserImpulseData->MoveRightLeftImpulse += 1.0f;
 		}
 		if( ( bRemapWASDKeys && Viewport->KeyState( FViewportNavigationCommands::Get().Left->GetActiveChord()->Key ) ) ||
 			( bRemapArrowKeys && Viewport->KeyState( EKeys::Left ) ) ||
-			Viewport->KeyState( EKeys::NumPadFour ) )
+			( bUnmodifiedPress && Viewport->KeyState( EKeys::NumPadFour ) ) )
 		{
 			CameraUserImpulseData->MoveRightLeftImpulse -= 1.0f;
 		}
 
 		// Up/down
 		if( ( bRemapWASDKeys && Viewport->KeyState( FViewportNavigationCommands::Get().Up->GetActiveChord()->Key ) ) ||
-			Viewport->KeyState( EKeys::PageUp ) || Viewport->KeyState( EKeys::NumPadNine ) || Viewport->KeyState( EKeys::Add ) )
+			( bUnmodifiedPress && ( Viewport->KeyState( EKeys::PageUp ) || Viewport->KeyState( EKeys::NumPadNine ) || Viewport->KeyState( EKeys::Add ) ) ) )
 		{
 			CameraUserImpulseData->MoveUpDownImpulse += 1.0f;
 		}
 		if( ( bRemapWASDKeys && Viewport->KeyState( FViewportNavigationCommands::Get().Down->GetActiveChord()->Key) ) ||
-			Viewport->KeyState( EKeys::PageDown ) || Viewport->KeyState( EKeys::NumPadSeven ) || Viewport->KeyState( EKeys::Subtract ) )
+			( bUnmodifiedPress && ( Viewport->KeyState( EKeys::PageDown ) || Viewport->KeyState( EKeys::NumPadSeven ) || Viewport->KeyState( EKeys::Subtract ) ) ) )
 		{
 			CameraUserImpulseData->MoveUpDownImpulse -= 1.0f;
 		}
 
 		// Zoom FOV out/in
 		if( ( bRemapWASDKeys && Viewport->KeyState( FViewportNavigationCommands::Get().FovZoomOut->GetActiveChord()->Key ) ) ||
-			Viewport->KeyState( EKeys::NumPadOne ) )
+			( bUnmodifiedPress && Viewport->KeyState( EKeys::NumPadOne ) ) )
 		{
 			CameraUserImpulseData->ZoomOutInImpulse += 1.0f;
 		}
 		if( ( bRemapWASDKeys && Viewport->KeyState( FViewportNavigationCommands::Get().FovZoomIn->GetActiveChord()->Key ) ) ||
-			Viewport->KeyState( EKeys::NumPadThree ) )
+			( bUnmodifiedPress && Viewport->KeyState( EKeys::NumPadThree ) ) )
 		{
 			CameraUserImpulseData->ZoomOutInImpulse -= 1.0f;
 		}
