Editor: Decoupled the enable flag for the 2D rotate+translate gizmo from the enable for the 2D snap layers

diff --git a/Engine/Source/Editor/UnrealEd/Private/SEditorViewport.cpp b/Engine/Source/Editor/UnrealEd/Private/SEditorViewport.cpp
--- a/Engine/Source/Editor/UnrealEd/Private/SEditorViewport.cpp
+++ b/Engine/Source/Editor/UnrealEd/Private/SEditorViewport.cpp
@@ -506,7 +506,7 @@ bool SEditorViewport::IsTranslateRotateModeVisible() const
 
 bool SEditorViewport::Is2DModeVisible() const
 {
-	return GetDefault<ULevelEditor2DSettings>()->bMode2DEnabled;
+	return GetDefault<ULevelEditor2DSettings>()->bEnable2DWidget;
 }
 
 bool SEditorViewport::IsCoordSystemActive(ECoordSystem CoordSystem) const
@@ -529,7 +529,7 @@ void SEditorViewport::OnCycleWidgetMode()
 			++WidgetModeAsInt;
 		}
 
-		if ((WidgetModeAsInt == FWidget::WM_2D) && (!GetDefault<ULevelEditor2DSettings>()->bMode2DEnabled))
+		if ((WidgetModeAsInt == FWidget::WM_2D) && (!GetDefault<ULevelEditor2DSettings>()->bEnable2DWidget))
 		{
 			++WidgetModeAsInt;
 		}
diff --git a/Engine/Source/Editor/UnrealEd/Private/STransformViewportToolbar.cpp b/Engine/Source/Editor/UnrealEd/Private/STransformViewportToolbar.cpp
--- a/Engine/Source/Editor/UnrealEd/Private/STransformViewportToolbar.cpp
+++ b/Engine/Source/Editor/UnrealEd/Private/STransformViewportToolbar.cpp
@@ -793,7 +793,7 @@ ECheckBoxState STransformViewportToolBar::IsLayer2DSnapChecked() const
 EVisibility STransformViewportToolBar::IsLayer2DSnapVisible() const
 {
 	const ULevelEditor2DSettings* Settings2D = GetDefault<ULevelEditor2DSettings>();
-	return Settings2D->bMode2DEnabled ? EVisibility::Visible : EVisibility::Collapsed;
+	return Settings2D->bEnableSnapLayers ? EVisibility::Visible : EVisibility::Collapsed;
 }
 
 ECheckBoxState STransformViewportToolBar::IsScaleGridSnapChecked() const
diff --git a/Engine/Source/Editor/UnrealEd/Private/Settings/SettingsClasses.cpp b/Engine/Source/Editor/UnrealEd/Private/Settings/SettingsClasses.cpp
--- a/Engine/Source/Editor/UnrealEd/Private/Settings/SettingsClasses.cpp
+++ b/Engine/Source/Editor/UnrealEd/Private/Settings/SettingsClasses.cpp
@@ -458,9 +458,9 @@ void ULevelEditorViewportSettings::PostEditChangeProperty( struct FPropertyChang
 	else if (Name == GET_MEMBER_NAME_CHECKED(ULevelEditorViewportSettings, bEnableLayerSnap))
 	{
 		ULevelEditor2DSettings* Settings2D = GetMutableDefault<ULevelEditor2DSettings>();
-		if (bEnableLayerSnap && !Settings2D->bMode2DEnabled)
+		if (bEnableLayerSnap && !Settings2D->bEnableSnapLayers)
 		{
-			Settings2D->bMode2DEnabled = true;
+			Settings2D->bEnableSnapLayers = true;
 		}
 	}
 
diff --git a/Engine/Source/Editor/UnrealEd/Public/Settings/EditorProjectSettings.h b/Engine/Source/Editor/UnrealEd/Public/Settings/EditorProjectSettings.h
--- a/Engine/Source/Editor/UnrealEd/Public/Settings/EditorProjectSettings.h
+++ b/Engine/Source/Editor/UnrealEd/Public/Settings/EditorProjectSettings.h
@@ -116,15 +116,19 @@ class UNREALED_API ULevelEditor2DSettings : public UDeveloperSettings
 
 public:
 	/** If enabled will allow 2D mode */
-	UPROPERTY(EditAnywhere, config, Category = General, meta = (DisplayName = "Enable 2D Mode"))
-	bool bMode2DEnabled;
+	UPROPERTY(EditAnywhere, config, Category=General, meta=(DisplayName="Enable 2D combined translate + rotate widget"))
+	bool bEnable2DWidget;
+
+	/** If enabled will allow 2D mode */
+	UPROPERTY(EditAnywhere, config, Category=LayerSnapping)
+	bool bEnableSnapLayers;
 
 	/** Snap axis */
-	UPROPERTY(EditAnywhere, config, Category = LayerSnapping)
+	UPROPERTY(EditAnywhere, config, Category=LayerSnapping, meta=(EditCondition=bEnableSnapLayers))
 	ELevelEditor2DAxis SnapAxis;
 
 	/** Snap layers that are displayed in the viewport toolbar */
-	UPROPERTY(EditAnywhere, config, Category = LayerSnapping)
+	UPROPERTY(EditAnywhere, config, Category=LayerSnapping, meta=(EditCondition=bEnableSnapLayers))
 	TArray<FMode2DLayer> SnapLayers;
 
 public:
