Fixed default editor directories not being applied if empty directories are set in editor per project user settings

This addresses UE-15574 - Save As... crashes the editor with a PackageNameToFilename failed to convert Fatal log message

diff --git a/Engine/Source/Editor/UnrealEd/Private/EditorDirectories.cpp b/Engine/Source/Editor/UnrealEd/Private/EditorDirectories.cpp
--- a/Engine/Source/Editor/UnrealEd/Private/EditorDirectories.cpp
+++ b/Engine/Source/Editor/UnrealEd/Private/EditorDirectories.cpp
@@ -12,22 +12,11 @@ FEditorDirectories& FEditorDirectories::Get()
 
 void FEditorDirectories::LoadLastDirectories()
 {
-		// Initialize "last dir" array.
-	const FString DefaultDir = FPaths::GameContentDir();
 	for( int32 CurDirectoryIndex = 0; CurDirectoryIndex < ARRAY_COUNT( LastDir ); ++CurDirectoryIndex )
 	{
-		// Default all directories to the game content folder
-		LastDir[ CurDirectoryIndex ] = DefaultDir;
+		LastDir[CurDirectoryIndex].Reset();
 	}
 
-	const FString DefaultMapDir = FPaths::GameContentDir() / TEXT("Maps");
-	if( IFileManager::Get().DirectoryExists( *DefaultMapDir ) )
-	{
-		LastDir[ELastDirectory::LEVEL] = DefaultMapDir;
-	}
-
-	LastDir[ELastDirectory::PROJECT] = FPaths::RootDir();
-
 	// NOTE: We append a "2" to the section name to enforce backwards compatibility.  "Directories" is deprecated.
 	GConfig->GetString( TEXT("Directories2"), TEXT("UNR"),				LastDir[ELastDirectory::UNR],					GEditorPerProjectIni );
 	GConfig->GetString( TEXT("Directories2"), TEXT("BRUSH"),			LastDir[ELastDirectory::BRUSH],				GEditorPerProjectIni );
@@ -42,6 +31,32 @@ void FEditorDirectories::LoadLastDirectories()
 	GConfig->GetString( TEXT("Directories2"), TEXT("Level"),			LastDir[ELastDirectory::LEVEL],					GEditorPerProjectIni );
 	GConfig->GetString( TEXT("Directories2"), TEXT("Project"),			LastDir[ELastDirectory::PROJECT],				GEditorPerProjectIni );
 
+	// Set up some defaults if they're note defined in the ini
+	const FString DefaultDir = FPaths::GameContentDir();
+	for( int32 CurDirectoryIndex = 0; CurDirectoryIndex < ARRAY_COUNT( LastDir ); ++CurDirectoryIndex )
+	{
+		if (LastDir[ CurDirectoryIndex ].IsEmpty())
+		{
+			// Default all directories to the game content folder
+			if (CurDirectoryIndex == ELastDirectory::LEVEL)
+			{
+				const FString DefaultMapDir = FPaths::GameContentDir() / TEXT("Maps");
+				if( IFileManager::Get().DirectoryExists( *DefaultMapDir ) )
+				{
+					LastDir[CurDirectoryIndex] = DefaultMapDir;
+					continue;
+				}
+			}
+			else if (CurDirectoryIndex == ELastDirectory::PROJECT)
+			{
+				LastDir[CurDirectoryIndex] = FPaths::RootDir();
+				continue;
+			}
+
+			// Set to the default dir
+			LastDir[ CurDirectoryIndex ] = DefaultDir;
+		}
+	}
 }
 
 /** Writes the current "LastDir" array back out to the config files */
