Cooking - Changing the cook on the fly server to use the autosave directory, not the entire save directory for locating additional assets.

diff --git a/Engine/Source/Editor/UnrealEd/Private/CookOnTheFlyServer.cpp b/Engine/Source/Editor/UnrealEd/Private/CookOnTheFlyServer.cpp
--- a/Engine/Source/Editor/UnrealEd/Private/CookOnTheFlyServer.cpp
+++ b/Engine/Source/Editor/UnrealEd/Private/CookOnTheFlyServer.cpp
@@ -22,6 +22,7 @@
 #include "UnrealEdMessages.h"
 #include "GameDelegates.h"
 #include "PhysicsPublic.h"
+#include "AutoSaveUtils.h"
 
 // cook by the book requirements
 #include "Commandlets/ChunkManifestGenerator.h"
@@ -2710,7 +2711,7 @@ void UCookOnTheFlyServer::CleanSandbox( const bool bIterative )
 
 void UCookOnTheFlyServer::GenerateAssetRegistry(const TArray<ITargetPlatform*>& Platforms)
 {
-	FPackageName::RegisterMountPoint(TEXT("/Temp/"), FPaths::GameSavedDir());
+	FPackageName::RegisterMountPoint(TEXT("/Temp/"), AutoSaveUtils::GetAutoSaveDir());
 
 	// load the interface
 	FAssetRegistryModule& AssetRegistryModule = FModuleManager::LoadModuleChecked<FAssetRegistryModule>(TEXT("AssetRegistry"));
@@ -2727,7 +2728,7 @@ void UCookOnTheFlyServer::GenerateAssetRegistry(const TArray<ITargetPlatform*>&
 		// We want the registry to be fully initialized when generating streaming manifests too.
 		bool bEditor = IsRealtimeMode();
 
-		// editor will scan asset registry automagically 
+		// editor will scan asset registry automatically 
 		if ( !bEditor )
 		{
 			TArray<FString> ScanPaths;
