Brought GatherTextFromSource commandlet configuration in line with the rest of the gather commandlets with regards to using an ExcludePathFilters rather than ExcludePaths setting.

diff --git a/Engine/Source/Editor/UnrealEd/Private/Commandlets/GatherTextFromSourceCommandlet.cpp b/Engine/Source/Editor/UnrealEd/Private/Commandlets/GatherTextFromSourceCommandlet.cpp
--- a/Engine/Source/Editor/UnrealEd/Private/Commandlets/GatherTextFromSourceCommandlet.cpp
+++ b/Engine/Source/Editor/UnrealEd/Private/Commandlets/GatherTextFromSourceCommandlet.cpp
@@ -79,9 +79,20 @@ int32 UGatherTextFromSourceCommandlet::Main( const FString& Params )
 		return -1;
 	}
 
-	// Exclude paths
-	TArray<FString> ExcludePaths;
-	GetPathArrayFromConfig(*SectionName, TEXT("ExcludePaths"), ExcludePaths, GatherTextConfigPath);
+	// ExcludePathFilters
+	TArray<FString> ExcludePathFilters;
+	GetPathArrayFromConfig(*SectionName, TEXT("ExcludePathFilters"), ExcludePathFilters, GatherTextConfigPath);
+
+	// ExcludePaths (DEPRECATED)
+	{
+		TArray<FString> ExcludePaths;
+		GetPathArrayFromConfig(*SectionName, TEXT("ExcludePaths"), ExcludePaths, GatherTextConfigPath);
+		if (ExcludePaths.Num())
+		{
+			ExcludePathFilters.Append(ExcludePaths);
+			UE_LOG(LogGatherTextFromSourceCommandlet, Warning, TEXT("ExcludePaths detected in section %s. ExcludePaths is deprecated, please use ExcludePathFilters."), *SectionName);
+		}
+	}
 
 	// FileNameFilters
 	TArray<FString> FileNameFilters;
@@ -143,7 +154,7 @@ int32 UGatherTextFromSourceCommandlet::Main( const FString& Params )
 		bool bExclude = false;
 
 		//Ensure it does not match the exclude paths if there are some.
-		for (FString& ExcludePath : ExcludePaths)
+		for (FString& ExcludePath : ExcludePathFilters)
 		{
 			if (FoundFile.MatchesWildcard(ExcludePath) )
 			{
@@ -168,7 +179,7 @@ int32 UGatherTextFromSourceCommandlet::Main( const FString& Params )
 		{
 			SpecifiedDirectoriesString.Append(FString(SpecifiedDirectoriesString.IsEmpty() ? TEXT("") : TEXT("\n")) + FString::Printf(TEXT("+ %s"), *SearchDirectoryPath));
 		}
-		for (FString& ExcludePath : ExcludePaths)
+		for (FString& ExcludePath : ExcludePathFilters)
 		{
 			SpecifiedDirectoriesString.Append(FString(SpecifiedDirectoriesString.IsEmpty() ? TEXT("") : TEXT("\n")) + FString::Printf(TEXT("- %s"), *ExcludePath));
 		}
