Fixed issue with asset registry including cooked content.
Only include autosaves directory when cooking from commandlet.

diff --git a/Engine/Source/Editor/UnrealEd/Private/CookOnTheFlyServer.cpp b/Engine/Source/Editor/UnrealEd/Private/CookOnTheFlyServer.cpp
--- a/Engine/Source/Editor/UnrealEd/Private/CookOnTheFlyServer.cpp
+++ b/Engine/Source/Editor/UnrealEd/Private/CookOnTheFlyServer.cpp
@@ -2711,7 +2711,12 @@ void UCookOnTheFlyServer::CleanSandbox( const bool bIterative )
 
 void UCookOnTheFlyServer::GenerateAssetRegistry(const TArray<ITargetPlatform*>& Platforms)
 {
-	FPackageName::RegisterMountPoint(TEXT("/Temp/"), AutoSaveUtils::GetAutoSaveDir());
+	if (IsCookingInEditor() == false)
+	{
+		// we want to register the temporary save directory if we are cooking outside the editor.  
+		// If we are cooking inside the editor we never use this directory so don't worry about registring it
+		FPackageName::RegisterMountPoint(TEXT("/TempAutosave/"), FPaths::GameSavedDir() / GEngine->PlayOnConsoleSaveDir);
+	}
 
 	// load the interface
 	FAssetRegistryModule& AssetRegistryModule = FModuleManager::LoadModuleChecked<FAssetRegistryModule>(TEXT("AssetRegistry"));
@@ -2728,7 +2733,7 @@ void UCookOnTheFlyServer::GenerateAssetRegistry(const TArray<ITargetPlatform*>&
 		// We want the registry to be fully initialized when generating streaming manifests too.
 		bool bEditor = IsRealtimeMode();
 
-		// editor will scan asset registry automatically 
+		// editor will scan asset registry automagically 
 		if ( !bEditor )
 		{
 			TArray<FString> ScanPaths;
