Fixed LandscapeMountain issue where begincacheplatformdata not finishing because of time out over several frames.
Use shortpackage names when compairing against cooked package list instead of long package names (fix for iterative asset registry being out of date).

diff --git a/Engine/Source/Editor/UnrealEd/Private/Commandlets/ChunkManifestGenerator.cpp b/Engine/Source/Editor/UnrealEd/Private/Commandlets/ChunkManifestGenerator.cpp
--- a/Engine/Source/Editor/UnrealEd/Private/Commandlets/ChunkManifestGenerator.cpp
+++ b/Engine/Source/Editor/UnrealEd/Private/Commandlets/ChunkManifestGenerator.cpp
@@ -434,8 +434,10 @@ bool FChunkManifestGenerator::LoadAssetRegistry(const FString& SandboxPath, cons
 		}
 		for (FAssetData LoadedAssetData : AssetRegistryData)
 		{
+			
+			FName ShortPackageName = FName(*FPaths::GetBaseFilename(LoadedAssetData.PackageName.ToString()));
 			if (PackagesToKeep &&
-				PackagesToKeep->Contains(LoadedAssetData.PackageName) == false)
+				PackagesToKeep->Contains(ShortPackageName) == false)
 			{
 				continue;
 			}
diff --git a/Engine/Source/Editor/UnrealEd/Private/CookOnTheFlyServer.cpp b/Engine/Source/Editor/UnrealEd/Private/CookOnTheFlyServer.cpp
--- a/Engine/Source/Editor/UnrealEd/Private/CookOnTheFlyServer.cpp
+++ b/Engine/Source/Editor/UnrealEd/Private/CookOnTheFlyServer.cpp
@@ -467,13 +467,20 @@ const UCookOnTheFlyServer::FCachedPackageFilename& UCookOnTheFlyServer::Cache(co
 		FPaths::MakeStandardFilename(StandardFilename);
 		StandardFileFName = FName(*StandardFilename);
 	}
-
+	PackageFilenameToPackageFNameCache.Add(StandardFileFName, PackageName);
 	return PackageFilenameCache.Emplace( PackageName, MoveTemp(FCachedPackageFilename(MoveTemp(PackageFilename),MoveTemp(StandardFilename), StandardFileFName)) );
 }
 
+
+const FName* UCookOnTheFlyServer::GetCachedPackageFilenameToPackageFName(const FName& StandardPackageFilename) const
+{
+	return PackageFilenameToPackageFNameCache.Find(StandardPackageFilename);
+}
+
 void UCookOnTheFlyServer::ClearPackageFilenameCache() const 
 {
 	PackageFilenameCache.Empty();
+	PackageFilenameToPackageFNameCache.Empty();
 	// need to clear the IniVersionStringsMap too
 	CachedIniVersionStringsMap.Empty();
 }
@@ -1358,9 +1365,11 @@ uint32 UCookOnTheFlyServer::TickCookOnTheSide( const float TimeSlice, uint32 &Co
 
 		auto BeginPackageCacheForCookedPlatformData = [&]( const TArray<UObject*>& ObjectsInPackage )
 		{
-			for ( int I = 0; I < ObjectsInPackage.Num(); ++I )
+			if (CurrentReentryData.bBeginCacheFinished)
+				return true;
+			for (; CurrentReentryData.BeginCacheCount < ObjectsInPackage.Num(); ++CurrentReentryData.BeginCacheCount)
 			{
-				const auto& Obj = ObjectsInPackage[I];
+				const auto& Obj = ObjectsInPackage[CurrentReentryData.BeginCacheCount];
 				for ( const auto& TargetPlatform : TargetPlatforms )
 				{
 					Obj->BeginCacheForCookedPlatformData( TargetPlatform );
@@ -1374,15 +1383,12 @@ uint32 UCookOnTheFlyServer::TickCookOnTheSide( const float TimeSlice, uint32 &Co
 					return false;
 				}
 			}
+			CurrentReentryData.bBeginCacheFinished = true;
 			return true;
 		};
 
 		auto FinishPackageCacheForCookedPlatformData = [&]( const TArray<UObject*>& ObjectsInPackage )
 		{
-			
-			// if we get here and bIsAllDataCached is true then BeginCacheFOrCookedPlatformData has been called once on every object
-			// so now 
-			CurrentReentryData.bBeginCacheFinished = true;
 			for ( const auto& Obj : ObjectsInPackage )
 			{
 				for ( const auto& TargetPlatform : TargetPlatforms )
@@ -2685,17 +2691,18 @@ void UCookOnTheFlyServer::CleanSandbox( const bool bIterative )
 #endif
 							IFileManager::Get().Delete(*CookedFilename);
 
-							PackagesKeptFromPreviousCook.Remove(StandardCookedFileFName);
 							CookedPackages.RemoveFileForPlatform(StandardCookedFileFName, PlatformFName);
 						}
 						else
 						{
-							PackagesKeptFromPreviousCook.Add(StandardCookedFileFName);
+							FName ShortPackageName = FName(*FPaths::GetBaseFilename(StandardCookedFileFName.ToString()));
+							PackagesKeptFromPreviousCook.Add(ShortPackageName);
 						}
 					}
 					else
 					{
-						PackagesKeptFromPreviousCook.Add(StandardCookedFileFName);
+						FName ShortPackageName = FName(*FPaths::GetBaseFilename(StandardCookedFileFName.ToString()));
+						PackagesKeptFromPreviousCook.Add(ShortPackageName);
 					}
 				}
 			}
