UE-15778 Get and Set nodes for variables from a BP's parent do not work correctly for new or data only BPs

diff --git a/Engine/Source/Editor/UnrealEd/Private/Kismet2/BlueprintEditorUtils.cpp b/Engine/Source/Editor/UnrealEd/Private/Kismet2/BlueprintEditorUtils.cpp
--- a/Engine/Source/Editor/UnrealEd/Private/Kismet2/BlueprintEditorUtils.cpp
+++ b/Engine/Source/Editor/UnrealEd/Private/Kismet2/BlueprintEditorUtils.cpp
@@ -2016,8 +2016,7 @@ void FBlueprintEditorUtils::MarkBlueprintAsStructurallyModified(UBlueprint* Blue
 				if (SkelBlueprint
 					&& SkelBlueprint->Status != BS_BeingCreated
 					&& !SkelBlueprint->bBeingCompiled
-					&& !SkelBlueprint->bIsRegeneratingOnLoad
-					&& SkelBlueprint->bHasBeenRegenerated)
+					&& !SkelBlueprint->bIsRegeneratingOnLoad)
 				{
 					TArray<UClass*> ChildrenOfClass;
 					GetDerivedClasses(SkelClass, ChildrenOfClass, false);
@@ -2028,10 +2027,9 @@ void FBlueprintEditorUtils::MarkBlueprintAsStructurallyModified(UBlueprint* Blue
 					Results.bSilentMode = true;
 					Results.bLogInfoOnly = true;
 
-					//if (!SkelBlueprint->bHasBeenRegenerated)
 					{
 						bool const bWasRegenerating = SkelBlueprint->bIsRegeneratingOnLoad;
-						SkelBlueprint->bIsRegeneratingOnLoad = bIsCompilingOnLoad;
+						SkelBlueprint->bIsRegeneratingOnLoad |= bIsCompilingOnLoad;
 
 						FKismetCompilerOptions CompileOptions;
 						CompileOptions.CompileType = EKismetCompileType::SkeletonOnly;
