Fixed crash pasting invalid data into material editor

diff --git a/Engine/Source/Editor/UnrealEd/Private/MaterialGraphNode.cpp b/Engine/Source/Editor/UnrealEd/Private/MaterialGraphNode.cpp
--- a/Engine/Source/Editor/UnrealEd/Private/MaterialGraphNode.cpp
+++ b/Engine/Source/Editor/UnrealEd/Private/MaterialGraphNode.cpp
@@ -137,7 +137,7 @@ bool UMaterialGraphNode::CanPasteHere(const UEdGraph* TargetGraph) const
 				bIsValidFunctionExpression = false;
 			}
 
-			if (bIsValidFunctionExpression && IsAllowedExpressionType(MaterialExpression->GetClass(), MaterialGraph->MaterialFunction != NULL))
+			if (bIsValidFunctionExpression && MaterialExpression && IsAllowedExpressionType(MaterialExpression->GetClass(), MaterialGraph->MaterialFunction != NULL))
 			{
 				return true;
 			}
