Failed to submit file with 2556671 - Added FCompactPose

diff --git a/Engine/Source/Editor/UnrealEd/Private/Animation/DebugSkelMeshComponent.cpp b/Engine/Source/Editor/UnrealEd/Private/Animation/DebugSkelMeshComponent.cpp
--- a/Engine/Source/Editor/UnrealEd/Private/Animation/DebugSkelMeshComponent.cpp
+++ b/Engine/Source/Editor/UnrealEd/Private/Animation/DebugSkelMeshComponent.cpp
@@ -485,14 +485,19 @@ void UDebugSkelMeshComponent::RefreshBoneTransforms(FActorComponentTickFunction*
 		{ 
 			if (Sequence->IsValidAdditive()) 
 			{ 
-				AdditiveBasePoses.AddUninitialized(PreviewInstance->RequiredBones.GetNumBones());
-				Sequence->GetAdditiveBasePose(AdditiveBasePoses, PreviewInstance->RequiredBones, FAnimExtractContext(PreviewInstance->CurrentTime));
-				
-				FA2CSPose CSPose;
-				CSPose.AllocateLocalPoses(AnimScriptInstance->RequiredBones, AdditiveBasePoses);
-				for(int32 i=0; i<AdditiveBasePoses.Num(); ++i)
+				FCSPose<FCompactPose> CSAdditiveBasePose;
 				{
-					AdditiveBasePoses[i] = CSPose.GetComponentSpaceTransform(i);
+					FCompactPose AdditiveBasePose;
+					AdditiveBasePose.SetBoneContainer(&PreviewInstance->RequiredBones);
+
+					Sequence->GetAdditiveBasePose(AdditiveBasePose, FAnimExtractContext(PreviewInstance->CurrentTime));
+					CSAdditiveBasePose.InitPose(AdditiveBasePose);
+				}
+
+				for (int32 i = 0; i < AdditiveBasePoses.Num(); ++i)
+				{
+					FCompactPoseBoneIndex CompactIndex = PreviewInstance->RequiredBones.MakeCompactPoseIndex(FMeshPoseBoneIndex(i));
+					AdditiveBasePoses[i] = CSAdditiveBasePose.GetComponentSpaceTransform(CompactIndex);
 				}
 			}
 		}
