Fix for loading the existing cooked asset registry when doing iterative cooking

diff --git a/Engine/Source/Editor/UnrealEd/Private/Commandlets/ChunkManifestGenerator.cpp b/Engine/Source/Editor/UnrealEd/Private/Commandlets/ChunkManifestGenerator.cpp
--- a/Engine/Source/Editor/UnrealEd/Private/Commandlets/ChunkManifestGenerator.cpp
+++ b/Engine/Source/Editor/UnrealEd/Private/Commandlets/ChunkManifestGenerator.cpp
@@ -432,7 +432,7 @@ bool FChunkManifestGenerator::LoadAssetRegistry(const FString& SandboxPath, cons
 		{
 			AssetRegistry.LoadRegistryData(*AssetRegistryReader, SavedAssetRegistryData);
 		}
-		for (FAssetData LoadedAssetData : AssetRegistryData)
+		for (auto& LoadedAssetData : AssetRegistryData)
 		{
 			
 			FName ShortPackageName = FName(*FPaths::GetBaseFilename(LoadedAssetData.PackageName.ToString()));
@@ -454,8 +454,13 @@ bool FChunkManifestGenerator::LoadAssetRegistry(const FString& SandboxPath, cons
 
 		for (const auto& SavedAsset : SavedAssetRegistryData)
 		{
-			if (PackagesToKeep && PackagesToKeep->Contains(SavedAsset.Value->PackageName))
-			AssetRegistryData.Add(*SavedAsset.Value);
+			FName ShortPackageName = FName(*FPaths::GetBaseFilename(SavedAsset.Value->PackageName.ToString()));
+
+			if (PackagesToKeep && PackagesToKeep->Contains(ShortPackageName))
+			{ 
+				AssetRegistryData.Add(*SavedAsset.Value);
+			}
+			
 			delete SavedAsset.Value;
 		}
 		SavedAssetRegistryData.Empty();
