Fixed crash when gathering text from source files due to not clearing the match progress of other matches when a match is accepted.

diff --git a/Engine/Source/Editor/UnrealEd/Private/Commandlets/GatherTextFromSourceCommandlet.cpp b/Engine/Source/Editor/UnrealEd/Private/Commandlets/GatherTextFromSourceCommandlet.cpp
--- a/Engine/Source/Editor/UnrealEd/Private/Commandlets/GatherTextFromSourceCommandlet.cpp
+++ b/Engine/Source/Editor/UnrealEd/Private/Commandlets/GatherTextFromSourceCommandlet.cpp
@@ -527,7 +527,10 @@ bool UGatherTextFromSourceCommandlet::ParseSourceText(const FString& Text, const
 				{
 					// Do the parse now
 					Parsables[PendingParseIdx]->TryParse(FString(ParsePoint), ParseCtxt);
-					ParsableMatchCounters[PendingParseIdx] = 0;
+					for(auto& Element : ParsableMatchCounters)
+					{
+						Element = 0;
+					}
 					PendingParseIdx = -1;
 					ParsePoint = NULL;
 				}
