Exported PO files are now sorted by namespace, then key, then source string (msgctxt, #., msgid).

diff --git a/Engine/Source/Editor/UnrealEd/Classes/Commandlets/InternationalizationExportCommandlet.h b/Engine/Source/Editor/UnrealEd/Classes/Commandlets/InternationalizationExportCommandlet.h
--- a/Engine/Source/Editor/UnrealEd/Classes/Commandlets/InternationalizationExportCommandlet.h
+++ b/Engine/Source/Editor/UnrealEd/Classes/Commandlets/InternationalizationExportCommandlet.h
@@ -362,6 +362,8 @@ public:
 		return Entries.CreateConstIterator();
 	}
 
+	void SortEntries();
+
 private:
 	/** Copying is not supported */
 	FPortableObjectFormatDOM( const FPortableObjectFormatDOM& ); 
diff --git a/Engine/Source/Editor/UnrealEd/Private/Commandlets/InternationalizationExportCommandlet.cpp b/Engine/Source/Editor/UnrealEd/Private/Commandlets/InternationalizationExportCommandlet.cpp
--- a/Engine/Source/Editor/UnrealEd/Private/Commandlets/InternationalizationExportCommandlet.cpp
+++ b/Engine/Source/Editor/UnrealEd/Private/Commandlets/InternationalizationExportCommandlet.cpp
@@ -926,6 +926,56 @@ TSharedPtr<FPortableObjectEntry> FPortableObjectFormatDOM::FindEntry( const FStr
 	return FindEntry( TempEntry );
 }
 
+void FPortableObjectFormatDOM::SortEntries()
+{
+	// Sort keys.
+	for (const TSharedPtr<FPortableObjectEntry>& Entry : Entries)
+	{
+		Entry->ExtractedComments.Sort();
+	}
+
+	// Sort by namespace, then keys, then source text.
+	const auto& SortingPredicate = [](const TSharedPtr<FPortableObjectEntry>& A, const TSharedPtr<FPortableObjectEntry>& B) -> bool
+	{
+		// Compare namespace
+		if (A->MsgCtxt < B->MsgCtxt)
+		{
+			return true;
+		}
+		else if (A->MsgCtxt > B->MsgCtxt)
+		{
+			return false;
+		}
+
+		// Compare keys
+		const int32 ExtractedCommentCount = FMath::Max(A->ExtractedComments.Num(), B->ExtractedComments.Num());
+		for (int32 i = 0; i < ExtractedCommentCount; ++i)
+		{
+			if (A->ExtractedComments[i] < B->ExtractedComments[i] || (!A->ExtractedComments.IsValidIndex(i) && B->ExtractedComments.IsValidIndex(i)))
+			{
+				return true;
+			}
+			else if (A->ExtractedComments[i] > B->ExtractedComments[i] || (A->ExtractedComments.IsValidIndex(i) && !B->ExtractedComments.IsValidIndex(i)))
+			{
+				return false;
+			}
+		}
+
+		// Compare source string
+		if (A->MsgId < B->MsgId)
+		{
+			return true;
+		}
+		else if (A->MsgId > B->MsgId)
+		{
+			return false;
+		}
+
+		return A.Get() < B.Get();
+	};
+	Entries.Sort(SortingPredicate);
+}
+
 void FPortableObjectEntry::AddExtractedComment( const FString& InComment )
 {
 	if(!InComment.IsEmpty())
@@ -1152,6 +1202,7 @@ bool UInternationalizationExportCommandlet::DoExport( const FString& SourcePath,
 
 				// Write out the Portable Object to .po file.
 				{
+					PortableObj.SortEntries();
 					FString OutputString = PortableObj.ToString();
 					FString OutputFileName = "";
 					if (bUseCultureDirectory)
