Fix crash in show uncompressed bones

diff --git a/Engine/Source/Editor/UnrealEd/Private/Animation/DebugSkelMeshComponent.cpp b/Engine/Source/Editor/UnrealEd/Private/Animation/DebugSkelMeshComponent.cpp
--- a/Engine/Source/Editor/UnrealEd/Private/Animation/DebugSkelMeshComponent.cpp
+++ b/Engine/Source/Editor/UnrealEd/Private/Animation/DebugSkelMeshComponent.cpp
@@ -413,9 +413,10 @@ void UDebugSkelMeshComponent::SetShowBoneWeight(bool bNewShowBoneWeight)
 void UDebugSkelMeshComponent::GenSpaceBases(TArray<FTransform>& OutSpaceBases)
 {
 	TArray<FTransform> TempLocalAtoms;
+	TempLocalAtoms.AddUninitialized(OutSpaceBases.Num());
 	TArray<FActiveVertexAnim> TempVertexAnims;
 	FVector TempRootBoneTranslation;
-	PerformAnimationEvaluation(SkeletalMesh, AnimScriptInstance, OutSpaceBases, CachedLocalAtoms, TempVertexAnims, TempRootBoneTranslation);
+	PerformAnimationEvaluation(SkeletalMesh, AnimScriptInstance, OutSpaceBases, TempLocalAtoms, TempVertexAnims, TempRootBoneTranslation);
 }
 
 void UDebugSkelMeshComponent::RefreshBoneTransforms(FActorComponentTickFunction* TickFunction)
