Added collection information to the asset and asset editor tooltips

UETOOL-332 - Collections 2.0
UETOOL-367 - Show which collections are on an asset

diff --git a/Engine/Source/Developer/CollectionManager/Private/CollectionManager.cpp b/Engine/Source/Developer/CollectionManager/Private/CollectionManager.cpp
--- a/Engine/Source/Developer/CollectionManager/Private/CollectionManager.cpp
+++ b/Engine/Source/Developer/CollectionManager/Private/CollectionManager.cpp
@@ -156,6 +156,29 @@ void FCollectionManager::GetCollectionsContainingObject(FName ObjectPath, EColle
 	}
 }
 
+FString FCollectionManager::GetCollectionsStringForObject(FName ObjectPath, ECollectionShareType::Type ShareType) const
+{
+	TArray<FString> CollectionNameStrings;
+
+	for (const auto& CachedCollection : CachedCollections)
+	{
+		const FCollectionNameType& CollectionKey = CachedCollection.Key;
+		const TSharedRef<FCollection>& Collection = CachedCollection.Value;
+		if ((ShareType == ECollectionShareType::CST_All || ShareType == CollectionKey.Type) && Collection->IsObjectInCollection(ObjectPath))
+		{
+			CollectionNameStrings.Add(CollectionKey.Name.ToString());
+		}
+	}
+	
+	if (CollectionNameStrings.Num() > 0)
+	{
+		CollectionNameStrings.Sort();
+		return FString::Join(CollectionNameStrings, TEXT(", "));
+	}
+
+	return FString();
+}
+
 void FCollectionManager::CreateUniqueCollectionName(const FName& BaseName, ECollectionShareType::Type ShareType, FName& OutCollectionName) const
 {
 	if (!ensure(ShareType != ECollectionShareType::CST_All))
diff --git a/Engine/Source/Developer/CollectionManager/Private/CollectionManager.h b/Engine/Source/Developer/CollectionManager/Private/CollectionManager.h
--- a/Engine/Source/Developer/CollectionManager/Private/CollectionManager.h
+++ b/Engine/Source/Developer/CollectionManager/Private/CollectionManager.h
@@ -15,6 +15,7 @@ public:
 	virtual bool GetObjectsInCollection(FName CollectionName, ECollectionShareType::Type ShareType, TArray<FName>& ObjectPaths) const override;
 	virtual bool GetClassesInCollection(FName CollectionName, ECollectionShareType::Type ShareType, TArray<FName>& ClassPaths) const override;
 	virtual void GetCollectionsContainingObject(FName ObjectPath, ECollectionShareType::Type ShareType, TArray<FName>& OutCollectionNames) const override;
+	virtual FString GetCollectionsStringForObject(FName ObjectPath, ECollectionShareType::Type ShareType) const override;
 	virtual void CreateUniqueCollectionName(const FName& BaseName, ECollectionShareType::Type ShareType, FName& OutCollectionName) const override;
 	virtual bool CreateCollection(FName CollectionName, ECollectionShareType::Type ShareType) override;
 	virtual bool RenameCollection(FName CurrentCollectionName, ECollectionShareType::Type CurrentShareType, FName NewCollectionName, ECollectionShareType::Type NewShareType) override;
diff --git a/Engine/Source/Developer/CollectionManager/Public/ICollectionManager.h b/Engine/Source/Developer/CollectionManager/Public/ICollectionManager.h
--- a/Engine/Source/Developer/CollectionManager/Public/ICollectionManager.h
+++ b/Engine/Source/Developer/CollectionManager/Public/ICollectionManager.h
@@ -28,6 +28,9 @@ public:
 	/** Returns a list of collections in which the specified object exists of the specified share type */
 	virtual void GetCollectionsContainingObject(FName ObjectPath, ECollectionShareType::Type ShareType, TArray<FName>& OutCollectionNames) const = 0;
 
+	/** Returns a string containing a comma separated list of collections in which the specified object exists of the specified share type */
+	virtual FString GetCollectionsStringForObject(FName ObjectPath, ECollectionShareType::Type ShareType) const = 0;
+
 	/** Creates a unique collection name for the given type taking the form BaseName+(unique number) */
 	virtual void CreateUniqueCollectionName(const FName& BaseName, ECollectionShareType::Type ShareType, FName& OutCollectionName) const = 0;
 
diff --git a/Engine/Source/Editor/BehaviorTreeEditor/Private/BehaviorTreeEditor.cpp b/Engine/Source/Editor/BehaviorTreeEditor/Private/BehaviorTreeEditor.cpp
--- a/Engine/Source/Editor/BehaviorTreeEditor/Private/BehaviorTreeEditor.cpp
+++ b/Engine/Source/Editor/BehaviorTreeEditor/Private/BehaviorTreeEditor.cpp
@@ -1575,7 +1575,18 @@ FText FBehaviorTreeEditor::GetToolkitName() const
 	const UObject* EditingObject = GetCurrentMode() == BehaviorTreeMode ? (UObject*)BehaviorTree : (UObject*)GetBlackboardData();
 	if(EditingObject != nullptr)
 	{
-		return FAssetEditorToolkit::GetDescriptionForObject(EditingObject);
+		return FAssetEditorToolkit::GetLabelForObject(EditingObject);
+	}
+
+	return FText();
+}
+
+FText FBehaviorTreeEditor::GetToolkitToolTipText() const
+{
+	const UObject* EditingObject = GetCurrentMode() == BehaviorTreeMode ? (UObject*)BehaviorTree : (UObject*)GetBlackboardData();
+	if(EditingObject != nullptr)
+	{
+		return FAssetEditorToolkit::GetToolTipTextForObject(EditingObject);
 	}
 
 	return FText();
diff --git a/Engine/Source/Editor/BehaviorTreeEditor/Private/BehaviorTreeEditor.h b/Engine/Source/Editor/BehaviorTreeEditor/Private/BehaviorTreeEditor.h
--- a/Engine/Source/Editor/BehaviorTreeEditor/Private/BehaviorTreeEditor.h
+++ b/Engine/Source/Editor/BehaviorTreeEditor/Private/BehaviorTreeEditor.h
@@ -27,6 +27,7 @@ public:
 	virtual FString GetWorldCentricTabPrefix() const override;
 	virtual FLinearColor GetWorldCentricTabColorScale() const override;
 	virtual FText GetToolkitName() const override;
+	virtual FText GetToolkitToolTipText() const override;
 	// End IToolkit interface
 
 	// Begin IBehaviorTreeEditor interface
diff --git a/Engine/Source/Editor/ContentBrowser/Private/AssetViewWidgets.cpp b/Engine/Source/Editor/ContentBrowser/Private/AssetViewWidgets.cpp
--- a/Engine/Source/Editor/ContentBrowser/Private/AssetViewWidgets.cpp
+++ b/Engine/Source/Editor/ContentBrowser/Private/AssetViewWidgets.cpp
@@ -7,6 +7,7 @@
 #include "AssetViewTypes.h"
 #include "AssetViewWidgets.h"
 #include "SThumbnailEditModeTools.h"
+#include "CollectionManagerModule.h"
 #include "DragAndDrop/AssetDragDropOp.h"
 #include "DragAndDrop/AssetPathDragDropOp.h"
 #include "DragDropHandler.h"
@@ -392,6 +393,18 @@ TSharedRef<SToolTip> SAssetViewItem::CreateToolTipWidget() const
 
 			// Add Path
 			AddToToolTipInfoBox( InfoBox, LOCTEXT("TileViewTooltipPath", "Path"), FText::FromName(AssetData.PackagePath), false );
+			
+			// Add Collections
+			{
+				static const FName CollectionManagerModuleName = "CollectionManager";
+				FCollectionManagerModule& CollectionManagerModule = FModuleManager::LoadModuleChecked<FCollectionManagerModule>(CollectionManagerModuleName);
+
+				const FString CollectionNames = CollectionManagerModule.Get().GetCollectionsStringForObject(AssetData.ObjectPath, ECollectionShareType::CST_All);
+				if (!CollectionNames.IsEmpty())
+				{
+					AddToToolTipInfoBox(InfoBox, LOCTEXT("AssetToolTipKey_Collections", "Collections"), FText::FromString(CollectionNames), false);
+				}
+			}
 
 			// If we are using a loaded class, find all the hidden tags so we don't display them
 			TMap<FName, UObject::FAssetRegistryTagMetadata> MetadataMap;
diff --git a/Engine/Source/Editor/UnrealEd/Private/Toolkits/AssetEditorToolkit.cpp b/Engine/Source/Editor/UnrealEd/Private/Toolkits/AssetEditorToolkit.cpp
--- a/Engine/Source/Editor/UnrealEd/Private/Toolkits/AssetEditorToolkit.cpp
+++ b/Engine/Source/Editor/UnrealEd/Private/Toolkits/AssetEditorToolkit.cpp
@@ -9,6 +9,7 @@
 #include "Editor/MainFrame/Public/MainFrame.h"
 #include "Editor/WorkspaceMenuStructure/Public/WorkspaceMenuStructureModule.h"
 #include "ClassIconFinder.h"
+#include "CollectionManagerModule.h"
 #include "IUserFeedbackModule.h"
 #include "IDocumentation.h"
 #include "ReferenceViewer.h"
@@ -76,6 +77,7 @@ void FAssetEditorToolkit::InitAssetEditor( const EToolkitMode::Type Mode, const
 
 		// Create the label and the link for the toolkit documentation.
 		TAttribute<FText> Label = TAttribute<FText>( this, &FAssetEditorToolkit::GetToolkitName );
+		TAttribute<FText> ToolTipText = TAttribute<FText>( this, &FAssetEditorToolkit::GetToolkitToolTipText );
 		FString DocLink = GetDocumentationLink();
 		if ( !DocLink.StartsWith( "Shared/" ) )
 		{
@@ -86,7 +88,7 @@ void FAssetEditorToolkit::InitAssetEditor( const EToolkitMode::Type Mode, const
 		NewMajorTab = SNew(SDockTab) 
 			.ContentPadding(0.0f) 
 			.TabRole(ETabRole::MajorTab)
-			.ToolTip( IDocumentation::Get()->CreateToolTip( Label, nullptr, DocLink, GetToolkitFName().ToString() ) )
+			.ToolTip( IDocumentation::Get()->CreateToolTip( ToolTipText, nullptr, DocLink, GetToolkitFName().ToString() ) ) 
 			.Icon( this, &FAssetEditorToolkit::GetDefaultTabIcon )
 			.Label( Label );
 
@@ -275,10 +277,19 @@ FText FAssetEditorToolkit::GetToolkitName() const
 
 	check (EditingObject != NULL);
 
-	return GetDescriptionForObject(EditingObject);
+	return GetLabelForObject(EditingObject);
 }
 
-FText FAssetEditorToolkit::GetDescriptionForObject(const UObject* InObject)
+FText FAssetEditorToolkit::GetToolkitToolTipText() const
+{
+	const UObject* EditingObject = GetEditingObject();
+
+	check (EditingObject != NULL);
+
+	return GetToolTipTextForObject(EditingObject);
+}
+
+FText FAssetEditorToolkit::GetLabelForObject(const UObject* InObject)
 {
 	const bool bDirtyState = InObject->GetOutermost()->IsDirty();
 	FString NameString;
@@ -297,6 +308,38 @@ FText FAssetEditorToolkit::GetDescriptionForObject(const UObject* InObject)
 	return FText::Format( LOCTEXT("AssetEditorAppLabel", "{ObjectName}{DirtyState}"), Args );
 }
 
+FText FAssetEditorToolkit::GetToolTipTextForObject(const UObject* InObject)
+{
+	FString ToolTipString;
+	if(const AActor* ObjectAsActor = Cast<AActor>(InObject))
+	{
+		ToolTipString += LOCTEXT("ToolTipActorLabel", "Actor").ToString();
+		ToolTipString += TEXT(": ");
+		ToolTipString += ObjectAsActor->GetActorLabel();
+	}
+	else
+	{
+		ToolTipString += LOCTEXT("ToolTipAssetLabel", "Asset").ToString();
+		ToolTipString += TEXT(": ");
+		ToolTipString += InObject->GetName();
+
+		static const FName CollectionManagerModuleName = "CollectionManager";
+		FCollectionManagerModule& CollectionManagerModule = FModuleManager::LoadModuleChecked<FCollectionManagerModule>(CollectionManagerModuleName);
+
+		const FString CollectionNames = CollectionManagerModule.Get().GetCollectionsStringForObject(*InObject->GetPathName(), ECollectionShareType::CST_All);
+		if (!CollectionNames.IsEmpty())
+		{
+			ToolTipString += TEXT("\n");
+
+			ToolTipString += LOCTEXT("ToolTipCollectionsLabel", "Collections").ToString();
+			ToolTipString += TEXT(": ");
+			ToolTipString += CollectionNames;
+		}
+	}
+
+	return FText::FromString(ToolTipString);
+}
+
 class FEdMode* FAssetEditorToolkit::GetEditorMode() const
 {
 	return NULL;
diff --git a/Engine/Source/Editor/UnrealEd/Public/Toolkits/AssetEditorToolkit.h b/Engine/Source/Editor/UnrealEd/Public/Toolkits/AssetEditorToolkit.h
--- a/Engine/Source/Editor/UnrealEd/Public/Toolkits/AssetEditorToolkit.h
+++ b/Engine/Source/Editor/UnrealEd/Public/Toolkits/AssetEditorToolkit.h
@@ -58,7 +58,8 @@ public:
 	virtual const TArray< UObject* >* GetObjectsCurrentlyBeingEdited() const override;
 	virtual FName GetToolkitFName() const override = 0;				// Must implement in derived class!
 	virtual FText GetBaseToolkitName() const override = 0;			// Must implement in derived class!
-	virtual FText GetToolkitName() const override;		
+	virtual FText GetToolkitName() const override;
+	virtual FText GetToolkitToolTipText() const override;
 	virtual FString GetWorldCentricTabPrefix() const override = 0;	// Must implement in derived class!
 	virtual class FEdMode* GetEditorMode() const override;
 
@@ -154,7 +155,14 @@ public:
 	 * @param	InObject	The object we want a description of
 	 * @return a formatted description of the object state (e.g. "MyObject*")
 	 */
-	static FText GetDescriptionForObject(const UObject* InObject);
+	static FText GetLabelForObject(const UObject* InObject);
+
+	/** 
+	 * Gets the text to display in a toolkit tooltip for an object 
+	 * @param	InObject	The object we want a description of
+	 * @return a formatted description of the object
+	 */
+	static FText GetToolTipTextForObject(const UObject* InObject);
 
 protected:
 
diff --git a/Engine/Source/Editor/UnrealEd/Public/Toolkits/BaseToolkit.h b/Engine/Source/Editor/UnrealEd/Public/Toolkits/BaseToolkit.h
--- a/Engine/Source/Editor/UnrealEd/Public/Toolkits/BaseToolkit.h
+++ b/Engine/Source/Editor/UnrealEd/Public/Toolkits/BaseToolkit.h
@@ -93,6 +93,7 @@ public:
 
 	/** IToolkit interface */
 	virtual FText GetToolkitName() const override { return GetBaseToolkitName(); }
+	virtual FText GetToolkitToolTipText() const override { return GetBaseToolkitName(); }
 	virtual FString GetWorldCentricTabPrefix() const override;
 	virtual bool IsAssetEditor() const override;
 	virtual const TArray< UObject* >* GetObjectsCurrentlyBeingEdited() const override;
diff --git a/Engine/Source/Editor/UnrealEd/Public/Toolkits/IToolkit.h b/Engine/Source/Editor/UnrealEd/Public/Toolkits/IToolkit.h
--- a/Engine/Source/Editor/UnrealEd/Public/Toolkits/IToolkit.h
+++ b/Engine/Source/Editor/UnrealEd/Public/Toolkits/IToolkit.h
@@ -72,6 +72,9 @@ public:
 	/** Returns the localized name of this toolkit */
 	virtual FText GetToolkitName() const = 0;
 
+	/** Returns the localized tooltip text of this toolkit */
+	virtual FText GetToolkitToolTipText() const = 0;
+
 	/** Returns the workspace menu category of this toolkit */
 	virtual TSharedRef<FWorkspaceItem> GetWorkspaceMenuCategory() const = 0;
 
