fixed normal map detection for 16bit PNG as reported here:
https://udn.unrealengine.com/questions/246349/16bits-png-import-in-content-browser.html

diff --git a/Engine/Source/Editor/UnrealEd/Private/NormalMapIdentification.cpp b/Engine/Source/Editor/UnrealEd/Private/NormalMapIdentification.cpp
--- a/Engine/Source/Editor/UnrealEd/Private/NormalMapIdentification.cpp
+++ b/Engine/Source/Editor/UnrealEd/Private/NormalMapIdentification.cpp
@@ -115,11 +115,11 @@ public:
 typedef SampleNormalMapPixel8<2, 1, 0, 3> SampleNormalMapPixelBGRA8;
 typedef SampleNormalMapPixel8<0, 1, 2, 3> SampleNormalMapPixelRGBA8;
 
-class SampleNormalMapPixel16 : public NormalMapSamplerBase
+class SampleNormalMapPixelRGBA16 : public NormalMapSamplerBase
 {
 public:
-	SampleNormalMapPixel16() {}
-	~SampleNormalMapPixel16() {}
+	SampleNormalMapPixelRGBA16() {}
+	~SampleNormalMapPixelRGBA16() {}
 
 	FLinearColor DoSampleColor( int32 X, int32 Y )
 	{
@@ -129,10 +129,10 @@ public:
 		const float OneOver65535 = 1.0f / 65535.0f;
 
 		// Need to verify component order
-		Result.R = (float)((uint16*)PixelToSample)[3] * OneOver65535;
-		Result.G = (float)((uint16*)PixelToSample)[2] * OneOver65535;
-		Result.B = (float)((uint16*)PixelToSample)[1] * OneOver65535;
-		Result.A = (float)((uint16*)PixelToSample)[0] * OneOver65535;
+		Result.R = (float)((uint16*)PixelToSample)[0] * OneOver65535;
+		Result.G = (float)((uint16*)PixelToSample)[1] * OneOver65535;
+		Result.B = (float)((uint16*)PixelToSample)[2] * OneOver65535;
+		Result.A = (float)((uint16*)PixelToSample)[3] * OneOver65535;
 
 		return Result;
 	}
@@ -370,7 +370,7 @@ static bool IsTextureANormalMap( UTexture* Texture )
 			break;
 		case TSF_RGBA16:
 			{
-				TNormalMapAnalyzer<SampleNormalMapPixel16> Analyzer;
+				TNormalMapAnalyzer<SampleNormalMapPixelRGBA16> Analyzer;
 				bIsNormalMap = Analyzer.DoesTextureLookLikelyToBeANormalMap( Texture );
 			}
 			break;
