Pulled the geometry test from the promotion pass, remade it in it's own spot.  Added a couple of bsp unit tests which we can expand upon.
Removed the BSP validation test as it's the same test.
BSP tests are now located in the GeometryTest.cpp file.

diff --git a/Engine/Source/Editor/UnrealEd/Private/Tests/AutomationEditorCommon.cpp b/Engine/Source/Editor/UnrealEd/Private/Tests/AutomationEditorCommon.cpp
--- a/Engine/Source/Editor/UnrealEd/Private/Tests/AutomationEditorCommon.cpp
+++ b/Engine/Source/Editor/UnrealEd/Private/Tests/AutomationEditorCommon.cpp
@@ -472,6 +472,24 @@ namespace AutomationEditorCommonUtils
 			}
 		}
 	}
+
+	bool SetOrthoViewportView(const FVector& ViewLocation, const FRotator& ViewRotation)
+	{
+		for (int32 i = 0; i < GEditor->LevelViewportClients.Num(); i++)
+		{
+			FLevelEditorViewportClient* ViewportClient = GEditor->LevelViewportClients[i];
+			if (!ViewportClient->IsOrtho())
+			{
+				ViewportClient->SetViewLocation(ViewLocation);
+				ViewportClient->SetViewRotation(ViewRotation);
+				return true;
+			}
+		}
+
+		UE_LOG(LogEditorAutomationTests, Log, TEXT("An ortho viewport was not found.  May affect the test results."));
+		return false;
+	}
+
 }
 
 ///////////////////////////////////////////////////////////////////////
@@ -644,7 +662,7 @@ bool FAddStaticMeshCommand::Update()
 
 	UE_LOG(LogEditorAutomationTests, Log, TEXT("Static Mesh cube has been added to 0, 0, 0."))
 
-		return true;
+	return true;
 }
 
 /**
diff --git a/Engine/Source/Editor/UnrealEd/Private/Tests/EditorAutomationTests.cpp b/Engine/Source/Editor/UnrealEd/Private/Tests/EditorAutomationTests.cpp
--- a/Engine/Source/Editor/UnrealEd/Private/Tests/EditorAutomationTests.cpp
+++ b/Engine/Source/Editor/UnrealEd/Private/Tests/EditorAutomationTests.cpp
@@ -1013,155 +1013,6 @@ bool FReinitializeRHIResources::RunTest(const FString& Parameters)
 }
 
 
-//////////////////////////////////////////////////////////////////////////
-/**
- * QA BSP Regression Testing
- */
-IMPLEMENT_SIMPLE_AUTOMATION_TEST(FBSPValidation, "System.QA.BSP Validation", EAutomationTestFlags::ATF_Editor);
-
-bool FBSPValidation::RunTest(const FString& Parameters)
-{
-	UWorld* World = AutomationEditorCommonUtils::CreateNewMap();
-	GEditor->Exec( World, TEXT("BRUSH Scale 1 1 1"));
-
-	for( int32 i = 0; i < GEditor->LevelViewportClients.Num(); i++ )
-	{
-		FLevelEditorViewportClient* ViewportClient = GEditor->LevelViewportClients[i];
-		if(!ViewportClient->IsOrtho() )
-		{
-			ViewportClient->SetViewLocation( FVector(176, 2625, 2075) );
-			ViewportClient->SetViewRotation( FRotator(319, 269, 1) );
-		}
-	}
-
-	//Cube Additive Brush
-	UCubeBuilder* CubeAdditiveBrushBuilder = Cast<UCubeBuilder>(GEditor->FindBrushBuilder( UCubeBuilder::StaticClass() ));
-	CubeAdditiveBrushBuilder->X = 4096.0f;
-	CubeAdditiveBrushBuilder->Y = 4096.0f;
-	CubeAdditiveBrushBuilder->Z = 128.0f;
-	CubeAdditiveBrushBuilder->Build(World);
-	GEditor->Exec( World, TEXT("BRUSH MOVETO X=0 Y=0 Z=0"));
-	GEditor->Exec( World, TEXT("BRUSH ADD"));
-
-	//Cone Additive Brush
-	UConeBuilder* ConeAdditiveBrushBuilder = Cast<UConeBuilder>(GEditor->FindBrushBuilder( UConeBuilder::StaticClass() ));
-	ConeAdditiveBrushBuilder->Z = 1024.0f;
-	ConeAdditiveBrushBuilder->CapZ = 256.0f;
-	ConeAdditiveBrushBuilder->OuterRadius = 512.0f;
-	ConeAdditiveBrushBuilder->InnerRadius = 384;
-	ConeAdditiveBrushBuilder->Sides = 32;
-	ConeAdditiveBrushBuilder->Build(World);
-	GEditor->Exec( World, TEXT("BRUSH MOVETO X=-1525 Y=-1777 Z=64"));
-	GEditor->Exec( World, TEXT("BRUSH ADD"));
-
-	//Sphere Additive Brush
-	UTetrahedronBuilder* TetraAdditiveBrushBuilder = Cast<UTetrahedronBuilder>(GEditor->FindBrushBuilder( UTetrahedronBuilder::StaticClass() ));
-	TetraAdditiveBrushBuilder->Radius = 512.0f;
-	TetraAdditiveBrushBuilder->SphereExtrapolation = 3;
-	TetraAdditiveBrushBuilder->Build(World);
-	GEditor->Exec( World, TEXT("BRUSH MOVETO X=-88 Y=-1777 Z=535"));
-	GEditor->Exec( World, TEXT("BRUSH ADD"));
-
-	//Cylinder Additive Brush
-	UCylinderBuilder* CylinderAdditiveBrushBuilder = Cast<UCylinderBuilder>(GEditor->FindBrushBuilder( UCylinderBuilder::StaticClass() ));
-	CylinderAdditiveBrushBuilder->Z = 1024.0f;
-	CylinderAdditiveBrushBuilder->OuterRadius = 512.0f;
-	CylinderAdditiveBrushBuilder->InnerRadius = 384.0f;
-	CylinderAdditiveBrushBuilder->Sides = 16;
-	CylinderAdditiveBrushBuilder->Build(World);
-	GEditor->Exec( World, TEXT("BRUSH MOVETO X=1338 Y=-1776 Z=535"));
-	GEditor->Exec( World, TEXT("BRUSH ADD"));
-
-	//Cylinder Additive Brush
-	USheetBuilder* SheetAdditiveBrushBuilder = Cast<USheetBuilder>(GEditor->FindBrushBuilder( USheetBuilder::StaticClass() ));
-	SheetAdditiveBrushBuilder->X = 512.0f;
-	SheetAdditiveBrushBuilder->Y = 512.0f;
-	SheetAdditiveBrushBuilder->XSegments = 1;
-	SheetAdditiveBrushBuilder->YSegments = 1;
-	SheetAdditiveBrushBuilder->Axis = AX_YAxis;
-	SheetAdditiveBrushBuilder->Build(World);
-	GEditor->Exec( World, TEXT("BRUSH MOVETO X=-760 Y=-346 Z=535"));
-	GEditor->Exec( World, TEXT("BRUSH ADD"));
-
-	//Volume Additive Brush
-	UVolumetricBuilder* VolumetricAdditiveBrushBuilder = Cast<UVolumetricBuilder>(GEditor->FindBrushBuilder( UVolumetricBuilder::StaticClass() ));
-	VolumetricAdditiveBrushBuilder->Z = 512.0f;
-	VolumetricAdditiveBrushBuilder->Radius = 128.0f;
-	VolumetricAdditiveBrushBuilder->NumSheets = 3;
-	VolumetricAdditiveBrushBuilder->Build(World);
-	GEditor->Exec( World, TEXT("BRUSH MOVETO X=445 Y=-345 Z=535"));
-	GEditor->Exec( World, TEXT("BRUSH ADD"));
-
-	//Linear Stair Additive Brush
-	ULinearStairBuilder* LinearStairAdditiveBrushBuilder = Cast<ULinearStairBuilder>(GEditor->FindBrushBuilder( ULinearStairBuilder::StaticClass() ));
-	LinearStairAdditiveBrushBuilder->StepLength = 64.0f;
-	LinearStairAdditiveBrushBuilder->StepHeight = 16.0f;
-	LinearStairAdditiveBrushBuilder->StepWidth = 256.0f;
-	LinearStairAdditiveBrushBuilder->NumSteps = 8;
-	LinearStairAdditiveBrushBuilder->AddToFirstStep = 0;
-	LinearStairAdditiveBrushBuilder->Build(World);
-	GEditor->Exec( World, TEXT("BRUSH MOVETO X=1464 Y=-345 Z=-61"));
-	GEditor->Exec( World, TEXT("BRUSH ADD"));
-
-	//Curved Stair Additive Brush
-	UCurvedStairBuilder* CurvedStairAdditiveBrushBuilder = Cast<UCurvedStairBuilder>(GEditor->FindBrushBuilder( UCurvedStairBuilder::StaticClass() ));
-	CurvedStairAdditiveBrushBuilder->InnerRadius = 240.0f;
-	CurvedStairAdditiveBrushBuilder->StepHeight = 16.0f;
-	CurvedStairAdditiveBrushBuilder->StepWidth = 256.0f;
-	CurvedStairAdditiveBrushBuilder->AngleOfCurve = 90.0f;
-	CurvedStairAdditiveBrushBuilder->NumSteps = 4;
-	CurvedStairAdditiveBrushBuilder->AddToFirstStep = 0;
-	CurvedStairAdditiveBrushBuilder->CounterClockwise = false;
-	CurvedStairAdditiveBrushBuilder->Build(World);
-	GEditor->Exec( World, TEXT("BRUSH MOVETO X=-1290 Y=263 Z=193"));
-	GEditor->Exec( World, TEXT("BRUSH ADD"));
-
-	//Spiral Stair Additive Brush
-	USpiralStairBuilder* SpiralStairAdditiveBrushBuilder = Cast<USpiralStairBuilder>(GEditor->FindBrushBuilder( USpiralStairBuilder::StaticClass() ));
-	SpiralStairAdditiveBrushBuilder->InnerRadius = 64;
-	SpiralStairAdditiveBrushBuilder->StepWidth = 256.0f;
-	SpiralStairAdditiveBrushBuilder->StepHeight = 16.0f;
-	SpiralStairAdditiveBrushBuilder->StepThickness = 32.0f;
-	SpiralStairAdditiveBrushBuilder->NumStepsPer360 = 8;
-	SpiralStairAdditiveBrushBuilder->NumSteps = 8;
-	SpiralStairAdditiveBrushBuilder->SlopedCeiling = true;
-	SpiralStairAdditiveBrushBuilder->Build(World);
-	GEditor->Exec( World, TEXT("BRUSH MOVETO X=850 Y=263 Z=193"));
-	GEditor->Exec( World, TEXT("BRUSH ADD"));
-
-	//Cylinder Additive Brush
-	UCylinderBuilder* CylinderSubtractiveBrushBuilder = Cast<UCylinderBuilder>(GEditor->FindBrushBuilder( UCylinderBuilder::StaticClass() ));
-	CylinderSubtractiveBrushBuilder->Z = 256;
-	CylinderSubtractiveBrushBuilder->OuterRadius = 512.0f;
-	CylinderSubtractiveBrushBuilder->InnerRadius = 384.0f;
-	CylinderSubtractiveBrushBuilder->Sides = 3;
-	CylinderSubtractiveBrushBuilder->Build(World);
-	GEditor->Exec( World, TEXT("BRUSH MOVETO X=0 Y=0 Z=128"));
-	GEditor->Exec( World, TEXT("BRUSH SUBTRACT"));
-
-	//Cube Subtractive Brush
-	UCubeBuilder* CubeSubtractiveBrushBuilder = Cast<UCubeBuilder>(GEditor->FindBrushBuilder( UCubeBuilder::StaticClass() ));
-	CubeSubtractiveBrushBuilder->X = 256.0f;
-	CubeSubtractiveBrushBuilder->Y = 1024.0f;
-	CubeSubtractiveBrushBuilder->Z = 256.0f;
-	CubeSubtractiveBrushBuilder->Build(World);
-	GEditor->Exec( World, TEXT("BRUSH MOVETO X=-88 Y=-1777 Z=535"));
-	GEditor->Exec( World, TEXT("BRUSH SUBTRACT"));
-
-	//Directional Light
-	const FTransform Transform(FVector(-611.0f, 242.0f, 805.0f));
-	ADirectionalLight* DirectionalLight = Cast<ADirectionalLight>(GEditor->AddActor(World->GetCurrentLevel(), ADirectionalLight::StaticClass(), Transform));
-	DirectionalLight->SetMobility(EComponentMobility::Movable);
-	DirectionalLight->SetActorRotation(FRotator(300, 250, -91));
-	DirectionalLight->SetBrightness(3.142f);
-	DirectionalLight->SetLightColor(FColor::White);
-
-	GLevelEditorModeTools().MapChangeNotify();
-	
-	return true;
-}
-
-
 //////////////////////////////////////////////////////////////////////////
 /**
  * QA Static Mesh Regression Testing
diff --git a/Engine/Source/Editor/UnrealEd/Private/Tests/EditorBuildPromotionTests.cpp b/Engine/Source/Editor/UnrealEd/Private/Tests/EditorBuildPromotionTests.cpp
--- a/Engine/Source/Editor/UnrealEd/Private/Tests/EditorBuildPromotionTests.cpp
+++ b/Engine/Source/Editor/UnrealEd/Private/Tests/EditorBuildPromotionTests.cpp
@@ -1120,49 +1120,6 @@ namespace EditorBuildPromotionTestUtils
 		return CreatedAsset;
 	}
 
-	/**
-	* Checks that the level geometry changed after adding a brush and checks that undo / redo works
-	*
-	* @param CurrentLevel - The current level
-	* @param VertsBefore - The number of verts before the brush was added
-	* @param VertsAfter - The number of verts after the brush was added
-	* @param BrushType - The type of brush to use in the log
-	* @param bAdditive - True if this brush was added in additive mode
-	*/
-	static void TestGeometryUndoRedo(ULevel* CurrentLevel, int32 VertsBefore, int32 VertsAfter, const FString& BrushType, bool bAdditive)
-	{
-		if (VertsAfter > VertsBefore)
-		{
-			UE_LOG(LogEditorBuildPromotionTests, Display, TEXT("Placed a %s in %s mode"), *BrushType, bAdditive ? TEXT("additive") : TEXT("subtractive"));
-		}
-		else
-		{
-			UE_LOG(LogEditorBuildPromotionTests, Error, TEXT("Failed to place a %s in %s mode"), *BrushType, bAdditive ? TEXT("additive") : TEXT("subtractive"));
-		}
-
-		//Undo
-		GEditor->UndoTransaction();
-		if (CurrentLevel->Model->Verts.Num() == VertsBefore)
-		{
-			UE_LOG(LogEditorBuildPromotionTests, Display, TEXT("Undo successful for %s"), *BrushType);
-		}
-		else
-		{
-			UE_LOG(LogEditorBuildPromotionTests, Error, TEXT("Undo failed for %s"), *BrushType);
-		}
-
-		//Redo
-		GEditor->RedoTransaction();
-		if (CurrentLevel->Model->Verts.Num() == VertsAfter)
-		{
-			UE_LOG(LogEditorBuildPromotionTests, Display, TEXT("Redo successful for %s"), *BrushType);
-		}
-		else
-		{
-			UE_LOG(LogEditorBuildPromotionTests, Error, TEXT("Redo failed for %s"), *BrushType);
-		}
-	}
-
 	/**
 	* Saves all assets in a given folder
 	*
@@ -1710,7 +1667,6 @@ namespace BuildPromotionTestHelper
 
 			// 2) Geometry
 			ADD_TEST_STAGE(Geometry_LevelCreationAndSetup,	TEXT("Level Creation and Setup"));
-			ADD_TEST_STAGE(Geometry_AddLevelGeometry,		TEXT("Level Geometry"));
 			ADD_TEST_STAGE(EndSection,						TEXT("Geometry Workflow"));
 
 			// 3) Lighting
@@ -1931,209 +1887,6 @@ namespace BuildPromotionTestHelper
 
 			return true;
 		}
-
-		/**
-		* Geometry Test Stage: Adding, Subtracting and Editing Geometry
-		*    Adds different world geometry to the map and adds a player start
-		*/
-		bool Geometry_AddLevelGeometry()
-		{
-			GEditor->Exec(CurrentWorld, TEXT("BRUSH Scale 1 1 1"));
-
-			for (int32 i = 0; i < GEditor->LevelViewportClients.Num(); i++)
-			{
-				FLevelEditorViewportClient* ViewportClient = GEditor->LevelViewportClients[i];
-				if (!ViewportClient->IsOrtho())
-				{
-					ViewportClient->SetViewLocation(FVector(176, 2625, 2075));
-					ViewportClient->SetViewRotation(FRotator(319, 269, 1));
-				}
-			}
-			ULevel* CurrentLevel = CurrentWorld->GetCurrentLevel();
-			int32 VertsBefore = CurrentLevel->Model->Verts.Num();
-
-			//Cube Additive Brush
-			UCubeBuilder* CubeAdditiveBrushBuilder = Cast<UCubeBuilder>(GEditor->FindBrushBuilder(UCubeBuilder::StaticClass()));
-			CubeAdditiveBrushBuilder->X = 8192.0f;
-			CubeAdditiveBrushBuilder->Y = 8192.0f;
-			CubeAdditiveBrushBuilder->Z = 128.0f;
-			CubeAdditiveBrushBuilder->Build(CurrentWorld);
-			GEditor->Exec(CurrentWorld, TEXT("BRUSH MOVETO X=0 Y=0 Z=0"));
-			GEditor->Exec(CurrentWorld, TEXT("BRUSH ADD"));
-			int32 VertsAfter = CurrentLevel->Model->Verts.Num();
-
-			//Log results and test Undo/Redo
-			EditorBuildPromotionTestUtils::TestGeometryUndoRedo(CurrentLevel, VertsBefore, VertsAfter, TEXT("cube"), true);
-
-			//Cone Additive Brush
-			VertsBefore = CurrentLevel->Model->Verts.Num();
-			UConeBuilder* ConeAdditiveBrushBuilder = Cast<UConeBuilder>(GEditor->FindBrushBuilder(UConeBuilder::StaticClass()));
-			ConeAdditiveBrushBuilder->Z = 1024.0f;
-			ConeAdditiveBrushBuilder->CapZ = 256.0f;
-			ConeAdditiveBrushBuilder->OuterRadius = 512.0f;
-			ConeAdditiveBrushBuilder->InnerRadius = 384;
-			ConeAdditiveBrushBuilder->Sides = 32;
-			ConeAdditiveBrushBuilder->Build(CurrentWorld);
-			GEditor->Exec(CurrentWorld, TEXT("BRUSH MOVETO X=-1525 Y=-1777 Z=64"));
-			GEditor->Exec(CurrentWorld, TEXT("BRUSH ADD"));
-			VertsAfter = CurrentLevel->Model->Verts.Num();
-
-			//Log results and test Undo/Redo
-			EditorBuildPromotionTestUtils::TestGeometryUndoRedo(CurrentLevel, VertsBefore, VertsAfter, TEXT("cone"), true);
-
-			//Sphere Additive Brush
-			VertsBefore = CurrentLevel->Model->Verts.Num();
-			UTetrahedronBuilder* TetraAdditiveBrushBuilder = Cast<UTetrahedronBuilder>(GEditor->FindBrushBuilder(UTetrahedronBuilder::StaticClass()));
-			TetraAdditiveBrushBuilder->Radius = 512.0f;
-			TetraAdditiveBrushBuilder->SphereExtrapolation = 3;
-			TetraAdditiveBrushBuilder->Build(CurrentWorld);
-			GEditor->Exec(CurrentWorld, TEXT("BRUSH MOVETO X=-88 Y=-1777 Z=535"));
-			GEditor->Exec(CurrentWorld, TEXT("BRUSH ADD"));
-			VertsAfter = CurrentLevel->Model->Verts.Num();
-
-			//Log results and test Undo/Redo
-			EditorBuildPromotionTestUtils::TestGeometryUndoRedo(CurrentLevel, VertsBefore, VertsAfter, TEXT("tetrahedron"), true);
-
-			//Cylinder Additive Brush
-			VertsBefore = CurrentLevel->Model->Verts.Num();
-			UCylinderBuilder* CylinderAdditiveBrushBuilder = Cast<UCylinderBuilder>(GEditor->FindBrushBuilder(UCylinderBuilder::StaticClass()));
-			CylinderAdditiveBrushBuilder->Z = 1024.0f;
-			CylinderAdditiveBrushBuilder->OuterRadius = 512.0f;
-			CylinderAdditiveBrushBuilder->InnerRadius = 384.0f;
-			CylinderAdditiveBrushBuilder->Sides = 16;
-			CylinderAdditiveBrushBuilder->Build(CurrentWorld);
-			GEditor->Exec(CurrentWorld, TEXT("BRUSH MOVETO X=1338 Y=-1776 Z=535"));
-			GEditor->Exec(CurrentWorld, TEXT("BRUSH ADD"));
-			VertsAfter = CurrentLevel->Model->Verts.Num();
-
-			//Log results and test Undo/Redo
-			EditorBuildPromotionTestUtils::TestGeometryUndoRedo(CurrentLevel, VertsBefore, VertsAfter, TEXT("cylinder"), true);
-
-			//Cylinder Additive Brush
-			VertsBefore = CurrentLevel->Model->Verts.Num();
-			USheetBuilder* SheetAdditiveBrushBuilder = Cast<USheetBuilder>(GEditor->FindBrushBuilder(USheetBuilder::StaticClass()));
-			SheetAdditiveBrushBuilder->X = 512.0f;
-			SheetAdditiveBrushBuilder->Y = 512.0f;
-			SheetAdditiveBrushBuilder->XSegments = 1;
-			SheetAdditiveBrushBuilder->YSegments = 1;
-			SheetAdditiveBrushBuilder->Axis = AX_YAxis;
-			SheetAdditiveBrushBuilder->Build(CurrentWorld);
-			GEditor->Exec(CurrentWorld, TEXT("BRUSH MOVETO X=-760 Y=-346 Z=535"));
-			GEditor->Exec(CurrentWorld, TEXT("BRUSH ADD"));
-			VertsAfter = CurrentLevel->Model->Verts.Num();
-
-			//Log results and test Undo/Redo
-			EditorBuildPromotionTestUtils::TestGeometryUndoRedo(CurrentLevel, VertsBefore, VertsAfter, TEXT("sheet"), true);
-
-			//Volume Additive Brush
-			VertsBefore = CurrentLevel->Model->Verts.Num();
-			UVolumetricBuilder* VolumetricAdditiveBrushBuilder = Cast<UVolumetricBuilder>(GEditor->FindBrushBuilder(UVolumetricBuilder::StaticClass()));
-			VolumetricAdditiveBrushBuilder->Z = 512.0f;
-			VolumetricAdditiveBrushBuilder->Radius = 128.0f;
-			VolumetricAdditiveBrushBuilder->NumSheets = 3;
-			VolumetricAdditiveBrushBuilder->Build(CurrentWorld);
-			GEditor->Exec(CurrentWorld, TEXT("BRUSH MOVETO X=445 Y=-345 Z=535"));
-			GEditor->Exec(CurrentWorld, TEXT("BRUSH ADD"));
-			VertsAfter = CurrentLevel->Model->Verts.Num();
-
-			//Log results and test Undo/Redo
-			EditorBuildPromotionTestUtils::TestGeometryUndoRedo(CurrentLevel, VertsBefore, VertsAfter, TEXT("volume"), true);
-
-			//Linear Stair Additive Brush
-			VertsBefore = CurrentLevel->Model->Verts.Num();
-			ULinearStairBuilder* LinearStairAdditiveBrushBuilder = Cast<ULinearStairBuilder>(GEditor->FindBrushBuilder(ULinearStairBuilder::StaticClass()));
-			LinearStairAdditiveBrushBuilder->StepLength = 64.0f;
-			LinearStairAdditiveBrushBuilder->StepHeight = 16.0f;
-			LinearStairAdditiveBrushBuilder->StepWidth = 256.0f;
-			LinearStairAdditiveBrushBuilder->NumSteps = 8;
-			LinearStairAdditiveBrushBuilder->AddToFirstStep = 0;
-			LinearStairAdditiveBrushBuilder->Build(CurrentWorld);
-			GEditor->Exec(CurrentWorld, TEXT("BRUSH MOVETO X=1464 Y=-345 Z=-61"));
-			GEditor->Exec(CurrentWorld, TEXT("BRUSH ADD"));
-			VertsAfter = CurrentLevel->Model->Verts.Num();
-
-			//Log results and test Undo/Redo
-			EditorBuildPromotionTestUtils::TestGeometryUndoRedo(CurrentLevel, VertsBefore, VertsAfter, TEXT("linear stair"), true);
-
-			//Curved Stair Additive Brush
-			VertsBefore = CurrentLevel->Model->Verts.Num();
-			UCurvedStairBuilder* CurvedStairAdditiveBrushBuilder = Cast<UCurvedStairBuilder>(GEditor->FindBrushBuilder(UCurvedStairBuilder::StaticClass()));
-			CurvedStairAdditiveBrushBuilder->InnerRadius = 240.0f;
-			CurvedStairAdditiveBrushBuilder->StepHeight = 16.0f;
-			CurvedStairAdditiveBrushBuilder->StepWidth = 256.0f;
-			CurvedStairAdditiveBrushBuilder->AngleOfCurve = 90.0f;
-			CurvedStairAdditiveBrushBuilder->NumSteps = 4;
-			CurvedStairAdditiveBrushBuilder->AddToFirstStep = 0;
-			CurvedStairAdditiveBrushBuilder->CounterClockwise = false;
-			CurvedStairAdditiveBrushBuilder->Build(CurrentWorld);
-			GEditor->Exec(CurrentWorld, TEXT("BRUSH MOVETO X=-1290 Y=263 Z=193"));
-			GEditor->Exec(CurrentWorld, TEXT("BRUSH ADD"));
-			VertsAfter = CurrentLevel->Model->Verts.Num();
-
-			//Log results and test Undo/Redo
-			EditorBuildPromotionTestUtils::TestGeometryUndoRedo(CurrentLevel, VertsBefore, VertsAfter, TEXT("curved stair"), true);
-
-			//Spiral Stair Additive Brush
-			VertsBefore = CurrentLevel->Model->Verts.Num();
-			USpiralStairBuilder* SpiralStairAdditiveBrushBuilder = Cast<USpiralStairBuilder>(GEditor->FindBrushBuilder(USpiralStairBuilder::StaticClass()));
-			SpiralStairAdditiveBrushBuilder->InnerRadius = 64;
-			SpiralStairAdditiveBrushBuilder->StepWidth = 256.0f;
-			SpiralStairAdditiveBrushBuilder->StepHeight = 16.0f;
-			SpiralStairAdditiveBrushBuilder->StepThickness = 32.0f;
-			SpiralStairAdditiveBrushBuilder->NumStepsPer360 = 8;
-			SpiralStairAdditiveBrushBuilder->NumSteps = 8;
-			SpiralStairAdditiveBrushBuilder->SlopedCeiling = true;
-			SpiralStairAdditiveBrushBuilder->Build(CurrentWorld);
-			GEditor->Exec(CurrentWorld, TEXT("BRUSH MOVETO X=850 Y=263 Z=193"));
-			GEditor->Exec(CurrentWorld, TEXT("BRUSH ADD"));
-			VertsAfter = CurrentLevel->Model->Verts.Num();
-
-			//Log results and test Undo/Redo
-			EditorBuildPromotionTestUtils::TestGeometryUndoRedo(CurrentLevel, VertsBefore, VertsAfter, TEXT("spiral stair"), true);
-
-			//Cylinder Additive Brush
-			VertsBefore = CurrentLevel->Model->Verts.Num();
-			UCylinderBuilder* CylinderSubtractiveBrushBuilder = Cast<UCylinderBuilder>(GEditor->FindBrushBuilder(UCylinderBuilder::StaticClass()));
-			CylinderSubtractiveBrushBuilder->Z = 256;
-			CylinderSubtractiveBrushBuilder->OuterRadius = 512.0f;
-			CylinderSubtractiveBrushBuilder->InnerRadius = 384.0f;
-			CylinderSubtractiveBrushBuilder->Sides = 3;
-			CylinderSubtractiveBrushBuilder->Build(CurrentWorld);
-			GEditor->Exec(CurrentWorld, TEXT("BRUSH MOVETO X=0 Y=0 Z=128"));
-			GEditor->Exec(CurrentWorld, TEXT("BRUSH SUBTRACT"));
-			VertsAfter = CurrentLevel->Model->Verts.Num();
-
-			//Log results and test Undo/Redo
-			EditorBuildPromotionTestUtils::TestGeometryUndoRedo(CurrentLevel, VertsBefore, VertsAfter, TEXT("cylinder"), true);
-
-			//Cube Subtractive Brush
-			VertsBefore = CurrentLevel->Model->Verts.Num();
-			UCubeBuilder* CubeSubtractiveBrushBuilder = Cast<UCubeBuilder>(GEditor->FindBrushBuilder(UCubeBuilder::StaticClass()));
-			CubeSubtractiveBrushBuilder->X = 256.0f;
-			CubeSubtractiveBrushBuilder->Y = 1024.0f;
-			CubeSubtractiveBrushBuilder->Z = 256.0f;
-			CubeSubtractiveBrushBuilder->Build(CurrentWorld);
-			GEditor->Exec(CurrentWorld, TEXT("BRUSH MOVETO X=-88 Y=-1777 Z=535"));
-			GEditor->Exec(CurrentWorld, TEXT("BRUSH SUBTRACT"));
-			VertsAfter = CurrentLevel->Model->Verts.Num();
-
-			//Log results and test Undo/Redo
-			EditorBuildPromotionTestUtils::TestGeometryUndoRedo(CurrentLevel, VertsBefore, VertsAfter, TEXT("cube"), false);
-
-			//Add a playerstart
-			const FTransform Transform(FRotator(0, -90, 0), FVector(0.f, 1750.f, 166.f));
-			AActor* PlayerStart = GEditor->AddActor(CurrentWorld->GetCurrentLevel(), APlayerStart::StaticClass(), Transform);
-			if (PlayerStart)
-			{
-				UE_LOG(LogEditorBuildPromotionTests, Display, TEXT("Placed a player start"));
-			}
-			else
-			{
-				UE_LOG(LogEditorBuildPromotionTests, Error, TEXT("Failed to place a PlayerStart"));
-			}
-
-			return true;
-		}
 		
 		/**
 		* Lighting Test Stage: Place a point Light
diff --git a/Engine/Source/Editor/UnrealEd/Private/Tests/GeometryTests.cpp b/Engine/Source/Editor/UnrealEd/Private/Tests/GeometryTests.cpp
new file mode 100644
--- /dev/null
+++ b/Engine/Source/Editor/UnrealEd/Private/Tests/GeometryTests.cpp
@@ -0,0 +1,403 @@
+// Copyright 1998-2015 Epic Games, Inc. All Rights Reserved.
+
+#include "UnrealEd.h"
+
+//Automation
+#include "AutomationTest.h"
+#include "AutomationEditorCommon.h"
+
+#define CUBE_VERTS_COUNT			8
+#define WORLD_ORIGIN_VECTOR			FVector(0.f,0.f,0.f)
+#define WORLD_ORIGIN_VECTOR_OFFSET	FVector(0.f,0.f,10.f)
+
+DEFINE_LOG_CATEGORY_STATIC(LogGeometryTests, All, All);
+
+
+namespace GeometryPromotionHelpers
+{
+	void MiscPreparationsForBrushCreation(const ULevel* InCurrentLevel, int32& CurrentLevelVerts, const FString& InBrushType)
+	{
+		UE_LOG(LogGeometryTests, Log, TEXT("Adding an '%s' brush."), *InBrushType);
+		CurrentLevelVerts = InCurrentLevel->Model->Verts.Num();
+	}
+
+	/**
+	* Checks that the level geometry changed after adding a brush and checks that undo / redo works
+	*
+	* @param CurrentLevel - The current level
+	* @param VertsBefore - The number of verts before the brush was added
+	* @param VertsAfter - The number of verts after the brush was added
+	* @param BrushType - The type of brush to use in the log
+	* @param bAdditive - True if this brush was added in additive mode
+	*/
+	void TestGeometryUndoRedo(ULevel* CurrentLevel, int32 VertsBefore, int32 VertsAfter, const FString& BrushType, bool bAdditive)
+	{
+		if (VertsAfter > VertsBefore)
+		{
+			UE_LOG(LogEditorAutomationTests, Log, TEXT("Placed a %s in %s mode"), *BrushType, bAdditive ? TEXT("additive") : TEXT("subtractive"));
+		}
+		else
+		{
+			UE_LOG(LogEditorAutomationTests, Error, TEXT("Failed to place a %s in %s mode"), *BrushType, bAdditive ? TEXT("additive") : TEXT("subtractive"));
+		}
+
+		//Undo
+		GEditor->UndoTransaction();
+		if (CurrentLevel->Model->Verts.Num() == VertsBefore)
+		{
+			UE_LOG(LogEditorAutomationTests, Log, TEXT("Undo successful for %s"), *BrushType);
+		}
+		else
+		{
+			UE_LOG(LogEditorAutomationTests, Error, TEXT("Undo failed for %s"), *BrushType);
+		}
+
+		//Redo
+		GEditor->RedoTransaction();
+		if (CurrentLevel->Model->Verts.Num() == VertsAfter)
+		{
+			UE_LOG(LogEditorAutomationTests, Log, TEXT("Redo successful for %s"), *BrushType);
+		}
+		else
+		{
+			UE_LOG(LogEditorAutomationTests, Error, TEXT("Redo failed for %s"), *BrushType);
+		}
+	}
+}
+
+
+//////////////////////////////////////////////////////////////////////////
+//BSP Promotion Test
+//////////////////////////////////////////////////////////////////////////
+
+/**
+* BSP Promotion Test
+* Adds one of every brush type to the world.  Most of which will be set as additive brushes.
+* There will be a few subtractive brushes intersecting with other objects.
+* Tests for number of surfaces and vertices, numbers of brushes created, and brush location.
+*/
+IMPLEMENT_SIMPLE_AUTOMATION_TEST(FGeometryTest, "System.Promotion.Editor.Geometry Test", EAutomationTestFlags::ATF_Editor);
+
+bool FGeometryTest::RunTest(const FString& Parameters)
+{
+	//** SETUP **//
+	// Create the world and setup the level.
+	UWorld* World = AutomationEditorCommonUtils::CreateNewMap();
+	ULevel* CurrentLevel = World->GetCurrentLevel();
+	GLevelEditorModeTools().MapChangeNotify();
+	// Set the view location so that it can see the entire scene.
+	AutomationEditorCommonUtils::SetOrthoViewportView(FVector(176, 2625, 2075), FRotator(319, 269, 1));
+	// Grabbing the current number of verts in the world to compare later on.
+	int32 VertsBefore = CurrentLevel->Model->Verts.Num();
+	GEditor->Exec(World, TEXT("BRUSH Scale 1 1 1"));
+
+	//** TEST **//
+	//Cube Additive Brush
+	AddLogItem(TEXT("Additive Cube added to the world."));
+	UE_LOG(LogGeometryTests, Log, TEXT("Adding an 'Additive Cube' brush."));
+	UCubeBuilder* CubeAdditiveBrushBuilder = Cast<UCubeBuilder>(GEditor->FindBrushBuilder(UCubeBuilder::StaticClass()));
+	CubeAdditiveBrushBuilder->X = 4096.0f;
+	CubeAdditiveBrushBuilder->Y = 4096.0f;
+	CubeAdditiveBrushBuilder->Z = 128.0f;
+	CubeAdditiveBrushBuilder->Build(World);
+	GEditor->Exec(World, TEXT("BRUSH MOVETO X=0 Y=0 Z=0"));
+	GEditor->Exec(World, TEXT("BRUSH ADD"));
+	// Test that Undo/Redo works for an additive brush. 
+	int32 VertsAfter = CurrentLevel->Model->Verts.Num();
+	GeometryPromotionHelpers::TestGeometryUndoRedo(CurrentLevel, VertsBefore, VertsAfter, TEXT("Cube"), true);
+
+	//Cone Additive Brush
+	GeometryPromotionHelpers::MiscPreparationsForBrushCreation(CurrentLevel, VertsBefore, TEXT("Additive Cone"));
+	UConeBuilder* ConeAdditiveBrushBuilder = Cast<UConeBuilder>(GEditor->FindBrushBuilder(UConeBuilder::StaticClass()));
+	ConeAdditiveBrushBuilder->Z = 1024.0f;
+	ConeAdditiveBrushBuilder->CapZ = 256.0f;
+	ConeAdditiveBrushBuilder->OuterRadius = 512.0f;
+	ConeAdditiveBrushBuilder->InnerRadius = 384;
+	ConeAdditiveBrushBuilder->Sides = 32;
+	ConeAdditiveBrushBuilder->Build(World);
+	GEditor->Exec(World, TEXT("BRUSH MOVETO X=-1525 Y=-1777 Z=64"));
+	GEditor->Exec(World, TEXT("BRUSH ADD"));
+
+	//Sphere Additive Brush
+	GeometryPromotionHelpers::MiscPreparationsForBrushCreation(CurrentLevel, VertsBefore, TEXT("Additive Sphere"));
+	UTetrahedronBuilder* TetraAdditiveBrushBuilder = Cast<UTetrahedronBuilder>(GEditor->FindBrushBuilder(UTetrahedronBuilder::StaticClass()));
+	TetraAdditiveBrushBuilder->Radius = 512.0f;
+	TetraAdditiveBrushBuilder->SphereExtrapolation = 3;
+	TetraAdditiveBrushBuilder->Build(World);
+	GEditor->Exec(World, TEXT("BRUSH MOVETO X=-88 Y=-1777 Z=535"));
+	GEditor->Exec(World, TEXT("BRUSH ADD"));
+
+	//Cylinder Additive Brush
+	GeometryPromotionHelpers::MiscPreparationsForBrushCreation(CurrentLevel, VertsBefore, TEXT("Additive Cylinder"));
+	UCylinderBuilder* CylinderAdditiveBrushBuilder = Cast<UCylinderBuilder>(GEditor->FindBrushBuilder(UCylinderBuilder::StaticClass()));
+	CylinderAdditiveBrushBuilder->Z = 1024.0f;
+	CylinderAdditiveBrushBuilder->OuterRadius = 512.0f;
+	CylinderAdditiveBrushBuilder->InnerRadius = 384.0f;
+	CylinderAdditiveBrushBuilder->Sides = 16;
+	CylinderAdditiveBrushBuilder->Build(World);
+	GEditor->Exec(World, TEXT("BRUSH MOVETO X=1338 Y=-1776 Z=535"));
+	GEditor->Exec(World, TEXT("BRUSH ADD"));
+
+	//Sheet Additive Brush
+	GeometryPromotionHelpers::MiscPreparationsForBrushCreation(CurrentLevel, VertsBefore, TEXT("Additive Sheet"));
+	USheetBuilder* SheetAdditiveBrushBuilder = Cast<USheetBuilder>(GEditor->FindBrushBuilder(USheetBuilder::StaticClass()));
+	SheetAdditiveBrushBuilder->X = 512.0f;
+	SheetAdditiveBrushBuilder->Y = 512.0f;
+	SheetAdditiveBrushBuilder->XSegments = 1;
+	SheetAdditiveBrushBuilder->YSegments = 1;
+	SheetAdditiveBrushBuilder->Axis = AX_YAxis;
+	SheetAdditiveBrushBuilder->Build(World);
+	GEditor->Exec(World, TEXT("BRUSH MOVETO X=-760 Y=-346 Z=535"));
+	GEditor->Exec(World, TEXT("BRUSH ADD"));
+
+	//Volumetric Additive Brush
+	GeometryPromotionHelpers::MiscPreparationsForBrushCreation(CurrentLevel, VertsBefore, TEXT("Additive Volumetric"));
+	UVolumetricBuilder* VolumetricAdditiveBrushBuilder = Cast<UVolumetricBuilder>(GEditor->FindBrushBuilder(UVolumetricBuilder::StaticClass()));
+	VolumetricAdditiveBrushBuilder->Z = 512.0f;
+	VolumetricAdditiveBrushBuilder->Radius = 128.0f;
+	VolumetricAdditiveBrushBuilder->NumSheets = 3;
+	VolumetricAdditiveBrushBuilder->Build(World);
+	GEditor->Exec(World, TEXT("BRUSH MOVETO X=445 Y=-345 Z=535"));
+	GEditor->Exec(World, TEXT("BRUSH ADD"));
+
+	//Linear Stair Additive Brush
+	GeometryPromotionHelpers::MiscPreparationsForBrushCreation(CurrentLevel, VertsBefore, TEXT("Additive Linear Stair"));
+	ULinearStairBuilder* LinearStairAdditiveBrushBuilder = Cast<ULinearStairBuilder>(GEditor->FindBrushBuilder(ULinearStairBuilder::StaticClass()));
+	LinearStairAdditiveBrushBuilder->StepLength = 64.0f;
+	LinearStairAdditiveBrushBuilder->StepHeight = 16.0f;
+	LinearStairAdditiveBrushBuilder->StepWidth = 256.0f;
+	LinearStairAdditiveBrushBuilder->NumSteps = 8;
+	LinearStairAdditiveBrushBuilder->AddToFirstStep = 0;
+	LinearStairAdditiveBrushBuilder->Build(World);
+	GEditor->Exec(World, TEXT("BRUSH MOVETO X=1464 Y=-345 Z=-61"));
+	GEditor->Exec(World, TEXT("BRUSH ADD"));
+
+	//Curved Stair Additive Brush
+	GeometryPromotionHelpers::MiscPreparationsForBrushCreation(CurrentLevel, VertsBefore, TEXT("Additive Curved Stair"));
+	UCurvedStairBuilder* CurvedStairAdditiveBrushBuilder = Cast<UCurvedStairBuilder>(GEditor->FindBrushBuilder(UCurvedStairBuilder::StaticClass()));
+	CurvedStairAdditiveBrushBuilder->InnerRadius = 240.0f;
+	CurvedStairAdditiveBrushBuilder->StepHeight = 16.0f;
+	CurvedStairAdditiveBrushBuilder->StepWidth = 256.0f;
+	CurvedStairAdditiveBrushBuilder->AngleOfCurve = 90.0f;
+	CurvedStairAdditiveBrushBuilder->NumSteps = 4;
+	CurvedStairAdditiveBrushBuilder->AddToFirstStep = 0;
+	CurvedStairAdditiveBrushBuilder->CounterClockwise = false;
+	CurvedStairAdditiveBrushBuilder->Build(World);
+	GEditor->Exec(World, TEXT("BRUSH MOVETO X=-1290 Y=263 Z=193"));
+	GEditor->Exec(World, TEXT("BRUSH ADD"));
+
+	//Spiral Stair Additive Brush
+	GeometryPromotionHelpers::MiscPreparationsForBrushCreation(CurrentLevel, VertsBefore, TEXT("Additive Spiral Stair"));
+	USpiralStairBuilder* SpiralStairAdditiveBrushBuilder = Cast<USpiralStairBuilder>(GEditor->FindBrushBuilder(USpiralStairBuilder::StaticClass()));
+	SpiralStairAdditiveBrushBuilder->InnerRadius = 64;
+	SpiralStairAdditiveBrushBuilder->StepWidth = 256.0f;
+	SpiralStairAdditiveBrushBuilder->StepHeight = 16.0f;
+	SpiralStairAdditiveBrushBuilder->StepThickness = 32.0f;
+	SpiralStairAdditiveBrushBuilder->NumStepsPer360 = 8;
+	SpiralStairAdditiveBrushBuilder->NumSteps = 8;
+	SpiralStairAdditiveBrushBuilder->SlopedCeiling = true;
+	SpiralStairAdditiveBrushBuilder->Build(World);
+	GEditor->Exec(World, TEXT("BRUSH MOVETO X=850 Y=263 Z=193"));
+	GEditor->Exec(World, TEXT("BRUSH ADD"));
+
+	//Cylinder Subtractive Brush
+	GeometryPromotionHelpers::MiscPreparationsForBrushCreation(CurrentLevel, VertsBefore, TEXT("Subtractive Cylinder"));
+	UCylinderBuilder* CylinderSubtractiveBrushBuilder = Cast<UCylinderBuilder>(GEditor->FindBrushBuilder(UCylinderBuilder::StaticClass()));
+	CylinderSubtractiveBrushBuilder->Z = 256;
+	CylinderSubtractiveBrushBuilder->OuterRadius = 512.0f;
+	CylinderSubtractiveBrushBuilder->InnerRadius = 384.0f;
+	CylinderSubtractiveBrushBuilder->Sides = 3;
+	CylinderSubtractiveBrushBuilder->Build(World);
+	GEditor->Exec(World, TEXT("BRUSH MOVETO X=0 Y=0 Z=128"));
+	GEditor->Exec(World, TEXT("BRUSH SUBTRACT"));
+	// Test that Undo/Redo works for a subtractive brush.  
+	VertsAfter = CurrentLevel->Model->Verts.Num();
+	GeometryPromotionHelpers::TestGeometryUndoRedo(CurrentLevel, VertsBefore, VertsAfter, TEXT("Cylinder"), false);
+
+	//Cube Subtractive Brush
+	GeometryPromotionHelpers::MiscPreparationsForBrushCreation(CurrentLevel, VertsBefore, TEXT("Subtractive Cube"));
+	UCubeBuilder* CubeSubtractiveBrushBuilder = Cast<UCubeBuilder>(GEditor->FindBrushBuilder(UCubeBuilder::StaticClass()));
+	CubeSubtractiveBrushBuilder->X = 256.0f;
+	CubeSubtractiveBrushBuilder->Y = 1024.0f;
+	CubeSubtractiveBrushBuilder->Z = 256.0f;
+	CubeSubtractiveBrushBuilder->Build(World);
+	GEditor->Exec(World, TEXT("BRUSH MOVETO X=-88 Y=-1777 Z=535"));
+	GEditor->Exec(World, TEXT("BRUSH SUBTRACT"));
+
+
+	//** VERIFY **//
+	// Gather the objects needed to run the tests.
+	// Get the levels BSP Model.
+	UModel* BSPModel = CurrentLevel->Model;
+	// Add the brush actors to an array.
+	TArray<ABrush*> AdditiveBSP;
+	TArray<ABrush*> SubtractiveBSP;
+	for (TObjectIterator<ABrush> It; It; ++It)
+	{
+		ABrush* FoundBrush = *It;
+
+		if (FoundBrush->BrushType == EBrushType::Brush_Add)
+		{
+			AdditiveBSP.Add(FoundBrush);
+		}
+		else if (FoundBrush->BrushType == EBrushType::Brush_Subtract)
+		{
+			SubtractiveBSP.Add(FoundBrush);
+		}
+	}
+
+	// Verify the BSP brushes that were created have been found.
+	TestEqual<int32>(TEXT("The number of Additive BSP brushes created is not correct, expected 9."), AdditiveBSP.Num(), 9);
+	TestEqual<int32>(TEXT("The number of Additive BSP brushes created is not correct, expected 2."), SubtractiveBSP.Num(), 2);
+
+	// Verify there are the correct amount of BSP surfaces are visible in the level.
+	TestEqual<int32>(TEXT("Incorrect total number of surfaces are being reported, expected 10."), BSPModel->Surfs.Num(), 276);
+
+	UE_LOG(LogGeometryTests, Log, TEXT("BSP Placement: Number of Additive: %i, Number of Subtractive: %i"), AdditiveBSP.Num(), SubtractiveBSP.Num());
+
+	//** TEARDOWN **//
+	AdditiveBSP.Empty();
+	SubtractiveBSP.Empty();
+
+	return true;
+}
+
+
+//////////////////////////////////////////////////////////////////////////
+//BSP Unit Tests
+//////////////////////////////////////////////////////////////////////////
+
+/**
+* BSP Cube - New Brush tests
+* Adds two new BSP cubes to the world. An Additive Brush and a Subtractive Brush.
+* Subtractive brush intersects will the additive brush.
+* Tests for number of surfaces and vertices, numbers of brushes created, and brush location.
+*/
+IMPLEMENT_SIMPLE_AUTOMATION_TEST(FBSPCubeAddTest, "System.Editor.Geometry.BSP Cube - New Brush", EAutomationTestFlags::ATF_Editor);
+
+bool FBSPCubeAddTest::RunTest(const FString& Parameters)
+{
+	//** SETUP **//
+	// Create the world and setup the level.
+	UWorld* World = AutomationEditorCommonUtils::CreateNewMap();
+	ULevel* CurrentLevel = World->GetCurrentLevel();
+	GLevelEditorModeTools().MapChangeNotify();
+	UBrushBuilder* DefaultBuilderBrush = World->GetDefaultBrush()->BrushBuilder;
+
+	//** TEST **//
+	// Add a new additive cube BSP to the world using it's default settings.
+	GEditor->Exec(World, TEXT("BRUSH Scale 1 1 1"));
+	World->GetDefaultBrush()->BrushBuilder = Cast<UCubeBuilder>(GEditor->FindBrushBuilder(UCubeBuilder::StaticClass()));
+	GEditor->Exec(World, TEXT("BRUSH MOVETO X=0 Y=0 Z=10"));
+	GEditor->Exec(World, TEXT("BRUSH ADD"));
+
+	// Add a new subtractive cube BSP to the world.  It will be interesting with the previous cube. 
+	// The cubes settings are being altered.
+	UCubeBuilder* CubeAdditiveBrushBuilder = Cast<UCubeBuilder>(GEditor->FindBrushBuilder(UCubeBuilder::StaticClass()));
+	CubeAdditiveBrushBuilder->X = 300.0f;
+	CubeAdditiveBrushBuilder->Y = 100.0f;
+	CubeAdditiveBrushBuilder->Z = 100.0f;
+	GEditor->Exec(World, TEXT("BRUSH Scale 1 1 1"));
+	CubeAdditiveBrushBuilder->Build(World);
+	GEditor->Exec(World, TEXT("BRUSH MOVETO X=0 Y=0 Z=0"));
+	GEditor->Exec(World, TEXT("BRUSH SUBTRACT"));
+
+	//** VERIFY **//
+	// Gather the objects needed to run the tests.
+	// Get the levels BSP Model.
+	UModel* BSPModel = CurrentLevel->Model;
+	// Set the brush actor variables to the generated additive and subtractive brushes.
+	ABrush* AdditiveBSP = nullptr;
+	ABrush* SubtractiveBSP = nullptr;
+	int32 BrushCount = 0;
+	for (TObjectIterator<ABrush> It; It; ++It)
+	{
+		ABrush* FoundBrush = *It;
+
+		if (FoundBrush->BrushType == EBrushType::Brush_Add)
+		{
+			AdditiveBSP = Cast<ABrush>(FoundBrush);
+			BrushCount++;
+		}
+		else if (FoundBrush->BrushType == EBrushType::Brush_Subtract)
+		{
+			SubtractiveBSP = Cast<ABrush>(FoundBrush);
+			BrushCount++;
+		}
+	}
+
+	// Verify the BSP brushes that were created have been found.
+	TestEqual<int32>(TEXT("The number of BSP brushes created is not correct, expected 2."), BrushCount, 2);
+	TestNotNull<ABrush>(TEXT("Additive BSP is NULL.  Most likely it wasn't created."), AdditiveBSP);
+	TestNotNull<ABrush>(TEXT("Subtractive BSP is NULL.  Most likely it wasn't created."), AdditiveBSP);
+
+	// Verify the location of the BSP.
+	TestEqual<FVector>(TEXT("Additive BSP is not located at the correct coordinates."), AdditiveBSP->GetActorLocation(), WORLD_ORIGIN_VECTOR_OFFSET);
+	TestEqual<FVector>(TEXT("Subtractive BSP is not located at the correct coordinates."), SubtractiveBSP->GetActorLocation(), WORLD_ORIGIN_VECTOR);
+
+	// Verify there are 8 vertices for the Additive and Subtractive BSP.
+	TestEqual<uint32>(TEXT("The number of vertices are incorrect for the additive brush."), AdditiveBSP->Brush->NumUniqueVertices, CUBE_VERTS_COUNT);
+	TestEqual<uint32>(TEXT("The number of vertices are incorrect for the subtractive brush."), SubtractiveBSP->Brush->NumUniqueVertices, CUBE_VERTS_COUNT);
+
+	// Verify there are the correct amount of BSP surfaces are visible in the level.
+	TestEqual<int32>(TEXT("Incorrect total number of surfaces are being reported, expected 10."), BSPModel->Surfs.Num(), 10);
+
+#if WITH_EDITOR
+	// Check the BSP for any other errors.
+	SubtractiveBSP->CheckForErrors();
+	AdditiveBSP->CheckForErrors();
+#endif // WITH_EDITOR
+
+	//** TEARDOWN **//
+	// Set the default builder brush back to its original state.
+	World->GetDefaultBrush()->BrushBuilder = Cast<UBrushBuilder>(DefaultBuilderBrush);
+
+	return true;
+}
+
+
+/**
+* BSP Cube - Undo and Redo
+* Adds an additive brush to the world and then does an undo/redo of that transaction.
+* Tests for number vertices in the level.
+*/
+IMPLEMENT_SIMPLE_AUTOMATION_TEST(FBSPCubeUndoRedoTest, "System.Editor.Geometry.BSP Cube - Undo and Redo", EAutomationTestFlags::ATF_Editor);
+
+bool FBSPCubeUndoRedoTest::RunTest(const FString& Parameters)
+{
+	//** SETUP **//
+	// Create the world and setup the level.
+	UWorld* World = AutomationEditorCommonUtils::CreateNewMap();
+	ULevel* CurrentLevel = World->GetCurrentLevel();
+	GLevelEditorModeTools().MapChangeNotify();
+	GEditor->Exec(World, TEXT("BRUSH Scale 1 1 1"));
+	
+	int32 DefaultVertCount = CurrentLevel->Model->Verts.Num();
+
+	// Add a cube brush to the world.
+	UCubeBuilder* CubeAdditiveBrushBuilder = Cast<UCubeBuilder>(GEditor->FindBrushBuilder(UCubeBuilder::StaticClass()));
+	CubeAdditiveBrushBuilder->Build(World);
+	GEditor->Exec(World, TEXT("BRUSH MOVETO X=0 Y=0 Z=0"));
+	GEditor->Exec(World, TEXT("BRUSH ADD"));
+
+	if (DefaultVertCount != CurrentLevel->Model->Verts.Num())
+	{
+		//** TEST **//
+		// Undo and Redo the placement of the brush.  Keeping track of the number of vertices for both actions.
+		GEditor->UndoTransaction();
+		int32 VertsFromUndoing = CurrentLevel->Model->Verts.Num();
+
+		GEditor->RedoTransaction();
+		int32 VertsFromRedoing = CurrentLevel->Model->Verts.Num();
+
+		//** VERIFY **//
+		// Compare the vert count before and after the placement.
+		TestEqual<int32>(TEXT("Undo brush placement appears to not have been undone."), DefaultVertCount, VertsFromUndoing);
+		TestNotEqual<int32>(TEXT("Redo brush placement appears to not have been redone."), DefaultVertCount, VertsFromRedoing);
+
+		return true;
+	}
+
+	UE_LOG(LogGeometryTests, Error, TEXT("Unable to test Undo/Redo since the brush actor was not placed into the level."))
+	return false;
+}
\ No newline at end of file
diff --git a/Engine/Source/Editor/UnrealEd/Public/Tests/AutomationEditorCommon.h b/Engine/Source/Editor/UnrealEd/Public/Tests/AutomationEditorCommon.h
--- a/Engine/Source/Editor/UnrealEd/Public/Tests/AutomationEditorCommon.h
+++ b/Engine/Source/Editor/UnrealEd/Public/Tests/AutomationEditorCommon.h
@@ -147,6 +147,14 @@ namespace AutomationEditorCommonUtils
 	* @param InDeviceName - Device Name 
 	*/
 	void GetLaunchOnDeviceID(FString& OutDeviceID, const FString& InMapName, const FString& InDeviceName);
+
+	/**
+	* Sets the first found ortho viewport camera to the desired location and rotation.
+
+	* @param ViewLocation - Desired location for the viewport view.
+	* @param ViewRotation - Desired rotation of the viewport view.
+	*/
+	bool SetOrthoViewportView(const FVector& ViewLocation, const FRotator& ViewRotation);
 }
 
 
