fixed 16bit float normal map detection

diff --git a/Engine/Source/Editor/UnrealEd/Private/NormalMapIdentification.cpp b/Engine/Source/Editor/UnrealEd/Private/NormalMapIdentification.cpp
--- a/Engine/Source/Editor/UnrealEd/Private/NormalMapIdentification.cpp
+++ b/Engine/Source/Editor/UnrealEd/Private/NormalMapIdentification.cpp
@@ -128,7 +128,6 @@ public:
 
 		const float OneOver65535 = 1.0f / 65535.0f;
 
-		// Need to verify component order
 		Result.R = (float)((uint16*)PixelToSample)[0] * OneOver65535;
 		Result.G = (float)((uint16*)PixelToSample)[1] * OneOver65535;
 		Result.B = (float)((uint16*)PixelToSample)[2] * OneOver65535;
@@ -153,11 +152,11 @@ public:
 		FLinearColor Result;
 		uint8* PixelToSample = SourceTextureData + ((Y * TextureSizeX + X) * 8);
 
-		// Need to verify component order
-		Result.R = (float)((FFloat16*)PixelToSample)[3];
-		Result.G = (float)((FFloat16*)PixelToSample)[2];
-		Result.B = (float)((FFloat16*)PixelToSample)[1];
-		Result.A = (float)((FFloat16*)PixelToSample)[0];
+		// this assume the normal map to be in linear (not the case if Photoshop converts a 8bit normalmap to float and saves it as 16bit dds)
+		Result.R = (float)((FFloat16*)PixelToSample)[0];
+		Result.G = (float)((FFloat16*)PixelToSample)[1];
+		Result.B = (float)((FFloat16*)PixelToSample)[2];
+		Result.A = (float)((FFloat16*)PixelToSample)[3];
 
 		return Result;
 	}
