Dont tick the directory watcher when the editor is loaded in project selection mode only.  Directory changes could send callbacks to methods that expect certian UI to be available which isnt when only the project dialog is open

diff --git a/Engine/Source/Editor/UnrealEd/Private/EditorEngine.cpp b/Engine/Source/Editor/UnrealEd/Private/EditorEngine.cpp
--- a/Engine/Source/Editor/UnrealEd/Private/EditorEngine.cpp
+++ b/Engine/Source/Editor/UnrealEd/Private/EditorEngine.cpp
@@ -1118,9 +1118,12 @@ void UEditorEngine::Tick( float DeltaSeconds, bool bIdleMode )
 
 	// tick the directory watcher
 	// @todo: Put me into an FTicker that is created when the DW module is loaded
-	static FName DirectoryWatcherName("DirectoryWatcher");
-	FDirectoryWatcherModule& DirectoryWatcherModule = FModuleManager::Get().LoadModuleChecked<FDirectoryWatcherModule>(DirectoryWatcherName);
-	DirectoryWatcherModule.Get()->Tick(DeltaSeconds);
+	if( !FApp::IsGameNameEmpty() )
+	{
+		static FName DirectoryWatcherName("DirectoryWatcher");
+		FDirectoryWatcherModule& DirectoryWatcherModule = FModuleManager::Get().LoadModuleChecked<FDirectoryWatcherModule>(DirectoryWatcherName);
+		DirectoryWatcherModule.Get()->Tick(DeltaSeconds);
+	}
 
 	if( bShouldTickEditorWorld )
 	{ 
