Fixed issue where messages of severity "Warning" were not being announced via a notification at the end of a PIE session.

FMessageLog::Notify takes, as a severity filter, the highest category which should _not_ be filtered.  So to display Warnings, it needs to be called with the default (EMessageSeverity::Info).

diff --git a/Engine/Source/Editor/UnrealEd/Private/PlayLevel.cpp b/Engine/Source/Editor/UnrealEd/Private/PlayLevel.cpp
--- a/Engine/Source/Editor/UnrealEd/Private/PlayLevel.cpp
+++ b/Engine/Source/Editor/UnrealEd/Private/PlayLevel.cpp
@@ -322,7 +322,7 @@ void UEditorEngine::EndPlayMap()
 	bUseVRPreviewForPlayWorld = false;
 
 	// display any info if required.
-	FMessageLog("PIE").Notify( LOCTEXT("PIEErrorsPresent", "Errors/warnings reported while playing in editor."), EMessageSeverity::Warning );
+	FMessageLog("PIE").Notify(LOCTEXT("PIEErrorsPresent", "Errors/warnings reported while playing in editor."));
 }
 
 void UEditorEngine::CleanupPIEOnlineSessions(TArray<FName> OnlineIdentifiers)
