[UE-16007] Fix for crash that occurs when you call IsDataOnlyBlueprint on a blueprint that has no ParentClass

diff --git a/Engine/Source/Editor/UnrealEd/Private/Kismet2/BlueprintEditorUtils.cpp b/Engine/Source/Editor/UnrealEd/Private/Kismet2/BlueprintEditorUtils.cpp
--- a/Engine/Source/Editor/UnrealEd/Private/Kismet2/BlueprintEditorUtils.cpp
+++ b/Engine/Source/Editor/UnrealEd/Private/Kismet2/BlueprintEditorUtils.cpp
@@ -2792,7 +2792,9 @@ bool FBlueprintEditorUtils::IsDataOnlyBlueprint(const UBlueprint* Blueprint)
 		return false;
 	}
 
-	if( Blueprint->ParentClass->IsChildOf( UActorComponent::StaticClass() ) )
+	// Note that the current implementation of IsChildOf will not crash when called on a nullptr, but
+	// I'm explicitly null checking because it seems unwise to rely on this behavior:
+	if (Blueprint->ParentClass && Blueprint->ParentClass->IsChildOf(UActorComponent::StaticClass()))
 	{
 		return false;
 	}
@@ -2834,7 +2836,7 @@ bool FBlueprintEditorUtils::IsDataOnlyBlueprint(const UBlueprint* Blueprint)
 
 	// Make sure there's nothing in the user construction script, other than an entry node
 	UEdGraph* UserConstructionScript = (Blueprint->FunctionGraphs.Num() == 1) ? Blueprint->FunctionGraphs[0] : NULL;
-	if (UserConstructionScript)
+	if (UserConstructionScript && Blueprint->ParentClass)
 	{
 		//Call parent construction script may be added automatically
 		UBlueprint* BlueprintParent = Cast<UBlueprint>(Blueprint->ParentClass->ClassGeneratedBy);
diff --git a/Engine/Source/Runtime/Engine/Classes/Engine/Blueprint.h b/Engine/Source/Runtime/Engine/Classes/Engine/Blueprint.h
--- a/Engine/Source/Runtime/Engine/Classes/Engine/Blueprint.h
+++ b/Engine/Source/Runtime/Engine/Classes/Engine/Blueprint.h
@@ -271,7 +271,11 @@ class ENGINE_API UBlueprint : public UBlueprintCore
 	UPROPERTY(config)
 	uint32 bRecompileOnLoad:1;
 
-	/** Pointer to the parent class that the generated class should derive from */
+	/** 
+		Pointer to the parent class that the generated class should derive from. This *can* be null under rare circumstances, 
+		one such case can be created by creating a blueprint (A) based on another blueprint (B), shutting down the editor, and
+		deleting the parent blueprint.
+	*/
 	UPROPERTY(AssetRegistrySearchable)
 	TSubclassOf<class UObject> ParentClass;
 
