Added bUseVertexNormalForHemisphereGather Lightmass primitive setting for foliage static lighting, in cases where artists want foliage lighting to match the underlying terrain

diff --git a/Engine/Source/Editor/UnrealEd/Private/Lightmass/Lightmass.cpp b/Engine/Source/Editor/UnrealEd/Private/Lightmass/Lightmass.cpp
--- a/Engine/Source/Editor/UnrealEd/Private/Lightmass/Lightmass.cpp
+++ b/Engine/Source/Editor/UnrealEd/Private/Lightmass/Lightmass.cpp
@@ -1477,7 +1477,8 @@ void FLightmassExporter::WriteMeshInstances( int32 Channel )
 						NewElementData.MaterialId = Material->GetLightingGuid();
 						NewElementData.bUseTwoSidedLighting = Primitive->LightmassSettings.bUseTwoSidedLighting;
 						NewElementData.bShadowIndirectOnly = Primitive->LightmassSettings.bShadowIndirectOnly;
-						NewElementData.bUseEmissiveForStaticLighting = Primitive->LightmassSettings.bUseEmissiveForStaticLighting;;
+						NewElementData.bUseEmissiveForStaticLighting = Primitive->LightmassSettings.bUseEmissiveForStaticLighting;
+						NewElementData.bUseVertexNormalForHemisphereGather = Primitive->LightmassSettings.bUseVertexNormalForHemisphereGather;
 						// Combine primitive and level boost settings so we don't have to send the level settings over to Lightmass  
 						NewElementData.EmissiveLightFalloffExponent = Primitive->LightmassSettings.EmissiveLightFalloffExponent;
 						NewElementData.EmissiveLightExplicitInfluenceRadius = Primitive->LightmassSettings.EmissiveLightExplicitInfluenceRadius;
@@ -1560,6 +1561,7 @@ void FLightmassExporter::WriteLandscapeInstances( int32 Channel )
 			NewElementData.bUseTwoSidedLighting = LMSetting.bUseTwoSidedLighting;
 			NewElementData.bShadowIndirectOnly = LMSetting.bShadowIndirectOnly;
 			NewElementData.bUseEmissiveForStaticLighting = LMSetting.bUseEmissiveForStaticLighting;
+			NewElementData.bUseVertexNormalForHemisphereGather = LMSetting.bUseVertexNormalForHemisphereGather;
 			// Combine primitive and level boost settings so we don't have to send the level settings over to Lightmass  
 			NewElementData.EmissiveLightFalloffExponent = LMSetting.EmissiveLightFalloffExponent;
 			NewElementData.EmissiveLightExplicitInfluenceRadius = LMSetting.EmissiveLightExplicitInfluenceRadius;
@@ -1664,6 +1666,7 @@ void FLightmassExporter::WriteMappings( int32 Channel )
 			TempData.bUseTwoSidedLighting = PrimitiveSettings.bUseTwoSidedLighting;
 			TempData.bShadowIndirectOnly = PrimitiveSettings.bShadowIndirectOnly;
 			TempData.bUseEmissiveForStaticLighting = PrimitiveSettings.bUseEmissiveForStaticLighting;
+			TempData.bUseVertexNormalForHemisphereGather = PrimitiveSettings.bUseVertexNormalForHemisphereGather;
 			TempData.EmissiveLightFalloffExponent = PrimitiveSettings.EmissiveLightFalloffExponent;
 			TempData.EmissiveLightExplicitInfluenceRadius = PrimitiveSettings.EmissiveLightExplicitInfluenceRadius;
 			TempData.EmissiveBoost = PrimitiveSettings.EmissiveBoost * LevelSettings.EmissiveBoost;
diff --git a/Engine/Source/Programs/UnrealLightmass/Private/Lighting/LightingMesh.cpp b/Engine/Source/Programs/UnrealLightmass/Private/Lighting/LightingMesh.cpp
--- a/Engine/Source/Programs/UnrealLightmass/Private/Lighting/LightingMesh.cpp
+++ b/Engine/Source/Programs/UnrealLightmass/Private/Lighting/LightingMesh.cpp
@@ -117,6 +117,7 @@ void FStaticLightingMesh::Import( FLightmassImporter& Importer )
 		CurrentMaterialElement.bUseTwoSidedLighting = METempData.bUseTwoSidedLighting;
 		CurrentMaterialElement.bShadowIndirectOnly = METempData.bShadowIndirectOnly;
 		CurrentMaterialElement.bUseEmissiveForStaticLighting = METempData.bUseEmissiveForStaticLighting;
+		CurrentMaterialElement.bUseVertexNormalForHemisphereGather = METempData.bUseVertexNormalForHemisphereGather;
 		// Validating data here instead of in Unreal since EmissiveLightFalloffExponent is used in so many different object types
 		CurrentMaterialElement.EmissiveLightFalloffExponent = FMath::Max(METempData.EmissiveLightFalloffExponent, 0.0f);
 		CurrentMaterialElement.EmissiveLightExplicitInfluenceRadius = FMath::Max(METempData.EmissiveLightExplicitInfluenceRadius, 0.0f);
diff --git a/Engine/Source/Programs/UnrealLightmass/Private/Lighting/LightingMesh.h b/Engine/Source/Programs/UnrealLightmass/Private/Lighting/LightingMesh.h
--- a/Engine/Source/Programs/UnrealLightmass/Private/Lighting/LightingMesh.h
+++ b/Engine/Source/Programs/UnrealLightmass/Private/Lighting/LightingMesh.h
@@ -315,6 +315,7 @@ public:
 	inline int32 GetNumElements() const { return MaterialElements.Num(); }
 	inline bool ShouldColorInvalidTexels() const { return bColorInvalidTexels; }
 	inline bool HasImportedNormal(int32 ElementIndex) const { return MaterialElements[ElementIndex].Material->NormalSize > 0; }
+	inline bool UseVertexNormalForHemisphereGather(int32 ElementIndex) const { return MaterialElements[ElementIndex].bUseVertexNormalForHemisphereGather; }
 
 	/**
 	 *	Returns the Guid for the object associated with this lighting mesh.
diff --git a/Engine/Source/Programs/UnrealLightmass/Private/Lighting/TextureMapping.cpp b/Engine/Source/Programs/UnrealLightmass/Private/Lighting/TextureMapping.cpp
--- a/Engine/Source/Programs/UnrealLightmass/Private/Lighting/TextureMapping.cpp
+++ b/Engine/Source/Programs/UnrealLightmass/Private/Lighting/TextureMapping.cpp
@@ -366,7 +366,9 @@ void FStaticLightingSystem::CacheIrradiancePhotonsTextureMapping(FStaticLighting
 
 				// Normalize the tangent basis and ensure it is orthonormal
 				CurrentVertex.WorldTangentZ = CurrentVertex.WorldTangentZ.GetUnsafeNormal3();
-				CurrentVertex.TriangleNormal = TexelToVertex.TriangleNormal.GetUnsafeNormal3();
+
+				const bool bUseVertexNormalForHemisphereGather = TextureMapping->Mesh->UseVertexNormalForHemisphereGather(TexelToVertex.ElementIndex);
+				CurrentVertex.TriangleNormal = bUseVertexNormalForHemisphereGather ? CurrentVertex.WorldTangentZ : TexelToVertex.TriangleNormal.GetUnsafeNormal3();
 				checkSlow(!CurrentVertex.TriangleNormal.ContainsNaN());
 
 				const FVector4 OriginalTangentX = CurrentVertex.WorldTangentX;
@@ -957,7 +959,9 @@ void FStaticLightingSystem::SetupTextureMapping(
 
 				// Normalize the tangent basis and ensure it is orthonormal
 				TexelToVertex.WorldTangentZ = TexelToVertex.WorldTangentZ.GetUnsafeNormal3();
-				TexelToVertex.TriangleNormal = TexelToVertex.TriangleNormal.GetUnsafeNormal3();
+
+				const bool bUseVertexNormalForHemisphereGather = TextureMapping->Mesh->UseVertexNormalForHemisphereGather(TexelToVertex.ElementIndex);
+				TexelToVertex.TriangleNormal = bUseVertexNormalForHemisphereGather ? TexelToVertex.WorldTangentZ : TexelToVertex.TriangleNormal.GetUnsafeNormal3();
 				checkSlow(!TexelToVertex.TriangleNormal.ContainsNaN());
 
 				const FVector4 OriginalTangentX = TexelToVertex.WorldTangentX;
diff --git a/Engine/Source/Programs/UnrealLightmass/Public/SceneExport.h b/Engine/Source/Programs/UnrealLightmass/Public/SceneExport.h
--- a/Engine/Source/Programs/UnrealLightmass/Public/SceneExport.h
+++ b/Engine/Source/Programs/UnrealLightmass/Public/SceneExport.h
@@ -882,6 +882,11 @@ struct FMaterialElementData
 	uint32 bShadowIndirectOnly:1;
 	/** If true, allow using the emissive for static lighting.						*/
 	uint32 bUseEmissiveForStaticLighting:1;
+	/** 
+	 * Typically the triangle normal is used for hemisphere gathering which prevents incorrect self-shadowing from artist-tweaked vertex normals. 
+	 * However in the case of foliage whose vertex normal has been setup to match the underlying terrain, gathering in the direction of the vertex normal is desired.
+	 */
+	uint32 bUseVertexNormalForHemisphereGather:1;
 	/** Direct lighting falloff exponent for mesh area lights created from emissive areas on this primitive. */
 	float EmissiveLightFalloffExponent;
 	/** 
@@ -904,6 +909,7 @@ struct FMaterialElementData
 		  bUseTwoSidedLighting(false)
 		, bShadowIndirectOnly(false)
 		, bUseEmissiveForStaticLighting(true)
+		, bUseVertexNormalForHemisphereGather(false)
 	    , EmissiveLightFalloffExponent(2.0f)
 		, EmissiveLightExplicitInfluenceRadius(0.0f)
 		, EmissiveBoost(1.0f)
diff --git a/Engine/Source/Runtime/Core/Public/UObject/ObjectVersion.h b/Engine/Source/Runtime/Core/Public/UObject/ObjectVersion.h
--- a/Engine/Source/Runtime/Core/Public/UObject/ObjectVersion.h
+++ b/Engine/Source/Runtime/Core/Public/UObject/ObjectVersion.h
@@ -516,6 +516,8 @@ enum EUnrealEngineObjectUE4Version
 	VER_UE4_SERIALIZE_TEXT_IN_PACKAGES,
 	// Added pivot to widget components, need to load old versions as a 0,0 pivot, new default is 0.5,0.5
 	VER_UE4_ADD_BLEND_MODE_TO_WIDGET_COMPONENT,
+	// Added lightmass primitive setting
+	VER_UE4_NEW_LIGHTMASS_PRIMITIVE_SETTING,
 
 	// -----<new versions can be added before this line>-------------------------------------------------
 	// - this needs to be the last line (see note below)
diff --git a/Engine/Source/Runtime/Engine/Classes/Engine/EngineTypes.h b/Engine/Source/Runtime/Engine/Classes/Engine/EngineTypes.h
--- a/Engine/Source/Runtime/Engine/Classes/Engine/EngineTypes.h
+++ b/Engine/Source/Runtime/Engine/Classes/Engine/EngineTypes.h
@@ -1191,6 +1191,13 @@ struct FLightmassPrimitiveSettings
 	UPROPERTY(EditAnywhere, BlueprintReadOnly, Category=Lightmass)
 	uint32 bUseEmissiveForStaticLighting:1;
 
+	/** 
+	 * Typically the triangle normal is used for hemisphere gathering which prevents incorrect self-shadowing from artist-tweaked vertex normals. 
+	 * However in the case of foliage whose vertex normal has been setup to match the underlying terrain, gathering in the direction of the vertex normal is desired.
+	 */
+	UPROPERTY(EditAnywhere, BlueprintReadOnly, Category=Lightmass)
+	uint32 bUseVertexNormalForHemisphereGather:1;
+
 	/** Direct lighting falloff exponent for mesh area lights created from emissive areas on this primitive. */
 	UPROPERTY()
 	float EmissiveLightFalloffExponent;
@@ -1220,6 +1227,7 @@ struct FLightmassPrimitiveSettings
 		bUseTwoSidedLighting = false;
 		bShadowIndirectOnly = false;
 		bUseEmissiveForStaticLighting = false;
+		bUseVertexNormalForHemisphereGather = false;
 		EmissiveLightFalloffExponent = 8.0f;
 		EmissiveLightExplicitInfluenceRadius = 0.0f;
 		EmissiveBoost = 1.0f;
@@ -1233,6 +1241,7 @@ struct FLightmassPrimitiveSettings
 		if ((A.bUseTwoSidedLighting != B.bUseTwoSidedLighting) ||
 			(A.bShadowIndirectOnly != B.bShadowIndirectOnly) || 
 			(A.bUseEmissiveForStaticLighting != B.bUseEmissiveForStaticLighting) || 
+			(A.bUseVertexNormalForHemisphereGather != B.bUseVertexNormalForHemisphereGather) || 
 			(fabsf(A.EmissiveLightFalloffExponent - B.EmissiveLightFalloffExponent) > SMALL_NUMBER) ||
 			(fabsf(A.EmissiveLightExplicitInfluenceRadius - B.EmissiveLightExplicitInfluenceRadius) > SMALL_NUMBER) ||
 			(fabsf(A.EmissiveBoost - B.EmissiveBoost) > SMALL_NUMBER) ||
diff --git a/Engine/Source/Runtime/Engine/Private/Polygon.cpp b/Engine/Source/Runtime/Engine/Private/Polygon.cpp
--- a/Engine/Source/Runtime/Engine/Private/Polygon.cpp
+++ b/Engine/Source/Runtime/Engine/Private/Polygon.cpp
@@ -37,6 +37,13 @@ FArchive& operator<<(FArchive& Ar, FLightmassPrimitiveSettings& Settings)
 	Ar << Temp;
 	Settings.bUseEmissiveForStaticLighting = Temp;
 
+	if (Ar.UE4Ver() >= VER_UE4_NEW_LIGHTMASS_PRIMITIVE_SETTING)
+	{
+		Temp = Settings.bUseVertexNormalForHemisphereGather;
+		Ar << Temp;
+		Settings.bUseVertexNormalForHemisphereGather = Temp;
+	}
+	
 	Ar << Settings.EmissiveLightFalloffExponent;
 	Ar << Settings.EmissiveLightExplicitInfluenceRadius;
 	
