Fix: HLOD: memory still being used after release
UE-14610

diff --git a/Engine/Source/Editor/UnrealEd/Private/HierarchicalLOD.cpp b/Engine/Source/Editor/UnrealEd/Private/HierarchicalLOD.cpp
--- a/Engine/Source/Editor/UnrealEd/Private/HierarchicalLOD.cpp
+++ b/Engine/Source/Editor/UnrealEd/Private/HierarchicalLOD.cpp
@@ -158,11 +158,14 @@ void FHierarchicalLODBuilder::BuildClusters(class ULevel* InLevel)
 		}
 	}
 
+	ULevel::BuildStreamingData(InLevel->OwningWorld, InLevel);
+
 	for (auto& Asset : AssetsToDelete)
 	{
+		Asset->MarkPendingKill();
 		ObjectTools::DeleteSingleObject(Asset, false);
 	}
-
+	
 	// garbage collect
 	CollectGarbage( GARBAGE_COLLECTION_KEEPFLAGS, true );
 
