Store pre-cooked package list as long package namesto avoid potential ambiguity with short package names. Seems reasonable now as the pre-cooked package list is no longer used for anything else.

diff --git a/Engine/Source/Editor/UnrealEd/Private/Commandlets/ChunkManifestGenerator.cpp b/Engine/Source/Editor/UnrealEd/Private/Commandlets/ChunkManifestGenerator.cpp
--- a/Engine/Source/Editor/UnrealEd/Private/Commandlets/ChunkManifestGenerator.cpp
+++ b/Engine/Source/Editor/UnrealEd/Private/Commandlets/ChunkManifestGenerator.cpp
@@ -434,10 +434,8 @@ bool FChunkManifestGenerator::LoadAssetRegistry(const FString& SandboxPath, cons
 		}
 		for (auto& LoadedAssetData : AssetRegistryData)
 		{
-			
-			FName ShortPackageName = FName(*FPaths::GetBaseFilename(LoadedAssetData.PackageName.ToString()));
 			if (PackagesToKeep &&
-				PackagesToKeep->Contains(ShortPackageName) == false)
+				PackagesToKeep->Contains(LoadedAssetData.PackageName) == false)
 			{
 				continue;
 			}
@@ -454,9 +452,7 @@ bool FChunkManifestGenerator::LoadAssetRegistry(const FString& SandboxPath, cons
 
 		for (const auto& SavedAsset : SavedAssetRegistryData)
 		{
-			FName ShortPackageName = FName(*FPaths::GetBaseFilename(SavedAsset.Value->PackageName.ToString()));
-
-			if (PackagesToKeep && PackagesToKeep->Contains(ShortPackageName))
+			if (PackagesToKeep && PackagesToKeep->Contains(SavedAsset.Value->PackageName))
 			{ 
 				AssetRegistryData.Add(*SavedAsset.Value);
 			}
diff --git a/Engine/Source/Editor/UnrealEd/Private/CookOnTheFlyServer.cpp b/Engine/Source/Editor/UnrealEd/Private/CookOnTheFlyServer.cpp
--- a/Engine/Source/Editor/UnrealEd/Private/CookOnTheFlyServer.cpp
+++ b/Engine/Source/Editor/UnrealEd/Private/CookOnTheFlyServer.cpp
@@ -2680,6 +2680,8 @@ void UCookOnTheFlyServer::CleanSandbox( const bool bIterative )
 					FDateTime DependentTimestamp;
 					FName StandardCookedFileFName = FName(*StandardCookedFilename);
 
+					bool bShouldKeep = true;
+
 					if (PDInfoModule.DeterminePackageDependentTimeStamp(*(FPaths::GetBaseFilename(StandardCookedFilename, false)), DependentTimestamp) == true)
 					{
 						double Diff = (CookedTimestamp - DependentTimestamp).GetTotalSeconds();
@@ -2692,18 +2694,26 @@ void UCookOnTheFlyServer::CleanSandbox( const bool bIterative )
 							IFileManager::Get().Delete(*CookedFilename);
 
 							CookedPackages.RemoveFileForPlatform(StandardCookedFileFName, PlatformFName);
+
+							bShouldKeep = false;
+						}
+					}
+
+					if (bShouldKeep)
+					{
+						FString LongPackageName;
+						FString FailureReason;
+						if (FPackageName::TryConvertFilenameToLongPackageName(StandardCookedFilename, LongPackageName, &FailureReason))
+						{
+							const FName LongPackageFName(*LongPackageName);
+							PackagesKeptFromPreviousCook.Add(LongPackageFName);
 						}
 						else
 						{
-							FName ShortPackageName = FName(*FPaths::GetBaseFilename(StandardCookedFileFName.ToString()));
-							PackagesKeptFromPreviousCook.Add(ShortPackageName);
+							LogCookerMessage(FString::Printf(TEXT("Unable to generate long package name for %s because %s"), *StandardCookedFilename, *FailureReason), EMessageSeverity::Warning);
+							UE_LOG(LogCook, Warning, TEXT("Unable to generate long package name for %s because %s"), *StandardCookedFilename, *FailureReason);
 						}
 					}
-					else
-					{
-						FName ShortPackageName = FName(*FPaths::GetBaseFilename(StandardCookedFileFName.ToString()));
-						PackagesKeptFromPreviousCook.Add(ShortPackageName);
-					}
 				}
 			}
 
