Fixed non-unity CIS error due to log category not being defined in that file

Not sure if this was intended or the category was used by mistake

diff --git a/Engine/Source/Editor/UnrealEd/Private/Commandlets/GenerateGatherArchiveCommandlet.cpp b/Engine/Source/Editor/UnrealEd/Private/Commandlets/GenerateGatherArchiveCommandlet.cpp
--- a/Engine/Source/Editor/UnrealEd/Private/Commandlets/GenerateGatherArchiveCommandlet.cpp
+++ b/Engine/Source/Editor/UnrealEd/Private/Commandlets/GenerateGatherArchiveCommandlet.cpp
@@ -11,6 +11,7 @@
 
 
 DEFINE_LOG_CATEGORY_STATIC(LogGenerateArchiveCommandlet, Log, All);
+DEFINE_LOG_CATEGORY_STATIC(LogGatherTextFromSourceCommandlet, Log, All);
 
 /**
  *	UGenerateGatherArchiveCommandlet
