Added support for disabling compression per CookByTheBook session.

diff --git a/Engine/Source/Developer/TargetPlatform/Public/Common/TargetPlatformBase.h b/Engine/Source/Developer/TargetPlatform/Public/Common/TargetPlatformBase.h
--- a/Engine/Source/Developer/TargetPlatform/Public/Common/TargetPlatformBase.h
+++ b/Engine/Source/Developer/TargetPlatform/Public/Common/TargetPlatformBase.h
@@ -356,6 +356,8 @@ public:
 		case ETargetPlatformFeatures::SdkConnectDisconnect:
 		case ETargetPlatformFeatures::UserCredentials:
 			break;
+		case ETargetPlatformFeatures::ShouldUseCompressedCookedPackages:
+			return true;
 		}
 
 		return false;
diff --git a/Engine/Source/Developer/TargetPlatform/Public/Interfaces/ITargetPlatform.h b/Engine/Source/Developer/TargetPlatform/Public/Interfaces/ITargetPlatform.h
--- a/Engine/Source/Developer/TargetPlatform/Public/Interfaces/ITargetPlatform.h
+++ b/Engine/Source/Developer/TargetPlatform/Public/Interfaces/ITargetPlatform.h
@@ -50,6 +50,9 @@ enum class ETargetPlatformFeatures
 
 	/** Vertex Shader Texture Sampling. */
 	VertexShaderTextureSampling,
+
+	/** Should use compressed cooked packages */
+	ShouldUseCompressedCookedPackages, 
 };
 
 
diff --git a/Engine/Source/Editor/UnrealEd/Classes/CookOnTheSide/CookOnTheFlyServer.h b/Engine/Source/Editor/UnrealEd/Classes/CookOnTheSide/CookOnTheFlyServer.h
--- a/Engine/Source/Editor/UnrealEd/Classes/CookOnTheSide/CookOnTheFlyServer.h
+++ b/Engine/Source/Editor/UnrealEd/Classes/CookOnTheSide/CookOnTheFlyServer.h
@@ -33,6 +33,8 @@ enum class ECookByTheBookOptions
 	MapsOnly = 0x2,				// cook only maps
 	NoDevContent = 0x4,			// don't include dev content
 	LeakTest = 0x8,				// test for uobject leaks after each level load
+	ForceDisableCompressed = 0x10, // force compression to be disabled even if the cooker was initialized with it enabled
+	ForceEnableCompressed = 0x20, // force compression to be on even if the cooker was initialized with it disabled
 };
 ENUM_CLASS_FLAGS(ECookByTheBookOptions);
 
@@ -591,7 +593,9 @@ private:
 			bRunning(false),
 			CookTime( 0.0 ),
 			CookStartTime( 0.0 ), 
-			bErrorOnEngineContentUse(false)
+			bErrorOnEngineContentUse(false),
+			bForceEnableCompressedPackages(false),
+			bForceDisableCompressedPackages(false)
 		{ }
 
 		/** Should we test for UObject leaks */
@@ -618,9 +622,14 @@ private:
 		TArray<FFilePlatformRequest> PreviousCookRequests; 
 		/** If we are based on a release version of the game this is the set of packages which were cooked in that release */
 		TMap<FName,TArray<FName> > BasedOnReleaseCookedPackages;
+		/** Timing information about cook by the book */
 		double CookTime;
 		double CookStartTime;
+		/** error when detecting engine content being used in this cook */
 		bool bErrorOnEngineContentUse;
+		/** force this cook by the book to enable \ disable package compression */
+		bool bForceEnableCompressedPackages;
+		bool bForceDisableCompressedPackages;
 	};
 	FCookByTheBookOptions* CookByTheBookOptions;
 	
diff --git a/Engine/Source/Editor/UnrealEd/Private/CookOnTheFlyServer.cpp b/Engine/Source/Editor/UnrealEd/Private/CookOnTheFlyServer.cpp
--- a/Engine/Source/Editor/UnrealEd/Private/CookOnTheFlyServer.cpp
+++ b/Engine/Source/Editor/UnrealEd/Private/CookOnTheFlyServer.cpp
@@ -2019,7 +2019,15 @@ bool UCookOnTheFlyServer::SaveCookedPackage( UPackage* Package, uint32 SaveFlags
 		EObjectFlags Flags = RF_NoFlags;
 		bool bPackageFullyLoaded = false;
 
-		if (IsCookFlagSet(ECookInitializationFlags::Compressed) )
+		bool bShouldCompressPackage = IsCookFlagSet(ECookInitializationFlags::Compressed);
+
+		if (CookByTheBookOptions)
+		{
+			bShouldCompressPackage &= (CookByTheBookOptions->bForceDisableCompressedPackages == false);
+			bShouldCompressPackage |= CookByTheBookOptions->bForceEnableCompressedPackages;
+		}
+
+		if (bShouldCompressPackage)
 		{
 			Package->PackageFlags |= PKG_StoreCompressed;
 		}
@@ -3413,6 +3421,8 @@ void UCookOnTheFlyServer::StartCookByTheBook( const FCookByTheBookStartupOptions
 	CookByTheBookOptions->bGenerateStreamingInstallManifests = CookByTheBookStartupOptions.bGenerateStreamingInstallManifests;
 	CookByTheBookOptions->bGenerateDependenciesForMaps = CookByTheBookStartupOptions.bGenerateDependenciesForMaps;
 	CookByTheBookOptions->CreateReleaseVersion = CreateReleaseVersion;	
+	CookByTheBookOptions->bForceEnableCompressedPackages = !!(CookOptions & ECookByTheBookOptions::ForceEnableCompressed);
+	CookByTheBookOptions->bForceDisableCompressedPackages = !!(CookOptions & ECookByTheBookOptions::ForceDisableCompressed);
 
 	bool bCookAll = (CookOptions & ECookByTheBookOptions::CookAll) != ECookByTheBookOptions::None;
 	bool bMapsOnly = (CookOptions & ECookByTheBookOptions::MapsOnly) != ECookByTheBookOptions::None;
