Fixed unity and non-unity by removing use of static log category from another .cpp file

diff --git a/Engine/Source/Editor/UnrealEd/Private/Commandlets/GenerateGatherArchiveCommandlet.cpp b/Engine/Source/Editor/UnrealEd/Private/Commandlets/GenerateGatherArchiveCommandlet.cpp
--- a/Engine/Source/Editor/UnrealEd/Private/Commandlets/GenerateGatherArchiveCommandlet.cpp
+++ b/Engine/Source/Editor/UnrealEd/Private/Commandlets/GenerateGatherArchiveCommandlet.cpp
@@ -11,7 +11,6 @@
 
 
 DEFINE_LOG_CATEGORY_STATIC(LogGenerateArchiveCommandlet, Log, All);
-DEFINE_LOG_CATEGORY_STATIC(LogGatherTextFromSourceCommandlet, Log, All);
 
 /**
  *	UGenerateGatherArchiveCommandlet
@@ -77,7 +76,7 @@ int32 UGenerateGatherArchiveCommandlet::Main( const FString& Params )
 		const bool WasSourceCultureSpecified = GetStringFromConfig( *SectionName, TEXT("SourceCulture"), SourceCulture, GatherTextConfigPath );
 		if (WasSourceCultureSpecified)
 		{
-			UE_LOG(LogGatherTextFromSourceCommandlet, Warning, TEXT("SourceCulture detected in section %s. SourceCulture is deprecated, please use NativeCulture."), *SectionName);
+			UE_LOG(LogGenerateArchiveCommandlet, Warning, TEXT("SourceCulture detected in section %s. SourceCulture is deprecated, please use NativeCulture."), *SectionName);
 		}
 
 		// Use SourceCulture if NativeCulture isn't specified.
