In the pin type selector dropdown, hitting the shift key will no longer move the cursor back to the start of the searchbox and hitting the down and up arrow keys will again navigate the listview.

diff --git a/Engine/Source/Editor/UnrealEd/Public/SListViewSelectorDropdownMenu.h b/Engine/Source/Editor/UnrealEd/Public/SListViewSelectorDropdownMenu.h
--- a/Engine/Source/Editor/UnrealEd/Public/SListViewSelectorDropdownMenu.h
+++ b/Engine/Source/Editor/UnrealEd/Public/SListViewSelectorDropdownMenu.h
@@ -2,7 +2,7 @@
 
 #pragma once
 
-/** This is a container widget to help refocus to a listview widget from a searchbox or other widgets that are used in conjunction.
+/** This is a container widget to help refocus to a listview widget from a searchbox or other text box widgets that are used in conjunction.
 	Will refocus when the up or down arrows are pressed, and will commit a selection when enter is pressed regardless of where focus is. */
 template <typename ListType>
 class SListViewSelectorDropdownMenu : public SCompoundWidget
@@ -17,7 +17,7 @@ public:
 	 * @param InDefaultfocusWidget		The default widget to give focus to when the listview does not handle an action
 	 * @param InTargetListView			If the SListViewSelectorDropdownMenu is handling the KeyDown event, focus will be applied to the TargetListView for certain keys that it can handle
 	 */
-	void Construct(const FArguments& InArgs, TSharedPtr< SWidget> InDefaultFocusWidget, TSharedPtr< SListView< ListType > > InTargetListView)
+	void Construct(const FArguments& InArgs, TSharedPtr<SSearchBox> InDefaultFocusWidget, TSharedPtr< SListView< ListType > > InTargetListView)
 	{
 		check(InDefaultFocusWidget.IsValid());
 		check(InTargetListView.IsValid());
@@ -25,6 +25,9 @@ public:
 		TargetListView = InTargetListView;
 		DefaultFocusWidget = InDefaultFocusWidget;
 
+		// We have some overrides for OnKeyDown handling of the SSearchBox to provide the seamless functionality
+		InDefaultFocusWidget->SetOnKeyDownHandler(FOnKeyDown::CreateSP(this, &SListViewSelectorDropdownMenu<ListType>::OnKeyDown));
+
 		this->ChildSlot
 		[
 			InArgs._Content.Widget
@@ -53,7 +56,13 @@ public:
 		else
 		{
 			TSharedPtr< SWidget > DefaultFocusWidgetPtr = DefaultFocusWidget.Pin();
-			return DefaultFocusWidgetPtr->OnKeyDown(FindChildGeometry(MyGeometry, TargetListViewPtr.ToSharedRef()), KeyEvent).SetUserFocus(DefaultFocusWidgetPtr.ToSharedRef(), EFocusCause::OtherWidgetLostFocus);
+			FReply Reply = DefaultFocusWidgetPtr->OnKeyDown(FindChildGeometry(MyGeometry, TargetListViewPtr.ToSharedRef()), KeyEvent);
+
+			if (!DefaultFocusWidgetPtr->HasKeyboardFocus())
+			{
+				Reply.SetUserFocus(DefaultFocusWidgetPtr.ToSharedRef(), EFocusCause::OtherWidgetLostFocus);
+			}
+			return Reply;
 		}
 		return FReply::Unhandled();
 	}
@@ -63,5 +72,5 @@ private:
 	/** The type tree view widget this is handling keyboard input for */
 	TWeakPtr< SListView<ListType> > TargetListView;
 	/** Widget to revert focus back to when this widget does not handle (or forward) a key input */
-	TWeakPtr< SWidget > DefaultFocusWidget;
+	TWeakPtr< SSearchBox > DefaultFocusWidget;
 };
\ No newline at end of file
