Fixing build issue from missing include.

diff --git a/Engine/Source/Editor/UnrealEd/Public/SListViewSelectorDropdownMenu.h b/Engine/Source/Editor/UnrealEd/Public/SListViewSelectorDropdownMenu.h
--- a/Engine/Source/Editor/UnrealEd/Public/SListViewSelectorDropdownMenu.h
+++ b/Engine/Source/Editor/UnrealEd/Public/SListViewSelectorDropdownMenu.h
@@ -2,6 +2,8 @@
 
 #pragma once
 
+#include "SSearchBox.h"
+
 /** This is a container widget to help refocus to a listview widget from a searchbox or other text box widgets that are used in conjunction.
 	Will refocus when the up or down arrows are pressed, and will commit a selection when enter is pressed regardless of where focus is. */
 template <typename ListType>
