Fixed mirror transforms on Brushes.

diff --git a/Engine/Source/Editor/UnrealEd/Private/UnrealEdSrv.cpp b/Engine/Source/Editor/UnrealEd/Private/UnrealEdSrv.cpp
--- a/Engine/Source/Editor/UnrealEd/Private/UnrealEdSrv.cpp
+++ b/Engine/Source/Editor/UnrealEd/Private/UnrealEdSrv.cpp
@@ -1745,45 +1745,23 @@ static void MirrorActors(const FVector& MirrorScale)
 		ABrush* Brush = Cast< ABrush >( Actor );
 		if( Brush && Brush->Brush )
 		{
+			// Brushes have to reverse their poly vertex order and recalculate the normal as negating one of the scale axes
+			// changes the handedness of the local transform.
 			Brush->Modify();
 			Brush->Brush->Modify();
 			Brush->Brush->Polys->Modify();
 
-			const FVector LocalToWorldOffset = ( Brush->GetActorLocation() - PivotLocation );
-			const FVector LocationOffset = ( LocalToWorldOffset * MirrorScale ) - LocalToWorldOffset;
-
-			Brush->SetActorLocation( Brush->GetActorLocation() + LocationOffset, false );
-			Brush->SetPrePivot( Brush->GetPrePivot() * MirrorScale );
-
 			for( int32 poly = 0 ; poly < Brush->Brush->Polys->Element.Num() ; poly++ )
 			{
 				FPoly* Poly = &(Brush->Brush->Polys->Element[poly]);
 
-				Poly->TextureU *= MirrorScale;
-				Poly->TextureV *= MirrorScale;
-
-				Poly->Base += LocalToWorldOffset;
-				Poly->Base *= MirrorScale;
-				Poly->Base -= LocalToWorldOffset;
-				Poly->Base -= LocationOffset;
-
-				for( int32 vtx = 0 ; vtx < Poly->Vertices.Num(); vtx++ )
-				{
-					Poly->Vertices[vtx] += LocalToWorldOffset;
-					Poly->Vertices[vtx] *= MirrorScale;
-					Poly->Vertices[vtx] -= LocalToWorldOffset;
-					Poly->Vertices[vtx] -= LocationOffset;
-				}
-
 				Poly->Reverse();
 				Poly->CalcNormal();
 			}
 		}
-		else
-		{
-			Actor->Modify();
-			Actor->EditorApplyMirror( MirrorScale, PivotLocation );
-		}
+
+		Actor->Modify();
+		Actor->EditorApplyMirror( MirrorScale, PivotLocation );
 
 		Actor->InvalidateLightingCache();
 		Actor->PostEditMove( true );
