Simply reconstruct all nodes on load, rather than reconstructing and compiling the skeleton class (which has already been compiled once, and will undergo full compile later). This is an optimization of blueprint load times.

diff --git a/Engine/Source/Editor/UnrealEd/Private/Kismet2/BlueprintEditorUtils.cpp b/Engine/Source/Editor/UnrealEd/Private/Kismet2/BlueprintEditorUtils.cpp
--- a/Engine/Source/Editor/UnrealEd/Private/Kismet2/BlueprintEditorUtils.cpp
+++ b/Engine/Source/Editor/UnrealEd/Private/Kismet2/BlueprintEditorUtils.cpp
@@ -492,6 +492,35 @@ void FBlueprintEditorUtils::RefreshAllNodes(UBlueprint* Blueprint)
 	}
 }
 
+
+void FBlueprintEditorUtils::ReconstructAllNodes(UBlueprint* Blueprint)
+{
+	if (!Blueprint || !Blueprint->HasAllFlags(RF_LoadCompleted))
+	{
+		UE_LOG(LogBlueprint, Warning,
+			TEXT("ReconstructAllNodes called on incompletly loaded blueprint '%s'"),
+			Blueprint ? *Blueprint->GetFullName() : TEXT("NULL"));
+		return;
+	}
+
+	TArray<UK2Node*> AllNodes;
+	FBlueprintEditorUtils::GetAllNodesOfClass(Blueprint, AllNodes);
+
+	const bool bIsMacro = (Blueprint->BlueprintType == BPTYPE_MacroLibrary);
+	if (AllNodes.Num() > 1)
+	{
+		AllNodes.Sort(FCompareNodePriority());
+	}
+
+	for (TArray<UK2Node*>::TIterator NodeIt(AllNodes); NodeIt; ++NodeIt)
+	{
+		UK2Node* CurrentNode = *NodeIt;
+		//@todo:  Do we really need per-schema refreshing?
+		const UEdGraphSchema* Schema = CurrentNode->GetGraph()->GetSchema();
+		Schema->ReconstructNode(*CurrentNode, true);
+	}
+}
+
 void FBlueprintEditorUtils::RefreshExternalBlueprintDependencyNodes(UBlueprint* Blueprint, UStruct* RefreshOnlyChild)
 {
 	BP_SCOPED_COMPILER_EVENT_STAT(EKismetCompilerStats_RefreshExternalDependencyNodes);
@@ -1452,15 +1481,18 @@ UClass* FBlueprintEditorUtils::RegenerateBlueprintClass(UBlueprint* Blueprint, U
 			// Make sure interfaces are up to date
 			FBlueprintEditorUtils::ConformImplementedInterfaces(Blueprint);
 
-			// Refresh all nodes to make sure function signatures are up to date, etc.
-			FBlueprintEditorUtils::RefreshAllNodes(Blueprint);
+			// Reconstruct all nodes, this will call AllocateDefaultPins, which ensures
+			// that nodes have a chance to create all the pins they'll expect when they compile.
+			// A good example of why this is necessary is UK2Node_BaseAsyncTask::AllocateDefaultPins
+			// and it's companion function UK2Node_BaseAsyncTask::ExpandNode.
+			FBlueprintEditorUtils::ReconstructAllNodes(Blueprint);
 
 			// Compile the actual blueprint
 			FKismetEditorUtilities::CompileBlueprint(Blueprint, true);
 		}
 		else if( bIsMacro )
 		{
-			// Just refresh all nodes in macro blueprints, but don't recompil
+			// Just refresh all nodes in macro blueprints, but don't recompile
 			FBlueprintEditorUtils::RefreshAllNodes(Blueprint);
 
 			if (ClassToRegenerate != nullptr)
diff --git a/Engine/Source/Editor/UnrealEd/Public/Kismet2/BlueprintEditorUtils.h b/Engine/Source/Editor/UnrealEd/Public/Kismet2/BlueprintEditorUtils.h
--- a/Engine/Source/Editor/UnrealEd/Public/Kismet2/BlueprintEditorUtils.h
+++ b/Engine/Source/Editor/UnrealEd/Public/Kismet2/BlueprintEditorUtils.h
@@ -95,6 +95,11 @@ public:
 	 */
 	static void RefreshAllNodes(UBlueprint* Blueprint);
 
+	/**
+	 * Reconstructs all nodes in the blueprint, node reconstruction order determined by FCompareNodePriority.
+	 */
+	static void ReconstructAllNodes(UBlueprint* Blueprint);
+
 	/**
 	 * Optimized refresh of nodes that depend on external blueprints.  Refreshes the nodes, but does not recompile the skeleton class
 	 */
