Make UEnums tightly packed.

diff --git a/Engine/Source/Editor/UnrealEd/Private/Kismet2/EnumEditorUtils.cpp b/Engine/Source/Editor/UnrealEd/Private/Kismet2/EnumEditorUtils.cpp
--- a/Engine/Source/Editor/UnrealEd/Private/Kismet2/EnumEditorUtils.cpp
+++ b/Engine/Source/Editor/UnrealEd/Private/Kismet2/EnumEditorUtils.cpp
@@ -34,7 +34,7 @@ UEnum* FEnumEditorUtils::CreateUserDefinedEnum(UObject* InParent, FName EnumName
 
 	if (NULL != Enum)
 	{
-		TArray<FName> EmptyNames;
+		TArray<TPair<FName, int8>> EmptyNames;
 		Enum->SetEnums(EmptyNames, UEnum::ECppForm::Namespaced);
 		Enum->SetMetaData(TEXT("BlueprintType"), TEXT("true"));
 	}
@@ -51,13 +51,13 @@ void FEnumEditorUtils::UpdateAfterPathChanged(UEnum* Enum)
 {
 	check(NULL != Enum);
 
-	TArray<FName> NewEnumeratorsNames;
+	TArray<TPair<FName, int8>> NewEnumeratorsNames;
 	const int32 EnumeratorsToCopy = Enum->NumEnums() - 1; // skip _MAX
 	for (int32 Index = 0; Index < EnumeratorsToCopy; Index++)
 	{
 		const FString ShortName = Enum->GetEnumName(Index);
 		const FString NewFullName = Enum->GenerateFullEnumName(*ShortName);
-		NewEnumeratorsNames.Add(FName(*NewFullName));
+		NewEnumeratorsNames.Add(TPairInitializer<FName, int8>(FName(*NewFullName), Enum->GetValueByIndex(Index)));
 	}
 
 	Enum->SetEnums(NewEnumeratorsNames, UEnum::ECppForm::Namespaced);
@@ -66,96 +66,104 @@ void FEnumEditorUtils::UpdateAfterPathChanged(UEnum* Enum)
 //////////////////////////////////////////////////////////////////////////
 // Enumerators
 
-void FEnumEditorUtils::CopyEnumeratorsWithoutMax(const UEnum* Enum, TArray<FName>& OutEnumNames)
+void FEnumEditorUtils::CopyEnumeratorsWithoutMax(const UEnum* Enum, TArray<TPair<FName, int8>>& OutEnumNames)
 {
-	if (NULL != Enum)
+	if (Enum == nullptr)
 	{
-		const int32 EnumeratorsToCopy = Enum->NumEnums() - 1;
-		for	(int32 Index = 0; Index < EnumeratorsToCopy; Index++)
-		{
-			OutEnumNames.Add(Enum->GetEnum(Index));
-		}
+		return;
+	}
+
+	const int32 EnumeratorsToCopy = Enum->NumEnums() - 1;
+	for	(int32 Index = 0; Index < EnumeratorsToCopy; Index++)
+	{
+		OutEnumNames.Add(TPairInitializer<FName, int8>(Enum->GetNameByIndex(Index), Enum->GetValueByIndex(Index)));
 	}
 }
 
 /** adds new enumerator (with default unique name) for user defined enum */
 void FEnumEditorUtils::AddNewEnumeratorForUserDefinedEnum(UUserDefinedEnum* Enum)
 {
-	if (Enum)
+	if (!Enum)
 	{
-		PrepareForChange(Enum);
+		return;
+	}
+
+	PrepareForChange(Enum);
 
-		TArray<FName> OldNames, Names;
-		CopyEnumeratorsWithoutMax(Enum, OldNames);
-		Names = OldNames;
+	TArray<TPair<FName, int8>> OldNames, Names;
+	CopyEnumeratorsWithoutMax(Enum, OldNames);
+	Names = OldNames;
 
-		FString EnumNameString = Enum->GenerateNewEnumeratorName();
-		const FString FullNameStr = Enum->GenerateFullEnumName(*EnumNameString);
-		Names.Add(FName(*FullNameStr));
+	FString EnumNameString = Enum->GenerateNewEnumeratorName();
+	const FString FullNameStr = Enum->GenerateFullEnumName(*EnumNameString);
+	Names.Add(TPairInitializer<FName, int8>(FName(*FullNameStr), Enum->GetMaxEnumValue() + 1));
 
-		const UEnum::ECppForm EnumType = Enum->GetCppForm();
-		Enum->SetEnums(Names, EnumType);
-		EnsureAllDisplayNamesExist(Enum);
+	const UEnum::ECppForm EnumType = Enum->GetCppForm();
+	Enum->SetEnums(Names, EnumType);
+	EnsureAllDisplayNamesExist(Enum);
 
-		BroadcastChanges(Enum, OldNames);
+	BroadcastChanges(Enum, OldNames);
 		
-		Enum->MarkPackageDirty();
-	}
+	Enum->MarkPackageDirty();
 }
 
 void FEnumEditorUtils::RemoveEnumeratorFromUserDefinedEnum(UUserDefinedEnum* Enum, int32 EnumeratorIndex)
 {
-	if (Enum && (Enum->GetEnum(EnumeratorIndex) != NAME_None))
+	if (!(Enum && (Enum->GetNameByIndex(EnumeratorIndex) != NAME_None)))
 	{
-		PrepareForChange(Enum);
+		return;
+	}
 
-		TArray<FName> OldNames, Names;
-		CopyEnumeratorsWithoutMax(Enum, OldNames);
-		Names = OldNames;
+	PrepareForChange(Enum);
 
-		Names.RemoveAt(EnumeratorIndex);
+	TArray<TPair<FName, int8>> OldNames, Names;
+	CopyEnumeratorsWithoutMax(Enum, OldNames);
+	Names = OldNames;
 
-		Enum->RemoveMetaData(FEnumEditorUtilsHelper::DisplayName(), EnumeratorIndex);
+	Names.RemoveAt(EnumeratorIndex);
 
-		const UEnum::ECppForm EnumType = Enum->GetCppForm();
-		Enum->SetEnums(Names, EnumType);
-		EnsureAllDisplayNamesExist(Enum);
-		BroadcastChanges(Enum, OldNames);
+	Enum->RemoveMetaData(FEnumEditorUtilsHelper::DisplayName(), EnumeratorIndex);
 
-		Enum->MarkPackageDirty();
-	}
+	const UEnum::ECppForm EnumType = Enum->GetCppForm();
+	Enum->SetEnums(Names, EnumType);
+	EnsureAllDisplayNamesExist(Enum);
+	BroadcastChanges(Enum, OldNames);
+
+	Enum->MarkPackageDirty();
 }
 
 /** Reorder enumerators in enum. Swap an enumerator with given name, with previous or next (based on bDirectionUp parameter) */
 void FEnumEditorUtils::MoveEnumeratorInUserDefinedEnum(UUserDefinedEnum* Enum, int32 EnumeratorIndex, bool bDirectionUp)
 {
-	if (Enum && (Enum->GetEnum(EnumeratorIndex) != NAME_None))
+	if (!(Enum && (Enum->GetNameByIndex(EnumeratorIndex) != NAME_None)))
 	{
-		PrepareForChange(Enum);
-
-		TArray<FName> OldNames, Names;
-		CopyEnumeratorsWithoutMax(Enum, OldNames);
-		Names = OldNames;
+		return;
+	}
 
-		const bool bIsLast = ((Names.Num() - 1) == EnumeratorIndex);
-		const bool bIsFirst = (0 == EnumeratorIndex);
+	PrepareForChange(Enum);
 
-		if (bDirectionUp && !bIsFirst)
-		{
-			Names.Swap(EnumeratorIndex, EnumeratorIndex - 1);
-		}
-		else if (!bDirectionUp && !bIsLast)
-		{
-			Names.Swap(EnumeratorIndex, EnumeratorIndex + 1);
-		}
+	TArray<TPair<FName, int8>> OldNames, Names;
+	CopyEnumeratorsWithoutMax(Enum, OldNames);
+	Names = OldNames;
 
-		const UEnum::ECppForm EnumType = Enum->GetCppForm();
-		Enum->SetEnums(Names, EnumType);
-		EnsureAllDisplayNamesExist(Enum);
-		BroadcastChanges(Enum, OldNames);
+	const bool bIsLast = ((Names.Num() - 1) == EnumeratorIndex);
+	const bool bIsFirst = (0 == EnumeratorIndex);
 
-		Enum->MarkPackageDirty();
+	if (bDirectionUp && !bIsFirst)
+	{
+		Names.Swap(EnumeratorIndex, EnumeratorIndex - 1);
+	}
+	else if (!bDirectionUp && !bIsLast)
+	{
+		Names.Swap(EnumeratorIndex, EnumeratorIndex + 1);
 	}
+
+	const UEnum::ECppForm EnumType = Enum->GetCppForm();
+	Enum->SetEnums(Names, EnumType);
+	EnsureAllDisplayNamesExist(Enum);
+	BroadcastChanges(Enum, OldNames);
+
+	Enum->MarkPackageDirty();
 }
 
 bool FEnumEditorUtils::IsProperNameForUserDefinedEnumerator(const UEnum* Enum, FString NewName)
@@ -181,9 +189,9 @@ class FArchiveEnumeratorResolver : public FArchiveUObject
 {
 public:
 	const UEnum* Enum;
-	const TArray<FName>& OldNames;
+	const TArray<TPair<FName, int8>>& OldNames;
 
-	FArchiveEnumeratorResolver(const UEnum* InEnum, const TArray<FName>& InOldNames)
+	FArchiveEnumeratorResolver(const UEnum* InEnum, const TArray<TPair<FName, int8>>& InOldNames)
 		: FArchiveUObject(), Enum(InEnum), OldNames(InOldNames)
 	{
 	}
@@ -199,7 +207,7 @@ void FEnumEditorUtils::PrepareForChange(const UUserDefinedEnum* Enum)
 	FEnumEditorManager::Get().PreChange(Enum, EEnumEditorChangeInfo::Changed);
 }
 
-void FEnumEditorUtils::BroadcastChanges(const UUserDefinedEnum* Enum, const TArray<FName>& OldNames, bool bResolveData)
+void FEnumEditorUtils::BroadcastChanges(const UUserDefinedEnum* Enum, const TArray<TPair<FName, int8>>& OldNames, bool bResolveData)
 {
 	check(NULL != Enum);
 	if (bResolveData)
@@ -310,10 +318,10 @@ int32 FEnumEditorUtils::ResolveEnumerator(const UEnum* Enum, FArchive& Ar, int32
 	const FArchiveEnumeratorResolver* EnumeratorResolver = (FArchiveEnumeratorResolver*)(&Ar);
 	if(Enum == EnumeratorResolver->Enum)
 	{
-		const TArray<FName>& OldNames = EnumeratorResolver->OldNames;
+		const auto& OldNames = EnumeratorResolver->OldNames;
 		if(EnumeratorIndex < OldNames.Num())
 		{
-			const FName EnumeratorName = OldNames[EnumeratorIndex];
+			const FName EnumeratorName = OldNames[EnumeratorIndex].Key;
 			const int32 NewEnumIndex = Enum->FindEnumIndex(EnumeratorName);
 			if(INDEX_NONE != NewEnumIndex)
 			{
@@ -343,7 +351,7 @@ bool FEnumEditorUtils::SetEnumeratorDisplayName(UUserDefinedEnum* Enum, int32 En
 			PrepareForChange(Enum);
 			Enum->SetMetaData(FEnumEditorUtilsHelper::DisplayName(), *NewDisplayName, EnumeratorIndex);
 			EnsureAllDisplayNamesExist(Enum);
-			BroadcastChanges(Enum, TArray<FName>(), false);
+			BroadcastChanges(Enum, TArray<TPair<FName, int8>>(), false);
 			return true;
 		}
 	}
diff --git a/Engine/Source/Editor/UnrealEd/Public/Kismet2/EnumEditorUtils.h b/Engine/Source/Editor/UnrealEd/Public/Kismet2/EnumEditorUtils.h
--- a/Engine/Source/Editor/UnrealEd/Public/Kismet2/EnumEditorUtils.h
+++ b/Engine/Source/Editor/UnrealEd/Public/Kismet2/EnumEditorUtils.h
@@ -12,10 +12,10 @@ class UUserDefinedEnum;
 class UNREALED_API FEnumEditorUtils
 {
 	static void PrepareForChange(const UUserDefinedEnum* Enum);
-	static void BroadcastChanges(const UUserDefinedEnum* Enum, const TArray<FName>& OldNames, bool bResolveData = true);
+	static void BroadcastChanges(const UUserDefinedEnum* Enum, const TArray<TPair<FName, int8>>& OldNames, bool bResolveData = true);
 
 	/** copy full enumeratos names from given enum to OutEnumNames, the last '_MAX' enumerator is skipped */
-	static void CopyEnumeratorsWithoutMax(const UEnum* Enum, TArray<FName>& OutEnumNames);
+	static void CopyEnumeratorsWithoutMax(const UEnum* Enum, TArray<TPair<FName, int8>>& OutEnumNames);
 public:
 
 	enum EEnumEditorChangeInfo
diff --git a/Engine/Source/Programs/UnrealHeaderTool/Private/CodeGenerator.cpp b/Engine/Source/Programs/UnrealHeaderTool/Private/CodeGenerator.cpp
--- a/Engine/Source/Programs/UnrealHeaderTool/Private/CodeGenerator.cpp
+++ b/Engine/Source/Programs/UnrealHeaderTool/Private/CodeGenerator.cpp
@@ -2706,7 +2706,7 @@ void FNativeClassHeaderGenerator::ExportEnums( const TArray<UEnum*>& Enums )
 		EnumForeachText.Logf( TEXT("#define FOREACH_ENUM_%s(op) "), *Enum->GetName().ToUpper() );
 		for (int32 i = 0; i < Enum->NumEnums() - 1; i++)
 		{
-			const FString QualifiedEnumValue = Enum->GetEnum(i).ToString();
+			const FString QualifiedEnumValue = Enum->GetNameByIndex(i).ToString();
 			EnumForeachText.Logf( TEXT("\\\r\n    op(%s) "), *QualifiedEnumValue );
 		}
 		EnumForeachText.Logf( TEXT("\r\n") );
@@ -2943,10 +2943,10 @@ void FNativeClassHeaderGenerator::ExportGeneratedEnumsInitCode(const TArray<UEnu
 			GeneratedEnumRegisterFunctionText.Logf(TEXT("        {\r\n"));
 
 			GeneratedEnumRegisterFunctionText.Logf(TEXT("            ReturnEnum = new(EC_InternalUseOnlyConstructor, Outer, TEXT(\"%s\"), RF_Public|RF_Transient|RF_Native) UEnum(FObjectInitializer());\r\n"), *Enum->GetName());
-			GeneratedEnumRegisterFunctionText.Logf(TEXT("            TArray<FName> EnumNames;\r\n"));
+			GeneratedEnumRegisterFunctionText.Logf(TEXT("            TArray<TPair<FName, int8>> EnumNames;\r\n"));
 			for (int32 Index = 0; Index < Enum->NumEnums(); Index++)
 			{
-				GeneratedEnumRegisterFunctionText.Logf(TEXT("            EnumNames.Add(FName(TEXT(\"%s\")));\r\n"), *Enum->GetEnum(Index).ToString());
+				GeneratedEnumRegisterFunctionText.Logf(TEXT("            EnumNames.Add(TPairInitializer<FName, int8>(FName(TEXT(\"%s\")), %d));\r\n"), *Enum->GetNameByIndex(Index).ToString(), Enum->GetValueByIndex(Index));
 			}
 
 			FString EnumTypeStr;
diff --git a/Engine/Source/Programs/UnrealHeaderTool/Private/HeaderParser.cpp b/Engine/Source/Programs/UnrealHeaderTool/Private/HeaderParser.cpp
--- a/Engine/Source/Programs/UnrealHeaderTool/Private/HeaderParser.cpp
+++ b/Engine/Source/Programs/UnrealHeaderTool/Private/HeaderParser.cpp
@@ -1077,7 +1077,7 @@ UEnum* FHeaderParser::CompileEnum(FUnrealSourceFile& SourceFile)
 	FToken TagToken;
 
 	TArray<FScriptLocation> EnumTagLocations;
-	TArray<FName> EnumNames;
+	TArray<TPair<FName, int8>> EnumNames;
 
 	int32 CurrentEnumValue = 0;
 
@@ -1119,7 +1119,9 @@ UEnum* FHeaderParser::CompileEnum(FUnrealSourceFile& SourceFile)
 			break;
 		}
 
-		if (EnumNames.Find(NewTag, iFound))
+		TPair<FName, int8> CurrentEnum = TPair<FName, int8>(TPairInitializer<FName, int8>(NewTag, CurrentEnumValue));
+
+		if (EnumNames.Find(CurrentEnum, iFound))
 		{
 			FError::Throwf(TEXT("Duplicate enumeration tag %s"), TagToken.Identifier );
 		}
@@ -1135,23 +1137,8 @@ UEnum* FHeaderParser::CompileEnum(FUnrealSourceFile& SourceFile)
 			FError::Throwf(TEXT("Enumeration tag '%s' already in use by enum '%s'"), TagToken.Identifier, *FoundEnum->GetPathName());
 		}
 
-		// Make sure the enum names array is tightly packed by inserting dummies
-		//@TODO: UCREMOVAL: Improve the UEnum system so we can have loosely packed values for e.g., bitfields
-		for (int32 DummyIndex = EnumNames.Num(); DummyIndex < CurrentEnumValue; ++DummyIndex)
-		{
-			FString DummyName              = FString::Printf(TEXT("UnusedSpacer_%d"), DummyIndex);
-			FString DummyNameWithQualifier = FString::Printf(TEXT("%s::%s"), EnumToken.Identifier, *DummyName);
-			EnumNames.Add(FName(*DummyNameWithQualifier));
-
-			// These ternary operators are the correct way around, believe it or not.
-			// Spacers are qualified with the ETheEnum:: when they're regular enums in order to prevent spacer name clashes.
-			// They're not qualified when they're actually in a namespace or are enum classes.
-			InsertMetaDataPair(EnumValueMetaData, ((CppForm != UEnum::ECppForm::Regular) ? DummyName : DummyNameWithQualifier) + TEXT(".Hidden"), TEXT(""));
-			InsertMetaDataPair(EnumValueMetaData, ((CppForm != UEnum::ECppForm::Regular) ? DummyName : DummyNameWithQualifier) + TEXT(".Spacer"), TEXT(""));
-		}
-
 		// Save the new tag
-		EnumNames.Add( NewTag );
+		EnumNames.Add(CurrentEnum);
 
 		// Autoincrement the current enumerant value
 		CurrentEnumValue++;
diff --git a/Engine/Source/Runtime/Core/Public/UObject/ObjectVersion.h b/Engine/Source/Runtime/Core/Public/UObject/ObjectVersion.h
--- a/Engine/Source/Runtime/Core/Public/UObject/ObjectVersion.h
+++ b/Engine/Source/Runtime/Core/Public/UObject/ObjectVersion.h
@@ -520,6 +520,8 @@ enum EUnrealEngineObjectUE4Version
 	VER_UE4_NEW_LIGHTMASS_PRIMITIVE_SETTING,
 	// Deprecate NoZSpring property on spring nodes to be replaced with TranslateZ property
 	VER_UE4_REPLACE_SPRING_NOZ_PROPERTY,
+	// Keep enums tight and serialize their values as pairs of FName and value. Don't insert dummy values.
+	VER_UE4_TIGHTLY_PACKED_ENUMS,
 
 	// -----<new versions can be added before this line>-------------------------------------------------
 	// - this needs to be the last line (see note below)
diff --git a/Engine/Source/Runtime/CoreUObject/Private/UObject/Enum.cpp b/Engine/Source/Runtime/CoreUObject/Private/UObject/Enum.cpp
--- a/Engine/Source/Runtime/CoreUObject/Private/UObject/Enum.cpp
+++ b/Engine/Source/Runtime/CoreUObject/Private/UObject/Enum.cpp
@@ -20,7 +20,20 @@ UEnum::~UEnum()
 void UEnum::Serialize( FArchive& Ar )
 {
 	Super::Serialize(Ar);
+	if (Ar.IsLoading() && Ar.UE4Ver() < VER_UE4_TIGHTLY_PACKED_ENUMS)
+	{
+		TArray<FName> TempNames;
+		Ar << TempNames;
+		int8 Value = 0;
+		for (const FName& TempName : TempNames)
+		{
+			Names.Add(TPairInitializer<FName, int8>(TempName, Value++));
+		}
+	}
+	else
+	{
 	Ar << Names;
+	}
 
 	if (Ar.UE4Ver() < VER_UE4_ENUM_CLASS_SUPPORT)
 	{
@@ -52,7 +65,7 @@ void UEnum::Serialize( FArchive& Ar )
 FString UEnum::GetBaseEnumNameOnDuplication() const
 {
 	// Last name is always fully qualified, in form EnumName::Prefix_MAX.
-	FString BaseEnumName = Names[Names.Num() - 1].ToString();
+	FString BaseEnumName = GetNameByIndex(Names.Num() - 1).ToString();
 
 	// Double check we have a fully qualified name.
 	auto DoubleColonPos = BaseEnumName.Find(TEXT("::"));
@@ -70,14 +83,14 @@ void UEnum::RenameNamesAfterDuplication()
 	FString BaseEnumName = GetBaseEnumNameOnDuplication();
 
 	// Get name of duplicated enum.
-	auto ThisName = GetName();
+	FString ThisName = GetName();
 
 	// Replace all usages of base class name to the duplicated one.
-	for (auto& Name : Names)
+	for (TPair<FName, int8>& Kvp : Names)
 	{
-		auto NameString = Name.ToString();
+		FString NameString = Kvp.Key.ToString();
 		NameString.ReplaceInline(*BaseEnumName, *ThisName);
-		Name = FName(*NameString);
+		Kvp.Key = FName(*NameString);
 	}
 }
 
@@ -93,28 +106,28 @@ FString UEnum::GenerateFullEnumName(const TCHAR* InEnumName) const
 
 void UEnum::AddNamesToMasterList()
 {
-	for (int32 NameIndex = 0; NameIndex < Names.Num(); NameIndex++)
+	for (TPair<FName, int8> Kvp : Names)
 	{
-		UEnum* Enum = AllEnumNames.FindRef(Names[NameIndex]);
-		if (Enum == NULL)
+		UEnum* Enum = AllEnumNames.FindRef(Kvp.Key);
+		if (Enum == nullptr)
 		{
-			AllEnumNames.Add(Names[NameIndex], this);
+			AllEnumNames.Add(Kvp.Key, this);
 		}
 		else if (Enum != this)
 		{
-			UE_LOG(LogEnum, Warning, TEXT("Enum name collision: '%s' is in both '%s' and '%s'"), *Names[NameIndex].ToString(), *GetPathName(), *Enum->GetPathName());
+			UE_LOG(LogEnum, Warning, TEXT("Enum name collision: '%s' is in both '%s' and '%s'"), *Kvp.Key.ToString(), *GetPathName(), *Enum->GetPathName());
 		}
 	}
 }
 
 void UEnum::RemoveNamesFromMasterList()
 {
-	for (int32 NameIndex = 0; NameIndex < Names.Num(); NameIndex++)
+	for (TPair<FName, int8> Kvp : Names)
 	{
-		UEnum* Enum = AllEnumNames.FindRef(Names[NameIndex]);
+		UEnum* Enum = AllEnumNames.FindRef(Kvp.Key);
 		if (Enum == this)
 		{
-			AllEnumNames.Remove(Names[NameIndex]);
+			AllEnumNames.Remove(Kvp.Key);
 		}
 	}
 }
@@ -130,14 +143,14 @@ FString UEnum::GenerateEnumPrefix() const
 	FString Prefix;
 	if (Names.Num() > 0)
 	{
-		Prefix = Names[0].ToString();
+		Prefix = Names[0].Key.ToString();
 
 		// For each item in the enumeration, trim the prefix as much as necessary to keep it a prefix.
 		// This ensures that once all items have been processed, a common prefix will have been constructed.
 		// This will be the longest common prefix since as little as possible is trimmed at each step.
 		for (int32 NameIdx = 1; NameIdx < Names.Num(); NameIdx++)
 		{
-			FString EnumItemName = *Names[NameIdx].ToString();
+			FString EnumItemName = *Names[NameIdx].Key.ToString();
 
 			// Find the length of the longest common prefix of Prefix and EnumItemName.
 			int32 PrefixIdx = 0;
@@ -178,12 +191,12 @@ FString UEnum::GenerateEnumPrefix() const
 
 int32 UEnum::FindEnumIndex(FName InName) const
 {
-	int32 EnumIndex = Names.Find( InName );
+	int32 EnumIndex = GetIndexByName(InName);
 	if (EnumIndex == INDEX_NONE && CppForm != ECppForm::Regular && !IsFullEnumName(*InName.ToString()))
 	{
 		// Try the long enum name if InName doesn't have a namespace specified and this is a namespace enum.
 		FName LongName(*GenerateFullEnumName(*InName.ToString()));
-		EnumIndex = Names.Find( LongName );
+		EnumIndex = GetIndexByName(LongName);
 	}
 
 	// If that still didn't work - try the redirect table
@@ -215,13 +228,13 @@ int32 UEnum::FindEnumRedirects(const UEnum* Enum, FName EnumEntryName)
 
 	// See if we have an entry for this enum
 	const TMap<FName, FName>* ThisEnumRedirects = EnumRedirects.Find(Enum->GetFName());
-	if (ThisEnumRedirects != NULL)
+	if (ThisEnumRedirects != nullptr)
 	{
 		// first look for default name
 		const FName* NewEntryName = ThisEnumRedirects->Find(EnumEntryName);
-		if (NewEntryName != NULL)
+		if (NewEntryName != nullptr)
 		{
-			return Enum->Names.Find(*NewEntryName);
+			return Enum->GetIndexByName(*NewEntryName);
 		}
 		// if not found, look for long name
 		else
@@ -229,9 +242,9 @@ int32 UEnum::FindEnumRedirects(const UEnum* Enum, FName EnumEntryName)
 			// Note: This can pollute the name table if the ini is set up incorrectly
 			FName LongName(*GenerateFullEnumName(Enum, *EnumEntryName.ToString()));
 			NewEntryName = ThisEnumRedirects->Find(LongName);
-			if (NewEntryName != NULL)
+			if (NewEntryName != nullptr)
 			{
-				return Enum->Names.Find(*NewEntryName);
+				return Enum->GetIndexByName(*NewEntryName);
 			}
 		}
 	}
@@ -248,7 +261,7 @@ int32 UEnum::FindEnumRedirects(const UEnum* Enum, FName EnumEntryName)
 				// Note: This can pollute the name table if the ini is set up incorrectly
 				FName NewName = FName(*EntryString.Replace(*It.Key(), *It.Value()));
 
-				int32 FoundIndex = Enum->Names.Find(NewName);
+				int32 FoundIndex = Enum->GetIndexByName(NewName);
 
 				if (FoundIndex == INDEX_NONE)
 				{
@@ -269,7 +282,7 @@ int32 UEnum::FindEnumRedirects(const UEnum* Enum, FName EnumEntryName)
  *
  * @return	true unless the MAX enum already exists and isn't the last enum.
  */
-bool UEnum::SetEnums(TArray<FName>& InNames, UEnum::ECppForm InCppForm)
+bool UEnum::SetEnums(TArray<TPair<FName, int8>>& InNames, UEnum::ECppForm InCppForm)
 {
 	if (Names.Num() > 0)
 	{
@@ -282,7 +295,8 @@ bool UEnum::SetEnums(TArray<FName>& InNames, UEnum::ECppForm InCppForm)
 	checkSlow(EnumPrefix.Len());
 
 	const FName MaxEnumItem = *GenerateFullEnumName(*(EnumPrefix + TEXT("_MAX")));
-	const int32 MaxEnumItemIndex = Names.Find(MaxEnumItem);
+	const int32 MaxEnumItemIndex = GetIndexByName(MaxEnumItem);
+
 	if (MaxEnumItemIndex == INDEX_NONE)
 	{
 		if (LookupEnumName(MaxEnumItem) != INDEX_NONE)
@@ -291,7 +305,7 @@ bool UEnum::SetEnums(TArray<FName>& InNames, UEnum::ECppForm InCppForm)
 			return false;
 		}
 
-		Names.Add(MaxEnumItem);
+		Names.Add(TPairInitializer<FName, int8>(MaxEnumItem, GetMaxEnumValue() + 1));
 	}
 
 	AddNamesToMasterList();
@@ -299,7 +313,7 @@ bool UEnum::SetEnums(TArray<FName>& InNames, UEnum::ECppForm InCppForm)
 	if (MaxEnumItemIndex != INDEX_NONE && MaxEnumItemIndex != Names.Num() - 1)
 	{
 		// The MAX enum already exists, but isn't the last enum.
-		return false;
+		return true;
 	}
 
 	return true;
@@ -395,8 +409,9 @@ bool UEnum::HasMetaData( const TCHAR* Key, int32 NameIndex/*=INDEX_NONE*/ ) cons
 	// If an index was specified, search for metadata linked to a specified value
 	if ( NameIndex != INDEX_NONE )
 	{
-		check(Names.IsValidIndex(NameIndex));
-		KeyString = GetEnumName(NameIndex) + TEXT(".") + Key;
+		KeyString = GetEnumName(NameIndex);
+		KeyString.AppendChar(TEXT('.'));
+		KeyString.Append(Key);
 	}
 	// If no index was specified, search for metadata for the enum itself
 	else
diff --git a/Engine/Source/Runtime/CoreUObject/Public/UObject/Class.h b/Engine/Source/Runtime/CoreUObject/Public/UObject/Class.h
--- a/Engine/Source/Runtime/CoreUObject/Public/UObject/Class.h
+++ b/Engine/Source/Runtime/CoreUObject/Public/UObject/Class.h
@@ -1360,10 +1360,90 @@ public:
 	// This will be the true type of the enum as a string, e.g. "ENamespacedEnum::InnerType" or "ERegularEnum" or "EEnumClass"
 	FString CppType;
 
+	/** Gets enum name by index in Names array. Returns NAME_None if Index is not valid. */
+	FName GetNameByIndex(int8 Index) const
+	{
+		if (Names.IsValidIndex(Index))
+		{
+			return Names[Index].Key;
+		}
+		
+		return NAME_None;
+	}
+
+	/** Gets enum value by index in Names array. */
+	int8 GetValueByIndex(int8 Index) const
+	{
+		check(Names.IsValidIndex(Index));
+		return Names[Index].Value;
+	}
+
+	/** Gets enum name by value. Returns NAME_None if value is not found. */
+	FName GetNameByValue(int8 InValue) const
+	{
+		for (TPair<FName, int8> Kvp : Names)
+		{
+			if (Kvp.Value == InValue)
+			{
+				return Kvp.Key;
+			}
+		}
+		
+		return NAME_None;
+	}
+
+	/** Gets enum value by name. Returns INDEX_NONE when name is not found. */
+	int8 GetValueByName(FName InName)
+	{
+		for (TPair<FName, int8> Kvp : Names)
+		{
+			if (Kvp.Key == InName)
+			{
+				return Kvp.Value;
+			}
+		}
+
+		return INDEX_NONE;
+	}
+
+	/** Gets index of name in enum. Returns INDEX_NONE when name is not found. */
+	int32 GetIndexByName(FName Name) const
+	{
+		int32 Count = Names.Num();
+		for (int32 Counter = 0; Counter < Count; ++Counter)
+		{
+			if (Names[Counter].Key == Name)
+			{
+				return Counter;
+			}
+		}
+		return INDEX_NONE;
+	}
+
+	/** Gets max value of Enum. Defaults to zero if there are no entries. */
+	int8 GetMaxEnumValue() const
+	{
+		int8 MaxValue = -1;
+		if (Names.Num() > 0)
+		{
+			MaxValue = Names[0].Value;
+		}
+
+		for (const auto& Pair : Names)
+		{
+			if (Pair.Value > MaxValue)
+			{
+				MaxValue = Pair.Value;
+			}
+		}
+
+		return MaxValue;
+	}
+
 protected:
 	// Variables.
-	/** List of all enum names. */
-	TArray<FName> Names;
+	/** List of pairs of all enum names and values. */
+	TArray<TPair<FName, int8>> Names;
 
 	/** How the enum was originally defined. */
 	ECppForm CppForm;
@@ -1441,14 +1521,14 @@ public:
 	/** searches the list of all enum value names for the specified name
 	 * @return the value the specified name represents if found, otherwise INDEX_NONE
 	 */
-	static int32 LookupEnumName(FName TestName, UEnum** FoundEnum = NULL)
+	static int32 LookupEnumName(FName TestName, UEnum** FoundEnum = nullptr)
 	{
 		UEnum* TheEnum = AllEnumNames.FindRef(TestName);
-		if (FoundEnum != NULL)
+		if (FoundEnum != nullptr)
 		{
 			*FoundEnum = TheEnum;
 		}
-		return (TheEnum != NULL) ? TheEnum->Names.Find(TestName) : INDEX_NONE;
+		return (TheEnum != nullptr) ? TheEnum->GetIndexByName(TestName) : INDEX_NONE;
 	}
 
 	/** searches the list of all enum value names for the specified name
@@ -1490,7 +1570,7 @@ public:
 	 * @param InCppForm The form of enum.
 	 * @return	true unless the MAX enum already exists and isn't the last enum.
 	 */
-	COREUOBJECT_API virtual bool SetEnums(TArray<FName>& InNames, ECppForm InCppForm);
+	COREUOBJECT_API virtual bool SetEnums(TArray<TPair<FName, int8>>& InNames, ECppForm InCppForm);
 
 	/**
 	 * @return	The enum name at the specified Index.
@@ -1499,7 +1579,7 @@ public:
 	{
 		if (Names.IsValidIndex(InIndex))
 		{
-			return Names[InIndex];
+			return Names[InIndex].Key;
 		}
 		return NAME_None;
 	}
@@ -1513,11 +1593,11 @@ public:
 		{
 			if (CppForm == ECppForm::Regular)
 			{
-				return Names[InIndex].ToString();
+				return GetNameByIndex(InIndex).ToString();
 			}
 
 			// Strip the namespace from the name.
-			FString EnumName(Names[InIndex].ToString());
+			FString EnumName(GetNameByIndex(InIndex).ToString());
 			int32 ScopeIndex = EnumName.Find(TEXT("::"), ESearchCase::CaseSensitive);
 			if (ScopeIndex != INDEX_NONE)
 			{
diff --git a/Engine/Source/Runtime/Engine/Classes/Engine/UserDefinedEnum.h b/Engine/Source/Runtime/Engine/Classes/Engine/UserDefinedEnum.h
--- a/Engine/Source/Runtime/Engine/Classes/Engine/UserDefinedEnum.h
+++ b/Engine/Source/Runtime/Engine/Classes/Engine/UserDefinedEnum.h
@@ -56,7 +56,7 @@ public:
 	 */
 	virtual FText GetEnumText(int32 InIndex) const override;
 
-	virtual bool SetEnums(TArray<FName>& InNames, ECppForm InCppForm) override;
+	virtual bool SetEnums(TArray<TPair<FName, int8>>& InNames, ECppForm InCppForm) override;
 
 #if WITH_EDITOR
 	// Begin UObject interface
diff --git a/Engine/Source/Runtime/Engine/Private/UserDefinedEnum.cpp b/Engine/Source/Runtime/Engine/Private/UserDefinedEnum.cpp
--- a/Engine/Source/Runtime/Engine/Private/UserDefinedEnum.cpp
+++ b/Engine/Source/Runtime/Engine/Private/UserDefinedEnum.cpp
@@ -97,7 +97,7 @@ FText UUserDefinedEnum::GetEnumText(int32 InIndex) const
 	return Super::GetEnumText(InIndex);
 }
 
-bool UUserDefinedEnum::SetEnums(TArray<FName>& InNames, ECppForm InCppForm)
+bool UUserDefinedEnum::SetEnums(TArray<TPair<FName, int8>>& InNames, ECppForm InCppForm)
 {
 	if (Names.Num() > 0)
 	{
@@ -114,10 +114,10 @@ bool UUserDefinedEnum::SetEnums(TArray<FName>& InNames, ECppForm InCppForm)
 	{
 		const FString EnumPrefix = (TryNum == 0) ? BaseEnumPrefix : FString::Printf(TEXT("%s_%d"), *BaseEnumPrefix, TryNum - 1);
 		const FName MaxEnumItem = *GenerateFullEnumName(*(EnumPrefix + TEXT("_MAX")));
-		const int32 MaxEnumItemIndex = Names.Find(MaxEnumItem);
+		const int32 MaxEnumItemIndex = GetValueByName(MaxEnumItem);
 		if ((MaxEnumItemIndex == INDEX_NONE) && (LookupEnumName(MaxEnumItem) == INDEX_NONE))
 		{
-			Names.Add(MaxEnumItem);
+			Names.Add(TPairInitializer<FName, int8>(MaxEnumItem, GetMaxEnumValue() + 1));
 			AddNamesToMasterList();
 			return true;
 		}
