Dont use ToolTipText till crash is fixed

diff --git a/Engine/Source/Editor/UnrealEd/Private/Toolkits/AssetEditorToolkit.cpp b/Engine/Source/Editor/UnrealEd/Private/Toolkits/AssetEditorToolkit.cpp
--- a/Engine/Source/Editor/UnrealEd/Private/Toolkits/AssetEditorToolkit.cpp
+++ b/Engine/Source/Editor/UnrealEd/Private/Toolkits/AssetEditorToolkit.cpp
@@ -88,7 +88,7 @@ void FAssetEditorToolkit::InitAssetEditor( const EToolkitMode::Type Mode, const
 		NewMajorTab = SNew(SDockTab) 
 			.ContentPadding(0.0f) 
 			.TabRole(ETabRole::MajorTab)
-			.ToolTip( IDocumentation::Get()->CreateToolTip( ToolTipText, nullptr, DocLink, GetToolkitFName().ToString() ) ) 
+			.ToolTip(IDocumentation::Get()->CreateToolTip(Label, nullptr, DocLink, GetToolkitFName().ToString()))
 			.Icon( this, &FAssetEditorToolkit::GetDefaultTabIcon )
 			.Label( Label );
 
