PR #1185: Fixed issue with animation sequence reimport not respecting 'Default Sample Rate' flag set by the import UI. (Contributed by [REDACTED_PERSON])

diff --git a/Engine/Source/Editor/UnrealEd/Classes/Factories/FbxAnimSequenceImportData.h b/Engine/Source/Editor/UnrealEd/Classes/Factories/FbxAnimSequenceImportData.h
--- a/Engine/Source/Editor/UnrealEd/Classes/Factories/FbxAnimSequenceImportData.h
+++ b/Engine/Source/Editor/UnrealEd/Classes/Factories/FbxAnimSequenceImportData.h
@@ -40,6 +40,10 @@ class UNREALED_API UFbxAnimSequenceImportData : public UFbxAssetImportData
 	UPROPERTY(EditAnywhere, AdvancedDisplay, Category=ImportSettings, meta=(DisplayName = "End Frame"))
 	int32	EndFrame;
 
+	/** Enable this option to use default sample rate for the imported animation at 30 frames per second */
+	UPROPERTY(EditAnywhere, AdvancedDisplay, config, Category = ImportSettings, meta = (ToolTip = "If enabled, samples all animation curves to 30 FPS"))
+	bool bUseDefaultSampleRate;
+
 	/** Name of source animation that was imported, used to reimport correct animation from the FBX file*/
 	UPROPERTY()
 	FString SourceAnimationName;
diff --git a/Engine/Source/Editor/UnrealEd/Classes/Factories/FbxImportUI.h b/Engine/Source/Editor/UnrealEd/Classes/Factories/FbxImportUI.h
--- a/Engine/Source/Editor/UnrealEd/Classes/Factories/FbxImportUI.h
+++ b/Engine/Source/Editor/UnrealEd/Classes/Factories/FbxImportUI.h
@@ -83,10 +83,6 @@ class UFbxImportUI : public UObject
 	UPROPERTY()
 	uint32 bImportRigidMesh:1;
 
-	/** Enable this option to use default sample rate for the imported animation at 30 frames per second */
-	UPROPERTY(EditAnywhere, AdvancedDisplay, config, Category=Animation, meta=(editcondition = "bImportAnimations", ToolTip="If enabled, samples all animation curves to 30 FPS", ImportType="SkeletalMesh|Animation"))
-	uint32 bUseDefaultSampleRate:1;
-
 	/** Whether to automatically create Unreal materials for materials found in the FBX scene */
 	UPROPERTY(EditAnywhere, config, Category = Material, meta = (OBJRestrict="true"))
 	uint32 bImportMaterials:1;
diff --git a/Engine/Source/Editor/UnrealEd/Private/Factories/EditorFactories.cpp b/Engine/Source/Editor/UnrealEd/Private/Factories/EditorFactories.cpp
--- a/Engine/Source/Editor/UnrealEd/Private/Factories/EditorFactories.cpp
+++ b/Engine/Source/Editor/UnrealEd/Private/Factories/EditorFactories.cpp
@@ -5871,7 +5871,6 @@ EReimportResult::Type UReimportFbxSkeletalMeshFactory::Reimport( UObject* Obj )
 	ReimportUI->bImportAnimations = false;
 	ReimportUI->AnimationName = TEXT("");
 	ReimportUI->bImportRigidMesh = false;
-	ReimportUI->bUseDefaultSampleRate = false;
 
 	bool bSuccess = false;
 
diff --git a/Engine/Source/Editor/UnrealEd/Private/Fbx/FbxMainImport.cpp b/Engine/Source/Editor/UnrealEd/Private/Fbx/FbxMainImport.cpp
--- a/Engine/Source/Editor/UnrealEd/Private/Fbx/FbxMainImport.cpp
+++ b/Engine/Source/Editor/UnrealEd/Private/Fbx/FbxMainImport.cpp
@@ -220,8 +220,6 @@ void ApplyImportUIToImportOptions(UFbxImportUI* ImportUI, FBXImportOptions& InOu
 		InOutImportOptions.ImportUniformScale	= AnimData->ImportUniformScale;
 	}
 
-	// only re-sample if they don't want to use default sample rate
-	InOutImportOptions.bResample = ImportUI->bUseDefaultSampleRate==false;
 	InOutImportOptions.bImportMorph = ImportUI->SkeletalMeshImportData->bImportMorphTargets;
 	InOutImportOptions.bUpdateSkeletonReferencePose = ImportUI->SkeletalMeshImportData->bUpdateSkeletonReferencePose;
 	InOutImportOptions.bImportRigidMesh = ImportUI->OriginalImportType == FBXIT_StaticMesh && ImportUI->MeshTypeToImport == FBXIT_SkeletalMesh;
@@ -244,6 +242,8 @@ void ApplyImportUIToImportOptions(UFbxImportUI* ImportUI, FBXImportOptions& InOu
 	InOutImportOptions.AnimationRange.X = ImportUI->AnimSequenceImportData->StartFrame;
 	InOutImportOptions.AnimationRange.Y = ImportUI->AnimSequenceImportData->EndFrame;
 	InOutImportOptions.AnimationName = ImportUI->AnimationName;
+	// only re-sample if they don't want to use default sample rate
+	InOutImportOptions.bResample = !ImportUI->AnimSequenceImportData->bUseDefaultSampleRate;
 	InOutImportOptions.bPreserveLocalTransform = ImportUI->AnimSequenceImportData->bPreserveLocalTransform;
 	InOutImportOptions.bDeleteExistingMorphTargetCurves = ImportUI->AnimSequenceImportData->bDeleteExistingMorphTargetCurves;
 	InOutImportOptions.bImportCustomAttribute = ImportUI->AnimSequenceImportData->bImportCustomAttribute;
diff --git a/Engine/Source/Editor/UnrealEd/Private/SkeletalMeshEdit.cpp b/Engine/Source/Editor/UnrealEd/Private/SkeletalMeshEdit.cpp
--- a/Engine/Source/Editor/UnrealEd/Private/SkeletalMeshEdit.cpp
+++ b/Engine/Source/Editor/UnrealEd/Private/SkeletalMeshEdit.cpp
@@ -185,7 +185,11 @@ bool UEditorEngine::ReimportFbxAnimation( USkeleton* Skeleton, UAnimSequence* An
 			if (CurAnimStack)
 			{
 				// set current anim stack
-				int32 ResampleRate = FbxImporter->GetMaxSampleRate(SortedLinks, FBXMeshNodeArray);
+				int32 ResampleRate = DEFAULT_SAMPLERATE;
+				if (FbxImporter->ImportOptions->bResample)
+				{
+					ResampleRate = FbxImporter->GetMaxSampleRate(SortedLinks, FBXMeshNodeArray);
+				}
 				FbxTimeSpan AnimTimeSpan = FbxImporter->GetAnimationTimeSpan(SortedLinks[0], CurAnimStack);
 				// for now it's not importing morph - in the future, this should be optional or saved with asset
 				if (FbxImporter->ValidateAnimStack(SortedLinks, FBXMeshNodeArray, CurAnimStack, ResampleRate, bImportMorphTracks, AnimTimeSpan))
