When using launch on request that the map package be saved before starting cook.
To force blueprints to be recompiled.

diff --git a/Engine/Source/Editor/UnrealEd/Private/PlayLevel.cpp b/Engine/Source/Editor/UnrealEd/Private/PlayLevel.cpp
--- a/Engine/Source/Editor/UnrealEd/Private/PlayLevel.cpp
+++ b/Engine/Source/Editor/UnrealEd/Private/PlayLevel.cpp
@@ -1528,23 +1528,26 @@ void UEditorEngine::PlayUsingLauncher()
 		FWorldContext & EditorContext = GetEditorWorldContext();
 		if (EditorContext.World()->WorldComposition || (LauncherProfile->GetCookMode() == ELauncherProfileCookModes::ByTheBookInEditor) || (LauncherProfile->GetCookMode() == ELauncherProfileCookModes::OnTheFlyInEditor) )
 		{
-			// Open world composition from original folder
-			// Or if using by book in editor don't need to resave the package just cook it by the book 
-			FString MapName = EditorContext.World()->GetOutermost()->GetName();
-			MapNames.Add(MapName);
-
-
 			// Daniel: Only reason we actually need to save any packages is because if a new package is created it won't be on disk yet and CookOnTheFly will early out if the package doesn't exist (even though it could be in memory and not require loading at all)
 			//			future me can optimize this by either adding extra allowances to CookOnTheFlyServer code or only saving packages which doesn't exist if it becomes a problem
 			// if this returns false, it means we should stop what we're doing and return to the editor
 			bool bPromptUserToSave = true;
-			bool bSaveMapPackages = false;
+			bool bSaveMapPackages = true;
 			bool bSaveContentPackages = true;
 			if (!FEditorFileUtils::SaveDirtyPackages(bPromptUserToSave, bSaveMapPackages, bSaveContentPackages))
 			{
 				CancelRequestPlaySession();
 				return;
 			}
+
+
+			// Open world composition from original folder
+			// Or if using by book in editor don't need to resave the package just cook it by the book 
+			FString MapName = EditorContext.World()->GetOutermost()->GetName();
+			MapNames.Add(MapName);
+
+
+			
 		}
 		else
 		{
