These RefreshNode calls are no longer needed... They add noise to my profiles of blueprint loading.

diff --git a/Engine/Source/Editor/UnrealEd/Private/Tests/BlueprintAutomationTests.cpp b/Engine/Source/Editor/UnrealEd/Private/Tests/BlueprintAutomationTests.cpp
--- a/Engine/Source/Editor/UnrealEd/Private/Tests/BlueprintAutomationTests.cpp
+++ b/Engine/Source/Editor/UnrealEd/Private/Tests/BlueprintAutomationTests.cpp
@@ -478,7 +478,6 @@ public:
 
 		bool bIsRegeneratingOnLoad = false;
 		bool bSkipGarbageCollection = true;
-		FBlueprintEditorUtils::RefreshAllNodes(BlueprintObj);
 		FKismetEditorUtilities::CompileBlueprint(BlueprintObj, bIsRegeneratingOnLoad, bSkipGarbageCollection);
 
 		if (BlueprintPackage != nullptr)
@@ -963,11 +962,6 @@ bool FBlueprintCompileOnLoadTest::RunTest(const FString& BlueprintAssetPath)
 				FArchiveReplaceObjectRef<UObject>(Subobj, ClassRedirects, /*bNullPrivateRefs=*/false, /*bIgnoreOuterRef=*/true, /*bIgnoreArchetypeRef=*/false);
 			}
 		}
-
-		UPackage* AssetPackage = ReloadedBlueprint->GetOutermost();
-		bool bHasUnsavedChanges = AssetPackage->IsDirty();
-		FBlueprintEditorUtils::RefreshAllNodes(ReloadedBlueprint);
-		AssetPackage->SetDirtyFlag(bHasUnsavedChanges);
 	}
 
 	// look for diffs between subsequent loads and log them as errors
