Made filter criteria for Blueprint reference replacement less strict.

diff --git a/Engine/Source/Editor/UnrealEd/Private/AssetDeleteModel.cpp b/Engine/Source/Editor/UnrealEd/Private/AssetDeleteModel.cpp
--- a/Engine/Source/Editor/UnrealEd/Private/AssetDeleteModel.cpp
+++ b/Engine/Source/Editor/UnrealEd/Private/AssetDeleteModel.cpp
@@ -9,6 +9,7 @@
 #include "AutoReimport/AutoReimportUtilities.h"
 #include "AutoReimport/AutoReimportManager.h"
 #include "ISourceControlModule.h"
+#include "BlueprintEditorUtils.h"
 
 #define LOCTEXT_NAMESPACE "FAssetDeleteModel"
 
@@ -319,20 +320,23 @@ bool FAssetDeleteModel::CanReplaceReferencesWith( const FAssetData& InAssetData
 	// Filter out blueprints of different types
 	if (FirstPendingDelete->IsChildOf(UBlueprint::StaticClass()) && AssetDataClass->IsChildOf(UBlueprint::StaticClass()))
 	{
-		// Get BP parent classes
+		// Get BP native parent classes
 		static const FName ParentClassName("ParentClass");
-		const UClass* OriginalBPClass = CastChecked<UBlueprint>(PendingDeletes[0]->GetObject())->ParentClass;
+		UClass* OriginalBPParentClass = CastChecked<UBlueprint>(PendingDeletes[0]->GetObject())->ParentClass;
 		const FString* BPClassNameToTest = InAssetData.TagsAndValues.Find(ParentClassName);
 
 		if (BPClassNameToTest && !BPClassNameToTest->IsEmpty())
 		{
-			UClass* ParentClass = FindObject<UClass>(NULL, **BPClassNameToTest);
-			if (!ParentClass)
+			UClass* ParentClassToTest = FindObject<UClass>(ANY_PACKAGE, **BPClassNameToTest);
+			if (!ParentClassToTest)
 			{
-				ParentClass = LoadObject<UClass>(NULL, **BPClassNameToTest);
+				ParentClassToTest = LoadObject<UClass>(nullptr, **BPClassNameToTest);
 			}
 
-			if (!ParentClass || !ParentClass->IsChildOf(OriginalBPClass))
+			UClass* NativeParentClassToReplace = FBlueprintEditorUtils::FindFirstNativeClass(OriginalBPParentClass);
+			UClass* NativeParentClassToTest = FBlueprintEditorUtils::FindFirstNativeClass(ParentClassToTest);
+
+			if (!NativeParentClassToTest || !NativeParentClassToTest->IsChildOf(NativeParentClassToReplace))
 			{
 				// If we couldn't determine the asset parent class (e.g. because the ParentClass tag wasn't present in the FAssetData), or if
 				// the asset parent class wasn't equal to or derived from the pending delete BP class, filter i
