PR #1188: Some fixes for AutomatedMapBuilds (Contributed by [REDACTED_PERSON])

diff --git a/Engine/Source/Editor/UnrealEd/Private/EditorBuildUtils.cpp b/Engine/Source/Editor/UnrealEd/Private/EditorBuildUtils.cpp
--- a/Engine/Source/Editor/UnrealEd/Private/EditorBuildUtils.cpp
+++ b/Engine/Source/Editor/UnrealEd/Private/EditorBuildUtils.cpp
@@ -36,6 +36,7 @@ FEditorBuildUtils::FEditorAutomatedBuildSettings::FEditorAutomatedBuildSettings(
 	UnableToCheckoutFilesBehavior( ABB_PromptOnError ),
 	NewMapBehavior( ABB_PromptOnError ),
 	FailedToSaveBehavior( ABB_PromptOnError ),
+	bUseSCC( true ),
 	bAutoAddNewFiles( true ),
 	bShutdownEditorOnCompletion( false )
 {}
@@ -134,7 +135,7 @@ bool FEditorBuildUtils::EditorAutomatedBuildAndSubmit( const FEditorAutomatedBui
 	}
 
 	// Finally, if everything has gone smoothly, submit the requested packages to source control
-	if ( bBuildSuccessful )
+	if ( bBuildSuccessful && BuildSettings.bUseSCC )
 	{
 		SubmitPackagesForAutomatedBuild( PackagesToSubmit, BuildSettings );
 	}
@@ -513,19 +514,12 @@ bool FEditorBuildUtils::PrepForAutomatedBuild( const FEditorAutomatedBuildSettin
 
 	// Source control is required for the automated build, so ensure that SCC support is compiled in and
 	// that the server is enabled and available for use
-	if ( !ISourceControlModule::Get().IsEnabled() || !SourceControlProvider.IsAvailable() )
+	if ( BuildSettings.bUseSCC && !(ISourceControlModule::Get().IsEnabled() && SourceControlProvider.IsAvailable() ) )
 	{
 		bBuildSuccessful = false;
 		LogErrorMessage( NSLOCTEXT("UnrealEd", "AutomatedBuild_Error_SCCError", "Cannot connect to source control; automated build aborted."), OutErrorMessages );
 	}
 
-	// Empty changelists aren't allowed; abort the build if one wasn't provided
-	if ( bBuildSuccessful && BuildSettings.ChangeDescription.Len() == 0 )
-	{
-		bBuildSuccessful = false;
-		LogErrorMessage( NSLOCTEXT("UnrealEd", "AutomatedBuild_Error_NoCLDesc", "A changelist description must be provided; automated build aborted."), OutErrorMessages );
-	}
-
 	TArray<UPackage*> PreviouslySavedWorldPackages;
 	TArray<UPackage*> PackagesToCheckout;
 	TArray<ULevel*> LevelsToSave;
@@ -574,7 +568,7 @@ bool FEditorBuildUtils::PrepForAutomatedBuild( const FEditorAutomatedBuildSettin
 	// Load the asset tools module
 	FAssetToolsModule& AssetToolsModule = FModuleManager::GetModuleChecked<FAssetToolsModule>("AssetTools");
 
-	if ( bBuildSuccessful )
+	if ( bBuildSuccessful && BuildSettings.bUseSCC )
 	{
 		// Update the source control status of any relevant world packages in order to determine which need to be
 		// checked out, added to the depot, etc.
@@ -602,6 +596,9 @@ bool FEditorBuildUtils::PrepForAutomatedBuild( const FEditorAutomatedBuildSettin
 					}
 				}
 			}
+			else if(SourceControlState->IsCheckedOut())
+			{
+			}
 			else if(SourceControlState->CanCheckout())
 			{
 				PackagesToCheckout.Add( CurPackage );
@@ -767,7 +764,14 @@ void FEditorBuildUtils::SubmitPackagesForAutomatedBuild( const TSet<UPackage*>&
 
 	// Now check in all the changes, including the files we added above
 	TSharedRef<FCheckIn, ESPMode::ThreadSafe> CheckInOperation = StaticCastSharedRef<FCheckIn>(ISourceControlOperation::Create<FCheckIn>());
-	CheckInOperation->SetDescription(NSLOCTEXT("UnrealEd", "AutomatedBuild_AutomaticSubmission", "[Automatic Submission]"));
+	if (BuildSettings.ChangeDescription.IsEmpty())
+	{
+		CheckInOperation->SetDescription(NSLOCTEXT("UnrealEd", "AutomatedBuild_AutomaticSubmission", "[Automatic Submission]"));
+	}
+	else
+	{
+		CheckInOperation->SetDescription(FText::FromString(BuildSettings.ChangeDescription));
+	}
 	SourceControlProvider.Execute( CheckInOperation, LevelsToSubmit, EConcurrency::Synchronous );
 }
 
diff --git a/Engine/Source/Editor/UnrealEd/Private/UnrealEdMisc.cpp b/Engine/Source/Editor/UnrealEd/Private/UnrealEdMisc.cpp
--- a/Engine/Source/Editor/UnrealEd/Private/UnrealEdMisc.cpp
+++ b/Engine/Source/Editor/UnrealEd/Private/UnrealEdMisc.cpp
@@ -355,6 +355,10 @@ void FUnrealEdMisc::OnInit()
 		{
 			AutomatedBuildSettings.BuildErrorBehavior = ParsedBool ? FEditorBuildUtils::ABB_ProceedOnError : FEditorBuildUtils::ABB_FailOnError;
 		}
+		if ( FParse::Bool( ParsedCmdLine, TEXT("UseSCC="), ParsedBool ) )
+		{
+			AutomatedBuildSettings.bUseSCC = ParsedBool;
+		}
 		if ( FParse::Bool( ParsedCmdLine, TEXT("IgnoreSCCErrors="), ParsedBool ) )
 		{
 			AutomatedBuildSettings.UnableToCheckoutFilesBehavior = ParsedBool ? FEditorBuildUtils::ABB_ProceedOnError : FEditorBuildUtils::ABB_FailOnError;
diff --git a/Engine/Source/Editor/UnrealEd/Public/EditorBuildUtils.h b/Engine/Source/Editor/UnrealEd/Public/EditorBuildUtils.h
--- a/Engine/Source/Editor/UnrealEd/Public/EditorBuildUtils.h
+++ b/Engine/Source/Editor/UnrealEd/Public/EditorBuildUtils.h
@@ -62,6 +62,9 @@ public:
 		/** Behavior to take when a saveable map fails to save correctly */
 		EAutomatedBuildBehavior FailedToSaveBehavior;
 
+		/** Use SCC to checkout/checkin files */
+		bool bUseSCC;
+
 		/** If true, built map files not already in the source control depot will be added */
 		bool bAutoAddNewFiles;
 
