Performance improvements for importing thousands of mesh sections from one FBX file.  Reduced log spam as well

diff --git a/Engine/Source/Editor/UnrealEd/Classes/Settings/EditorExperimentalSettings.h b/Engine/Source/Editor/UnrealEd/Classes/Settings/EditorExperimentalSettings.h
--- a/Engine/Source/Editor/UnrealEd/Classes/Settings/EditorExperimentalSettings.h
+++ b/Engine/Source/Editor/UnrealEd/Classes/Settings/EditorExperimentalSettings.h
@@ -93,6 +93,9 @@ public:
 	UPROPERTY(EditAnywhere, config, Category = Rendering, meta = (DisplayName = "Enable Metal/High-end mobile rendering preview"))
 	bool bFeatureLevelES31Preview;
 
+	UPROPERTY(EditAnywhere, config, Category = Import)
+	bool bEnableLargeFileImportImprovements;
+
 	/**
 	 * Returns an event delegate that is executed when a setting has changed.
 	 *
diff --git a/Engine/Source/Editor/UnrealEd/Private/Fbx/FbxStaticMeshImport.cpp b/Engine/Source/Editor/UnrealEd/Private/Fbx/FbxStaticMeshImport.cpp
--- a/Engine/Source/Editor/UnrealEd/Private/Fbx/FbxStaticMeshImport.cpp
+++ b/Engine/Source/Editor/UnrealEd/Private/Fbx/FbxStaticMeshImport.cpp
@@ -93,7 +93,7 @@ UStaticMesh* UnFbx::FFbxImporter::ImportStaticMesh(UObject* InParent, FbxNode* N
 	return ImportStaticMeshAsSingle(InParent, MeshNodeArray, Name, Flags, ImportData, InStaticMesh, LODIndex);
 }
 
-bool UnFbx::FFbxImporter::BuildStaticMeshFromGeometry(FbxMesh* Mesh, UStaticMesh* StaticMesh, TArray<FFbxMaterial>& MeshMaterials, int LODIndex,
+bool UnFbx::FFbxImporter::BuildStaticMeshFromGeometry(FbxMesh* Mesh, UStaticMesh* StaticMesh, TArray<FFbxMaterial>& MeshMaterials, int LODIndex,FRawMesh* InRawMesh,
 													  EVertexColorImportOption::Type VertexColorImportOption, const TMap<FVector, FColor>& ExistingVertexColorData, const FColor& VertexOverrideColor)
 {
 	check(StaticMesh->SourceModels.IsValidIndex(LODIndex));
@@ -440,9 +440,23 @@ bool UnFbx::FFbxImporter::BuildStaticMeshFromGeometry(FbxMesh* Mesh, UStaticMesh
 	int32 WedgeCount = TriangleCount * 3;
 	bool OddNegativeScale = IsOddNegativeScale(TotalMatrix);
 
-	// Load the existing raw mesh.
-	FRawMesh RawMesh;
-	SrcModel.RawMeshBulkData->LoadRawMesh(RawMesh);
+	FRawMesh* RawMeshPtr;
+
+	TScopedPointer<FRawMesh> NewMesh;
+	if( !InRawMesh )
+	{
+		// Load the existing raw mesh.
+		NewMesh = new FRawMesh;
+		SrcModel.RawMeshBulkData->LoadRawMesh(*NewMesh);
+		RawMeshPtr = NewMesh.GetOwnedPointer();
+	}
+	else
+	{
+		RawMeshPtr = InRawMesh;
+	}
+
+	FRawMesh& RawMesh = *RawMeshPtr;
+
 	int32 VertexOffset = RawMesh.VertexPositions.Num();
 	int32 WedgeOffset = RawMesh.WedgeIndices.Num();
 	int32 TriangleOffset = RawMesh.FaceMaterialIndices.Num();
@@ -693,8 +707,11 @@ bool UnFbx::FFbxImporter::BuildStaticMeshFromGeometry(FbxMesh* Mesh, UStaticMesh
 	}
 
 
-	// Store the new raw mesh.
-	SrcModel.RawMeshBulkData->SaveRawMesh(RawMesh);
+	// Store the new raw mesh if needed
+	if( !InRawMesh )
+	{
+		SrcModel.RawMeshBulkData->SaveRawMesh(RawMesh);
+	}
 	
 	//
 	// clean up.  This needs to happen before the mesh is destroyed
@@ -987,6 +1004,15 @@ UStaticMesh* UnFbx::FFbxImporter::ImportStaticMeshAsSingle(UObject* InParent, TA
 	StaticMesh->LightMapResolution = 64;
 	StaticMesh->LightMapCoordinateIndex = 1;
 
+	// Note: This is only used if large file improvements are enabled.
+	FRawMesh RawMesh;
+
+	const bool bEnableLargeFileImportImprovements = GetDefault<UEditorExperimentalSettings>()->bEnableLargeFileImportImprovements;
+	if( bEnableLargeFileImportImprovements )
+	{
+		SrcModel.RawMeshBulkData->LoadRawMesh(RawMesh);
+	}
+
 	TArray<FFbxMaterial> MeshMaterials;
 	for (MeshIndex = 0; MeshIndex < MeshNodeArray.Num(); MeshIndex++ )
 	{
@@ -994,7 +1020,7 @@ UStaticMesh* UnFbx::FFbxImporter::ImportStaticMeshAsSingle(UObject* InParent, TA
 
 		if (Node->GetMesh())
 		{
-			if (!BuildStaticMeshFromGeometry(Node->GetMesh(), StaticMesh, MeshMaterials, LODIndex,
+			if (!BuildStaticMeshFromGeometry(Node->GetMesh(), StaticMesh, MeshMaterials, LODIndex, bEnableLargeFileImportImprovements ? &RawMesh : nullptr,
 											 VertexColorImportOption, ExistingVertexColorData, ImportOptions->VertexOverrideColor))
 			{
 				bBuildStatus = false;
@@ -1003,12 +1029,18 @@ UStaticMesh* UnFbx::FFbxImporter::ImportStaticMeshAsSingle(UObject* InParent, TA
 		}
 	}
 
+	if( bEnableLargeFileImportImprovements )
+	{
+		// Store the new raw mesh.
+		SrcModel.RawMeshBulkData->SaveRawMesh(RawMesh);
+	}
+
 	if (bBuildStatus)
 	{
-		UE_LOG(LogFbx,Log,TEXT("== Initial material list:"));
+		UE_LOG(LogFbx,Verbose,TEXT("== Initial material list:"));
 		for (int32 MaterialIndex = 0; MaterialIndex < MeshMaterials.Num(); ++MaterialIndex)
 		{
-			UE_LOG(LogFbx,Log,TEXT("%d: %s"),MaterialIndex,*MeshMaterials[MaterialIndex].GetName() );
+			UE_LOG(LogFbx,Verbose,TEXT("%d: %s"),MaterialIndex,*MeshMaterials[MaterialIndex].GetName() );
 		}
 
 		// Compress the materials array by removing any duplicates.
@@ -1048,7 +1080,7 @@ UStaticMesh* UnFbx::FFbxImporter::ImportStaticMeshAsSingle(UObject* InParent, TA
 			}
 			else
 			{
-				UE_LOG(LogFbx,Log,TEXT("  remap %d -> %d"), MaterialIndex, MaterialMap[MaterialIndex]);
+				UE_LOG(LogFbx,Verbose,TEXT("  remap %d -> %d"), MaterialIndex, MaterialMap[MaterialIndex]);
 			}
 		}
 
@@ -1092,15 +1124,15 @@ UStaticMesh* UnFbx::FFbxImporter::ImportStaticMeshAsSingle(UObject* InParent, TA
 		}
 		SortedMaterialIndex.Sort();
 
-		UE_LOG(LogFbx,Log,TEXT("== After sorting:"));
+		UE_LOG(LogFbx,Verbose,TEXT("== After sorting:"));
 		TArray<FFbxMaterial> SortedMaterials;
 		for (int32 SortedIndex = 0; SortedIndex < SortedMaterialIndex.Num(); ++SortedIndex)
 		{
 			int32 RemappedIndex = SortedMaterialIndex[SortedIndex] & 0xffff;
 			SortedMaterials.Add(UniqueMaterials[RemappedIndex]);
-			UE_LOG(LogFbx,Log,TEXT("%d: %s"),SortedIndex,*UniqueMaterials[RemappedIndex].GetName());
+			UE_LOG(LogFbx,Verbose,TEXT("%d: %s"),SortedIndex,*UniqueMaterials[RemappedIndex].GetName());
 		}
-		UE_LOG(LogFbx,Log,TEXT("== Mapping table:"));
+		UE_LOG(LogFbx,Verbose,TEXT("== Mapping table:"));
 		for (int32 MaterialIndex = 0; MaterialIndex < MaterialMap.Num(); ++MaterialIndex)
 		{
 			for (int32 SortedIndex = 0; SortedIndex < SortedMaterialIndex.Num(); ++SortedIndex)
@@ -1108,7 +1140,7 @@ UStaticMesh* UnFbx::FFbxImporter::ImportStaticMeshAsSingle(UObject* InParent, TA
 				int32 RemappedIndex = SortedMaterialIndex[SortedIndex] & 0xffff;
 				if (MaterialMap[MaterialIndex] == RemappedIndex)
 				{
-					UE_LOG(LogFbx,Log,TEXT("  sort %d -> %d"), MaterialIndex, SortedIndex);
+					UE_LOG(LogFbx,Verbose,TEXT("  sort %d -> %d"), MaterialIndex, SortedIndex);
 					MaterialMap[MaterialIndex] = SortedIndex;
 					break;
 				}
@@ -1229,8 +1261,14 @@ UStaticMesh* UnFbx::FFbxImporter::ImportStaticMeshAsSingle(UObject* InParent, TA
 			SrcModel.BuildSettings.bGenerateLightmapUVs = false;
 		}
 
+		TArray<FText> BuildErrors;
 		StaticMesh->LODGroup = ImportOptions->StaticMeshLODGroup;
-		StaticMesh->Build(false);
+		StaticMesh->Build(false, &BuildErrors);
+
+		for( FText& Error : BuildErrors )
+		{
+			AddTokenizedErrorMessage(FTokenizedMessage::Create(EMessageSeverity::Warning, Error), FFbxErrors::StaticMesh_BuildError );
+		}
 		
 		// this is damage control. After build, we'd like to absolutely sure that 
 		// all index is pointing correctly and they're all used. Otherwise we remove them
diff --git a/Engine/Source/Editor/UnrealEd/Private/FbxImporter.h b/Engine/Source/Editor/UnrealEd/Private/FbxImporter.h
--- a/Engine/Source/Editor/UnrealEd/Private/FbxImporter.h
+++ b/Engine/Source/Editor/UnrealEd/Private/FbxImporter.h
@@ -697,7 +697,7 @@ protected:
 	 * @param LODIndex	LOD level to set up for StaticMesh
 	 * @return bool true if set up successfully
 	 */
-	bool BuildStaticMeshFromGeometry(FbxMesh* FbxMesh, UStaticMesh* StaticMesh, TArray<FFbxMaterial>& MeshMaterials, int LODIndex,
+	bool BuildStaticMeshFromGeometry(FbxMesh* Mesh, UStaticMesh* StaticMesh, TArray<FFbxMaterial>& MeshMaterials, int LODIndex,FRawMesh* InRawMesh,
 									 EVertexColorImportOption::Type VertexColorImportOption, const TMap<FVector, FColor>& ExistingVertexColorData, const FColor& VertexOverrideColor);
 	
 	/**
diff --git a/Engine/Source/Runtime/Core/Private/Misc/FbxErrors.cpp b/Engine/Source/Runtime/Core/Private/Misc/FbxErrors.cpp
--- a/Engine/Source/Runtime/Core/Private/Misc/FbxErrors.cpp
+++ b/Engine/Source/Runtime/Core/Private/Misc/FbxErrors.cpp
@@ -49,6 +49,8 @@ FName FFbxErrors::StaticMesh_UVSetLayoutProblem(TEXT("StaticMesh_UVSetLayoutProb
 
 FName FFbxErrors::StaticMesh_NoTriangles(TEXT("StaticMesh_NoTriangles"));
 
+FName FFbxErrors::StaticMesh_BuildError(TEXT("StaticMesh_BuildError"));
+
 /** SkeletalMesh **/
 FName FFbxErrors::SkeletalMesh_DifferentRoots(TEXT("SkeletalMesh_DifferentRoot"));
 
diff --git a/Engine/Source/Runtime/Core/Public/Misc/FbxErrors.h b/Engine/Source/Runtime/Core/Public/Misc/FbxErrors.h
--- a/Engine/Source/Runtime/Core/Public/Misc/FbxErrors.h
+++ b/Engine/Source/Runtime/Core/Public/Misc/FbxErrors.h
@@ -57,6 +57,8 @@ struct CORE_API FFbxErrors
 
 	static FName StaticMesh_NoTriangles;
 
+	static FName StaticMesh_BuildError;
+
 	/** SkeletalMesh **/
 	static FName SkeletalMesh_DifferentRoots;
 
diff --git a/Engine/Source/Runtime/Engine/Classes/Engine/StaticMesh.h b/Engine/Source/Runtime/Engine/Classes/Engine/StaticMesh.h
--- a/Engine/Source/Runtime/Engine/Classes/Engine/StaticMesh.h
+++ b/Engine/Source/Runtime/Engine/Classes/Engine/StaticMesh.h
@@ -461,7 +461,7 @@ public:
 	 * Rebuilds renderable data for this static mesh.
 	 * @param bSilent - If true will not popup a progress dialog.
 	 */
-	ENGINE_API void Build(bool bSilent = false);
+	ENGINE_API void Build(bool bSilent = false, TArray<FText>* OutErrors = nullptr);
 
 	/**
 	 * Initialize the static mesh's render resources.
diff --git a/Engine/Source/Runtime/Engine/Private/StaticMeshBuild.cpp b/Engine/Source/Runtime/Engine/Private/StaticMeshBuild.cpp
--- a/Engine/Source/Runtime/Engine/Private/StaticMeshBuild.cpp
+++ b/Engine/Source/Runtime/Engine/Private/StaticMeshBuild.cpp
@@ -48,7 +48,7 @@ static bool HasBadTangents(UStaticMesh* Mesh)
 }
 #endif // #if WITH_EDITOR
 
-void UStaticMesh::Build(bool bSilent)
+void UStaticMesh::Build(bool bSilent, TArray<FText>* OutErrors)
 {
 #if WITH_EDITOR
 	if (IsTemplate())
@@ -119,9 +119,9 @@ void UStaticMesh::Build(bool bSilent)
 			Arguments.Add( TEXT("Options"), SourceModels[0].BuildSettings.bRecomputeTangents ? FText::GetEmpty() : LOCTEXT("MeshRecomputeTangents", "Consider enabling Recompute Tangents in the mesh's Build Settings.") );
 			const FText WarningMsg = FText::Format( LOCTEXT("MeshHasDegenerateTangents", "{Meshname} has degenerate tangent bases which will result in incorrect shading. {Options}"), Arguments );
 			UE_LOG(LogStaticMesh,Warning,TEXT("%s"),*WarningMsg.ToString());
-			if (!bSilent)
+			if (!bSilent && OutErrors)
 			{
-				FMessageDialog::Open(EAppMsgType::Ok, WarningMsg);
+				OutErrors->Add( WarningMsg );
 			}
 		}
 
