Updates to how a couple of the tests are named.

diff --git a/Engine/Source/Editor/UnrealEd/Private/Tests/GeometryTests.cpp b/Engine/Source/Editor/UnrealEd/Private/Tests/GeometryTests.cpp
--- a/Engine/Source/Editor/UnrealEd/Private/Tests/GeometryTests.cpp
+++ b/Engine/Source/Editor/UnrealEd/Private/Tests/GeometryTests.cpp
@@ -73,18 +73,23 @@ namespace GeometryPromotionHelpers
 /**
 * BSP Promotion Test
 * Adds one of every brush type to the world.  Most of which will be set as additive brushes.
+* Undo and redo the placement of an additive and subtractive brush.
 * There will be a few subtractive brushes intersecting with other objects.
 * Tests for number of surfaces and vertices, numbers of brushes created, and brush location.
 */
-IMPLEMENT_SIMPLE_AUTOMATION_TEST(FGeometryTest, "System.Promotion.Editor.Geometry Test", EAutomationTestFlags::ATF_Editor);
+IMPLEMENT_SIMPLE_AUTOMATION_TEST(FGeometryPromotionValidation, "System.Promotion.Editor.Geometry Validation", EAutomationTestFlags::ATF_Editor);
 
-bool FGeometryTest::RunTest(const FString& Parameters)
+bool FGeometryPromotionValidation::RunTest(const FString& Parameters)
 {
 	//** SETUP **//
+	// Set the Test Description
+	FString Description = TEXT("Geometry Validation:\n- Adds one of every brush type to the world.\n- Undo and redo the placement of an additive and subtractive brush.\n- Verify by finding the number of surfaces, vertices, brushes created, and their location.\n");
+	AddLogItem(Description);
+
 	// Create the world and setup the level.
 	UWorld* World = AutomationEditorCommonUtils::CreateNewMap();
-	ULevel* CurrentLevel = World->GetCurrentLevel();
 	GLevelEditorModeTools().MapChangeNotify();
+	ULevel* CurrentLevel = World->GetCurrentLevel();
 	// Set the view location so that it can see the entire scene.
 	AutomationEditorCommonUtils::SetOrthoViewportView(FVector(176, 2625, 2075), FRotator(319, 269, 1));
 	// Grabbing the current number of verts in the world to compare later on.
@@ -93,7 +98,6 @@ bool FGeometryTest::RunTest(const FString& Parameters)
 
 	//** TEST **//
 	//Cube Additive Brush
-	AddLogItem(TEXT("Additive Cube added to the world."));
 	UE_LOG(LogGeometryTests, Log, TEXT("Adding an 'Additive Cube' brush."));
 	UCubeBuilder* CubeAdditiveBrushBuilder = Cast<UCubeBuilder>(GEditor->FindBrushBuilder(UCubeBuilder::StaticClass()));
 	CubeAdditiveBrushBuilder->X = 4096.0f;
@@ -252,9 +256,9 @@ bool FGeometryTest::RunTest(const FString& Parameters)
 
 	// Verify there are the correct amount of BSP surfaces are visible in the level.
 	TestEqual<int32>(TEXT("Incorrect total number of surfaces are being reported, expected 10."), BSPModel->Surfs.Num(), 276);
+	UE_LOG(LogGeometryTests, Log, TEXT("The number of BSP Brushes Found in the level: Additive: %i, Subtractive: %i"), AdditiveBSP.Num(), SubtractiveBSP.Num());
 
-	UE_LOG(LogGeometryTests, Log, TEXT("BSP Placement: Number of Additive: %i, Number of Subtractive: %i"), AdditiveBSP.Num(), SubtractiveBSP.Num());
-
+	
 	//** TEARDOWN **//
 	AdditiveBSP.Empty();
 	SubtractiveBSP.Empty();
@@ -273,15 +277,15 @@ bool FGeometryTest::RunTest(const FString& Parameters)
 * Subtractive brush intersects will the additive brush.
 * Tests for number of surfaces and vertices, numbers of brushes created, and brush location.
 */
-IMPLEMENT_SIMPLE_AUTOMATION_TEST(FBSPCubeAddTest, "System.Editor.Geometry.BSP Cube - New Brush", EAutomationTestFlags::ATF_Editor);
+IMPLEMENT_SIMPLE_AUTOMATION_TEST(FBSPCubePlacement, "System.Editor.Geometry.BSP Cube.Brush Placement", EAutomationTestFlags::ATF_Editor);
 
-bool FBSPCubeAddTest::RunTest(const FString& Parameters)
+bool FBSPCubePlacement::RunTest(const FString& Parameters)
 {
 	//** SETUP **//
 	// Create the world and setup the level.
 	UWorld* World = AutomationEditorCommonUtils::CreateNewMap();
-	ULevel* CurrentLevel = World->GetCurrentLevel();
 	GLevelEditorModeTools().MapChangeNotify();
+	ULevel* CurrentLevel = World->GetCurrentLevel();
 	UBrushBuilder* DefaultBuilderBrush = World->GetDefaultBrush()->BrushBuilder;
 
 	//** TEST **//
@@ -361,9 +365,9 @@ bool FBSPCubeAddTest::RunTest(const FString& Parameters)
 * Adds an additive brush to the world and then does an undo/redo of that transaction.
 * Tests for number vertices in the level.
 */
-IMPLEMENT_SIMPLE_AUTOMATION_TEST(FBSPCubeUndoRedoTest, "System.Editor.Geometry.BSP Cube - Undo and Redo", EAutomationTestFlags::ATF_Editor);
+IMPLEMENT_SIMPLE_AUTOMATION_TEST(FBSPCubeUndoRedoPlacement, "System.Editor.Geometry.BSP Cube.Brush Placement Undo Redo", EAutomationTestFlags::ATF_Editor);
 
-bool FBSPCubeUndoRedoTest::RunTest(const FString& Parameters)
+bool FBSPCubeUndoRedoPlacement::RunTest(const FString& Parameters)
 {
 	//** SETUP **//
 	// Create the world and setup the level.
diff --git a/Engine/Source/Runtime/Engine/Private/Tests/EngineAutomationTests.cpp b/Engine/Source/Runtime/Engine/Private/Tests/EngineAutomationTests.cpp
--- a/Engine/Source/Runtime/Engine/Private/Tests/EngineAutomationTests.cpp
+++ b/Engine/Source/Runtime/Engine/Private/Tests/EngineAutomationTests.cpp
@@ -428,28 +428,54 @@ bool FCinematicFPSPerfTest::RunTest(const FString& Parameters)
 }
 
 
-IMPLEMENT_SIMPLE_AUTOMATION_TEST(FLogTypesTest, "System.Automation Framework.Logging Test", EAutomationTestFlags::ATF_None)
+IMPLEMENT_SIMPLE_AUTOMATION_TEST(FLogAddMessage, "System.Automation.Log.Add Log Message", EAutomationTestFlags::ATF_None)
 
-bool FLogTypesTest::RunTest(const FString& Parameters)
+bool FLogAddMessage::RunTest(const FString& Parameters)
 {
-	UE_LOG(LogEngineAutomationTests, Log, TEXT("This 'Log' log is created using the UE_LOG Macro"));
-	UE_LOG(LogEngineAutomationTests, Display, TEXT("This 'Display' log is created using the UE_LOG Macro"));
+	//** TEST **//
+	AddLogItem(TEXT("Test log message."));
 
-	ExecutionInfo.LogItems.Add(TEXT("This was added to the ExecutionInfo.LogItems"));
-	ExecutionInfo.Errors.Add(TEXT("This was added to the ExecutionInfo.Error."));
-	if (ExecutionInfo.Errors.Last().Contains(TEXT("This was added to the ExecutionInfo.Error.")))
-	{
-		UE_LOG(LogEngineAutomationTests, Display, TEXT("ExecutionInfo.Error has the correct info when used."));
-		ExecutionInfo.Errors.Empty();
-	}
+	//** VERIFY **//
+	TestEqual<FString>(TEXT("Test log message was not added to the ExecutionInfo.Log array."), ExecutionInfo.LogItems.Last(), TEXT("Test log message."));
+	
+	//** TEARDOWN **//
+	// We have to empty this log array so that it doesn't show in the automation results window as it may cause confusion.
+	ExecutionInfo.LogItems.Empty();
 
-	ExecutionInfo.Warnings.Add(TEXT("This was added to the ExecutionInfo.Warnings"));
-	if (ExecutionInfo.Warnings.Last().Contains(TEXT("This was added to the ExecutionInfo.Warnings")))
-	{
-		UE_LOG(LogEngineAutomationTests, Display, TEXT("ExecutionInfo.Warnings has the correct info when used."));
-		ExecutionInfo.Warnings.Empty();
-	}
+	return true;
+}
+
+IMPLEMENT_SIMPLE_AUTOMATION_TEST(FLogAddWarning, "System.Automation.Log.Add Warning Message", EAutomationTestFlags::ATF_None)
+
+bool FLogAddWarning::RunTest(const FString& Parameters)
+{
+	//** TEST **//
+	AddWarning(TEXT("Test warning message."));
+
+	//** VERIFY **//
+	FString CurrentWarningMessage = ExecutionInfo.Warnings.Last();
+	// The warnings array is emptied so that it doesn't cause a false positive warning for this test.
+	ExecutionInfo.Warnings.Empty();
+
+	TestEqual<FString>(TEXT("Test warning message was not added to the ExecutionInfo.Warning array."), CurrentWarningMessage, TEXT("Test warning message."));
+
+	return true;
+}
 
+IMPLEMENT_SIMPLE_AUTOMATION_TEST(FLogAddError, "System.Automation.Log.Add Error Message", EAutomationTestFlags::ATF_None)
+
+bool FLogAddError::RunTest(const FString& Parameters)
+{
+	//** TEST **//
+	AddError(TEXT("Test error message"));
+	
+	//** VERIFY **//
+	FString CurrentErrorMessage = ExecutionInfo.Errors.Last();
+	// The errors array is emptied so that this doesn't cause a false positive failure for this test.
+	ExecutionInfo.Errors.Empty();
+
+	TestEqual<FString>(TEXT("Test error message was not added to the ExecutionInfo.Error array."), CurrentErrorMessage, TEXT("Test error message"));
+	
 	return true;
 }
 
