Unified asset import data across all asset types.

The majority of asset reimport machinery is now held within UAssetImportData and FAssetImportInfo. This change allows us to remove the CachedAssetsBySourceFileName map from FAssetRegistry that was previously added to support auto reimport for 4.7.

The functionality of CachedAssetsBySourceFileName has been wrapped up in FAssetSourceFilenameCache, which sits on top of the asset registry in the editor.

This unification allows us to also consistently store the file timestamps and MD5 hashes of imported files. This subsequently allows auto-reimports to be ignored where a source content file has not actually changed (its MD5 is the same).

This addresses UETOOL-365 - Rework AssetRegistry changes to support auto-reimport

diff --git a/Engine/Plugins/2D/Paper2D/Source/Paper2D/Classes/PaperTileMap.h b/Engine/Plugins/2D/Paper2D/Source/Paper2D/Classes/PaperTileMap.h
--- a/Engine/Plugins/2D/Paper2D/Source/Paper2D/Classes/PaperTileMap.h
+++ b/Engine/Plugins/2D/Paper2D/Source/Paper2D/Classes/PaperTileMap.h
@@ -109,7 +109,7 @@ public:
 #if WITH_EDITORONLY_DATA
 	/** Importing data and options used for this tile map */
 	UPROPERTY(Category=ImportSettings, VisibleAnywhere, Instanced)
-	UAssetImportData* AssetImportData;
+	class UAssetImportData* AssetImportData;
 
 	/** The currently selected layer index */
 	UPROPERTY()
diff --git a/Engine/Plugins/2D/Paper2D/Source/Paper2D/Private/PaperTileMap.cpp b/Engine/Plugins/2D/Paper2D/Source/Paper2D/Private/PaperTileMap.cpp
--- a/Engine/Plugins/2D/Paper2D/Source/Paper2D/Private/PaperTileMap.cpp
+++ b/Engine/Plugins/2D/Paper2D/Source/Paper2D/Private/PaperTileMap.cpp
@@ -203,7 +203,7 @@ void UPaperTileMap::GetAssetRegistryTags(TArray<FAssetRegistryTag>& OutTags) con
 {
 	if (AssetImportData)
 	{
-		OutTags.Add( FAssetRegistryTag(SourceFileTagName(), AssetImportData->SourceFilePath, FAssetRegistryTag::TT_Hidden) );
+		OutTags.Add( FAssetRegistryTag(SourceFileTagName(), AssetImportData->ToJson(), FAssetRegistryTag::TT_Hidden) );
 	}
 
 	Super::GetAssetRegistryTags(OutTags);
diff --git a/Engine/Plugins/2D/Paper2D/Source/Paper2DEditor/Private/TileMapEditing/TileMapAssetImportData.cpp b/Engine/Plugins/2D/Paper2D/Source/Paper2DEditor/Private/TileMapEditing/TileMapAssetImportData.cpp
--- a/Engine/Plugins/2D/Paper2D/Source/Paper2DEditor/Private/TileMapEditing/TileMapAssetImportData.cpp
+++ b/Engine/Plugins/2D/Paper2D/Source/Paper2DEditor/Private/TileMapEditing/TileMapAssetImportData.cpp
@@ -25,8 +25,7 @@ UTileMapAssetImportData* UTileMapAssetImportData::GetImportDataForTileMap(UPaper
 		// Try to preserve the source file path if possible
 		if (TileMap->AssetImportData != nullptr)
 		{
-			ImportData->SourceFilePath = TileMap->AssetImportData->SourceFilePath;
-			ImportData->SourceFileTimestamp = TileMap->AssetImportData->SourceFileTimestamp;
+			ImportData->CopyFrom(*TileMap->AssetImportData);
 		}
 
 		TileMap->AssetImportData = ImportData;
diff --git a/Engine/Plugins/2D/Paper2D/Source/PaperSpriteSheetImporter/Private/PaperSpriteSheet.cpp b/Engine/Plugins/2D/Paper2D/Source/PaperSpriteSheetImporter/Private/PaperSpriteSheet.cpp
--- a/Engine/Plugins/2D/Paper2D/Source/PaperSpriteSheetImporter/Private/PaperSpriteSheet.cpp
+++ b/Engine/Plugins/2D/Paper2D/Source/PaperSpriteSheetImporter/Private/PaperSpriteSheet.cpp
@@ -13,7 +13,7 @@ void UPaperSpriteSheet::GetAssetRegistryTags(TArray<FAssetRegistryTag>& OutTags)
 {
 	if (AssetImportData)
 	{
-		OutTags.Add( FAssetRegistryTag(SourceFileTagName(), AssetImportData->SourceFilePath, FAssetRegistryTag::TT_Hidden) );
+		OutTags.Add( FAssetRegistryTag(SourceFileTagName(), AssetImportData->ToJson(), FAssetRegistryTag::TT_Hidden) );
 	}
 
 	Super::GetAssetRegistryTags(OutTags);
diff --git a/Engine/Plugins/2D/Paper2D/Source/PaperSpriteSheetImporter/Private/PaperSpriteSheetAssetTypeActions.cpp b/Engine/Plugins/2D/Paper2D/Source/PaperSpriteSheetImporter/Private/PaperSpriteSheetAssetTypeActions.cpp
--- a/Engine/Plugins/2D/Paper2D/Source/PaperSpriteSheetImporter/Private/PaperSpriteSheetAssetTypeActions.cpp
+++ b/Engine/Plugins/2D/Paper2D/Source/PaperSpriteSheetImporter/Private/PaperSpriteSheetAssetTypeActions.cpp
@@ -48,7 +48,7 @@ void FPaperSpriteSheetAssetTypeActions::GetResolvedSourceFilePaths(const TArray<
 		const auto SpriteSheet = CastChecked<UPaperSpriteSheet>(Asset);
 		if (SpriteSheet->AssetImportData)
 		{
-			OutSourceFilePaths.Add(FReimportManager::ResolveImportFilename(SpriteSheet->AssetImportData->SourceFilePath, SpriteSheet));
+			SpriteSheet->AssetImportData->ExtractFilenames(OutSourceFilePaths);
 		}
 	}
 }
diff --git a/Engine/Plugins/2D/Paper2D/Source/PaperSpriteSheetImporter/Private/PaperSpriteSheetImportFactory.cpp b/Engine/Plugins/2D/Paper2D/Source/PaperSpriteSheetImporter/Private/PaperSpriteSheetImportFactory.cpp
--- a/Engine/Plugins/2D/Paper2D/Source/PaperSpriteSheetImporter/Private/PaperSpriteSheetImportFactory.cpp
+++ b/Engine/Plugins/2D/Paper2D/Source/PaperSpriteSheetImporter/Private/PaperSpriteSheetImportFactory.cpp
@@ -65,8 +65,7 @@ UObject* UPaperSpriteSheetImportFactory::FactoryCreateText(UClass* InClass, UObj
 				SpriteSheet->AssetImportData = NewObject<UAssetImportData>(SpriteSheet);
 			}
 
-			SpriteSheet->AssetImportData->SourceFilePath = FReimportManager::SanitizeImportFilename(CurrentFilename, SpriteSheet);
-			SpriteSheet->AssetImportData->SourceFileTimestamp = IFileManager::Get().GetTimeStamp(*CurrentFilename).ToString();
+			SpriteSheet->AssetImportData->Update(CurrentFilename);
 			SpriteSheet->AssetImportData->bDirty = false;
 
 			Result = SpriteSheet;
diff --git a/Engine/Plugins/2D/Paper2D/Source/PaperSpriteSheetImporter/Private/PaperSpriteSheetReimportFactory.cpp b/Engine/Plugins/2D/Paper2D/Source/PaperSpriteSheetImporter/Private/PaperSpriteSheetReimportFactory.cpp
--- a/Engine/Plugins/2D/Paper2D/Source/PaperSpriteSheetImporter/Private/PaperSpriteSheetReimportFactory.cpp
+++ b/Engine/Plugins/2D/Paper2D/Source/PaperSpriteSheetImporter/Private/PaperSpriteSheetReimportFactory.cpp
@@ -22,7 +22,7 @@ bool UPaperSpriteSheetReimportFactory::CanReimport(UObject* Obj, TArray<FString>
 	UPaperSpriteSheet* SpriteSheet = Cast<UPaperSpriteSheet>(Obj);
 	if (SpriteSheet && SpriteSheet->AssetImportData)
 	{
-		OutFilenames.Add(FReimportManager::ResolveImportFilename(SpriteSheet->AssetImportData->SourceFilePath, SpriteSheet));
+		SpriteSheet->AssetImportData->ExtractFilenames(OutFilenames);
 		return true;
 	}
 	return false;
@@ -33,7 +33,7 @@ void UPaperSpriteSheetReimportFactory::SetReimportPaths(UObject* Obj, const TArr
 	UPaperSpriteSheet* SpriteSheet = Cast<UPaperSpriteSheet>(Obj);
 	if (SpriteSheet && ensure(NewReimportPaths.Num() == 1))
 	{
-		SpriteSheet->AssetImportData->SourceFilePath = FReimportManager::ResolveImportFilename(NewReimportPaths[0], SpriteSheet);
+		SpriteSheet->AssetImportData->UpdateFilenameOnly(NewReimportPaths[0]);
 	}
 }
 
@@ -46,7 +46,7 @@ EReimportResult::Type UPaperSpriteSheetReimportFactory::Reimport(UObject* Obj)
 	}
 
 	// Make sure file is valid and exists
-	const FString Filename = FReimportManager::ResolveImportFilename(SpriteSheet->AssetImportData->SourceFilePath, SpriteSheet);
+	const FString Filename = SpriteSheet->AssetImportData->GetFirstFilename();
 	if (!Filename.Len() || IFileManager::Get().FileSize(*Filename) == INDEX_NONE)
 	{
 		return EReimportResult::Failed;
diff --git a/Engine/Plugins/2D/Paper2D/Source/PaperTiledImporter/Private/PaperTiledImporterFactory.cpp b/Engine/Plugins/2D/Paper2D/Source/PaperTiledImporter/Private/PaperTiledImporterFactory.cpp
--- a/Engine/Plugins/2D/Paper2D/Source/PaperTiledImporter/Private/PaperTiledImporterFactory.cpp
+++ b/Engine/Plugins/2D/Paper2D/Source/PaperTiledImporter/Private/PaperTiledImporterFactory.cpp
@@ -271,8 +271,7 @@ UObject* UPaperTiledImporterFactory::FactoryCreateText(UClass* InClass, UObject*
 	{
 		// Store the current file path and timestamp for re-import purposes
 		UTileMapAssetImportData* ImportData = UTileMapAssetImportData::GetImportDataForTileMap(Result);
-		ImportData->SourceFilePath = FReimportManager::SanitizeImportFilename(CurrentFilename, Result);
-		ImportData->SourceFileTimestamp = IFileManager::Get().GetTimeStamp(*CurrentFilename).ToString();
+		ImportData->Update(CurrentFilename);
 
 		//GlobalInfo.TileSets
 	}
@@ -321,7 +320,7 @@ bool UPaperTiledImporterFactory::CanReimport(UObject* Obj, TArray<FString>& OutF
 	{
 		if (TileMap->AssetImportData != nullptr)
 		{
-			OutFilenames.Add(FReimportManager::ResolveImportFilename(TileMap->AssetImportData->SourceFilePath, TileMap));
+			TileMap->AssetImportData->ExtractFilenames(OutFilenames);
 		}
 		else
 		{
@@ -339,8 +338,7 @@ void UPaperTiledImporterFactory::SetReimportPaths(UObject* Obj, const TArray<FSt
 		if (ensure(NewReimportPaths.Num() == 1))
 		{
 			UTileMapAssetImportData* ImportData = UTileMapAssetImportData::GetImportDataForTileMap(TileMap);
-
-			ImportData->SourceFilePath = FReimportManager::SanitizeImportFilename(NewReimportPaths[0], TileMap);
+			ImportData->UpdateFilenameOnly(NewReimportPaths[0]);
 		}
 	}
 }
diff --git a/Engine/Plugins/2D/Paper2D/Source/SpriterImporter/Private/PaperSpriterImportData.cpp b/Engine/Plugins/2D/Paper2D/Source/SpriterImporter/Private/PaperSpriterImportData.cpp
--- a/Engine/Plugins/2D/Paper2D/Source/SpriterImporter/Private/PaperSpriterImportData.cpp
+++ b/Engine/Plugins/2D/Paper2D/Source/SpriterImporter/Private/PaperSpriterImportData.cpp
@@ -16,7 +16,7 @@ void UPaperSpriterImportData::GetAssetRegistryTags(TArray<FAssetRegistryTag>& Ou
 {
 	if (AssetImportData != nullptr)
 	{
-		OutTags.Add( FAssetRegistryTag(SourceFileTagName(), AssetImportData->SourceFilePath, FAssetRegistryTag::TT_Hidden) );
+		OutTags.Add( FAssetRegistryTag(SourceFileTagName(), AssetImportData->ToJson(), FAssetRegistryTag::TT_Hidden) );
 	}
 
 	Super::GetAssetRegistryTags(OutTags);
diff --git a/Engine/Plugins/Editor/SpeedTreeImporter/Source/SpeedTreeImporter/Private/ReimportSpeedTreeFactory.cpp b/Engine/Plugins/Editor/SpeedTreeImporter/Source/SpeedTreeImporter/Private/ReimportSpeedTreeFactory.cpp
--- a/Engine/Plugins/Editor/SpeedTreeImporter/Source/SpeedTreeImporter/Private/ReimportSpeedTreeFactory.cpp
+++ b/Engine/Plugins/Editor/SpeedTreeImporter/Source/SpeedTreeImporter/Private/ReimportSpeedTreeFactory.cpp
@@ -29,7 +29,7 @@ bool UReimportSpeedTreeFactory::CanReimport(UObject* Obj, TArray<FString>& OutFi
 	UStaticMesh* Mesh = Cast<UStaticMesh>(Obj);
 	if (Mesh && Mesh->AssetImportData)
 	{
-		OutFilenames.Add(FReimportManager::ResolveImportFilename(Mesh->AssetImportData->SourceFilePath, Mesh));
+		Mesh->AssetImportData->ExtractFilenames(OutFilenames);
 		return true;
 	}
 #endif // #if WITH_SPEEDTREE
@@ -42,7 +42,7 @@ void UReimportSpeedTreeFactory::SetReimportPaths(UObject* Obj, const TArray<FStr
 	UStaticMesh* Mesh = Cast<UStaticMesh>(Obj);
 	if (Mesh && ensure(NewReimportPaths.Num() == 1))
 	{
-		Mesh->AssetImportData->SourceFilePath = FReimportManager::ResolveImportFilename(NewReimportPaths[0], Mesh);
+		Mesh->AssetImportData->UpdateFilenameOnly(NewReimportPaths[0]);
 	}
 #endif // #if WITH_SPEEDTREE
 }
@@ -56,7 +56,7 @@ EReimportResult::Type UReimportSpeedTreeFactory::Reimport(UObject* Obj)
 		return EReimportResult::Failed;
 	}
 
-	const FString Filename = FReimportManager::ResolveImportFilename(Mesh->AssetImportData->SourceFilePath, Mesh);
+	const FString Filename = Mesh->AssetImportData->GetFirstFilename();
 	const FString FileExtension = FPaths::GetExtension(Filename);
 	const bool bIsSRT = FCString::Stricmp(*FileExtension, TEXT("SRT")) == 0;
 
diff --git a/Engine/Plugins/Editor/SpeedTreeImporter/Source/SpeedTreeImporter/Private/SpeedTreeImportFactory.cpp b/Engine/Plugins/Editor/SpeedTreeImporter/Source/SpeedTreeImporter/Private/SpeedTreeImportFactory.cpp
--- a/Engine/Plugins/Editor/SpeedTreeImporter/Source/SpeedTreeImporter/Private/SpeedTreeImportFactory.cpp
+++ b/Engine/Plugins/Editor/SpeedTreeImporter/Source/SpeedTreeImporter/Private/SpeedTreeImportFactory.cpp
@@ -489,8 +489,7 @@ UTexture* CreateSpeedTreeMaterialTexture(UObject* Parent, FString Filename, bool
 		UnrealTexture = (UTexture*)TextureFact->FactoryCreateBinary(UTexture2D::StaticClass(), Package, *TextureName, RF_Standalone|RF_Public, NULL, *Extension, PtrTexture, PtrTexture + TextureData.Num(), GWarn);
 		if (UnrealTexture != NULL)
 		{
-			UnrealTexture->SourceFilePath = FReimportManager::SanitizeImportFilename(Filename, UnrealTexture);
-			UnrealTexture->SourceFileTimestamp = IFileManager::Get().GetTimeStamp(*Filename).ToString();
+			UnrealTexture->AssetImportData->Update(Filename);
 
 			// Notify the asset registry
 			FAssetRegistryModule::AssetCreated(UnrealTexture);
@@ -1235,8 +1234,7 @@ UObject* USpeedTreeImportFactory::FactoryCreateBinary( UClass* InClass, UObject*
 
 				// @todo AssetImportData make a data class for speed tree assets
 				StaticMesh->AssetImportData = NewObject<UAssetImportData>(StaticMesh);
-				StaticMesh->AssetImportData->SourceFilePath = FReimportManager::SanitizeImportFilename(UFactory::GetCurrentFilename(), StaticMesh);
-				StaticMesh->AssetImportData->SourceFileTimestamp = IFileManager::Get().GetTimeStamp(*UFactory::GetCurrentFilename()).ToString();
+				StaticMesh->AssetImportData->Update(UFactory::GetCurrentFilename());
 				StaticMesh->AssetImportData->bDirty = false;
 				
 				// clear out any old data
diff --git a/Engine/Plugins/ScriptPlugin/Source/ScriptEditorPlugin/Private/ScriptFactory.cpp b/Engine/Plugins/ScriptPlugin/Source/ScriptEditorPlugin/Private/ScriptFactory.cpp
--- a/Engine/Plugins/ScriptPlugin/Source/ScriptEditorPlugin/Private/ScriptFactory.cpp
+++ b/Engine/Plugins/ScriptPlugin/Source/ScriptEditorPlugin/Private/ScriptFactory.cpp
@@ -66,9 +66,10 @@ UObject* UScriptFactory::FactoryCreateText(UClass* InClass, UObject* InParent, F
 	    {
 		    NewBlueprint = CastChecked<UScriptBlueprint>(FKismetEditorUtilities::CreateBlueprint(ParentClass, InParent, InName, BPTYPE_Normal, UScriptBlueprint::StaticClass(), UScriptBlueprintGeneratedClass::StaticClass(), "UScriptFactory"));
 		}
-		NewBlueprint->SourceFilePath = FReimportManager::SanitizeImportFilename(CurrentFilename, NewBlueprint);
+		
 		NewBlueprint->SourceCode = Buffer;
-		NewBlueprint->SourceFileTimestamp = IFileManager::Get().GetTimeStamp(*NewBlueprint->SourceFilePath).ToString();
+
+		NewBlueprint->AssetImportData->Update(CurrentFilename);
 
 		// Need to make sure we compile with the new source code
 		FKismetEditorUtilities::CompileBlueprint(NewBlueprint);
@@ -101,7 +102,7 @@ bool UReimportScriptFactory::CanReimport(UObject* Obj, TArray<FString>& OutFilen
 	UScriptBlueprint* ScriptClass = Cast<UScriptBlueprint>(Obj);
 	if (ScriptClass)
 	{
-		OutFilenames.Add(FReimportManager::ResolveImportFilename(ScriptClass->SourceFilePath, ScriptClass));
+		ScriptClass->AssetImportData->ExtractFilenames(OutFilenames);
 		return true;
 	}
 	return false;
@@ -112,7 +113,7 @@ void UReimportScriptFactory::SetReimportPaths(UObject* Obj, const TArray<FString
 	UScriptBlueprint* ScriptClass = Cast<UScriptBlueprint>(Obj);
 	if (ScriptClass && ensure(NewReimportPaths.Num() == 1))
 	{
-		ScriptClass->SourceFilePath = FReimportManager::SanitizeImportFilename(NewReimportPaths[0], Obj);
+		ScriptClass->AssetImportData->UpdateFilenameOnly(NewReimportPaths[0]);
 	}
 }
 
@@ -129,7 +130,7 @@ EReimportResult::Type UReimportScriptFactory::Reimport(UObject* Obj)
 
 	TGuardValue<UScriptBlueprint*> OriginalScriptGuardValue(OriginalScript, ScriptClass);
 
-	const FString ResolvedSourceFilePath = FReimportManager::ResolveImportFilename(ScriptClass->SourceFilePath, ScriptClass);
+	const FString ResolvedSourceFilePath = ScriptClass->AssetImportData->GetFirstFilename();
 	if (!ResolvedSourceFilePath.Len())
 	{
 		return EReimportResult::Failed;
diff --git a/Engine/Plugins/ScriptPlugin/Source/ScriptPlugin/Classes/ScriptBlueprint.h b/Engine/Plugins/ScriptPlugin/Source/ScriptPlugin/Classes/ScriptBlueprint.h
--- a/Engine/Plugins/ScriptPlugin/Source/ScriptPlugin/Classes/ScriptBlueprint.h
+++ b/Engine/Plugins/ScriptPlugin/Source/ScriptPlugin/Classes/ScriptBlueprint.h
@@ -13,16 +13,17 @@ UCLASS(BlueprintType, EarlyAccessPreview)
 class SCRIPTPLUGIN_API UScriptBlueprint : public UBlueprint
 {
 	GENERATED_UCLASS_BODY()
-
 public:
+	
+#if WITH_EDITORONLY_DATA
+	UPROPERTY()
+	FString SourceFilePath_DEPRECATED;
 
-	/** Path to the script text used to construct this asset. Relative to the object's package, BaseDir() or absolute */
-	UPROPERTY(Category = Script, VisibleAnywhere, BlueprintReadWrite)
-	FString SourceFilePath;
+	UPROPERTY()
+	class UAssetImportData* AssetImportData;
+#endif
 
-	/** Date/Time-stamp of the file from the last import. */
-	UPROPERTY(Category = Script, VisibleAnywhere)
-	FString SourceFileTimestamp;
+public:
 
 	/** Generated script bytecode */
 	UPROPERTY()
@@ -35,6 +36,7 @@ public:
 #if WITH_EDITORONLY_DATA
 	/** Override to ensure we write out the asset import data */
 	virtual void GetAssetRegistryTags(TArray<FAssetRegistryTag>& OutTags) const;
+	virtual void PostLoad() override;
 #endif
 
 #if WITH_EDITOR
@@ -47,6 +49,8 @@ public:
 	// End of UBlueprint interface
 
 	static bool ValidateGeneratedClass(const UClass* InClass);
+	
+	bool IsCodeDirty() const;
 	void UpdateScriptStatus();
 
 	/**
diff --git a/Engine/Plugins/ScriptPlugin/Source/ScriptPlugin/Private/ScriptBlueprint.cpp b/Engine/Plugins/ScriptPlugin/Source/ScriptPlugin/Private/ScriptBlueprint.cpp
--- a/Engine/Plugins/ScriptPlugin/Source/ScriptPlugin/Private/ScriptBlueprint.cpp
+++ b/Engine/Plugins/ScriptPlugin/Source/ScriptPlugin/Private/ScriptBlueprint.cpp
@@ -11,15 +11,26 @@
 UScriptBlueprint::UScriptBlueprint(const FObjectInitializer& ObjectInitializer)
 	: Super(ObjectInitializer)
 {
+#if WITH_EDITORONLY_DATA
+	AssetImportData = CreateEditorOnlyDefaultSubobject<UAssetImportData>(TEXT("AssetImportData"));
+#endif
 }
 
 #if WITH_EDITORONLY_DATA
 void UScriptBlueprint::GetAssetRegistryTags(TArray<FAssetRegistryTag>& OutTags) const
 {
-	OutTags.Add( FAssetRegistryTag(SourceFileTagName(), SourceFilePath, FAssetRegistryTag::TT_Hidden) );
+	OutTags.Add( FAssetRegistryTag(SourceFileTagName(), AssetImportData->ToJson(), FAssetRegistryTag::TT_Hidden) );
 
 	Super::GetAssetRegistryTags(OutTags);
 }
+void UScriptBlueprint::PostLoad()
+{
+	Super::PostLoad();
+	if (!SourceFilePath_DEPRECATED.IsEmpty())
+	{
+		AssetImportData->UpdateFilenameOnly(SourceFilePath_DEPRECATED);
+	}
+}
 #endif
 
 #if WITH_EDITOR
@@ -46,18 +57,25 @@ bool UScriptBlueprint::ValidateGeneratedClass(const UClass* InClass)
 	return Result;
 }
 
-void UScriptBlueprint::UpdateScriptStatus()
+bool UScriptBlueprint::IsCodeDirty() const
 {
-	FDateTime OldTimeStamp;
-	bool bCodeDirty = !FDateTime::Parse(SourceFileTimestamp, OldTimeStamp);
-	
-	if (!bCodeDirty)
+	const TArray<FAssetImportInfo::FSourceFile>& Data = AssetImportData->GetSourceFileData();
+
+	if (Data.Num() == 1)
 	{
-		FDateTime TimeStamp = IFileManager::Get().GetTimeStamp(*SourceFilePath);
-		bCodeDirty = TimeStamp > OldTimeStamp;
+		// Check the timestamp of the file as it is now, and the last imported timestamp
+		if (IFileManager::Get().GetTimeStamp(*AssetImportData->GetFirstFilename()) <= Data[0].Timestamp)
+		{
+			return false;
+		}
 	}
 
-	if (bCodeDirty)
+	return true;
+}
+
+void UScriptBlueprint::UpdateScriptStatus()
+{
+	if (IsCodeDirty())
 	{
 		FBlueprintEditorUtils::MarkBlueprintAsModified(this);
 	}
@@ -65,18 +83,14 @@ void UScriptBlueprint::UpdateScriptStatus()
 
 void UScriptBlueprint::UpdateSourceCodeIfChanged()
 {
-	FDateTime OldTimeStamp;	
-	bool bCodeDirty = !FDateTime::Parse(SourceFileTimestamp, OldTimeStamp);
-	FDateTime TimeStamp = IFileManager::Get().GetTimeStamp(*SourceFilePath);
-	bCodeDirty = bCodeDirty || (TimeStamp > OldTimeStamp);
-
-	if (bCodeDirty)
+	if (IsCodeDirty())
 	{
 		FString NewScript;
-		if (FFileHelper::LoadFileToString(NewScript, *SourceFilePath))
+		FString Filename = AssetImportData->GetFirstFilename();
+		if (FFileHelper::LoadFileToString(NewScript, *Filename))
 		{
 			SourceCode = NewScript;
-			SourceFileTimestamp = TimeStamp.ToString();
+			AssetImportData->Update(Filename);
 		}
 	}
 }
diff --git a/Engine/Source/Developer/AssetTools/Private/AssetTypeActions/AssetTypeActions_AnimSequence.cpp b/Engine/Source/Developer/AssetTools/Private/AssetTypeActions/AssetTypeActions_AnimSequence.cpp
--- a/Engine/Source/Developer/AssetTools/Private/AssetTypeActions/AssetTypeActions_AnimSequence.cpp
+++ b/Engine/Source/Developer/AssetTools/Private/AssetTypeActions/AssetTypeActions_AnimSequence.cpp
@@ -58,36 +58,36 @@ void FAssetTypeActions_AnimSequence::GetResolvedSourceFilePaths(const TArray<UOb
 		const auto AnimSequence = CastChecked<UAnimSequence>(Asset);
 		if (AnimSequence->AssetImportData)
 		{
-			OutSourceFilePaths.Add(FReimportManager::ResolveImportFilename(AnimSequence->AssetImportData->SourceFilePath, AnimSequence));
+			AnimSequence->AssetImportData->ExtractFilenames(OutSourceFilePaths);
 		}
 	}
 }
 
 void FAssetTypeActions_AnimSequence::ExecuteReimportWithNewSource(TArray<TWeakObjectPtr<UAnimSequence>> Objects)
 {
-	const FString EmptyPath = TEXT("");
+	FAssetImportInfo EmptyImportInfo;
 
 	for (auto ObjIt = Objects.CreateConstIterator(); ObjIt; ++ObjIt)
 	{
 		auto Object = (*ObjIt).Get();
 		if (Object)
 		{
-			FString OldPath = EmptyPath;
+			FAssetImportInfo OldImportData;
 
-			// clear the source, I'm not sure if I should just remove all AssetImportData, but this is only clearing the source
-			// so everything else stays except the path
 			if (Object->AssetImportData)
 			{
-				OldPath = Object->AssetImportData->SourceFilePath;
-				Object->AssetImportData->SourceFilePath = EmptyPath;
+				// Make note of the old import data
+				OldImportData.CopyFrom(*Object->AssetImportData);
+				// And reset the import data
+				Object->AssetImportData->CopyFrom(EmptyImportInfo);
 			}
 
 			bool bSuccess = FReimportManager::Instance()->Reimport(Object, /*bAskForNewFileIfMissing=*/true);
 
 			// restore the old source path in case reimport was not successful
-			if (!bSuccess && Object->AssetImportData && OldPath != EmptyPath)
+			if (!bSuccess && Object->AssetImportData)
 			{
-				Object->AssetImportData->SourceFilePath = OldPath;
+				Object->AssetImportData->CopyFrom(OldImportData);
 			}
 		}
 	}
diff --git a/Engine/Source/Developer/AssetTools/Private/AssetTypeActions/AssetTypeActions_Curve.cpp b/Engine/Source/Developer/AssetTools/Private/AssetTypeActions/AssetTypeActions_Curve.cpp
--- a/Engine/Source/Developer/AssetTools/Private/AssetTypeActions/AssetTypeActions_Curve.cpp
+++ b/Engine/Source/Developer/AssetTools/Private/AssetTypeActions/AssetTypeActions_Curve.cpp
@@ -24,6 +24,6 @@ void FAssetTypeActions_Curve::GetResolvedSourceFilePaths(const TArray<UObject*>&
 	for (auto& Asset : TypeAssets)
 	{
 		const auto Curve = CastChecked<UCurveBase>(Asset);
-		OutSourceFilePaths.Add(FReimportManager::ResolveImportFilename(Curve->ImportPath, Curve));
+		Curve->AssetImportData->ExtractFilenames(OutSourceFilePaths);
 	}
 }
\ No newline at end of file
diff --git a/Engine/Source/Developer/AssetTools/Private/AssetTypeActions/AssetTypeActions_CurveTable.cpp b/Engine/Source/Developer/AssetTools/Private/AssetTypeActions/AssetTypeActions_CurveTable.cpp
--- a/Engine/Source/Developer/AssetTools/Private/AssetTypeActions/AssetTypeActions_CurveTable.cpp
+++ b/Engine/Source/Developer/AssetTools/Private/AssetTypeActions/AssetTypeActions_CurveTable.cpp
@@ -20,7 +20,7 @@ void FAssetTypeActions_CurveTable::GetActions( const TArray<UObject*>& InObjects
 		const UCurveTable* CurTable = Cast<UCurveTable>((*TableIter).Get());
 		if (CurTable)
 		{
-			ImportPaths.Add(FReimportManager::ResolveImportFilename(CurTable->ImportPath, CurTable));
+			CurTable->AssetImportData->ExtractFilenames(ImportPaths);
 		}
 	}
 
@@ -79,7 +79,7 @@ void FAssetTypeActions_CurveTable::ExecuteExportAsCSV(TArray< TWeakObjectPtr<UOb
 		if (CurTable)
 		{
 			const FText Title = FText::Format(LOCTEXT("CurveTable_ExportCSVDialogTitle", "Export '{0}' as CSV..."), FText::FromString(*CurTable->GetName()));
-			const FString CurrentFilename = (CurTable->ImportPath.IsEmpty()) ? TEXT("") : FReimportManager::ResolveImportFilename(CurTable->ImportPath, CurTable);
+			const FString CurrentFilename = CurTable->AssetImportData->GetFirstFilename();
 			const FString FileTypes = TEXT("Curve Table CSV (*.csv)|*.csv");
 
 			TArray<FString> OutFilenames;
@@ -120,7 +120,7 @@ void FAssetTypeActions_CurveTable::ExecuteExportAsJSON(TArray< TWeakObjectPtr<UO
 		if (CurTable)
 		{
 			const FText Title = FText::Format(LOCTEXT("CurveTable_ExportCSVDialogTitle", "Export '{0}' as JSON..."), FText::FromString(*CurTable->GetName()));
-			const FString CurrentFilename = (CurTable->ImportPath.IsEmpty()) ? TEXT("") : FReimportManager::ResolveImportFilename(CurTable->ImportPath, CurTable);
+			const FString CurrentFilename = CurTable->AssetImportData->GetFirstFilename();
 			const FString FileTypes = TEXT("Curve Table JSON (*.json)|*.json");
 
 			TArray<FString> OutFilenames;
@@ -160,7 +160,7 @@ void FAssetTypeActions_CurveTable::GetResolvedSourceFilePaths(const TArray<UObje
 	for (auto& Asset : TypeAssets)
 	{
 		const auto CurveTable = CastChecked<UCurveTable>(Asset);
-		OutSourceFilePaths.Add(FReimportManager::ResolveImportFilename(CurveTable->ImportPath, CurveTable));
+		CurveTable->AssetImportData->ExtractFilenames(OutSourceFilePaths);
 	}
 }
 
diff --git a/Engine/Source/Developer/AssetTools/Private/AssetTypeActions/AssetTypeActions_DataTable.cpp b/Engine/Source/Developer/AssetTools/Private/AssetTypeActions/AssetTypeActions_DataTable.cpp
--- a/Engine/Source/Developer/AssetTools/Private/AssetTypeActions/AssetTypeActions_DataTable.cpp
+++ b/Engine/Source/Developer/AssetTools/Private/AssetTypeActions/AssetTypeActions_DataTable.cpp
@@ -20,7 +20,7 @@ void FAssetTypeActions_DataTable::GetActions( const TArray<UObject*>& InObjects,
 		const UDataTable* CurTable = Cast<UDataTable>((*TableIter).Get());
 		if (CurTable)
 		{
-			ImportPaths.Add(FReimportManager::ResolveImportFilename(CurTable->ImportPath, CurTable));
+			CurTable->AssetImportData->ExtractFilenames(ImportPaths);
 		}
 	}
 
@@ -79,7 +79,7 @@ void FAssetTypeActions_DataTable::ExecuteExportAsCSV(TArray< TWeakObjectPtr<UObj
 		if (DataTable)
 		{
 			const FText Title = FText::Format(LOCTEXT("DataTable_ExportCSVDialogTitle", "Export '{0}' as CSV..."), FText::FromString(*DataTable->GetName()));
-			const FString CurrentFilename = (DataTable->ImportPath.IsEmpty()) ? TEXT("") : FReimportManager::ResolveImportFilename(DataTable->ImportPath, DataTable);
+			const FString CurrentFilename = DataTable->AssetImportData->GetFirstFilename();
 			const FString FileTypes = TEXT("Data Table CSV (*.csv)|*.csv");
 
 			TArray<FString> OutFilenames;
@@ -120,7 +120,7 @@ void FAssetTypeActions_DataTable::ExecuteExportAsJSON(TArray< TWeakObjectPtr<UOb
 		if (DataTable)
 		{
 			const FText Title = FText::Format(LOCTEXT("DataTable_ExportCSVDialogTitle", "Export '{0}' as JSON..."), FText::FromString(*DataTable->GetName()));
-			const FString CurrentFilename = (DataTable->ImportPath.IsEmpty()) ? TEXT("") : FReimportManager::ResolveImportFilename(DataTable->ImportPath, DataTable);
+			const FString CurrentFilename = DataTable->AssetImportData->GetFirstFilename();
 			const FString FileTypes = TEXT("Data Table JSON (*.json)|*.json");
 
 			TArray<FString> OutFilenames;
@@ -160,7 +160,7 @@ void FAssetTypeActions_DataTable::GetResolvedSourceFilePaths(const TArray<UObjec
 	for (auto& Asset : TypeAssets)
 	{
 		const auto DataTable = CastChecked<UDataTable>(Asset);
-		OutSourceFilePaths.Add(FReimportManager::ResolveImportFilename(DataTable->ImportPath, DataTable));
+		DataTable->AssetImportData->ExtractFilenames(OutSourceFilePaths);
 	}
 }
 
diff --git a/Engine/Source/Developer/AssetTools/Private/AssetTypeActions/AssetTypeActions_SkeletalMesh.cpp b/Engine/Source/Developer/AssetTools/Private/AssetTypeActions/AssetTypeActions_SkeletalMesh.cpp
--- a/Engine/Source/Developer/AssetTools/Private/AssetTypeActions/AssetTypeActions_SkeletalMesh.cpp
+++ b/Engine/Source/Developer/AssetTools/Private/AssetTypeActions/AssetTypeActions_SkeletalMesh.cpp
@@ -553,7 +553,7 @@ void FAssetTypeActions_SkeletalMesh::GetResolvedSourceFilePaths(const TArray<UOb
 		const auto SkeletalMesh = CastChecked<USkeletalMesh>(Asset);
 		if (SkeletalMesh->AssetImportData)
 		{
-			OutSourceFilePaths.Add(FReimportManager::ResolveImportFilename(SkeletalMesh->AssetImportData->SourceFilePath, SkeletalMesh));
+			SkeletalMesh->AssetImportData->ExtractFilenames(OutSourceFilePaths);
 		}
 	}
 }
diff --git a/Engine/Source/Developer/AssetTools/Private/AssetTypeActions/AssetTypeActions_SoundWave.cpp b/Engine/Source/Developer/AssetTools/Private/AssetTypeActions/AssetTypeActions_SoundWave.cpp
--- a/Engine/Source/Developer/AssetTools/Private/AssetTypeActions/AssetTypeActions_SoundWave.cpp
+++ b/Engine/Source/Developer/AssetTools/Private/AssetTypeActions/AssetTypeActions_SoundWave.cpp
@@ -36,7 +36,7 @@ void FAssetTypeActions_SoundWave::GetResolvedSourceFilePaths(const TArray<UObjec
 	for (auto& Asset : TypeAssets)
 	{
 		const auto SoundWave = CastChecked<USoundWave>(Asset);
-		OutSourceFilePaths.Add(FReimportManager::ResolveImportFilename(SoundWave->SourceFilePath, SoundWave));
+		SoundWave->AssetImportData->ExtractFilenames(OutSourceFilePaths);
 	}
 }
 
diff --git a/Engine/Source/Developer/AssetTools/Private/AssetTypeActions/AssetTypeActions_StaticMesh.cpp b/Engine/Source/Developer/AssetTools/Private/AssetTypeActions/AssetTypeActions_StaticMesh.cpp
--- a/Engine/Source/Developer/AssetTools/Private/AssetTypeActions/AssetTypeActions_StaticMesh.cpp
+++ b/Engine/Source/Developer/AssetTools/Private/AssetTypeActions/AssetTypeActions_StaticMesh.cpp
@@ -94,7 +94,7 @@ void FAssetTypeActions_StaticMesh::GetResolvedSourceFilePaths(const TArray<UObje
 		const auto StaticMesh = CastChecked<UStaticMesh>(Asset);
 		if (StaticMesh->AssetImportData)
 		{
-			OutSourceFilePaths.Add(FReimportManager::ResolveImportFilename(StaticMesh->AssetImportData->SourceFilePath, StaticMesh));
+			StaticMesh->AssetImportData->ExtractFilenames(OutSourceFilePaths);
 		}
 	}
 }
diff --git a/Engine/Source/Developer/AssetTools/Private/AssetTypeActions/AssetTypeActions_Texture.cpp b/Engine/Source/Developer/AssetTools/Private/AssetTypeActions/AssetTypeActions_Texture.cpp
--- a/Engine/Source/Developer/AssetTools/Private/AssetTypeActions/AssetTypeActions_Texture.cpp
+++ b/Engine/Source/Developer/AssetTools/Private/AssetTypeActions/AssetTypeActions_Texture.cpp
@@ -45,7 +45,7 @@ void FAssetTypeActions_Texture::GetResolvedSourceFilePaths(const TArray<UObject*
 	for (auto& Asset : TypeAssets)
 	{
 		UTexture* Texture = CastChecked<UTexture>(Asset);
-		OutSourceFilePaths.Add(FReimportManager::ResolveImportFilename(Texture->SourceFilePath, Texture));
+		Texture->AssetImportData->ExtractFilenames(OutSourceFilePaths);
 	}
 }
 
diff --git a/Engine/Source/Developer/AssetTools/Private/AssetTypeActions/AssetTypeActions_VectorFieldStatic.cpp b/Engine/Source/Developer/AssetTools/Private/AssetTypeActions/AssetTypeActions_VectorFieldStatic.cpp
--- a/Engine/Source/Developer/AssetTools/Private/AssetTypeActions/AssetTypeActions_VectorFieldStatic.cpp
+++ b/Engine/Source/Developer/AssetTools/Private/AssetTypeActions/AssetTypeActions_VectorFieldStatic.cpp
@@ -7,6 +7,6 @@ void FAssetTypeActions_VectorFieldStatic::GetResolvedSourceFilePaths(const TArra
 	for (auto& Asset : TypeAssets)
 	{
 		const auto StaticVectorField = CastChecked<UVectorFieldStatic>(Asset);
-		OutSourceFilePaths.Add(FReimportManager::ResolveImportFilename(StaticVectorField->SourceFilePath, StaticVectorField));
+		StaticVectorField->AssetImportData->ExtractFilenames(OutSourceFilePaths);
 	}
 }
\ No newline at end of file
diff --git a/Engine/Source/Editor/CurveTableEditor/Private/CurveTableEditorModule.cpp b/Engine/Source/Editor/CurveTableEditor/Private/CurveTableEditorModule.cpp
--- a/Engine/Source/Editor/CurveTableEditor/Private/CurveTableEditorModule.cpp
+++ b/Engine/Source/Editor/CurveTableEditor/Private/CurveTableEditorModule.cpp
@@ -5,6 +5,7 @@
 #include "CurveTableEditor.h"
 #include "Toolkits/ToolkitManager.h"
 #include "Engine/CurveTable.h"
+#include "EditorFramework/AssetImportData.h"
 
 IMPLEMENT_MODULE( FCurveTableEditorModule, CurveTableEditor );
 
diff --git a/Engine/Source/Editor/DestructibleMeshEditor/Private/SDestructibleMeshEditorViewport.cpp b/Engine/Source/Editor/DestructibleMeshEditor/Private/SDestructibleMeshEditorViewport.cpp
--- a/Engine/Source/Editor/DestructibleMeshEditor/Private/SDestructibleMeshEditorViewport.cpp
+++ b/Engine/Source/Editor/DestructibleMeshEditor/Private/SDestructibleMeshEditorViewport.cpp
@@ -245,11 +245,11 @@ bool FDestructibleMeshEditorViewportClient::CanRefreshFromStaticMesh()
 		return false;
 	}
 
-	FDateTime CurrentSourceTimestamp;
-	if (DestructibleMesh->SourceStaticMesh->AssetImportData == NULL || !FDateTime::Parse(DestructibleMesh->SourceStaticMesh->AssetImportData->SourceFileTimestamp,
-						 CurrentSourceTimestamp))
+	FAssetImportInfo* ImportInfo = DestructibleMesh->SourceStaticMesh->AssetImportData;
+	FDateTime CurrentSourceTimestamp = FDateTime::MinValue();
+	if (ImportInfo && ImportInfo->GetSourceFileData().Num() == 1)
 	{
-		CurrentSourceTimestamp = FDateTime::MinValue();
+		CurrentSourceTimestamp = ImportInfo->GetSourceFileData()[0].Timestamp;
 	}
 
 	return (CurrentSourceTimestamp > DestructibleMesh->SourceSMImportTimestamp);
diff --git a/Engine/Source/Editor/StaticMeshEditor/Private/StaticMeshEditor.cpp b/Engine/Source/Editor/StaticMeshEditor/Private/StaticMeshEditor.cpp
--- a/Engine/Source/Editor/StaticMeshEditor/Private/StaticMeshEditor.cpp
+++ b/Engine/Source/Editor/StaticMeshEditor/Private/StaticMeshEditor.cpp
@@ -1838,7 +1838,7 @@ void FStaticMeshEditor::ExecuteFindInExplorer()
 {
 	if ( ensure(StaticMesh->AssetImportData) )
 	{
-		const FString SourceFilePath = FReimportManager::ResolveImportFilename(StaticMesh->AssetImportData->SourceFilePath, StaticMesh);
+		const FString SourceFilePath = StaticMesh->AssetImportData->GetFirstFilename();
 		if ( SourceFilePath.Len() && IFileManager::Get().FileSize( *SourceFilePath ) != INDEX_NONE )
 		{
 			FPlatformProcess::ExploreFolder( *FPaths::GetPath(SourceFilePath) );
@@ -1853,7 +1853,7 @@ bool FStaticMeshEditor::CanExecuteSourceCommands() const
 		return false;
 	}
 
-	const FString& SourceFilePath = FReimportManager::ResolveImportFilename(StaticMesh->AssetImportData->SourceFilePath, StaticMesh);
+	const FString& SourceFilePath = StaticMesh->AssetImportData->GetFirstFilename();
 
 	return SourceFilePath.Len() && IFileManager::Get().FileSize(*SourceFilePath) != INDEX_NONE;
 }
diff --git a/Engine/Source/Editor/TextureEditor/Private/TextureEditorToolkit.cpp b/Engine/Source/Editor/TextureEditor/Private/TextureEditorToolkit.cpp
--- a/Engine/Source/Editor/TextureEditor/Private/TextureEditorToolkit.cpp
+++ b/Engine/Source/Editor/TextureEditor/Private/TextureEditorToolkit.cpp
@@ -1101,20 +1101,6 @@ void FTextureEditorToolkit::HandleReimportManagerPreReimport( UObject* InObject
 	SavedCompressionSetting = Texture->DeferCompression;
 	Texture->DeferCompression = true;
 
-	// Reduce the year to make sure the texture will always be reloaded even if it hasn't changed on disk
-	FDateTime TimeStamp;
-
-	if (FDateTime::Parse(Texture->SourceFileTimestamp, TimeStamp))
-	{
-		TimeStamp -= FTimespan(365, 0, 0, 0);
-	}
-	else
-	{
-		TimeStamp = FDateTime::MinValue();
-	}
-
-	Texture->SourceFileTimestamp = TimeStamp.ToString();
-
 	// Disable viewport rendering until the texture has finished re-importing
 	TextureViewport->DisableRendering();
 }
diff --git a/Engine/Source/Editor/UnrealEd/Private/AutoReimport/AssetSourceFilenameCache.cpp b/Engine/Source/Editor/UnrealEd/Private/AutoReimport/AssetSourceFilenameCache.cpp
new file mode 100644
--- /dev/null
+++ b/Engine/Source/Editor/UnrealEd/Private/AutoReimport/AssetSourceFilenameCache.cpp
@@ -0,0 +1,213 @@
+// Copyright 1998-2015 Epic Games, Inc. All Rights Reserved.
+
+#include "UnrealEd.h"
+#include "AutoReimport/AssetSourceFilenameCache.h"
+#include "AutoReimportUtilities.h"
+#include "AssetRegistryModule.h"
+
+FAssetSourceFilenameCache::FAssetSourceFilenameCache()
+{
+	IAssetRegistry& AssetRegistry = FModuleManager::LoadModuleChecked<FAssetRegistryModule>("AssetRegistry").Get();
+
+	AssetRegistry.OnAssetAdded().AddRaw(this, &FAssetSourceFilenameCache::HandleOnAssetAdded);
+	AssetRegistry.OnAssetRemoved().AddRaw(this, &FAssetSourceFilenameCache::HandleOnAssetRemoved);
+	AssetRegistry.OnAssetRenamed().AddRaw(this, &FAssetSourceFilenameCache::HandleOnAssetRenamed);
+
+	UAssetImportData::OnImportDataChanged.AddRaw(this, &FAssetSourceFilenameCache::HandleOnAssetUpdated);
+
+	TArray<FAssetData> Assets;
+	if (AssetRegistry.GetAllAssets(Assets))
+	{
+		for (auto& Asset : Assets)
+		{
+			HandleOnAssetAdded(Asset);
+		}
+	}
+}
+
+FAssetSourceFilenameCache& FAssetSourceFilenameCache::Get()
+{
+	static FAssetSourceFilenameCache Cache;
+	return Cache;
+}
+
+void FAssetSourceFilenameCache::Shutdown()
+{
+	auto* ModulePtr = FModuleManager::GetModulePtr<FAssetRegistryModule>("AssetRegistry");
+	if (ModulePtr)
+	{
+		ModulePtr->Get().OnAssetAdded().RemoveAll(this);
+		ModulePtr->Get().OnAssetRemoved().RemoveAll(this);
+		ModulePtr->Get().OnAssetRenamed().RemoveAll(this);
+	}
+	
+	AssetRenamedEvent.Clear();
+
+	UAssetImportData::OnImportDataChanged.RemoveAll(this);
+}
+
+TOptional<FAssetImportInfo> FAssetSourceFilenameCache::ExtractAssetImportInfo(const TArray<UObject::FAssetRegistryTag>& InTags)
+{
+	static const FName LegacySourceFilePathName("SourceFile");
+
+	TOptional<FAssetImportInfo> Info;
+
+	for (const auto& Tag : InTags)
+	{
+		if (Tag.Name == UObject::SourceFileTagName())
+		{
+			Info = FAssetImportInfo::FromJson(*Tag.Value);
+			// We're done
+			break;
+		}
+		else if (Tag.Name == LegacySourceFilePathName)
+		{
+			FAssetImportInfo Legacy;
+			Legacy.Insert(Tag.Value);
+			Info = Legacy;
+			// Keep looking for a newer json version
+		}
+	}
+
+	return Info;
+}
+
+TOptional<FAssetImportInfo> FAssetSourceFilenameCache::ExtractAssetImportInfo(const TMap<FName, FString>& InTags)
+{
+	static const FName LegacySourceFilePathName("SourceFile");
+
+	if (const FString* ImportDataString = InTags.Find(UObject::SourceFileTagName()))
+	{
+		return FAssetImportInfo::FromJson(*ImportDataString);
+	}
+	else if (const FString* LegacyFilename = InTags.Find(LegacySourceFilePathName))
+	{
+		FAssetImportInfo Legacy;
+		Legacy.Insert(*LegacyFilename);
+		return Legacy;
+	}
+	else
+	{
+		return TOptional<FAssetImportInfo>();
+	}
+}
+
+void FAssetSourceFilenameCache::HandleOnAssetAdded(const FAssetData& AssetData)
+{
+	TOptional<FAssetImportInfo> ImportData = ExtractAssetImportInfo(AssetData.TagsAndValues);
+	if (ImportData.IsSet())
+	{
+		for (const auto& SourceFile : ImportData->GetSourceFileData())
+		{
+			SourceFileToObjectPathCache.FindOrAdd(FPaths::GetCleanFilename(SourceFile.RelativeFilename)).Add(AssetData.ObjectPath);
+		}
+	}
+}
+
+void FAssetSourceFilenameCache::HandleOnAssetRemoved(const FAssetData& AssetData)
+{
+	TOptional<FAssetImportInfo> ImportData = ExtractAssetImportInfo(AssetData.TagsAndValues);
+	if (ImportData.IsSet())
+	{
+		for (auto& SourceFile : ImportData->GetSourceFileData())
+		{
+			FString CleanFilename = FPaths::GetCleanFilename(SourceFile.RelativeFilename);
+			if (auto* Objects = SourceFileToObjectPathCache.Find(CleanFilename))
+			{
+				Objects->Remove(AssetData.ObjectPath);
+				if (Objects->Num() == 0)
+				{
+					SourceFileToObjectPathCache.Remove(CleanFilename);
+				}
+			}
+		}
+	}
+}
+
+void FAssetSourceFilenameCache::HandleOnAssetRenamed(const FAssetData& AssetData, const FString& OldPath)
+{
+	TOptional<FAssetImportInfo> ImportData = ExtractAssetImportInfo(AssetData.TagsAndValues);
+	if (ImportData.IsSet())
+	{
+		FName OldPathName = *OldPath;
+
+		for (auto& SourceFile : ImportData->GetSourceFileData())
+		{
+			FString CleanFilename = FPaths::GetCleanFilename(SourceFile.RelativeFilename);
+
+			if (auto* Objects = SourceFileToObjectPathCache.Find(CleanFilename))
+			{
+				Objects->Remove(OldPathName);
+				if (Objects->Num() == 0)
+				{
+					SourceFileToObjectPathCache.Remove(CleanFilename);
+				}
+			}
+
+			SourceFileToObjectPathCache.FindOrAdd(CleanFilename).Add(AssetData.ObjectPath);
+		}
+	}
+
+	AssetRenamedEvent.Broadcast(AssetData, OldPath);
+}
+
+void FAssetSourceFilenameCache::HandleOnAssetUpdated(const FAssetImportInfo& OldData, const UAssetImportData* ImportData)
+{
+	FName ObjectPath = FName(*ImportData->GetOuter()->GetPathName());
+
+	for (const auto& SourceFile : OldData.GetSourceFileData())
+	{
+		FString CleanFilename = FPaths::GetCleanFilename(SourceFile.RelativeFilename);
+		if (auto* Objects = SourceFileToObjectPathCache.Find(CleanFilename))
+		{
+			Objects->Remove(ObjectPath);
+			if (Objects->Num() == 0)
+			{
+				SourceFileToObjectPathCache.Remove(CleanFilename);
+			}
+		}
+	}
+
+	for (const auto& SourceFile : ImportData->GetSourceFileData())
+	{
+		SourceFileToObjectPathCache.FindOrAdd(FPaths::GetCleanFilename(SourceFile.RelativeFilename)).Add(ObjectPath);
+	}
+}
+
+TArray<FAssetData> FAssetSourceFilenameCache::GetAssetsPertainingToFile(const IAssetRegistry& Registry, const FString& AbsoluteFilename) const
+{
+	TArray<FAssetData> Assets;
+	
+	if (const auto* ObjectPaths = SourceFileToObjectPathCache.Find(FPaths::GetCleanFilename(AbsoluteFilename)))
+	{
+		for (const FName& Path : *ObjectPaths)
+		{
+			FAssetData Asset = Registry.GetAssetByObjectPath(Path);
+			TOptional<FAssetImportInfo> ImportInfo = ExtractAssetImportInfo(Asset.TagsAndValues);
+			if (ImportInfo.IsSet())
+			{
+				auto AssetPackagePath = FPackageName::LongPackageNameToFilename(Asset.PackagePath.ToString() / TEXT(""));
+
+				// Attempt to find the matching source filename in the list of imported sorce files (generally there are only one of these)
+				const bool bWasImportedFromFile = ImportInfo->GetSourceFileData().ContainsByPredicate([&](const FAssetImportInfo::FSourceFile& File){
+
+					if (AbsoluteFilename == FPaths::ConvertRelativePathToFull(AssetPackagePath / File.RelativeFilename) ||
+						AbsoluteFilename == FPaths::ConvertRelativePathToFull(File.RelativeFilename))
+					{
+						return true;
+					}
+
+					return false;
+				});
+
+				if (bWasImportedFromFile)
+				{
+					Assets.Emplace();
+					Swap(Asset, Assets[Assets.Num() - 1]);
+				}
+			}
+		}
+	}
+
+	return Assets;
+}
\ No newline at end of file
diff --git a/Engine/Source/Editor/UnrealEd/Private/AutoReimport/AutoReimportManager.cpp b/Engine/Source/Editor/UnrealEd/Private/AutoReimport/AutoReimportManager.cpp
--- a/Engine/Source/Editor/UnrealEd/Private/AutoReimport/AutoReimportManager.cpp
+++ b/Engine/Source/Editor/UnrealEd/Private/AutoReimport/AutoReimportManager.cpp
@@ -13,6 +13,7 @@
 #include "AssetRegistryModule.h"
 #include "ReimportFeedbackContext.h"
 #include "MessageLogModule.h"
+#include "AssetSourceFilenameCache.h"
 
 #define LOCTEXT_NAMESPACE "AutoReimportManager"
 #define yield if (TimeLimit.Exceeded()) return
@@ -143,9 +144,6 @@ private:
 	/** Get the number of unprocessed changes that are not part of the current processing operation */
 	int32 GetNumUnprocessedChanges() const;
 
-	/** Get the progress text to display on the context notification */
-	FText GetProgressText() const;
-
 private:
 
 	/** A state machine holding information about the current state of the manager */
@@ -235,8 +233,7 @@ FAutoReimportManager::FAutoReimportManager()
 		MessageLogModule.RegisterLogListing("AssetReimport", LOCTEXT("AssetReimportLabel", "Asset Reimport"));
 	}
 
-	IAssetRegistry& AssetRegistry = FModuleManager::LoadModuleChecked<FAssetRegistryModule>("AssetRegistry").Get();
-	AssetRegistry.OnAssetRenamed().AddRaw(this, &FAutoReimportManager::OnAssetRenamed);
+	FAssetSourceFilenameCache::Get().OnAssetRenamed().AddRaw(this, &FAutoReimportManager::OnAssetRenamed);
 
 	// Only set this up for content directories if the user has this enabled
 	if (Settings->bMonitorContentDirectories)
@@ -322,7 +319,7 @@ void FAutoReimportManager::Destroy()
 {
 	if (auto* AssetRegistryModule = FModuleManager::GetModulePtr<FAssetRegistryModule>("AssetRegistry"))
 	{
-		AssetRegistryModule->Get().OnAssetRenamed().RemoveAll(this);
+		FAssetSourceFilenameCache::Get().OnAssetRenamed().RemoveAll(this);
 	}
 	
 	if (auto* Settings = GetMutableDefault<UEditorLoadingSavingSettings>())
@@ -357,23 +354,23 @@ void FAutoReimportManager::OnAssetRenamed(const FAssetData& AssetData, const FSt
 	// Additionally, we rename the source file if it matched the name of the asset before the rename/move.
 	//	- If we rename the source file, then we also update the reimport paths for the asset
 
-	TArray<FString> SourceFilesRelativeToOldPath;
-	for (const auto& Pair : AssetData.TagsAndValues)
+	
+	TOptional<FAssetImportInfo> ImportInfo = FAssetSourceFilenameCache::ExtractAssetImportInfo(AssetData.TagsAndValues);
+	if (!ImportInfo.IsSet())
 	{
-		if (Pair.Key.IsEqual(UObject::SourceFileTagName(), ENameCase::IgnoreCase, false))
-		{
-			SourceFilesRelativeToOldPath.Add(Pair.Value);
-		}
+		return;
 	}
 
+	const TArray<FAssetImportInfo::FSourceFile>& SourceFileData = ImportInfo->GetSourceFileData();
+
 	// We move the file with the asset provided it is the only file referenced, and sits right beside the uasset file
-	if (SourceFilesRelativeToOldPath.Num() == 1 && !SourceFilesRelativeToOldPath[0].GetCharArray().ContainsByPredicate([](const TCHAR Char) { return Char == '/' || Char == '\\'; }))
+	if (SourceFileData.Num() == 1 && !SourceFileData[0].RelativeFilename.GetCharArray().ContainsByPredicate([](const TCHAR Char) { return Char == '/' || Char == '\\'; }))
 	{
 		const FString AbsoluteSrcPath = FPaths::ConvertRelativePathToFull(FPackageName::LongPackageNameToFilename(FPackageName::GetLongPackagePath(OldPath) / TEXT("")));
 		const FString AbsoluteDstPath = FPaths::ConvertRelativePathToFull(FPackageName::LongPackageNameToFilename(AssetData.PackagePath.ToString() / TEXT("")));
 
 		const FString OldAssetName = FPackageName::GetLongPackageAssetName(FPackageName::ObjectPathToPackageName(OldPath));
-		FString NewFileName = FPaths::GetBaseFilename(SourceFilesRelativeToOldPath[0]);
+		FString NewFileName = FPaths::GetBaseFilename(SourceFileData[0].RelativeFilename);
 
 		bool bRequireReimportPathUpdate = false;
 		if (PackageTools::SanitizePackageName(NewFileName) == OldAssetName)
@@ -382,13 +379,12 @@ void FAutoReimportManager::OnAssetRenamed(const FAssetData& AssetData, const FSt
 			bRequireReimportPathUpdate = true;
 		}
 
-		const FString SrcFile = AbsoluteSrcPath / SourceFilesRelativeToOldPath[0];
-		const FString DstFile = AbsoluteDstPath / NewFileName + TEXT(".") + FPaths::GetExtension(SourceFilesRelativeToOldPath[0]);
+		const FString SrcFile = AbsoluteSrcPath / SourceFileData[0].RelativeFilename;
+		const FString DstFile = AbsoluteDstPath / NewFileName + TEXT(".") + FPaths::GetExtension(SourceFileData[0].RelativeFilename);
 
-		// We can't do this if multiple assets reference the same file. We should be checking for > 1 referencing asset, but the asset registry
-		// filter lookup won't return the recently renamed package because it will be Empty by now, so we check for *anything* referencing the asset (assuming that we'll never find *this* asset).
+		// We can't do this if multiple assets reference the same file
 		const IAssetRegistry& Registry = FModuleManager::LoadModuleChecked<FAssetRegistryModule>("AssetRegistry").Get();
-		if (Utils::FindAssetsPertainingToFile(Registry, SrcFile).Num() > 0)
+		if (Utils::FindAssetsPertainingToFile(Registry, SrcFile).Num() > 1)
 		{
 			return;
 		}
@@ -417,28 +413,6 @@ int32 FAutoReimportManager::GetNumUnprocessedChanges() const
 	}, 0);
 }
 
-FText FAutoReimportManager::GetProgressText() const
-{
-	FFormatOrderedArguments Args;
-	{
-		const int32 Progress = Utils::Reduce(DirectoryMonitors, [](const FContentDirectoryMonitor& Monitor, int32 Total){
-			return Total + Monitor.GetWorkProgress();
-		}, 0);
-
-		Args.Add(Progress);
-	}
-
-	{
-		const int32 Total = Utils::Reduce(DirectoryMonitors, [](const FContentDirectoryMonitor& Monitor, int32 InTotal){
-			return InTotal + Monitor.GetTotalWork();
-		}, 0);
-
-		Args.Add(Total);
-	}
-
-	return FText::Format(LOCTEXT("ProcessingChanges", "Processing outstanding content changes ({0} of {1})..."), Args);
-}
-
 TOptional<ECurrentState> FAutoReimportManager::ProcessAdditions(const FTimeLimit& TimeLimit)
 {
 	auto NewState = HandlePauseAbort(ECurrentState::ProcessAdditions);
@@ -536,6 +510,8 @@ TOptional<ECurrentState> FAutoReimportManager::ProcessDeletions()
 		Monitor.ExtractAssetsToDelete(Registry, AssetsToDelete);
 	}
 
+	FeedbackContextOverride->MainTask->EnterProgressFrame(AssetsToDelete.Num());
+
 	if (AssetsToDelete.Num() > 0)
 	{
 		for (const auto& AssetData : AssetsToDelete)
@@ -604,6 +580,8 @@ TOptional<ECurrentState> FAutoReimportManager::Abort()
 		Monitor.Abort();
 	}
 
+	PackagesToSave.Empty();
+
 	Cleanup();
 	return ECurrentState::Idle;
 }
@@ -648,30 +626,26 @@ TOptional<ECurrentState> FAutoReimportManager::Idle()
 	{
 		State = EProcessState::Running;
 
+		int32 TotalWork = 0;
 		// We have some changes so kick off the process
 		for (auto& Monitor : DirectoryMonitors)
 		{
-			Monitor.StartProcessing();
+			TotalWork += Monitor.StartProcessing();
 		}
 
-		// Check that we actually have anything to do before kicking off a process
-		const int32 WorkTotal = Utils::Reduce(DirectoryMonitors, [](const FContentDirectoryMonitor& Monitor, int32 Total){
-			return Total + Monitor.GetTotalWork();
-		}, 0);
-
-		if (WorkTotal > 0)
+		if (TotalWork > 0)
 		{
 			if (!FeedbackContextOverride.IsValid())
 			{
 				// Create a new feedback context override
 				FeedbackContextOverride = MakeShareable(new FReimportFeedbackContext(
-					TAttribute<FText>(this, &FAutoReimportManager::GetProgressText),
 					FSimpleDelegate::CreateSP(this, &FAutoReimportManager::OnPauseClicked),
 					FSimpleDelegate::CreateSP(this, &FAutoReimportManager::OnAbortClicked))
 				);
 			}
 
-			FeedbackContextOverride->Show();
+			FeedbackContextOverride->Show(TotalWork);
+
 			return ECurrentState::ProcessAdditions;
 		}
 	}
@@ -681,10 +655,7 @@ TOptional<ECurrentState> FAutoReimportManager::Idle()
 
 void FAutoReimportManager::Cleanup()
 {
-	if (GetNumUnprocessedChanges() == 0)
-	{
-		FeedbackContextOverride->Hide();
-	}
+	FeedbackContextOverride->Hide();
 }
 
 void FAutoReimportManager::Tick(float DeltaTime)
diff --git a/Engine/Source/Editor/UnrealEd/Private/AutoReimport/AutoReimportUtilities.cpp b/Engine/Source/Editor/UnrealEd/Private/AutoReimport/AutoReimportUtilities.cpp
--- a/Engine/Source/Editor/UnrealEd/Private/AutoReimport/AutoReimportUtilities.cpp
+++ b/Engine/Source/Editor/UnrealEd/Private/AutoReimport/AutoReimportUtilities.cpp
@@ -5,6 +5,7 @@
 #include "EditorReimportHandler.h"
 #include "ARFilter.h"
 #include "IAssetRegistry.h"
+#include "AssetSourceFilenameCache.h"
 
 bool MatchExtensionString(const TCHAR* Filename, const TCHAR* Extensions)
 {
@@ -188,37 +189,7 @@ namespace Utils
 {
 	TArray<FAssetData> FindAssetsPertainingToFile(const IAssetRegistry& Registry, const FString& AbsoluteFilename)
 	{
-		TArray<FAssetData> Assets;
-
-		const FString LeafName = FPaths::GetCleanFilename(AbsoluteFilename);
-		const FName TagName = UObject::SourceFileTagName();
-
-		FARFilter Filter;
-		Filter.SourceFilenames.Add(*LeafName);
-		Filter.bIncludeOnlyOnDiskAssets = true;
-		Registry.GetAssets(Filter, Assets);
-
-		Assets.RemoveAll([&](const FAssetData& Asset){
-			for (const auto& Pair : Asset.TagsAndValues)
-			{
-				// We don't compare numbers on FNames for this check because the tag "ReimportPath" may exist multiple times
-				const bool bCompareNumber = false;
-				if (Pair.Key.IsEqual(TagName, ENameCase::IgnoreCase, bCompareNumber))
-				{
-					// Either relative to the asset itself, or relative to the base path, or absolute.
-					// Would ideally use FReimportManager::ResolveImportFilename here, but we don't want to ask the file system whether the file exists.
-					auto RelativeToPackagePath = FPackageName::LongPackageNameToFilename(Asset.PackagePath.ToString() / TEXT("")) / Pair.Value;
-					if (AbsoluteFilename == FPaths::ConvertRelativePathToFull(RelativeToPackagePath) ||
-						AbsoluteFilename == FPaths::ConvertRelativePathToFull(Pair.Value))
-					{
-						return false;
-					}
-				}
-			}
-			return true;
-		});
-		
-		return Assets;
+		return FAssetSourceFilenameCache::Get().GetAssetsPertainingToFile(Registry, AbsoluteFilename);
 	}
 
 	TArray<FString> ExtractSourceFilePaths(UObject* Object)
@@ -236,11 +207,19 @@ namespace Utils
 		const FName TagName = UObject::SourceFileTagName();
 		for (const auto& Tag : TagList)
 		{
-			// We don't compare numbers on FNames for this check because the tag "ReimportPath" may exist multiple times
-			const bool bCompareNumber = false;
-			if (Tag.Name.IsEqual(TagName, ENameCase::IgnoreCase, bCompareNumber))
+			if (Tag.Name == TagName)
 			{
-				OutSourceFiles.Add(FReimportManager::ResolveImportFilename(Tag.Value, Object));
+				int32 PreviousNum = OutSourceFiles.Num();
+
+				TOptional<FAssetImportInfo> ImportInfo = FAssetImportInfo::FromJson(Tag.Value);
+				if (ImportInfo.IsSet())
+				{
+					for (const auto& File : ImportInfo->GetSourceFileData())
+					{
+						OutSourceFiles.Emplace(UAssetImportData::ResolveImportFilename(File.RelativeFilename, Object->GetOutermost()));
+					}
+				}
+				break;
 			}
 		}
 	}
diff --git a/Engine/Source/Editor/UnrealEd/Private/AutoReimport/AutoReimportUtilities.h b/Engine/Source/Editor/UnrealEd/Private/AutoReimport/AutoReimportUtilities.h
--- a/Engine/Source/Editor/UnrealEd/Private/AutoReimport/AutoReimportUtilities.h
+++ b/Engine/Source/Editor/UnrealEd/Private/AutoReimport/AutoReimportUtilities.h
@@ -92,54 +92,6 @@ private:
 	double StartTime;
 };
 
-/** Simple helper struct to ease the caching of MD5 hashes */
-struct FMD5Hash
-{
-	/** Default constructor */
-	FMD5Hash() : bIsValid(false) {}
-
-	/** Check whether this has hash is valid or not */
-	bool IsValid() const { return bIsValid; }
-
-	/** Set up the MD5 hash from a container */
-	void Set(FMD5& MD5)
-	{
-		MD5.Final(Bytes);
-		bIsValid = true;
-	}
-
-	/** Compare one hash with another */
-	friend bool operator==(const FMD5Hash& LHS, const FMD5Hash& RHS)
-	{
-		return LHS.bIsValid == RHS.bIsValid && (!LHS.bIsValid || FMemory::Memcmp(LHS.Bytes, RHS.Bytes, 16) == 0);
-	}
-
-	/** Compare one hash with another */
-	friend bool operator!=(const FMD5Hash& LHS, const FMD5Hash& RHS)
-	{
-		return LHS.bIsValid != RHS.bIsValid || (LHS.bIsValid && FMemory::Memcmp(LHS.Bytes, RHS.Bytes, 16) != 0);
-	}
-
-	/** Serialise this hash */
-	friend FArchive& operator<<(FArchive& Ar, FMD5Hash& Hash)
-	{
-		Ar << Hash.bIsValid;
-		if (Hash.bIsValid)
-		{
-			Ar.Serialize(Hash.Bytes, 16);
-		}
-
-		return Ar;
-	}
-
-private:
-	/** Whether this hash is valid or not */
-	bool bIsValid;
-
-	/** The bytes this hash comprises */
-	uint8 Bytes[16];
-};
-
 /** A rule that checks whether a file is applicable or not */
 struct IMatchRule
 {
@@ -212,7 +164,7 @@ namespace Utils
 		}
 		return Accumulator;
 	}
-	
+
 	/** Find a list of assets that were once imported from the specified filename */
 	TArray<FAssetData> FindAssetsPertainingToFile(const IAssetRegistry& Registry, const FString& AbsoluteFilename);
 
diff --git a/Engine/Source/Editor/UnrealEd/Private/AutoReimport/ContentDirectoryMonitor.cpp b/Engine/Source/Editor/UnrealEd/Private/AutoReimport/ContentDirectoryMonitor.cpp
--- a/Engine/Source/Editor/UnrealEd/Private/AutoReimport/ContentDirectoryMonitor.cpp
+++ b/Engine/Source/Editor/UnrealEd/Private/AutoReimport/ContentDirectoryMonitor.cpp
@@ -9,6 +9,7 @@
 #include "ObjectTools.h"
 #include "AssetToolsModule.h"
 #include "ReimportFeedbackContext.h"
+#include "AssetSourceFilenameCache.h"
 
 #define LOCTEXT_NAMESPACE "ContentDirectoryMonitor"
 
@@ -33,6 +34,13 @@ FFileCacheConfig GenerateFileCacheConfig(const FString& InPath, const FMatchRule
 	Config.PathType = EPathType::Relative;
 
 	Config.bDetectChangesSinceLastRun = GetDefault<UEditorLoadingSavingSettings>()->bDetectChangesOnRestart;
+
+	// We only detect changes for when the file *contents* have changed (not its timestamp)
+	Config
+		.DetectMoves(true)
+		.DetectChangesFor(FFileCacheConfig::Timestamp, false)
+		.DetectChangesFor(FFileCacheConfig::FileHash, true);
+
 	return Config;
 }
 
@@ -81,56 +89,57 @@ void FContentDirectoryMonitor::Tick()
 	}
 }
 
-void FContentDirectoryMonitor::StartProcessing()
+int32 FContentDirectoryMonitor::StartProcessing()
 {
-	WorkProgress = TotalWork = 0;
-
 	// We only process things that haven't changed for a given threshold
 	auto& FileManager = IFileManager::Get();
 	const FDateTime Threshold = FDateTime::UtcNow() - FTimespan(0, 0, GetDefault<UEditorLoadingSavingSettings>()->AutoReimportThreshold);
+
+	// Get all the changes that have happend beyond our import threshold
 	auto OutstandingChanges = Cache.FilterOutstandingChanges([&](const FUpdateCacheTransaction& Transaction, const FDateTime& TimeOfChange){
-		return TimeOfChange < Threshold;
+		return TimeOfChange <= Threshold;
 	});
 
-	if (OutstandingChanges.Num() != 0)
+	if (OutstandingChanges.Num() == 0)
 	{
-		const auto* Settings = GetDefault<UEditorLoadingSavingSettings>();
+		return 0;
+	}
 
-		for (auto& Transaction : OutstandingChanges)
+	const auto* Settings = GetDefault<UEditorLoadingSavingSettings>();
+	for (auto& Transaction : OutstandingChanges)
+	{
+		switch(Transaction.Action)
 		{
-			switch(Transaction.Action)
-			{
-				case EFileAction::Added:
-					if (Settings->bAutoCreateAssets && !MountedContentPath.IsEmpty())
-					{
-						AddedFiles.Emplace(MoveTemp(Transaction));
-					}
-					else
-					{
-						Cache.CompleteTransaction(MoveTemp(Transaction));
-					}
-					break;
+			case EFileAction::Added:
+				if (Settings->bAutoCreateAssets && !MountedContentPath.IsEmpty())
+				{
+					AddedFiles.Emplace(MoveTemp(Transaction));
+				}
+				else
+				{
+					Cache.CompleteTransaction(MoveTemp(Transaction));
+				}
+				break;
 
-				case EFileAction::Moved:
-				case EFileAction::Modified:
-					ModifiedFiles.Emplace(MoveTemp(Transaction));
-					break;
+			case EFileAction::Moved:
+			case EFileAction::Modified:
+				ModifiedFiles.Emplace(MoveTemp(Transaction));
+				break;
 
-				case EFileAction::Removed:
-					if (Settings->bAutoDeleteAssets && !MountedContentPath.IsEmpty())
-					{
-						DeletedFiles.Emplace(MoveTemp(Transaction));
-					}
-					else
-					{
-						Cache.CompleteTransaction(MoveTemp(Transaction));
-					}
-					break;
-			}
+			case EFileAction::Removed:
+				if (Settings->bAutoDeleteAssets && !MountedContentPath.IsEmpty())
+				{
+					DeletedFiles.Emplace(MoveTemp(Transaction));
+				}
+				else
+				{
+					Cache.CompleteTransaction(MoveTemp(Transaction));
+				}
+				break;
 		}
-
-		TotalWork = AddedFiles.Num() + ModifiedFiles.Num() + DeletedFiles.Num();
 	}
+
+	return AddedFiles.Num() + ModifiedFiles.Num() + DeletedFiles.Num();
 }
 
 UObject* AttemptImport(UClass* InFactoryType, UPackage* Package, FName InName, bool& bCancelled, const FString& FullFilename)
@@ -140,7 +149,7 @@ UObject* AttemptImport(UClass* InFactoryType, UPackage* Package, FName InName, b
 	if (UFactory* Factory = NewObject<UFactory>(GetTransientPackage(), InFactoryType))
 	{
 		Factory->AddToRoot();
-		if (Factory->ConfigureProperties())
+		if (Factory->ConfigureProperties() && Factory->SupportedClass)
 		{
 			Asset = UFactory::StaticImportObject(Factory->SupportedClass, Package, InName, RF_Public | RF_Standalone, bCancelled, *FullFilename, nullptr, Factory);
 		}
@@ -157,7 +166,7 @@ void FContentDirectoryMonitor::ProcessAdditions(const IAssetRegistry& Registry,
 	{
 		auto& Addition = AddedFiles[Index];
 
-		WorkProgress += 1;
+		Context.MainTask->EnterProgressFrame();
 
 		if (bCancelled)
 		{
@@ -175,13 +184,14 @@ void FContentDirectoryMonitor::ProcessAdditions(const IAssetRegistry& Registry,
 			FText Message;
 			if (ExistingReferences.Num() == 1)
 			{
-				Message = FText::Format(LOCTEXT("Info_AlreadyImported_Single", "Ignoring new file {0} as it's already imported by {1}."), FText::FromString(FPaths::GetCleanFilename(FullFilename)), FText::FromName(ExistingReferences[0].AssetName));
+				FText Message = FText::Format(LOCTEXT("Info_AlreadyImported_Single", "Ignoring new file {0} as it's already imported by {1}."), FText::FromString(FPaths::GetCleanFilename(FullFilename)), FText::FromName(ExistingReferences[0].AssetName));
+				Context.GetMessageLog().Message(EMessageSeverity::Info, Message);
 			}
 			else
 			{
-				Message = FText::Format(LOCTEXT("Info_AlreadyImported_Multiple", "Ignoring new file {0} as it's already imported by {1} assets."), FText::FromString(FPaths::GetCleanFilename(FullFilename)), FText::AsNumber(ExistingReferences.Num()));
+				FText Message = FText::Format(LOCTEXT("Info_AlreadyImported_Multiple", "Ignoring new file {0} as it's already imported by {1} assets."), FText::FromString(FPaths::GetCleanFilename(FullFilename)), FText::AsNumber(ExistingReferences.Num()));
+				Context.GetMessageLog().Message(EMessageSeverity::Info, Message);
 			}
-			Context.AddMessage(EMessageSeverity::Info, Message);
 
 			// We don't try and import new files that are already associated with an asset
 			Cache.CompleteTransaction(MoveTemp(Addition));
@@ -286,6 +296,8 @@ void FContentDirectoryMonitor::ProcessModifications(const IAssetRegistry& Regist
 
 	for (int32 Index = 0; Index < ModifiedFiles.Num(); ++Index)
 	{
+		Context.MainTask->EnterProgressFrame();
+
 		auto& Change = ModifiedFiles[Index];
 		const FString FullFilename = Cache.GetDirectory() + Change.Filename.Get();
 
@@ -300,12 +312,21 @@ void FContentDirectoryMonitor::ProcessModifications(const IAssetRegistry& Regist
 				UObject* Asset = Assets[0].GetAsset();
 				if (Asset && Utils::ExtractSourceFilePaths(Asset).Num() == 1)
 				{
+					UPackage* ExistingPackage = Asset->GetOutermost();
+
 					const bool bAssetWasDirty = IsAssetDirty(Asset);
 
 					const FString NewAssetName = ObjectTools::SanitizeObjectName(FPaths::GetBaseFilename(Change.Filename.Get()));
 					const FString PackagePath = PackageTools::SanitizePackageName(MountedContentPath / FPaths::GetPath(Change.Filename.Get()));
 					const FString FullDestPath = PackagePath / NewAssetName;
 
+					if (ExistingPackage && ExistingPackage->FileName.ToString() == FullDestPath)
+					{
+						// No need to process this asset - it's already been moved to the right location
+						Cache.CompleteTransaction(MoveTemp(Change));
+						continue;
+					}
+
 					const FText SrcPathText = FText::FromString(Assets[0].PackageName.ToString()),
 						DstPathText = FText::FromString(FullDestPath);
 
@@ -344,6 +365,20 @@ void FContentDirectoryMonitor::ProcessModifications(const IAssetRegistry& Regist
 				UObject* Asset = AssetData.GetAsset();
 				if (Asset)
 				{
+					TArray<UObject::FAssetRegistryTag> Tags;
+					Asset->GetAssetRegistryTags(Tags);
+					TOptional<FAssetImportInfo> Info = FAssetSourceFilenameCache::ExtractAssetImportInfo(Tags);
+
+					// Check if the source file that this asset last imported was the same as the one we're going to reimport.
+					// If it is, there's no reason to auto-reimport it
+					if (Info.IsSet() && Info->GetSourceFileData().Num() == 1)
+					{
+						if (Info->GetSourceFileData()[0].FileHash == Change.FileData.FileHash)
+						{
+							continue;
+						}
+					}
+
 					const bool bAssetWasDirty = IsAssetDirty(Asset);
 					if (!ReimportManager->Reimport(Asset, false /* Ask for new file */, false /* Show notification */))
 					{
@@ -363,8 +398,7 @@ void FContentDirectoryMonitor::ProcessModifications(const IAssetRegistry& Regist
 
 		// Let the cache know that we've dealt with this change
 		Cache.CompleteTransaction(MoveTemp(Change));
-		
-		WorkProgress += 1;
+
 		if (TimeLimit.Exceeded())
 		{
 			ModifiedFiles.RemoveAt(0, Index + 1);
@@ -388,7 +422,6 @@ void FContentDirectoryMonitor::ExtractAssetsToDelete(const IAssetRegistry& Regis
 		Cache.CompleteTransaction(MoveTemp(Deletion));
 	}
 
-	WorkProgress += DeletedFiles.Num();
 	DeletedFiles.Empty();
 }
 
@@ -416,8 +449,6 @@ void FContentDirectoryMonitor::Abort()
 	{
 		Cache.CompleteTransaction(MoveTemp(Change));
 	}
-
-	WorkProgress = TotalWork;
 }
 
 #undef LOCTEXT_NAMESPACE
diff --git a/Engine/Source/Editor/UnrealEd/Private/AutoReimport/ContentDirectoryMonitor.h b/Engine/Source/Editor/UnrealEd/Private/AutoReimport/ContentDirectoryMonitor.h
--- a/Engine/Source/Editor/UnrealEd/Private/AutoReimport/ContentDirectoryMonitor.h
+++ b/Engine/Source/Editor/UnrealEd/Private/AutoReimport/ContentDirectoryMonitor.h
@@ -6,6 +6,8 @@
 #include "FileCache.h"
 #include "MessageLog.h"
 
+class FReimportFeedbackContext;
+
 /** Class responsible for watching a specific content directory for changes */
 class FContentDirectoryMonitor
 {
@@ -23,13 +25,13 @@ public:
 	void Tick();
 
 	/** Start processing any outstanding changes this monitor is aware of */
-	void StartProcessing();
+	int32 StartProcessing();
 
 	/** Extract the files we need to import from our outstanding changes (happens first)*/ 
-	void ProcessAdditions(const IAssetRegistry& Registry, const FTimeLimit& TimeLimit, TArray<UPackage*>& OutPackagesToSave, const TMap<FString, TArray<UFactory*>>& InFactoriesByExtension, class FReimportFeedbackContext& Context);
+	void ProcessAdditions(const IAssetRegistry& Registry, const FTimeLimit& TimeLimit, TArray<UPackage*>& OutPackagesToSave, const TMap<FString, TArray<UFactory*>>& InFactoriesByExtension, FReimportFeedbackContext& Context);
 
 	/** Process the outstanding changes that we have cached */
-	void ProcessModifications(const IAssetRegistry& Registry, const FTimeLimit& TimeLimit, TArray<UPackage*>& OutPackagesToSave, class FReimportFeedbackContext& Context);
+	void ProcessModifications(const IAssetRegistry& Registry, const FTimeLimit& TimeLimit, TArray<UPackage*>& OutPackagesToSave, FReimportFeedbackContext& Context);
 
 	/** Extract the assets we need to delete from our outstanding changes (happens last) */ 
 	void ExtractAssetsToDelete(const IAssetRegistry& Registry, TArray<FAssetData>& OutAssetsToDelete);
@@ -57,12 +59,6 @@ public:
 	/** Get the number of outstanding changes that we potentially have to process (when not already processing) */
 	int32 GetNumUnprocessedChanges() const { return Cache.GetNumOutstandingChanges(); }
 
-	/** Get the total amount of work this monitor has to perform in the current processing operation */
-	int32 GetTotalWork() const { return TotalWork; }
-
-	/** Get the total amount of work this monitor has performed in the current processing operation */
-	int32 GetWorkProgress() const { return WorkProgress; }
-
 private:
 
 	/** The file cache that monitors and reflects the content directory. */
@@ -74,9 +70,6 @@ private:
 	/** A list of file system changes that are due to be processed */
 	TArray<FUpdateCacheTransaction> AddedFiles, ModifiedFiles, DeletedFiles;
 
-	/** The number of changes we've processed out of the OutstandingChanges array */
-	int32 TotalWork, WorkProgress;
-
 	/** The last time we attempted to save the cache file */
 	double LastSaveTime;
 
diff --git a/Engine/Source/Editor/UnrealEd/Private/AutoReimport/FileCache.cpp b/Engine/Source/Editor/UnrealEd/Private/AutoReimport/FileCache.cpp
--- a/Engine/Source/Editor/UnrealEd/Private/AutoReimport/FileCache.cpp
+++ b/Engine/Source/Editor/UnrealEd/Private/AutoReimport/FileCache.cpp
@@ -61,42 +61,6 @@ EFileAction ToFileAction(FFileChangeData::EFileChangeAction InAction)
 	}
 }
 
-/** Synchronously read a file's MD5 */
-FMD5Hash ReadFileMD5(const FString& Filename, TArray<uint8>* Buffer = nullptr)
-{
-	FArchive* Ar = IFileManager::Get().CreateFileReader(*Filename);
-
-	FMD5Hash Hash;
-	if (Ar)
-	{
-		TArray<uint8> LocalScratch;
-		if (!Buffer)
-		{
-			LocalScratch.SetNumUninitialized(1024*64);
-			Buffer = &LocalScratch; //-V506
-		}
-		FMD5 MD5;
-
-		const int64 Size = Ar->TotalSize();
-		int64 Position = 0;
-
-		// Read in BufferSize chunks
-		while (Position < Size)
-		{
-			const auto ReadNum = FMath::Min(Size - Position, (int64)Buffer->Num());
-			Ar->Serialize(Buffer->GetData(), ReadNum);
-			MD5.Update(Buffer->GetData(), ReadNum);
-
-			Position += ReadNum;
-		}
-
-		Hash.Set(MD5);
-		delete Ar;
-	}
-
-	return Hash;
-}
-
 const FGuid FFileCacheCustomVersion::Key(0x8E7DDCB3, 0x80DA47BB, 0x9FD346A2, 0x93984DF6);
 FCustomVersionRegistration GRegisterFileCacheVersion(FFileCacheCustomVersion::Key, FFileCacheCustomVersion::Latest, TEXT("FileCacheVersion"));
 
@@ -181,7 +145,7 @@ private:
 
 FAsyncTaskThread AsyncTaskThread;
 
-/** Threading stategy for FAsyncFileHasher:
+/** Threading strategy for FAsyncFileHasher:
  *	The task is constructed on the main thread with its Data.
  *	The array 'Data' *never* changes size. The task thread moves along setting file hashes, while the main thread
  * 	trails behind accessing the completed entries. We should thus never have 2 threads accessing the same memory,
@@ -197,6 +161,7 @@ FAsyncFileHasher::FAsyncFileHasher(TArray<FFilenameAndHash> InFilesThatNeedHashi
 
 TArray<FFilenameAndHash> FAsyncFileHasher::GetCompletedData()
 {
+	// Don't need to lock here since the thread will never look at the array before CurrentIndex.
 	TArray<FFilenameAndHash> Local;
 	const int32 CompletedIndex = CurrentIndex.GetValue();
 
@@ -225,7 +190,7 @@ IAsyncFileCacheTask::EProgressResult FAsyncFileHasher::Tick(const FTimeLimit& Li
 	for (; CurrentIndex.GetValue() < Data.Num(); )
 	{
 		const auto Index = CurrentIndex.GetValue();
-		Data[Index].FileHash = ReadFileMD5(Data[Index].AbsoluteFilename, &ScratchBuffer);
+		Data[Index].FileHash = FMD5Hash::HashFile(*Data[Index].AbsoluteFilename, &ScratchBuffer);
 
 		CurrentIndex.Increment();
 
@@ -238,7 +203,7 @@ IAsyncFileCacheTask::EProgressResult FAsyncFileHasher::Tick(const FTimeLimit& Li
 	return EProgressResult::Finished;
 }
 
-/** Threading stategy for FAsyncDirectoryReader:
+/** Threading strategy for FAsyncDirectoryReader:
  *	The directory reader owns the cached and live state until it has completely finished. Once IsComplete() is true, the main thread can 
  *	have access to both the cached and farmed data.
  */
@@ -388,6 +353,9 @@ FFileCache::FFileCache(const FFileCacheConfig& InConfig)
 	// Ensure the directory has a trailing /
 	Config.Directory /= TEXT("");
 
+	// bDetectMoves implies bRequireFileHashes
+	Config.bRequireFileHashes = Config.bRequireFileHashes || Config.bDetectMoves;
+
 	DirectoryReader = MakeShareable(new FAsyncDirectoryReader(Config.Directory, Config.PathType));
 	DirectoryReader->SetMatchRules(Config.Rules);
 
@@ -425,6 +393,7 @@ void FFileCache::Destroy()
 
 	DirectoryReader = nullptr;
 	AsyncFileHasher = nullptr;
+	DirtyFileHasher = nullptr;
 
 	DirtyFiles.Empty();
 	CachedDirectoryState = FDirectoryState();
@@ -579,7 +548,7 @@ TOptional<FString> FFileCache::GetTransactionPath(const FString& InAbsolutePath)
 	}
 }
 
-void FFileCache::DiffDirtyFiles(TMap<FImmutableString, FDateTime>& InDirtyFiles, TArray<FUpdateCacheTransaction>& OutTransactions, const FDirectoryState* InFileSystemState, const FDateTime& ThresholdTime) const
+void FFileCache::DiffDirtyFiles(TMap<FImmutableString, FFileData>& InDirtyFiles, TArray<FUpdateCacheTransaction>& OutTransactions, const FDirectoryState* InFileSystemState) const
 {
 	TArray<uint8> ScratchBuffer;
 	ScratchBuffer.SetNumUninitialized(1024*1024);
@@ -592,11 +561,6 @@ void FFileCache::DiffDirtyFiles(TMap<FImmutableString, FDateTime>& InDirtyFiles,
 
 	for (const auto& Pair : InDirtyFiles)
 	{
-		if (Pair.Value > ThresholdTime)
-		{
-			continue;
-		}
-
 		const auto& File = Pair.Key;
 		FString AbsoluteFilename = GetAbsolutePath(File.Get());
 		
@@ -605,13 +569,37 @@ void FFileCache::DiffDirtyFiles(TMap<FImmutableString, FDateTime>& InDirtyFiles,
 		const bool bFileExists = InFileSystemState ? InFileSystemState->Files.Find(File) != nullptr : PlatformFile.FileExists(*AbsoluteFilename);
 		if (bFileExists)
 		{
-			const FFileData FileData = InFileSystemState ? *InFileSystemState->Files.Find(File) : FFileData(FileManager.GetTimeStamp(*AbsoluteFilename), ReadFileMD5(AbsoluteFilename, &ScratchBuffer));
+			FFileData FileData;
+			if (const auto* FoundData = InFileSystemState ? InFileSystemState->Files.Find(File) : nullptr)
+			{
+				FileData = *FoundData;
+			}
+			else
+			{
+				// The dirty file timestamp is the time that the file was dirtied, not necessarily its modification time
+				FileData = FFileData(FileManager.GetTimeStamp(*AbsoluteFilename), Pair.Value.FileHash);
+			}
+
+			if (Config.bRequireFileHashes && !FileData.FileHash.IsValid())
+			{
+				// We don't have this file's hash yet. Temporarily ignore it.
+				continue;
+			}
 
 			// Do we think it exists in the cache?
 			if (CachedState)
 			{
-				// Has it changed?
-				if (CachedState->Timestamp != FileData.Timestamp)
+				// A file has changed if its hash is now different
+				if (Config.bRequireFileHashes &&
+					Config.ChangeDetectionBits[FFileCacheConfig::FileHash] &&
+					CachedState->FileHash != FileData.FileHash
+					)
+				{
+					ModifiedFiles.Add(File, FileData);
+				}
+				// or the timestamp has changed
+				else if (Config.ChangeDetectionBits[FFileCacheConfig::Timestamp] &&
+					CachedState->Timestamp != FileData.Timestamp)
 				{
 					ModifiedFiles.Add(File, FileData);
 				}
@@ -647,13 +635,27 @@ void FFileCache::DiffDirtyFiles(TMap<FImmutableString, FDateTime>& InDirtyFiles,
 	// Rename / move detection
 	if (Config.bDetectMoves)
 	{
+		bool bHavePendingHashes = false;
+
+		// Remove any additions that don't have their hash generated yet
+		for (auto AdIt = AddedFiles.CreateIterator(); AdIt; ++AdIt)
+		{
+			if (!AdIt.Value().FileHash.IsValid())
+			{
+				bHavePendingHashes = true;
+				AdIt.RemoveCurrent();
+			}
+		}
+
+		// We can only detect renames or moves for files that have had their file hash harvested.
+		// If we can't find a valid move destination for this file, and we have pending hashes, ignore the removal until we can be sure it's not a move
 		for (auto RemoveIt = RemovedFiles.CreateIterator(); RemoveIt; ++RemoveIt)
 		{
 			const auto* CachedState = CachedDirectoryState.Files.Find(*RemoveIt);
 			if (CachedState && CachedState->FileHash.IsValid())
 			{
 				for (auto AdIt = AddedFiles.CreateIterator(); AdIt; ++AdIt)
-				{				
+				{
 					if (AdIt.Value().FileHash == CachedState->FileHash)
 					{
 						// Found a move destination!
@@ -661,10 +663,19 @@ void FFileCache::DiffDirtyFiles(TMap<FImmutableString, FDateTime>& InDirtyFiles,
 
 						AdIt.RemoveCurrent();
 						RemoveIt.RemoveCurrent();
-						break;
+						goto next;
 					}
 				}
+
+				// We can't be sure this isn't a move (yet) so temporarily ignore this
+				if (bHavePendingHashes)
+				{
+					RemoveIt.RemoveCurrent();
+				}
 			}
+
+		next:
+			continue;
 		}
 	}
 
@@ -689,25 +700,29 @@ void FFileCache::DiffDirtyFiles(TMap<FImmutableString, FDateTime>& InDirtyFiles,
 
 TArray<FUpdateCacheTransaction> FFileCache::GetOutstandingChanges()
 {
+	HarvestDirtyFileHashes();
+
 	TArray<FUpdateCacheTransaction> PendingTransactions;
 	DiffDirtyFiles(DirtyFiles, PendingTransactions);
 	DirtyFiles.Empty();
-	return PendingTransactions;	
+	return PendingTransactions;
 }
 
 TArray<FUpdateCacheTransaction> FFileCache::FilterOutstandingChanges(const TFunctionRef<bool(const FUpdateCacheTransaction&, const FDateTime&)>& InPredicate)
 {
-	// We don't diff things that have only just changed, to ensure that add/delete pairs correctly get picked up as renames
-	FDateTime Threshold = FDateTime::UtcNow() - FTimespan(0,0,0,0,500);
+	HarvestDirtyFileHashes();
 
 	TArray<FUpdateCacheTransaction> AllTransactions;
-	DiffDirtyFiles(DirtyFiles, AllTransactions, nullptr, Threshold);
+	DiffDirtyFiles(DirtyFiles, AllTransactions, nullptr);
 
+	// Filter the transactions based on the predicate
 	TArray<FUpdateCacheTransaction> FilteredTransactions;
 	for (auto& Transaction : AllTransactions)
 	{
-		auto TimeOfChange = DirtyFiles[Transaction.Filename];
-		if (InPredicate(Transaction, TimeOfChange))
+		FFileData FileData = DirtyFiles.FindRef(Transaction.Filename);
+
+		// Timestamp is the time the file was dirtied, not necessarily the timestamp of the file
+		if (InPredicate(Transaction, FileData.Timestamp))
 		{
 			DirtyFiles.Remove(Transaction.Filename);
 			FilteredTransactions.Add(MoveTemp(Transaction));
@@ -724,7 +739,7 @@ void FFileCache::IgnoreNewFile(const FString& Filename)
 	{
 		DirtyFiles.Remove(TransactionPath.GetValue());
 		
-		const FFileData FileData(IFileManager::Get().GetTimeStamp(*Filename), ReadFileMD5(Filename));
+		const FFileData FileData(IFileManager::Get().GetTimeStamp(*Filename), FMD5Hash::HashFile(*Filename));
 		CompleteTransaction(FUpdateCacheTransaction(MoveTemp(TransactionPath.GetValue()), EFileAction::Added, FileData));	
 	}
 }
@@ -736,7 +751,7 @@ void FFileCache::IgnoreFileModification(const FString& Filename)
 	{
 		DirtyFiles.Remove(TransactionPath.GetValue());
 		
-		const FFileData FileData(IFileManager::Get().GetTimeStamp(*Filename), ReadFileMD5(Filename));
+		const FFileData FileData(IFileManager::Get().GetTimeStamp(*Filename), FMD5Hash::HashFile(*Filename));
 		CompleteTransaction(FUpdateCacheTransaction(MoveTemp(TransactionPath.GetValue()), EFileAction::Modified, FileData));	
 	}
 }
@@ -751,7 +766,7 @@ void FFileCache::IgnoreMovedFile(const FString& SrcFilename, const FString& DstF
 		DirtyFiles.Remove(SrcTransactionPath.GetValue());
 		DirtyFiles.Remove(DstTransactionPath.GetValue());
 
-		const FFileData FileData(IFileManager::Get().GetTimeStamp(*DstFilename), ReadFileMD5(DstFilename));
+		const FFileData FileData(IFileManager::Get().GetTimeStamp(*DstFilename), FMD5Hash::HashFile(*DstFilename));
 		CompleteTransaction(FUpdateCacheTransaction(MoveTemp(SrcTransactionPath.GetValue()), MoveTemp(DstTransactionPath.GetValue()), FileData));	
 	}
 }
@@ -833,7 +848,7 @@ void FFileCache::Tick()
 		{
 			ReadStateFromAsyncReader();
 
-			if (Config.bDetectMoves)
+			if (Config.bRequireFileHashes)
 			{
 				auto FilesThatNeedHashing = DirectoryReader->GetFilesThatNeedHashing();
 				if (FilesThatNeedHashing.Num() > 0)
@@ -903,20 +918,26 @@ void FFileCache::ReadStateFromAsyncReader()
 		CachedDirectoryState = MoveTemp(CachedState.GetValue());
 	}
 
-	// Build up a list of dirty files
-	TMap<FImmutableString, FDateTime> LocalDirtyFiles;
-	auto Now = FDateTime::UtcNow();
-
+	const FDateTime Now = FDateTime::UtcNow();
 	// We already have cached data so we need to compare it with the harvested data
 	// to detect additions, modifications, and removals
 	for (const auto& FilenameAndData : LiveState->Files)
 	{
 		const FString& Filename = FilenameAndData.Key.Get();
 
-		const auto* CachedData = CachedDirectoryState.Files.Find(Filename);
-		if (!CachedData || CachedData->Timestamp != FilenameAndData.Value.Timestamp)
+		// If the file we've discovered was not applicable to the old cache, we can't report a change for it as we don't know if it's new or not, just add it straight to the cache.
+		if (!CachedDirectoryState.Rules.IsFileApplicable(*Filename))
 		{
-			LocalDirtyFiles.Add(Filename, Now);
+			CachedDirectoryState.Files.Add(FilenameAndData.Key, FilenameAndData.Value);
+			bSavedCacheDirty = true;
+		}
+		else
+		{
+			const auto* CachedData = CachedDirectoryState.Files.Find(Filename);
+			if (!CachedData || CachedData->Timestamp != FilenameAndData.Value.Timestamp)
+			{
+				DirtyFiles.Add(FilenameAndData.Key, FFileData(Now, FMD5Hash()));
+			}
 		}
 	}
 
@@ -924,85 +945,75 @@ void FFileCache::ReadStateFromAsyncReader()
 	for (auto It = CachedDirectoryState.Files.CreateIterator(); It; ++It)
 	{
 		const FImmutableString& Filename = It.Key();
-		if (!LiveState->Files.Contains(Filename))
+		if (LiveState->Rules.IsFileApplicable(*Filename.Get()) && !LiveState->Files.Contains(Filename))
 		{
-			LocalDirtyFiles.Add(Filename, Now);
+			DirtyFiles.Add(Filename, FFileData(Now, FMD5Hash()));
 		}
 	}
 
-	if (LocalDirtyFiles.Num() != 0)
-	{
-		TArray<FUpdateCacheTransaction> TempTransactions;
-		DiffDirtyFiles(LocalDirtyFiles, TempTransactions, &LiveState.GetValue());
-
-		// Any file changes that we've detected that are not relevant to the previously cached state, we just add to the cache immediately (without telling the user)
-		// This is because we have no idea whether they are new changes or not, because they were not previously cached.
-		for (auto& Transaction : TempTransactions)
-		{
-			// Generate paths that are relative to the watch directory to pass to the match rules
-			const TCHAR* TransactionFilename = *Transaction.Filename.Get();
-			const TCHAR* MovedFromFilename = *Transaction.MovedFromFilename.Get();
-			if (Config.PathType == EPathType::Absolute)
-			{
-				if (Transaction.Filename.Get().Len() > Config.Directory.Len())
-				{
-					TransactionFilename += Config.Directory.Len();
-				}
+	RescanForDirtyFileHashes();
 
-				if (Transaction.MovedFromFilename.Get().Len() > Config.Directory.Len())
-				{
-					MovedFromFilename += Config.Directory.Len();
-				}
-			}
-
-			const bool bAppliesToOldCache = CachedDirectoryState.Rules.IsFileApplicable(TransactionFilename) &&
-				(Transaction.Action != EFileAction::Moved || CachedDirectoryState.Rules.IsFileApplicable(MovedFromFilename));
-
-			const bool bAppliesToNewCache = LiveState->Rules.IsFileApplicable(TransactionFilename) &&
-				(Transaction.Action != EFileAction::Moved || LiveState->Rules.IsFileApplicable(MovedFromFilename));
+	// Update the applicable extensions now that we've updated the cache
+	CachedDirectoryState.Rules = LiveState->Rules;
+}
 
-			if (bAppliesToOldCache && bAppliesToNewCache)
-			{
-				// Just throw away the transaction
-				continue;
-			}
-			else
-			{
-				// Remove it from our dirty files
-				LocalDirtyFiles.Remove(Transaction.Filename);
-				if (Transaction.Action == EFileAction::Moved)
-				{
-					LocalDirtyFiles.Remove(Transaction.MovedFromFilename);
-				}
+void FFileCache::HarvestDirtyFileHashes()
+{
+	if (!DirtyFileHasher.IsValid())
+	{
+		return;
+	}
+	else for (FFilenameAndHash& Data : DirtyFileHasher->GetCompletedData())
+	{
+		FImmutableString CachePath = (Config.PathType == EPathType::Relative) ? *Data.AbsoluteFilename + Config.Directory.Len() : *Data.AbsoluteFilename;
 
-				// Update the cache immediately and remove it from the pending transactions
-				CompleteTransaction(MoveTemp(Transaction));
-			}
+		if (auto* FileData = DirtyFiles.Find(CachePath))
+		{
+			FileData->FileHash = Data.FileHash;
 		}
 	}
 
-	// Now the only files left in LocalDirtyFiles are things that definitely apply to the current cache, and used to apply to the old cache
-	for (auto& Pair : LocalDirtyFiles)
+	if (DirtyFileHasher->IsComplete())
 	{
-		if (!DirtyFiles.Contains(Pair.Key))
+		DirtyFileHasher = nullptr;
+	}
+}
+
+void FFileCache::RescanForDirtyFileHashes()
+{
+	TArray<FFilenameAndHash> FilesThatNeedHashing;
+
+	for (const auto& Pair : DirtyFiles)
+	{
+		if (!Pair.Value.FileHash.IsValid())
 		{
-			DirtyFiles.Add(MoveTemp(Pair.Key), Pair.Value);
+			FilesThatNeedHashing.Emplace(GetAbsolutePath(Pair.Key.Get()));
 		}
 	}
 
-	// Update the applicable extensions now that we've updated the cache
-	CachedDirectoryState.Rules = LiveState->Rules;
+	if (FilesThatNeedHashing.Num() > 0)
+	{
+		// Re-create the dirty file hasher with the new data that needs hashing. The old task will clean itself up if it already exists.
+		DirtyFileHasher = MakeShareable(new FAsyncFileHasher(MoveTemp(FilesThatNeedHashing)));
+		AsyncTaskThread.AddTask(DirtyFileHasher);
+	}
 }
 
 void FFileCache::OnDirectoryChanged(const TArray<FFileChangeData>& FileChanges)
 {
-	FDateTime Now = FDateTime::UtcNow();
+	// Harvest any completed data from the file hasher before we discard it
+	HarvestDirtyFileHashes();
+
+	const FDateTime Now = FDateTime::UtcNow();
 	for (const auto& ThisEntry : FileChanges)
 	{
 		auto TransactionPath = GetTransactionPath(ThisEntry.Filename);
 		if (TransactionPath.IsSet())
 		{
-			DirtyFiles.Add(MoveTemp(TransactionPath.GetValue()), Now);
+			// Add the file that changed to the dirty files map, potentially invalidating the MD5 hash (we'll need to calculate it again)
+			DirtyFiles.Add(MoveTemp(TransactionPath.GetValue()), FFileData(Now, FMD5Hash()));
 		}
 	}
+
+	RescanForDirtyFileHashes();
 }
\ No newline at end of file
diff --git a/Engine/Source/Editor/UnrealEd/Private/AutoReimport/FileCache.h b/Engine/Source/Editor/UnrealEd/Private/AutoReimport/FileCache.h
--- a/Engine/Source/Editor/UnrealEd/Private/AutoReimport/FileCache.h
+++ b/Engine/Source/Editor/UnrealEd/Private/AutoReimport/FileCache.h
@@ -190,12 +190,12 @@ struct FAsyncFileHasher : public IAsyncFileCacheTask
 	/** Return any completed filenames and their corresponding hashes */
 	TArray<FFilenameAndHash> GetCompletedData();
 
-	/** Returns true when this directory reader has finished scanning the directory */
+	/** Returns true when this task has finished hashing all its files */
 	virtual bool IsComplete() const override;
 
 protected:
 
-	/** Tick this reader (discover new directories / files). Returns progress state. */
+	/** Tick this reader (hashes as many files as possible in the time allowed). Returns progress state. */
 	virtual EProgressResult Tick(const FTimeLimit& Limit) override;
 
 	/** The array of data that we will process */
@@ -298,9 +298,22 @@ private:
 /** Configuration structure required to construct a FFileCache */
 struct FFileCacheConfig
 {
+	/** Enum that specifies what changes are required for a change to be reported. When combined, any valid change is reported. */
+	enum EChangeDetection
+	{
+		/** Report modifications when the timestamp of a file changes */
+		Timestamp,
+		/** Report modifications when the contents of a file changes */
+		FileHash,
+	};
+
 	FFileCacheConfig(FString InDirectory, FString InCacheFile)
-		: Directory(InDirectory), CacheFile(InCacheFile), PathType(EPathType::Relative), bDetectChangesSinceLastRun(false), bDetectMoves(true)
-	{}
+		: Directory(InDirectory), CacheFile(InCacheFile), PathType(EPathType::Relative), bDetectChangesSinceLastRun(false)
+		, ChangeDetectionBits(false, 2)
+	{
+		DetectMoves(true);
+		ChangeDetectionBits[EChangeDetection::Timestamp] = true;
+	}
 
 	/** String specifying the directory on disk that the cache should reflect */
 	FString Directory;
@@ -317,9 +330,46 @@ struct FFileCacheConfig
 	/** When true, changes to the directory since the cache shutdown will be detected and reported. When false, said changes will silently be applied to the serialized cache. */
 	bool bDetectChangesSinceLastRun;
 
-	/** True to detect moves and renames (based on file hash), false otherwise. When true, an additional thread will be launched on startup to harvest MD5 hashes.
-	 *  Changes that occur *before* hashes have been gatherd will not be detected as moves/renames. */
+	/** Set up this cache to detect moves */
+	FFileCacheConfig& DetectMoves(bool bInDetectMoves)
+	{
+		bDetectMoves = bInDetectMoves;
+		if (bDetectMoves)
+		{
+			bRequireFileHashes = true;
+		}
+		return *this;
+	}
+	
+	/** Set up this cache to generate MD5 hashes for its constituent files */
+	FFileCacheConfig& RequireFileHashes(bool bInRequireFileHashes)
+	{
+		if (ensureMsg(bInRequireFileHashes || !bDetectMoves, TEXT("Unable to disable file hashing when move detection is enabled")))
+		{
+			bRequireFileHashes = bInRequireFileHashes;
+		}
+		return *this;
+	}
+
+	/** Instruct the cache to report the specified changes to files */
+	FFileCacheConfig& DetectChangesFor(EChangeDetection ChangeType, bool Value)
+	{
+		ChangeDetectionBits[ChangeType] = Value;
+		return *this;
+	}
+
+private:
+
+	/** True to detect moves and renames (based on file hash), false otherwise. Implies bRequireFileHashes. */
 	bool bDetectMoves;
+
+	/** When true, the cache will also calculate MD5 hashes for files. When true, an additional thread will be launched on startup to harvest unknown MD5 hashes for the directory. */
+	bool bRequireFileHashes;
+
+	/** Bitfied specifying how we will be detecting changes. See EChangeDetection. */
+	TBitArray<> ChangeDetectionBits;
+
+	friend class FFileCache;
 };
 
 /**
@@ -379,7 +429,8 @@ public:
 	/** Get the number of pending changes to the cache. */
 	int32 GetNumOutstandingChanges() const { return DirtyFiles.Num(); }
 
-	/** Get pending changes to the cache. Transactions must be returned to CompleteTransaction to update the cache. */
+	/** Get pending changes to the cache. Transactions must be returned to CompleteTransaction to update the cache.
+	 *  Filter predicate recieves a transaction and the time the change was reported. */
 	TArray<FUpdateCacheTransaction> FilterOutstandingChanges(const TFunctionRef<bool(const FUpdateCacheTransaction&, const FDateTime&)>& InPredicate);
 	TArray<FUpdateCacheTransaction> GetOutstandingChanges();
 
@@ -392,7 +443,10 @@ private:
 	 *	Optionally takes a directory state from which we can retrieve current file system state, without having to ask the FS directly.
 	 *  Optionally exclude files that have changed since the specified threshold, to ensure that related events get grouped together correctly.
 	 */
-	void DiffDirtyFiles(TMap<FImmutableString, FDateTime>& InDirtyFiles, TArray<FUpdateCacheTransaction>& OutTransactions, const FDirectoryState* InFileSystemState = nullptr, const FDateTime& ThreasholdTime = FDateTime::UtcNow()) const;
+	void DiffDirtyFiles(TMap<FImmutableString, FFileData>& InDirtyFiles, TArray<FUpdateCacheTransaction>& OutTransactions, const FDirectoryState* InFileSystemState = nullptr) const;
+
+	/** Detect a rename for the specified file */
+	void DetectRename();
 
 	/** Get the absolute path from a transaction filename */
 	FString GetAbsolutePath(const FString& InTransactionPath) const;
@@ -409,17 +463,24 @@ private:
 	/** Called when the initial async reader has finished harvesting file system timestamps */
 	void ReadStateFromAsyncReader();
 
+	/** Called to harvest any file hashes that have been generated by the DirtyFileHasher thread task */
+	void HarvestDirtyFileHashes();
+	void RescanForDirtyFileHashes();
+
 private:
 
 	/** Configuration settings applied on construction */
 	FFileCacheConfig Config;
 
-	/** 'Asynchronous' directory reader responsible for gathering all file/timestamp information recursively from our cache directory */
+	/** Asynchronous directory reader responsible for gathering all file/timestamp information recursively from our cache directory */
 	TSharedPtr<FAsyncDirectoryReader, ESPMode::ThreadSafe> DirectoryReader;
+	/** Asynchronous task used to harvest the MD5 hashes of a set of filenames */
 	TSharedPtr<FAsyncFileHasher, ESPMode::ThreadSafe> AsyncFileHasher;
+	/** Asynchronous task used to harvest the MD5 hashes of a set of recently changed filenames */
+	TSharedPtr<FAsyncFileHasher, ESPMode::ThreadSafe> DirtyFileHasher;
 
-	/** A map of dirty file times that we will use to report changes to the user */
-	TMap<FImmutableString, FDateTime> DirtyFiles;
+	/** A map of dirty files that we will use to report changes to the user. Timestamp value of FFileData here pertains to the *time of the change* not the timestamp of the file. */
+	TMap<FImmutableString, FFileData> DirtyFiles;
 
 	/** Our in-memory view of the cached directory state. */
 	FDirectoryState CachedDirectoryState;
diff --git a/Engine/Source/Editor/UnrealEd/Private/AutoReimport/ReimportFeedbackContext.cpp b/Engine/Source/Editor/UnrealEd/Private/AutoReimport/ReimportFeedbackContext.cpp
--- a/Engine/Source/Editor/UnrealEd/Private/AutoReimport/ReimportFeedbackContext.cpp
+++ b/Engine/Source/Editor/UnrealEd/Private/AutoReimport/ReimportFeedbackContext.cpp
@@ -45,8 +45,6 @@ class SWidgetStack : public SVerticalBox
 	{
 		const float Lerp = SizeCurve.GetLerp();
 		FVector2D DesiredSize = ComputeTotalSize() * Lerp + StartSizeOffset * (1.f-Lerp);
-		DesiredSize.X = FMath::Min(500.f, DesiredSize.X);
-
 		return DesiredSize;
 	}
 
@@ -226,8 +224,7 @@ class SReimportFeedback : public SCompoundWidget
 public:
 	SLATE_BEGIN_ARGS(SReimportFeedback) : _ExpireDuration(3.f) {}
 
-		SLATE_ATTRIBUTE(FText, MainText)
-
+		SLATE_ARGUMENT(TWeakPtr<FReimportFeedbackContext>, FeedbackContext)
 		SLATE_ARGUMENT(float, ExpireDuration)
 		SLATE_EVENT(FSimpleDelegate, OnExpired)
 
@@ -236,13 +233,21 @@ public:
 
 	SLATE_END_ARGS()
 
+
+	virtual FVector2D ComputeDesiredSize(float LayoutScale) const override
+	{
+		auto Size = SCompoundWidget::ComputeDesiredSize(LayoutScale);
+		// The width is determined by the top row, plus some padding
+		Size.X = TopRow->GetDesiredSize().X + 2*20;
+		return Size;
+	}
 	/** Construct this widget */
 	void Construct(const FArguments& InArgs)
 	{
 		ExpireDuration = InArgs._ExpireDuration;
 		OnExpired = InArgs._OnExpired;
+		FeedbackContext = InArgs._FeedbackContext;
 
-		MainText = InArgs._MainText;
 		bPaused = false;
 		bExpired = false;
 
@@ -260,22 +265,22 @@ public:
 				SNew(SVerticalBox)
 
 				+ SVerticalBox::Slot()
-				.Padding(FMargin(0))
 				.AutoHeight()
 				[
-					SNew(SHorizontalBox)
+					SAssignNew(TopRow, SHorizontalBox)
 
 					+ SHorizontalBox::Slot()
+					.VAlign(VAlign_Center)
 					[
 						SNew(STextBlock)
-						.Text(this, &SReimportFeedback::GetMainText)
+						.Text(LOCTEXT("ProcessingChanges", "Processing source file changes..."))
 						.Font(FCoreStyle::Get().GetFontStyle(TEXT("NotificationList.FontLight")))
 					]
 
 					+ SHorizontalBox::Slot()
 					.AutoWidth()
 					.VAlign(VAlign_Center)
-					.Padding(FMargin(0,0,4,0))
+					.Padding(FMargin(4,0,4,0))
 					[
 						SAssignNew(PauseButton, SButton)
 						.ButtonStyle(FEditorStyle::Get(), "HoverHintOnly")
@@ -304,6 +309,21 @@ public:
 					]
 				]
 
+				+ SVerticalBox::Slot()
+				.Padding(FMargin(0, 1))
+				.AutoHeight()
+				[
+					SNew(SBox)
+					.HeightOverride(2)
+					[
+						SAssignNew(ProgressBar, SProgressBar)
+						.BorderPadding(FVector2D::ZeroVector)
+						.Percent( this, &SReimportFeedback::GetProgressFraction )
+						.BackgroundImage( FEditorStyle::GetBrush("ProgressBar.ThinBackground") )
+						.FillImage( FEditorStyle::GetBrush("ProgressBar.ThinFill") )
+					]
+				]
+
 				+ SVerticalBox::Slot()
 				.Padding(FMargin(0, 5, 0, 0))
 				.AutoHeight()
@@ -335,14 +355,13 @@ public:
 	/** Disable input to this widget's dynamic content (except the message log hyperlink) */
 	void Disable()
 	{
-		// Get the final text to display
-		CachedMainText = MainText.Get(FText());
 
 		ExpireTimeout = FTimeLimit(ExpireDuration);
 
 		WidgetStack->SetVisibility(EVisibility::HitTestInvisible);
 		PauseButton->SetVisibility(EVisibility::Collapsed);
 		AbortButton->SetVisibility(EVisibility::Collapsed);
+		ProgressBar->SetVisibility(EVisibility::Collapsed);
 	}
 
 	/** Enable, if previously disabled */
@@ -354,10 +373,21 @@ public:
 		WidgetStack->SetVisibility(EVisibility::Visible);
 		PauseButton->SetVisibility(EVisibility::Visible);
 		AbortButton->SetVisibility(EVisibility::Visible);
+		ProgressBar->SetVisibility(EVisibility::Visible);
 	}
 
 private:
 	
+	TOptional<float> GetProgressFraction() const
+	{
+		auto PinnedContext = FeedbackContext.Pin();
+		if (PinnedContext.IsValid() && PinnedContext->GetScopeStack().Num() >= 0 )
+		{
+			return PinnedContext->GetScopeStack().GetProgressFraction(0);
+		}
+		return 1.f;
+	}
+
 	virtual void Tick(const FGeometry& AllottedGeometry, const double InCurrentTime, const float InDeltaTime) override
 	{
 		if (bExpired)
@@ -372,16 +402,6 @@ private:
 				bExpired = true;
 			}
 		}
-		else
-		{
-			CachedMainText = MainText.Get(FText());
-		}
-	}
-
-	/** Get the main text of this widget */
-	FText GetMainText() const
-	{
-		return CachedMainText;
 	}
 
 	/** Get the play/pause image */
@@ -424,29 +444,29 @@ private:
 	/** Event that is called when this widget has been inactive and open for too long, and will fade out */
 	FSimpleDelegate OnExpired;
 
-	/** Cached main text for the notification */
-	FText CachedMainText;
-	TAttribute<FText> MainText;
-
 	/** Whether we are paused and/or expired */
 	bool bPaused, bExpired;
 
 	/** The widget stack, displaying contextural information about the current state of the process */
 	TSharedPtr<SWidgetStack> WidgetStack;
-	TSharedPtr<SButton> PauseButton, AbortButton;
+	TSharedPtr<SWidget> PauseButton, AbortButton, ProgressBar, TopRow;
+	TWeakPtr<FReimportFeedbackContext> FeedbackContext;
 };
 
-FReimportFeedbackContext::FReimportFeedbackContext(const TAttribute<FText>& InMainText, const FSimpleDelegate& InOnPauseClicked, const FSimpleDelegate& InOnAbortClicked)
+FReimportFeedbackContext::FReimportFeedbackContext(const FSimpleDelegate& InOnPauseClicked, const FSimpleDelegate& InOnAbortClicked)
 	: bSuppressSlowTaskMessages(false)
-	, MainTextAttribute(InMainText)
 	, OnPauseClickedEvent(InOnPauseClicked)
 	, OnAbortClickedEvent(InOnAbortClicked)
 	, MessageLog("AssetReimport")
 {
 }
 
-void FReimportFeedbackContext::Show()
+void FReimportFeedbackContext::Show(int32 TotalWork)
 {
+	// Important - we first destroy the old main task, then create a new one to ensure that they are (removed from/added to) the scope stack in the correct order
+	MainTask.Reset();
+	MainTask.Reset(new FScopedSlowTask(TotalWork, FText(), true, *this));
+
 	if (NotificationContent.IsValid())
 	{
 		NotificationContent->Enable();
@@ -454,8 +474,8 @@ void FReimportFeedbackContext::Show()
 	else
 	{
 		NotificationContent = SNew(SReimportFeedback)
+			.FeedbackContext(AsShared())
 			.OnExpired(this, &FReimportFeedbackContext::OnNotificationExpired)
-			.MainText(MainTextAttribute)
 			.OnPauseClicked(OnPauseClickedEvent)
 			.OnAbortClicked(OnAbortClickedEvent);
 
@@ -470,6 +490,8 @@ void FReimportFeedbackContext::Show()
 
 void FReimportFeedbackContext::Hide()
 {
+	MainTask.Reset();
+
 	if (Notification.IsValid())
 	{
 		NotificationContent->Disable();
@@ -514,8 +536,29 @@ void FReimportFeedbackContext::StartSlowTask(const FText& Task, bool bShowCancel
 
 	if (NotificationContent.IsValid() && !bSuppressSlowTaskMessages && !Task.IsEmpty())
 	{
-		NotificationContent->Add(SNew(STextBlock).Text(Task));
+		if (SlowTaskText.IsValid())
+		{
+			SlowTaskText->SetText(FText::Format(LOCTEXT("SlowTaskPattern_Default", "{0} (0%)"), Task));
+		}
+		else
+		{
+			NotificationContent->Add(SAssignNew(SlowTaskText, STextBlock).Text(Task));
+		}
+	}
+}
+
+void FReimportFeedbackContext::ProgressReported(const float TotalProgressInterp, FText DisplayMessage)
+{
+	if (SlowTaskText.IsValid())
+	{
+		SlowTaskText->SetText(FText::Format(LOCTEXT("SlowTaskPattern", "{0} ({1}%)"), DisplayMessage, FText::AsNumber(int(TotalProgressInterp * 100))));
 	}
 }
 
+void FReimportFeedbackContext::FinalizeSlowTask()
+{
+	SlowTaskText->SetVisibility(EVisibility::Collapsed);
+	SlowTaskText = nullptr;
+}
+
 #undef LOCTEXT_NAMESPACE
\ No newline at end of file
diff --git a/Engine/Source/Editor/UnrealEd/Private/AutoReimport/ReimportFeedbackContext.h b/Engine/Source/Editor/UnrealEd/Private/AutoReimport/ReimportFeedbackContext.h
--- a/Engine/Source/Editor/UnrealEd/Private/AutoReimport/ReimportFeedbackContext.h
+++ b/Engine/Source/Editor/UnrealEd/Private/AutoReimport/ReimportFeedbackContext.h
@@ -13,13 +13,13 @@ class FReimportFeedbackContext : public FFeedbackContext, public INotificationWi
 {
 public:
 	/** Constructor */
-	FReimportFeedbackContext(const TAttribute<FText>& InMainText, const FSimpleDelegate& InOnPauseClicked, const FSimpleDelegate& InOnAbortClicked);
+	FReimportFeedbackContext(const FSimpleDelegate& InOnPauseClicked, const FSimpleDelegate& InOnAbortClicked);
 
 	/** Destroy this reimport context */
 	void Hide();
 
 	/** Revive the notification if it was destroyed */
-	void Show();
+	void Show(int32 TotalWork);
 
 	/** Get the message log that this context is using */
 	FMessageLog& GetMessageLog() { return MessageLog; }
@@ -34,22 +34,29 @@ public:
 	virtual void Serialize( const TCHAR* V, ELogVerbosity::Type Verbosity, const class FName& Category ) override {}
 	virtual void OnSetCompletionState(SNotificationItem::ECompletionState State) override {}
 	virtual TSharedRef<SWidget> AsWidget() override;
-	virtual void StartSlowTask( const FText& Task, bool bShowCancelButton=false) override;
+	virtual void StartSlowTask(const FText& Task, bool bShowCancelButton=false) override;
+	virtual void ProgressReported(const float TotalProgressInterp, FText DisplayMessage) override;
+	virtual void FinalizeSlowTask() override;
 
 	/** True if we will suppress slow task messages from being shown on the UI, false if otherwise */
 	bool bSuppressSlowTaskMessages;
 
+	/** Publicly accessible task that defines the entire import progress */
+	TUniquePtr<FScopedSlowTask> MainTask;
+
 private:
 
 	/** Called when our notification has expired */
 	void OnNotificationExpired();
 
-	TAttribute<FText> MainTextAttribute;
 	FSimpleDelegate OnPauseClickedEvent, OnAbortClickedEvent;
 
 	/** Message log for output of errors and log messages */
 	FMessageLog MessageLog;
 
+	/** Nested slow task text */
+	TSharedPtr<class STextBlock> SlowTaskText;
+
 	/** The notification that is shown when the context is active */
 	TSharedPtr<class SNotificationItem> Notification;
 
diff --git a/Engine/Source/Editor/UnrealEd/Private/Editor.cpp b/Engine/Source/Editor/UnrealEd/Private/Editor.cpp
--- a/Engine/Source/Editor/UnrealEd/Private/Editor.cpp
+++ b/Engine/Source/Editor/UnrealEd/Private/Editor.cpp
@@ -442,49 +442,6 @@ void FReimportManager::GetNewReimportPath(UObject* Obj, TArray<FString>& InOutFi
 	}
 }
 
-FString FReimportManager::SanitizeImportFilename(const FString& InPath, const UObject* Obj)
-{
-	const UPackage* Package = Obj->GetOutermost();
-	if (Package)
-	{
-		const bool		bIncludeDot = true;
-		const FString	PackagePath	= Package->GetPathName();
-		const FName		MountPoint	= FPackageName::GetPackageMountPoint(PackagePath);
-		const FString	PackageFilename = FPackageName::LongPackageNameToFilename(PackagePath, FPaths::GetExtension(InPath, bIncludeDot));
-		const FString	AbsolutePath = FPaths::ConvertRelativePathToFull(InPath);
-
-		if ((MountPoint == FName("Engine") && AbsolutePath.StartsWith(FPaths::ConvertRelativePathToFull(FPaths::EngineContentDir()))) ||
-			(MountPoint == FName("Game") &&	AbsolutePath.StartsWith(FPaths::ConvertRelativePathToFull(FPaths::GameDir()))))
-		{
-			FString RelativePath = InPath;
-			FPaths::MakePathRelativeTo(RelativePath, *PackageFilename);
-			return RelativePath;
-		}
-	}
-
-	return IFileManager::Get().ConvertToRelativePath(*InPath);
-}
-
-
-FString FReimportManager::ResolveImportFilename(const FString& InRelativePath, const UObject* Obj)
-{
-	FString RelativePath = InRelativePath;
-
-	const UPackage* Package = Obj->GetOutermost();
-	if (Package)
-	{
-		// Relative to the package filename?
-		const FString PathRelativeToPackage = FPaths::GetPath(FPackageName::LongPackageNameToFilename(Package->GetPathName())) / InRelativePath;
-		if (FPaths::FileExists(PathRelativeToPackage))
-		{
-			RelativePath = PathRelativeToPackage;
-		}
-	}
-
-	// Convert relative paths
-	return FPaths::ConvertRelativePathToFull(RelativePath);
-}
-
 FReimportManager::FReimportManager()
 {
 	// Create reimport handler for textures
diff --git a/Engine/Source/Editor/UnrealEd/Private/Factories/ApexDestructibleAssetImport.cpp b/Engine/Source/Editor/UnrealEd/Private/Factories/ApexDestructibleAssetImport.cpp
--- a/Engine/Source/Editor/UnrealEd/Private/Factories/ApexDestructibleAssetImport.cpp
+++ b/Engine/Source/Editor/UnrealEd/Private/Factories/ApexDestructibleAssetImport.cpp
@@ -1003,8 +1003,7 @@ UDestructibleMesh* ImportDestructibleMeshFromApexDestructibleAsset(UObject* InPa
 		// Store the current file path and timestamp for re-import purposes
 		// @todo AssetImportData make a data class for Apex destructible assets
 		DestructibleMesh->AssetImportData = NewObject<UAssetImportData>(DestructibleMesh);
-		DestructibleMesh->AssetImportData->SourceFilePath = FReimportManager::SanitizeImportFilename(UFactory::CurrentFilename, DestructibleMesh);
-		DestructibleMesh->AssetImportData->SourceFileTimestamp = IFileManager::Get().GetTimeStamp(*UFactory::CurrentFilename).ToString();
+		DestructibleMesh->AssetImportData->Update(UFactory::CurrentFilename);
 		DestructibleMesh->AssetImportData->bDirty = false;
 	}
 
diff --git a/Engine/Source/Editor/UnrealEd/Private/Factories/CSVImportFactory.cpp b/Engine/Source/Editor/UnrealEd/Private/Factories/CSVImportFactory.cpp
--- a/Engine/Source/Editor/UnrealEd/Private/Factories/CSVImportFactory.cpp
+++ b/Engine/Source/Editor/UnrealEd/Private/Factories/CSVImportFactory.cpp
@@ -155,7 +155,7 @@ UObject* UCSVImportFactory::FactoryCreateText( UClass* InClass, UObject* InParen
 			// Create/reset table
 			UDataTable* NewTable = NewObject<UDataTable>(InParent, InName, Flags);
 			NewTable->RowStruct = ImportRowStruct;
-			NewTable->ImportPath = FReimportManager::SanitizeImportFilename(CurrentFilename, NewTable);
+			NewTable->AssetImportData->Update(CurrentFilename);
 			// Go ahead and create table from string
 			Problems = DoImportDataTable(NewTable, String);
 
@@ -173,7 +173,7 @@ UObject* UCSVImportFactory::FactoryCreateText( UClass* InClass, UObject* InParen
 
 			// Create/reset table
 			UCurveTable* NewTable = NewObject<UCurveTable>(InParent, InName, Flags);
-			NewTable->ImportPath = FReimportManager::SanitizeImportFilename(CurrentFilename, NewTable);
+			NewTable->AssetImportData->Update(CurrentFilename);
 
 			// Go ahead and create table from string
 			Problems = DoImportCurveTable(NewTable, String, ImportCurveInterpMode);
@@ -192,7 +192,7 @@ UObject* UCSVImportFactory::FactoryCreateText( UClass* InClass, UObject* InParen
 			Problems = DoImportCurve(NewCurve, String);
 
 			UE_LOG(LogCSVImportFactory, Log, TEXT("Imported Curve '%s' - %d Problems"), *InName.ToString(), Problems.Num());
-			NewCurve->ImportPath = FReimportManager::SanitizeImportFilename(CurrentFilename, NewCurve);
+			NewCurve->AssetImportData->Update(CurrentFilename);
 			NewAsset = NewCurve;
 		}
 		
@@ -223,15 +223,15 @@ bool UCSVImportFactory::ReimportCSV( UObject* Obj )
 	bool bHandled = false;
 	if(UCurveBase* Curve = Cast<UCurveBase>(Obj))
 	{
-		bHandled = Reimport(Curve, FReimportManager::ResolveImportFilename(Curve->ImportPath, Curve));
+		bHandled = Reimport(Curve, Curve->AssetImportData->GetFirstFilename());
 	}
 	else if(UCurveTable* CurveTable = Cast<UCurveTable>(Obj))
 	{
-		bHandled = Reimport(CurveTable, FReimportManager::ResolveImportFilename(CurveTable->ImportPath, CurveTable));
+		bHandled = Reimport(CurveTable, CurveTable->AssetImportData->GetFirstFilename());
 	}
 	else if(UDataTable* DataTable = Cast<UDataTable>(Obj))
 	{
-		bHandled = Reimport(DataTable, FReimportManager::ResolveImportFilename(DataTable->ImportPath, DataTable));
+		bHandled = Reimport(DataTable, DataTable->AssetImportData->GetFirstFilename());
 	}
 	return bHandled;
 }
@@ -305,7 +305,7 @@ bool UReimportDataTableFactory::CanReimport( UObject* Obj, TArray<FString>& OutF
 	UDataTable* DataTable = Cast<UDataTable>(Obj);
 	if(DataTable)
 	{
-		OutFilenames.Add(FReimportManager::ResolveImportFilename(DataTable->ImportPath, DataTable));
+		DataTable->AssetImportData->ExtractFilenames(OutFilenames);
 		return true;
 	}
 	return false;
@@ -316,7 +316,7 @@ void UReimportDataTableFactory::SetReimportPaths( UObject* Obj, const TArray<FSt
 	UDataTable* DataTable = Cast<UDataTable>(Obj);
 	if(DataTable && ensure(NewReimportPaths.Num() == 1))
 	{
-		DataTable->ImportPath = FReimportManager::SanitizeImportFilename(NewReimportPaths[0], DataTable);
+		DataTable->AssetImportData->UpdateFilenameOnly(NewReimportPaths[0]);
 	}
 }
 
@@ -350,7 +350,7 @@ bool UReimportCurveTableFactory::CanReimport( UObject* Obj, TArray<FString>& Out
 	UCurveTable* CurveTable = Cast<UCurveTable>(Obj);
 	if(CurveTable)
 	{
-		OutFilenames.Add(FReimportManager::ResolveImportFilename(CurveTable->ImportPath, CurveTable));
+		CurveTable->AssetImportData->ExtractFilenames(OutFilenames);
 		return true;
 	}
 	return false;
@@ -361,7 +361,7 @@ void UReimportCurveTableFactory::SetReimportPaths( UObject* Obj, const TArray<FS
 	UCurveTable* CurveTable = Cast<UCurveTable>(Obj);
 	if(CurveTable && ensure(NewReimportPaths.Num() == 1))
 	{
-		CurveTable->ImportPath = FReimportManager::SanitizeImportFilename(NewReimportPaths[0], CurveTable);
+		CurveTable->AssetImportData->UpdateFilenameOnly(NewReimportPaths[0]);
 	}
 }
 
@@ -391,7 +391,7 @@ bool UReimportCurveFactory::CanReimport( UObject* Obj, TArray<FString>& OutFilen
 	UCurveBase* CurveBase = Cast<UCurveBase>(Obj);
 	if(CurveBase)
 	{
-		OutFilenames.Add(FReimportManager::ResolveImportFilename(CurveBase->ImportPath, CurveBase));
+		CurveBase->AssetImportData->ExtractFilenames(OutFilenames);
 		return true;
 	}
 	return false;
@@ -402,7 +402,7 @@ void UReimportCurveFactory::SetReimportPaths( UObject* Obj, const TArray<FString
 	UCurveBase* CurveBase = Cast<UCurveBase>(Obj);
 	if(CurveBase && ensure(NewReimportPaths.Num() == 1))
 	{
-		CurveBase->ImportPath = FReimportManager::SanitizeImportFilename(NewReimportPaths[0], CurveBase);
+		CurveBase->AssetImportData->UpdateFilenameOnly(NewReimportPaths[0]);
 	}
 }
 
diff --git a/Engine/Source/Editor/UnrealEd/Private/Factories/EditorFactories.cpp b/Engine/Source/Editor/UnrealEd/Private/Factories/EditorFactories.cpp
--- a/Engine/Source/Editor/UnrealEd/Private/Factories/EditorFactories.cpp
+++ b/Engine/Source/Editor/UnrealEd/Private/Factories/EditorFactories.cpp
@@ -1781,8 +1781,7 @@ UObject* USoundFactory::FactoryCreateBinary
 		}
 		
 		// Store the current file path and timestamp for re-import purposes
-		Sound->SourceFilePath = FReimportManager::SanitizeImportFilename( *CurrentFilename, Sound );
-		Sound->SourceFileTimestamp = IFileManager::Get().GetTimeStamp(*CurrentFilename).ToString();
+		Sound->AssetImportData->Update(CurrentFilename);
 
 		// Compressed data is now out of date.
 		Sound->InvalidateCompressedData();
@@ -1898,7 +1897,7 @@ bool UReimportSoundFactory::CanReimport( UObject* Obj, TArray<FString>& OutFilen
 	USoundWave* SoundWave = Cast<USoundWave>(Obj);
 	if(SoundWave && SoundWave->NumChannels < 3)
 	{
-		OutFilenames.Add(FReimportManager::ResolveImportFilename(SoundWave->SourceFilePath, SoundWave));
+		SoundWave->AssetImportData->ExtractFilenames(OutFilenames);
 		return true;
 	}
 	return false;
@@ -1909,7 +1908,7 @@ void UReimportSoundFactory::SetReimportPaths( UObject* Obj, const TArray<FString
 	USoundWave* SoundWave = Cast<USoundWave>(Obj);
 	if(SoundWave && ensure(NewReimportPaths.Num() == 1))
 	{
-		SoundWave->SourceFilePath = FReimportManager::ResolveImportFilename(NewReimportPaths[0], SoundWave);
+		SoundWave->AssetImportData->UpdateFilenameOnly(NewReimportPaths[0]);
 	}
 }
 
@@ -1923,7 +1922,7 @@ EReimportResult::Type UReimportSoundFactory::Reimport( UObject* Obj )
 
 	USoundWave* SoundWave = Cast<USoundWave>( Obj );
 
-	const FString Filename = FReimportManager::ResolveImportFilename(SoundWave->SourceFilePath, SoundWave);
+	const FString Filename = SoundWave->AssetImportData->GetFirstFilename();
 	const FString FileExtension = FPaths::GetExtension(Filename);
 	const bool bIsWav = ( FCString::Stricmp( *FileExtension, TEXT("WAV") ) == 0 );
 
@@ -2100,8 +2099,7 @@ UObject* USoundSurroundFactory::FactoryCreateBinary
 		}
 
 		// Store the current file path and timestamp for re-import purposes
-		Sound->SourceFilePath = FReimportManager::SanitizeImportFilename( *CurrentFilename, Sound );
-		Sound->SourceFileTimestamp = IFileManager::Get().GetTimeStamp(*CurrentFilename).ToString();
+		Sound->AssetImportData->Update(CurrentFilename);
 
 		// Compressed data is now out of date.
 		Sound->InvalidateCompressedData();
@@ -2229,33 +2227,24 @@ bool UReimportSoundSurroundFactory::CanReimport( UObject* Obj, TArray<FString>&
 	USoundWave* SoundWave = Cast<USoundWave>(Obj);
 	if(SoundWave && SoundWave->NumChannels > 2)
 	{
-		bool bGeneratedFilenames = false;
-
-		if (!SoundWave->SourceFilePath.IsEmpty())
+		FString SourceFilename = SoundWave->AssetImportData->GetFirstFilename();
+		if (!SourceFilename.IsEmpty() && FactoryCanImport(SourceFilename))
 		{
-			// Convert to a FString to check it has correct formatting
-			const FString SourceFilename = FReimportManager::ResolveImportFilename(SoundWave->SourceFilePath, SoundWave);
-			if (FactoryCanImport(SourceFilename))
-			{
-				// Get filename with speaker location removed
-				FString BaseFilename = FPaths::GetBaseFilename(SourceFilename).LeftChop(3);
-				FString FileExtension = FPaths::GetExtension(SourceFilename, true);
-				FString FilePath = FPaths::GetPath(SourceFilename);
+			// Get filename with speaker location removed
+			FString BaseFilename = FPaths::GetBaseFilename(SourceFilename).LeftChop(3);
+			FString FileExtension = FPaths::GetExtension(SourceFilename, true);
+			FString FilePath = FPaths::GetPath(SourceFilename);
 
-				// Add a filename for each speaker location we have Channel Size data for
-				for (int32 ChannelIndex = 0; ChannelIndex < SoundWave->ChannelSizes.Num(); ++ChannelIndex)
+			// Add a filename for each speaker location we have Channel Size data for
+			for (int32 ChannelIndex = 0; ChannelIndex < SoundWave->ChannelSizes.Num(); ++ChannelIndex)
+			{
+				if (SoundWave->ChannelSizes[ChannelIndex])
 				{
-					if (SoundWave->ChannelSizes[ChannelIndex])
-					{
-						OutFilenames.Add(FilePath + TEXT("//") + BaseFilename + SurroundSpeakerLocations[ChannelIndex] + FileExtension);
-					}
+					OutFilenames.Add(FilePath + TEXT("//") + BaseFilename + SurroundSpeakerLocations[ChannelIndex] + FileExtension);
 				}
-
-				bGeneratedFilenames = true;
 			}
 		}
-
-		if (!bGeneratedFilenames)
+		else
 		{
 			// We failed to generate possible filenames, fill the array with a blank string for each channel
 			for (int32 ChannelIndex = 0; ChannelIndex < SoundWave->NumChannels; ++ChannelIndex)
@@ -4442,8 +4431,7 @@ UObject* UTextureFactory::FactoryCreateBinary
 	Texture->MipGenSettings			= MipGenSettings;
 	Texture->bPreserveBorder		= bPreserveBorder;
 
-	Texture->SourceFilePath         = FReimportManager::SanitizeImportFilename(CurrentFilename, Texture);
-	Texture->SourceFileTimestamp	= IFileManager::Get().GetTimeStamp(*CurrentFilename).ToString();
+	Texture->AssetImportData->Update(CurrentFilename);
 
 	UTexture2D* Texture2D = Cast<UTexture2D>(Texture);
 
@@ -5512,7 +5500,7 @@ bool UReimportTextureFactory::CanReimport( UObject* Obj, TArray<FString>& OutFil
 	UTexture* pTex = Cast<UTexture>(Obj);
 	if( pTex && !pTex->IsA<UTextureRenderTarget>() )
 	{
-		OutFilenames.Add(FReimportManager::ResolveImportFilename(pTex->SourceFilePath, pTex));
+		pTex->AssetImportData->ExtractFilenames(OutFilenames);
 		return true;
 	}
 	return false;
@@ -5523,7 +5511,7 @@ void UReimportTextureFactory::SetReimportPaths( UObject* Obj, const TArray<FStri
 	UTexture* pTex = Cast<UTexture>(Obj);
 	if(pTex && ensure(NewReimportPaths.Num() == 1))
 	{
-		pTex->SourceFilePath = FReimportManager::SanitizeImportFilename(NewReimportPaths[0], Obj);
+		pTex->AssetImportData->UpdateFilenameOnly(NewReimportPaths[0]);
 	}
 }
 
@@ -5541,7 +5529,7 @@ EReimportResult::Type UReimportTextureFactory::Reimport( UObject* Obj )
 	
 	TGuardValue<UTexture*> OriginalTexGuardValue(pOriginalTex, pTex);
 
-	const FString ResolvedSourceFilePath = FReimportManager::ResolveImportFilename(pTex->SourceFilePath, pTex);
+	const FString ResolvedSourceFilePath = pTex->AssetImportData->GetFirstFilename();
 	if (!ResolvedSourceFilePath.Len())
 	{
 		// Since this is a new system most textures don't have paths, so logging has been commented out
@@ -5629,14 +5617,15 @@ bool UReimportFbxStaticMeshFactory::CanReimport( UObject* Obj, TArray<FString>&
 	{
 		if ( Mesh->AssetImportData )
 		{
-			if (FPaths::GetExtension(Mesh->AssetImportData->SourceFilePath).ToLower() == "srt")
+			const FString& Filename = Mesh->AssetImportData->GetFirstFilename();
+			if (FPaths::GetExtension(Filename) == "srt")
 			{
 				// SpeedTrees need to use their own importer
 				return false;
 			}
 			else
 			{
-				OutFilenames.Add(FReimportManager::ResolveImportFilename(Mesh->AssetImportData->SourceFilePath, Mesh));
+				OutFilenames.Add(Filename);
 			}
 		}
 		else
@@ -5655,7 +5644,7 @@ void UReimportFbxStaticMeshFactory::SetReimportPaths( UObject* Obj, const TArray
 	{
 		UFbxStaticMeshImportData* ImportData = UFbxStaticMeshImportData::GetImportDataForStaticMesh(Mesh, ImportUI->StaticMeshImportData);
 
-		ImportData->SourceFilePath = FReimportManager::SanitizeImportFilename(NewReimportPaths[0], Mesh);
+		ImportData->UpdateFilenameOnly(NewReimportPaths[0]);
 	}
 }
 
@@ -5699,7 +5688,7 @@ EReimportResult::Type UReimportFbxStaticMeshFactory::Reimport( UObject* Obj )
 
 	if( !bOperationCanceled && ensure(ImportData) )
 	{
-		const FString Filename = FReimportManager::ResolveImportFilename(ImportData->SourceFilePath, Mesh);
+		const FString& Filename = ImportData->GetFirstFilename();
 		const FString FileExtension = FPaths::GetExtension(Filename);
 		const bool bIsValidFile = FileExtension.Equals( TEXT("fbx"), ESearchCase::IgnoreCase ) || FileExtension.Equals( "obj",  ESearchCase::IgnoreCase );
 
@@ -5819,7 +5808,7 @@ bool UReimportFbxSkeletalMeshFactory::CanReimport( UObject* Obj, TArray<FString>
 	{
 		if ( SkeletalMesh->AssetImportData )
 		{
-			OutFilenames.Add(FReimportManager::ResolveImportFilename(SkeletalMesh->AssetImportData->SourceFilePath, SkeletalMesh));
+			SkeletalMesh->AssetImportData->ExtractFilenames(OutFilenames);
 		}
 		else
 		{
@@ -5836,8 +5825,7 @@ void UReimportFbxSkeletalMeshFactory::SetReimportPaths( UObject* Obj, const TArr
 	if(SkeletalMesh && ensure(NewReimportPaths.Num() == 1))
 	{
 		UFbxSkeletalMeshImportData* ImportData = UFbxSkeletalMeshImportData::GetImportDataForSkeletalMesh(SkeletalMesh, ImportUI->SkeletalMeshImportData);
-
-		ImportData->SourceFilePath = FReimportManager::SanitizeImportFilename(NewReimportPaths[0], SkeletalMesh);
+		ImportData->UpdateFilenameOnly(NewReimportPaths[0]);
 	}
 }
 
@@ -5902,7 +5890,7 @@ EReimportResult::Type UReimportFbxSkeletalMeshFactory::Reimport( UObject* Obj )
 
 	if( !bOperationCanceled && ensure(ImportData) )
 	{
-		const FString Filename = FReimportManager::ResolveImportFilename(ImportData->SourceFilePath, SkeletalMesh);
+		const FString Filename = ImportData->GetFirstFilename();
 		UE_LOG(LogEditorFactories, Log, TEXT("Performing atomic reimport of [%s]"), *Filename);
 
 		// Ensure that the file provided by the path exists
@@ -5999,7 +5987,7 @@ bool UReimportFbxAnimSequenceFactory::CanReimport( UObject* Obj, TArray<FString>
 	{
 		if ( AnimSequence->AssetImportData )
 		{
-			OutFilenames.Add(FReimportManager::ResolveImportFilename(AnimSequence->AssetImportData->SourceFilePath, AnimSequence));
+			AnimSequence->AssetImportData->ExtractFilenames(OutFilenames);
 		}
 		else
 		{
@@ -6017,7 +6005,7 @@ void UReimportFbxAnimSequenceFactory::SetReimportPaths( UObject* Obj, const TArr
 	{
 		UFbxAnimSequenceImportData* ImportData = UFbxAnimSequenceImportData::GetImportDataForAnimSequence(AnimSequence, ImportUI->AnimSequenceImportData);
 
-		ImportData->SourceFilePath = FReimportManager::SanitizeImportFilename(NewReimportPaths[0], AnimSequence);
+		ImportData->UpdateFilenameOnly(NewReimportPaths[0]);
 	}
 }
 
@@ -6036,7 +6024,7 @@ EReimportResult::Type UReimportFbxAnimSequenceFactory::Reimport( UObject* Obj )
 		return EReimportResult::Failed;
 	}
 
-	const FString Filename = FReimportManager::ResolveImportFilename(ImportData->SourceFilePath, AnimSequence);
+	const FString Filename = ImportData->GetFirstFilename();
 	const FString FileExtension = FPaths::GetExtension(Filename);
 	const bool bIsNotFBXFile = ( FileExtension.Len() > 0  && FCString::Stricmp( *FileExtension, TEXT("FBX") ) != 0 );
 
@@ -6085,8 +6073,7 @@ EReimportResult::Type UReimportFbxAnimSequenceFactory::Reimport( UObject* Obj )
 		UE_LOG(LogEditorFactories, Log, TEXT("-- imported successfully") );
 
 		// update the data in case the file source has changed
-		ImportData->SourceFilePath = FReimportManager::SanitizeImportFilename(UFactory::CurrentFilename, AnimSequence);
-		ImportData->SourceFileTimestamp = IFileManager::Get().GetTimeStamp(*UFactory::CurrentFilename).ToString();
+		ImportData->Update(UFactory::CurrentFilename);
 		ImportData->bDirty = false;
 
 		// Try to find the outer package so we can dirty it up
@@ -6787,7 +6774,7 @@ bool UReimportDestructibleMeshFactory::CanReimport( UObject* Obj, TArray<FString
 	{
 		if ( DestructibleMesh->AssetImportData )
 		{
-			OutFilenames.Add(FReimportManager::ResolveImportFilename(DestructibleMesh->AssetImportData->SourceFilePath, DestructibleMesh));
+			DestructibleMesh->AssetImportData->ExtractFilenames(OutFilenames);
 		}
 		else
 		{
@@ -6809,7 +6796,7 @@ void UReimportDestructibleMeshFactory::SetReimportPaths( UObject* Obj, const TAr
 			DestructibleMesh->AssetImportData = NewObject<UAssetImportData>(DestructibleMesh);
 		}
 
-		DestructibleMesh->AssetImportData->SourceFilePath = FReimportManager::SanitizeImportFilename(NewReimportPaths[0], DestructibleMesh);
+		DestructibleMesh->AssetImportData->UpdateFilenameOnly(NewReimportPaths[0]);
 	}
 }
 
@@ -6829,7 +6816,7 @@ EReimportResult::Type UReimportDestructibleMeshFactory::Reimport( UObject* Obj )
 		DestructibleMesh->AssetImportData = NewObject<UAssetImportData>(DestructibleMesh);
 	}
 
-	const FString Filename = FReimportManager::ResolveImportFilename(DestructibleMesh->AssetImportData->SourceFilePath, DestructibleMesh);
+	const FString Filename = DestructibleMesh->AssetImportData->GetFirstFilename();
 
 	// If there is no file path provided, can't reimport from source
 	if ( !Filename.Len() )
diff --git a/Engine/Source/Editor/UnrealEd/Private/Factories/VectorFieldFactory.cpp b/Engine/Source/Editor/UnrealEd/Private/Factories/VectorFieldFactory.cpp
--- a/Engine/Source/Editor/UnrealEd/Private/Factories/VectorFieldFactory.cpp
+++ b/Engine/Source/Editor/UnrealEd/Private/Factories/VectorFieldFactory.cpp
@@ -227,8 +227,7 @@ UObject* UVectorFieldStaticFactory::FactoryCreateBinary(
 				VectorField->SizeZ = FileContents.GridZ;
 				VectorField->Bounds = FileContents.Bounds;
 				
-				VectorField->SourceFilePath = FReimportManager::SanitizeImportFilename(CurrentFilename, VectorField);
-				VectorField->SourceFileTimestamp = IFileManager::Get().GetTimeStamp(*CurrentFilename).ToString();
+				VectorField->AssetImportData->Update(CurrentFilename);
 
 				// Convert vectors to 16-bit FP and store.
 				check( (FileContents.Values.Num() % 3) == 0 );
@@ -282,7 +281,7 @@ bool UReimportVectorFieldStaticFactory::CanReimport( UObject* Obj, TArray<FStrin
 	UVectorFieldStatic* VectorFieldStatic = Cast<UVectorFieldStatic>(Obj);
 	if(VectorFieldStatic)
 	{
-		OutFilenames.Add(FReimportManager::ResolveImportFilename(VectorFieldStatic->SourceFilePath, VectorFieldStatic));
+		VectorFieldStatic->AssetImportData->ExtractFilenames(OutFilenames);
 		return true;
 	}
 	return false;
@@ -293,7 +292,7 @@ void UReimportVectorFieldStaticFactory::SetReimportPaths( UObject* Obj, const TA
 	UVectorFieldStatic* VectorFieldStatic = Cast<UVectorFieldStatic>(Obj);
 	if(VectorFieldStatic && ensure(NewReimportPaths.Num() == 1))
 	{
-		VectorFieldStatic->SourceFilePath = FReimportManager::SanitizeImportFilename(NewReimportPaths[0], VectorFieldStatic);
+		VectorFieldStatic->AssetImportData->UpdateFilenameOnly(NewReimportPaths[0]);
 	}
 }
 
@@ -306,13 +305,13 @@ EReimportResult::Type UReimportVectorFieldStaticFactory::Reimport( UObject* Obj
 
 	UVectorFieldStatic* VectorFieldStatic = Cast<UVectorFieldStatic>(Obj);
 
-	if ( VectorFieldStatic->SourceFilePath.Len() <= 0 )
+	if ( VectorFieldStatic->AssetImportData->GetSourceFileData().Num() != 1 )
 	{
 		// No source art path. Can't reimport.
 		return EReimportResult::Failed;
 	}
 
-	const FString ReImportFilename = FReimportManager::ResolveImportFilename(VectorFieldStatic->SourceFilePath, VectorFieldStatic);
+	const FString& ReImportFilename = VectorFieldStatic->AssetImportData->GetSourceFileData()[0].RelativeFilename;
 
 	UE_LOG(LogVectorFieldFactory, Log, TEXT("Performing atomic reimport of [%s]"), *ReImportFilename);
 
diff --git a/Engine/Source/Editor/UnrealEd/Private/Fbx/FbxAnimSequenceImportData.cpp b/Engine/Source/Editor/UnrealEd/Private/Fbx/FbxAnimSequenceImportData.cpp
--- a/Engine/Source/Editor/UnrealEd/Private/Fbx/FbxAnimSequenceImportData.cpp
+++ b/Engine/Source/Editor/UnrealEd/Private/Fbx/FbxAnimSequenceImportData.cpp
@@ -21,8 +21,7 @@ UFbxAnimSequenceImportData* UFbxAnimSequenceImportData::GetImportDataForAnimSequ
 		// Try to preserve the source file path if possible
 		if ( AnimSequence->AssetImportData != NULL )
 		{
-			ImportData->SourceFilePath = AnimSequence->AssetImportData->SourceFilePath;
-			ImportData->SourceFileTimestamp = AnimSequence->AssetImportData->SourceFileTimestamp;
+			ImportData->CopyFrom(*AnimSequence->AssetImportData);
 		}
 
 		AnimSequence->AssetImportData = ImportData;
diff --git a/Engine/Source/Editor/UnrealEd/Private/Fbx/FbxFactory.cpp b/Engine/Source/Editor/UnrealEd/Private/Fbx/FbxFactory.cpp
--- a/Engine/Source/Editor/UnrealEd/Private/Fbx/FbxFactory.cpp
+++ b/Engine/Source/Editor/UnrealEd/Private/Fbx/FbxFactory.cpp
@@ -386,7 +386,7 @@ UObject* UFbxFactory::FactoryCreateBinary
 							{
 								FName OutputName = FbxImporter->MakeNameForMesh(Name.ToString(), SkelMeshNodeArray[0]);
 
-								USkeletalMesh* NewMesh = FbxImporter->ImportSkeletalMesh( InParent, SkelMeshNodeArray, OutputName, Flags, ImportUI->SkeletalMeshImportData, FPaths::GetBaseFilename(Filename) );
+								USkeletalMesh* NewMesh = FbxImporter->ImportSkeletalMesh( InParent, SkelMeshNodeArray, OutputName, Flags, ImportUI->SkeletalMeshImportData );
 								NewObject = NewMesh;
 
 								if ( NewMesh && ImportUI->bImportAnimations )
@@ -404,7 +404,7 @@ UObject* UFbxFactory::FactoryCreateBinary
 							{
 								USkeletalMesh* BaseSkeletalMesh = Cast<USkeletalMesh>(NewObject);
 								FName LODObjectName = NAME_None;
-								USkeletalMesh *LODObject = FbxImporter->ImportSkeletalMesh( GetTransientPackage(), SkelMeshNodeArray, LODObjectName, RF_NoFlags, ImportUI->SkeletalMeshImportData, FPaths::GetBaseFilename(Filename) );
+								USkeletalMesh *LODObject = FbxImporter->ImportSkeletalMesh( GetTransientPackage(), SkelMeshNodeArray, LODObjectName, RF_NoFlags, ImportUI->SkeletalMeshImportData );
 								bool bImportSucceeded = FbxImporter->ImportSkeletalMeshLOD(LODObject, BaseSkeletalMesh, LODIndex, false);
 
 								if (bImportSucceeded)
@@ -424,7 +424,7 @@ UObject* UFbxFactory::FactoryCreateBinary
 								uint32 bImportMaterials = ImportOptions->bImportMaterials;
 								ImportOptions->bImportMaterials = 0;
 
-								FbxImporter->ImportFbxMorphTarget(SkelMeshNodeArray, Cast<USkeletalMesh>(NewObject), InParent, Filename, LODIndex);
+								FbxImporter->ImportFbxMorphTarget(SkelMeshNodeArray, Cast<USkeletalMesh>(NewObject), InParent, LODIndex);
 							
 								ImportOptions->bImportMaterials = !!bImportMaterials;
 							}
diff --git a/Engine/Source/Editor/UnrealEd/Private/Fbx/FbxSkeletalMeshImport.cpp b/Engine/Source/Editor/UnrealEd/Private/Fbx/FbxSkeletalMeshImport.cpp
--- a/Engine/Source/Editor/UnrealEd/Private/Fbx/FbxSkeletalMeshImport.cpp
+++ b/Engine/Source/Editor/UnrealEd/Private/Fbx/FbxSkeletalMeshImport.cpp
@@ -1058,7 +1058,7 @@ bool UnFbx::FFbxImporter::ImportBone(TArray<FbxNode*>& NodeArray, FSkeletalMeshI
 	return true;
 }
 
-USkeletalMesh* UnFbx::FFbxImporter::ImportSkeletalMesh(UObject* InParent, TArray<FbxNode*>& NodeArray, const FName& Name, EObjectFlags Flags, UFbxSkeletalMeshImportData* TemplateImportData, FString Filename, TArray<FbxShape*> *FbxShapeArray, FSkeletalMeshImportData* OutData, bool bCreateRenderData )
+USkeletalMesh* UnFbx::FFbxImporter::ImportSkeletalMesh(UObject* InParent, TArray<FbxNode*>& NodeArray, const FName& Name, EObjectFlags Flags, UFbxSkeletalMeshImportData* TemplateImportData, TArray<FbxShape*> *FbxShapeArray, FSkeletalMeshImportData* OutData, bool bCreateRenderData )
 {
 	if (NodeArray.Num() == 0)
 	{
@@ -1320,8 +1320,7 @@ USkeletalMesh* UnFbx::FFbxImporter::ImportSkeletalMesh(UObject* InParent, TArray
 
 		// Store the current file path and timestamp for re-import purposes
 		UFbxSkeletalMeshImportData* ImportData = UFbxSkeletalMeshImportData::GetImportDataForSkeletalMesh(SkeletalMesh, TemplateImportData);
-		SkeletalMesh->AssetImportData->SourceFilePath = FReimportManager::SanitizeImportFilename(UFactory::CurrentFilename, SkeletalMesh);
-		SkeletalMesh->AssetImportData->SourceFileTimestamp = IFileManager::Get().GetTimeStamp(*UFactory::CurrentFilename).ToString();
+		SkeletalMesh->AssetImportData->Update(UFactory::CurrentFilename);
 		SkeletalMesh->AssetImportData->bDirty = false;
 
 		SkeletalMesh->CalculateInvRefMatrices();
@@ -1580,12 +1579,12 @@ USkeletalMesh* UnFbx::FFbxImporter::ReimportSkeletalMesh(USkeletalMesh* Mesh, UF
 
 			if (LODIndex == 0)
 			{
-				NewMesh = ImportSkeletalMesh( Mesh->GetOuter(), SkelMeshNodeArray, *Mesh->GetName(), RF_Public|RF_Standalone, TemplateImportData, FReimportManager::ResolveImportFilename(TemplateImportData->SourceFilePath, Mesh));
+				NewMesh = ImportSkeletalMesh( Mesh->GetOuter(), SkelMeshNodeArray, *Mesh->GetName(), RF_Public|RF_Standalone, TemplateImportData);
 			}
 			else if (NewMesh) // the base skeletal mesh is imported successfully
 			{
 				USkeletalMesh* BaseSkeletalMesh = Cast<USkeletalMesh>(NewMesh);
-				UObject *LODObject = ImportSkeletalMesh( GetTransientPackage(), SkelMeshNodeArray, NAME_None, RF_NoFlags, TemplateImportData, FReimportManager::ResolveImportFilename(TemplateImportData->SourceFilePath, BaseSkeletalMesh) );
+				UObject *LODObject = ImportSkeletalMesh( GetTransientPackage(), SkelMeshNodeArray, NAME_None, RF_NoFlags, TemplateImportData);
 				ImportSkeletalMeshLOD( Cast<USkeletalMesh>(LODObject), BaseSkeletalMesh, LODIndex, false);
 
 				// Set LOD Model's DisplayFactor
@@ -1601,7 +1600,7 @@ USkeletalMesh* UnFbx::FFbxImporter::ReimportSkeletalMesh(USkeletalMesh* Mesh, UF
 			if ( NewMesh)
 			{
 				// @fixme: @question : where do they import this morph? where to? What morph target sets?
-				ImportFbxMorphTarget(SkelMeshNodeArray, NewMesh, NewMesh->GetOutermost(), FReimportManager::ResolveImportFilename(TemplateImportData->SourceFilePath, NewMesh), LODIndex);
+				ImportFbxMorphTarget(SkelMeshNodeArray, NewMesh, NewMesh->GetOutermost(), LODIndex);
 			}
 		}
 	}
@@ -2676,7 +2675,7 @@ private:
 	bool bKeepOverlappingVertices;
 };
 
-void UnFbx::FFbxImporter::ImportMorphTargetsInternal( TArray<FbxNode*>& SkelMeshNodeArray, USkeletalMesh* BaseSkelMesh, UObject* InParent, const FString& InFilename, int32 LODIndex )
+void UnFbx::FFbxImporter::ImportMorphTargetsInternal( TArray<FbxNode*>& SkelMeshNodeArray, USkeletalMesh* BaseSkelMesh, UObject* InParent, int32 LODIndex )
 {
 	FbxString ShapeNodeName;
 	TMap<FString, TArray<FbxShape*>> ShapeNameToShapeArray;
@@ -2803,7 +2802,7 @@ void UnFbx::FFbxImporter::ImportMorphTargetsInternal( TArray<FbxNode*>& SkelMesh
 			// now we get a shape for whole mesh, import to unreal as a morph target
 			// @todo AssetImportData do we need import data for this temp mesh?
 			UFbxSkeletalMeshImportData* TmpMeshImportData = NULL;
-			USkeletalMesh* TmpSkeletalMesh = (USkeletalMesh*)ImportSkeletalMesh(GetTransientPackage(), SkelMeshNodeArray, NAME_None, (EObjectFlags)0, TmpMeshImportData, FPaths::GetBaseFilename(InFilename), &ShapeArray, &ImportData, false);
+			USkeletalMesh* TmpSkeletalMesh = (USkeletalMesh*)ImportSkeletalMesh(GetTransientPackage(), SkelMeshNodeArray, NAME_None, (EObjectFlags)0, TmpMeshImportData, &ShapeArray, &ImportData, false);
 			TempMeshes.Add(TmpSkeletalMesh);
 			MorphTargets.Add(Result);
 
@@ -2861,7 +2860,7 @@ void UnFbx::FFbxImporter::ImportMorphTargetsInternal( TArray<FbxNode*>& SkelMesh
 }	
 
 // Import Morph target
-void UnFbx::FFbxImporter::ImportFbxMorphTarget(TArray<FbxNode*> &SkelMeshNodeArray, USkeletalMesh* BaseSkelMesh, UObject* InParent, const FString& Filename, int32 LODIndex)
+void UnFbx::FFbxImporter::ImportFbxMorphTarget(TArray<FbxNode*> &SkelMeshNodeArray, USkeletalMesh* BaseSkelMesh, UObject* InParent, int32 LODIndex)
 {
 	bool bHasMorph = false;
 	int32 NodeIndex;
@@ -2881,7 +2880,7 @@ void UnFbx::FFbxImporter::ImportFbxMorphTarget(TArray<FbxNode*> &SkelMeshNodeArr
 	
 	if (bHasMorph)
 	{
-		ImportMorphTargetsInternal( SkelMeshNodeArray, BaseSkelMesh, InParent, Filename, LODIndex );
+		ImportMorphTargetsInternal( SkelMeshNodeArray, BaseSkelMesh, InParent, LODIndex );
 	}
 }
 
diff --git a/Engine/Source/Editor/UnrealEd/Private/Fbx/FbxSkeletalMeshImportData.cpp b/Engine/Source/Editor/UnrealEd/Private/Fbx/FbxSkeletalMeshImportData.cpp
--- a/Engine/Source/Editor/UnrealEd/Private/Fbx/FbxSkeletalMeshImportData.cpp
+++ b/Engine/Source/Editor/UnrealEd/Private/Fbx/FbxSkeletalMeshImportData.cpp
@@ -20,8 +20,7 @@ UFbxSkeletalMeshImportData* UFbxSkeletalMeshImportData::GetImportDataForSkeletal
 		// Try to preserve the source file path if possible
 		if ( SkeletalMesh->AssetImportData != NULL )
 		{
-			ImportData->SourceFilePath = SkeletalMesh->AssetImportData->SourceFilePath;
-			ImportData->SourceFileTimestamp = SkeletalMesh->AssetImportData->SourceFileTimestamp;
+			ImportData->CopyFrom(*SkeletalMesh->AssetImportData);
 		}
 
 		SkeletalMesh->AssetImportData = ImportData;
diff --git a/Engine/Source/Editor/UnrealEd/Private/Fbx/FbxStaticMeshImport.cpp b/Engine/Source/Editor/UnrealEd/Private/Fbx/FbxStaticMeshImport.cpp
--- a/Engine/Source/Editor/UnrealEd/Private/Fbx/FbxStaticMeshImport.cpp
+++ b/Engine/Source/Editor/UnrealEd/Private/Fbx/FbxStaticMeshImport.cpp
@@ -1241,8 +1241,7 @@ UStaticMesh* UnFbx::FFbxImporter::ImportStaticMeshAsSingle(UObject* InParent, TA
 		}
 
 		UFbxStaticMeshImportData* ImportData = UFbxStaticMeshImportData::GetImportDataForStaticMesh(StaticMesh, TemplateImportData);
-		ImportData->SourceFilePath = FReimportManager::SanitizeImportFilename(UFactory::CurrentFilename, StaticMesh);
-		ImportData->SourceFileTimestamp = IFileManager::Get().GetTimeStamp(*UFactory::CurrentFilename).ToString();
+		ImportData->Update(UFactory::CurrentFilename);
 		ImportData->bDirty = false;
 
 		// @todo This overrides restored values currently but we need to be able to import over the existing settings if the user chose to do so.
diff --git a/Engine/Source/Editor/UnrealEd/Private/Fbx/FbxStaticMeshImportData.cpp b/Engine/Source/Editor/UnrealEd/Private/Fbx/FbxStaticMeshImportData.cpp
--- a/Engine/Source/Editor/UnrealEd/Private/Fbx/FbxStaticMeshImportData.cpp
+++ b/Engine/Source/Editor/UnrealEd/Private/Fbx/FbxStaticMeshImportData.cpp
@@ -26,8 +26,7 @@ UFbxStaticMeshImportData* UFbxStaticMeshImportData::GetImportDataForStaticMesh(U
 		// Try to preserve the source file path if possible
 		if ( StaticMesh->AssetImportData != NULL )
 		{
-			ImportData->SourceFilePath = StaticMesh->AssetImportData->SourceFilePath;
-			ImportData->SourceFileTimestamp = StaticMesh->AssetImportData->SourceFileTimestamp;
+			ImportData->CopyFrom(*StaticMesh->AssetImportData);
 		}
 
 		StaticMesh->AssetImportData = ImportData;
diff --git a/Engine/Source/Editor/UnrealEd/Private/FbxImporter.h b/Engine/Source/Editor/UnrealEd/Private/FbxImporter.h
--- a/Engine/Source/Editor/UnrealEd/Private/FbxImporter.h
+++ b/Engine/Source/Editor/UnrealEd/Private/FbxImporter.h
@@ -383,14 +383,13 @@ public:
 	 * @param NodeArray	Fbx Nodes to import
 	 * @param Name	the Unreal Mesh name after import
 	 * @param Flags
-	 * @param Filename	Fbx file name
 	 * @param FbxShapeArray	Fbx Morph objects.
 	 * @param OutData - Optional import data to populate
 	 * @param bCreateRenderData - Whether or not skeletal mesh rendering data will be created.
 	 *
 	 * @return The USkeletalMesh object created
 	 */
-	USkeletalMesh* ImportSkeletalMesh(UObject* InParent, TArray<FbxNode*>& NodeArray, const FName& Name, EObjectFlags Flags, UFbxSkeletalMeshImportData* TemplateImportData, FString Filename, TArray<FbxShape*> *FbxShapeArray=NULL, FSkeletalMeshImportData* OutData=NULL, bool bCreateRenderData = true );
+	USkeletalMesh* ImportSkeletalMesh(UObject* InParent, TArray<FbxNode*>& NodeArray, const FName& Name, EObjectFlags Flags, UFbxSkeletalMeshImportData* TemplateImportData, TArray<FbxShape*> *FbxShapeArray=NULL, FSkeletalMeshImportData* OutData=NULL, bool bCreateRenderData = true );
 
 	/**
 	 * Add to the animation set, the animations contained within the FBX scene, for the given skeletal mesh
@@ -444,10 +443,9 @@ public:
 	 *
 	 * @param SkelMeshNodeArray - Fbx Nodes that the base Skeletal Mesh construct from
 	 * @param BaseSkelMesh - base Skeletal Mesh
-	 * @param Filename - Fbx file name
 	 * @param LODIndex - LOD index
 	 */
-	void ImportFbxMorphTarget(TArray<FbxNode*> &SkelMeshNodeArray, USkeletalMesh* BaseSkelMesh, UObject* Parent, const FString& Filename, int32 LODIndex);
+	void ImportFbxMorphTarget(TArray<FbxNode*> &SkelMeshNodeArray, USkeletalMesh* BaseSkelMesh, UObject* Parent, int32 LODIndex);
 
 	/**
 	 * Import LOD object for skeletal mesh
@@ -622,10 +620,9 @@ private:
 	 *
 	 * @param SkelMeshNodeArray - Fbx Nodes that the base Skeletal Mesh construct from
 	 * @param BaseSkelMesh - base Skeletal Mesh
-	 * @param Filename - Fbx file name
 	 * @param LODIndex - LOD index of the skeletal mesh
 	 */
-	void ImportMorphTargetsInternal( TArray<FbxNode*>& SkelMeshNodeArray, USkeletalMesh* BaseSkelMesh, UObject* Parent, const FString& InFilename, int32 LODIndex );
+	void ImportMorphTargetsInternal( TArray<FbxNode*>& SkelMeshNodeArray, USkeletalMesh* BaseSkelMesh, UObject* Parent, int32 LODIndex );
 
 	/**
 	* sub-method called from ImportSkeletalMeshLOD method
diff --git a/Engine/Source/Editor/UnrealEd/Private/FbxMeshUtils.cpp b/Engine/Source/Editor/UnrealEd/Private/FbxMeshUtils.cpp
--- a/Engine/Source/Editor/UnrealEd/Private/FbxMeshUtils.cpp
+++ b/Engine/Source/Editor/UnrealEd/Private/FbxMeshUtils.cpp
@@ -281,7 +281,7 @@ namespace FbxMeshUtils
 					USkeletalMesh* TempSkelMesh = NULL;
 					// @todo AssetImportData does this temp skeletal mesh need import data?
 					UFbxSkeletalMeshImportData* TempAssetImportData = NULL;
-					TempSkelMesh = (USkeletalMesh*)FFbxImporter->ImportSkeletalMesh(GetTransientPackage(), bUseLODs? SkelMeshNodeArray: *MeshObject, NAME_None, (EObjectFlags)0, TempAssetImportData, FPaths::GetBaseFilename(Filename));
+					TempSkelMesh = (USkeletalMesh*)FFbxImporter->ImportSkeletalMesh(GetTransientPackage(), bUseLODs? SkelMeshNodeArray: *MeshObject, NAME_None, (EObjectFlags)0, TempAssetImportData);
 
 					// Add imported mesh to existing model
 					bool bImportSucceeded = false;
@@ -295,7 +295,7 @@ namespace FbxMeshUtils
 
 					if(ImportOptions->bImportMorph)
 					{
-						FFbxImporter->ImportFbxMorphTarget(SkelMeshNodeArray, SelectedSkelMesh, SelectedSkelMesh->GetOutermost(), Filename, SelectedLOD);
+						FFbxImporter->ImportFbxMorphTarget(SkelMeshNodeArray, SelectedSkelMesh, SelectedSkelMesh->GetOutermost(), SelectedLOD);
 					}
 
 					if (bImportSucceeded)
diff --git a/Engine/Source/Editor/UnrealEd/Private/SkeletalMeshEdit.cpp b/Engine/Source/Editor/UnrealEd/Private/SkeletalMeshEdit.cpp
--- a/Engine/Source/Editor/UnrealEd/Private/SkeletalMeshEdit.cpp
+++ b/Engine/Source/Editor/UnrealEd/Private/SkeletalMeshEdit.cpp
@@ -77,8 +77,7 @@ UAnimSequence * UEditorEngine::ImportFbxAnimation( USkeleton* Skeleton, UObject*
 			{
 				// since to know full path, reimport will need to do same
 				UFbxAnimSequenceImportData* ImportData = UFbxAnimSequenceImportData::GetImportDataForAnimSequence(NewAnimation, TemplateImportData);
-				ImportData->SourceFilePath = FReimportManager::SanitizeImportFilename(UFactory::CurrentFilename, NewAnimation);
-				ImportData->SourceFileTimestamp = IFileManager::Get().GetTimeStamp(*UFactory::CurrentFilename).ToString();
+				ImportData->Update(UFactory::CurrentFilename);
 				ImportData->bDirty = false;
 			}
 		}
@@ -549,8 +548,7 @@ UAnimSequence * UnFbx::FFbxImporter::ImportAnimations(USkeleton* Skeleton, UObje
 
 		// since to know full path, reimport will need to do same
 		UFbxAnimSequenceImportData* ImportData = UFbxAnimSequenceImportData::GetImportDataForAnimSequence(DestSeq, TemplateImportData);
-		ImportData->SourceFilePath = FReimportManager::SanitizeImportFilename(UFactory::CurrentFilename, DestSeq);
-		ImportData->SourceFileTimestamp = IFileManager::Get().GetTimeStamp(*UFactory::CurrentFilename).ToString();
+		ImportData->Update(UFactory::CurrentFilename);
 		ImportData->bDirty = false;
 
 		ImportAnimation(Skeleton, DestSeq, Name, SortedLinks, NodeArray, CurAnimStack, ResampleRate, AnimTimeSpan);
diff --git a/Engine/Source/Editor/UnrealEd/Private/UnrealEdEngine.cpp b/Engine/Source/Editor/UnrealEd/Private/UnrealEdEngine.cpp
--- a/Engine/Source/Editor/UnrealEd/Private/UnrealEdEngine.cpp
+++ b/Engine/Source/Editor/UnrealEd/Private/UnrealEdEngine.cpp
@@ -31,6 +31,7 @@
 #include "AutoReimport/AutoReimportManager.h"
 #include "NotificationManager.h"
 #include "SNotificationList.h"
+#include "AutoReimport/AssetSourceFilenameCache.h"
 #include "UObject/UObjectThreadContext.h"
 #include "Components/BillboardComponent.h"
 #include "Components/ArrowComponent.h"
@@ -309,6 +310,8 @@ void UUnrealEdEngine::MakeSortedSpriteInfo(TArray<FSpriteCategoryInfo>& OutSorte
 
 void UUnrealEdEngine::PreExit()
 {
+	FAssetSourceFilenameCache::Get().Shutdown();
+
 	// Notify edit modes we're mode at exit
 	FEditorModeRegistry::Get().Shutdown();
 
diff --git a/Engine/Source/Editor/UnrealEd/Public/AutoReimport/AssetSourceFilenameCache.h b/Engine/Source/Editor/UnrealEd/Public/AutoReimport/AssetSourceFilenameCache.h
new file mode 100644
--- /dev/null
+++ b/Engine/Source/Editor/UnrealEd/Public/AutoReimport/AssetSourceFilenameCache.h
@@ -0,0 +1,42 @@
+// Copyright 1998-2015 Epic Games, Inc. All Rights Reserved.
+
+#pragma once
+
+class IAssetRegistry;
+
+/** Class responsible for maintaing a cache of clean source file names (bla.txt) to asset data */
+class UNREALED_API FAssetSourceFilenameCache
+{
+public:
+	FAssetSourceFilenameCache();
+
+	/** Singleton access */
+	static FAssetSourceFilenameCache& Get();
+
+	/** Helper functions to extract asset import information from asset registry tags */
+	static TOptional<FAssetImportInfo> ExtractAssetImportInfo(const TMap<FName, FString>& InTags);
+	static TOptional<FAssetImportInfo> ExtractAssetImportInfo(const TArray<UObject::FAssetRegistryTag>& InTags);
+
+	/** Retrieve a list of assets that were imported from the specified filename */
+	TArray<FAssetData> GetAssetsPertainingToFile(const IAssetRegistry& Registry, const FString& AbsoluteFilename) const;
+
+	/** Shutdown this instance of the cache */
+	void Shutdown();
+
+	/** Event for when an asset has been renamed, and has been updated in our source file cache */
+	DECLARE_EVENT_TwoParams( FAssetSourceFilenameCache, FAssetRenamedEvent, const FAssetData&, const FString& );
+	FAssetRenamedEvent& OnAssetRenamed() { return AssetRenamedEvent; }
+
+private:
+	/** Delegate bindings that keep the cache up-to-date */
+	void HandleOnAssetAdded(const FAssetData& AssetData);
+	void HandleOnAssetRemoved(const FAssetData& AssetData);
+	void HandleOnAssetRenamed(const FAssetData& AssetData, const FString& OldPath);
+	void HandleOnAssetUpdated(const FAssetImportInfo& OldData, const UAssetImportData* ImportData);
+	
+	/** Event that is triggered when an asset has been renamed, and we've updated our cache */
+	FAssetRenamedEvent AssetRenamedEvent;
+
+	/** Map of clean filenames (no leading path information) to object paths that were imported with that file */
+	TMap<FString, TSet<FName>> SourceFileToObjectPathCache;
+};
\ No newline at end of file
diff --git a/Engine/Source/Editor/UnrealEd/Public/EditorReimportHandler.h b/Engine/Source/Editor/UnrealEd/Public/EditorReimportHandler.h
--- a/Engine/Source/Editor/UnrealEd/Public/EditorReimportHandler.h
+++ b/Engine/Source/Editor/UnrealEd/Public/EditorReimportHandler.h
@@ -59,26 +59,6 @@ public:
 	 */
 	UNREALED_API virtual void UpdateReimportPaths( UObject* Obj, const TArray<FString>& InFilenames );
 
-	/**
-	 * Convert a file path to be relative to the specified object, if it resides in the same package folder
-	 * 
-	 * @param	InPath Absolute (or relative by cwd) path to the source file
-	 * @param	Obj Object file to make InPath relative to
-	 *
-	 * @return	A path relative to the specified object, BaseDir() or failing that, an absolute path
-	 */
-	UNREALED_API static FString SanitizeImportFilename(const FString& InPath, const UObject* Obj);
-
-	/**
-	 * Take an import filename and resolve it to an absolute filename on disk
-	 * 
-	 * @param	InRelativePath Path relative to the specified file or BaseDir(), or an absolute path
-	 * @param	Obj Object file to make InPath relative to
-	 *
-	 * @return	An absolute path to the import file
-	 */
-	UNREALED_API static FString ResolveImportFilename(const FString& InRelativePath, const UObject* Obj);
-
 	/**
 	 * Gets the delegate that's fired prior to reimporting an asset
 	 *
diff --git a/Engine/Source/Runtime/AssetRegistry/Private/AssetRegistry.h b/Engine/Source/Runtime/AssetRegistry/Private/AssetRegistry.h
--- a/Engine/Source/Runtime/AssetRegistry/Private/AssetRegistry.h
+++ b/Engine/Source/Runtime/AssetRegistry/Private/AssetRegistry.h
@@ -182,11 +182,6 @@ private:
 	/** The map of asset tag to asset data for assets saved to disk */
 	TMap<FName, TArray<FAssetData*> > CachedAssetsByTag;
 
-#if WITH_EDITORONLY_DATA
-	/** A map of clean source file name (asset.ext) to assets saved to disk */
-	TMap<FName, TArray<FAssetData*>> CachedAssetsBySourceFileName;
-#endif
-
 	/** A map of object names to dependency data */
 	TMap<FName, FDependsNode*> CachedDependsNodes;
 
diff --git a/Engine/Source/Runtime/AssetRegistry/Private/PackageReader.cpp b/Engine/Source/Runtime/AssetRegistry/Private/PackageReader.cpp
--- a/Engine/Source/Runtime/AssetRegistry/Private/PackageReader.cpp
+++ b/Engine/Source/Runtime/AssetRegistry/Private/PackageReader.cpp
@@ -141,7 +141,7 @@ bool FPackageReader::ReadAssetRegistryData (TArray<FBackgroundAssetData*>& Asset
 		if (bLegacyPackage || bNoMapAsset)
 		{
 			const FString AssetName = FPackageName::GetLongPackageAssetName(PackageName);
-			TMultiMap<FString, FString> TagsAndValues;
+			TMap<FString, FString> TagsAndValues;
 			AssetDataList.Add(new FBackgroundAssetData(PackageName, PackagePath, TEXT(""), AssetName, TEXT("World"), TagsAndValues, PackageFileSummary.ChunkIDs));
 		}
 	}
@@ -156,7 +156,7 @@ bool FPackageReader::ReadAssetRegistryData (TArray<FBackgroundAssetData*>& Asset
 		*this << ObjectClassName;
 		*this << TagCount;
 
-		TMultiMap<FString, FString> TagsAndValues;
+		TMap<FString, FString> TagsAndValues;
 
 		for(int32 TagIdx = 0; TagIdx < TagCount; ++TagIdx)
 		{
@@ -246,7 +246,7 @@ bool FPackageReader::ReadAssetDataFromThumbnailCache(TArray<FBackgroundAssetData
 		}
 
 		// Create a new FBackgroundAssetData for this asset and update it with the gathered data
-		TMultiMap<FString, FString> TagsAndValues;
+		TMap<FString, FString> TagsAndValues;
 		AssetDataList.Add(new FBackgroundAssetData(PackageName, PackagePath, GroupNames, AssetName, AssetClassName, TagsAndValues, PackageFileSummary.ChunkIDs));
 	}
 
diff --git a/Engine/Source/Runtime/Core/Public/Misc/SecureHash.h b/Engine/Source/Runtime/Core/Public/Misc/SecureHash.h
--- a/Engine/Source/Runtime/Core/Public/Misc/SecureHash.h
+++ b/Engine/Source/Runtime/Core/Public/Misc/SecureHash.h
@@ -87,6 +87,67 @@ private:
 };
 //@}
 
+struct FMD5Hash;
+
+namespace LexicalConversion
+{
+	CORE_API FString ToString(const FMD5Hash& Hash);
+	CORE_API void FromString(FMD5Hash& Hash, const TCHAR* Buffer);
+}
+
+/** Simple helper struct to ease the caching of MD5 hashes */
+struct FMD5Hash
+{
+	/** Default constructor */
+	FMD5Hash() : bIsValid(false) {}
+
+	/** Check whether this has hash is valid or not */
+	bool IsValid() const { return bIsValid; }
+
+	/** Set up the MD5 hash from a container */
+	void Set(FMD5& MD5)
+	{
+		MD5.Final(Bytes);
+		bIsValid = true;
+	}
+
+	/** Compare one hash with another */
+	friend bool operator==(const FMD5Hash& LHS, const FMD5Hash& RHS)
+	{
+		return LHS.bIsValid == RHS.bIsValid && (!LHS.bIsValid || FMemory::Memcmp(LHS.Bytes, RHS.Bytes, 16) == 0);
+	}
+
+	/** Compare one hash with another */
+	friend bool operator!=(const FMD5Hash& LHS, const FMD5Hash& RHS)
+	{
+		return LHS.bIsValid != RHS.bIsValid || (LHS.bIsValid && FMemory::Memcmp(LHS.Bytes, RHS.Bytes, 16) != 0);
+	}
+
+	/** Serialise this hash */
+	friend FArchive& operator<<(FArchive& Ar, FMD5Hash& Hash)
+	{
+		Ar << Hash.bIsValid;
+		if (Hash.bIsValid)
+		{
+			Ar.Serialize(Hash.Bytes, 16);
+		}
+
+		return Ar;
+	}
+
+	/** Hash the specified file contents (using the optionally supplied scratch buffer) */
+	CORE_API static FMD5Hash HashFile(const TCHAR* InFilename, TArray<uint8>* Buffer = nullptr);
+
+private:
+	/** Whether this hash is valid or not */
+	bool bIsValid;
+
+	/** The bytes this hash comprises */
+	uint8 Bytes[16];
+
+	friend FString LexicalConversion::ToString(const FMD5Hash&);
+	friend void LexicalConversion::FromString(FMD5Hash& Hash, const TCHAR*);
+};
 
 /*-----------------------------------------------------------------------------
 	SHA-1 functions.
diff --git a/Engine/Source/Runtime/Engine/Classes/Curves/CurveBase.h b/Engine/Source/Runtime/Engine/Classes/Curves/CurveBase.h
--- a/Engine/Source/Runtime/Engine/Classes/Curves/CurveBase.h
+++ b/Engine/Source/Runtime/Engine/Classes/Curves/CurveBase.h
@@ -464,10 +464,6 @@ class ENGINE_API UCurveBase : public UObject, public FCurveOwnerInterface
 {
 	GENERATED_UCLASS_BODY()
 
-	/** The filename imported to create this object. Relative to this object's package, BaseDir() or absolute */
-	UPROPERTY()
-	FString ImportPath;
-
 	/** Get the time range across all curves */
 	UFUNCTION(BlueprintCallable, Category="Math|Curves")
 	void GetTimeRange(float& MinTime, float& MaxTime) const;
@@ -514,6 +510,15 @@ public:
 #if WITH_EDITORONLY_DATA
 	/** Override to ensure we write out the asset import data */
 	virtual void GetAssetRegistryTags(TArray<FAssetRegistryTag>& OutTags) const override;
+	virtual void PostLoad() override;
+
+	UPROPERTY()
+	class UAssetImportData* AssetImportData;
+
+	/** The filename imported to create this object. Relative to this object's package, BaseDir() or absolute */
+	UPROPERTY()
+	FString ImportPath_DEPRECATED;
+
 #endif
 	// End UObject interface
 };
diff --git a/Engine/Source/Runtime/Engine/Classes/EditorFramework/AssetImportData.h b/Engine/Source/Runtime/Engine/Classes/EditorFramework/AssetImportData.h
--- a/Engine/Source/Runtime/Engine/Classes/EditorFramework/AssetImportData.h
+++ b/Engine/Source/Runtime/Engine/Classes/EditorFramework/AssetImportData.h
@@ -1,34 +1,102 @@
 // Copyright 1998-2015 Epic Games, Inc. All Rights Reserved.
 
 #pragma once
+
 #include "AssetImportData.generated.h"
 
-/**
- * A class to hold import options and data for assets. This class has many children in editor code.
- */
-UCLASS(MinimalAPI)
-class UAssetImportData : public UObject
+/** Struct that is used to store an array of asset import data as an asset registry tag */
+struct ENGINE_API FAssetImportInfo
 {
-	GENERATED_UCLASS_BODY()
+#if WITH_EDITORONLY_DATA
+	struct FSourceFile
+	{
+		FSourceFile(FString InRelativeFilename, const FDateTime& InTimestamp = 0, const FMD5Hash& InFileHash = FMD5Hash())
+			: RelativeFilename(MoveTemp(InRelativeFilename))
+			, Timestamp(0)
+			, FileHash(InFileHash)
+		{}
+
+		/** The path to the file that this asset was imported from. Relative to either the asset's package, BaseDir(), or absolute */
+		FString RelativeFilename;
+
+		/** The timestamp of the file when it was imported (as UTC). 0 when unknown. */
+		FDateTime Timestamp;
+
+		/** The MD5 hash of the file when it was imported. Invalid when unknown. */
+		FMD5Hash FileHash;
+	};
+
+	/** Convert this import information to JSON */
+	FString ToJson() const;
+
+	/** Attempt to parse an asset import structure from the specified json string. */
+	static TOptional<FAssetImportInfo> FromJson(FString InJsonString);
 
-	/** Path to the resource used to construct this static mesh. Relative to the object's package, BaseDir() or absolute */
-	UPROPERTY(EditAnywhere, Category=ImportSettings)
-	FString SourceFilePath;
+	/** Insert information pertaining to a new source file to this structure */
+	void Insert(const FSourceFile& InSourceFile) { SourceFiles.Add(InSourceFile); }
 
-	/** Date/Time-stamp of the file from the last import */
-	UPROPERTY(VisibleAnywhere, Category=ImportSettings)
-	FString SourceFileTimestamp;
+	/** Copy the information from another asset import info structure into this one. */
+	void CopyFrom(const FAssetImportInfo& Other) { SourceFiles = Other.SourceFiles; }
 
+	/** Const public access to the source file data */
+	const TArray<FSourceFile>& GetSourceFileData() const { return SourceFiles; }
+
+protected:
+
+	/** Array of information pertaining to the source files that this asset was imported from */
+	TArray<FSourceFile> SourceFiles;
+#endif // WITH_EDITORONLY_DATA
+};
+
+DECLARE_MULTICAST_DELEGATE_TwoParams(FOnImportDataChanged, const FAssetImportInfo& /*OldData*/, const class UAssetImportData* /* NewData */);
+
+UCLASS()
+class ENGINE_API UAssetImportData : public UObject, public FAssetImportInfo
+{
+public:
+	GENERATED_UCLASS_BODY()
+
+#if WITH_EDITORONLY_DATA
 	/** If true, the settings have been modified but not reimported yet */
 	UPROPERTY()
 	uint32 bDirty : 1;
 
-	// UObject interface
-#if WITH_EDITOR
-	ENGINE_API virtual void PostEditChangeProperty(struct FPropertyChangedEvent& PropertyChangedEvent) override;
-#endif // WITH_EDITOR
-	// End of UObject interface
-};
+public:
+
+	/** Static event that is broadcast whenever any asset has updated its import data */
+	static FOnImportDataChanged OnImportDataChanged;
+
+	/** Update this import data using the specified file. Called when an asset has been imported from a file. */
+	void Update(const FString& AbsoluteFilename);
+
+	/** Update this import data using the specified filename. Will not update the imported timestamp or MD5. */
+	void UpdateFilenameOnly(const FString& InPath);
+
+	/** Helper function to return the first filename stored in this data. The resulting filename will be absolute (ie, not relative to the asset).  */
+	FString GetFirstFilename() const;
+
+	/** Extract all the (resolved) filenames from this data  */
+	void ExtractFilenames(TArray<FString>& AbsoluteFilenames) const;
 
+	/** Extract all the (resolved) filenames from this data  */
+	TArray<FString> ExtractFilenames() const;
+
+	/** Resolve a filename that is relative to either the specified package, BaseDir() or absolute */
+	static FString ResolveImportFilename(const FString& InRelativePath, const UPackage* Outermost);
+
+protected:
+
+	/** Convert an absolute import path so that it's relative to either this object's package, BaseDir() or leave it absolute */
+	FString SanitizeImportFilename(const FString& InPath) const;
+
+	/** Resolve a filename that is relative to either this object's package, BaseDir() or absolute */
+	FString ResolveImportFilename(const FString& InRelativePath) const;
+
+	/** Overridden serialize function to write out the underlying data as json */
+	virtual void Serialize(FArchive& Ar) override;
+	virtual void PostEditChangeProperty(struct FPropertyChangedEvent& PropertyChangedEvent) override;
+	
+#endif		// WITH_EDITORONLY_DATA
+};
 
 
diff --git a/Engine/Source/Runtime/Engine/Classes/VectorField/VectorFieldStatic.h b/Engine/Source/Runtime/Engine/Classes/VectorField/VectorFieldStatic.h
--- a/Engine/Source/Runtime/Engine/Classes/VectorField/VectorFieldStatic.h
+++ b/Engine/Source/Runtime/Engine/Classes/VectorField/VectorFieldStatic.h
@@ -34,13 +34,11 @@ public:
 	FByteBulkData SourceData;
 
 #if WITH_EDITORONLY_DATA
-	/** Path to the resource used to construct this vector field. Relative to the object's package, BaseDir() or absolute. */
-	UPROPERTY(Category=SourceAsset, VisibleAnywhere)
-	FString SourceFilePath;
+	UPROPERTY()
+	FString SourceFilePath_DEPRECATED;
 
-	/** Date/Time-stamp of the file from the last import */
-	UPROPERTY(Category=SourceAsset, VisibleAnywhere)
-	FString SourceFileTimestamp;
+	UPROPERTY()
+	class UAssetImportData* AssetImportData;
 #endif // WITH_EDITORONLY_DATA
 
 	// Begin UObject interface.
diff --git a/Engine/Source/Runtime/Engine/Private/DestructibleMesh.cpp b/Engine/Source/Runtime/Engine/Private/DestructibleMesh.cpp
--- a/Engine/Source/Runtime/Engine/Private/DestructibleMesh.cpp
+++ b/Engine/Source/Runtime/Engine/Private/DestructibleMesh.cpp
@@ -638,9 +638,11 @@ ENGINE_API bool UDestructibleMesh::BuildFromStaticMesh( UStaticMesh& StaticMesh
 	}
 
 	SourceStaticMesh = &StaticMesh;
-	if ( StaticMesh.AssetImportData == NULL || !FDateTime::Parse(StaticMesh.AssetImportData->SourceFileTimestamp, SourceSMImportTimestamp) )
+
+	SourceSMImportTimestamp = FDateTime::MinValue();
+	if ( StaticMesh.AssetImportData && StaticMesh.AssetImportData->GetSourceFileData().Num() == 1 )
 	{
-		SourceSMImportTimestamp = FDateTime::MinValue();
+		SourceSMImportTimestamp = StaticMesh.AssetImportData->GetSourceFileData()[0].Timestamp;
 	}
 
 #if WITH_APEX
diff --git a/Engine/Source/Runtime/Engine/Private/EditorFramework/AssetImportData.cpp b/Engine/Source/Runtime/Engine/Private/EditorFramework/AssetImportData.cpp
--- a/Engine/Source/Runtime/Engine/Private/EditorFramework/AssetImportData.cpp
+++ b/Engine/Source/Runtime/Engine/Private/EditorFramework/AssetImportData.cpp
@@ -1,19 +1,200 @@
 // Copyright 1998-2015 Epic Games, Inc. All Rights Reserved.
 
 #include "EnginePrivate.h"
+
 #include "EditorFramework/AssetImportData.h"
 
+// This whole class is compiled out in non-editor
 UAssetImportData::UAssetImportData(const FObjectInitializer& ObjectInitializer)
 	: Super(ObjectInitializer)
 {
+#if WITH_EDITORONLY_DATA
 	bDirty = false;
+#endif
+}
+
+#if WITH_EDITORONLY_DATA
+
+FOnImportDataChanged UAssetImportData::OnImportDataChanged;
+
+FString FAssetImportInfo::ToJson() const
+{
+	FString Json;
+	Json.Reserve(1024);
+	Json += TEXT("[");
+
+	for (int32 Index = 0; Index < SourceFiles.Num(); ++Index)
+	{
+		Json += FString::Printf(TEXT("{ \"RelativeFilename\" : \"%s\", \"Timestamp\" : %d, \"FileMD5\" : \"%s\" }"),
+			*SourceFiles[Index].RelativeFilename,
+			SourceFiles[Index].Timestamp.ToUnixTimestamp(),
+			*LexicalConversion::ToString(SourceFiles[Index].FileHash)
+			);
+
+		if (Index != SourceFiles.Num() - 1)
+		{
+			Json += TEXT(",");
+		}
+	}
+
+	Json += TEXT("]");
+	return Json;
+}
+
+TOptional<FAssetImportInfo> FAssetImportInfo::FromJson(FString InJsonString)
+{
+	// Load json
+	TSharedRef<TJsonReader<TCHAR>> Reader = FJsonStringReader::Create(MoveTemp(InJsonString));
+
+	TArray<TSharedPtr<FJsonValue>> JSONSourceFiles;
+	if (!FJsonSerializer::Deserialize(Reader, JSONSourceFiles))
+	{
+		return TOptional<FAssetImportInfo>();
+	}
+	
+	FAssetImportInfo Info;
+
+	for (const auto& Value : JSONSourceFiles)
+	{
+		const TSharedPtr<FJsonObject>& SourceFile = Value->AsObject();
+		if (!SourceFile.IsValid())
+		{
+			continue;
+		}
+
+		FString RelativeFilename, TimestampString, MD5String;
+		SourceFile->TryGetStringField("RelativeFilename", RelativeFilename);
+		SourceFile->TryGetStringField("Timestamp", TimestampString);
+		SourceFile->TryGetStringField("FileMD5", MD5String);
+
+		if (RelativeFilename.IsEmpty())
+		{
+			continue;
+		}
+
+		int64 UnixTimestamp = 0;
+		LexicalConversion::FromString(UnixTimestamp, *TimestampString);
+
+		FMD5Hash FileHash;
+		LexicalConversion::FromString(FileHash, *MD5String);
+
+		Info.SourceFiles.Emplace(MoveTemp(RelativeFilename), UnixTimestamp, FileHash);
+	}
+
+	return Info;
+}
+
+void UAssetImportData::UpdateFilenameOnly(const FString& InPath)
+{
+	SourceFiles.Reset();
+	SourceFiles.Emplace(SanitizeImportFilename(InPath));
+}
+
+void UAssetImportData::Update(const FString& InPath)
+{
+	FAssetImportInfo Old;
+	Old.CopyFrom(*this);
+
+	SourceFiles.Reset();
+	SourceFiles.Emplace(
+		SanitizeImportFilename(InPath),
+		IFileManager::Get().GetTimeStamp(*InPath),
+		FMD5Hash::HashFile(*InPath)
+		);
+
+	bDirty = true;
+	
+	OnImportDataChanged.Broadcast(Old, this);
+}
+
+FString UAssetImportData::GetFirstFilename() const
+{
+	return SourceFiles.Num() > 0 ? ResolveImportFilename(SourceFiles[0].RelativeFilename) : FString();
+}
+
+void UAssetImportData::ExtractFilenames(TArray<FString>& AbsoluteFilenames) const
+{
+	for (const auto& File : SourceFiles)
+	{
+		AbsoluteFilenames.Add(ResolveImportFilename(File.RelativeFilename));
+	}
+}
+
+FString UAssetImportData::SanitizeImportFilename(const FString& InPath) const
+{
+	const UPackage* Package = GetOutermost();
+	if (Package)
+	{
+		const bool		bIncludeDot = true;
+		const FString	PackagePath	= Package->GetPathName();
+		const FName		MountPoint	= FPackageName::GetPackageMountPoint(PackagePath);
+		const FString	PackageFilename = FPackageName::LongPackageNameToFilename(PackagePath, FPaths::GetExtension(InPath, bIncludeDot));
+		const FString	AbsolutePath = FPaths::ConvertRelativePathToFull(InPath);
+
+		if ((MountPoint == FName("Engine") && AbsolutePath.StartsWith(FPaths::ConvertRelativePathToFull(FPaths::EngineContentDir()))) ||
+			(MountPoint == FName("Game") &&	AbsolutePath.StartsWith(FPaths::ConvertRelativePathToFull(FPaths::GameDir()))))
+		{
+			FString RelativePath = InPath;
+			FPaths::MakePathRelativeTo(RelativePath, *PackageFilename);
+			return RelativePath;
+		}
+	}
+
+	return IFileManager::Get().ConvertToRelativePath(*InPath);
+}
+
+FString UAssetImportData::ResolveImportFilename(const FString& InRelativePath, const UPackage* Outermost)
+{
+	FString RelativePath = InRelativePath;
+
+	if (Outermost)
+	{
+		// Relative to the package filename?
+		const FString PathRelativeToPackage = FPaths::GetPath(FPackageName::LongPackageNameToFilename(Outermost->GetPathName())) / InRelativePath;
+		if (FPaths::FileExists(PathRelativeToPackage))
+		{
+			RelativePath = PathRelativeToPackage;
+		}
+	}
+
+	// Convert relative paths
+	return FPaths::ConvertRelativePathToFull(RelativePath);	
+}
+
+FString UAssetImportData::ResolveImportFilename(const FString& InRelativePath) const
+{
+	return ResolveImportFilename(InRelativePath, GetOutermost());
+}
+
+void UAssetImportData::Serialize(FArchive& Ar)
+{
+	if (Ar.UE4Ver() >= VER_UE4_ASSET_IMPORT_DATA_AS_JSON)
+	{
+		FString Json;
+		if (Ar.IsLoading())
+		{
+			Ar << Json;
+			TOptional<FAssetImportInfo> Copy = FromJson(MoveTemp(Json));
+			if (Copy.IsSet())
+			{
+				CopyFrom(Copy.GetValue());
+			}
+		}
+		else if (Ar.IsSaving())
+		{
+			Json = ToJson();
+			Ar << Json;
+		}
+	}
+
+	Super::Serialize(Ar);
 }
 
-#if WITH_EDITOR
 void UAssetImportData::PostEditChangeProperty(FPropertyChangedEvent& PropertyChangedEvent)
 {
 	Super::PostEditChangeProperty(PropertyChangedEvent);
 
 	bDirty = true;
 }
-#endif // WITH_EDITOR
+
+#endif // WITH_EDITORONLY_DATA
\ No newline at end of file
diff --git a/Engine/Source/Runtime/Engine/Private/StaticMesh.cpp b/Engine/Source/Runtime/Engine/Private/StaticMesh.cpp
--- a/Engine/Source/Runtime/Engine/Private/StaticMesh.cpp
+++ b/Engine/Source/Runtime/Engine/Private/StaticMesh.cpp
@@ -1539,7 +1539,7 @@ void UStaticMesh::GetAssetRegistryTags(TArray<FAssetRegistryTag>& OutTags) const
 #if WITH_EDITORONLY_DATA
 	if (AssetImportData)
 	{
-		OutTags.Add( FAssetRegistryTag(SourceFileTagName(), AssetImportData->SourceFilePath, FAssetRegistryTag::TT_Hidden) );
+		OutTags.Add( FAssetRegistryTag(SourceFileTagName(), AssetImportData->ToJson(), FAssetRegistryTag::TT_Hidden) );
 	}
 #endif
 
@@ -1881,8 +1881,10 @@ void UStaticMesh::Serialize(FArchive& Ar)
 			AssetImportData = NewObject<UAssetImportData>(this);
 		}
 
-		AssetImportData->SourceFilePath = SourceFilePath_DEPRECATED;
-		AssetImportData->SourceFileTimestamp = SourceFileTimestamp_DEPRECATED;
+		FAssetImportInfo Info;
+		Info.Insert(FAssetImportInfo::FSourceFile(SourceFilePath_DEPRECATED));
+		AssetImportData->CopyFrom(Info);
+		
 		SourceFilePath_DEPRECATED = TEXT("");
 		SourceFileTimestamp_DEPRECATED = TEXT("");
 	}
