Fixed issue where BSP surface transforms could get broken when applying them to geometry surfaces which are out of sync with the owner Brush (because Auto Build BSP is disabled).

diff --git a/Engine/Source/Editor/UnrealEd/Private/EditorBsp.cpp b/Engine/Source/Editor/UnrealEd/Private/EditorBsp.cpp
--- a/Engine/Source/Editor/UnrealEd/Private/EditorBsp.cpp
+++ b/Engine/Source/Editor/UnrealEd/Private/EditorBsp.cpp
@@ -1160,6 +1160,13 @@ int UEditorEngine::bspBrushCSG
 	const FRotator Rotation = Actor->GetActorRotation();
 	const FVector Location = Actor->GetActorLocation();
 
+	// Cache actor transform which is used for the geometry being built
+	Brush->OwnerLocationWhenLastBuilt = Location;
+	Brush->OwnerRotationWhenLastBuilt = Rotation;
+	Brush->OwnerScaleWhenLastBuilt = Scale;
+	Brush->OwnerPrepivotWhenLastBuilt = PrePivot;
+	Brush->bCachedOwnerTransformValid = true;
+
 	for( i=0; i<Brush->Polys->Element.Num(); i++ )
 	{
 		FPoly& CurrentPoly = Brush->Polys->Element[i];
diff --git a/Engine/Source/Editor/UnrealEd/Private/EditorCsg.cpp b/Engine/Source/Editor/UnrealEd/Private/EditorCsg.cpp
--- a/Engine/Source/Editor/UnrealEd/Private/EditorCsg.cpp
+++ b/Engine/Source/Editor/UnrealEd/Private/EditorCsg.cpp
@@ -357,17 +357,23 @@ void UEditorEngine::polyUpdateMaster
 	if( !Actor )
 		return;
 
-	const FVector ActorLocation = Actor->GetActorLocation();
-	const FVector ActorPrePivot = Actor->GetPrePivot();
-	const FVector ActorScale = Actor->GetActorScale();
-	const FRotator ActorRotation = -Actor->GetActorRotation();
+	UModel* Brush = Actor->Brush;
+	check(Brush);
 
-	for( int32 iEdPoly = Surf.iBrushPoly; iEdPoly < Actor->Brush->Polys->Element.Num(); iEdPoly++ )
+	// Use transform cached when the geometry was last built, in case the current Actor transform has changed since then
+	// (e.g. because Auto Update BSP is disabled)
+	check(Brush->bCachedOwnerTransformValid);
+	const FVector ActorLocation = Brush->OwnerLocationWhenLastBuilt;
+	const FVector ActorPrePivot = Brush->OwnerPrepivotWhenLastBuilt;
+	const FVector ActorScale = Brush->OwnerScaleWhenLastBuilt;
+	const FRotator ActorRotation = -Brush->OwnerRotationWhenLastBuilt;
+
+	for( int32 iEdPoly = Surf.iBrushPoly; iEdPoly < Brush->Polys->Element.Num(); iEdPoly++ )
 	{
-		FPoly& MasterEdPoly = Actor->Brush->Polys->Element[iEdPoly];
+		FPoly& MasterEdPoly = Brush->Polys->Element[iEdPoly];
 		if( iEdPoly==Surf.iBrushPoly || MasterEdPoly.iLink==Surf.iBrushPoly )
 		{
-			Actor->Brush->Polys->Element.ModifyItem( iEdPoly );
+			Brush->Polys->Element.ModifyItem( iEdPoly );
 
 			MasterEdPoly.Material  = Surf.Material;
 			MasterEdPoly.PolyFlags = Surf.PolyFlags & ~(PF_NoEdit);
diff --git a/Engine/Source/Runtime/Engine/Private/Model.cpp b/Engine/Source/Runtime/Engine/Private/Model.cpp
--- a/Engine/Source/Runtime/Engine/Private/Model.cpp
+++ b/Engine/Source/Runtime/Engine/Private/Model.cpp
@@ -528,6 +528,9 @@ UModel::UModel(const FObjectInitializer& ObjectInitializer)
 	, Points(this)
 	, Surfs(this)
 	, VertexBuffer(this)
+#if WITH_EDITOR
+	, bCachedOwnerTransformValid(false)
+#endif
 {
 
 }
@@ -541,6 +544,9 @@ UModel::UModel(FVTableHelper& Helper)
 	, Points(this)
 	, Surfs(this)
 	, VertexBuffer(this)
+#if WITH_EDITOR
+	, bCachedOwnerTransformValid(false)
+#endif
 {
 
 }
diff --git a/Engine/Source/Runtime/Engine/Public/Model.h b/Engine/Source/Runtime/Engine/Public/Model.h
--- a/Engine/Source/Runtime/Engine/Public/Model.h
+++ b/Engine/Source/Runtime/Engine/Public/Model.h
@@ -406,6 +406,15 @@ class UModel : public UObject
 
 	/** The level used to generate NodeGroups */
 	const ULevel* LightingLevel;
+
+	/** Cached transform of the owner brush when the geometry was last built */
+	FVector OwnerLocationWhenLastBuilt;
+	FRotator OwnerRotationWhenLastBuilt;
+	FVector OwnerScaleWhenLastBuilt;
+	FVector OwnerPrepivotWhenLastBuilt;
+
+	/** Specifies whether the above cached transform is valid */
+	bool bCachedOwnerTransformValid;
 #endif // WITH_EDITOR
 
 	// Other variables.
