Fix shadowed variable in FBX import

diff --git a/Engine/Source/Editor/UnrealEd/Private/Fbx/FbxStaticMeshImport.cpp b/Engine/Source/Editor/UnrealEd/Private/Fbx/FbxStaticMeshImport.cpp
--- a/Engine/Source/Editor/UnrealEd/Private/Fbx/FbxStaticMeshImport.cpp
+++ b/Engine/Source/Editor/UnrealEd/Private/Fbx/FbxStaticMeshImport.cpp
@@ -1151,8 +1151,8 @@ UStaticMesh* UnFbx::FFbxImporter::ImportStaticMeshAsSingle(UObject* InParent, TA
 		int32 MaxMaterialIndex = 0;
 		int32 FirstOpenUVChannel = 1;
 		{
-			FRawMesh RawMesh;
-			SrcModel.RawMeshBulkData->LoadRawMesh(RawMesh);
+			FRawMesh LocalRawMesh;
+			SrcModel.RawMeshBulkData->LoadRawMesh(LocalRawMesh);
 
 			if (bDoRemap)
 			{
@@ -1198,7 +1198,7 @@ UStaticMesh* UnFbx::FFbxImporter::ImportStaticMeshAsSingle(UObject* InParent, TA
 				}
 			}
 
-			SrcModel.RawMeshBulkData->SaveRawMesh(RawMesh);
+			SrcModel.RawMeshBulkData->SaveRawMesh(LocalRawMesh);
 		}
 
 		// Setup per-section info and the materials array.
