Added support for instanced mesh components (e.g. foliage) to Fbx export
Fbx export now instances meshes where possible, reducing file size and improving export speed of large scenes
Also a little bit of cleanup and const-ifying

diff --git a/Engine/Source/Editor/UnrealEd/Private/EditorExporters.cpp b/Engine/Source/Editor/UnrealEd/Private/EditorExporters.cpp
--- a/Engine/Source/Editor/UnrealEd/Private/EditorExporters.cpp
+++ b/Engine/Source/Editor/UnrealEd/Private/EditorExporters.cpp
@@ -1559,7 +1559,7 @@ bool UStaticMeshExporterOBJ::ExportText(const FExportObjectInnerContext* Context
 		Ar.Log( TEXT("# UnrealEd OBJ exporter\r\n") );
 
 		// Currently, we only export LOD 0 of the static mesh. In the future, we could potentially export all available LODs
-		FStaticMeshLODResources& RenderData = StaticMesh->GetLODForExport(0);
+		const FStaticMeshLODResources& RenderData = StaticMesh->GetLODForExport(0);
 		FRawMesh RawMesh;
 		StaticMesh->SourceModels[0].RawMeshBulkData->LoadRawMesh(RawMesh);
 
@@ -1689,7 +1689,7 @@ bool UStaticMeshExporterOBJ::ExportText(const FExportObjectInnerContext* Context
 		File->Logf( TEXT("# UnrealEd OBJ exporter (_Internal)\r\n") );
 
 		// Currently, we only export LOD 0 of the static mesh. In the future, we could potentially export all available LODs
-		FStaticMeshLODResources& RenderData = StaticMesh->GetLODForExport(0);
+		const FStaticMeshLODResources& RenderData = StaticMesh->GetLODForExport(0);
 		uint32 VertexCount = RenderData.GetNumVertices();
 
 		check(VertexCount == RenderData.VertexBuffer.GetNumVertices());
diff --git a/Engine/Source/Editor/UnrealEd/Private/Fbx/FbxMainExport.cpp b/Engine/Source/Editor/UnrealEd/Private/Fbx/FbxMainExport.cpp
--- a/Engine/Source/Editor/UnrealEd/Private/Fbx/FbxMainExport.cpp
+++ b/Engine/Source/Editor/UnrealEd/Private/Fbx/FbxMainExport.cpp
@@ -219,6 +219,7 @@ void FFbxExporter::CloseDocument()
 	FbxActors.Reset();
 	FbxSkeletonRoots.Reset();
 	FbxMaterials.Reset();
+	FbxMeshes.Reset();
 	FbxNodeNameToIndexMap.Reset();
 	
 	if (Scene)
@@ -427,32 +428,32 @@ void FFbxExporter::ExportBrush(ABrush* Actor, UModel* InModel, bool bConvertToSt
 {
 	if (Scene == NULL || Actor == NULL || !Actor->GetBrushComponent()) return;
 
- 	if (!bConvertToStaticMesh)
- 	{
- 		// Retrieve the information structures, verifying the integrity of the data.
- 		UModel* Model = Actor->GetBrushComponent()->Brush;
+	if (!bConvertToStaticMesh)
+	{
+		// Retrieve the information structures, verifying the integrity of the data.
+		UModel* Model = Actor->GetBrushComponent()->Brush;
 
- 		if (Model == NULL || Model->VertexBuffer.Vertices.Num() < 3 || Model->MaterialIndexBuffers.Num() == 0) return;
+		if (Model == NULL || Model->VertexBuffer.Vertices.Num() < 3 || Model->MaterialIndexBuffers.Num() == 0) return;
  
- 		// Create the FBX actor, the FBX geometry and instantiate it.
- 		FbxNode* FbxActor = ExportActor( Actor, NULL );
- 		Scene->GetRootNode()->AddChild(FbxActor);
+		// Create the FBX actor, the FBX geometry and instantiate it.
+		FbxNode* FbxActor = ExportActor( Actor, NULL );
+		Scene->GetRootNode()->AddChild(FbxActor);
  
- 		// Export the mesh information
- 		ExportModel(Model, FbxActor, TCHAR_TO_UTF8(*Actor->GetName()));
- 	}
- 	else
- 	{
+		// Export the mesh information
+		ExportModel(Model, FbxActor, TCHAR_TO_UTF8(*Actor->GetName()));
+	}
+	else
+	{
 		FRawMesh Mesh;
 		TArray<UMaterialInterface*>	Materials;
 		GetBrushMesh(Actor,Actor->Brush,Mesh,Materials);
 
- 		if( Mesh.VertexPositions.Num() )
+		if( Mesh.VertexPositions.Num() )
 		{
 			UStaticMesh* StaticMesh = CreateStaticMesh(Mesh,Materials,GetTransientPackage(),Actor->GetFName());
 			ExportStaticMesh( StaticMesh, &Materials );
- 		}
- 	}
+		}
+	}
 }
 
 void FFbxExporter::ExportModel(UModel* Model, FbxNode* Node, const char* Name)
@@ -607,10 +608,10 @@ void FFbxExporter::ExportModel(UModel* Model, FbxNode* Node, const char* Name)
 	Node->SetNodeAttribute(Mesh);
 }
 
-void FFbxExporter::ExportStaticMesh( AActor* Actor, UStaticMeshComponent* StaticMeshComponent, AMatineeActor* InMatineeActor )
+void FFbxExporter::ExportStaticMesh(AActor* Actor, UStaticMeshComponent* StaticMeshComponent, AMatineeActor* InMatineeActor)
 {
 	if (Scene == NULL || Actor == NULL || StaticMeshComponent == NULL)
-	{ 
+	{
 		return;
 	}
 
@@ -621,7 +622,6 @@ void FFbxExporter::ExportStaticMesh( AActor* Actor, UStaticMeshComponent* Static
 		return;
 	}
 	int32 LODIndex = StaticMeshComponent->ForcedLodModel-1;
-	FStaticMeshLODResources& RenderMesh = StaticMesh->GetLODForExport(LODIndex);
 
 	FString FbxNodeName = GetActorNodeName(Actor, InMatineeActor);
 
@@ -632,8 +632,8 @@ void FFbxExporter::ExportStaticMesh( AActor* Actor, UStaticMeshComponent* Static
 		ColorBuffer = StaticMeshComponent->LODData[LODIndex].OverrideVertexColors;
 	}
 
-	FbxNode* FbxActor = ExportActor( Actor, InMatineeActor );
-	ExportStaticMeshToFbx(StaticMesh, RenderMesh, *FbxNodeName, FbxActor, -1, ColorBuffer);
+	FbxNode* FbxActor = ExportActor(Actor, InMatineeActor);
+	ExportStaticMeshToFbx(StaticMesh, LODIndex, *FbxNodeName, FbxActor, -1, ColorBuffer);
 }
 
 struct FBSPExportData
@@ -737,32 +737,33 @@ void FFbxExporter::ExportBSP( UModel* Model, bool bSelectedOnly )
 			const FVector& TextureY = Model->Vectors[Surf.vTextureV];
 			const FVector& Normal = Model->Vectors[Surf.vNormal];
 
-			int32 StartIndex=ExportData->CurrentVertAddIndex;
-	
-			int32 VertexIndex = 0;
-			for( ; VertexIndex < Node.NumVertices ; VertexIndex++ )
-			{					
+			const int32 StartIndex = ExportData->CurrentVertAddIndex;
+
+			for(int32 VertexIndex = 0; VertexIndex < Node.NumVertices ; VertexIndex++ )
+			{
 				const FVert& Vert = Model->Verts[Node.iVertPool + VertexIndex];
 				const FVector& Vertex = Model->Points[Vert.pVertex];
 				Mesh.VertexPositions[ExportData->CurrentVertAddIndex+VertexIndex] = Vertex;
 			}
-			ExportData->CurrentVertAddIndex += VertexIndex;
+			ExportData->CurrentVertAddIndex += Node.NumVertices;
 
-			for( int32 StartVertexIndex = 1 ; StartVertexIndex < Node.NumVertices - 1 ; StartVertexIndex++ )
+			for (int32 StartVertexIndex = 1; StartVertexIndex < Node.NumVertices - 1; ++StartVertexIndex)
 			{
- 				// These map the node's vertices to the 3 triangle indices to triangulate the convex polygon.
- 				int32 TriVertIndices[3] = {	Node.iVertPool + StartVertexIndex + 1,
- 					Node.iVertPool + StartVertexIndex,
- 					Node.iVertPool };
-				
-				int32 WedgeIndices[3] = {	StartIndex + StartVertexIndex + 1 ,
-					 						StartIndex + StartVertexIndex,
-					 						StartIndex };
+				// These map the node's vertices to the 3 triangle indices to triangulate the convex polygon.
+				int32 TriVertIndices[3] = {
+					Node.iVertPool + StartVertexIndex + 1,
+					Node.iVertPool + StartVertexIndex,
+					Node.iVertPool };
+
+				int32 WedgeIndices[3] = {
+					StartIndex + StartVertexIndex + 1,
+					StartIndex + StartVertexIndex,
+					StartIndex };
 
 				Mesh.FaceMaterialIndices[ExportData->CurrentFaceAddIndex] = MaterialIndex;
 				Mesh.FaceSmoothingMasks[ExportData->CurrentFaceAddIndex] =  ( 1 << ( Node.iSurf % 32 ) );
 
-				for(uint32 WedgeIndex = 0; WedgeIndex < 3; WedgeIndex++)
+				for (uint32 WedgeIndex = 0; WedgeIndex < 3; ++WedgeIndex)
 				{
 					const FVert& Vert = Model->Verts[TriVertIndices[WedgeIndex]];
 					const FVector& Vertex = Model->Points[Vert.pVertex];
@@ -801,10 +802,9 @@ void FFbxExporter::ExportStaticMesh( UStaticMesh* StaticMesh, const TArray<UMate
 	if (Scene == NULL || StaticMesh == NULL || !StaticMesh->HasValidRenderData()) return;
 	FString MeshName;
 	StaticMesh->GetName(MeshName);
-	FStaticMeshLODResources& RenderMesh = StaticMesh->GetLODForExport(0);
 	FbxNode* MeshNode = FbxNode::Create(Scene, TCHAR_TO_UTF8(*MeshName));
 	Scene->GetRootNode()->AddChild(MeshNode);
-	ExportStaticMeshToFbx(StaticMesh, RenderMesh, *MeshName, MeshNode, -1, NULL, MaterialOrder );
+	ExportStaticMeshToFbx(StaticMesh, 0, *MeshName, MeshNode, -1, NULL, MaterialOrder);
 }
 
 void FFbxExporter::ExportStaticMeshLightMap( UStaticMesh* StaticMesh, int32 LODIndex, int32 UVChannel )
@@ -813,10 +813,9 @@ void FFbxExporter::ExportStaticMeshLightMap( UStaticMesh* StaticMesh, int32 LODI
 
 	FString MeshName;
 	StaticMesh->GetName(MeshName);
-	FStaticMeshLODResources& RenderMesh = StaticMesh->GetLODForExport(LODIndex);
 	FbxNode* MeshNode = FbxNode::Create(Scene, TCHAR_TO_UTF8(*MeshName));
 	Scene->GetRootNode()->AddChild(MeshNode);
-	ExportStaticMeshToFbx(StaticMesh, RenderMesh, *MeshName, MeshNode, UVChannel);
+	ExportStaticMeshToFbx(StaticMesh, LODIndex, *MeshName, MeshNode, UVChannel);
 }
 
 void FFbxExporter::ExportSkeletalMesh( USkeletalMesh* SkeletalMesh )
@@ -1190,21 +1189,23 @@ FbxNode* FFbxExporter::ExportActor(AActor* Actor, AMatineeActor* InMatineeActor,
 
 				if (StaticMeshComp && StaticMeshComp->StaticMesh)
 				{
-					int32 LODIndex = StaticMeshComp->ForcedLodModel;
-					FStaticMeshLODResources& RenderMesh = StaticMeshComp->StaticMesh->GetLODForExport(LODIndex);
-
 					if (USplineMeshComponent* SplineMeshComp = Cast<USplineMeshComponent>(StaticMeshComp))
 					{
-						ExportSplineMeshToFbx(SplineMeshComp, RenderMesh, *SplineMeshComp->GetName(), ExportNode);
+						ExportSplineMeshToFbx(SplineMeshComp, *SplineMeshComp->GetName(), ExportNode);
+					}
+					else if (UInstancedStaticMeshComponent* InstancedMeshComp = Cast<UInstancedStaticMeshComponent>(StaticMeshComp))
+					{
+						ExportInstancedMeshToFbx(InstancedMeshComp, *InstancedMeshComp->GetName(), ExportNode);
 					}
 					else
 					{
-						ExportStaticMeshToFbx(StaticMeshComp->StaticMesh, RenderMesh, *StaticMeshComp->GetName(), ExportNode);
+						const int32 LODIndex = (StaticMeshComp->ForcedLodModel > 0 ? StaticMeshComp->ForcedLodModel - 1 : /* auto-select*/ 0);
+						ExportStaticMeshToFbx(StaticMeshComp->StaticMesh, LODIndex, *StaticMeshComp->GetName(), ExportNode);
 					}
 				}
 				else if (SkelMeshComp && SkelMeshComp->SkeletalMesh)
 				{
-					ExportSkeletalMeshComponent( SkelMeshComp, *SkelMeshComp->GetName(), ExportNode);
+					ExportSkeletalMeshComponent(SkelMeshComp, *SkelMeshComp->GetName(), ExportNode);
 				}
 				else if (ChildActorComp && ChildActorComp->ChildActor)
 				{
@@ -1755,7 +1756,7 @@ bool FFbxExporter::FindSkeleton(const USkeletalMeshComponent* SkelComp, TArray<F
  * @param VertRemap		Index of each UV (out)
  * @param UniqueVerts	
  */
-void DetermineUVsToWeld(TArray<int32>& VertRemap, TArray<int32>& UniqueVerts, FStaticMeshVertexBuffer& VertexBuffer, int32 TexCoordSourceIndex)
+void DetermineUVsToWeld(TArray<int32>& VertRemap, TArray<int32>& UniqueVerts, const FStaticMeshVertexBuffer& VertexBuffer, int32 TexCoordSourceIndex)
 {
 	const int32 VertexCount = VertexBuffer.GetNumVertices();
 
@@ -1785,7 +1786,7 @@ void DetermineUVsToWeld(TArray<int32>& VertRemap, TArray<int32>& UniqueVerts, FS
 	}
 }
 
-void DetermineVertsToWeld(TArray<int32>& VertRemap, TArray<int32>& UniqueVerts, FStaticMeshLODResources& RenderMesh)
+void DetermineVertsToWeld(TArray<int32>& VertRemap, TArray<int32>& UniqueVerts, const FStaticMeshLODResources& RenderMesh)
 {
 	const int32 VertexCount = RenderMesh.VertexBuffer.GetNumVertices();
 
@@ -1817,285 +1818,300 @@ void DetermineVertsToWeld(TArray<int32>& VertRemap, TArray<int32>& UniqueVerts,
 
 /**
  * Exports a static mesh
+ * @param StaticMesh	The static mesh to export
  * @param RenderMesh	The static mesh render data to export
  * @param MeshName		The name of the mesh for the FBX file
  * @param FbxActor		The fbx node representing the mesh
- * @param LightmapUVChannel Optional UV channel to export
+ * @param ExportLOD		The LOD of the mesh to export
+ * @param LightmapUVChannel If set, performs a "lightmap export" and exports only the single given UV channel
  * @param ColorBuffer	Vertex color overrides to export
  * @param MaterialOrderOverride	Optional ordering of materials to set up correct material ID's across multiple meshes being export such as BSP surfaces which share common materials. Should be used sparingly
  */
-FbxNode* FFbxExporter::ExportStaticMeshToFbx(UStaticMesh* StaticMesh, FStaticMeshLODResources& RenderMesh, const TCHAR* MeshName, FbxNode* FbxActor, int32 LightmapUVChannel, FColorVertexBuffer* ColorBuffer, const TArray<UMaterialInterface*>* MaterialOrderOverride )
+FbxNode* FFbxExporter::ExportStaticMeshToFbx(const UStaticMesh* StaticMesh, int32 ExportLOD, const TCHAR* MeshName, FbxNode* FbxActor, int32 LightmapUVChannel /*= -1*/, const FColorVertexBuffer* ColorBuffer /*= NULL*/, const TArray<UMaterialInterface*>* MaterialOrderOverride /*= NULL*/)
 {
-	// Verify the integrity of the static mesh.
-	if (RenderMesh.VertexBuffer.GetNumVertices() == 0)
+	FbxMesh* Mesh = nullptr;
+	if (ExportLOD == 0 && LightmapUVChannel == -1 && ColorBuffer == nullptr && MaterialOrderOverride == nullptr)
 	{
-		return NULL;
+		Mesh = FbxMeshes.FindRef(StaticMesh);
 	}
 
-	if (RenderMesh.Sections.Num() == 0)
+	if (!Mesh)
 	{
-		return NULL;
-	}
+		Mesh = FbxMesh::Create(Scene, TCHAR_TO_UTF8(MeshName));
 
-	// Remaps an Unreal vert to final reduced vertex list
-	TArray<int32> VertRemap;
-	TArray<int32> UniqueVerts;
+		const FStaticMeshLODResources& RenderMesh = StaticMesh->GetLODForExport(ExportLOD);
 
-	if ( bStaticMeshExportUnWeldedVerts == false )
-	{
-		// Weld verts
-		DetermineVertsToWeld(VertRemap, UniqueVerts, RenderMesh);
-	}
-	else
-	{
-		// Do not weld verts
-		VertRemap.Add(RenderMesh.VertexBuffer.GetNumVertices());
-		for(int32 i=0; i < VertRemap.Num(); i++)
+		// Verify the integrity of the static mesh.
+		if (RenderMesh.VertexBuffer.GetNumVertices() == 0)
 		{
-			VertRemap[i] = i;
+			return nullptr;
 		}
-		UniqueVerts = VertRemap;
-	}
 
-	FbxMesh* Mesh = FbxMesh::Create(Scene, TCHAR_TO_UTF8(MeshName));
-
-	// Create and fill in the vertex position data source.
-	// The position vertices are duplicated, for some reason, retrieve only the first half vertices.
-	const int32 VertexCount = VertRemap.Num();
-	const int32 PolygonsCount = RenderMesh.Sections.Num();
-	
-	Mesh->InitControlPoints(UniqueVerts.Num());
+		if (RenderMesh.Sections.Num() == 0)
+		{
+			return nullptr;
+		}
 
-	FbxVector4* ControlPoints = Mesh->GetControlPoints();
-	for (int32 PosIndex = 0; PosIndex < UniqueVerts.Num(); ++PosIndex)
-	{
-		int32 UnrealPosIndex = UniqueVerts[PosIndex];
-		FVector Position = RenderMesh.PositionVertexBuffer.VertexPosition(UnrealPosIndex);
-		ControlPoints[PosIndex] = FbxVector4(Position.X, -Position.Y, Position.Z);
-	}
-	
-	// Set the normals on Layer 0.
-	FbxLayer* Layer = Mesh->GetLayer(0);
-	if (Layer == NULL)
-	{
-		Mesh->CreateLayer();
-		Layer = Mesh->GetLayer(0);
-	}
+		// Remaps an Unreal vert to final reduced vertex list
+		TArray<int32> VertRemap;
+		TArray<int32> UniqueVerts;
 
-	// Build list of Indices re-used multiple times to lookup Normals, UVs, other per face vertex information
-	TArray<uint32> Indices;
-	for (int32 PolygonsIndex = 0; PolygonsIndex < PolygonsCount; ++PolygonsIndex)
-	{
-		FIndexArrayView RawIndices = RenderMesh.IndexBuffer.GetArrayView();
-		FStaticMeshSection& Polygons = RenderMesh.Sections[PolygonsIndex];
-		const uint32 TriangleCount = Polygons.NumTriangles;
-		for (uint32 TriangleIndex = 0; TriangleIndex < TriangleCount; ++TriangleIndex)
+		if (bStaticMeshExportUnWeldedVerts == false)
 		{
-			for (uint32 PointIndex = 0; PointIndex < 3; PointIndex++)
+			// Weld verts
+			DetermineVertsToWeld(VertRemap, UniqueVerts, RenderMesh);
+		}
+		else
+		{
+			// Do not weld verts
+			VertRemap.Add(RenderMesh.VertexBuffer.GetNumVertices());
+			for (int32 i = 0; i < VertRemap.Num(); i++)
 			{
-				uint32 UnrealVertIndex = RawIndices[Polygons.FirstIndex + ((TriangleIndex * 3) + PointIndex)];
-				Indices.Add(UnrealVertIndex);
+				VertRemap[i] = i;
 			}
+			UniqueVerts = VertRemap;
 		}
-	}
-
-	// Create and fill in the per-face-vertex normal data source.
-	// We extract the Z-tangent and drop the X/Y-tangents which are also stored in the render mesh.
-	FbxLayerElementNormal* LayerElementNormal= FbxLayerElementNormal::Create(Mesh, "");
 
-	// Set 3 normals per triangle instead of storing normals on positional control points
-	LayerElementNormal->SetMappingMode(FbxLayerElement::eByPolygonVertex);
-
-	// Set the normal values for every polygon vertex.
-	LayerElementNormal->SetReferenceMode(FbxLayerElement::eDirect);
+		// Create and fill in the vertex position data source.
+		// The position vertices are duplicated, for some reason, retrieve only the first half vertices.
+		const int32 VertexCount = VertRemap.Num();
+		const int32 PolygonsCount = RenderMesh.Sections.Num();
 
-	TArray<FbxVector4> FbxNormals;
-	FbxNormals.AddUninitialized(VertexCount);
-	for (int32 NormalIndex = 0; NormalIndex < VertexCount; ++NormalIndex)
-	{
-		FVector Normal = (FVector) (RenderMesh.VertexBuffer.VertexTangentZ(NormalIndex));
-		FbxVector4& FbxNormal = FbxNormals[NormalIndex];
-		FbxNormal = FbxVector4(Normal.X, -Normal.Y, Normal.Z);
-		FbxNormal.Normalize();
-	}
+		Mesh->InitControlPoints(UniqueVerts.Num());
 
-	// Add one normal per each face index (3 per triangle)
-	for (int32 i=0; i < Indices.Num(); i++)
-	{
-		uint32 UnrealVertIndex = Indices[i];
-		LayerElementNormal->GetDirectArray().Add( FbxNormals[UnrealVertIndex] );
-	}
-	Layer->SetNormals(LayerElementNormal);
-	FbxNormals.Empty();
-
-	// Create and fill in the per-face-vertex texture coordinate data source(s).
-	// Create UV for Diffuse channel.
-	int32 TexCoordSourceCount = (LightmapUVChannel == -1)? RenderMesh.VertexBuffer.GetNumTexCoords(): LightmapUVChannel + 1;
-	int32 TexCoordSourceIndex = (LightmapUVChannel == -1)? 0: LightmapUVChannel;
-	TCHAR UVChannelName[32] = { 0 };
-	for (; TexCoordSourceIndex < TexCoordSourceCount; ++TexCoordSourceIndex)
-	{
-		FbxLayer* UVsLayer = (LightmapUVChannel == -1)? Mesh->GetLayer(TexCoordSourceIndex): Mesh->GetLayer(0);
-		if (UVsLayer == NULL)
+		FbxVector4* ControlPoints = Mesh->GetControlPoints();
+		for (int32 PosIndex = 0; PosIndex < UniqueVerts.Num(); ++PosIndex)
 		{
-			Mesh->CreateLayer();
-			UVsLayer = (LightmapUVChannel == -1)? Mesh->GetLayer(TexCoordSourceIndex): Mesh->GetLayer(0);
+			int32 UnrealPosIndex = UniqueVerts[PosIndex];
+			FVector Position = RenderMesh.PositionVertexBuffer.VertexPosition(UnrealPosIndex);
+			ControlPoints[PosIndex] = FbxVector4(Position.X, -Position.Y, Position.Z);
 		}
 
-		if ((LightmapUVChannel >= 0) || ((LightmapUVChannel == -1) && (TexCoordSourceIndex == 1)))
+		// Set the normals on Layer 0.
+		FbxLayer* Layer = Mesh->GetLayer(0);
+		if (Layer == nullptr)
 		{
-			FCString::Sprintf(UVChannelName, TEXT("LightMapUV"));
+			Mesh->CreateLayer();
+			Layer = Mesh->GetLayer(0);
 		}
-		else
+
+		// Build list of Indices re-used multiple times to lookup Normals, UVs, other per face vertex information
+		TArray<uint32> Indices;
+		for (int32 PolygonsIndex = 0; PolygonsIndex < PolygonsCount; ++PolygonsIndex)
 		{
-			FCString::Sprintf(UVChannelName, TEXT(""));
+			FIndexArrayView RawIndices = RenderMesh.IndexBuffer.GetArrayView();
+			const FStaticMeshSection& Polygons = RenderMesh.Sections[PolygonsIndex];
+			const uint32 TriangleCount = Polygons.NumTriangles;
+			for (uint32 TriangleIndex = 0; TriangleIndex < TriangleCount; ++TriangleIndex)
+			{
+				for (uint32 PointIndex = 0; PointIndex < 3; PointIndex++)
+				{
+					uint32 UnrealVertIndex = RawIndices[Polygons.FirstIndex + ((TriangleIndex * 3) + PointIndex)];
+					Indices.Add(UnrealVertIndex);
+				}
+			}
 		}
 
-		FbxLayerElementUV* UVDiffuseLayer = FbxLayerElementUV::Create(Mesh, TCHAR_TO_UTF8(UVChannelName));
+		// Create and fill in the per-face-vertex normal data source.
+		// We extract the Z-tangent and drop the X/Y-tangents which are also stored in the render mesh.
+		FbxLayerElementNormal* LayerElementNormal = FbxLayerElementNormal::Create(Mesh, "");
 
-		// Note: when eINDEX_TO_DIRECT is used, IndexArray must be 3xTriangle count, DirectArray can be smaller
-		UVDiffuseLayer->SetMappingMode(FbxLayerElement::eByPolygonVertex);
-		UVDiffuseLayer->SetReferenceMode(FbxLayerElement::eIndexToDirect);
+		// Set 3 normals per triangle instead of storing normals on positional control points
+		LayerElementNormal->SetMappingMode(FbxLayerElement::eByPolygonVertex);
 
-		TArray<int32> UvsRemap;
-		TArray<int32> UniqueUVs;
-		if ( bStaticMeshExportUnWeldedVerts == false )
-		{
-			// Weld UVs
-			DetermineUVsToWeld(UvsRemap, UniqueUVs, RenderMesh.VertexBuffer, TexCoordSourceIndex);
-		}
-		else
+		// Set the normal values for every polygon vertex.
+		LayerElementNormal->SetReferenceMode(FbxLayerElement::eDirect);
+
+		TArray<FbxVector4> FbxNormals;
+		FbxNormals.AddUninitialized(VertexCount);
+		for (int32 NormalIndex = 0; NormalIndex < VertexCount; ++NormalIndex)
 		{
-			// Do not weld UVs
-			UvsRemap = VertRemap;
-			UniqueUVs = UvsRemap;
+			FVector Normal = (FVector)(RenderMesh.VertexBuffer.VertexTangentZ(NormalIndex));
+			FbxVector4& FbxNormal = FbxNormals[NormalIndex];
+			FbxNormal = FbxVector4(Normal.X, -Normal.Y, Normal.Z);
+			FbxNormal.Normalize();
 		}
-		
-		// Create the texture coordinate data source.
-		for (int32 i=0; i < UniqueUVs.Num(); i++)
+
+		// Add one normal per each face index (3 per triangle)
+		for (int32 FbxVertIndex = 0; FbxVertIndex < Indices.Num(); FbxVertIndex++)
 		{
-			int32 UnrealVertIndex = UniqueUVs[i];
-			const FVector2D& TexCoord = RenderMesh.VertexBuffer.GetVertexUV(UnrealVertIndex, TexCoordSourceIndex);
-			UVDiffuseLayer->GetDirectArray().Add(FbxVector2(TexCoord.X, -TexCoord.Y + 1.0));
+			uint32 UnrealVertIndex = Indices[FbxVertIndex];
+			LayerElementNormal->GetDirectArray().Add(FbxNormals[UnrealVertIndex]);
 		}
+		Layer->SetNormals(LayerElementNormal);
+		FbxNormals.Empty();
 
-		// For each face index, point to a texture uv
-		UVDiffuseLayer->GetIndexArray().SetCount( Indices.Num() );
-		for (int32 i=0; i < Indices.Num(); i++)
+		// Create and fill in the per-face-vertex texture coordinate data source(s).
+		// Create UV for Diffuse channel.
+		int32 TexCoordSourceCount = (LightmapUVChannel == -1) ? RenderMesh.VertexBuffer.GetNumTexCoords() : LightmapUVChannel + 1;
+		int32 TexCoordSourceIndex = (LightmapUVChannel == -1) ? 0 : LightmapUVChannel;
+		for (; TexCoordSourceIndex < TexCoordSourceCount; ++TexCoordSourceIndex)
 		{
-			uint32 UnrealVertIndex = Indices[i];
-			int32 NewVertIndex = UvsRemap[UnrealVertIndex];
-			UVDiffuseLayer->GetIndexArray().SetAt(i,NewVertIndex);
-		}
+			FbxLayer* UVsLayer = (LightmapUVChannel == -1) ? Mesh->GetLayer(TexCoordSourceIndex) : Mesh->GetLayer(0);
+			if (UVsLayer == NULL)
+			{
+				Mesh->CreateLayer();
+				UVsLayer = (LightmapUVChannel == -1) ? Mesh->GetLayer(TexCoordSourceIndex) : Mesh->GetLayer(0);
+			}
+			check(UVsLayer);
 
-		UVsLayer->SetUVs(UVDiffuseLayer, FbxLayerElement::eTextureDiffuse);
-	}
-	
-	FbxLayerElementMaterial* MatLayer = FbxLayerElementMaterial::Create(Mesh, "");
-	MatLayer->SetMappingMode(FbxLayerElement::eByPolygon);
-	MatLayer->SetReferenceMode(FbxLayerElement::eIndexToDirect);
-	Layer->SetMaterials(MatLayer);
-	
-	// Keep track of the number of tri's we export
-	uint32 AccountedTriangles = 0;
-	for (int32 PolygonsIndex = 0; PolygonsIndex < PolygonsCount; ++PolygonsIndex)
-	{
-		FStaticMeshSection& Polygons = RenderMesh.Sections[PolygonsIndex];
-		FIndexArrayView RawIndices = RenderMesh.IndexBuffer.GetArrayView();
-		UMaterialInterface* Material = StaticMesh->GetMaterial(Polygons.MaterialIndex);
+			TCHAR* UVChannelName = TEXT("");
+			if ((LightmapUVChannel >= 0) || ((LightmapUVChannel == -1) && (TexCoordSourceIndex == StaticMesh->LightMapCoordinateIndex)))
+			{
+				UVChannelName = TEXT("LightMapUV");
+			}
 
-		FbxSurfaceMaterial* FbxMaterial = Material ? ExportMaterial(Material->GetMaterial()) : NULL;
-		if (!FbxMaterial)
-		{
-			FbxMaterial = CreateDefaultMaterial();
-		}
-		int32 MatIndex = FbxActor->AddMaterial(FbxMaterial);
-		
-		// Determine the actual material index
-		int32 ActualIndex = MatIndex;
+			FbxLayerElementUV* UVDiffuseLayer = FbxLayerElementUV::Create(Mesh, TCHAR_TO_UTF8(UVChannelName));
 
-		if( MaterialOrderOverride )
-		{
-			ActualIndex = MaterialOrderOverride->Find( Material );
+			// Note: when eINDEX_TO_DIRECT is used, IndexArray must be 3xTriangle count, DirectArray can be smaller
+			UVDiffuseLayer->SetMappingMode(FbxLayerElement::eByPolygonVertex);
+			UVDiffuseLayer->SetReferenceMode(FbxLayerElement::eIndexToDirect);
+
+			TArray<int32> UvsRemap;
+			TArray<int32> UniqueUVs;
+			if (bStaticMeshExportUnWeldedVerts == false)
+			{
+				// Weld UVs
+				DetermineUVsToWeld(UvsRemap, UniqueUVs, RenderMesh.VertexBuffer, TexCoordSourceIndex);
+			}
+			else
+			{
+				// Do not weld UVs
+				UvsRemap = VertRemap;
+				UniqueUVs = UvsRemap;
+			}
+
+			// Create the texture coordinate data source.
+			for (int32 FbxVertIndex = 0; FbxVertIndex < UniqueUVs.Num(); FbxVertIndex++)
+			{
+				int32 UnrealVertIndex = UniqueUVs[FbxVertIndex];
+				const FVector2D& TexCoord = RenderMesh.VertexBuffer.GetVertexUV(UnrealVertIndex, TexCoordSourceIndex);
+				UVDiffuseLayer->GetDirectArray().Add(FbxVector2(TexCoord.X, -TexCoord.Y + 1.0));
+			}
+
+			// For each face index, point to a texture uv
+			UVDiffuseLayer->GetIndexArray().SetCount(Indices.Num());
+			for (int32 FbxVertIndex = 0; FbxVertIndex < Indices.Num(); FbxVertIndex++)
+			{
+				uint32 UnrealVertIndex = Indices[FbxVertIndex];
+				int32 NewVertIndex = UvsRemap[UnrealVertIndex];
+				UVDiffuseLayer->GetIndexArray().SetAt(FbxVertIndex, NewVertIndex);
+			}
+
+			UVsLayer->SetUVs(UVDiffuseLayer, FbxLayerElement::eTextureDiffuse);
 		}
-		// Static meshes contain one triangle list per element.
-		// [GLAFORTE] Could it occasionally contain triangle strips? How do I know?
-		uint32 TriangleCount = Polygons.NumTriangles;
-		
-		// Copy over the index buffer into the FBX polygons set.
-		for (uint32 TriangleIndex = 0; TriangleIndex < TriangleCount; ++TriangleIndex)
+
+		FbxLayerElementMaterial* MatLayer = FbxLayerElementMaterial::Create(Mesh, "");
+		MatLayer->SetMappingMode(FbxLayerElement::eByPolygon);
+		MatLayer->SetReferenceMode(FbxLayerElement::eIndexToDirect);
+		Layer->SetMaterials(MatLayer);
+
+		// Keep track of the number of tri's we export
+		uint32 AccountedTriangles = 0;
+		for (int32 PolygonsIndex = 0; PolygonsIndex < PolygonsCount; ++PolygonsIndex)
 		{
-			Mesh->BeginPolygon(ActualIndex);
-			for (uint32 PointIndex = 0; PointIndex < 3; PointIndex++)
+			const FStaticMeshSection& Polygons = RenderMesh.Sections[PolygonsIndex];
+			FIndexArrayView RawIndices = RenderMesh.IndexBuffer.GetArrayView();
+			UMaterialInterface* Material = StaticMesh->GetMaterial(Polygons.MaterialIndex);
+
+			FbxSurfaceMaterial* FbxMaterial = Material ? ExportMaterial(Material->GetMaterial()) : NULL;
+			if (!FbxMaterial)
 			{
-				uint32 OriginalUnrealVertIndex = RawIndices[Polygons.FirstIndex + ((TriangleIndex * 3) + PointIndex)];
-				int32 RemappedVertIndex = VertRemap[OriginalUnrealVertIndex];
-				Mesh->AddPolygon( RemappedVertIndex );
+				FbxMaterial = CreateDefaultMaterial();
 			}
-			Mesh->EndPolygon();
-		}
+			int32 MatIndex = FbxActor->AddMaterial(FbxMaterial);
 
-		AccountedTriangles += TriangleCount;
-	}
+			// Determine the actual material index
+			int32 ActualMatIndex = MatIndex;
+
+			if (MaterialOrderOverride)
+			{
+				ActualMatIndex = MaterialOrderOverride->Find(Material);
+			}
+			// Static meshes contain one triangle list per element.
+			// [GLAFORTE] Could it occasionally contain triangle strips? How do I know?
+			uint32 TriangleCount = Polygons.NumTriangles;
+
+			// Copy over the index buffer into the FBX polygons set.
+			for (uint32 TriangleIndex = 0; TriangleIndex < TriangleCount; ++TriangleIndex)
+			{
+				Mesh->BeginPolygon(ActualMatIndex);
+				for (uint32 PointIndex = 0; PointIndex < 3; PointIndex++)
+				{
+					uint32 OriginalUnrealVertIndex = RawIndices[Polygons.FirstIndex + ((TriangleIndex * 3) + PointIndex)];
+					int32 RemappedVertIndex = VertRemap[OriginalUnrealVertIndex];
+					Mesh->AddPolygon(RemappedVertIndex);
+				}
+				Mesh->EndPolygon();
+			}
+
+			AccountedTriangles += TriangleCount;
+		}
 
 #if TODO_FBX
-	// Throw a warning if this is a lightmap export and the exported poly count does not match the raw triangle data count
-	if (LightmapUVChannel != -1 && AccountedTriangles != RenderMesh.RawTriangles.GetElementCount())
-	{
-		FMessageDialog::Open( EAppMsgType::Ok, NSLOCTEXT("UnrealEd", "StaticMeshEditor_LightmapExportFewerTriangles", "Fewer polygons have been exported than the raw triangle count.  This Lightmapped UV mesh may contain fewer triangles than the destination mesh on import.") );
-	}
+		// Throw a warning if this is a lightmap export and the exported poly count does not match the raw triangle data count
+		if (LightmapUVChannel != -1 && AccountedTriangles != RenderMesh.RawTriangles.GetElementCount())
+		{
+			FMessageDialog::Open( EAppMsgType::Ok, NSLOCTEXT("UnrealEd", "StaticMeshEditor_LightmapExportFewerTriangles", "Fewer polygons have been exported than the raw triangle count.  This Lightmapped UV mesh may contain fewer triangles than the destination mesh on import.") );
+		}
 
-	// Create and fill in the smoothing data source.
-	FbxLayerElementSmoothing* SmoothingInfo = FbxLayerElementSmoothing::Create(Mesh, "");
-	SmoothingInfo->SetMappingMode(FbxLayerElement::eByPolygon);
-	SmoothingInfo->SetReferenceMode(FbxLayerElement::eDirect);
-	FbxLayerElementArrayTemplate<int>& SmoothingArray = SmoothingInfo->GetDirectArray();
-	Layer->SetSmoothing(SmoothingInfo);
+		// Create and fill in the smoothing data source.
+		FbxLayerElementSmoothing* SmoothingInfo = FbxLayerElementSmoothing::Create(Mesh, "");
+		SmoothingInfo->SetMappingMode(FbxLayerElement::eByPolygon);
+		SmoothingInfo->SetReferenceMode(FbxLayerElement::eDirect);
+		FbxLayerElementArrayTemplate<int>& SmoothingArray = SmoothingInfo->GetDirectArray();
+		Layer->SetSmoothing(SmoothingInfo);
 
-	// This is broken. We are exporting the render mesh but providing smoothing
-	// information from the source mesh. The render triangles are not in the
-	// same order. Therefore we should export the raw mesh or not export
-	// smoothing group information!
-	int32 TriangleCount = RenderMesh.RawTriangles.GetElementCount();
-	FStaticMeshTriangle* RawTriangleData = (FStaticMeshTriangle*)RenderMesh.RawTriangles.Lock(LOCK_READ_ONLY);
-	for( int32 TriangleIndex = 0 ; TriangleIndex < TriangleCount ; TriangleIndex++ )
-	{
-		FStaticMeshTriangle* Triangle = (RawTriangleData++);
-		
-		SmoothingArray.Add(Triangle->SmoothingMask);
-	}
-	RenderMesh.RawTriangles.Unlock();
+		// This is broken. We are exporting the render mesh but providing smoothing
+		// information from the source mesh. The render triangles are not in the
+		// same order. Therefore we should export the raw mesh or not export
+		// smoothing group information!
+		int32 TriangleCount = RenderMesh.RawTriangles.GetElementCount();
+		FStaticMeshTriangle* RawTriangleData = (FStaticMeshTriangle*)RenderMesh.RawTriangles.Lock(LOCK_READ_ONLY);
+		for (int32 TriangleIndex = 0; TriangleIndex < TriangleCount; TriangleIndex++)
+		{
+			FStaticMeshTriangle* Triangle = (RawTriangleData++);
+
+			SmoothingArray.Add(Triangle->SmoothingMask);
+		}
+		RenderMesh.RawTriangles.Unlock();
 #endif // #if TODO_FBX
 
-	// Create and fill in the vertex color data source.
-	FColorVertexBuffer* ColorBufferToUse = ColorBuffer? ColorBuffer: &RenderMesh.ColorVertexBuffer;
-	uint32 ColorVertexCount = ColorBufferToUse->GetNumVertices();
-	
-	// Only export vertex colors if they exist
-	if (ColorVertexCount > 0)
-	{
-		FbxLayerElementVertexColor* VertexColor = FbxLayerElementVertexColor::Create(Mesh, "");
-		VertexColor->SetMappingMode(FbxLayerElement::eByPolygonVertex);
-		VertexColor->SetReferenceMode(FbxLayerElement::eIndexToDirect);
-		FbxLayerElementArrayTemplate<FbxColor>& VertexColorArray = VertexColor->GetDirectArray();
-		Layer->SetVertexColors(VertexColor);
+		// Create and fill in the vertex color data source.
+		const FColorVertexBuffer* ColorBufferToUse = ColorBuffer ? ColorBuffer : &RenderMesh.ColorVertexBuffer;
+		uint32 ColorVertexCount = ColorBufferToUse->GetNumVertices();
 
-		for (int32 i=0; i < Indices.Num(); i++)
+		// Only export vertex colors if they exist
+		if (ColorVertexCount > 0)
 		{
-			FLinearColor VertColor(1.0f, 1.0f, 1.0f);
-			uint32 UnrealVertIndex = Indices[i];
-			if (UnrealVertIndex < ColorVertexCount)
+			FbxLayerElementVertexColor* VertexColor = FbxLayerElementVertexColor::Create(Mesh, "");
+			VertexColor->SetMappingMode(FbxLayerElement::eByPolygonVertex);
+			VertexColor->SetReferenceMode(FbxLayerElement::eIndexToDirect);
+			FbxLayerElementArrayTemplate<FbxColor>& VertexColorArray = VertexColor->GetDirectArray();
+			Layer->SetVertexColors(VertexColor);
+
+			for (int32 FbxVertIndex = 0; FbxVertIndex < Indices.Num(); FbxVertIndex++)
 			{
-				VertColor = ColorBufferToUse->VertexColor(UnrealVertIndex).ReinterpretAsLinear();
+				FLinearColor VertColor(1.0f, 1.0f, 1.0f);
+				uint32 UnrealVertIndex = Indices[FbxVertIndex];
+				if (UnrealVertIndex < ColorVertexCount)
+				{
+					VertColor = ColorBufferToUse->VertexColor(UnrealVertIndex).ReinterpretAsLinear();
+				}
+
+				VertexColorArray.Add(FbxColor(VertColor.R, VertColor.G, VertColor.B, VertColor.A));
 			}
 
-			VertexColorArray.Add( FbxColor(VertColor.R, VertColor.G, VertColor.B, VertColor.A ));
+			VertexColor->GetIndexArray().SetCount(Indices.Num());
+			for (int32 FbxVertIndex = 0; FbxVertIndex < Indices.Num(); FbxVertIndex++)
+			{
+				VertexColor->GetIndexArray().SetAt(FbxVertIndex, FbxVertIndex);
+			}
 		}
 
-		VertexColor->GetIndexArray().SetCount( Indices.Num() );
-		for (int32 i=0; i < Indices.Num(); i++)
+		if (ExportLOD == 0 && LightmapUVChannel == -1 && ColorBuffer == nullptr && MaterialOrderOverride == nullptr)
 		{
-			VertexColor->GetIndexArray().SetAt(i,i);
+			FbxMeshes.Add(StaticMesh, Mesh);
 		}
 	}
 
@@ -2104,20 +2120,23 @@ FbxNode* FFbxExporter::ExportStaticMeshToFbx(UStaticMesh* StaticMesh, FStaticMes
 	return FbxActor;
 }
 
-FbxNode* FFbxExporter::ExportSplineMeshToFbx(USplineMeshComponent* SplineMeshComp, FStaticMeshLODResources& RenderMesh, const TCHAR* MeshName, FbxNode* FbxActor)
+void FFbxExporter::ExportSplineMeshToFbx(const USplineMeshComponent* SplineMeshComp, const TCHAR* MeshName, FbxNode* FbxActor)
 {
 	const UStaticMesh* StaticMesh = SplineMeshComp->StaticMesh;
 	check(StaticMesh);
 
+	const int32 LODIndex = (SplineMeshComp->ForcedLodModel > 0 ? SplineMeshComp->ForcedLodModel - 1 : /* auto-select*/ 0);
+	const FStaticMeshLODResources& RenderMesh = StaticMesh->GetLODForExport(LODIndex);
+
 	// Verify the integrity of the static mesh.
 	if (RenderMesh.VertexBuffer.GetNumVertices() == 0)
 	{
-		return NULL;
+		return;
 	}
 
 	if (RenderMesh.Sections.Num() == 0)
 	{
-		return NULL;
+		return;
 	}
 
 	// Remaps an Unreal vert to final reduced vertex list
@@ -2175,7 +2194,7 @@ FbxNode* FFbxExporter::ExportSplineMeshToFbx(USplineMeshComponent* SplineMeshCom
 	for (int32 PolygonsIndex = 0; PolygonsIndex < PolygonsCount; ++PolygonsIndex)
 	{
 		FIndexArrayView RawIndices = RenderMesh.IndexBuffer.GetArrayView();
-		FStaticMeshSection& Polygons = RenderMesh.Sections[PolygonsIndex];
+		const FStaticMeshSection& Polygons = RenderMesh.Sections[PolygonsIndex];
 		const uint32 TriangleCount = Polygons.NumTriangles;
 		for (uint32 TriangleIndex = 0; TriangleIndex < TriangleCount; ++TriangleIndex)
 		{
@@ -2228,10 +2247,10 @@ FbxNode* FFbxExporter::ExportSplineMeshToFbx(USplineMeshComponent* SplineMeshCom
 		if (UVsLayer == NULL)
 		{
 			Mesh->CreateLayer();
-			UVsLayer =  Mesh->GetLayer(TexCoordSourceIndex);
+			UVsLayer = Mesh->GetLayer(TexCoordSourceIndex);
 		}
 
-		if (TexCoordSourceIndex == 1)
+		if (TexCoordSourceIndex == StaticMesh->LightMapCoordinateIndex)
 		{
 			FCString::Sprintf(UVChannelName, TEXT("LightMapUV"));
 		}
@@ -2269,11 +2288,11 @@ FbxNode* FFbxExporter::ExportSplineMeshToFbx(USplineMeshComponent* SplineMeshCom
 
 		// For each face index, point to a texture uv
 		UVDiffuseLayer->GetIndexArray().SetCount(Indices.Num());
-		for (int32 i = 0; i < Indices.Num(); i++)
+		for (int32 FbxVertIndex = 0; FbxVertIndex < Indices.Num(); FbxVertIndex++)
 		{
-			uint32 UnrealVertIndex = Indices[i];
+			uint32 UnrealVertIndex = Indices[FbxVertIndex];
 			int32 NewVertIndex = UvsRemap[UnrealVertIndex];
-			UVDiffuseLayer->GetIndexArray().SetAt(i, NewVertIndex);
+			UVDiffuseLayer->GetIndexArray().SetAt(FbxVertIndex, NewVertIndex);
 		}
 
 		UVsLayer->SetUVs(UVDiffuseLayer, FbxLayerElement::eTextureDiffuse);
@@ -2286,7 +2305,7 @@ FbxNode* FFbxExporter::ExportSplineMeshToFbx(USplineMeshComponent* SplineMeshCom
 
 	for (int32 PolygonsIndex = 0; PolygonsIndex < PolygonsCount; ++PolygonsIndex)
 	{
-		FStaticMeshSection& Polygons = RenderMesh.Sections[PolygonsIndex];
+		const FStaticMeshSection& Polygons = RenderMesh.Sections[PolygonsIndex];
 		FIndexArrayView RawIndices = RenderMesh.IndexBuffer.GetArrayView();
 		UMaterialInterface* Material = StaticMesh->GetMaterial(Polygons.MaterialIndex);
 
@@ -2331,7 +2350,7 @@ FbxNode* FFbxExporter::ExportSplineMeshToFbx(USplineMeshComponent* SplineMeshCom
 #endif // #if TODO_FBX
 
 	// Create and fill in the vertex color data source.
-	FColorVertexBuffer* ColorBufferToUse = &RenderMesh.ColorVertexBuffer;
+	const FColorVertexBuffer* ColorBufferToUse = &RenderMesh.ColorVertexBuffer;
 	uint32 ColorVertexCount = ColorBufferToUse->GetNumVertices();
 
 	// Only export vertex colors if they exist
@@ -2343,10 +2362,10 @@ FbxNode* FFbxExporter::ExportSplineMeshToFbx(USplineMeshComponent* SplineMeshCom
 		FbxLayerElementArrayTemplate<FbxColor>& VertexColorArray = VertexColor->GetDirectArray();
 		Layer->SetVertexColors(VertexColor);
 
-		for (int32 i = 0; i < Indices.Num(); i++)
+		for (int32 FbxVertIndex = 0; FbxVertIndex < Indices.Num(); FbxVertIndex++)
 		{
 			FLinearColor VertColor(1.0f, 1.0f, 1.0f);
-			uint32 UnrealVertIndex = Indices[i];
+			uint32 UnrealVertIndex = Indices[FbxVertIndex];
 			if (UnrealVertIndex < ColorVertexCount)
 			{
 				VertColor = ColorBufferToUse->VertexColor(UnrealVertIndex).ReinterpretAsLinear();
@@ -2356,21 +2375,44 @@ FbxNode* FFbxExporter::ExportSplineMeshToFbx(USplineMeshComponent* SplineMeshCom
 		}
 
 		VertexColor->GetIndexArray().SetCount(Indices.Num());
-		for (int32 i = 0; i < Indices.Num(); i++)
+		for (int32 FbxVertIndex = 0; FbxVertIndex < Indices.Num(); FbxVertIndex++)
 		{
-			VertexColor->GetIndexArray().SetAt(i, i);
+			VertexColor->GetIndexArray().SetAt(FbxVertIndex, FbxVertIndex);
 		}
 	}
 
 	FbxActor->SetNodeAttribute(Mesh);
+}
 
-	return FbxActor;
+void FFbxExporter::ExportInstancedMeshToFbx(const UInstancedStaticMeshComponent* InstancedMeshComp, const TCHAR* MeshName, FbxNode* FbxActor)
+{
+	const UStaticMesh* StaticMesh = InstancedMeshComp->StaticMesh;
+	check(StaticMesh);
+
+	const int32 LODIndex = (InstancedMeshComp->ForcedLodModel > 0 ? InstancedMeshComp->ForcedLodModel - 1 : /* auto-select*/ 0);
+	const int32 NumInstances = InstancedMeshComp->GetInstanceCount();
+	for (int32 InstanceIndex = 0; InstanceIndex < NumInstances; ++InstanceIndex)
+	{
+		FTransform RelativeTransform;
+		if (ensure(InstancedMeshComp->GetInstanceTransform(InstanceIndex, RelativeTransform, /*bWorldSpace=*/false)))
+		{
+			FbxNode* InstNode = FbxNode::Create(Scene, TCHAR_TO_UTF8(*FString::Printf(TEXT("%d"), InstanceIndex)));
+
+			InstNode->LclTranslation.Set(Converter.ConvertToFbxPos(RelativeTransform.GetTranslation()));
+			InstNode->LclRotation.Set(Converter.ConvertToFbxRot(RelativeTransform.GetRotation().Euler()));
+			InstNode->LclScaling.Set(Converter.ConvertToFbxScale(RelativeTransform.GetScale3D()));
+
+			// Todo - export once and then clone the node
+			ExportStaticMeshToFbx(StaticMesh, LODIndex, *FString::Printf(TEXT("%d"), InstanceIndex), InstNode);
+			FbxActor->AddChild(InstNode);
+		}
+	}
 }
 
 /**
  * Exports a Landscape
  */
-FbxNode* FFbxExporter::ExportLandscapeToFbx(ALandscapeProxy* Landscape, const TCHAR* MeshName, FbxNode* FbxActor, bool bSelectedOnly)
+void FFbxExporter::ExportLandscapeToFbx(ALandscapeProxy* Landscape, const TCHAR* MeshName, FbxNode* FbxActor, bool bSelectedOnly)
 {
 	const ULandscapeInfo* LandscapeInfo = Landscape->GetLandscapeInfo(false);
 	
@@ -2575,8 +2617,6 @@ FbxNode* FFbxExporter::ExportLandscapeToFbx(ALandscapeProxy* Landscape, const TC
 	}
 
 	FbxActor->SetNodeAttribute(Mesh);
-
-	return FbxActor;
 }
 
 
diff --git a/Engine/Source/Editor/UnrealEd/Private/FbxExporter.h b/Engine/Source/Editor/UnrealEd/Private/FbxExporter.h
--- a/Engine/Source/Editor/UnrealEd/Private/FbxExporter.h
+++ b/Engine/Source/Editor/UnrealEd/Private/FbxExporter.h
@@ -37,6 +37,7 @@
 class ALandscapeProxy;
 struct FAnimControlTrackKey;
 class USplineMeshComponent;
+class UInstancedStaticMeshComponent;
 class ALight;
 class UInterpTrackInstMove;
 class AMatineeActor;
@@ -178,10 +179,11 @@ private:
 	FFbxDataConverter Converter;
 	
 	TMap<FString,int32> FbxNodeNameToIndexMap;
-	TMap<AActor*, FbxNode*> FbxActors;
-	TMap<USkeletalMeshComponent*, FbxNode*> FbxSkeletonRoots;
-	TMap<UMaterial*, FbxSurfaceMaterial*> FbxMaterials;
-	
+	TMap<const AActor*, FbxNode*> FbxActors;
+	TMap<const USkeletalMeshComponent*, FbxNode*> FbxSkeletonRoots;
+	TMap<const UMaterial*, FbxSurfaceMaterial*> FbxMaterials;
+	TMap<const UStaticMesh*, FbxMesh*> FbxMeshes;
+
 	/** The frames-per-second (FPS) used when baking transforms */
 	static const float BakeTransformsFPS;
 	
@@ -199,23 +201,32 @@ private:
 	/**
 	 * Exports a static mesh
 	 * @param StaticMesh	The static mesh to export
-	 * @param RenderMesh	The static mesh render data to export
 	 * @param MeshName		The name of the mesh for the FBX file
 	 * @param FbxActor		The fbx node representing the mesh
+	 * @param ExportLOD		The LOD of the mesh to export
 	 * @param LightmapUVChannel Optional UV channel to export
 	 * @param ColorBuffer	Vertex color overrides to export
 	 * @param MaterialOrderOverride	Optional ordering of materials to set up correct material ID's across multiple meshes being export such as BSP surfaces which share common materials. Should be used sparingly
 	 */
-	FbxNode* ExportStaticMeshToFbx(UStaticMesh* StaticMesh, FStaticMeshLODResources& RenderMesh, const TCHAR* MeshName, FbxNode* FbxActor, int32 LightmapUVChannel = -1, FColorVertexBuffer* ColorBuffer = NULL, const TArray<UMaterialInterface*>* MaterialOrderOverride = NULL );
+	FbxNode* ExportStaticMeshToFbx(const UStaticMesh* StaticMesh, int32 ExportLOD, const TCHAR* MeshName, FbxNode* FbxActor, int32 LightmapUVChannel = -1, const FColorVertexBuffer* ColorBuffer = NULL, const TArray<UMaterialInterface*>* MaterialOrderOverride = NULL);
 
 	/**
-	* Exports a spline mesh
-	* @param SplineMeshComp	The spline mesh component to export
-	* @param RenderMesh		The mesh render data to export
-	* @param MeshName		The name of the mesh for the FBX file
-	* @param FbxActor		The fbx node representing the mesh
-	*/
-	FbxNode* ExportSplineMeshToFbx(USplineMeshComponent* SplineMeshComp, FStaticMeshLODResources& RenderMesh, const TCHAR* MeshName, FbxNode* FbxActor);
+	 * Exports a spline mesh
+	 * @param SplineMeshComp	The spline mesh component to export
+	 * @param RenderMesh		The mesh render data to export
+	 * @param MeshName		The name of the mesh for the FBX file
+	 * @param FbxActor		The fbx node representing the mesh
+	 */
+	void ExportSplineMeshToFbx(const USplineMeshComponent* SplineMeshComp, const TCHAR* MeshName, FbxNode* FbxActor);
+
+	/**
+	 * Exports an instanced mesh
+	 * @param InstancedMeshComp	The instanced mesh component to export
+	 * @param RenderMesh		The mesh render data to export
+	 * @param MeshName		The name of the mesh for the FBX file
+	 * @param FbxActor		The fbx node representing the mesh
+	 */
+	void ExportInstancedMeshToFbx(const UInstancedStaticMeshComponent* InstancedMeshComp, const TCHAR* MeshName, FbxNode* FbxActor);
 
 	/**
 	* Exports a landscape
@@ -223,7 +234,7 @@ private:
 	* @param MeshName		The name of the mesh for the FBX file
 	* @param FbxActor		The fbx node representing the mesh
 	*/
-	FbxNode* ExportLandscapeToFbx(ALandscapeProxy* Landscape, const TCHAR* MeshName, FbxNode* FbxActor, bool bSelectedOnly);
+	void ExportLandscapeToFbx(ALandscapeProxy* Landscape, const TCHAR* MeshName, FbxNode* FbxActor, bool bSelectedOnly);
 
 	/**
 	 * Adds FBX skeleton nodes to the FbxScene based on the skeleton in the given USkeletalMesh, and fills
diff --git a/Engine/Source/Editor/UnrealEd/Private/UnrealEdSrv.cpp b/Engine/Source/Editor/UnrealEd/Private/UnrealEdSrv.cpp
--- a/Engine/Source/Editor/UnrealEd/Private/UnrealEdSrv.cpp
+++ b/Engine/Source/Editor/UnrealEd/Private/UnrealEdSrv.cpp
@@ -1821,11 +1821,11 @@ TArray<FPoly*> GetSelectedPolygons()
 					int32 NumLods = StaticMesh->GetNumLODs();
 					if ( NumLods )
 					{
-						FStaticMeshLODResources& MeshLodZero = StaticMesh->GetLODForExport(0);
+						const FStaticMeshLODResources& MeshLodZero = StaticMesh->GetLODForExport(0);
 						int32 NumTriangles = MeshLodZero.GetNumTriangles();
 						int32 NumVertices = MeshLodZero.GetNumVertices();
 			
-						FPositionVertexBuffer& PositionVertexBuffer = MeshLodZero.PositionVertexBuffer;
+						const FPositionVertexBuffer& PositionVertexBuffer = MeshLodZero.PositionVertexBuffer;
 						FIndexArrayView Indices = MeshLodZero.DepthOnlyIndexBuffer.GetArrayView();
 
 						for ( int32 TriangleIndex = 0; TriangleIndex < NumTriangles; TriangleIndex++ )
diff --git a/Engine/Source/Runtime/Engine/Classes/Engine/StaticMesh.h b/Engine/Source/Runtime/Engine/Classes/Engine/StaticMesh.h
--- a/Engine/Source/Runtime/Engine/Classes/Engine/StaticMesh.h
+++ b/Engine/Source/Runtime/Engine/Classes/Engine/StaticMesh.h
@@ -10,6 +10,7 @@
 
 // Forward declarations
 class UFoliageType_InstancedStaticMesh;
+struct FStaticMeshLODResources;
 
 /*-----------------------------------------------------------------------------
 	Legacy mesh optimization settings.
@@ -512,7 +513,7 @@ public:
 	 * Returns the render data to use for exporting the specified LOD. This method should always
 	 * be called when exporting a static mesh.
 	 */
-	ENGINE_API struct FStaticMeshLODResources& GetLODForExport( int32 LODIndex );
+	ENGINE_API const FStaticMeshLODResources& GetLODForExport(int32 LODIndex) const;
 
 	/**
 	 * Static: Processes the specified static mesh for light map UV problems
diff --git a/Engine/Source/Runtime/Engine/Private/DestructibleMesh.cpp b/Engine/Source/Runtime/Engine/Private/DestructibleMesh.cpp
--- a/Engine/Source/Runtime/Engine/Private/DestructibleMesh.cpp
+++ b/Engine/Source/Runtime/Engine/Private/DestructibleMesh.cpp
@@ -487,7 +487,7 @@ void UDestructibleMesh::CreateFractureSettings()
 
 #if WITH_APEX && WITH_EDITORONLY_DATA
 
-bool CreateSubmeshFromSMSection(FStaticMeshLODResources& RenderMesh, int32 SubmeshIdx, FStaticMeshSection& Section, physx::NxExplicitSubmeshData& SubmeshData, TArray<physx::NxExplicitRenderTriangle>& Triangles)
+bool CreateSubmeshFromSMSection(const FStaticMeshLODResources& RenderMesh, int32 SubmeshIdx, const FStaticMeshSection& Section, physx::NxExplicitSubmeshData& SubmeshData, TArray<physx::NxExplicitRenderTriangle>& Triangles)
 {
 	// Create submesh descriptor, just a material name and a vertex format
 	FCStringAnsi::Strncpy(SubmeshData.mMaterialName, TCHAR_TO_ANSI(*FString::Printf(TEXT("Material%d"),Section.MaterialIndex)), physx::NxExplicitSubmeshData::MaterialNameBufferSize);
@@ -547,12 +547,12 @@ bool UDestructibleMesh::BuildFractureSettingsFromStaticMesh(UStaticMesh* StaticM
 
 	// Array of render meshes to create the fracture settings from. The first RenderMesh is the root mesh, while the other ones are 
 	// imported to form the chunk meshes ( if no FBX was imported as level 1 chunks, this will just contain the root mesh )
-	TArray<FStaticMeshLODResources*> RenderMeshes;
+	TArray<const FStaticMeshLODResources*> RenderMeshes;
 	TArray<UStaticMesh*> StaticMeshes;
 	TArray<uint32> MeshPartitions;
 
 	// Do we have a main static mesh?
-	FStaticMeshLODResources& MainRenderMesh = StaticMesh->GetLODForExport(0);
+	const FStaticMeshLODResources& MainRenderMesh = StaticMesh->GetLODForExport(0);
 
 	// Keep track of overall triangle and submesh count
 	int32 OverallTriangleCount = MainRenderMesh.GetNumTriangles();
@@ -568,7 +568,7 @@ bool UDestructibleMesh::BuildFractureSettingsFromStaticMesh(UStaticMesh* StaticM
 	// Get the rendermeshes for the chunks ( if any )
 	for (int32 ChunkMeshIdx=0; ChunkMeshIdx < FractureChunkMeshes.Num(); ++ChunkMeshIdx)
 	{
-		FStaticMeshLODResources& ChunkRM = FractureChunkMeshes[ChunkMeshIdx]->GetLODForExport(0);
+		const FStaticMeshLODResources& ChunkRM = FractureChunkMeshes[ChunkMeshIdx]->GetLODForExport(0);
 		
 		RenderMeshes.Add(&ChunkRM);
 		StaticMeshes.Add(FractureChunkMeshes[ChunkMeshIdx]);
@@ -596,7 +596,7 @@ bool UDestructibleMesh::BuildFractureSettingsFromStaticMesh(UStaticMesh* StaticM
 	// Start creating apex data
 	for (int32 MeshIdx=0; MeshIdx < RenderMeshes.Num(); ++MeshIdx)
 	{
-		FStaticMeshLODResources& RenderMesh = *RenderMeshes[MeshIdx];
+		const FStaticMeshLODResources& RenderMesh = *RenderMeshes[MeshIdx];
 
 		// Get the current static mesh to retrieve the material from
 		UStaticMesh* CurrentStaticMesh = StaticMeshes[MeshIdx];
@@ -604,7 +604,7 @@ bool UDestructibleMesh::BuildFractureSettingsFromStaticMesh(UStaticMesh* StaticM
 		// Loop through elements
 		for (int32 SectionIndex = 0; SectionIndex < RenderMesh.Sections.Num(); ++SectionIndex, ++SubmeshIndexCounter)
 		{
-			FStaticMeshSection& Section = RenderMesh.Sections[SectionIndex];
+			const FStaticMeshSection& Section = RenderMesh.Sections[SectionIndex];
 			physx::NxExplicitSubmeshData& SubmeshData = Submeshes[SubmeshIndexCounter];
 
 			// Parallel materials array
diff --git a/Engine/Source/Runtime/Engine/Private/StaticMesh.cpp b/Engine/Source/Runtime/Engine/Private/StaticMesh.cpp
--- a/Engine/Source/Runtime/Engine/Private/StaticMesh.cpp
+++ b/Engine/Source/Runtime/Engine/Private/StaticMesh.cpp
@@ -2522,7 +2522,7 @@ UMaterialInterface* UStaticMesh::GetMaterial(int32 MaterialIndex) const
  * Returns the render data to use for exporting the specified LOD. This method should always
  * be called when exporting a static mesh.
  */
-FStaticMeshLODResources& UStaticMesh::GetLODForExport( int32 LODIndex )
+const FStaticMeshLODResources& UStaticMesh::GetLODForExport(int32 LODIndex) const
 {
 	check(RenderData);
 	LODIndex = FMath::Clamp<int32>( LODIndex, 0, RenderData->LODResources.Num()-1 );
