Fix compile error on clang, and removed some pointless UTF8 <-> TCHAR conversions

diff --git a/Engine/Source/Editor/UnrealEd/Private/Fbx/FbxMainExport.cpp b/Engine/Source/Editor/UnrealEd/Private/Fbx/FbxMainExport.cpp
--- a/Engine/Source/Editor/UnrealEd/Private/Fbx/FbxMainExport.cpp
+++ b/Engine/Source/Editor/UnrealEd/Private/Fbx/FbxMainExport.cpp
@@ -1108,7 +1108,7 @@ FbxNode* FFbxExporter::ExportActor(AActor* Actor, AMatineeActor* InMatineeActor,
 
 			if (FbxPivotNodeName == FbxNodeName)
 			{
-				FbxPivotNodeName += UTF8_TO_TCHAR("_pivot");
+				FbxPivotNodeName += TEXT("_pivot");
 			}
 
 			FbxNode* PivotNode = FbxNode::Create(Scene, TCHAR_TO_UTF8(*FbxPivotNodeName));
@@ -1955,13 +1955,13 @@ FbxNode* FFbxExporter::ExportStaticMeshToFbx(const UStaticMesh* StaticMesh, int3
 			}
 			check(UVsLayer);
 
-			TCHAR* UVChannelName = TEXT("");
+			const char* UVChannelName = ""; // actually UTF8 as required by Fbx, but can't use UE4's UTF8CHAR type because that's a uint8 aka *unsigned* char
 			if ((LightmapUVChannel >= 0) || ((LightmapUVChannel == -1) && (TexCoordSourceIndex == StaticMesh->LightMapCoordinateIndex)))
 			{
-				UVChannelName = TEXT("LightMapUV");
+				UVChannelName = "LightMapUV";
 			}
 
-			FbxLayerElementUV* UVDiffuseLayer = FbxLayerElementUV::Create(Mesh, TCHAR_TO_UTF8(UVChannelName));
+			FbxLayerElementUV* UVDiffuseLayer = FbxLayerElementUV::Create(Mesh, UVChannelName);
 
 			// Note: when eINDEX_TO_DIRECT is used, IndexArray must be 3xTriangle count, DirectArray can be smaller
 			UVDiffuseLayer->SetMappingMode(FbxLayerElement::eByPolygonVertex);
@@ -2240,7 +2240,6 @@ void FFbxExporter::ExportSplineMeshToFbx(const USplineMeshComponent* SplineMeshC
 	// Create and fill in the per-face-vertex texture coordinate data source(s).
 	// Create UV for Diffuse channel.
 	int32 TexCoordSourceCount = RenderMesh.VertexBuffer.GetNumTexCoords();
-	TCHAR UVChannelName[32] = { 0 };
 	for (int32 TexCoordSourceIndex = 0; TexCoordSourceIndex < TexCoordSourceCount; ++TexCoordSourceIndex)
 	{
 		FbxLayer* UVsLayer = Mesh->GetLayer(TexCoordSourceIndex);
@@ -2250,16 +2249,13 @@ void FFbxExporter::ExportSplineMeshToFbx(const USplineMeshComponent* SplineMeshC
 			UVsLayer = Mesh->GetLayer(TexCoordSourceIndex);
 		}
 
+		const char* UVChannelName = ""; // actually UTF8 as required by Fbx, but can't use UE4's UTF8CHAR type because that's a uint8 aka *unsigned* char
 		if (TexCoordSourceIndex == StaticMesh->LightMapCoordinateIndex)
 		{
-			FCString::Sprintf(UVChannelName, TEXT("LightMapUV"));
-		}
-		else
-		{
-			FCString::Sprintf(UVChannelName, TEXT(""));
+			UVChannelName = "LightMapUV";
 		}
 
-		FbxLayerElementUV* UVDiffuseLayer = FbxLayerElementUV::Create(Mesh, TCHAR_TO_UTF8(UVChannelName));
+		FbxLayerElementUV* UVDiffuseLayer = FbxLayerElementUV::Create(Mesh, UVChannelName);
 
 		// Note: when eINDEX_TO_DIRECT is used, IndexArray must be 3xTriangle count, DirectArray can be smaller
 		UVDiffuseLayer->SetMappingMode(FbxLayerElement::eByPolygonVertex);
