Added code to ensure that AssetImportData is valid in the editor, and added null checks for code that can be executed outside of the editor.

This further addresses UE-16277.

diff --git a/Engine/Plugins/2D/Paper2D/Source/Paper2D/Classes/PaperTileMap.h b/Engine/Plugins/2D/Paper2D/Source/Paper2D/Classes/PaperTileMap.h
--- a/Engine/Plugins/2D/Paper2D/Source/Paper2D/Classes/PaperTileMap.h
+++ b/Engine/Plugins/2D/Paper2D/Source/Paper2D/Classes/PaperTileMap.h
@@ -135,6 +135,7 @@ public:
 #endif
 #if WITH_EDITORONLY_DATA
 	virtual void GetAssetRegistryTags(TArray<FAssetRegistryTag>& OutTags) const override;
+	virtual void Serialize(FArchive& Ar) override;
 #endif
 	// End of UObject interface
 
diff --git a/Engine/Plugins/2D/Paper2D/Source/Paper2D/Private/PaperTileMap.cpp b/Engine/Plugins/2D/Paper2D/Source/Paper2D/Private/PaperTileMap.cpp
--- a/Engine/Plugins/2D/Paper2D/Source/Paper2D/Private/PaperTileMap.cpp
+++ b/Engine/Plugins/2D/Paper2D/Source/Paper2D/Private/PaperTileMap.cpp
@@ -29,6 +29,7 @@ UPaperTileMap::UPaperTileMap(const FObjectInitializer& ObjectInitializer)
 #if WITH_EDITORONLY_DATA
 	SelectedLayerIndex = INDEX_NONE;
 	BackgroundColor = FColor(55, 55, 55);
+	AssetImportData = CreateEditorOnlyDefaultSubobject<UAssetImportData>(TEXT("AssetImportData"));
 #endif
 
 	LayerNameIndex = 0;
@@ -175,6 +176,19 @@ void UPaperTileMap::PostLoad()
 	ValidateSelectedLayerIndex();
 }
 
+#if WITH_EDITORONLY_DATA
+void UPaperTileMap::Serialize(FArchive& Ar)
+{
+	Super::Serialize(Ar);
+	
+	if ( Ar.IsLoading() && Ar.UE4Ver() < VER_UE4_ASSET_IMPORT_DATA_AS_JSON && !AssetImportData)
+	{
+		// AssetImportData should always be valid
+		AssetImportData = NewObject<UAssetImportData>(this, TEXT("AssetImportData"));
+	}
+}
+#endif
+
 void UPaperTileMap::ValidateSelectedLayerIndex()
 {
 	if (!TileLayers.IsValidIndex(SelectedLayerIndex))
diff --git a/Engine/Plugins/2D/Paper2D/Source/PaperSpriteSheetImporter/Private/PaperSpriteSheet.cpp b/Engine/Plugins/2D/Paper2D/Source/PaperSpriteSheetImporter/Private/PaperSpriteSheet.cpp
--- a/Engine/Plugins/2D/Paper2D/Source/PaperSpriteSheetImporter/Private/PaperSpriteSheet.cpp
+++ b/Engine/Plugins/2D/Paper2D/Source/PaperSpriteSheetImporter/Private/PaperSpriteSheet.cpp
@@ -6,6 +6,9 @@
 UPaperSpriteSheet::UPaperSpriteSheet(const FObjectInitializer& ObjectInitializer)
 	: Super(ObjectInitializer)
 {
+#if WITH_EDITORONLY_DATA
+	AssetImportData = CreateEditorOnlyDefaultSubobject<UAssetImportData>(TEXT("AssetImportData"));
+#endif
 }
 
 #if WITH_EDITORONLY_DATA
@@ -18,4 +21,14 @@ void UPaperSpriteSheet::GetAssetRegistryTags(TArray<FAssetRegistryTag>& OutTags)
 
 	Super::GetAssetRegistryTags(OutTags);
 }
+void UPaperSpriteSheet::Serialize(FArchive& Ar)
+{
+	Super::Serialize(Ar);
+
+	if (Ar.IsLoading() && Ar.UE4Ver() < VER_UE4_ASSET_IMPORT_DATA_AS_JSON && !AssetImportData)
+	{
+		// AssetImportData should always be valid
+		AssetImportData = NewObject<UAssetImportData>(this, TEXT("AssetImportData"));
+	}
+}
 #endif
\ No newline at end of file
diff --git a/Engine/Plugins/2D/Paper2D/Source/PaperSpriteSheetImporter/Private/PaperSpriteSheet.h b/Engine/Plugins/2D/Paper2D/Source/PaperSpriteSheetImporter/Private/PaperSpriteSheet.h
--- a/Engine/Plugins/2D/Paper2D/Source/PaperSpriteSheetImporter/Private/PaperSpriteSheet.h
+++ b/Engine/Plugins/2D/Paper2D/Source/PaperSpriteSheetImporter/Private/PaperSpriteSheet.h
@@ -40,6 +40,7 @@ public:
 
 	// UObject interface
 	virtual void GetAssetRegistryTags(TArray<FAssetRegistryTag>& OutTags) const override;
+	virtual void Serialize(FArchive& Ar) override;
 	// End of UObject interface
 #endif
 };
diff --git a/Engine/Plugins/2D/Paper2D/Source/PaperSpriteSheetImporter/Private/PaperSpriteSheetImportFactory.cpp b/Engine/Plugins/2D/Paper2D/Source/PaperSpriteSheetImporter/Private/PaperSpriteSheetImportFactory.cpp
--- a/Engine/Plugins/2D/Paper2D/Source/PaperSpriteSheetImporter/Private/PaperSpriteSheetImportFactory.cpp
+++ b/Engine/Plugins/2D/Paper2D/Source/PaperSpriteSheetImporter/Private/PaperSpriteSheetImportFactory.cpp
@@ -60,11 +60,6 @@ UObject* UPaperSpriteSheetImportFactory::FactoryCreateText(UClass* InClass, UObj
 		UPaperSpriteSheet* SpriteSheet = NewObject<UPaperSpriteSheet>(InParent, InName, Flags);
 		if (Importer.PerformImport(LongPackagePath, Flags, SpriteSheet))
 		{
-			if (SpriteSheet->AssetImportData == nullptr)
-			{
-				SpriteSheet->AssetImportData = NewObject<UAssetImportData>(SpriteSheet);
-			}
-
 			SpriteSheet->AssetImportData->Update(CurrentFilename);
 			SpriteSheet->AssetImportData->bDirty = false;
 
diff --git a/Engine/Plugins/Editor/SpeedTreeImporter/Source/SpeedTreeImporter/Private/SpeedTreeImportFactory.cpp b/Engine/Plugins/Editor/SpeedTreeImporter/Source/SpeedTreeImporter/Private/SpeedTreeImportFactory.cpp
--- a/Engine/Plugins/Editor/SpeedTreeImporter/Source/SpeedTreeImporter/Private/SpeedTreeImportFactory.cpp
+++ b/Engine/Plugins/Editor/SpeedTreeImporter/Source/SpeedTreeImporter/Private/SpeedTreeImportFactory.cpp
@@ -1233,7 +1233,6 @@ UObject* USpeedTreeImportFactory::FactoryCreateBinary( UClass* InClass, UObject*
 				StaticMesh = NewObject<UStaticMesh>(Package, FName(*MeshName), Flags | RF_Public);
 
 				// @todo AssetImportData make a data class for speed tree assets
-				StaticMesh->AssetImportData = NewObject<UAssetImportData>(StaticMesh);
 				StaticMesh->AssetImportData->Update(UFactory::GetCurrentFilename());
 				StaticMesh->AssetImportData->bDirty = false;
 				
diff --git a/Engine/Plugins/ScriptPlugin/Source/ScriptPlugin/Private/ScriptBlueprint.cpp b/Engine/Plugins/ScriptPlugin/Source/ScriptPlugin/Private/ScriptBlueprint.cpp
--- a/Engine/Plugins/ScriptPlugin/Source/ScriptPlugin/Private/ScriptBlueprint.cpp
+++ b/Engine/Plugins/ScriptPlugin/Source/ScriptPlugin/Private/ScriptBlueprint.cpp
@@ -19,14 +19,17 @@ UScriptBlueprint::UScriptBlueprint(const FObjectInitializer& ObjectInitializer)
 #if WITH_EDITORONLY_DATA
 void UScriptBlueprint::GetAssetRegistryTags(TArray<FAssetRegistryTag>& OutTags) const
 {
-	OutTags.Add( FAssetRegistryTag(SourceFileTagName(), AssetImportData->ToJson(), FAssetRegistryTag::TT_Hidden) );
+	if (AssetImportData)
+	{
+		OutTags.Add( FAssetRegistryTag(SourceFileTagName(), AssetImportData->ToJson(), FAssetRegistryTag::TT_Hidden) );
+	}
 
 	Super::GetAssetRegistryTags(OutTags);
 }
 void UScriptBlueprint::PostLoad()
 {
 	Super::PostLoad();
-	if (!SourceFilePath_DEPRECATED.IsEmpty())
+	if (!SourceFilePath_DEPRECATED.IsEmpty() && AssetImportData)
 	{
 		AssetImportData->UpdateFilenameOnly(SourceFilePath_DEPRECATED);
 	}
@@ -59,6 +62,11 @@ bool UScriptBlueprint::ValidateGeneratedClass(const UClass* InClass)
 
 bool UScriptBlueprint::IsCodeDirty() const
 {
+	if (!AssetImportData)
+	{
+		return true;
+	}
+
 	const TArray<FAssetImportInfo::FSourceFile>& Data = AssetImportData->GetSourceFileData();
 
 	if (Data.Num() == 1)
@@ -83,7 +91,7 @@ void UScriptBlueprint::UpdateScriptStatus()
 
 void UScriptBlueprint::UpdateSourceCodeIfChanged()
 {
-	if (IsCodeDirty())
+	if (IsCodeDirty() && AssetImportData)
 	{
 		FString NewScript;
 		FString Filename = AssetImportData->GetFirstFilename();
diff --git a/Engine/Source/Developer/AssetTools/Private/AssetTypeActions/AssetTypeActions_AnimSequence.cpp b/Engine/Source/Developer/AssetTools/Private/AssetTypeActions/AssetTypeActions_AnimSequence.cpp
--- a/Engine/Source/Developer/AssetTools/Private/AssetTypeActions/AssetTypeActions_AnimSequence.cpp
+++ b/Engine/Source/Developer/AssetTools/Private/AssetTypeActions/AssetTypeActions_AnimSequence.cpp
@@ -56,10 +56,7 @@ void FAssetTypeActions_AnimSequence::GetResolvedSourceFilePaths(const TArray<UOb
 	for (auto& Asset : TypeAssets)
 	{
 		const auto AnimSequence = CastChecked<UAnimSequence>(Asset);
-		if (AnimSequence->AssetImportData)
-		{
-			AnimSequence->AssetImportData->ExtractFilenames(OutSourceFilePaths);
-		}
+		AnimSequence->AssetImportData->ExtractFilenames(OutSourceFilePaths);
 	}
 }
 
@@ -74,18 +71,15 @@ void FAssetTypeActions_AnimSequence::ExecuteReimportWithNewSource(TArray<TWeakOb
 		{
 			FAssetImportInfo OldImportData;
 
-			if (Object->AssetImportData)
-			{
-				// Make note of the old import data
-				OldImportData.CopyFrom(*Object->AssetImportData);
-				// And reset the import data
-				Object->AssetImportData->CopyFrom(EmptyImportInfo);
-			}
+			// Make note of the old import data
+			OldImportData.CopyFrom(*Object->AssetImportData);
+			// And reset the import data
+			Object->AssetImportData->CopyFrom(EmptyImportInfo);
 
 			bool bSuccess = FReimportManager::Instance()->Reimport(Object, /*bAskForNewFileIfMissing=*/true);
 
 			// restore the old source path in case reimport was not successful
-			if (!bSuccess && Object->AssetImportData)
+			if (!bSuccess)
 			{
 				Object->AssetImportData->CopyFrom(OldImportData);
 			}
diff --git a/Engine/Source/Developer/AssetTools/Private/AssetTypeActions/AssetTypeActions_SkeletalMesh.cpp b/Engine/Source/Developer/AssetTools/Private/AssetTypeActions/AssetTypeActions_SkeletalMesh.cpp
--- a/Engine/Source/Developer/AssetTools/Private/AssetTypeActions/AssetTypeActions_SkeletalMesh.cpp
+++ b/Engine/Source/Developer/AssetTools/Private/AssetTypeActions/AssetTypeActions_SkeletalMesh.cpp
@@ -551,10 +551,7 @@ void FAssetTypeActions_SkeletalMesh::GetResolvedSourceFilePaths(const TArray<UOb
 	for (auto& Asset : TypeAssets)
 	{
 		const auto SkeletalMesh = CastChecked<USkeletalMesh>(Asset);
-		if (SkeletalMesh->AssetImportData)
-		{
-			SkeletalMesh->AssetImportData->ExtractFilenames(OutSourceFilePaths);
-		}
+		SkeletalMesh->AssetImportData->ExtractFilenames(OutSourceFilePaths);
 	}
 }
 
diff --git a/Engine/Source/Developer/AssetTools/Private/AssetTypeActions/AssetTypeActions_StaticMesh.cpp b/Engine/Source/Developer/AssetTools/Private/AssetTypeActions/AssetTypeActions_StaticMesh.cpp
--- a/Engine/Source/Developer/AssetTools/Private/AssetTypeActions/AssetTypeActions_StaticMesh.cpp
+++ b/Engine/Source/Developer/AssetTools/Private/AssetTypeActions/AssetTypeActions_StaticMesh.cpp
@@ -92,10 +92,7 @@ void FAssetTypeActions_StaticMesh::GetResolvedSourceFilePaths(const TArray<UObje
 	for (auto& Asset : TypeAssets)
 	{
 		const auto StaticMesh = CastChecked<UStaticMesh>(Asset);
-		if (StaticMesh->AssetImportData)
-		{
-			StaticMesh->AssetImportData->ExtractFilenames(OutSourceFilePaths);
-		}
+		StaticMesh->AssetImportData->ExtractFilenames(OutSourceFilePaths);
 	}
 }
 
diff --git a/Engine/Source/Editor/UnrealEd/Private/Factories/EditorFactories.cpp b/Engine/Source/Editor/UnrealEd/Private/Factories/EditorFactories.cpp
--- a/Engine/Source/Editor/UnrealEd/Private/Factories/EditorFactories.cpp
+++ b/Engine/Source/Editor/UnrealEd/Private/Factories/EditorFactories.cpp
@@ -5800,14 +5800,7 @@ bool UReimportFbxSkeletalMeshFactory::CanReimport( UObject* Obj, TArray<FString>
 	USkeletalMesh* SkeletalMesh = Cast<USkeletalMesh>(Obj);
 	if(SkeletalMesh && !Obj->IsA( UDestructibleMesh::StaticClass() ))
 	{
-		if ( SkeletalMesh->AssetImportData )
-		{
-			SkeletalMesh->AssetImportData->ExtractFilenames(OutFilenames);
-		}
-		else
-		{
-			OutFilenames.Add(TEXT(""));
-		}
+		SkeletalMesh->AssetImportData->ExtractFilenames(OutFilenames);
 		return true;
 	}
 	return false;
@@ -5979,14 +5972,7 @@ bool UReimportFbxAnimSequenceFactory::CanReimport( UObject* Obj, TArray<FString>
 	UAnimSequence* AnimSequence = Cast<UAnimSequence>(Obj);
 	if(AnimSequence)
 	{
-		if ( AnimSequence->AssetImportData )
-		{
-			AnimSequence->AssetImportData->ExtractFilenames(OutFilenames);
-		}
-		else
-		{
-			OutFilenames.Add(TEXT(""));
-		}
+		AnimSequence->AssetImportData->ExtractFilenames(OutFilenames);
 		return true;
 	}
 	return false;
@@ -6784,12 +6770,6 @@ void UReimportDestructibleMeshFactory::SetReimportPaths( UObject* Obj, const TAr
 	UDestructibleMesh* DestructibleMesh = Cast<UDestructibleMesh>(Obj);
 	if(DestructibleMesh && ensure(NewReimportPaths.Num() == 1))
 	{
-		if ( DestructibleMesh->AssetImportData == nullptr )
-		{
-			// @todo AssetImportData make an apex destructible import data class
-			DestructibleMesh->AssetImportData = NewObject<UAssetImportData>(DestructibleMesh);
-		}
-
 		DestructibleMesh->AssetImportData->UpdateFilenameOnly(NewReimportPaths[0]);
 	}
 }
@@ -6804,12 +6784,6 @@ EReimportResult::Type UReimportDestructibleMeshFactory::Reimport( UObject* Obj )
 
 	UDestructibleMesh* DestructibleMesh = Cast<UDestructibleMesh>( Obj );
 
-	if ( DestructibleMesh->AssetImportData == nullptr )
-	{
-		// @todo AssetImportData make an apex destructible import data class
-		DestructibleMesh->AssetImportData = NewObject<UAssetImportData>(DestructibleMesh);
-	}
-
 	const FString Filename = DestructibleMesh->AssetImportData->GetFirstFilename();
 
 	// If there is no file path provided, can't reimport from source
diff --git a/Engine/Source/Runtime/Engine/Private/Animation/AnimSequence.cpp b/Engine/Source/Runtime/Engine/Private/Animation/AnimSequence.cpp
--- a/Engine/Source/Runtime/Engine/Private/Animation/AnimSequence.cpp
+++ b/Engine/Source/Runtime/Engine/Private/Animation/AnimSequence.cpp
@@ -114,6 +114,9 @@ UAnimSequence::UAnimSequence(const FObjectInitializer& ObjectInitializer)
 {
 
 	RateScale = 1.0;
+#if WITH_EDITORONLY_DATA
+	AssetImportData = CreateEditorOnlyDefaultSubobject<UAssetImportData>(TEXT("AssetImportData"));
+#endif
 }
 
 SIZE_T UAnimSequence::GetResourceSize(EResourceSizeMode::Type Mode)
@@ -278,14 +281,16 @@ void UAnimSequence::Serialize(FArchive& Ar)
 	}
 
 #if WITH_EDITORONLY_DATA
+	if ( Ar.IsLoading() && Ar.UE4Ver() < VER_UE4_ASSET_IMPORT_DATA_AS_JSON && !AssetImportData)
+	{
+		// AssetImportData should always be valid
+		AssetImportData = NewObject<UAssetImportData>(this, TEXT("AssetImportData"));
+	}
+
 	// SourceFilePath and SourceFileTimestamp were moved into a subobject
-	if ( Ar.IsLoading() && Ar.UE4Ver() < VER_UE4_ADDED_FBX_ASSET_IMPORT_DATA )
+	if ( Ar.IsLoading() && Ar.UE4Ver() < VER_UE4_ADDED_FBX_ASSET_IMPORT_DATA && AssetImportData)
 	{
-		if ( AssetImportData == NULL )
-		{
-			AssetImportData = CreateEditorOnlyDefaultSubobject<UAssetImportData>(TEXT("AssetImportData"));
-		}
-		
+		// AssetImportData should always have been set up in the constructor where this is relevant
 		FAssetImportInfo Info;
 		Info.Insert(FAssetImportInfo::FSourceFile(SourceFilePath_DEPRECATED));
 		AssetImportData->CopyFrom(Info);
@@ -293,6 +298,7 @@ void UAnimSequence::Serialize(FArchive& Ar)
 		SourceFilePath_DEPRECATED = TEXT("");
 		SourceFileTimestamp_DEPRECATED = TEXT("");
 	}
+
 #endif // WITH_EDITORONLY_DATA
 }
 
diff --git a/Engine/Source/Runtime/Engine/Private/CurveTable.cpp b/Engine/Source/Runtime/Engine/Private/CurveTable.cpp
--- a/Engine/Source/Runtime/Engine/Private/CurveTable.cpp
+++ b/Engine/Source/Runtime/Engine/Private/CurveTable.cpp
@@ -99,7 +99,10 @@ void UCurveTable::FinishDestroy()
 #if WITH_EDITORONLY_DATA
 void UCurveTable::GetAssetRegistryTags(TArray<FAssetRegistryTag>& OutTags) const
 {
-	OutTags.Add( FAssetRegistryTag(SourceFileTagName(), AssetImportData->ToJson(), FAssetRegistryTag::TT_Hidden) );
+	if (AssetImportData)
+	{
+		OutTags.Add( FAssetRegistryTag(SourceFileTagName(), AssetImportData->ToJson(), FAssetRegistryTag::TT_Hidden) );
+	}
 
 	Super::GetAssetRegistryTags(OutTags);
 }
diff --git a/Engine/Source/Runtime/Engine/Private/DataTable.cpp b/Engine/Source/Runtime/Engine/Private/DataTable.cpp
--- a/Engine/Source/Runtime/Engine/Private/DataTable.cpp
+++ b/Engine/Source/Runtime/Engine/Private/DataTable.cpp
@@ -163,7 +163,10 @@ void UDataTable::FinishDestroy()
 #if WITH_EDITORONLY_DATA
 void UDataTable::GetAssetRegistryTags(TArray<FAssetRegistryTag>& OutTags) const
 {
-	OutTags.Add( FAssetRegistryTag(SourceFileTagName(), AssetImportData->ToJson(), FAssetRegistryTag::TT_Hidden) );
+	if (AssetImportData)
+	{
+		OutTags.Add( FAssetRegistryTag(SourceFileTagName(), AssetImportData->ToJson(), FAssetRegistryTag::TT_Hidden) );
+	}
 
 	Super::GetAssetRegistryTags(OutTags);
 }
diff --git a/Engine/Source/Runtime/Engine/Private/SkeletalMesh.cpp b/Engine/Source/Runtime/Engine/Private/SkeletalMesh.cpp
--- a/Engine/Source/Runtime/Engine/Private/SkeletalMesh.cpp
+++ b/Engine/Source/Runtime/Engine/Private/SkeletalMesh.cpp
@@ -1917,6 +1917,7 @@ USkeletalMesh::USkeletalMesh(const FObjectInitializer& ObjectInitializer)
 	StreamingDistanceMultiplier = 1.0f;
 #if WITH_EDITORONLY_DATA
 	SelectedEditorSection = INDEX_NONE;
+	AssetImportData = CreateEditorOnlyDefaultSubobject<UAssetImportData>(TEXT("AssetImportData"));
 #endif
 	ImportedResource = MakeShareable(new FSkeletalMeshResource());
 }
@@ -2342,14 +2343,16 @@ void USkeletalMesh::Serialize( FArchive& Ar )
 	}
 
 #if WITH_EDITORONLY_DATA
+	if ( Ar.IsLoading() && Ar.UE4Ver() < VER_UE4_ASSET_IMPORT_DATA_AS_JSON && !AssetImportData)
+	{
+		// AssetImportData should always be valid
+		AssetImportData = NewObject<UAssetImportData>(this, TEXT("AssetImportData"));
+	}
+	
 	// SourceFilePath and SourceFileTimestamp were moved into a subobject
-	if ( Ar.IsLoading() && Ar.UE4Ver() < VER_UE4_ADDED_FBX_ASSET_IMPORT_DATA )
+	if ( Ar.IsLoading() && Ar.UE4Ver() < VER_UE4_ADDED_FBX_ASSET_IMPORT_DATA && AssetImportData)
 	{
-		if ( AssetImportData == NULL )
-		{
-			AssetImportData = NewObject<UAssetImportData>(this);
-		}
-
+		// AssetImportData should always have been set up in the constructor where this is relevant
 		FAssetImportInfo Info;
 		Info.Insert(FAssetImportInfo::FSourceFile(SourceFilePath_DEPRECATED));
 		AssetImportData->CopyFrom(Info);
diff --git a/Engine/Source/Runtime/Engine/Private/SoundWave.cpp b/Engine/Source/Runtime/Engine/Private/SoundWave.cpp
--- a/Engine/Source/Runtime/Engine/Private/SoundWave.cpp
+++ b/Engine/Source/Runtime/Engine/Private/SoundWave.cpp
@@ -135,7 +135,10 @@ void USoundWave::GetAssetRegistryTags(TArray<FAssetRegistryTag>& OutTags) const
 	Super::GetAssetRegistryTags(OutTags);
 	
 #if WITH_EDITORONLY_DATA
-	OutTags.Add( FAssetRegistryTag(SourceFileTagName(), AssetImportData->ToJson(), FAssetRegistryTag::TT_Hidden) );
+	if (AssetImportData)
+	{
+		OutTags.Add( FAssetRegistryTag(SourceFileTagName(), AssetImportData->ToJson(), FAssetRegistryTag::TT_Hidden) );
+	}
 #endif
 	// GetCompressedDataSize could technically modify this->CompressedFormatData therefore it is not const, however this information
 	// is very useful in the asset registry so we will allow GetCompressedDataSize to be modified if the formats do not exist
diff --git a/Engine/Source/Runtime/Engine/Private/StaticMesh.cpp b/Engine/Source/Runtime/Engine/Private/StaticMesh.cpp
--- a/Engine/Source/Runtime/Engine/Private/StaticMesh.cpp
+++ b/Engine/Source/Runtime/Engine/Private/StaticMesh.cpp
@@ -1233,6 +1233,7 @@ UStaticMesh::UStaticMesh(const FObjectInitializer& ObjectInitializer)
 #if WITH_EDITORONLY_DATA
 	AutoLODPixelError = 1.0f;
 	bAutoComputeLODScreenSize=true;
+	AssetImportData = CreateEditorOnlyDefaultSubobject<UAssetImportData>(TEXT("AssetImportData"));
 #endif // #if WITH_EDITORONLY_DATA
 	LightMapResolution = 4;
 	LpvBiasMultiplier = 1.0f;
@@ -1873,14 +1874,16 @@ void UStaticMesh::Serialize(FArchive& Ar)
 	}
 
 #if WITH_EDITORONLY_DATA
+	if ( Ar.IsLoading() && Ar.UE4Ver() < VER_UE4_ASSET_IMPORT_DATA_AS_JSON && !AssetImportData)
+	{
+		// AssetImportData should always be valid
+		AssetImportData = NewObject<UAssetImportData>(this, TEXT("AssetImportData"));
+	}
+	
 	// SourceFilePath and SourceFileTimestamp were moved into a subobject
-	if ( Ar.IsLoading() && Ar.UE4Ver() < VER_UE4_ADDED_FBX_ASSET_IMPORT_DATA )
+	if ( Ar.IsLoading() && Ar.UE4Ver() < VER_UE4_ADDED_FBX_ASSET_IMPORT_DATA && AssetImportData )
 	{
-		if ( AssetImportData == NULL )
-		{
-			AssetImportData = NewObject<UAssetImportData>(this);
-		}
-
+		// AssetImportData should always have been set up in the constructor where this is relevant
 		FAssetImportInfo Info;
 		Info.Insert(FAssetImportInfo::FSourceFile(SourceFilePath_DEPRECATED));
 		AssetImportData->CopyFrom(Info);
diff --git a/Engine/Source/Runtime/Engine/Private/Texture.cpp b/Engine/Source/Runtime/Engine/Private/Texture.cpp
--- a/Engine/Source/Runtime/Engine/Private/Texture.cpp
+++ b/Engine/Source/Runtime/Engine/Private/Texture.cpp
@@ -391,7 +391,10 @@ void UTexture::PreSave()
 #if WITH_EDITORONLY_DATA
 void UTexture::GetAssetRegistryTags(TArray<FAssetRegistryTag>& OutTags) const
 {
-	OutTags.Add( FAssetRegistryTag(SourceFileTagName(), AssetImportData->ToJson(), FAssetRegistryTag::TT_Hidden) );
+	if (AssetImportData)
+	{
+		OutTags.Add( FAssetRegistryTag(SourceFileTagName(), AssetImportData->ToJson(), FAssetRegistryTag::TT_Hidden) );
+	}
 
 	Super::GetAssetRegistryTags(OutTags);
 }
diff --git a/Engine/Source/Runtime/Engine/Private/VectorField.cpp b/Engine/Source/Runtime/Engine/Private/VectorField.cpp
--- a/Engine/Source/Runtime/Engine/Private/VectorField.cpp
+++ b/Engine/Source/Runtime/Engine/Private/VectorField.cpp
@@ -352,7 +352,10 @@ void UVectorFieldStatic::PostEditChangeProperty(FPropertyChangedEvent& PropertyC
 #if WITH_EDITORONLY_DATA
 void UVectorFieldStatic::GetAssetRegistryTags(TArray<FAssetRegistryTag>& OutTags) const
 {
-	OutTags.Add( FAssetRegistryTag(SourceFileTagName(), AssetImportData->ToJson(), FAssetRegistryTag::TT_Hidden) );
+	if (AssetImportData)
+	{
+		OutTags.Add( FAssetRegistryTag(SourceFileTagName(), AssetImportData->ToJson(), FAssetRegistryTag::TT_Hidden) );
+	}
 
 	Super::GetAssetRegistryTags(OutTags);
 }
