Removed experimental setting for large fbx file import.  The improvements are always enabled now
Removed the limitation for how many materials a mesh can have.

diff --git a/Engine/Source/Developer/MeshUtilities/Private/MeshUtilities.cpp b/Engine/Source/Developer/MeshUtilities/Private/MeshUtilities.cpp
--- a/Engine/Source/Developer/MeshUtilities/Private/MeshUtilities.cpp
+++ b/Engine/Source/Developer/MeshUtilities/Private/MeshUtilities.cpp
@@ -2933,7 +2933,6 @@ public:
 			{
 				MaxMaterialIndex = FMath::Max<int32>(RawMesh.FaceMaterialIndices[FaceIndex], MaxMaterialIndex);
 			}
-			MaxMaterialIndex = FMath::Min(MaxMaterialIndex, MAX_MESH_MATERIAL_INDEX);
 
 			while (MaxMaterialIndex >= LODModel.Sections.Num())
 			{
diff --git a/Engine/Source/Developer/MeshUtilities/Public/MeshUtilities.h b/Engine/Source/Developer/MeshUtilities/Public/MeshUtilities.h
--- a/Engine/Source/Developer/MeshUtilities/Public/MeshUtilities.h
+++ b/Engine/Source/Developer/MeshUtilities/Public/MeshUtilities.h
@@ -4,9 +4,6 @@
 
 #include "ModuleInterface.h"
 
-/** The maximum index for a material in a mesh */
-#define MAX_MESH_MATERIAL_INDEX 256
-
 /**
  * Mesh reduction interface.
  */
diff --git a/Engine/Source/Editor/UnrealEd/Classes/Settings/EditorExperimentalSettings.h b/Engine/Source/Editor/UnrealEd/Classes/Settings/EditorExperimentalSettings.h
--- a/Engine/Source/Editor/UnrealEd/Classes/Settings/EditorExperimentalSettings.h
+++ b/Engine/Source/Editor/UnrealEd/Classes/Settings/EditorExperimentalSettings.h
@@ -93,9 +93,6 @@ public:
 	UPROPERTY(EditAnywhere, config, Category = Rendering, meta = (DisplayName = "Enable Metal/High-end mobile rendering preview"))
 	bool bFeatureLevelES31Preview;
 
-	UPROPERTY(EditAnywhere, config, Category = Import)
-	bool bEnableLargeFileImportImprovements;
-
 	/**
 	 * Returns an event delegate that is executed when a setting has changed.
 	 *
diff --git a/Engine/Source/Editor/UnrealEd/Private/Fbx/FbxStaticMeshImport.cpp b/Engine/Source/Editor/UnrealEd/Private/Fbx/FbxStaticMeshImport.cpp
--- a/Engine/Source/Editor/UnrealEd/Private/Fbx/FbxStaticMeshImport.cpp
+++ b/Engine/Source/Editor/UnrealEd/Private/Fbx/FbxStaticMeshImport.cpp
@@ -51,6 +51,8 @@
 
 #define LOCTEXT_NAMESPACE "FbxStaticMeshImport"
 
+static const int32 LARGE_MESH_MATERIAL_INDEX_THRESHOLD = 64;
+
 using namespace UnFbx;
 
 struct ExistingStaticMeshData;
@@ -93,7 +95,7 @@ UStaticMesh* UnFbx::FFbxImporter::ImportStaticMesh(UObject* InParent, FbxNode* N
 	return ImportStaticMeshAsSingle(InParent, MeshNodeArray, Name, Flags, ImportData, InStaticMesh, LODIndex);
 }
 
-bool UnFbx::FFbxImporter::BuildStaticMeshFromGeometry(FbxMesh* Mesh, UStaticMesh* StaticMesh, TArray<FFbxMaterial>& MeshMaterials, int LODIndex,FRawMesh* InRawMesh,
+bool UnFbx::FFbxImporter::BuildStaticMeshFromGeometry(FbxMesh* Mesh, UStaticMesh* StaticMesh, TArray<FFbxMaterial>& MeshMaterials, int LODIndex,FRawMesh& RawMesh,
 													  EVertexColorImportOption::Type VertexColorImportOption, const TMap<FVector, FColor>& ExistingVertexColorData, const FColor& VertexOverrideColor)
 {
 	check(StaticMesh->SourceModels.IsValidIndex(LODIndex));
@@ -440,23 +442,6 @@ bool UnFbx::FFbxImporter::BuildStaticMeshFromGeometry(FbxMesh* Mesh, UStaticMesh
 	int32 WedgeCount = TriangleCount * 3;
 	bool OddNegativeScale = IsOddNegativeScale(TotalMatrix);
 
-	FRawMesh* RawMeshPtr;
-
-	TScopedPointer<FRawMesh> NewMesh;
-	if( !InRawMesh )
-	{
-		// Load the existing raw mesh.
-		NewMesh = new FRawMesh;
-		SrcModel.RawMeshBulkData->LoadRawMesh(*NewMesh);
-		RawMeshPtr = NewMesh.GetOwnedPointer();
-	}
-	else
-	{
-		RawMeshPtr = InRawMesh;
-	}
-
-	FRawMesh& RawMesh = *RawMeshPtr;
-
 	int32 VertexOffset = RawMesh.VertexPositions.Num();
 	int32 WedgeOffset = RawMesh.WedgeIndices.Num();
 	int32 TriangleOffset = RawMesh.FaceMaterialIndices.Num();
@@ -706,13 +691,6 @@ bool UnFbx::FFbxImporter::BuildStaticMeshFromGeometry(FbxMesh* Mesh, UStaticMesh
 		RawMesh.FaceMaterialIndices[DestTriangleIndex] = MaterialIndex;
 	}
 
-
-	// Store the new raw mesh if needed
-	if( !InRawMesh )
-	{
-		SrcModel.RawMeshBulkData->SaveRawMesh(RawMesh);
-	}
-	
 	//
 	// clean up.  This needs to happen before the mesh is destroyed
 	//
@@ -1004,14 +982,9 @@ UStaticMesh* UnFbx::FFbxImporter::ImportStaticMeshAsSingle(UObject* InParent, TA
 	StaticMesh->LightMapResolution = 64;
 	StaticMesh->LightMapCoordinateIndex = 1;
 
-	// Note: This is only used if large file improvements are enabled.
-	FRawMesh RawMesh;
 
-	const bool bEnableLargeFileImportImprovements = GetDefault<UEditorExperimentalSettings>()->bEnableLargeFileImportImprovements;
-	if( bEnableLargeFileImportImprovements )
-	{
-		SrcModel.RawMeshBulkData->LoadRawMesh(RawMesh);
-	}
+	FRawMesh RawMesh;
+	SrcModel.RawMeshBulkData->LoadRawMesh(RawMesh);
 
 	TArray<FFbxMaterial> MeshMaterials;
 	for (MeshIndex = 0; MeshIndex < MeshNodeArray.Num(); MeshIndex++ )
@@ -1020,7 +993,7 @@ UStaticMesh* UnFbx::FFbxImporter::ImportStaticMeshAsSingle(UObject* InParent, TA
 
 		if (Node->GetMesh())
 		{
-			if (!BuildStaticMeshFromGeometry(Node->GetMesh(), StaticMesh, MeshMaterials, LODIndex, bEnableLargeFileImportImprovements ? &RawMesh : nullptr,
+			if (!BuildStaticMeshFromGeometry(Node->GetMesh(), StaticMesh, MeshMaterials, LODIndex, RawMesh,
 											 VertexColorImportOption, ExistingVertexColorData, ImportOptions->VertexOverrideColor))
 			{
 				bBuildStatus = false;
@@ -1029,11 +1002,9 @@ UStaticMesh* UnFbx::FFbxImporter::ImportStaticMeshAsSingle(UObject* InParent, TA
 		}
 	}
 
-	if( bEnableLargeFileImportImprovements )
-	{
-		// Store the new raw mesh.
-		SrcModel.RawMeshBulkData->SaveRawMesh(RawMesh);
-	}
+	// Store the new raw mesh.
+	SrcModel.RawMeshBulkData->SaveRawMesh(RawMesh);
+
 
 	if (bBuildStatus)
 	{
@@ -1056,11 +1027,6 @@ UStaticMesh* UnFbx::FFbxImporter::ImportStaticMeshAsSingle(UObject* InParent, TA
 				{
 					int32 UniqueIndex = MaterialMap[OtherMaterialIndex];
 
-					if( UniqueIndex > MAX_MESH_MATERIAL_INDEX )
-					{
-						UniqueIndex = MAX_MESH_MATERIAL_INDEX;
-					}
-
 					MaterialMap.Add(UniqueIndex);
 					bDoRemap = true;
 					bUnique = false;
@@ -1071,11 +1037,6 @@ UStaticMesh* UnFbx::FFbxImporter::ImportStaticMeshAsSingle(UObject* InParent, TA
 			{
 				int32 UniqueIndex = UniqueMaterials.Add(MeshMaterials[MaterialIndex]);
 
-				if (UniqueIndex > MAX_MESH_MATERIAL_INDEX)
-				{
-					UniqueIndex = MAX_MESH_MATERIAL_INDEX;
-				}
-
 				MaterialMap.Add( UniqueIndex );
 			}
 			else
@@ -1084,13 +1045,12 @@ UStaticMesh* UnFbx::FFbxImporter::ImportStaticMeshAsSingle(UObject* InParent, TA
 			}
 		}
 
-		if (UniqueMaterials.Num() > MAX_MESH_MATERIAL_INDEX)
+		if (UniqueMaterials.Num() > LARGE_MESH_MATERIAL_INDEX_THRESHOLD)
 		{
 			AddTokenizedErrorMessage(
 				FTokenizedMessage::Create(
 				EMessageSeverity::Warning,
-				FText::Format(LOCTEXT("Error_TooManyMaterials", "StaticMesh has too many({1}) materials. Clamping materials to {0} which may produce unexpected results. Break apart your mesh into multiple pieces to fix this."),
-				FText::AsNumber(MAX_MESH_MATERIAL_INDEX),
+				FText::Format(LOCTEXT("Error_TooManyMaterials", "StaticMesh has a large number({1}) of materials and may render inefficently.  Consider breaking up the mesh into multiple Static Mesh Assets"),
 				FText::AsNumber(UniqueMaterials.Num())
 				)), 
 				FFbxErrors::StaticMesh_TooManyMaterials);
diff --git a/Engine/Source/Editor/UnrealEd/Private/FbxImporter.h b/Engine/Source/Editor/UnrealEd/Private/FbxImporter.h
--- a/Engine/Source/Editor/UnrealEd/Private/FbxImporter.h
+++ b/Engine/Source/Editor/UnrealEd/Private/FbxImporter.h
@@ -694,7 +694,7 @@ protected:
 	 * @param LODIndex	LOD level to set up for StaticMesh
 	 * @return bool true if set up successfully
 	 */
-	bool BuildStaticMeshFromGeometry(FbxMesh* Mesh, UStaticMesh* StaticMesh, TArray<FFbxMaterial>& MeshMaterials, int LODIndex,FRawMesh* InRawMesh,
+	bool BuildStaticMeshFromGeometry(FbxMesh* Mesh, UStaticMesh* StaticMesh, TArray<FFbxMaterial>& MeshMaterials, int LODIndex,FRawMesh& RawMesh,
 									 EVertexColorImportOption::Type VertexColorImportOption, const TMap<FVector, FColor>& ExistingVertexColorData, const FColor& VertexOverrideColor);
 	
 	/**
