fixed clang warnings

diff --git a/Engine/Source/Editor/GraphEditor/Private/SGraphNode.cpp b/Engine/Source/Editor/GraphEditor/Private/SGraphNode.cpp
--- a/Engine/Source/Editor/GraphEditor/Private/SGraphNode.cpp
+++ b/Engine/Source/Editor/GraphEditor/Private/SGraphNode.cpp
@@ -180,9 +180,9 @@ void SGraphNode::OnCommentTextCommitted(const FText& NewComment, ETextCommit::Ty
 	GetNodeObj()->OnUpdateCommentText(NewComment.ToString());
 }
 
-void SGraphNode::OnCommentBubbleToggled(bool bCommentBubbleVisible)
+void SGraphNode::OnCommentBubbleToggled(bool bInCommentBubbleVisible)
 {
-	GetNodeObj()->OnCommentBubbleToggled(bCommentBubbleVisible);
+	GetNodeObj()->OnCommentBubbleToggled(bInCommentBubbleVisible);
 }
 
 void SGraphNode::SetDisallowedPinConnectionEvent(SGraphEditor::FOnDisallowedPinConnection InOnDisallowedPinConnection)
diff --git a/Engine/Source/Editor/GraphEditor/Public/SGraphNode.h b/Engine/Source/Editor/GraphEditor/Public/SGraphNode.h
--- a/Engine/Source/Editor/GraphEditor/Public/SGraphNode.h
+++ b/Engine/Source/Editor/GraphEditor/Public/SGraphNode.h
@@ -186,7 +186,7 @@ public:
 	/** called to replace this nodes comment text */
 	void OnCommentTextCommitted(const FText& NewComment, ETextCommit::Type CommitInfo);
 	/** called when the node's comment bubble is toggled */
-	void OnCommentBubbleToggled(bool bCommentBubbleVisible);
+	void OnCommentBubbleToggled(bool bInCommentBubbleVisible);
 	/** returns true if a rename is pending on this node */
 	bool IsRenamePending() const { return bRenameIsPending; }
 
diff --git a/Engine/Source/Editor/UnrealEd/Classes/MaterialGraph/MaterialGraphNode.h b/Engine/Source/Editor/UnrealEd/Classes/MaterialGraph/MaterialGraphNode.h
--- a/Engine/Source/Editor/UnrealEd/Classes/MaterialGraph/MaterialGraphNode.h
+++ b/Engine/Source/Editor/UnrealEd/Classes/MaterialGraph/MaterialGraphNode.h
@@ -60,7 +60,7 @@ public:
 	virtual void NodeConnectionListChanged() override;
 	virtual void OnRenameNode(const FString& NewName) override;
 	virtual void OnUpdateCommentText( const FString& NewComment ) override;
-	virtual void OnCommentBubbleToggled( bool bCommentBubbleVisible ) override;
+	virtual void OnCommentBubbleToggled( bool bInCommentBubbleVisible ) override;
 	virtual void GetPinHoverText(const UEdGraphPin& Pin, FString& HoverTextOut) const override;
 	virtual FString GetDocumentationExcerptName() const override;
 	// End UEdGraphNode interface.
diff --git a/Engine/Source/Editor/UnrealEd/Private/MaterialGraphNode.cpp b/Engine/Source/Editor/UnrealEd/Private/MaterialGraphNode.cpp
--- a/Engine/Source/Editor/UnrealEd/Private/MaterialGraphNode.cpp
+++ b/Engine/Source/Editor/UnrealEd/Private/MaterialGraphNode.cpp
@@ -633,12 +633,12 @@ void UMaterialGraphNode::OnUpdateCommentText( const FString& NewComment )
 	}
 }
 
-void UMaterialGraphNode::OnCommentBubbleToggled( bool bCommentBubbleVisible )
+void UMaterialGraphNode::OnCommentBubbleToggled( bool bInCommentBubbleVisible )
 {
 	if ( MaterialExpression )
 	{
 		MaterialExpression->Modify();
-		MaterialExpression->bCommentBubbleVisible = bCommentBubbleVisible;
+		MaterialExpression->bCommentBubbleVisible = bInCommentBubbleVisible;
 		MaterialDirtyDelegate.ExecuteIfBound();
 	}
 }
