UE4 - fix race condition on feature packs and global shaders

diff --git a/Engine/Source/Editor/UnrealEd/Private/FeaturePack/FeaturePack.cpp b/Engine/Source/Editor/UnrealEd/Private/FeaturePack/FeaturePack.cpp
--- a/Engine/Source/Editor/UnrealEd/Private/FeaturePack/FeaturePack.cpp
+++ b/Engine/Source/Editor/UnrealEd/Private/FeaturePack/FeaturePack.cpp
@@ -7,7 +7,8 @@
 #include "AssetRegistryModule.h"
 #include "LevelEditor.h"
 #include "Toolkits/AssetEditorManager.h"
-
+#include "GlobalShader.h"
+#include "ShaderCompiler.h"
 
 DEFINE_LOG_CATEGORY_STATIC(LogFeaturePack, Log, All);
 
@@ -78,6 +79,14 @@ void FFeaturePack::ParsePacks()
 			}
 			else
 			{
+				static bool bOnce = false;
+				if (!bOnce)
+				{
+					bOnce = true;
+					TArray<int32> ShaderMapIds;
+					ShaderMapIds.Add(GlobalShaderMapId);
+					GShaderCompilingManager->FinishCompilation(TEXT("Global"), ShaderMapIds);
+				}
 				// Save any imported assets.
 				TArray<UPackage*> ToSave;
 				for (auto ImportedObject : EachPackData.ImportedObjects)
