fixed clang warning

diff --git a/Engine/Source/Editor/UnrealEd/Private/AutoReimport/ContentDirectoryMonitor.cpp b/Engine/Source/Editor/UnrealEd/Private/AutoReimport/ContentDirectoryMonitor.cpp
--- a/Engine/Source/Editor/UnrealEd/Private/AutoReimport/ContentDirectoryMonitor.cpp
+++ b/Engine/Source/Editor/UnrealEd/Private/AutoReimport/ContentDirectoryMonitor.cpp
@@ -181,7 +181,6 @@ void FContentDirectoryMonitor::ProcessAdditions(const IAssetRegistry& Registry,
 		auto ExistingReferences = Utils::FindAssetsPertainingToFile(Registry, FullFilename);
 		if (ExistingReferences.Num() != 0)
 		{
-			FText Message;
 			if (ExistingReferences.Num() == 1)
 			{
 				FText Message = FText::Format(LOCTEXT("Info_AlreadyImported_Single", "Ignoring new file {0} as it's already imported by {1}."), FText::FromString(FPaths::GetCleanFilename(FullFilename)), FText::FromName(ExistingReferences[0].AssetName));
