UE4 - alternate fix for crash when generating thumbnails for added feature packs

diff --git a/Engine/Source/Editor/UnrealEd/Private/FeaturePack/FeaturePack.cpp b/Engine/Source/Editor/UnrealEd/Private/FeaturePack/FeaturePack.cpp
--- a/Engine/Source/Editor/UnrealEd/Private/FeaturePack/FeaturePack.cpp
+++ b/Engine/Source/Editor/UnrealEd/Private/FeaturePack/FeaturePack.cpp
@@ -7,8 +7,6 @@
 #include "AssetRegistryModule.h"
 #include "LevelEditor.h"
 #include "Toolkits/AssetEditorManager.h"
-#include "GlobalShader.h"
-#include "ShaderCompiler.h"
 
 DEFINE_LOG_CATEGORY_STATIC(LogFeaturePack, Log, All);
 
@@ -79,14 +77,6 @@ void FFeaturePack::ParsePacks()
 			}
 			else
 			{
-				static bool bOnce = false;
-				if (!bOnce)
-				{
-					bOnce = true;
-					TArray<int32> ShaderMapIds;
-					ShaderMapIds.Add(GlobalShaderMapId);
-					GShaderCompilingManager->FinishCompilation(TEXT("Global"), ShaderMapIds);
-				}
 				// Save any imported assets.
 				TArray<UPackage*> ToSave;
 				for (auto ImportedObject : EachPackData.ImportedObjects)
diff --git a/Engine/Source/Editor/UnrealEd/Private/ObjectTools.cpp b/Engine/Source/Editor/UnrealEd/Private/ObjectTools.cpp
--- a/Engine/Source/Editor/UnrealEd/Private/ObjectTools.cpp
+++ b/Engine/Source/Editor/UnrealEd/Private/ObjectTools.cpp
@@ -41,6 +41,7 @@
 #include "CanvasTypes.h"
 #include "Engine/SCS_Node.h"
 #include "Engine/UserDefinedStruct.h"
+#include "ShaderCompiler.h"
 
 DEFINE_LOG_CATEGORY_STATIC(LogObjectTools, Log, All);
 
@@ -3557,6 +3558,10 @@ namespace ThumbnailTools
 		FThumbnailRenderingInfo* RenderInfo =
 			GUnrealEd->GetThumbnailManager()->GetRenderingInfo( InObject );
 
+		if (GShaderCompilingManager)
+		{
+			GShaderCompilingManager->ProcessAsyncResults(false, true);
+		}
 		// Wait for all textures to be streamed in before we render the thumbnail
 		// @todo CB: This helps but doesn't result in 100%-streamed-in resources every time! :(
 		if( InFlushMode == EThumbnailTextureFlushMode::AlwaysFlush )
