Disabled compressed cooked packages in launch on.

diff --git a/Engine/Source/Editor/UnrealEd/Private/CookOnTheFlyServer.cpp b/Engine/Source/Editor/UnrealEd/Private/CookOnTheFlyServer.cpp
--- a/Engine/Source/Editor/UnrealEd/Private/CookOnTheFlyServer.cpp
+++ b/Engine/Source/Editor/UnrealEd/Private/CookOnTheFlyServer.cpp
@@ -2027,11 +2027,6 @@ bool UCookOnTheFlyServer::SaveCookedPackage( UPackage* Package, uint32 SaveFlags
 			bShouldCompressPackage |= CookByTheBookOptions->bForceEnableCompressedPackages;
 		}
 
-		if (bShouldCompressPackage)
-		{
-			Package->PackageFlags |= PKG_StoreCompressed;
-		}
-
 		ITargetPlatformManagerModule& TPM = GetTargetPlatformManagerRef();
 
 		static TArray<ITargetPlatform*> ActiveStartupPlatforms = TPM.GetCookingTargetPlatforms();
@@ -2065,6 +2060,15 @@ bool UCookOnTheFlyServer::SaveCookedPackage( UPackage* Package, uint32 SaveFlags
 			}
 		}
 		
+		for (const auto& TargetPlatform : Platforms)
+		{
+			bShouldCompressPackage &= TargetPlatform->SupportsFeature(ETargetPlatformFeatures::ShouldUseCompressedCookedPackages);
+		}
+
+		if (bShouldCompressPackage)
+		{
+			Package->PackageFlags |= PKG_StoreCompressed;
+		}
 
 		for (ITargetPlatform* Target : Platforms)
 		{
