DDC commandlet now blocks on async mesh distance field tasks and actually writes the results to the DDC

diff --git a/Engine/Source/Editor/UnrealEd/Private/Commandlets/DerivedDataCacheCommandlet.cpp b/Engine/Source/Editor/UnrealEd/Private/Commandlets/DerivedDataCacheCommandlet.cpp
--- a/Engine/Source/Editor/UnrealEd/Private/Commandlets/DerivedDataCacheCommandlet.cpp
+++ b/Engine/Source/Editor/UnrealEd/Private/Commandlets/DerivedDataCacheCommandlet.cpp
@@ -11,6 +11,7 @@
 #include "TargetPlatform.h"
 #include "IConsoleManager.h"
 #include "ShaderCompiler.h"
+#include "DistanceFieldAtlas.h"
 
 DEFINE_LOG_CATEGORY_STATIC(LogDerivedDataCacheCommandlet, Log, All);
 
@@ -240,6 +241,7 @@ int32 UDerivedDataCacheCommandlet::Main( const FString& Params )
 	IConsoleManager::Get().ProcessUserConsoleInput(TEXT("Tex.DerivedDataTimings"), *GWarn, NULL);
 	UE_LOG(LogDerivedDataCacheCommandlet, Display, TEXT("Waiting for shaders to finish."));
 	GShaderCompilingManager->FinishAllCompilation();
+	GDistanceFieldAsyncQueue->BlockUntilAllBuildsComplete();
 	UE_LOG(LogDerivedDataCacheCommandlet, Display, TEXT("Done waiting for shaders to finish."));
 	GetDerivedDataCacheRef().WaitForQuiescence(true);
 
diff --git a/Engine/Source/Runtime/Engine/Private/DistanceFieldAtlas.cpp b/Engine/Source/Runtime/Engine/Private/DistanceFieldAtlas.cpp
--- a/Engine/Source/Runtime/Engine/Private/DistanceFieldAtlas.cpp
+++ b/Engine/Source/Runtime/Engine/Private/DistanceFieldAtlas.cpp
@@ -421,6 +421,16 @@ void FDistanceFieldAsyncQueue::BlockUntilBuildComplete(UStaticMesh* StaticMesh,
 	}
 }
 
+void FDistanceFieldAsyncQueue::BlockUntilAllBuildsComplete()
+{
+	do 
+	{
+		ProcessAsyncTasks();
+		FPlatformProcess::Sleep(.01f);
+	} 
+	while (GetNumOutstandingTasks() > 0);
+}
+
 void FDistanceFieldAsyncQueue::Build(FAsyncDistanceFieldTask* Task, FQueuedThreadPool& ThreadPool)
 {
 #if WITH_EDITOR
diff --git a/Engine/Source/Runtime/Engine/Public/DistanceFieldAtlas.h b/Engine/Source/Runtime/Engine/Public/DistanceFieldAtlas.h
--- a/Engine/Source/Runtime/Engine/Public/DistanceFieldAtlas.h
+++ b/Engine/Source/Runtime/Engine/Public/DistanceFieldAtlas.h
@@ -190,6 +190,9 @@ public:
 	/** Blocks the main thread until the async build of the specified mesh is complete. */
 	ENGINE_API void BlockUntilBuildComplete(UStaticMesh* StaticMesh, bool bWarnIfBlocked);
 
+	/** Blocks the main thread until all async builds complete. */
+	ENGINE_API void BlockUntilAllBuildsComplete();
+
 	/** Called once per frame, fetches completed tasks and applies them to the scene. */
 	ENGINE_API void ProcessAsyncTasks();
 
