Files now remember when they have been made writable so that they do not prompt again.

diff --git a/Engine/Source/Editor/UnrealEd/Private/FileHelpers.cpp b/Engine/Source/Editor/UnrealEd/Private/FileHelpers.cpp
--- a/Engine/Source/Editor/UnrealEd/Private/FileHelpers.cpp
+++ b/Engine/Source/Editor/UnrealEd/Private/FileHelpers.cpp
@@ -39,6 +39,7 @@ DEFINE_LOG_CATEGORY_STATIC(LogFileHelpers, Log, All);
 bool FEditorFileUtils::bIsLoadingDefaultStartupMap = false;
 bool FEditorFileUtils::bIsPromptingForCheckoutAndSave = false;
 TSet<FString> FEditorFileUtils::PackagesNotSavedDuringSaveAll;
+TSet<FString> FEditorFileUtils::PackagesNotToPromptAnyMore;
 
 static const FString InvalidFilenames[] = {
 	TEXT("CON"), TEXT("PRN"), TEXT("AUX"), TEXT("CLOCK$"), TEXT("NUL"),
@@ -1020,7 +1021,7 @@ static bool IsCheckOutSelectedDisabled()
 	return !(ISourceControlModule::Get().IsEnabled() && ISourceControlModule::Get().GetProvider().IsAvailable());
 }
 
-static bool AddCheckoutPackageItems(bool bCheckDirty, TArray<UPackage*> PackagesToCheckOut, TArray<UPackage*>* OutPackagesNotNeedingCheckout, bool* bOutHavePackageToCheckOut)
+bool FEditorFileUtils::AddCheckoutPackageItems(bool bCheckDirty, TArray<UPackage*> PackagesToCheckOut, TArray<UPackage*>* OutPackagesNotNeedingCheckout, bool* bOutHavePackageToCheckOut)
 {
 	ISourceControlProvider& SourceControlProvider = ISourceControlModule::Get().GetProvider();
 	if (ISourceControlModule::Get().IsEnabled() && SourceControlProvider.IsAvailable())
@@ -1066,17 +1067,28 @@ static bool AddCheckoutPackageItems(bool bCheckDirty, TArray<UPackage*> Packages
 
 		if (!bSCCCanEdit && (bIsSourceControlled && (!bCheckDirty || (bCheckDirty && CurPackage->IsDirty()))) && !SourceControlState->IsCheckedOut())
 		{
-			if (SourceControlState.IsValid() && !SourceControlState->IsCurrent())
+			if (SourceControlState.IsValid() && (!SourceControlState->IsCurrent() || SourceControlState->IsCheckedOutOther()))
 			{
-				// This package is not at the head revision and it should be ghosted as a result
-				CheckoutPackagesDialogModule.AddPackageItem(CurPackage, CurPackage->GetName(), ECheckBoxState::Unchecked, true, TEXT("SavePackages.SCC_DlgNotCurrent"), SourceControlState->GetDisplayTooltip().ToString());
-				bShowWarning = true;
-			}
-			else if (SourceControlState.IsValid() && SourceControlState->IsCheckedOutOther())
-			{
-				// This package is checked out by someone else so it should be ghosted
-				CheckoutPackagesDialogModule.AddPackageItem(CurPackage, CurPackage->GetName(), ECheckBoxState::Unchecked, true, TEXT("SavePackages.SCC_DlgCheckedOutOther"), SourceControlState->GetDisplayTooltip().ToString());
-				bShowWarning = true;
+				if (!PackagesNotToPromptAnyMore.Contains(CurPackage->GetName()))
+				{
+					if (!SourceControlState->IsCurrent())
+					{
+						// This package is not at the head revision and it should be ghosted as a result
+						CheckoutPackagesDialogModule.AddPackageItem(CurPackage, CurPackage->GetName(), ECheckBoxState::Unchecked, true, TEXT("SavePackages.SCC_DlgNotCurrent"), SourceControlState->GetDisplayTooltip().ToString());
+					}
+					else if (SourceControlState->IsCheckedOutOther())
+					{
+						// This package is checked out by someone else so it should be ghosted
+						CheckoutPackagesDialogModule.AddPackageItem(CurPackage, CurPackage->GetName(), ECheckBoxState::Unchecked, true, TEXT("SavePackages.SCC_DlgCheckedOutOther"), SourceControlState->GetDisplayTooltip().ToString());
+					}
+					bShowWarning = true;
+					bPackagesAdded = true;
+				}
+				else
+				{
+					// File has already been made writable, just allow it to be saved without prompting
+					OutPackagesNotNeedingCheckout->Add(CurPackage);
+				}
 			}
 			else
 			{
@@ -1086,8 +1098,9 @@ static bool AddCheckoutPackageItems(bool bCheckDirty, TArray<UPackage*> Packages
 				//Add this package to the dialog if its not checked out, in the source control depot, dirty(if we are checking), and read only
 				//This package could also be marked for delete, which we will treat as SCC_ReadOnly until it is time to check it out. At that time, we will revert it.
 				CheckoutPackagesDialogModule.AddPackageItem(CurPackage, CurPackage->GetName(), ECheckBoxState::Checked, false, TEXT("SavePackages.SCC_DlgReadOnly"), Tooltip.ToString());
+				PackagesNotToPromptAnyMore.Remove(CurPackage->GetName());
+				bPackagesAdded = true;
 			}
-			bPackagesAdded = true;
 		}
 		else if (bPkgReadOnly && bFoundFile && (IsCheckOutSelectedDisabled() || !bCareAboutReadOnly))
 		{
@@ -1099,12 +1112,14 @@ static bool AddCheckoutPackageItems(bool bCheckDirty, TArray<UPackage*> Packages
 			// This package is read only but source control is not available, show the dialog so users can save the package by making the file writable or by connecting to source control.
 			// If we don't care about read-only state, we should allow the user to make the file writable whatever the state of source control.
 			CheckoutPackagesDialogModule.AddPackageItem(CurPackage, CurPackage->GetName(), ECheckBoxState::Unchecked, bIsDisabled, TEXT("SavePackages.SCC_DlgReadOnly"), Tooltip.ToString());
+			PackagesNotToPromptAnyMore.Remove(CurPackage->GetName());
 			bPackagesAdded = true;
 		}
 		else if (OutPackagesNotNeedingCheckout)
 		{
 			// The current package does not need to be checked out in order to save.
 			OutPackagesNotNeedingCheckout->Add(CurPackage);
+			PackagesNotToPromptAnyMore.Remove(CurPackage->GetName());
 		}
 	}
 
@@ -1129,7 +1144,7 @@ static bool AddCheckoutPackageItems(bool bCheckDirty, TArray<UPackage*> Packages
 	return bPackagesAdded;
 }
 
-static void UpdateCheckoutPackageItems(bool bCheckDirty, TArray<UPackage*> PackagesToCheckOut, TArray<UPackage*>* OutPackagesNotNeedingCheckout)
+void FEditorFileUtils::UpdateCheckoutPackageItems(bool bCheckDirty, TArray<UPackage*> PackagesToCheckOut, TArray<UPackage*>* OutPackagesNotNeedingCheckout)
 {
 	AddCheckoutPackageItems(bCheckDirty, PackagesToCheckOut, OutPackagesNotNeedingCheckout, nullptr);
 }
@@ -1251,7 +1266,8 @@ bool FEditorFileUtils::PromptToCheckoutPackages(bool bCheckDirty, const TArray<U
 						// Knock off the read only flag from the current file attributes
 						if (FPlatformFileManager::Get().GetPlatformFile().SetReadOnly(*Filename, false))
 						{
-							if ( OutPackagesCheckedOutOrMadeWritable )
+							PackagesNotToPromptAnyMore.Add(PackageToMakeWritable->GetName());
+							if (OutPackagesCheckedOutOrMadeWritable)
 							{
 								OutPackagesCheckedOutOrMadeWritable->Add(PackageToMakeWritable);
 							}
@@ -1288,6 +1304,12 @@ bool FEditorFileUtils::PromptToCheckoutPackages(bool bCheckDirty, const TArray<U
 	// Update again to catch potentially new SCC states
 	ISourceControlModule::Get().QueueStatusUpdate(PackagesToCheckOut);
 
+	// If any files were just checked out, remove any pending flag to show a notification prompting for checkout.
+	if (PackagesToCheckOut.Num() > 0)
+	{
+		GUnrealEd->bNeedToPromptForCheckout = false;
+	}
+
 	if (OutPackagesNotNeedingCheckout)
 	{
 		ISourceControlModule::Get().QueueStatusUpdate(*OutPackagesNotNeedingCheckout);
diff --git a/Engine/Source/Editor/UnrealEd/Public/FileHelpers.h b/Engine/Source/Editor/UnrealEd/Public/FileHelpers.h
--- a/Engine/Source/Editor/UnrealEd/Public/FileHelpers.h
+++ b/Engine/Source/Editor/UnrealEd/Public/FileHelpers.h
@@ -414,6 +414,13 @@ public:
 	UNREALED_API static FString ExtractPackageName(const FString& ObjectPath);
 
 private:
+
+	/** Private method used to build a list of items requiring checkout */
+	static bool AddCheckoutPackageItems(bool bCheckDirty, TArray<UPackage*> PackagesToCheckOut, TArray<UPackage*>* OutPackagesNotNeedingCheckout, bool* bOutHavePackageToCheckOut);
+
+	/** Callback from PackagesDialog used to update the list of items when the source control state changes */
+	static void UpdateCheckoutPackageItems(bool bCheckDirty, TArray<UPackage*> PackagesToCheckOut, TArray<UPackage*>* OutPackagesNotNeedingCheckout);
+
 	static bool bIsLoadingDefaultStartupMap;
 
 	/** Flag used to determine if the checkout and save prompt is already open to prevent re-entrance */
@@ -422,4 +429,6 @@ private:
 	// Set of packages to ignore for save/checkout when using SaveAll.
 	static TSet<FString> PackagesNotSavedDuringSaveAll;
 
+	// Set of packages which should no longer prompt for checkouts / to be made writable
+	static TSet<FString> PackagesNotToPromptAnyMore;
 };
