UMG/Blueprint - Using the new GetBlueprintTypesForClass function to find UWdgetBlueprint for UUserWidget derivatives instead of attemptings to make standard UBlueprints.

diff --git a/Engine/Source/Editor/UnrealEd/Private/Factories/EditorFactories.cpp b/Engine/Source/Editor/UnrealEd/Private/Factories/EditorFactories.cpp
--- a/Engine/Source/Editor/UnrealEd/Private/Factories/EditorFactories.cpp
+++ b/Engine/Source/Editor/UnrealEd/Private/Factories/EditorFactories.cpp
@@ -103,6 +103,7 @@
 #include "GameFramework/TouchInterface.h"
 #include "Engine/DataTable.h"
 #include "DataTableEditorUtils.h"
+#include "Editor/KismetCompiler/Public/KismetCompilerModule.h"
 
 #include "Kismet2/BlueprintEditorUtils.h"
 
@@ -6189,7 +6190,13 @@ UObject* UBlueprintFactory::FactoryCreateNew(UClass* Class, UObject* InParent, F
 	}
 	else
 	{
-		return FKismetEditorUtilities::CreateBlueprint(ParentClass, InParent, Name, BPTYPE_Normal, UBlueprint::StaticClass(), UBlueprintGeneratedClass::StaticClass(), CallingContext);
+		UClass* BlueprintClass = nullptr;
+		UClass* BlueprintGeneratedClass = nullptr;
+
+		IKismetCompilerInterface& KismetCompilerModule = FModuleManager::LoadModuleChecked<IKismetCompilerInterface>("KismetCompiler");
+		KismetCompilerModule.GetBlueprintTypesForClass(ParentClass, BlueprintClass, BlueprintGeneratedClass);
+
+		return FKismetEditorUtilities::CreateBlueprint(ParentClass, InParent, Name, BPTYPE_Normal, BlueprintClass, BlueprintGeneratedClass, CallingContext);
 	}
 }
 
