Implemented hierarchy timer for cooker packags (need to be enabled manually).

diff --git a/Engine/Source/Editor/UnrealEd/Private/CookOnTheFlyServer.cpp b/Engine/Source/Editor/UnrealEd/Private/CookOnTheFlyServer.cpp
--- a/Engine/Source/Editor/UnrealEd/Private/CookOnTheFlyServer.cpp
+++ b/Engine/Source/Editor/UnrealEd/Private/CookOnTheFlyServer.cpp
@@ -51,6 +51,7 @@ DEFINE_LOG_CATEGORY_STATIC(LogCook, Log, All);
 #if OUTPUT_TIMING
 
 #define HEIRARCHICAL_TIMER 1
+#define PERPACKAGE_TIMER 0
 
 struct FTimerInfo
 {
@@ -104,6 +105,21 @@ public:
 	{
 	}
 
+	~FHierarchicalTimerInfo()
+	{
+		Parent = NULL;
+		ClearChildren();
+	}
+
+	void ClearChildren()
+	{
+		for (auto& Child : Children)
+		{
+			delete Child.Value;
+		}
+		Children.Empty();
+	}
+
 	FHierarchicalTimerInfo* FindChild(const FString& Name)
 	{
 		FHierarchicalTimerInfo* Child = (FHierarchicalTimerInfo*)(Children.FindRef(Name));
@@ -159,19 +175,24 @@ public:
 		{
 			++GScopeDepth;
 		}
-		Index = GTimerInfo.Emplace(MoveTemp(Name), 0.0);
-		Started = false;
+
 #if HEIRARCHICAL_TIMER
 		HeirarchyTimerInfo = CurrentTimerInfo->FindChild(Name);
 		CurrentTimerInfo = HeirarchyTimerInfo;
 #endif
+#if PERPACKAGE_TIMER
+		Index = GTimerInfo.Emplace(MoveTemp(Name), 0.0);
+#endif
+		Started = false;
 	}
 
 	void Start()
 	{
 		if ( !Started )
 		{
+#if PERPACKAGE_TIMER
 			GTimerInfo[Index].Length -= FPlatformTime::Seconds();
+#endif
 			Started = true;
 #if HEIRARCHICAL_TIMER
 			HeirarchyTimerInfo->Length -= FPlatformTime::Seconds();
@@ -186,7 +207,9 @@ public:
 #if HEIRARCHICAL_TIMER
 			HeirarchyTimerInfo->Length += FPlatformTime::Seconds();
 #endif
+#if PERPACKAGE_TIMER
 			GTimerInfo[Index].Length += FPlatformTime::Seconds();
+#endif
 			Started = false;
 		}
 	}
@@ -213,7 +236,7 @@ int FScopeTimer::GScopeDepth = 0;
 
 void OutputTimers()
 {
-	
+#if PERPACKAGE_TIMER
 	if ( GTimerInfo.Num() <= 0 )
 		return;
 	
@@ -292,6 +315,39 @@ void OutputTimers()
 	}
 
 	GTimerInfo.Empty();
+#endif
+}
+
+void OutputHierarchyTimers(const FHierarchicalTimerInfo* TimerInfo, FString& Output, int32& Depth)
+{
+	// put our line into the output buffer
+	/*for (int32 I = 0; I < Depth; ++I)
+	{
+		Output += TEXT("\t");
+	}*/
+	
+	Output += FString::Printf(TEXT("%s: %fms\n"), *TimerInfo->Name, TimerInfo->Length * 1000);
+
+	++Depth;
+	for (const auto& ChildInfo : TimerInfo->Children)
+	{
+		OutputHierarchyTimers(ChildInfo.Value, Output, Depth);
+	}
+	--Depth;
+}
+
+void OutputHierarchyTimers()
+{
+	FHierarchicalTimerInfo* TimerInfo = &RootTimerInfo;
+
+	int32 Depth = 0;
+	FString Output = TEXT("Hierarchy timer information\nName:  Length(ms)");
+	OutputHierarchyTimers(TimerInfo, Output, Depth);
+
+
+	UE_LOG(LogCook, Display, TEXT("%s"), *Output);
+
+	RootTimerInfo.ClearChildren();
 }
 
 
@@ -307,6 +363,7 @@ void OutputTimers()
 #define ACCUMULATE_TIMER_STOP(name) ScopeTimer##name.Stop();
 
 #define OUTPUT_TIMERS() OutputTimers();
+#define OUTPUT_HIERARCHYTIMERS() OutputHierarchyTimers();
 
 #else
 #define CREATE_TIMER(name)
@@ -320,7 +377,7 @@ void OutputTimers()
 #define ACCUMULATE_TIMER_STOP(name) 
 
 #define OUTPUT_TIMERS()
-
+#define OUTPUT_HIERARCHYTIMERS()
 #endif
 ////////////////////////////////////////////////////////////////
 /// Cook on the fly server
@@ -3298,6 +3355,8 @@ void UCookOnTheFlyServer::CookByTheBookFinished()
 	CookByTheBookOptions->BasedOnReleaseCookedPackages.Empty();
 
 	CookByTheBookOptions->bRunning = false;
+
+	OUTPUT_HIERARCHYTIMERS();
 }
 
 void UCookOnTheFlyServer::QueueCancelCookByTheBook()
