Fix for a recent bug introduced by CL 2568322.

diff --git a/Engine/Source/Editor/UnrealEd/Private/EditorCsg.cpp b/Engine/Source/Editor/UnrealEd/Private/EditorCsg.cpp
--- a/Engine/Source/Editor/UnrealEd/Private/EditorCsg.cpp
+++ b/Engine/Source/Editor/UnrealEd/Private/EditorCsg.cpp
@@ -360,13 +360,28 @@ void UEditorEngine::polyUpdateMaster
 	UModel* Brush = Actor->Brush;
 	check(Brush);
 
-	// Use transform cached when the geometry was last built, in case the current Actor transform has changed since then
-	// (e.g. because Auto Update BSP is disabled)
-	check(Brush->bCachedOwnerTransformValid);
-	const FVector ActorLocation = Brush->OwnerLocationWhenLastBuilt;
-	const FVector ActorPrePivot = Brush->OwnerPrepivotWhenLastBuilt;
-	const FVector ActorScale = Brush->OwnerScaleWhenLastBuilt;
-	const FRotator ActorRotation = -Brush->OwnerRotationWhenLastBuilt;
+	FVector ActorLocation;
+	FVector ActorPrePivot;
+	FVector ActorScale;
+	FRotator ActorRotation;
+
+	if (Brush->bCachedOwnerTransformValid)
+	{
+		// Use transform cached when the geometry was last built, in case the current Actor transform has changed since then
+		// (e.g. because Auto Update BSP is disabled)
+		ActorLocation = Brush->OwnerLocationWhenLastBuilt;
+		ActorPrePivot = Brush->OwnerPrepivotWhenLastBuilt;
+		ActorScale = Brush->OwnerScaleWhenLastBuilt;
+		ActorRotation = -Brush->OwnerRotationWhenLastBuilt;
+	}
+	else
+	{
+		// No cached owner transform, so use the current one
+		ActorLocation = Actor->GetActorLocation();
+		ActorPrePivot = Actor->GetPrePivot();
+		ActorScale = Actor->GetActorScale();
+		ActorRotation = -Actor->GetActorRotation();
+	}
 
 	for( int32 iEdPoly = Surf.iBrushPoly; iEdPoly < Brush->Polys->Element.Num(); iEdPoly++ )
 	{
