Removed redundant URL encoding when launching editor with map name. This was internally kicking off an unrequired package search on disk. Minor optimization.

diff --git a/Engine/Source/Editor/UnrealEd/Private/UnrealEdMisc.cpp b/Engine/Source/Editor/UnrealEd/Private/UnrealEdMisc.cpp
--- a/Engine/Source/Editor/UnrealEd/Private/UnrealEdMisc.cpp
+++ b/Engine/Source/Editor/UnrealEd/Private/UnrealEdMisc.cpp
@@ -267,9 +267,6 @@ void FUnrealEdMisc::OnInit()
 			 // If it's not a parameter
 			 ParsedMapName.StartsWith( TEXT("-") ) == false )
 		{
-			FURL DefaultURL;
-			FURL URL( &DefaultURL, *ParsedMapName, TRAVEL_Partial );
-			
 			FString InitialMapName;
 				
 			// If the specified package exists
