In Editor lighting test. part of the refactoring of the promotion pass.

diff --git a/Engine/Source/Editor/UnrealEd/Private/Tests/LightingTests.cpp b/Engine/Source/Editor/UnrealEd/Private/Tests/LightingTests.cpp
new file mode 100644
--- /dev/null
+++ b/Engine/Source/Editor/UnrealEd/Private/Tests/LightingTests.cpp
@@ -0,0 +1,253 @@
+// Copyright 1998-2015 Epic Games, Inc. All Rights Reserved.
+
+#include "UnrealEd.h"
+
+//Automation
+#include "AutomationTest.h"
+#include "AutomationEditorCommon.h"
+
+#define POINTLIGHT_TRANSFORM			FTransform()
+#define POINT_LIGHT_UPDATED_LOCATION	FVector(50.f, 50.f, 50.f)
+#define POINT_LIGHT_UPDATED_ROTATION	FRotator(-60.0f, -110.0f, -91.0f)		// X: -91, Y: -60, Z: -110
+#define POINT_LIGHT_UPDATED_SCALE3D		FVector(2.f, 2.f, 2.f)
+
+#define LOCTEXT_NAMESPACE "EditorLightingBuildPromotionTests"
+
+DEFINE_LOG_CATEGORY_STATIC(LogLightingTests, All, All);
+
+namespace LightingTestHelpers
+{
+	// Searches through the list of actors in the level to find the actor that matches the desired name.
+	// @CurrentLevel - The level to search through.
+	// @ActorsName - The name of the actor to search for.
+	bool DoesActorExistInTheLevel(const ULevel* CurrentLevel, const FString& ActorsName, const UClass* ActorToFind)
+	{
+		for (int32 i = 0; i < CurrentLevel->Actors.Num(); i++)
+		{
+			FString PossibleActorsName = CurrentLevel->Actors[i]->GetName();
+			if (PossibleActorsName.Contains(ActorsName) && ActorToFind == CurrentLevel->Actors[i]->GetClass())
+			{
+				UE_LOG(LogLightingTests, Log, TEXT("Found: %s"), *PossibleActorsName);
+				return true;
+			}
+		}
+		return false;
+	}
+
+	void GetActorCurrentLocation(const ULevel* CurrentLevel, const FString& ActorsName, FVector& OutActorLocation)
+	{
+		for (int32 i = 0; i < CurrentLevel->Actors.Num(); i++)
+		{
+			FString PossibleActorsName = CurrentLevel->Actors[i]->GetName();
+			if (PossibleActorsName.Contains(ActorsName))
+			{
+				OutActorLocation = CurrentLevel->Actors[i]->GetActorLocation();
+				break;
+			}
+		}
+	}
+
+	void GetActorCurrentRotation(const ULevel* CurrentLevel, const FString ActorsName, FRotator& OutActorRotation)
+	{
+		for (int32 i = 0; i < CurrentLevel->Actors.Num(); i++)
+		{
+			FString PossibleActorsName = CurrentLevel->Actors[i]->GetName();
+			if (PossibleActorsName.Contains(ActorsName))
+			{
+				OutActorRotation = CurrentLevel->Actors[i]->GetActorRotation();
+				break;
+			}
+		}
+	}
+
+	void GetActorCurrentScale3D(const ULevel* CurrentLevel, const FString ActorsName, FVector& OutActorScale3D)
+	{
+		for (int32 i = 0; i < CurrentLevel->Actors.Num(); i++)
+		{
+			FString PossibleActorsName = CurrentLevel->Actors[i]->GetName();
+			if (PossibleActorsName.Contains(ActorsName))
+			{
+				OutActorScale3D = CurrentLevel->Actors[i]->GetActorScale3D();
+				break;
+			}
+		}
+	}
+
+	/**
+	* Sets an object property value by name
+	* @param TargetObject - The object to modify
+	* @param InVariableName - The name of the property
+	*/
+	void SetPropertyByName(UObject* TargetObject, const FString& InVariableName, const FString& NewValueString)
+	{
+		UProperty* FoundProperty = FindField<UProperty>(TargetObject->GetClass(), *InVariableName);
+		if (FoundProperty)
+		{
+			const FScopedTransaction PropertyChanged(LOCTEXT("PropertyChanged", "Object Property Change"));
+
+			TargetObject->Modify();
+
+			TargetObject->PreEditChange(FoundProperty);
+			FoundProperty->ImportText(*NewValueString, FoundProperty->ContainerPtrToValuePtr<uint8>(TargetObject), 0, TargetObject);
+			FPropertyChangedEvent PropertyChangedEvent(FoundProperty, EPropertyChangeType::ValueSet);
+			TargetObject->PostEditChangeProperty(PropertyChangedEvent);
+		}
+	}
+}
+
+//////////////////////////////////////////////////////////////////////////
+//Lighting Promotion Test
+//////////////////////////////////////////////////////////////////////////
+
+/**
+* Lighting Promotion Test - Place a Point Light and move it to a new location.
+*/
+IMPLEMENT_SIMPLE_AUTOMATION_TEST(FLightingPromotionPointLightPlaceRotScaleTest, "System.Promotion.Editor.Lighting.Place Scale Rotate", EAutomationTestFlags::ATF_Editor)
+
+bool FLightingPromotionPointLightPlaceRotScaleTest::RunTest(const FString& Parameters)
+{
+	//** SETUP **//
+	// Create the world.
+	UWorld* World = AutomationEditorCommonUtils::CreateNewMap();
+	ULevel* CurrentLevel = World->GetCurrentLevel();
+	// Test Summary
+	AddLogItem(TEXT("A point light creation.\n- Point light is placed into the world.\n- The actor is moved.\n- The actor is rotated.\n- The actor's scale is enlarged."));
+
+	if (!LightingTestHelpers::DoesActorExistInTheLevel(CurrentLevel, TEXT("PointLight"), APointLight::StaticClass()))
+	{
+		//** TEST **//
+		// Add a point light to the level.
+		APointLight* PointLight = Cast<APointLight>(GEditor->AddActor(World->GetCurrentLevel(), APointLight::StaticClass(), FTransform()));
+		// Set the actors location, rotation, and scale3D.
+		PointLight->SetActorLocation(POINT_LIGHT_UPDATED_LOCATION);
+		PointLight->SetActorRotation(POINT_LIGHT_UPDATED_ROTATION);
+		PointLight->SetActorScale3D(POINT_LIGHT_UPDATED_SCALE3D);
+
+		//** VERIFY **//
+		FVector CurrentLocation;
+		LightingTestHelpers::GetActorCurrentLocation(CurrentLevel, PointLight->GetName(), CurrentLocation);
+		FRotator CurrentRotation;
+		LightingTestHelpers::GetActorCurrentRotation(CurrentLevel, PointLight->GetName(), CurrentRotation);
+		FVector CurrentScale3D;
+		LightingTestHelpers::GetActorCurrentScale3D(CurrentLevel, PointLight->GetName(), CurrentScale3D);
+		bool RotationsAreEqual = CurrentRotation.Equals(POINT_LIGHT_UPDATED_ROTATION, 1);
+
+		TestTrue(TEXT("The placed point light was not found."), LightingTestHelpers::DoesActorExistInTheLevel(CurrentLevel, PointLight->GetName(), PointLight->GetClass()));
+		TestEqual<FVector>(TEXT("The point light is not in correct location"), POINT_LIGHT_UPDATED_LOCATION, CurrentLocation);
+		TestTrue(TEXT("The point light is not rotated correctly."), RotationsAreEqual);
+		TestEqual<FVector>(TEXT("The point light is not scaled correctly."), POINT_LIGHT_UPDATED_SCALE3D, CurrentScale3D);
+
+		return true;
+	}
+
+	AddError(TEXT("A point light already exists in this level which will block the verification of a new point light."));
+	return false;
+}
+
+/**
+* Lighting Promotion Test - Place a Point Light and move it to a new location.
+*/
+IMPLEMENT_SIMPLE_AUTOMATION_TEST(FLightingPromotionModifyProperties, "System.Promotion.Editor.Lighting.Modify Properties", EAutomationTestFlags::ATF_Editor)
+
+bool FLightingPromotionModifyProperties::RunTest(const FString& Parameters)
+{
+	//** SETUP **//
+	UWorld* World = AutomationEditorCommonUtils::CreateNewMap();
+	ULevel* CurrentLevel = World->GetCurrentLevel();
+	// Test Summary
+	AddLogItem(TEXT("The properties values for a point light are modified.\n- Intensity is set to 1000.\n- Color is set to R=0,G=0,B=25.\n- Attenuation Radius is set to 1024."));
+
+	if (!LightingTestHelpers::DoesActorExistInTheLevel(CurrentLevel, TEXT("PointLight"), APointLight::StaticClass()))
+	{
+		//** TEST **//
+		// Add a point light to the level.
+		APointLight* PointLight = Cast<APointLight>(GEditor->AddActor(World->GetCurrentLevel(), APointLight::StaticClass(), FTransform()));
+		// Modify the Lights Intensity, Light Color, and Attenuation Radius using it's properties.
+		LightingTestHelpers::SetPropertyByName(PointLight->PointLightComponent, TEXT("Intensity"), TEXT("1000.f"));
+		LightingTestHelpers::SetPropertyByName(PointLight->PointLightComponent, TEXT("LightColor"), TEXT("(R=0,G=0,B=255)"));
+		LightingTestHelpers::SetPropertyByName(PointLight->PointLightComponent, TEXT("AttenuationRadius"), TEXT("1024.f"));
+
+		//** VERIFY **//
+		TestEqual<float>(TEXT("Light brightness property was not modified."), 1000.f, PointLight->PointLightComponent->Intensity);
+		TestEqual<FColor>(TEXT("Light color property was not modified."), FColor(0,0,255), PointLight->PointLightComponent->LightColor);
+		TestEqual<float>(TEXT("Light attenuation radius was not modified."), 1024.f, PointLight->PointLightComponent->AttenuationRadius);
+
+		return true;
+	}
+
+	AddError(TEXT("A point light already exists in this level which will block the verification of a new point light."));
+	return false;
+}
+
+
+
+
+//////////////////////////////////////////////////////////////////////////
+//Lighting Tests
+//////////////////////////////////////////////////////////////////////////
+
+/**
+* Place a point light in the world with it's default settings
+* True if the light exists in the levels actor array, otherwise False.
+*/
+IMPLEMENT_SIMPLE_AUTOMATION_TEST(FLightPointLightPlacement, "System.Editor.Lighting.Point Light.Placement", EAutomationTestFlags::ATF_Editor)
+
+bool FLightPointLightPlacement::RunTest(const FString& Parameters)
+{
+	//** SETUP **//
+	// Create a new level.
+	UWorld* World = AutomationEditorCommonUtils::CreateNewMap();
+	ULevel* CurrentLevel = World->GetCurrentLevel();
+	// Light Setup
+	APointLight* PointLight = nullptr;
+	const FTransform Transform;
+	
+	if (!LightingTestHelpers::DoesActorExistInTheLevel(CurrentLevel, TEXT("PointLight"), APointLight::StaticClass()))
+	{
+		//** TEST **//
+		// Add a point light to the level.
+		PointLight = Cast<APointLight>(GEditor->AddActor(World->GetCurrentLevel(), APointLight::StaticClass(), Transform));
+
+		//** VERIFY **//
+		TestTrue(TEXT("The placed point light was not found."), LightingTestHelpers::DoesActorExistInTheLevel(CurrentLevel, PointLight->GetName(), PointLight->GetClass()));
+		return true;
+	}
+
+	AddError(TEXT("A point light already exists in this level which will block the verification of a new point light."));
+	return false;
+}
+
+/**
+* Place a point light in the world with it's default settings
+* True if the light exists in the levels actor array, otherwise False.
+*/
+IMPLEMENT_SIMPLE_AUTOMATION_TEST(FLightPointLightSetLocation, "System.Editor.Lighting.Point Light.Set Location", EAutomationTestFlags::ATF_Editor)
+
+bool FLightPointLightSetLocation::RunTest(const FString& Parameters)
+{
+	//** SETUP **//
+	// Create a new level.
+	UWorld* World = AutomationEditorCommonUtils::CreateNewMap();
+	ULevel* CurrentLevel = World->GetCurrentLevel();
+	// Light Setup
+	APointLight* PointLight = nullptr;
+	const FTransform Transform;
+
+	if (!LightingTestHelpers::DoesActorExistInTheLevel(CurrentLevel, TEXT("PointLight"), APointLight::StaticClass()))
+	{
+		//** TEST **//
+		// Add a point light to the level.
+		PointLight = Cast<APointLight>(GEditor->AddActor(World->GetCurrentLevel(), APointLight::StaticClass(), Transform));
+		PointLight->SetActorLocation(POINT_LIGHT_UPDATED_LOCATION);
+
+		//** VERIFY **//
+		FVector CurrentLocation;
+		LightingTestHelpers::GetActorCurrentLocation(CurrentLevel, PointLight->GetName(), CurrentLocation);
+
+		TestEqual<FVector>(TEXT("The point light is not in correct location"), POINT_LIGHT_UPDATED_LOCATION, CurrentLocation);
+		return true;
+	}
+
+	AddError(TEXT("A point light already exists in this level which will block the verification of a new point light."));
+	return false;
+}
\ No newline at end of file
