Fix crash on Assigning Material to Element 1 in Reimported Static Mesh

https://jira.ol.epicgames.net/browse/UE-15663

diff --git a/Engine/Source/Editor/UnrealEd/Private/Fbx/FbxStaticMeshImport.cpp b/Engine/Source/Editor/UnrealEd/Private/Fbx/FbxStaticMeshImport.cpp
--- a/Engine/Source/Editor/UnrealEd/Private/Fbx/FbxStaticMeshImport.cpp
+++ b/Engine/Source/Editor/UnrealEd/Private/Fbx/FbxStaticMeshImport.cpp
@@ -971,8 +971,8 @@ UStaticMesh* UnFbx::FFbxImporter::ImportStaticMeshAsSingle(UObject* InParent, TA
 	if( InStaticMesh != NULL && LODIndex > 0 && !SrcModel.RawMeshBulkData->IsEmpty() )
 	{
 		// clear out the old mesh data
-		FRawMesh RawMesh;
-		SrcModel.RawMeshBulkData->SaveRawMesh( RawMesh );
+		FRawMesh EmptyRawMesh;
+		SrcModel.RawMeshBulkData->SaveRawMesh( EmptyRawMesh );
 	}
 	
 	// make sure it has a new lighting guid
@@ -983,8 +983,8 @@ UStaticMesh* UnFbx::FFbxImporter::ImportStaticMeshAsSingle(UObject* InParent, TA
 	StaticMesh->LightMapCoordinateIndex = 1;
 
 
-	FRawMesh RawMesh;
-	SrcModel.RawMeshBulkData->LoadRawMesh(RawMesh);
+	FRawMesh NewRawMesh;
+	SrcModel.RawMeshBulkData->LoadRawMesh(NewRawMesh);
 
 	TArray<FFbxMaterial> MeshMaterials;
 	for (MeshIndex = 0; MeshIndex < MeshNodeArray.Num(); MeshIndex++ )
@@ -993,7 +993,7 @@ UStaticMesh* UnFbx::FFbxImporter::ImportStaticMeshAsSingle(UObject* InParent, TA
 
 		if (Node->GetMesh())
 		{
-			if (!BuildStaticMeshFromGeometry(Node->GetMesh(), StaticMesh, MeshMaterials, LODIndex, RawMesh,
+			if (!BuildStaticMeshFromGeometry(Node->GetMesh(), StaticMesh, MeshMaterials, LODIndex, NewRawMesh,
 											 VertexColorImportOption, ExistingVertexColorData, ImportOptions->VertexOverrideColor))
 			{
 				bBuildStatus = false;
@@ -1003,7 +1003,7 @@ UStaticMesh* UnFbx::FFbxImporter::ImportStaticMeshAsSingle(UObject* InParent, TA
 	}
 
 	// Store the new raw mesh.
-	SrcModel.RawMeshBulkData->SaveRawMesh(RawMesh);
+	SrcModel.RawMeshBulkData->SaveRawMesh(NewRawMesh);
 
 
 	if (bBuildStatus)
@@ -1116,26 +1116,26 @@ UStaticMesh* UnFbx::FFbxImporter::ImportStaticMeshAsSingle(UObject* InParent, TA
 
 			if (bDoRemap)
 			{
-				for (int32 TriIndex = 0; TriIndex < RawMesh.FaceMaterialIndices.Num(); ++TriIndex)
+				for (int32 TriIndex = 0; TriIndex < LocalRawMesh.FaceMaterialIndices.Num(); ++TriIndex)
 				{
-					RawMesh.FaceMaterialIndices[TriIndex] = MaterialMap[RawMesh.FaceMaterialIndices[TriIndex]];
+					LocalRawMesh.FaceMaterialIndices[TriIndex] = MaterialMap[LocalRawMesh.FaceMaterialIndices[TriIndex]];
 				}
 			}
 
 			// Compact material indices so that we won't have any sections with zero triangles.
-			RawMesh.CompactMaterialIndices();
+			LocalRawMesh.CompactMaterialIndices();
 
 			// Also compact the sorted materials array.
-			if (RawMesh.MaterialIndexToImportIndex.Num() > 0)
+			if (LocalRawMesh.MaterialIndexToImportIndex.Num() > 0)
 			{
 				TArray<FFbxMaterial> OldSortedMaterials;
 
 				Exchange(OldSortedMaterials,SortedMaterials);
-				SortedMaterials.Empty(RawMesh.MaterialIndexToImportIndex.Num());
-				for (int32 MaterialIndex = 0; MaterialIndex < RawMesh.MaterialIndexToImportIndex.Num(); ++MaterialIndex)
+				SortedMaterials.Empty(LocalRawMesh.MaterialIndexToImportIndex.Num());
+				for (int32 MaterialIndex = 0; MaterialIndex < LocalRawMesh.MaterialIndexToImportIndex.Num(); ++MaterialIndex)
 				{
 					FFbxMaterial Material;
-					int32 ImportIndex = RawMesh.MaterialIndexToImportIndex[MaterialIndex];
+					int32 ImportIndex = LocalRawMesh.MaterialIndexToImportIndex[MaterialIndex];
 					if (OldSortedMaterials.IsValidIndex(ImportIndex))
 					{
 						Material = OldSortedMaterials[ImportIndex];
@@ -1144,14 +1144,14 @@ UStaticMesh* UnFbx::FFbxImporter::ImportStaticMeshAsSingle(UObject* InParent, TA
 				}
 			}
 
-			for (int32 TriIndex = 0; TriIndex < RawMesh.FaceMaterialIndices.Num(); ++TriIndex)
+			for (int32 TriIndex = 0; TriIndex < LocalRawMesh.FaceMaterialIndices.Num(); ++TriIndex)
 			{
-				MaxMaterialIndex = FMath::Max<int32>(MaxMaterialIndex,RawMesh.FaceMaterialIndices[TriIndex]);
+				MaxMaterialIndex = FMath::Max<int32>(MaxMaterialIndex,LocalRawMesh.FaceMaterialIndices[TriIndex]);
 			}
 
 			for( int32 i = 0; i < MAX_MESH_TEXTURE_COORDS; i++ )
 			{
-				if( RawMesh.WedgeTexCoords[i].Num() == 0 )
+				if( LocalRawMesh.WedgeTexCoords[i].Num() == 0 )
 				{
 					FirstOpenUVChannel = i;
 					break;
@@ -1167,21 +1167,26 @@ UStaticMesh* UnFbx::FFbxImporter::ImportStaticMeshAsSingle(UObject* InParent, TA
 			StaticMesh->Materials.Empty();
 		}
 		
+		// Build a new map of sections with the unique material set
+		FMeshSectionInfoMap NewMap;
 		int32 NumMaterials = FMath::Min(SortedMaterials.Num(),MaxMaterialIndex+1);
 		for (int32 MaterialIndex = 0; MaterialIndex < NumMaterials; ++MaterialIndex)
 		{
 			FMeshSectionInfo Info = StaticMesh->SectionInfoMap.Get(LODIndex, MaterialIndex);
-			//int32 Index = StaticMesh->Materials.Find( SortedMaterials[MaterialIndex] );
-			//if( Index == INDEX_NONE )
-			//{
+
 			int32 Index = StaticMesh->Materials.Add(SortedMaterials[MaterialIndex].Material);
-			//}
+
 			Info.MaterialIndex = Index;
-			StaticMesh->SectionInfoMap.Set(LODIndex, MaterialIndex, Info);
-			
+			NewMap.Set( LODIndex, MaterialIndex, Info);
 		}
 
-	
+		// Copy the final section map into the static mesh
+		StaticMesh->SectionInfoMap.Clear();
+		StaticMesh->SectionInfoMap.CopyFrom(NewMap);
+
+		FRawMesh LocalRawMesh;
+		SrcModel.RawMeshBulkData->LoadRawMesh(LocalRawMesh);
+
 		// Setup default LOD settings based on the selected LOD group.
 		if (ExistingMesh == NULL && LODIndex == 0)
 		{
