Fix another FBX material import crash when the number of materials changes between reimports

diff --git a/Engine/Source/Editor/UnrealEd/Private/StaticMeshEdit.cpp b/Engine/Source/Editor/UnrealEd/Private/StaticMeshEdit.cpp
--- a/Engine/Source/Editor/UnrealEd/Private/StaticMeshEdit.cpp
+++ b/Engine/Source/Editor/UnrealEd/Private/StaticMeshEdit.cpp
@@ -924,10 +924,11 @@ void RestoreExistingMeshData(struct ExistingStaticMeshData* ExistingMeshDataPtr,
 {
 	if ( ExistingMeshDataPtr && NewMesh )
 	{
-		NewMesh->SectionInfoMap.Clear();
-		NewMesh->SectionInfoMap.CopyFrom(ExistingMeshDataPtr->ExistingSectionInfoMap);
-
-		NewMesh->Materials.Reset();
+		int32 NumCommonMaterials = FMath::Min(NewMesh->Materials.Num(), ExistingMeshDataPtr->ExistingLODData[0].ExistingMaterials.Num());
+		for (int32 MaterialIndex = 0; MaterialIndex < NumCommonMaterials; ++MaterialIndex)
+		{
+			NewMesh->Materials[MaterialIndex] = ExistingMeshDataPtr->ExistingLODData[0].ExistingMaterials[MaterialIndex];
+		}
 
 		int32 NumCommonLODs = FMath::Min<int32>(ExistingMeshDataPtr->ExistingLODData.Num(), NewMesh->SourceModels.Num());
 		for(int32 i=0; i<NumCommonLODs; i++)
@@ -936,10 +937,6 @@ void RestoreExistingMeshData(struct ExistingStaticMeshData* ExistingMeshDataPtr,
 			NewMesh->SourceModels[i].ReductionSettings = ExistingMeshDataPtr->ExistingLODData[i].ExistingReductionSettings;
 			NewMesh->SourceModels[i].ScreenSize = ExistingMeshDataPtr->ExistingLODData[i].ExistingScreenSize;
 
-			if (ExistingMeshDataPtr->ExistingLODData[i].ExistingMaterials.Num() > 0)
-			{
-				NewMesh->Materials.Append(ExistingMeshDataPtr->ExistingLODData[i].ExistingMaterials);
-			}
 		}
 
 		for(int32 i=NumCommonLODs; i < ExistingMeshDataPtr->ExistingLODData.Num(); ++i)
