Added in the promotion test for duplicating a point light.
Added in the ScopedTransaction.h include to fix the nonunity failure

diff --git a/Engine/Source/Editor/UnrealEd/Private/Tests/LightingTests.cpp b/Engine/Source/Editor/UnrealEd/Private/Tests/LightingTests.cpp
--- a/Engine/Source/Editor/UnrealEd/Private/Tests/LightingTests.cpp
+++ b/Engine/Source/Editor/UnrealEd/Private/Tests/LightingTests.cpp
@@ -2,10 +2,13 @@
 
 #include "UnrealEd.h"
 
-//Automation
+// Automation
 #include "AutomationTest.h"
 #include "AutomationEditorCommon.h"
 
+// Misc
+#include "ScopedTransaction.h"
+
 #define POINTLIGHT_TRANSFORM			FTransform()
 #define POINT_LIGHT_UPDATED_LOCATION	FVector(50.f, 50.f, 50.f)
 #define POINT_LIGHT_UPDATED_ROTATION	FRotator(-60.0f, -110.0f, -91.0f)		// X: -91, Y: -60, Z: -110
@@ -93,6 +96,13 @@ namespace LightingTestHelpers
 			TargetObject->PostEditChangeProperty(PropertyChangedEvent);
 		}
 	}
+
+	void SelectActorInLevel(AActor* ActorToSelect)
+	{
+		//Deselect everything and then select the actor.
+		GEditor->SelectNone(false, true, false);
+		GEditor->SelectActor(ActorToSelect, true, false, true);
+	}
 }
 
 //////////////////////////////////////////////////////////////////////////
@@ -111,7 +121,7 @@ bool FLightingPromotionPointLightPlaceRotScaleTest::RunTest(const FString& Param
 	UWorld* World = AutomationEditorCommonUtils::CreateNewMap();
 	ULevel* CurrentLevel = World->GetCurrentLevel();
 	// Test Summary
-	AddLogItem(TEXT("A point light creation.\n- Point light is placed into the world.\n- The actor is moved.\n- The actor is rotated.\n- The actor's scale is enlarged."));
+	AddLogItem(TEXT("Place, Scale, and Rotate.\n- A Point light is placed into the world.\n- The light is moved.\n- The light is rotated.\n- The light is scaled up."));
 
 	if (!LightingTestHelpers::DoesActorExistInTheLevel(CurrentLevel, TEXT("PointLight"), APointLight::StaticClass()))
 	{
@@ -145,7 +155,7 @@ bool FLightingPromotionPointLightPlaceRotScaleTest::RunTest(const FString& Param
 }
 
 /**
-* Lighting Promotion Test - Place a Point Light and move it to a new location.
+* Lighting Promotion Test - Modify a point lights properties.
 */
 IMPLEMENT_SIMPLE_AUTOMATION_TEST(FLightingPromotionModifyProperties, "System.Promotion.Editor.Lighting.Modify Properties", EAutomationTestFlags::ATF_Editor)
 
@@ -155,7 +165,7 @@ bool FLightingPromotionModifyProperties::RunTest(const FString& Parameters)
 	UWorld* World = AutomationEditorCommonUtils::CreateNewMap();
 	ULevel* CurrentLevel = World->GetCurrentLevel();
 	// Test Summary
-	AddLogItem(TEXT("The properties values for a point light are modified.\n- Intensity is set to 1000.\n- Color is set to R=0,G=0,B=25.\n- Attenuation Radius is set to 1024."));
+	AddLogItem(TEXT("The properties values for a point light are modified.\n- Intensity is set to 1000.\n- Color is set to R=0,G=0,B=255.\n- Attenuation Radius is set to 1024."));
 
 	if (!LightingTestHelpers::DoesActorExistInTheLevel(CurrentLevel, TEXT("PointLight"), APointLight::StaticClass()))
 	{
@@ -180,6 +190,51 @@ bool FLightingPromotionModifyProperties::RunTest(const FString& Parameters)
 }
 
 
+/**
+* Lighting Promotion Test - Duplicate/Copy Paste a point light.
+*/
+IMPLEMENT_SIMPLE_AUTOMATION_TEST(FLightingPromotionDuplicationTest, "System.Promotion.Editor.Lighting.Duplicate and Copy Paste", EAutomationTestFlags::ATF_Editor)
+
+bool FLightingPromotionDuplicationTest::RunTest(const FString& Parameters)
+{
+	//** SETUP **//
+	UWorld* World = AutomationEditorCommonUtils::CreateNewMap();
+	ULevel* CurrentLevel = World->GetCurrentLevel();
+	// Test Summary
+	AddLogItem(TEXT("Duplicate and Copy Paste\n- Duplicates a point light.\n- Copies and Pastes a point light."));
+
+	if (!LightingTestHelpers::DoesActorExistInTheLevel(CurrentLevel, TEXT("PointLight"), APointLight::StaticClass()))
+	{
+		//** TEST **//
+		// Add a point light to the level.
+		APointLight* PointLight = Cast<APointLight>(GEditor->AddActor(World->GetCurrentLevel(), APointLight::StaticClass(), FTransform()));
+		// Deselect all and then Select the light
+		LightingTestHelpers::SelectActorInLevel(PointLight);
+		// Copy and Paste.
+		GEngine->Exec(World, TEXT("EDIT COPY"));
+		GEngine->Exec(World, TEXT("EDIT PASTE"));
+		// Deselect all and then select a light
+		LightingTestHelpers::SelectActorInLevel(PointLight);
+		// Duplicate the light
+		GEngine->Exec(World, TEXT("DUPLICATE"));
+
+		//** Verify **//
+		int32 NumberOfPointLights = 0;
+		// Count the number of point lights in the level.
+		for (TObjectIterator<APointLight> It; It; ++It)
+		{
+			NumberOfPointLights++;
+		}
+		
+		// We are expecting three point lights to be in the level now.
+		TestEqual<int32>(TEXT("After duplicating a light the total number of them in the level is not correct."), 3, NumberOfPointLights);
+
+		return true;
+	}
+
+	AddError(TEXT("A point light already exists in this level which would dirty the test results."));
+	return false;
+}
 
 
 //////////////////////////////////////////////////////////////////////////
