Moved particle promotion tests into Cascade
Exposed promotion test keybinding methods
Refactored keybinding settings tests

diff --git a/Engine/Source/Editor/Cascade/Private/Tests/ParticleSystemTests.cpp b/Engine/Source/Editor/Cascade/Private/Tests/ParticleSystemTests.cpp
new file mode 100644
--- /dev/null
+++ b/Engine/Source/Editor/Cascade/Private/Tests/ParticleSystemTests.cpp
@@ -0,0 +1,238 @@
+// Copyright 1998-2015 Epic Games, Inc. All Rights Reserved.
+
+#include "CascadeModule.h"
+
+//Automation
+#include "Tests/AutomationTestSettings.h"
+#include "AutomationEditorCommon.h"
+#include "AutomationEditorPromotionCommon.h"
+
+//Particle system
+#include "Particles/ParticleSystem.h"
+#include "Particles/ParticleEmitter.h"
+#include "Particles/ParticleLODLevel.h"
+#include "Particles/Size/ParticleModuleSize.h"
+
+#include "ScopedTransaction.h"
+
+#define LOCTEXT_NAMESPACE "ParticleEditorPromotionTests"
+
+DEFINE_LOG_CATEGORY_STATIC(LogParticleEditorPromotionTests, Log, All);
+
+/**
+* Helper functions used by the build promotion automation test
+*/
+namespace ParticleEditorPromotionTestUtils
+{
+	/**
+	* Gets saved settings for Particle Editor promotion
+	*/
+	static FParticleEditorPromotionSettings TestSettings()
+	{
+		UAutomationTestSettings const* AutomationTestSettings = GetDefault<UAutomationTestSettings>();
+		check(AutomationTestSettings);
+
+		return AutomationTestSettings->ParticleEditorPromotionTest;
+	}
+
+	/**
+	* Imports an asset using the supplied factory, class, and name
+	*
+	* @param CreateFactory - The factory to use to create the asset
+	* @param AssetClass - The class of the new asset
+	* @param AssetName - The name to use for the new asset
+	*/
+	static UObject* CreateAsset(UFactory* CreateFactory, UClass* AssetClass, const FString& AssetName)
+	{
+		FString PackageName = FString::Printf(TEXT("%s/%s"), *FEditorPromotionTestUtilities::GetGamePath(), *AssetName);
+		UPackage* AssetPackage = CreatePackage(NULL, *PackageName);
+		EObjectFlags Flags = RF_Public | RF_Standalone;
+
+		UObject* CreatedAsset = CreateFactory->FactoryCreateNew(AssetClass, AssetPackage, FName(*AssetName), Flags, NULL, GWarn);
+
+		if (CreatedAsset)
+		{
+			// Notify the asset registry
+			FAssetRegistryModule::AssetCreated(CreatedAsset);
+
+			// Mark the package dirty...
+			AssetPackage->MarkPackageDirty();
+
+			UE_LOG(LogParticleEditorPromotionTests, Display, TEXT("Created asset %s (%s)"), *AssetName, *AssetClass->GetName());
+		}
+		else
+		{
+			UE_LOG(LogParticleEditorPromotionTests, Error, TEXT("Unable to create asset of type %s"), *AssetClass->GetName());
+		}
+
+		return CreatedAsset;
+	}
+
+	// @TODO - Replace with direct Cascade iteraction if possible
+	/**
+	* Sends the AssetEditor->SaveAsset UI command
+	*/
+	static void SendSaveCascadeCommand()
+	{
+		const FString Context = TEXT("AssetEditor");
+		const FString Command = TEXT("SaveAsset");
+
+		FInputChord CurrentSaveChord = FEditorPromotionTestUtilities::GetOrSetUICommand(Context, Command);
+
+		const FName FocusWidgetType(TEXT("SCascadeEmitterCanvas"));
+		FEditorPromotionTestUtilities::SendCommandToCurrentEditor(CurrentSaveChord, FocusWidgetType);
+	}
+}
+
+/**
+* Container for items related to the create asset test
+*/
+namespace ParticleEditorPromotionTestHelper
+{
+	struct FParticleEditorPromotionTest
+	{
+		/** Pointer to the execution info of this test */
+		FAutomationTestExecutionInfo* TestExecutionInfo;
+
+		/** Function definition for the test stage functions */
+		typedef bool(ParticleEditorPromotionTestHelper::FParticleEditorPromotionTest::*TestStageFn)();
+
+		/** List of test stage functions */
+		TArray<TestStageFn> TestStages;
+
+		/** The index of the test stage we are on */
+		int32 CurrentStage;
+
+		// Test assets
+
+		/** Particle created from the "Creating a Particle" stage */
+		UParticleSystem* CreatedPS;
+
+		/**
+		* Constructor
+		*/
+		FParticleEditorPromotionTest(FAutomationTestExecutionInfo* InExecutionInfo) :
+			CurrentStage(0)
+		{
+			FMemory::Memzero(this, sizeof(FParticleEditorPromotionTest));
+
+			TestExecutionInfo = InExecutionInfo;
+
+			//Add test stage functions to array in order of execution
+			TestStages.Add(&ParticleEditorPromotionTestHelper::FParticleEditorPromotionTest::ContentBrowser_CreateAParticleSystem_Part1);
+			TestStages.Add(&ParticleEditorPromotionTestHelper::FParticleEditorPromotionTest::ContentBrowser_CreateAParticleSystem_Part2);
+		}
+		/**
+		* Runs the current test stage
+		*/
+		bool Update()
+		{
+			bool bStageComplete = (this->*TestStages[CurrentStage])();
+			if (bStageComplete)
+			{
+				CurrentStage++;
+			}
+			return CurrentStage >= TestStages.Num();
+		}
+
+	private:
+
+		/**
+		* ContentBrowser Test Stage: Creating a particle system (Part 1)
+		*    Creates a new particle system and opens the cascade editor
+		*/
+		bool ContentBrowser_CreateAParticleSystem_Part1()
+		{
+			// @TODO: This should open an existing particle in Cascade; creation should be a different test
+			// @TODO: If saving is unnecessary, maybe do edits on default particle, undo (extra testing!), then close without saving
+			//Create a Particle system
+			UParticleSystemFactoryNew* PSFactory = NewObject<UParticleSystemFactoryNew>();
+			const FString& PSName(TEXT("PROMO_ParticleSystem"));
+			CreatedPS = Cast<UParticleSystem>(ParticleEditorPromotionTestUtils::CreateAsset(PSFactory, UParticleSystem::StaticClass(), PSName));
+			if (CreatedPS)
+			{
+				// @TODO: Add verification that this completed successfully.
+				FAssetEditorManager::Get().OpenEditorForAsset(CreatedPS);
+				UE_LOG(LogParticleEditorPromotionTests, Display, TEXT("Opened the cascade editor"));
+			}
+			else
+			{
+				UE_LOG(LogParticleEditorPromotionTests, Error, TEXT("Failed to create a new ParticleSystem"));
+			}
+
+			return true;
+		}
+
+		/**
+		* ContentBrowser Test Stage: Creating a particle system (Part 2)
+		*    Finds and modifies the StartSize of the particle system, saves the asset, and then closes the editor
+		*/
+		bool ContentBrowser_CreateAParticleSystem_Part2()
+		{
+			// @TODO: This should create, edit, and remove a temp asset instead of relying on a previously-created one
+			if (CreatedPS)
+			{
+				IAssetEditorInstance* AssetEditor = FAssetEditorManager::Get().FindEditorForAsset(CreatedPS, true);
+
+				bool bModifiedSize = false;
+				UParticleLODLevel* DefaultLOD = CreatedPS->Emitters[0]->LODLevels[0];
+				for (int32 i = 0; i < DefaultLOD->Modules.Num(); ++i)
+				{
+					UParticleModuleSize* SizeModule = Cast<UParticleModuleSize>(DefaultLOD->Modules[i]);
+					if (SizeModule)
+					{
+						UDistributionVectorUniform* Distribution = Cast<UDistributionVectorUniform>(SizeModule->StartSize.Distribution);
+						FEditorPromotionTestUtilities::SetPropertyByName(Distribution, TEXT("Max"), TEXT("(X=100,Y=100,Z=100)"));
+						FEditorPromotionTestUtilities::SetPropertyByName(Distribution, TEXT("Min"), TEXT("(X=100,Y=100,Z=100)"));
+						bModifiedSize = true;
+					}
+				}
+
+				if (bModifiedSize)
+				{
+					UE_LOG(LogParticleEditorPromotionTests, Display, TEXT("Modified ParticleSystem StartSize (Min and Max)"));
+					// @TODO: Do this with actual editor instead of simulated keyboard commands
+					ParticleEditorPromotionTestUtils::SendSaveCascadeCommand();
+					// @TODO: Is saving necessary? If so, do we save all assets created during test? If not, why not?
+					// @TODO: If we are going to save, verify saving succeeded before counting as a pass
+					UE_LOG(LogParticleEditorPromotionTests, Display, TEXT("Saved the particle system"));
+				}
+				else
+				{
+					UE_LOG(LogParticleEditorPromotionTests, Error, TEXT("Failed to modify ParticleSystem StartSize"));
+				}
+
+				// @TODO: Close only Cascade instead of all editors
+				FAssetEditorManager::Get().CloseAllAssetEditors();
+				UE_LOG(LogParticleEditorPromotionTests, Display, TEXT("Closed the cascade editor"));
+			}
+			// @TODO: Return false if test failed after removing sequential-testing setup
+			return true;
+		}
+	};
+}
+
+
+/**
+* Latent command to run the main build promotion test
+*/
+DEFINE_LATENT_AUTOMATION_COMMAND_ONE_PARAMETER(FRunParticleEditorPromotionTestCommand, TSharedPtr<ParticleEditorPromotionTestHelper::FParticleEditorPromotionTest>, ParticleEditorPromotionTest);
+bool FRunParticleEditorPromotionTestCommand::Update()
+{
+	return ParticleEditorPromotionTest->Update();
+}
+
+
+/**
+* Automation test that handles the build promotion process
+*/
+IMPLEMENT_SIMPLE_AUTOMATION_TEST(FParticleEditorPromotionTest, "System.Promotion.Editor.Particle Editor", EAutomationTestFlags::ATF_Editor)
+
+bool FParticleEditorPromotionTest::RunTest(const FString& Parameters)
+{
+	TSharedPtr<ParticleEditorPromotionTestHelper::FParticleEditorPromotionTest> ParticleEditorPromotionTest = MakeShareable(new ParticleEditorPromotionTestHelper::FParticleEditorPromotionTest(&ExecutionInfo));
+	ADD_LATENT_AUTOMATION_COMMAND(FRunParticleEditorPromotionTestCommand(ParticleEditorPromotionTest));
+	return true;
+}
+
+#undef LOCTEXT_NAMESPACE
\ No newline at end of file
diff --git a/Engine/Source/Editor/UnrealEd/Classes/Factories/ParticleSystemFactoryNew.h b/Engine/Source/Editor/UnrealEd/Classes/Factories/ParticleSystemFactoryNew.h
--- a/Engine/Source/Editor/UnrealEd/Classes/Factories/ParticleSystemFactoryNew.h
+++ b/Engine/Source/Editor/UnrealEd/Classes/Factories/ParticleSystemFactoryNew.h
@@ -7,14 +7,14 @@
 #pragma once
 #include "ParticleSystemFactoryNew.generated.h"
 
-UCLASS(hidecategories=Object)
+UCLASS(hidecategories=Object, MinimalAPI)
 class UParticleSystemFactoryNew : public UFactory
 {
 	GENERATED_UCLASS_BODY()
 
 
 	// Begin UFactory Interface
-	virtual UObject* FactoryCreateNew(UClass* Class,UObject* InParent,FName Name,EObjectFlags Flags,UObject* Context,FFeedbackContext* Warn) override;
+	UNREALED_API virtual UObject* FactoryCreateNew(UClass* Class,UObject* InParent,FName Name,EObjectFlags Flags,UObject* Context,FFeedbackContext* Warn) override;
 	// Begin UFactory Interface	
 };
 
diff --git a/Engine/Source/Editor/UnrealEd/Private/Tests/AutomationEditorPromotionCommon.cpp b/Engine/Source/Editor/UnrealEd/Private/Tests/AutomationEditorPromotionCommon.cpp
--- a/Engine/Source/Editor/UnrealEd/Private/Tests/AutomationEditorPromotionCommon.cpp
+++ b/Engine/Source/Editor/UnrealEd/Private/Tests/AutomationEditorPromotionCommon.cpp
@@ -3,6 +3,34 @@
 #include "UnrealEd.h"
 #include "AutomationEditorPromotionCommon.h"
 #include "AssetRegistryModule.h"
+#include "ScopedTransaction.h"
+
+#define LOCTEXT_NAMESPACE "EditorPromotionTestCommon"
+
+/**
+* Finds a visible widget by type.     SLOW!!!!!
+*
+* @param InParent - We search this widget and its children for a matching widget (recursive)
+* @param InWidgetType - The widget type we are searching for
+*/
+TSharedPtr<SWidget> FEditorPromotionTestUtilities::FindFirstWidgetByClass(TSharedRef<SWidget> InParent, const FName& InWidgetType)
+{
+	if (InParent->GetType() == InWidgetType)
+	{
+		return InParent;
+	}
+	FChildren* Children = InParent->GetChildren();
+	for (int32 i = 0; i < Children->Num(); ++i)
+	{
+		TSharedRef<SWidget> ChildWidget = Children->GetChildAt(i);
+		TSharedPtr<SWidget> FoundWidget = FindFirstWidgetByClass(ChildWidget, InWidgetType);
+		if (FoundWidget.IsValid())
+		{
+			return FoundWidget;
+		}
+	}
+	return NULL;
+}
 
 /**
 * Gets the base path for this asset
@@ -41,4 +69,118 @@ UMaterial* FEditorPromotionTestUtilities::CreateMaterialFromTexture(UTexture* In
 	}
 
 	return Cast<UMaterial>(CreatedAsset);
-}
\ No newline at end of file
+}
+
+/**
+* Sets an editor keyboard shortcut
+*
+* @param CommandContext - The context of the command
+* @param Command - The command name to set
+* @param NewChord - The new input chord to assign
+*/
+bool FEditorPromotionTestUtilities::SetEditorKeybinding(const FString& CommandContext, const FString& Command, const FInputChord& NewChord)
+{
+	TSharedPtr<FUICommandInfo> UICommand = FInputBindingManager::Get().FindCommandInContext(*CommandContext, *Command);
+	if (UICommand.IsValid())
+	{
+		UICommand->SetActiveChord(NewChord);
+		FInputBindingManager::Get().NotifyActiveChordChanged(*UICommand.Get());
+		FInputBindingManager::Get().SaveInputBindings();
+		return true;
+	}
+	else
+	{
+		return false;
+	}
+}
+
+/**
+* Gets an editor keyboard shortcut
+*
+* @param CommandContext - The context of the command
+* @param Command - The command name to get
+*/
+FInputChord FEditorPromotionTestUtilities::GetEditorKeybinding(const FString& CommandContext, const FString& Command)
+{
+	FInputChord Keybinding;
+	TSharedPtr<FUICommandInfo> UICommand = FInputBindingManager::Get().FindCommandInContext(*CommandContext, *Command);
+	if (UICommand.IsValid())
+	{
+		TSharedRef<const FInputChord> ActiveChord = UICommand->GetActiveChord();
+		Keybinding = ActiveChord.Get();
+	}
+	return Keybinding;
+}
+
+/**
+* Gets the current input chord or sets a new one if it doesn't exist
+*
+* @param Context - The context of the UI Command
+* @param Command - The name of the UI command
+*/
+FInputChord FEditorPromotionTestUtilities::GetOrSetUICommand(const FString& Context, const FString& Command)
+{
+	FInputChord CurrentChord = GetEditorKeybinding(Context, Command);
+
+	//If there is no current keybinding, set one
+	if (!CurrentChord.Key.IsValid())
+	{
+		FInputChord NewChord(EKeys::J, EModifierKey::Control);
+		SetEditorKeybinding(Context, Command, NewChord);
+		CurrentChord = NewChord;
+	}
+
+	return CurrentChord;
+}
+
+/**
+* Sends a UI command to the active top level window after focusing on a widget of a given type
+*
+* @param InChord - The chord to send to the window
+* @param WidgetTypeToFocus - The widget type to find and focus on
+*/
+void FEditorPromotionTestUtilities::SendCommandToCurrentEditor(const FInputChord& InChord, const FName& WidgetTypeToFocus)
+{
+	//Focus the asset Editor / Graph 
+	TSharedRef<SWindow> EditorWindow = FSlateApplication::Get().GetActiveTopLevelWindow().ToSharedRef();
+	FSlateApplication::Get().ProcessWindowActivatedEvent(FWindowActivateEvent(FWindowActivateEvent::EA_Activate, EditorWindow));
+	TSharedPtr<SWidget> FocusWidget = FindFirstWidgetByClass(EditorWindow, WidgetTypeToFocus);
+	if (FocusWidget.IsValid())
+	{
+		FSlateApplication::Get().SetKeyboardFocus(FocusWidget.ToSharedRef(), EFocusCause::SetDirectly);
+
+		//Send the command
+		FModifierKeysState ModifierKeys(InChord.NeedsShift(), false, InChord.NeedsControl(), false, InChord.NeedsAlt(), false, InChord.NeedsCommand(), false, false);
+		FKeyEvent KeyEvent(InChord.Key, ModifierKeys, 0/*UserIndex*/, false, 0, 0);
+		FSlateApplication::Get().ProcessKeyDownEvent(KeyEvent);
+		FSlateApplication::Get().ProcessKeyUpEvent(KeyEvent);
+	}
+	else
+	{
+		//UE_LOG(LogParticleEditorPromotionTests, Error, TEXT("Could not find widget %s to send UI command"), *WidgetTypeToFocus.ToString());
+	}
+}
+
+/**
+* Sets an object property value by name
+*
+* @param TargetObject - The object to modify
+* @param InVariableName - The name of the property
+*/
+void FEditorPromotionTestUtilities::SetPropertyByName(UObject* TargetObject, const FString& InVariableName, const FString& NewValueString)
+{
+	UProperty* FoundProperty = FindField<UProperty>(TargetObject->GetClass(), *InVariableName);
+	if (FoundProperty)
+	{
+		const FScopedTransaction PropertyChanged(LOCTEXT("PropertyChanged", "Object Property Change"));
+
+		TargetObject->Modify();
+
+		TargetObject->PreEditChange(FoundProperty);
+		FoundProperty->ImportText(*NewValueString, FoundProperty->ContainerPtrToValuePtr<uint8>(TargetObject), 0, TargetObject);
+		FPropertyChangedEvent PropertyChangedEvent(FoundProperty, EPropertyChangeType::ValueSet);
+		TargetObject->PostEditChangeProperty(PropertyChangedEvent);
+	}
+}
+
+#undef LOCTEXT_NAMESPACE
diff --git a/Engine/Source/Editor/UnrealEd/Private/Tests/EditorBuildPromotionTests.cpp b/Engine/Source/Editor/UnrealEd/Private/Tests/EditorBuildPromotionTests.cpp
--- a/Engine/Source/Editor/UnrealEd/Private/Tests/EditorBuildPromotionTests.cpp
+++ b/Engine/Source/Editor/UnrealEd/Private/Tests/EditorBuildPromotionTests.cpp
@@ -161,50 +161,6 @@ namespace EditorBuildPromotionTestUtils
 		double Duration;
 	};
 
-	/**
-	* Sets an editor keyboard shortcut
-	*
-	* @param CommandContext - The context of the command
-	* @param Command - The command name to set
-	* @param NewChord - The new input chord to assign
-	*/
-	static void SetEditorKeybinding(const FString& CommandContext, const FString& Command, const FInputChord& NewChord)
-	{
-		TSharedPtr<FUICommandInfo> UICommand = FInputBindingManager::Get().FindCommandInContext(*CommandContext, *Command);
-		if (UICommand.IsValid())
-		{
-			UICommand->SetActiveChord(NewChord);
-			FInputBindingManager::Get().NotifyActiveChordChanged(*UICommand.Get());
-			FInputBindingManager::Get().SaveInputBindings();
-		}
-		else
-		{
-			UE_LOG(LogEditorBuildPromotionTests, Error, TEXT("Could not find keybinding for %s using context %s"), *Command, *CommandContext);
-		}
-	}
-
-	/**
-	* Gets an editor keyboard shortcut
-	*
-	* @param CommandContext - The context of the command
-	* @param Command - The command name to get
-	* @param NewChord - The current input chord that is assigned
-	*/
-	static bool GetEditorKeybinding(const FString& CommandContext, const FString& Command, FInputChord& CurrentChord)
-	{
-		TSharedPtr<FUICommandInfo> UICommand = FInputBindingManager::Get().FindCommandInContext(*CommandContext, *Command);
-		if (UICommand.IsValid())
-		{
-			TSharedRef<const FInputChord> ActiveChord = UICommand->GetActiveChord();
-			CurrentChord = ActiveChord.Get();
-		}
-		else
-		{
-			UE_LOG(LogEditorBuildPromotionTests, Error, TEXT("Could not find keybinding for %s using context %s"), *Command, *CommandContext);
-		}
-		return true;
-	}
-
 	/**
 	* Exports the current editor keybindings
 	*
@@ -253,81 +209,6 @@ namespace EditorBuildPromotionTestUtils
 		GConfig->LoadFile(GEditorPerProjectIni);
 	}
 
-	/**
-	* Finds a visible widget by type.     SLOW!!!!!
-	*
-	* @param InParent - We search this widget and its children for a matching widget (recursive)
-	* @param InWidgetType - The widget type we are searching for
-	*/
-	static TSharedPtr<SWidget> FindFirstWidgetByClass(TSharedRef<SWidget> InParent, const FName& InWidgetType)
-	{
-		if (InParent->GetType() == InWidgetType)
-		{
-			return InParent;
-		}
-		FChildren* Children = InParent->GetChildren();
-		for (int32 i = 0; i < Children->Num(); ++i)
-		{
-			TSharedRef<SWidget> ChildWidget = Children->GetChildAt(i);
-			TSharedPtr<SWidget> FoundWidget = FindFirstWidgetByClass(ChildWidget, InWidgetType);
-			if (FoundWidget.IsValid())
-			{
-				return FoundWidget;
-			}
-		}
-		return NULL;
-	}
-
-	/**
-	* Sends a UI command to the active top level window after focusing on a widget of a given type
-	*
-	* @param InChord - The chord to send to the window
-	* @param WidgetTypeToFocus - The widget type to find and focus on
-	*/
-	static void SendCommandToCurrentEditor(const FInputChord& InChord, const FName& WidgetTypeToFocus)
-	{
-		//Focus the asset Editor / Graph 
-		TSharedRef<SWindow> EditorWindow = FSlateApplication::Get().GetActiveTopLevelWindow().ToSharedRef();
-		FSlateApplication::Get().ProcessWindowActivatedEvent(FWindowActivateEvent(FWindowActivateEvent::EA_Activate, EditorWindow));
-		TSharedPtr<SWidget> FocusWidget = FindFirstWidgetByClass(EditorWindow, WidgetTypeToFocus);
-		if (FocusWidget.IsValid())
-		{
-			FSlateApplication::Get().SetKeyboardFocus(FocusWidget.ToSharedRef(), EFocusCause::SetDirectly);
-
-			//Send the command
-			FModifierKeysState ModifierKeys(InChord.NeedsShift(), false, InChord.NeedsControl(), false, InChord.NeedsAlt(), false, InChord.NeedsCommand(), false, false);
-			FKeyEvent KeyEvent(InChord.Key, ModifierKeys, 0/*UserIndex*/, false, 0, 0);
-			FSlateApplication::Get().ProcessKeyDownEvent(KeyEvent);
-			FSlateApplication::Get().ProcessKeyUpEvent(KeyEvent);
-		}
-		else
-		{
-			UE_LOG(LogEditorBuildPromotionTests, Error, TEXT("Could not find widget %s to send UI command"), *WidgetTypeToFocus.ToString());
-		}
-	}
-
-	/**
-	* Gets the current input chord or sets a new one if it doesn't exist
-	*
-	* @param Context - The context of the UI Command
-	* @param Command - The name of the UI command
-	*/
-	static FInputChord GetOrSetUICommand(const FString& Context, const FString& Command)
-	{
-		FInputChord CurrentChord;
-		GetEditorKeybinding(Context, Command, CurrentChord);
-
-		//If there is no current keybinding, set one
-		if (!CurrentChord.Key.IsValid())
-		{
-			FInputChord NewChord(EKeys::J, EModifierKey::Control);
-			SetEditorKeybinding(Context, Command, NewChord);
-			CurrentChord = NewChord;
-		}
-
-		return CurrentChord;
-	}
-
 	/**
 	* Sends the MaterialEditor->Apply UI command
 	*/
@@ -335,24 +216,10 @@ namespace EditorBuildPromotionTestUtils
 	{
 		const FString Context = TEXT("MaterialEditor");
 		const FString Command = TEXT("Apply");
-		FInputChord CurrentApplyChord = GetOrSetUICommand(Context, Command);
+		FInputChord CurrentApplyChord = FEditorPromotionTestUtilities::GetOrSetUICommand(Context, Command);
 
 		const FName FocusWidgetType(TEXT("SGraphEditor"));
-		SendCommandToCurrentEditor(CurrentApplyChord, FocusWidgetType);
-	}
-
-	/**
-	* Sends the AssetEditor->SaveAsset UI command
-	*/
-	static void SendSaveCascadeCommand()
-	{
-		const FString Context = TEXT("AssetEditor");
-		const FString Command = TEXT("SaveAsset");
-
-		FInputChord CurrentSaveChord = GetOrSetUICommand(Context, Command);
-
-		const FName FocusWidgetType(TEXT("SCascadeEmitterCanvas"));
-		SendCommandToCurrentEditor(CurrentSaveChord, FocusWidgetType);
+		FEditorPromotionTestUtilities::SendCommandToCurrentEditor(CurrentApplyChord, FocusWidgetType);
 	}
 
 	/**
@@ -363,10 +230,10 @@ namespace EditorBuildPromotionTestUtils
 		const FString Context = TEXT("BlueprintEditor");
 		const FString Command = TEXT("ResetCamera");
 
-		FInputChord CurrentSaveChord = GetOrSetUICommand(Context, Command);
+		FInputChord CurrentSaveChord = FEditorPromotionTestUtilities::GetOrSetUICommand(Context, Command);
 
 		const FName FocusWidgetType(TEXT("SSCSEditorViewport"));
-		SendCommandToCurrentEditor(CurrentSaveChord, FocusWidgetType);
+		FEditorPromotionTestUtilities::SendCommandToCurrentEditor(CurrentSaveChord, FocusWidgetType);
 	}
 
 	/**
@@ -446,6 +313,7 @@ namespace EditorBuildPromotionTestUtils
 	*/
 	static void SetComponentAsRoot(UBlueprint* InBlueprint, USCS_Node* NewRoot)
 	{
+		// @FIXME: Current usages doesn't guarantee NewRoot is valid!!! Check first!
 		//Get all the construction script nodes
 		TArray<USCS_Node*> AllNodes = InBlueprint->SimpleConstructionScript->GetAllNodes();
 
@@ -868,24 +736,49 @@ namespace EditorBuildPromotionTestUtils
 	}
 
 	/**
-	* Sets an object property value by name
+	* Creates a new keybinding chord and sets it for the supplied command and context
 	*
-	* @param TargetObject - The object to modify
-	* @param InVariableName - The name of the property
+	* @param CommandContext - The context of the command
+	* @param Command - The command name to get
+	* @param Key - The keybinding chord key
+	* @param ModifierKey - The keybinding chord modifier key
 	*/
-	static void SetPropertyByName(UObject* TargetObject, const FString& InVariableName, const FString& NewValueString)
+	static FInputChord SetKeybinding(const FString& CommandContext, const FString& Command, const FKey Key, const EModifierKey::Type ModifierKey)
 	{
-		UProperty* FoundProperty = FindField<UProperty>(TargetObject->GetClass(), *InVariableName);
-		if (FoundProperty)
+		FInputChord NewChord(Key, ModifierKey);
+		if (!FEditorPromotionTestUtilities::SetEditorKeybinding(CommandContext, Command, NewChord))
 		{
-			const FScopedTransaction PropertyChanged(LOCTEXT("PropertyChanged", "Object Property Change"));
-
-			TargetObject->Modify();
+			// Trigger a failure when used in an automated test
+			UE_LOG(LogEditorBuildPromotionTests, Error, TEXT("Could not find keybinding for %s using context %s"), *Command, *CommandContext);
+		}
+		return NewChord;
+	}
 
-			TargetObject->PreEditChange(FoundProperty);
-			FoundProperty->ImportText(*NewValueString, FoundProperty->ContainerPtrToValuePtr<uint8>(TargetObject), 0, TargetObject);
-			FPropertyChangedEvent PropertyChangedEvent(FoundProperty, EPropertyChangeType::ValueSet);
-			TargetObject->PostEditChangeProperty(PropertyChangedEvent);
+	/**
+	* Retrieves the current keybinding for a command and compares it against the expected binding.
+	* Triggers an automation test failure if keybind cannot be retrieved or does not match expected binding.
+	*
+	* @param CommandContext - The context of the command
+	* @param Command - The command name to get
+	* @param ExpectedChord - The chord value to compare against
+	*/
+	static void CompareKeybindings(const FString& CommandContext, const FString& Command, FInputChord ExpectedChord)
+	{
+		FInputChord CurrentChord = FEditorPromotionTestUtilities::GetEditorKeybinding(CommandContext, Command);
+		if (!CurrentChord.IsValidChord())
+		{
+			UE_LOG(LogEditorBuildPromotionTests, Error, TEXT("Could not find keybinding for %s using context %s"), *Command, *CommandContext);
+		}
+		else
+		{
+			if (CurrentChord == ExpectedChord)
+			{
+				UE_LOG(LogEditorBuildPromotionTests, Display, TEXT("%s keybinding correct."), *Command);
+			}
+			else
+			{
+				UE_LOG(LogEditorBuildPromotionTests, Error, TEXT("%s keybinding incorrect."), *Command);
+			}
 		}
 	}
 
@@ -1087,39 +980,6 @@ namespace EditorBuildPromotionTestUtils
 		}
 	}
 
-	/**
-	* Imports an asset using the supplied factory, class, and name
-	*
-	* @param CreateFactory - The factory to use to create the asset
-	* @param AssetClass - The class of the new asset
-	* @param AssetName - The name to use for the new asset
-	*/
-	static UObject* CreateAsset(UFactory* CreateFactory, UClass* AssetClass, const FString& AssetName)
-	{
-		FString PackageName = FString::Printf(TEXT("%s/%s"), *FEditorPromotionTestUtilities::GetGamePath(), *AssetName);
-		UPackage* AssetPackage = CreatePackage(NULL, *PackageName);
-		EObjectFlags Flags = RF_Public | RF_Standalone;
-
-		UObject* CreatedAsset = CreateFactory->FactoryCreateNew(AssetClass, AssetPackage, FName(*AssetName), Flags, NULL, GWarn);
-
-		if (CreatedAsset)
-		{
-			// Notify the asset registry
-			FAssetRegistryModule::AssetCreated(CreatedAsset);
-
-			// Mark the package dirty...
-			AssetPackage->MarkPackageDirty();
-
-			UE_LOG(LogEditorBuildPromotionTests, Display, TEXT("Created asset %s (%s)"), *AssetName, *AssetClass->GetName());
-		}
-		else
-		{
-			UE_LOG(LogEditorBuildPromotionTests, Error, TEXT("Unable to create asset of type %s"), *AssetClass->GetName());
-		}
-
-		return CreatedAsset;
-	}
-
 	/**
 	* Saves all assets in a given folder
 	*
@@ -1406,7 +1266,7 @@ namespace BuildPromotionTestHelper
 		}
 
 		/**
-		* Modifies a property on the current asset
+		* Modifies a property on the current asset, undoes and redoes the property change, then saves changed asset
 		*/
 		void ChangeProperty()
 		{
@@ -1415,7 +1275,7 @@ namespace BuildPromotionTestHelper
 			FString NewPropertyValue = Assets[AssetIndex].PropertyValue;
 
 			FString OldPropertyValue = EditorBuildPromotionTestUtils::GetPropertyByName(CurrentAsset, PropertyName);
-			EditorBuildPromotionTestUtils::SetPropertyByName(CurrentAsset, PropertyName, NewPropertyValue);
+			FEditorPromotionTestUtilities::SetPropertyByName(CurrentAsset, PropertyName, NewPropertyValue);
 			UE_LOG(LogEditorBuildPromotionTests, Display, TEXT("Modified asset.  %s = %s"), *PropertyName, *NewPropertyValue);
 
 			//Get the property again and use that to compare the redo action.  Parsing the new value may change the formatting a bit. ie) 100 becomes 100.0000
@@ -1600,8 +1460,8 @@ namespace BuildPromotionTestHelper
 		UMaterial* SCTestMat;
 		FString ChosenMaterialColor;
 
-		/** Particle System created from the "Creating a Particle System" stage */
-		UParticleSystem* CreatedPS;
+		/** Particle System loaded from Automation Settings for Blueprint Pass */
+		UParticleSystem* LoadedParticleSystem;
 
 		/** Helper for opening, modifying, and placing assets */
 		FOpenAssetHelper* OpenAssetHelper;
@@ -1691,8 +1551,6 @@ namespace BuildPromotionTestHelper
 			ADD_TEST_STAGE(ContentBrowser_OpenAssets_Part2,				TEXT("Open Asset Types"));
 			ADD_TEST_STAGE(ContentBrowser_ReimportAsset,				TEXT("Re-import Assets"));
 			ADD_TEST_STAGE(ContentBrowser_AssignAMaterial,				TEXT("Assigning a Material"));
-			ADD_TEST_STAGE(ContentBrowser_CreateAParticleSystem_Part1,	TEXT("Creating a Particle System"));
-			ADD_TEST_STAGE(ContentBrowser_CreateAParticleSystem_Part2,	TEXT("Creating a Particle System"));
 			ADD_TEST_STAGE(EndSection, TEXT("Content Browser"));
 
 			// 6) Blueprints
@@ -1923,11 +1781,11 @@ namespace BuildPromotionTestHelper
 			if (PointLight)
 			{
 				UPointLightComponent* PointLightComp = PointLight->PointLightComponent;
-				EditorBuildPromotionTestUtils::SetPropertyByName(PointLightComp, TEXT("Intensity"), TEXT("1000.f"));
+				FEditorPromotionTestUtilities::SetPropertyByName(PointLightComp, TEXT("Intensity"), TEXT("1000.f"));
 				UE_LOG(LogEditorBuildPromotionTests, Display, TEXT("Modified the PointLight's Intensity"));
-				EditorBuildPromotionTestUtils::SetPropertyByName(PointLightComp, TEXT("LightColor"), TEXT("(R=0,G=0,B=255)"));
+				FEditorPromotionTestUtilities::SetPropertyByName(PointLightComp, TEXT("LightColor"), TEXT("(R=0,G=0,B=255)"));
 				UE_LOG(LogEditorBuildPromotionTests, Display, TEXT("Modified the PointLight's LightColor"));
-				EditorBuildPromotionTestUtils::SetPropertyByName(PointLightComp, TEXT("AttenuationRadius"), TEXT("1024.f"));
+				FEditorPromotionTestUtilities::SetPropertyByName(PointLightComp, TEXT("AttenuationRadius"), TEXT("1024.f"));
 				UE_LOG(LogEditorBuildPromotionTests, Display, TEXT("Modified the PointLight's AttenuationRadius"));
 			}
 			else
@@ -2317,7 +2175,7 @@ namespace BuildPromotionTestHelper
 					if (ColorParam)
 					{
 						EditorMaterial->Modify();
-						EditorBuildPromotionTestUtils::SetPropertyByName(ColorParam, TEXT("Constant"), ColorValue);
+						FEditorPromotionTestUtilities::SetPropertyByName(ColorParam, TEXT("Constant"), ColorValue);
 						MaterialEditor->UpdateMaterialAfterGraphChange();
 						MaterialEditor->ForceRefreshExpressionPreviews();
 						EditorBuildPromotionTestUtils::SendUpdateMaterialCommand();
@@ -2456,7 +2314,7 @@ namespace BuildPromotionTestHelper
 			}
 			
 			// Particle System
-			AssetPackagePath = AutomationTestSettings->BuildPromotionTest.OpenAssets.ParticleSystemAsset.FilePath;
+			AssetPackagePath = AutomationTestSettings->BuildPromotionTest.OpenAssets.ParticleSystemAsset.FilePath;  // @TODO: Use an Engine asset
 			if (AssetPackagePath.Len() > 0)
 			{
 				AssetData = FEditorAutomationTestUtilities::GetAssetDataFromPackagePath(AssetPackagePath);
@@ -2638,75 +2496,11 @@ namespace BuildPromotionTestHelper
 			return true;
 		}
 
-		/**
-		* ContentBrowser Test Stage: Creating a particle system (Part 1)
-		*    Creates a new particle system and opens the cascade editor
-		*/
-		bool ContentBrowser_CreateAParticleSystem_Part1()
-		{
-			//Create a Particle system
-			UParticleSystemFactoryNew* PSFactory = NewObject<UParticleSystemFactoryNew>();
-			const FString& PSName(TEXT("BP_ParticleSystem"));
-			CreatedPS = Cast<UParticleSystem>(EditorBuildPromotionTestUtils::CreateAsset(PSFactory, UParticleSystem::StaticClass(), PSName));
-			if (CreatedPS)
-			{
-				FAssetEditorManager::Get().OpenEditorForAsset(CreatedPS);
-				UE_LOG(LogEditorBuildPromotionTests, Display, TEXT("Opened the cascade editor"));
-			}
-			else
-			{
-				UE_LOG(LogEditorBuildPromotionTests, Error, TEXT("Failed to create a new ParticleSystem"));
-			}
-
-			return true;
-		}
-
-		/**
-		* ContentBrowser Test Stage: Creating a particle system (Part 2)
-		*    Finds and modifies the StartSize of the particle system, saves the asset, and then closes the editor
-		*/
-		bool ContentBrowser_CreateAParticleSystem_Part2()
-		{
-			if (CreatedPS)
-			{
-				IAssetEditorInstance* AssetEditor = FAssetEditorManager::Get().FindEditorForAsset(CreatedPS, true);
-
-				bool bModifiedSize = false;
-				UParticleLODLevel* DefaultLOD = CreatedPS->Emitters[0]->LODLevels[0];
-				for (int32 i = 0; i < DefaultLOD->Modules.Num(); ++i)
-				{
-					UParticleModuleSize* SizeModule = Cast<UParticleModuleSize>(DefaultLOD->Modules[i]);
-					if (SizeModule)
-					{
-						UDistributionVectorUniform* Distribution = Cast<UDistributionVectorUniform>(SizeModule->StartSize.Distribution);
-						EditorBuildPromotionTestUtils::SetPropertyByName(Distribution, TEXT("Max"), TEXT("(X=100,Y=100,Z=100)"));
-						EditorBuildPromotionTestUtils::SetPropertyByName(Distribution, TEXT("Min"), TEXT("(X=100,Y=100,Z=100)"));
-						bModifiedSize = true;
-					}
-				}
-
-				if (bModifiedSize)
-				{
-					UE_LOG(LogEditorBuildPromotionTests, Display, TEXT("Modified ParticleSystem StartSize (Min and Max)"));
-					EditorBuildPromotionTestUtils::SendSaveCascadeCommand();
-					UE_LOG(LogEditorBuildPromotionTests, Display, TEXT("Saved the particle system"));
-				}
-				else
-				{
-					UE_LOG(LogEditorBuildPromotionTests, Error, TEXT("Failed to modify ParticleSystem StartSize"));
-				}
-
-				FAssetEditorManager::Get().CloseAllAssetEditors();
-				UE_LOG(LogEditorBuildPromotionTests, Display, TEXT("Closed the cascade editor"));
-			}
-			return true;
-		}
-
 		/**
 		* Saves the blueprint stored in BlueprintObject
 		*/
 		void SaveBlueprint()
-		{
+		{  // @FIXME: Saving BP currently fails.
 			if (BlueprintObject && BlueprintPackage)
 			{
 				BlueprintPackage->SetDirtyFlag(true);
@@ -2748,7 +2542,14 @@ namespace BuildPromotionTestHelper
 				SecondBlueprintMesh = Cast<UStaticMesh>(AssetData.GetAsset());
 			}
 
-			if (FirstBlueprintMesh && SecondBlueprintMesh && CreatedPS)
+			const FString ParticleSystemPath = AutomationTestSettings->ParticleEditorPromotionTest.DefaultParticleAsset.FilePath;
+			if (ParticleSystemPath.Len() > 0)
+			{
+				AssetData = FEditorAutomationTestUtilities::GetAssetDataFromPackagePath(ParticleSystemPath);
+				LoadedParticleSystem = Cast<UParticleSystem>(AssetData.GetAsset());
+			}
+
+			if (FirstBlueprintMesh && SecondBlueprintMesh && LoadedParticleSystem)
 			{
 				UBlueprintFactory* Factory = NewObject<UBlueprintFactory>();
 				Factory->ParentClass = AActor::StaticClass();
@@ -2786,8 +2587,8 @@ namespace BuildPromotionTestHelper
 			}
 			else
 			{
-				SkippedTests.Add(TEXT("All Blueprint tests. (Missing a required mesh)"));
-				UE_LOG(LogEditorBuildPromotionTests, Warning, TEXT("SKIPPING BLUEPRINT TESTS.  Invalid or missing FirstMeshPath or SecondMeshPath in AutomationTestSettings."));
+				SkippedTests.Add(TEXT("All Blueprint tests. (Missing a required mesh or particle system)"));
+				UE_LOG(LogEditorBuildPromotionTests, Warning, TEXT("SKIPPING BLUEPRINT TESTS.  Invalid or missing FirstMeshPath or SecondMeshPath in AutomationTestSettings, or particle system was not created."));
 			}
 			
 			return true;
@@ -2900,7 +2701,7 @@ namespace BuildPromotionTestHelper
 					UE_LOG(LogEditorBuildPromotionTests, Error, TEXT("Failed to create the second mesh component"));
 				}
 
-				PSNode = EditorBuildPromotionTestUtils::CreateBlueprintComponent(BlueprintObject, CreatedPS);
+				PSNode = EditorBuildPromotionTestUtils::CreateBlueprintComponent(BlueprintObject, LoadedParticleSystem);
 				if (PSNode)
 				{
 					UE_LOG(LogEditorBuildPromotionTests, Display, TEXT("Added a particle system component"));
@@ -2911,6 +2712,7 @@ namespace BuildPromotionTestHelper
 				}
 
 				//Set the Particle System as the root
+				// @FIXME: This will probably do bad things if PSNode doesn't exist?
 				EditorBuildPromotionTestUtils::SetComponentAsRoot(BlueprintObject, PSNode);
 				UE_LOG(LogEditorBuildPromotionTests, Display, TEXT("Set the particle system component as the new root"));
 
@@ -2923,6 +2725,7 @@ namespace BuildPromotionTestHelper
 				FBlueprintEditorUtils::RenameComponentMemberVariable(BlueprintObject, OtherMeshNode, OtherMeshName);
 				UE_LOG(LogEditorBuildPromotionTests, Display, TEXT("Renamed the second mesh component to SecondMesh"));
 
+				// @FIXME: This will probably also do bad things if PSNode doesn't exist?
 				const FName PSName(TEXT("ParticleSys"));
 				FBlueprintEditorUtils::RenameComponentMemberVariable(BlueprintObject, PSNode, PSName);
 				UE_LOG(LogEditorBuildPromotionTests, Display, TEXT("Renamed the particle system component to ParticleSys"));
@@ -3384,7 +3187,7 @@ namespace BuildPromotionTestHelper
 					UE_LOG(LogEditorBuildPromotionTests, Error, TEXT("Failed to create a new function graph"));
 				}
 
-				AddParticleSystemNode = EditorBuildPromotionTestUtils::CreateAddComponentActionNode(BlueprintObject, CustomGraph, CreatedPS);
+				AddParticleSystemNode = EditorBuildPromotionTestUtils::CreateAddComponentActionNode(BlueprintObject, CustomGraph, LoadedParticleSystem);
 				if (AddParticleSystemNode)
 				{
 					UE_LOG(LogEditorBuildPromotionTests, Display, TEXT("Created an AddParticleSystem node"));
@@ -3454,7 +3257,7 @@ namespace BuildPromotionTestHelper
 				if (PlacedBlueprint)
 				{
 					//Set the text
-					EditorBuildPromotionTestUtils::SetPropertyByName(PlacedBlueprint, EditorBuildPromotionTestUtils::BlueprintStringVariableName.ToString(), TEXT("Print String works!"));
+					FEditorPromotionTestUtilities::SetPropertyByName(PlacedBlueprint, EditorBuildPromotionTestUtils::BlueprintStringVariableName.ToString(), TEXT("Print String works!"));
 					UE_LOG(LogEditorBuildPromotionTests, Display, TEXT("Placed the blueprint and set the display string to \"Print String works!\""));
 				}
 				else
@@ -4005,22 +3808,18 @@ bool FBuildPromotionSettingsTest::RunTest(const FString& Parameters)
 	EditorBuildPromotionTestUtils::ExportEditorSettings(TargetOriginalPreferenceFile);
 
 //New Editor Settings
-	UE_LOG(LogEditorBuildPromotionTests, Display, TEXT("Binding create empty layer shortcut"));
-
 	//Bind H to CreateEmptyLayer keybinding
-	FInputChord NewCreateChord(EKeys::H, EModifierKey::None);
-	EditorBuildPromotionTestUtils::SetEditorKeybinding(TEXT("LayersView"), TEXT("CreateEmptyLayer"), NewCreateChord);
-
-	UE_LOG(LogEditorBuildPromotionTests, Display, TEXT("Binding request rename layer shortcut"));
-
+	UE_LOG(LogEditorBuildPromotionTests, Display, TEXT("Binding create empty layer shortcut"));
+	FInputChord NewCreateChord = EditorBuildPromotionTestUtils::SetKeybinding(TEXT("LayersView"), TEXT("CreateEmptyLayer"), EKeys::H, EModifierKey::None);
+	
 	//Bind J to RequestRenameLayer
-	FInputChord NewRenameChord(EKeys::J, EModifierKey::None);
-	EditorBuildPromotionTestUtils::SetEditorKeybinding(TEXT("LayersView"), TEXT("RequestRenameLayer"), NewRenameChord);
-
+	UE_LOG(LogEditorBuildPromotionTests, Display, TEXT("Binding request rename layer shortcut"));
+	FInputChord NewRenameChord = EditorBuildPromotionTestUtils::SetKeybinding(TEXT("LayersView"), TEXT("RequestRenameLayer"), EKeys::J, EModifierKey::None);
+	
+	// Bind CTRL+L to PIE
 	UE_LOG(LogEditorBuildPromotionTests, Display, TEXT("Binding play shortcut (PIE)"));
-	FInputChord NewPIEChord(EKeys::L, EModifierKey::Control);
-	EditorBuildPromotionTestUtils::SetEditorKeybinding(TEXT("PlayWorld"), TEXT("RepeatLastPlay"), NewPIEChord);
-
+	FInputChord NewPIEChord = EditorBuildPromotionTestUtils::SetKeybinding(TEXT("PlayWorld"), TEXT("RepeatLastPlay"), EKeys::L, EModifierKey::Control);
+	
 	//Export the keybindings
 	const FString TargetKeybindFile = FString::Printf(TEXT("%s/BuildPromotion/Keybindings-%d.ini"), *FPaths::AutomationDir(), GEngineVersion.GetChangelist());
 	UE_LOG(LogEditorBuildPromotionTests, Display, TEXT("Exporting keybind"));
@@ -4029,7 +3828,7 @@ bool FBuildPromotionSettingsTest::RunTest(const FString& Parameters)
 
 	UEditorStyleSettings* EditorStyleSettings = GetMutableDefault<UEditorStyleSettings>();
 	FString OldStyleSetting = EditorBuildPromotionTestUtils::GetPropertyByName(EditorStyleSettings, TEXT("bUseSmallToolBarIcons"));
-	EditorBuildPromotionTestUtils::SetPropertyByName(EditorStyleSettings, TEXT("bUseSmallToolBarIcons"), TEXT("true"));
+	FEditorPromotionTestUtilities::SetPropertyByName(EditorStyleSettings, TEXT("bUseSmallToolBarIcons"), TEXT("true"));
 
 	UE_LOG(LogEditorBuildPromotionTests, Display, TEXT("Set UseSmallToolBarIcons"));
 
@@ -4043,29 +3842,12 @@ bool FBuildPromotionSettingsTest::RunTest(const FString& Parameters)
 	UE_LOG(LogEditorBuildPromotionTests, Display, TEXT("Exported keybindings and preferences to %s/BuildPromotion"), *FPaths::AutomationDir());
 
 	//Change the setting back
-	EditorBuildPromotionTestUtils::SetPropertyByName(EditorStyleSettings, TEXT("bUseSmallToolBarIcons"), OldStyleSetting);
+	FEditorPromotionTestUtilities::SetPropertyByName(EditorStyleSettings, TEXT("bUseSmallToolBarIcons"), OldStyleSetting);
 
-	FInputChord CurrentCreateChord;
-	EditorBuildPromotionTestUtils::GetEditorKeybinding(TEXT("LayersView"), TEXT("CreateEmptyLayer"), CurrentCreateChord);
-	if (CurrentCreateChord.Key == NewCreateChord.Key)
-	{
-		UE_LOG(LogEditorBuildPromotionTests, Display, TEXT("CreateEmptyLayer keybinding correct."));
-	}
-	else
-	{
-		UE_LOG(LogEditorBuildPromotionTests, Error, TEXT("CreateEmptyLayer keybinding incorrect."));
-	}
-
-	FInputChord CurrentRenameChord;
-	EditorBuildPromotionTestUtils::GetEditorKeybinding(TEXT("LayersView"), TEXT("RequestRenameLayer"), CurrentRenameChord);
-	if (CurrentRenameChord.Key == NewRenameChord.Key)
-	{
-		UE_LOG(LogEditorBuildPromotionTests, Display, TEXT("RequestRenameLayer keybinding correct."));
-	}
-	else
-	{
-		UE_LOG(LogEditorBuildPromotionTests, Error, TEXT("RequestRenameLayer keybinding incorrect."));
-	}
+	// Verify keybindings were assigned correctly
+	EditorBuildPromotionTestUtils::CompareKeybindings(TEXT("LayersView"), TEXT("CreateEmptyLayer"), NewCreateChord);
+	EditorBuildPromotionTestUtils::CompareKeybindings(TEXT("LayersView"), TEXT("RequestRenameLayer"), NewRenameChord);
+	EditorBuildPromotionTestUtils::CompareKeybindings(TEXT("PlayWorld"), TEXT("RepeatLastPlay"), NewPIEChord);
 
 	//Focus the main editor
 	TArray< TSharedRef<SWindow> > AllWindows;
diff --git a/Engine/Source/Editor/UnrealEd/Public/Tests/AutomationEditorPromotionCommon.h b/Engine/Source/Editor/UnrealEd/Public/Tests/AutomationEditorPromotionCommon.h
--- a/Engine/Source/Editor/UnrealEd/Public/Tests/AutomationEditorPromotionCommon.h
+++ b/Engine/Source/Editor/UnrealEd/Public/Tests/AutomationEditorPromotionCommon.h
@@ -10,6 +10,17 @@
 class FEditorPromotionTestUtilities
 {
 
+private:
+
+	/**
+	* Finds a visible widget by type.     SLOW!!!!!
+	*
+	* @param InParent - We search this widget and its children for a matching widget (recursive)
+	* @param InWidgetType - The widget type we are searching for
+	*/
+	static TSharedPtr<SWidget> FindFirstWidgetByClass(TSharedRef<SWidget> InParent, const FName& InWidgetType);
+
+
 public:
 
 	/**
@@ -24,4 +35,49 @@ public:
 	* @param InTexture - The texture to use as the diffuse for the new material
 	*/
 	UNREALED_API static UMaterial* CreateMaterialFromTexture(UTexture* InTexture);
+
+
+	/**
+	* Sets an editor keyboard shortcut
+	*
+	* @param CommandContext - The context of the command
+	* @param Command - The command name to set
+	* @param NewChord - The new input chord to assign
+	*/
+	UNREALED_API static bool SetEditorKeybinding(const FString& CommandContext, const FString& Command, const FInputChord& NewChord);
+
+	
+	/**
+	* Gets an editor keyboard shortcut
+	*
+	* @param CommandContext - The context of the command
+	* @param Command - The command name to get
+	*/
+	UNREALED_API static FInputChord GetEditorKeybinding(const FString& CommandContext, const FString& Command); //, FInputChord& CurrentChord);
+
+	/**
+	* Gets the current input chord or sets a new one if it doesn't exist
+	*
+	* @param Context - The context of the UI Command
+	* @param Command - The name of the UI command
+	*/
+	UNREALED_API static FInputChord GetOrSetUICommand(const FString& Context, const FString& Command);
+
+
+	/**
+	* Sends a UI command to the active top level window after focusing on a widget of a given type
+	*
+	* @param InChord - The chord to send to the window
+	* @param WidgetTypeToFocus - The widget type to find and focus on
+	*/
+	UNREALED_API static void SendCommandToCurrentEditor(const FInputChord& InChord, const FName& WidgetTypeToFocus);
+
+
+	/**
+	* Sets an object property value by name
+	*
+	* @param TargetObject - The object to modify
+	* @param InVariableName - The name of the property
+	*/
+	UNREALED_API static void SetPropertyByName(UObject* TargetObject, const FString& InVariableName, const FString& NewValueString);
 };
\ No newline at end of file
diff --git a/Engine/Source/Runtime/Engine/Classes/Tests/AutomationTestSettings.h b/Engine/Source/Runtime/Engine/Classes/Tests/AutomationTestSettings.h
--- a/Engine/Source/Runtime/Engine/Classes/Tests/AutomationTestSettings.h
+++ b/Engine/Source/Runtime/Engine/Classes/Tests/AutomationTestSettings.h
@@ -250,19 +250,33 @@ struct FMaterialEditorPromotionSettings
 	GENERATED_USTRUCT_BODY()
 
 	/** Default material asset to apply to static meshes **/
-	UPROPERTY(EditAnywhere, Category = Automation)
+	UPROPERTY(EditAnywhere, Category = Automation, meta = (FilePathFilter = "uasset"))
 	FFilePath DefaultMaterialAsset;
 
 	/** Default material asset to apply to static meshes **/
-	UPROPERTY(EditAnywhere, Category = Automation)
+	UPROPERTY(EditAnywhere, Category = Automation, meta = (FilePathFilter = "uasset"))
 	FFilePath DefaultDiffuseTexture;
 
 	/** Default material asset to apply to static meshes **/
-	UPROPERTY(EditAnywhere, Category = Automation)
+	UPROPERTY(EditAnywhere, Category = Automation, meta = (FilePathFilter = "uasset"))
 	FFilePath DefaultNormalTexture;
 
 };
 
+
+/**
+* Holds settings for the particle editor build promotion tests
+*/
+USTRUCT()
+struct FParticleEditorPromotionSettings
+{
+	GENERATED_USTRUCT_BODY()
+
+	/** Default particle asset to use for tests*/
+	UPROPERTY(EditAnywhere, Category = Automation, meta = (FilePathFilter = "uasset"))
+	FFilePath DefaultParticleAsset;
+};
+
 /**
 * Holds settings for the editor build promotion test
 */
@@ -364,6 +378,12 @@ public:
 	UPROPERTY(EditAnywhere, config, Category = Automation)
 	FMaterialEditorPromotionSettings MaterialEditorPromotionTest;
 
+	/**
+	* Particle editor promotion test settings
+	*/
+	UPROPERTY(EditAnywhere, config, Category = Automation)
+	FParticleEditorPromotionSettings ParticleEditorPromotionTest;
+
 	/**
 	* Modules to load that have engine tests
 	*/
