Don't use copied asset data within manifest generator to find chunk IDs. Added separate look up. Fixes UE-16053

diff --git a/Engine/Source/Editor/UnrealEd/Private/Commandlets/ChunkManifestGenerator.cpp b/Engine/Source/Editor/UnrealEd/Private/Commandlets/ChunkManifestGenerator.cpp
--- a/Engine/Source/Editor/UnrealEd/Private/Commandlets/ChunkManifestGenerator.cpp
+++ b/Engine/Source/Editor/UnrealEd/Private/Commandlets/ChunkManifestGenerator.cpp
@@ -193,6 +193,12 @@ void FChunkManifestGenerator::Initialize(bool InGenerateChunks)
 			{
 				RegistryChunkIDs.Add(ChunkID);
 			}
+			auto* FoundIDList = PackageChunkIDMap.Find(AssetData.PackageName);
+			if (!FoundIDList)
+			{
+				FoundIDList = &PackageChunkIDMap.Add(AssetData.PackageName);
+			}
+			FoundIDList->AddUnique(ChunkID);
 		}
 		// Now clear the original chunk id list. We will fill it with real IDs when cooking.
 		AssetData.ChunkIDs.Empty();
diff --git a/Engine/Source/Editor/UnrealEd/Private/Commandlets/ChunkManifestGenerator.h b/Engine/Source/Editor/UnrealEd/Private/Commandlets/ChunkManifestGenerator.h
--- a/Engine/Source/Editor/UnrealEd/Private/Commandlets/ChunkManifestGenerator.h
+++ b/Engine/Source/Editor/UnrealEd/Private/Commandlets/ChunkManifestGenerator.h
@@ -22,6 +22,8 @@ class FChunkManifestGenerator
 	TSet<FName> AssetsLoadedWithLastPackage;
 	/** The entire asset registry data */
 	TArray<FAssetData> AssetRegistryData;
+	/** Lookup for the original ChunkID Mappings. AssetRegistryData is modified making it invalid to query */
+	TMap<FName, TArray<int32> > PackageChunkIDMap;
 	/** Maps packages to assets from the asset registry */
 	TMap<FName, TArray<int32> > PackageToRegistryDataMap;
 	/** Should the chunks be generated or only asset registry */
@@ -167,14 +169,11 @@ class FChunkManifestGenerator
 	*/
 	FORCEINLINE TArray<int32> GetAssetRegistryChunkAssignments(const FName& PackageFName)
 	{
-		//get the objects in this package
 		TArray<int32> RegistryChunkIDs;
-		for ( const auto& AssetData : AssetRegistryData )
+		auto* FoundIDs = PackageChunkIDMap.Find(PackageFName);
+		if (FoundIDs)
 		{
-			if ( AssetData.PackageName == PackageFName ) 
-			{
-				RegistryChunkIDs.Append( AssetData.ChunkIDs );
-			}
+			RegistryChunkIDs = *FoundIDs;
 		}
 		return RegistryChunkIDs;
 	}
