[UE-15868] Fix a potential crash when copying an array of USTRUCTs from one Actor instance to another.

change summary:
- modified EditorUtilities::CopySinglePropertyRecursive() to ensure that target Struct properties are initialized prior to copying fields

diff --git a/Engine/Source/Editor/UnrealEd/Private/Editor.cpp b/Engine/Source/Editor/UnrealEd/Private/Editor.cpp
--- a/Engine/Source/Editor/UnrealEd/Private/Editor.cpp
+++ b/Engine/Source/Editor/UnrealEd/Private/Editor.cpp
@@ -812,6 +812,9 @@ namespace EditorUtilities
 		}
 		else if (UStructProperty* const StructProperty = Cast<UStructProperty>(InProperty))
 		{
+			// Ensure that the target struct is initialized before copying fields from the source.
+			StructProperty->InitializeValue_InContainer(InTargetPtr);
+
 			const int32 PropertyArrayDim = InProperty->ArrayDim;
 			for (int32 ArrayIndex = 0; ArrayIndex < PropertyArrayDim; ArrayIndex++)
 			{
