Fix hitches with ComponentTypeRegistry attempting to load already loaded objects

diff --git a/Engine/Source/Editor/UnrealEd/Private/ComponentTypeRegistry.cpp b/Engine/Source/Editor/UnrealEd/Private/ComponentTypeRegistry.cpp
--- a/Engine/Source/Editor/UnrealEd/Private/ComponentTypeRegistry.cpp
+++ b/Engine/Source/Editor/UnrealEd/Private/ComponentTypeRegistry.cpp
@@ -48,6 +48,18 @@ static const FString CommonClassGroup(TEXT("Common"));
 // This has to stay in sync with logic in FKismetCompilerContext::FinishCompilingClass
 static const FString BlueprintComponents(TEXT("Custom"));
 
+template <typename ObjectType>
+static ObjectType* FindOrLoadObject( const FString& ObjectPath )
+{
+	ObjectType* Object = FindObject<ObjectType>( nullptr, *ObjectPath );
+	if( !Object )
+	{
+		Object = LoadObject<ObjectType>( nullptr, *ObjectPath );
+	}
+
+	return Object;
+}
+
 void FComponentTypeRegistryData::AddBasicShapeComponents(TArray<FComponentClassComboEntryPtr>& SortedClassList)
 {
 	FString BasicShapesHeading = LOCTEXT("BasicShapesHeading", "Basic Shapes").ToString();
@@ -57,13 +69,14 @@ void FComponentTypeRegistryData::AddBasicShapeComponents(TArray<FComponentClassC
 		UStaticMeshComponent* SMC = Cast<UStaticMeshComponent>(Component);
 		if (SMC)
 		{
-			SMC->SetMaterial(0, LoadObject<UMaterial>(nullptr, TEXT("/Engine/BasicShapes/BasicShapeMaterial.BasicShapeMaterial")));
+			const FString MaterialName = TEXT("/Engine/BasicShapes/BasicShapeMaterial.BasicShapeMaterial");
+			SMC->SetMaterial(0, FindOrLoadObject<UMaterial>(MaterialName));
 		}
 	};
 
 	{
 		FComponentEntryCustomizationArgs Args;
-		Args.AssetOverride = LoadObject<UStaticMesh>(nullptr, *UActorFactoryBasicShape::BasicCube.ToString());
+		Args.AssetOverride = FindOrLoadObject<UStaticMesh>(UActorFactoryBasicShape::BasicCube.ToString());
 		Args.OnComponentCreated = FOnComponentCreated::CreateStatic(OnBasicShapeCreated);
 		Args.ComponentNameOverride = LOCTEXT("BasicCubeShapeDisplayName", "Cube").ToString();
 		Args.IconOverrideBrushName = FName("ClassIcon.Cube");
@@ -83,7 +96,7 @@ void FComponentTypeRegistryData::AddBasicShapeComponents(TArray<FComponentClassC
 
 	{
 		FComponentEntryCustomizationArgs Args;
-		Args.AssetOverride = LoadObject<UStaticMesh>(nullptr, *UActorFactoryBasicShape::BasicSphere.ToString());
+		Args.AssetOverride = FindOrLoadObject<UStaticMesh>(UActorFactoryBasicShape::BasicSphere.ToString());
 		Args.OnComponentCreated = FOnComponentCreated::CreateStatic(OnBasicShapeCreated);
 		Args.ComponentNameOverride = LOCTEXT("BasicSphereShapeDisplayName", "Sphere").ToString();
 		Args.IconOverrideBrushName = FName("ClassIcon.Sphere");
@@ -102,7 +115,7 @@ void FComponentTypeRegistryData::AddBasicShapeComponents(TArray<FComponentClassC
 
 	{
 		FComponentEntryCustomizationArgs Args;
-		Args.AssetOverride = LoadObject<UStaticMesh>(nullptr, *UActorFactoryBasicShape::BasicCylinder.ToString());
+		Args.AssetOverride = FindOrLoadObject<UStaticMesh>(UActorFactoryBasicShape::BasicCylinder.ToString());
 		Args.OnComponentCreated = FOnComponentCreated::CreateStatic(OnBasicShapeCreated);
 		Args.ComponentNameOverride = LOCTEXT("BasicCylinderShapeDisplayName", "Cylinder").ToString();
 		Args.IconOverrideBrushName = FName("ClassIcon.Cylinder");
@@ -113,7 +126,7 @@ void FComponentTypeRegistryData::AddBasicShapeComponents(TArray<FComponentClassC
 
 	{
 		FComponentEntryCustomizationArgs Args;
-		Args.AssetOverride = LoadObject<UStaticMesh>(nullptr, *UActorFactoryBasicShape::BasicCone.ToString());
+		Args.AssetOverride = FindOrLoadObject<UStaticMesh>(UActorFactoryBasicShape::BasicCone.ToString());
 		Args.OnComponentCreated = FOnComponentCreated::CreateStatic(OnBasicShapeCreated);
 		Args.ComponentNameOverride = LOCTEXT("BasicConeShapeDisplayName", "Cone").ToString();
 		Args.IconOverrideBrushName = FName("ClassIcon.Cone");
