PR #1148: Fix divide by zero when importing 0-length animations (Contributed by [REDACTED_PERSON])

Unshelved from pending changelist '2552016'

diff --git a/Engine/Source/Editor/UnrealEd/Private/SkeletalMeshEdit.cpp b/Engine/Source/Editor/UnrealEd/Private/SkeletalMeshEdit.cpp
--- a/Engine/Source/Editor/UnrealEd/Private/SkeletalMeshEdit.cpp
+++ b/Engine/Source/Editor/UnrealEd/Private/SkeletalMeshEdit.cpp
@@ -1147,7 +1147,7 @@ bool UnFbx::FFbxImporter::ImportAnimation(USkeleton* Skeleton, UAnimSequence * D
 		FMatrix AddedMatrix = Converter.ConvertMatrix(FbxAddedMatrix);
 
 		const int32 NumSamplingKeys = FMath::FloorToInt(AnimTimeSpan.GetDuration().GetSecondDouble() * ResampleRate);
-		const FbxTime TimeIncrement = AnimTimeSpan.GetDuration() / NumSamplingKeys;
+		const FbxTime TimeIncrement = AnimTimeSpan.GetDuration() / FMath::Max(NumSamplingKeys,1);
 		for(int32 SourceTrackIdx = 0; SourceTrackIdx < FbxRawBoneNames.Num(); ++SourceTrackIdx)
 		{
 			int32 NumKeysForTrack = 0;
@@ -1298,4 +1298,4 @@ bool UnFbx::FFbxImporter::ImportAnimation(USkeleton* Skeleton, UAnimSequence * D
 	return true;
 }
 
-#undef LOCTEXT_NAMESPACE
\ No newline at end of file
+#undef LOCTEXT_NAMESPACE
