Removed some of the lighting tests from the big promotion pass.

diff --git a/Engine/Source/Editor/UnrealEd/Private/Tests/EditorBuildPromotionTests.cpp b/Engine/Source/Editor/UnrealEd/Private/Tests/EditorBuildPromotionTests.cpp
--- a/Engine/Source/Editor/UnrealEd/Private/Tests/EditorBuildPromotionTests.cpp
+++ b/Engine/Source/Editor/UnrealEd/Private/Tests/EditorBuildPromotionTests.cpp
@@ -921,7 +921,7 @@ namespace EditorBuildPromotionTestUtils
 		}
 		else
 		{
-			UE_LOG(LogEditorBuildPromotionTests, Error, TEXT("Failed to apply material to static mesh because mesh does not exist"));
+			UE_LOG(LogEditorBuildPromotionTests, Warning, TEXT("Failed to apply material to static mesh because mesh does not exist"));
 			return false;
 		}
 	}
@@ -1530,9 +1530,6 @@ namespace BuildPromotionTestHelper
 			ADD_TEST_STAGE(EndSection,						TEXT("Geometry Workflow"));
 
 			// 3) Lighting
-			ADD_TEST_STAGE(Lighting_PlaceAPointLight,		TEXT("Place a Point Light"));
-			ADD_TEST_STAGE(Lighting_PointLightProperties,	TEXT("Point Light Properties"));
-			ADD_TEST_STAGE(Lighting_DuplicatePointLight,	TEXT("Duplicate Point Light"));
 			ADD_TEST_STAGE(Lighting_BuildLighting_Part1,	TEXT("Build Lighting"));
 			ADD_TEST_STAGE(Lighting_BuildLighting_Part2,	TEXT("Build Lighting"));
 			ADD_TEST_STAGE(EndSection, TEXT("Lighting Workflow"));
@@ -1746,111 +1743,6 @@ namespace BuildPromotionTestHelper
 			return true;
 		}
 		
-		/**
-		* Lighting Test Stage: Place a point Light
-		*    Adds a point light to the map and sets its location, rotation, and scale
-		*/
-		bool Lighting_PlaceAPointLight()
-		{
-			const FTransform Transform(FVector(0.0f, 50.0f, 400.0f));
-			PointLight = Cast<APointLight>(GEditor->AddActor(CurrentWorld->GetCurrentLevel(), APointLight::StaticClass(), Transform));
-			if (PointLight)
-			{
-				UE_LOG(LogEditorBuildPromotionTests, Display, TEXT("Placed a PointLight"));
-				PointLight->SetActorLocation(FVector(100.f, 50.f, 200.f));
-				UE_LOG(LogEditorBuildPromotionTests, Display, TEXT("Modified the PointLight's Location"));
-				PointLight->SetActorRotation(FRotator(300, 250, -91));
-				UE_LOG(LogEditorBuildPromotionTests, Display, TEXT("Modified the PointLight's Rotation"));
-				PointLight->SetActorScale3D(FVector(2.f, 2.f, 2.f));
-				UE_LOG(LogEditorBuildPromotionTests, Display, TEXT("Modified the PointLight's Scale"));
-			}
-			else
-			{
-				UE_LOG(LogEditorBuildPromotionTests, Error, TEXT("Error creating point light"));
-			}
-
-			return true;
-		}
-
-		/**
-		* Lighting Test Stage: Point Light Properties
-		*    Modifies the point light properties (Intensity, LightColor, and AttenuationRadius)
-		*/
-		bool Lighting_PointLightProperties()
-		{
-			if (PointLight)
-			{
-				UPointLightComponent* PointLightComp = PointLight->PointLightComponent;
-				FEditorPromotionTestUtilities::SetPropertyByName(PointLightComp, TEXT("Intensity"), TEXT("1000.f"));
-				UE_LOG(LogEditorBuildPromotionTests, Display, TEXT("Modified the PointLight's Intensity"));
-				FEditorPromotionTestUtilities::SetPropertyByName(PointLightComp, TEXT("LightColor"), TEXT("(R=0,G=0,B=255)"));
-				UE_LOG(LogEditorBuildPromotionTests, Display, TEXT("Modified the PointLight's LightColor"));
-				FEditorPromotionTestUtilities::SetPropertyByName(PointLightComp, TEXT("AttenuationRadius"), TEXT("1024.f"));
-				UE_LOG(LogEditorBuildPromotionTests, Display, TEXT("Modified the PointLight's AttenuationRadius"));
-			}
-			else
-			{
-				SkippedTests.Add(TEXT("Lighting: Point Light Properties. (Failed to create light)"));
-			}
-
-			return true;
-		}
-
-		/**
-		* Lighting Test Stage: Duplicate point light
-		*    Duplicates the point light created earlier using Copy/Paste and Duplicate
-		*/
-		bool Lighting_DuplicatePointLight()
-		{
-			if (PointLight)
-			{
-				int32 StartingLightCount = EditorBuildPromotionTestUtils::GetNumActors(CurrentWorld, APointLight::StaticClass());
-
-				//Select the light
-				GEditor->SelectNone(false, true, false);
-				GEditor->SelectActor(PointLight, true, false, true);
-
-				GEngine->Exec(CurrentWorld, TEXT("EDIT COPY"));
-				GEngine->Exec(CurrentWorld, TEXT("EDIT PASTE"));
-
-				//Check if we gained a new light
-				int32 CurrentLightCount = EditorBuildPromotionTestUtils::GetNumActors(CurrentWorld, APointLight::StaticClass());
-				if (CurrentLightCount - StartingLightCount == 1)
-				{
-					UE_LOG(LogEditorBuildPromotionTests, Display, TEXT("Duplicated the PointLight using copy / paste"));
-				}
-				else
-				{
-					UE_LOG(LogEditorBuildPromotionTests, Error, TEXT("Failed to duplicate light using copy / paste"));
-				}
-
-				//Reset the starting count
-				StartingLightCount = CurrentLightCount;
-
-				//Select the light
-				GEditor->SelectNone(false, true, false);
-				GEditor->SelectActor(PointLight, true, false, true);
-
-				GEngine->Exec(CurrentWorld, TEXT("DUPLICATE"));
-
-				//Check if we gained a new light
-				CurrentLightCount = EditorBuildPromotionTestUtils::GetNumActors(CurrentWorld, APointLight::StaticClass());
-				if (CurrentLightCount - StartingLightCount == 1)
-				{
-					UE_LOG(LogEditorBuildPromotionTests, Display, TEXT("Duplicated the PointLight using duplicate"));
-				}
-				else
-				{
-					UE_LOG(LogEditorBuildPromotionTests, Error, TEXT("Failed to duplicate light using duplicate"));
-				}
-			}
-			else
-			{
-				SkippedTests.Add(TEXT("Lighting: Duplicate point light. (Failed to create light)"));
-			}
-
-			return true;
-		}
 
 		/**
 		* Lighting Test Stage: Build Lighting (Part 1)
