Fixed duplicated brushes not appearing in Game

 https://jira.ol.epicgames.net/browse/UE-16648

diff --git a/Engine/Source/Editor/UnrealEd/Private/EditorObject.cpp b/Engine/Source/Editor/UnrealEd/Private/EditorObject.cpp
--- a/Engine/Source/Editor/UnrealEd/Private/EditorObject.cpp
+++ b/Engine/Source/Editor/UnrealEd/Private/EditorObject.cpp
@@ -620,7 +620,7 @@ static const TCHAR* ImportProperties(
 	}
 
 	// Prepare brush.
-	if( ImportedBrush && ObjectStruct->IsChildOf(ABrush::StaticClass()) )
+	if( ImportedBrush && ObjectStruct->IsChildOf<ABrush>() && !ObjectStruct->IsChildOf<AVolume>() )
 	{
 		check(GIsEditor);
 		ABrush* Actor = (ABrush*)DestData;
