Integrate - Fixed crash building lighting due to zero area BSP nodes creating 0 resolution texture mappings

diff --git a/Engine/Source/Editor/UnrealEd/Private/Lightmass/Lightmass.cpp b/Engine/Source/Editor/UnrealEd/Private/Lightmass/Lightmass.cpp
--- a/Engine/Source/Editor/UnrealEd/Private/Lightmass/Lightmass.cpp
+++ b/Engine/Source/Editor/UnrealEd/Private/Lightmass/Lightmass.cpp
@@ -1360,6 +1360,7 @@ void FLightmassExporter::WriteBaseTextureMappingData( int32 Channel, const FStat
 	
 	Lightmass::FStaticLightingTextureMappingData TextureMappingData;
 	FMemory::Memzero(&TextureMappingData,sizeof(TextureMappingData));
+	check(TextureMapping->SizeX > 0 && TextureMapping->SizeY > 0);
 	TextureMappingData.SizeX = TextureMapping->SizeX;
 	TextureMappingData.SizeY = TextureMapping->SizeY;
 	TextureMappingData.LightmapTextureCoordinateIndex = TextureMapping->LightmapTextureCoordinateIndex;
diff --git a/Engine/Source/Editor/UnrealEd/Private/StaticLightingSystem/StaticLightingSystem.cpp b/Engine/Source/Editor/UnrealEd/Private/StaticLightingSystem/StaticLightingSystem.cpp
--- a/Engine/Source/Editor/UnrealEd/Private/StaticLightingSystem/StaticLightingSystem.cpp
+++ b/Engine/Source/Editor/UnrealEd/Private/StaticLightingSystem/StaticLightingSystem.cpp
@@ -1369,6 +1369,11 @@ void FStaticLightingSystem::AddBSPStaticLightingInfo(ULevel* Level, bool bBuildL
 
 			// fill out the NodeGroup/mapping, as UModelComponent::GetStaticLightingInfo did
 			SomeModelComponent->GetSurfaceLightMapResolution(SurfaceIndex, true, NodeGroup->SizeX, NodeGroup->SizeY, NodeGroup->WorldToMap, &NodeGroup->Nodes);
+
+			// Make sure mapping will have valid size
+			NodeGroup->SizeX = FMath::Max(NodeGroup->SizeX, 1);
+			NodeGroup->SizeY = FMath::Max(NodeGroup->SizeY, 1);
+
 			NodeGroup->MapToWorld = NodeGroup->WorldToMap.InverseFast();
 
 			// Cache the surface's vertices and triangles.
