Put in missing #undef LOCTEXT_NAMESPACE

diff --git a/Engine/Source/Editor/UnrealEd/Private/Tests/LightingTests.cpp b/Engine/Source/Editor/UnrealEd/Private/Tests/LightingTests.cpp
--- a/Engine/Source/Editor/UnrealEd/Private/Tests/LightingTests.cpp
+++ b/Engine/Source/Editor/UnrealEd/Private/Tests/LightingTests.cpp
@@ -305,4 +305,6 @@ bool FLightPointLightSetLocation::RunTest(const FString& Parameters)
 
 	AddError(TEXT("A point light already exists in this level which will block the verification of a new point light."));
 	return false;
-}
\ No newline at end of file
+}
+
+#undef LOCTEXT_NAMESPACE
\ No newline at end of file
