Force mark all PIE world objects pending kill to avoid failure of GC to clean up PIE worlds

diff --git a/Engine/Source/Editor/UnrealEd/Private/PlayLevel.cpp b/Engine/Source/Editor/UnrealEd/Private/PlayLevel.cpp
--- a/Engine/Source/Editor/UnrealEd/Private/PlayLevel.cpp
+++ b/Engine/Source/Editor/UnrealEd/Private/PlayLevel.cpp
@@ -158,15 +158,31 @@ void UEditorEngine::EndPlayMap()
 	}
 	CleanupGameViewport();
 
+	// Gather all the worlds relevant to the PIE levels we are ending
+	TSet<UWorld*> CurrentPlayWorlds;
+	for (const FWorldContext& ThisContext : WorldList)
+	{
+		if (ThisContext.WorldType == EWorldType::PIE && ThisContext.World())
+		{
+			for (auto LevelIt(ThisContext.World()->GetLevelIterator()); LevelIt; ++LevelIt)
+			{
+				if (const ULevel* Level = *LevelIt)
+				{
+					CurrentPlayWorlds.Add(CastChecked<UWorld>(Level->GetOuter()));
+				}
+			}
+		}
+	}
 
 	// find objects like Textures in the playworld levels that won't get garbage collected as they are marked RF_Standalone
+	// mark everything else contained in the PIE worlds to be deleted
 	for( FObjectIterator It; It; ++It )
 	{
 		UObject* Object = *It;
 
-		if((Object->GetOutermost()->PackageFlags & PKG_PlayInEditor) != 0)
+		if ((Object->GetOutermost()->PackageFlags & PKG_PlayInEditor) != 0)
 		{
-			if(Object->HasAnyFlags(RF_Standalone))
+			if (Object->HasAnyFlags(RF_Standalone))
 			{
 				// Clear RF_Standalone flag from objects in the levels used for PIE so they get cleaned up.
 				Object->ClearFlags(RF_Standalone);
@@ -174,6 +190,14 @@ void UEditorEngine::EndPlayMap()
 			// Close any asset editors that are currently editing this object
 			FAssetEditorManager::Get().CloseAllEditorsForAsset(Object);
 		}
+		else if (!Object->IsPendingKill())
+		{
+			UWorld* InWorld = Object->GetTypedOuter<UWorld>();
+			if (InWorld && CurrentPlayWorlds.Contains(InWorld))
+			{
+				Object->MarkPendingKill();
+			}
+		}
 	}
 
 	// Clean up each world individually
