Fixing a crash due to the TestBase not being set for the linetrace test.

diff --git a/Engine/Source/Editor/UnrealEd/Private/Tests/CollisionAutomationTests.cpp b/Engine/Source/Editor/UnrealEd/Private/Tests/CollisionAutomationTests.cpp
--- a/Engine/Source/Editor/UnrealEd/Private/Tests/CollisionAutomationTests.cpp
+++ b/Engine/Source/Editor/UnrealEd/Private/Tests/CollisionAutomationTests.cpp
@@ -238,6 +238,7 @@ bool FLineTraceSingleByChannel::RunTest(const FString& Parameters)
 	// Create map
 	UWorld* World = AutomationEditorCommonUtils::CreateNewMap();
 	TestNotNull(TEXT("Failed to create world for Physics.Collision.Ray Test. Tests aborted."), World);
+	CollisionAutomationTests::TestBase = this;
 
 	static FName TraceIdent = FName(TEXT("TestTrace"));
 
