Exposed time dilation in PhAT

#UE-16714

diff --git a/Engine/Source/Editor/PhAT/Private/PhATPreviewViewportClient.cpp b/Engine/Source/Editor/PhAT/Private/PhATPreviewViewportClient.cpp
--- a/Engine/Source/Editor/PhAT/Private/PhATPreviewViewportClient.cpp
+++ b/Engine/Source/Editor/PhAT/Private/PhATPreviewViewportClient.cpp
@@ -680,7 +680,7 @@ void FPhATEdPreviewViewportClient::Tick(float DeltaSeconds)
 		SharedData->EditorSkelComp->SetPhysicsBlendWeight(SharedData->EditorSimOptions->PhysicsBlend);
 	}
 
-	World->Tick(LEVELTICK_All, DeltaSeconds);
+	World->Tick(LEVELTICK_All, DeltaSeconds * SharedData->EditorSimOptions->TimeDilation);
 
 	if(SharedData->Recorder.InRecording())
 	{
diff --git a/Engine/Source/Editor/UnrealEd/Classes/Preferences/PhATSimOptions.h b/Engine/Source/Editor/UnrealEd/Classes/Preferences/PhATSimOptions.h
--- a/Engine/Source/Editor/UnrealEd/Classes/Preferences/PhATSimOptions.h
+++ b/Engine/Source/Editor/UnrealEd/Classes/Preferences/PhATSimOptions.h
@@ -32,6 +32,10 @@ class UNREALED_API UPhATSimOptions : public UObject
 	UPROPERTY(EditAnywhere, config, Category = Simulation)
 	int32 MaxFPS;
 
+	/** Dilate time by scale*/
+	UPROPERTY(EditAnywhere, config, Category = Simulation, meta = (UIMin = 0) )
+	float TimeDilation;
+
 	/** Linear damping of mouse spring forces */
 	UPROPERTY(EditAnywhere, config, Category=MouseSpring)
 	float HandleLinearDamping;
diff --git a/Engine/Source/Editor/UnrealEd/Private/PreferenceStubs.cpp b/Engine/Source/Editor/UnrealEd/Private/PreferenceStubs.cpp
--- a/Engine/Source/Editor/UnrealEd/Private/PreferenceStubs.cpp
+++ b/Engine/Source/Editor/UnrealEd/Private/PreferenceStubs.cpp
@@ -22,6 +22,7 @@ UPhATSimOptions::UPhATSimOptions(const FObjectInitializer& ObjectInitializer)
 {
 	PhysicsBlend = 1.0f;
 	MaxFPS = -1;
+	TimeDilation = 1.f;
 }
 
 UMaterialEditorOptions::UMaterialEditorOptions(const FObjectInitializer& ObjectInitializer)
