Fixed create anim BP option in content browser

diff --git a/Engine/Source/Editor/UnrealEd/Private/AnimationEditorUtils.cpp b/Engine/Source/Editor/UnrealEd/Private/AnimationEditorUtils.cpp
--- a/Engine/Source/Editor/UnrealEd/Private/AnimationEditorUtils.cpp
+++ b/Engine/Source/Editor/UnrealEd/Private/AnimationEditorUtils.cpp
@@ -329,18 +329,19 @@ namespace AnimationEditorUtils
 	{
 		MenuBuilder.BeginSection("CreateAnimAssets", LOCTEXT("CreateAnimAssetsMenuHeading", "Anim Assets"));
 		{
-// commented out for UE-13743
-/*
-			MenuBuilder.AddMenuEntry(
-				LOCTEXT("Skeleton_NewAnimBlueprint", "Anim Blueprint"),
-				LOCTEXT("Skeleton_NewAnimBlueprintTooltip", "Creates an Anim Blueprint using the selected skeleton."),
-				FSlateIcon(FEditorStyle::GetStyleSetName(), "ClassIcon.AnimBlueprint"),
-				FUIAction(
-					FExecuteAction::CreateStatic(&CreateNewAnimBlueprint, Skeletons, AssetCreated),
-					FCanExecuteAction()
-					)
-				);
-*/
+			// only allow for content browser until we support multi assets so we can open new persona with this BP
+			if (bInContentBrowser)
+			{
+				MenuBuilder.AddMenuEntry(
+					LOCTEXT("Skeleton_NewAnimBlueprint", "Anim Blueprint"),
+					LOCTEXT("Skeleton_NewAnimBlueprintTooltip", "Creates an Anim Blueprint using the selected skeleton."),
+					FSlateIcon(FEditorStyle::GetStyleSetName(), "ClassIcon.AnimBlueprint"),
+					FUIAction(
+						FExecuteAction::CreateStatic(&CreateNewAnimBlueprint, Skeletons, AssetCreated),
+						FCanExecuteAction()
+						)
+					);
+			}
 
 			MenuBuilder.AddMenuEntry(
 				LOCTEXT("Skeleton_NewAnimComposite", "Anim Composite"),
