Fix retarget base pose getting reset on reimport

UE-16468

diff --git a/Engine/Source/Editor/UnrealEd/Private/Factories/SkeletalMeshImport.cpp b/Engine/Source/Editor/UnrealEd/Private/Factories/SkeletalMeshImport.cpp
--- a/Engine/Source/Editor/UnrealEd/Private/Factories/SkeletalMeshImport.cpp
+++ b/Engine/Source/Editor/UnrealEd/Private/Factories/SkeletalMeshImport.cpp
@@ -703,6 +703,7 @@ ExistingSkelMeshData* SaveExistingSkelMeshData(USkeletalMesh* ExistingSkelMesh)
 
 		ExistingMeshDataPtr->ExistingSockets = ExistingSkelMesh->GetMeshOnlySocketList();
 		ExistingMeshDataPtr->ExistingMaterials = ExistingSkelMesh->Materials;
+		ExistingMeshDataPtr->ExistingRetargetBasePose = ExistingSkelMesh->RetargetBasePose;
 
 		if( ImportedResource->LODModels.Num() > 0 &&
 			ExistingSkelMesh->LODInfo.Num() == ImportedResource->LODModels.Num() )
@@ -779,6 +780,13 @@ void RestoreExistingSkelMeshData(ExistingSkelMeshData* MeshData, USkeletalMesh*
 
 		SkeletalMesh->Materials = MeshData->ExistingMaterials;
 
+		// this is not ideal. Ideally we'll have to save only diff with indicating which joints, 
+		// but for now, we allow them to keep the previous pose IF the element count is same
+		if (MeshData->ExistingRetargetBasePose.Num() == SkeletalMesh->RefSkeleton.GetNum())
+		{
+			SkeletalMesh->RetargetBasePose = MeshData->ExistingRetargetBasePose;
+		}
+
 		// Assign sockets from old version of this SkeletalMesh.
 		// Only copy ones for bones that exist in the new mesh.
 		for(int32 i=0; i<MeshData->ExistingSockets.Num(); i++)
diff --git a/Engine/Source/Editor/UnrealEd/Public/SkelImport.h b/Engine/Source/Editor/UnrealEd/Public/SkelImport.h
--- a/Engine/Source/Editor/UnrealEd/Public/SkelImport.h
+++ b/Engine/Source/Editor/UnrealEd/Public/SkelImport.h
@@ -130,6 +130,7 @@ struct ExistingSkelMeshData
 	TArray<UMorphTarget*>					ExistingMorphTargets;
 	TArray<UPhysicsAsset*>					ExistingPhysicsAssets;
 	USkeleton*								ExistingSkeleton;
+	TArray<FTransform>						ExistingRetargetBasePose;
 
 	bool									bExistingUseFullPrecisionUVs;
 
