Fixed broken legacy support for SourceFileSearchFilters setting in GatherTextFromSource steps in localization configuration scripts.

diff --git a/Engine/Source/Editor/UnrealEd/Private/Commandlets/GatherTextFromSourceCommandlet.cpp b/Engine/Source/Editor/UnrealEd/Private/Commandlets/GatherTextFromSourceCommandlet.cpp
--- a/Engine/Source/Editor/UnrealEd/Private/Commandlets/GatherTextFromSourceCommandlet.cpp
+++ b/Engine/Source/Editor/UnrealEd/Private/Commandlets/GatherTextFromSourceCommandlet.cpp
@@ -101,7 +101,7 @@ int32 UGatherTextFromSourceCommandlet::Main( const FString& Params )
 	// SourceFileSearchFilters (DEPRECATED)
 	{
 		TArray<FString> SourceFileSearchFilters;
-		GetPathArrayFromConfig(*SectionName, TEXT("SourceFileSearchFilters"), SourceFileSearchFilters, GatherTextConfigPath);
+		GetStringArrayFromConfig(*SectionName, TEXT("SourceFileSearchFilters"), SourceFileSearchFilters, GatherTextConfigPath);
 		if (SourceFileSearchFilters.Num())
 		{
 			FileNameFilters.Append(SourceFileSearchFilters);
