Added "source location"s as extracted comments in PO files exported from localization archives, due to OneSky not exposing them when they are reference comments.

diff --git a/Engine/Source/Editor/UnrealEd/Private/Commandlets/InternationalizationExportCommandlet.cpp b/Engine/Source/Editor/UnrealEd/Private/Commandlets/InternationalizationExportCommandlet.cpp
--- a/Engine/Source/Editor/UnrealEd/Private/Commandlets/InternationalizationExportCommandlet.cpp
+++ b/Engine/Source/Editor/UnrealEd/Private/Commandlets/InternationalizationExportCommandlet.cpp
@@ -1194,6 +1194,7 @@ bool UInternationalizationExportCommandlet::DoExport( const FString& SourcePath,
 								FString PORefString = ConvertSrcLocationToPORef( ContextIter->SourceLocation );
 								PoEntry->AddReference( PORefString ); // Source location.
 								PoEntry->AddExtractedComment( ContextIter->Key ); // "Notes from Programmer" in the form of the Key.
+								PoEntry->AddExtractedComment( PORefString ); // "Notes from Programmer" in the form of the Source Location, since this comes in handy too and OneSky doesn't properly show references, only comments.
 								PortableObj.AddEntry( PoEntry );
 							}
 						}
