[UEBP-60] Skipping a skeleton compile class that was redundant, also adding more accurate timer than GBlueprintCompileTime, which missed some compile related stuff.

diff --git a/Engine/Source/Editor/KismetCompiler/Private/KismetCompilerModule.cpp b/Engine/Source/Editor/KismetCompiler/Private/KismetCompilerModule.cpp
--- a/Engine/Source/Editor/KismetCompiler/Private/KismetCompilerModule.cpp
+++ b/Engine/Source/Editor/KismetCompiler/Private/KismetCompilerModule.cpp
@@ -128,10 +128,12 @@ void FKismet2CompilerModule::CompileStructure(UUserDefinedStruct* Struct, FCompi
 	FUserDefinedStructureCompilerUtils::CompileStruct(Struct, Results, true);
 }
 
+extern UNREALED_API FSecondsCounterData BlueprintCompileAndLoadTimerData;
+
 // Compiles a blueprint.
 void FKismet2CompilerModule::CompileBlueprint(class UBlueprint* Blueprint, const FKismetCompilerOptions& CompileOptions, FCompilerResultsLog& Results, TSharedPtr<FBlueprintCompileReinstancer> ParentReinstancer, TArray<UObject*>* ObjLoaded)
 {
-	SCOPE_SECONDS_COUNTER(GBlueprintCompileTime);
+	FSecondsCounterScope Timer(BlueprintCompileAndLoadTimerData);
 	BP_SCOPED_COMPILER_EVENT_STAT(EKismetCompilerStats_CompileTime);
 
 	Results.SetSourceName(Blueprint->GetName());
@@ -144,7 +146,8 @@ void FKismet2CompilerModule::CompileBlueprint(class UBlueprint* Blueprint, const
 	}
 
 	if (CompileOptions.CompileType != EKismetCompileType::Cpp
-		&& CompileOptions.CompileType != EKismetCompileType::BytecodeOnly )
+		&& CompileOptions.CompileType != EKismetCompileType::BytecodeOnly
+		&& CompileOptions.bRegenerateSkelton )
 	{
 		BP_SCOPED_COMPILER_EVENT_STAT(EKismetCompilerStats_CompileSkeletonClass);
 		auto SkeletonReinstancer = FBlueprintCompileReinstancer::Create(Blueprint->SkeletonGeneratedClass);
diff --git a/Engine/Source/Editor/UnrealEd/Private/Kismet2/BlueprintEditorUtils.cpp b/Engine/Source/Editor/UnrealEd/Private/Kismet2/BlueprintEditorUtils.cpp
--- a/Engine/Source/Editor/UnrealEd/Private/Kismet2/BlueprintEditorUtils.cpp
+++ b/Engine/Source/Editor/UnrealEd/Private/Kismet2/BlueprintEditorUtils.cpp
@@ -1448,7 +1448,7 @@ UClass* FBlueprintEditorUtils::RegenerateBlueprintClass(UBlueprint* Blueprint, U
 		// Make sure all used external classes/functions/structures/macros/etc are loaded and linked
 		FRegenerationHelper::LinkExternalDependencies(Blueprint, ObjLoaded);
 
-		FKismetEditorUtilities::GenerateBlueprintSkeleton(Blueprint);
+		bool bSkeletonUpToDate = FKismetEditorUtilities::GenerateBlueprintSkeleton(Blueprint);
 
 		static FBoolConfigValueHelper ReplaceBlueprintWithClass(TEXT("EditoronlyBP"), TEXT("bReplaceBlueprintWithClass"));
 		if (ReplaceBlueprintWithClass)
@@ -1488,7 +1488,7 @@ UClass* FBlueprintEditorUtils::RegenerateBlueprintClass(UBlueprint* Blueprint, U
 			FBlueprintEditorUtils::ReconstructAllNodes(Blueprint);
 
 			// Compile the actual blueprint
-			FKismetEditorUtilities::CompileBlueprint(Blueprint, true);
+			FKismetEditorUtilities::CompileBlueprint(Blueprint, true, false, false, nullptr, bSkeletonUpToDate);
 		}
 		else if( bIsMacro )
 		{
@@ -1818,8 +1818,12 @@ void FBlueprintEditorUtils::PropagateParentBlueprintDefaults(UClass* ClassToProp
 	}
 }
 
+UNREALED_API FSecondsCounterData BlueprintCompileAndLoadTimerData;
+
 void FBlueprintEditorUtils::PostDuplicateBlueprint(UBlueprint* Blueprint)
 {
+	FSecondsCounterScope Timer(BlueprintCompileAndLoadTimerData); 
+	
 	// Only recompile after duplication if this isn't PIE
 	if (!GIsPlayInEditorWorld)
 	{
@@ -2033,10 +2037,14 @@ void FBlueprintEditorUtils::UpdateDelegatesInBlueprint(UBlueprint* Blueprint)
 // Blueprint has materially changed.  Recompile the skeleton, notify observers, and mark the package as dirty.
 void FBlueprintEditorUtils::MarkBlueprintAsStructurallyModified(UBlueprint* Blueprint)
 {
+	FSecondsCounterScope Timer(BlueprintCompileAndLoadTimerData);
+	
 	struct FRefreshHelper
 	{
 		static void SkeletalRecompileChildren(TArray<UClass*> SkelClassesToRecompile, bool bIsCompilingOnLoad)
 		{
+			FSecondsCounterScope Timer(BlueprintCompileAndLoadTimerData);
+			
 			for (auto SkelClass : SkelClassesToRecompile)
 			{
 				if (SkelClass->HasAnyClassFlags(CLASS_NewerVersionExists))
diff --git a/Engine/Source/Editor/UnrealEd/Private/Kismet2/Kismet2.cpp b/Engine/Source/Editor/UnrealEd/Private/Kismet2/Kismet2.cpp
--- a/Engine/Source/Editor/UnrealEd/Private/Kismet2/Kismet2.cpp
+++ b/Engine/Source/Editor/UnrealEd/Private/Kismet2/Kismet2.cpp
@@ -367,9 +367,12 @@ void FKismetEditorUtilities::CreateDefaultEventGraphs(UBlueprint* Blueprint)
 	Blueprint->LastEditedDocuments.AddUnique(Ubergraph);
 }
 
+extern UNREALED_API FSecondsCounterData BlueprintCompileAndLoadTimerData;
+
 /** Create a new Blueprint and initialize it to a valid state. */
 UBlueprint* FKismetEditorUtilities::CreateBlueprint(UClass* ParentClass, UObject* Outer, const FName NewBPName, EBlueprintType BlueprintType, TSubclassOf<UBlueprint> BlueprintClassType, TSubclassOf<UBlueprintGeneratedClass> BlueprintGeneratedClassType, FName CallingContext)
 {
+	FSecondsCounterScope Timer(BlueprintCompileAndLoadTimerData);
 	check(FindObject<UBlueprint>(Outer, *NewBPName.ToString()) == NULL); 
 
 	// Not all types are legal for all parent classes, if the parent class is const then the blueprint cannot be an ubergraph-bearing one
@@ -679,8 +682,9 @@ bool FKismetEditorUtilities::IsReferencedByUndoBuffer(UBlueprint* Blueprint)
 	return (TotalReferenceCount > NonUndoReferenceCount);
 }
 
-void FKismetEditorUtilities::CompileBlueprint(UBlueprint* BlueprintObj, bool bIsRegeneratingOnLoad, bool bSkipGarbageCollection, bool bSaveIntermediateProducts, FCompilerResultsLog* pResults)
+void FKismetEditorUtilities::CompileBlueprint(UBlueprint* BlueprintObj, bool bIsRegeneratingOnLoad, bool bSkipGarbageCollection, bool bSaveIntermediateProducts, FCompilerResultsLog* pResults, bool bSkeletonUpToDate )
 {
+	FSecondsCounterScope Timer(BlueprintCompileAndLoadTimerData); 
 	BP_SCOPED_COMPILER_EVENT_STAT(EKismetCompilerStats_CompileBlueprint);
 
 	// Broadcast pre-compile
@@ -734,6 +738,7 @@ void FKismetEditorUtilities::CompileBlueprint(UBlueprint* BlueprintObj, bool bIs
 
 	FKismetCompilerOptions CompileOptions;
 	CompileOptions.bSaveIntermediateProducts = bSaveIntermediateProducts;
+	CompileOptions.bRegenerateSkelton = !bSkeletonUpToDate;
 	Compiler.CompileBlueprint(BlueprintObj, CompileOptions, Results, ReinstanceHelper);
 
 	FBlueprintEditorUtils::UpdateDelegatesInBlueprint(BlueprintObj);
@@ -871,8 +876,10 @@ void FKismetEditorUtilities::CompileBlueprint(UBlueprint* BlueprintObj, bool bIs
 }
 
 /** Generates a blueprint skeleton only.  Minimal compile, no notifications will be sent, no GC, etc.  Only successful if there isn't already a skeleton generated */
-void FKismetEditorUtilities::GenerateBlueprintSkeleton(UBlueprint* BlueprintObj, bool bForceRegeneration)
+bool FKismetEditorUtilities::GenerateBlueprintSkeleton(UBlueprint* BlueprintObj, bool bForceRegeneration)
 {
+	bool bRegeneratedSkeleton = false;
+	FSecondsCounterScope Timer(BlueprintCompileAndLoadTimerData); 
 	check(BlueprintObj);
 
 	if( BlueprintObj->SkeletonGeneratedClass == NULL || bForceRegeneration )
@@ -888,6 +895,7 @@ void FKismetEditorUtilities::GenerateBlueprintSkeleton(UBlueprint* BlueprintObj,
 		FKismetCompilerOptions CompileOptions;
 		CompileOptions.CompileType = EKismetCompileType::SkeletonOnly;
 		Compiler.CompileBlueprint(BlueprintObj, CompileOptions, Results);
+		bRegeneratedSkeleton = true;
 
 		// Restore the package dirty flag here
 		if( Package != NULL )
@@ -895,11 +903,14 @@ void FKismetEditorUtilities::GenerateBlueprintSkeleton(UBlueprint* BlueprintObj,
 			Package->SetDirtyFlag(bIsPackageDirty);
 		}
 	}
+	return bRegeneratedSkeleton;
 }
 
 /** Recompiles the bytecode of a blueprint only.  Should only be run for recompiling dependencies during compile on load */
 void FKismetEditorUtilities::RecompileBlueprintBytecode(UBlueprint* BlueprintObj, TArray<UObject*>* ObjLoaded)
 {
+	FSecondsCounterScope Timer(BlueprintCompileAndLoadTimerData); 
+	
 	check(BlueprintObj);
 	checkf(BlueprintObj->GeneratedClass, TEXT("Invalid generated class for %s"), *BlueprintObj->GetName());
 
diff --git a/Engine/Source/Editor/UnrealEd/Private/Kismet2/KismetReinstanceUtilities.cpp b/Engine/Source/Editor/UnrealEd/Private/Kismet2/KismetReinstanceUtilities.cpp
--- a/Engine/Source/Editor/UnrealEd/Private/Kismet2/KismetReinstanceUtilities.cpp
+++ b/Engine/Source/Editor/UnrealEd/Private/Kismet2/KismetReinstanceUtilities.cpp
@@ -510,8 +510,12 @@ void FBlueprintCompileReinstancer::BlueprintWasRecompiled(UBlueprint* BP, bool b
 	}
 }
 
+extern UNREALED_API FSecondsCounterData BlueprintCompileAndLoadTimerData;
+
 void FBlueprintCompileReinstancer::ReinstanceObjects(bool bForceAlwaysReinstance)
 {
+	FSecondsCounterScope Timer(BlueprintCompileAndLoadTimerData);
+
 	BP_SCOPED_COMPILER_EVENT_STAT(EKismetCompilerStats_ReinstanceObjects);
 	
 	// Make sure we only reinstance classes once!
diff --git a/Engine/Source/Editor/UnrealEd/Public/Kismet2/KismetEditorUtilities.h b/Engine/Source/Editor/UnrealEd/Public/Kismet2/KismetEditorUtilities.h
--- a/Engine/Source/Editor/UnrealEd/Public/Kismet2/KismetEditorUtilities.h
+++ b/Engine/Source/Editor/UnrealEd/Public/Kismet2/KismetEditorUtilities.h
@@ -70,10 +70,10 @@ public:
 	static void CreateDefaultEventGraphs(UBlueprint* Blueprint);
 
 	/** Tries to compile a blueprint, updating any actors in the editor who are using the old class, etc... */
-	static void CompileBlueprint(UBlueprint* BlueprintObj, bool bIsRegeneratingOnLoad=false, bool bSkipGarbageCollection=false, bool bSaveIntermediateProducts = false, class FCompilerResultsLog* pResults = NULL);
+	static void CompileBlueprint(UBlueprint* BlueprintObj, bool bIsRegeneratingOnLoad=false, bool bSkipGarbageCollection=false, bool bSaveIntermediateProducts = false, class FCompilerResultsLog* pResults = NULL, bool bSkeletonUpToDate = false);
 
 	/** Generates a blueprint skeleton only.  Minimal compile, no notifications will be sent, no GC, etc.  Only successful if there isn't already a skeleton generated */
-	static void GenerateBlueprintSkeleton(UBlueprint* BlueprintObj, bool bForceRegeneration = false);
+	static bool GenerateBlueprintSkeleton(UBlueprint* BlueprintObj, bool bForceRegeneration = false);
 
 	/** Recompiles the bytecode of a blueprint only.  Should only be run for recompiling dependencies during compile on load */
 	static void RecompileBlueprintBytecode(UBlueprint* BlueprintObj, TArray<UObject*>* ObjLoaded = NULL);
diff --git a/Engine/Source/Runtime/Engine/Classes/Engine/Blueprint.h b/Engine/Source/Runtime/Engine/Classes/Engine/Blueprint.h
--- a/Engine/Source/Runtime/Engine/Classes/Engine/Blueprint.h
+++ b/Engine/Source/Runtime/Engine/Classes/Engine/Blueprint.h
@@ -76,6 +76,9 @@ public:
 	/** Whether or not to save intermediate build products (temporary graphs and expanded macros) for debugging */
 	bool bSaveIntermediateProducts;
 
+	/** Whether to regenerate the skeleton first, when compiling on load we don't need to regenerate the skeleton. */
+	bool bRegenerateSkelton;
+
 	/** Whether or not this compile is for a duplicated blueprint */
 	bool bIsDuplicationInstigated;
 
@@ -104,6 +107,7 @@ public:
 	FKismetCompilerOptions()
 		: CompileType(EKismetCompileType::Full)
 		, bSaveIntermediateProducts(false)
+		, bRegenerateSkelton(true)
 		, bIsDuplicationInstigated(false)
 	{
 	};
diff --git a/Engine/Source/Runtime/Engine/Private/BlueprintGeneratedClass.cpp b/Engine/Source/Runtime/Engine/Private/BlueprintGeneratedClass.cpp
--- a/Engine/Source/Runtime/Engine/Private/BlueprintGeneratedClass.cpp
+++ b/Engine/Source/Runtime/Engine/Private/BlueprintGeneratedClass.cpp
@@ -187,8 +187,12 @@ struct FConditionalRecompileClassHepler
 	}
 };
 
+extern UNREALED_API FSecondsCounterData BlueprintCompileAndLoadTimerData;
+
 void UBlueprintGeneratedClass::ConditionalRecompileClass(TArray<UObject*>* ObjLoaded)
 {
+	FSecondsCounterScope Timer(BlueprintCompileAndLoadTimerData);
+
 	UBlueprint* GeneratingBP = Cast<UBlueprint>(ClassGeneratedBy);
 	if (GeneratingBP && (GeneratingBP->SkeletonGeneratedClass != this))
 	{
