Fix shadow variables

diff --git a/Engine/Source/Editor/UnrealEd/Private/Kismet2/BlueprintEditorUtils.cpp b/Engine/Source/Editor/UnrealEd/Private/Kismet2/BlueprintEditorUtils.cpp
--- a/Engine/Source/Editor/UnrealEd/Private/Kismet2/BlueprintEditorUtils.cpp
+++ b/Engine/Source/Editor/UnrealEd/Private/Kismet2/BlueprintEditorUtils.cpp
@@ -2043,7 +2043,7 @@ void FBlueprintEditorUtils::MarkBlueprintAsStructurallyModified(UBlueprint* Blue
 	{
 		static void SkeletalRecompileChildren(TArray<UClass*> SkelClassesToRecompile, bool bIsCompilingOnLoad)
 		{
-			FSecondsCounterScope Timer(BlueprintCompileAndLoadTimerData);
+			FSecondsCounterScope SkeletalRecompileTimer(BlueprintCompileAndLoadTimerData);
 			
 			for (auto SkelClass : SkelClassesToRecompile)
 			{
