Added standalone window position field and always center field in advanced settings under play. Added boolean to change between the previous centering of the screen and moving to the window position field.

UE-16501

diff --git a/Engine/Source/Editor/UnrealEd/Classes/Settings/LevelEditorPlaySettings.h b/Engine/Source/Editor/UnrealEd/Classes/Settings/LevelEditorPlaySettings.h
--- a/Engine/Source/Editor/UnrealEd/Classes/Settings/LevelEditorPlaySettings.h
+++ b/Engine/Source/Editor/UnrealEd/Classes/Settings/LevelEditorPlaySettings.h
@@ -200,6 +200,14 @@ public:
 	UPROPERTY(config, EditAnywhere, Category=PlayInStandaloneGame)
 	int32 StandaloneWindowHeight;
 
+	/** The position of the standalone game window on the screen in pixels. */
+	UPROPERTY(config, EditAnywhere, Category=PlayInStandaloneGame)
+	FIntPoint StandaloneWindowPosition;
+
+	/** Whether the standalone game window should be centered on the screen. */
+	UPROPERTY(config, EditAnywhere, Category=PlayInStandaloneGame)
+	bool CenterStandaloneWindow;
+
 	/** Whether sound should be disabled when playing standalone games. */
 	UPROPERTY(config , EditAnywhere, Category=PlayInStandaloneGame, AdvancedDisplay)
 	uint32 DisableStandaloneSound:1;
diff --git a/Engine/Source/Editor/UnrealEd/Private/PlayLevel.cpp b/Engine/Source/Editor/UnrealEd/Private/PlayLevel.cpp
--- a/Engine/Source/Editor/UnrealEd/Private/PlayLevel.cpp
+++ b/Engine/Source/Editor/UnrealEd/Private/PlayLevel.cpp
@@ -1156,7 +1156,7 @@ void UEditorEngine::PlayStandaloneLocalPc(FString MapNameOverride, FIntPoint* Wi
 	{
 		CmdLine = GenerateCmdLineForNextPieInstance(*WindowPos, PIENum, bIsServer && CanPlayNetDedicated);
 	}
-	
+
 	const FString URLParms = bIsServer && !CanPlayNetDedicated ? TEXT("?Listen") : FString();
 
 	// select map to play
@@ -1226,15 +1226,36 @@ void UEditorEngine::PlayStandaloneLocalPc(FString MapNameOverride, FIntPoint* Wi
 	FIntPoint WinSize(0, 0);
 	GetWindowSizeForInstanceType(WinSize, PlayInSettings);
 	
-	FString Params = FString::Printf(TEXT("%s %s -game -PIEVIACONSOLE -ResX=%d -ResY=%d %s%s %s"),
-		*GameNameOrProjectFile,
-		*BuildPlayWorldURL(*SavedMapNames[0], false, URLParms),
-		WinSize.X,
-		WinSize.Y,
-		*FCommandLine::GetSubprocessCommandline(),
-		*AdditionalParameters,
-		*CmdLine
-	);
+	// Check if centered
+	FString Params;
+	if (PlayInSettings->CenterStandaloneWindow)
+	{
+		Params = FString::Printf(TEXT("%s %s -game -PIEVIACONSOLE -ResX=%d -ResY=%d %s%s %s"),
+			*GameNameOrProjectFile,
+			*BuildPlayWorldURL(*SavedMapNames[0], false, URLParms),
+			WinSize.X,
+			WinSize.Y,
+			*FCommandLine::GetSubprocessCommandline(),
+			*AdditionalParameters,
+			*CmdLine
+			);
+	}
+	else
+	{
+		FIntPoint WinPos(0, 0);
+
+		Params = FString::Printf(TEXT("%s %s -game -PIEVIACONSOLE -WinX=%d -WinY=%d -ResX=%d -ResY=%d %s%s %s"),
+			*GameNameOrProjectFile,
+			*BuildPlayWorldURL(*SavedMapNames[0], false, URLParms),
+			PlayInSettings->StandaloneWindowPosition.X,
+			PlayInSettings->StandaloneWindowPosition.Y,
+			WinSize.X,
+			WinSize.Y,
+			*FCommandLine::GetSubprocessCommandline(),
+			*AdditionalParameters,
+			*CmdLine
+		);
+	}
 
 	// launch the game process
 	FString GamePath = FPlatformProcess::GenerateApplicationPath(FApp::GetName(), FApp::GetBuildConfiguration());
diff --git a/Engine/Source/Editor/UnrealEd/Private/Settings/LevelEditorPlaySettingsCustomization.h b/Engine/Source/Editor/UnrealEd/Private/Settings/LevelEditorPlaySettingsCustomization.h
--- a/Engine/Source/Editor/UnrealEd/Private/Settings/LevelEditorPlaySettingsCustomization.h
+++ b/Engine/Source/Editor/UnrealEd/Private/Settings/LevelEditorPlaySettingsCustomization.h
@@ -389,9 +389,13 @@ public:
 			// standalone window size
 			TSharedRef<IPropertyHandle> WindowHeightHandle = LayoutBuilder.GetProperty(GET_MEMBER_NAME_CHECKED(ULevelEditorPlaySettings, StandaloneWindowHeight));
 			TSharedRef<IPropertyHandle> WindowWidthHandle = LayoutBuilder.GetProperty(GET_MEMBER_NAME_CHECKED(ULevelEditorPlaySettings, StandaloneWindowWidth));
+			TSharedRef<IPropertyHandle> WindowPositionHandle = LayoutBuilder.GetProperty(GET_MEMBER_NAME_CHECKED(ULevelEditorPlaySettings, StandaloneWindowPosition));
+			TSharedRef<IPropertyHandle> CenterNewWindowHandle = LayoutBuilder.GetProperty(GET_MEMBER_NAME_CHECKED(ULevelEditorPlaySettings, CenterStandaloneWindow));
 
 			WindowHeightHandle->MarkHiddenByCustomization();
 			WindowWidthHandle->MarkHiddenByCustomization();
+			WindowPositionHandle->MarkHiddenByCustomization();
+			CenterNewWindowHandle->MarkHiddenByCustomization();
 
 			PlayInStandaloneCategory.AddCustomRow(LOCTEXT("PlayInStandaloneWindowDetails", "Standalone Window Size"), false)
 				.NameContent()
@@ -407,6 +411,21 @@ public:
 					SNew(SScreenResolutionCustomization, &LayoutBuilder, WindowHeightHandle, WindowWidthHandle)
 				];
 
+			PlayInStandaloneCategory.AddCustomRow(LOCTEXT("PlayInStandaloneWindowDetails", "Standalone Window Position"), false)
+				.NameContent()
+				[
+					SNew(STextBlock)
+						.Font(LayoutBuilder.GetDetailFont())
+						.Text(LOCTEXT("StandaloneWindowPosName", "Standalone Window Position"))
+						.ToolTipText(LOCTEXT("StandaloneWindowSizeTooltip", "Sets the width and height of standalone game windows (in pixels)"))
+				]
+				.ValueContent()
+				.MaxDesiredWidth(MaxPropertyWidth)
+				[
+					SNew(SScreenPositionCustomization, &LayoutBuilder, WindowPositionHandle, CenterNewWindowHandle)
+				];
+
+
 			// command line options
 			TSharedPtr<IPropertyHandle> DisableStandaloneSoundProperty = LayoutBuilder.GetProperty(GET_MEMBER_NAME_CHECKED(ULevelEditorPlaySettings, DisableStandaloneSound));
 
