Fixed issue where selecting a surface with no actor selected does not bring up the details panel specialization for the surface properties.

diff --git a/Engine/Source/Editor/LevelEditor/Private/LevelEditorGenericDetails.cpp b/Engine/Source/Editor/LevelEditor/Private/LevelEditorGenericDetails.cpp
--- a/Engine/Source/Editor/LevelEditor/Private/LevelEditorGenericDetails.cpp
+++ b/Engine/Source/Editor/LevelEditor/Private/LevelEditorGenericDetails.cpp
@@ -23,7 +23,7 @@ TSharedRef<IDetailCustomization> FLevelEditorGenericDetails::MakeInstance()
 void FLevelEditorGenericDetails::CustomizeDetails( IDetailLayoutBuilder& DetailLayout )
 {
 	FSelectedActorInfo SelectionInfo = AssetSelectionUtils::GetSelectedActorInfo();	
-	if( AssetSelectionUtils::GetNumSelectedSurfaces( SelectionInfo.SharedWorld ) > 0 )
+	if( AssetSelectionUtils::IsAnySurfaceSelected( SelectionInfo.SharedWorld ) )
 	{
 		AddSurfaceDetails( DetailLayout );
 	}
diff --git a/Engine/Source/Editor/PropertyEditor/Private/SDetailsView.cpp b/Engine/Source/Editor/PropertyEditor/Private/SDetailsView.cpp
--- a/Engine/Source/Editor/PropertyEditor/Private/SDetailsView.cpp
+++ b/Engine/Source/Editor/PropertyEditor/Private/SDetailsView.cpp
@@ -387,6 +387,11 @@ bool SDetailsView::ShouldSetNewObjects( const TArray< TWeakObjectPtr< UObject >
 			}
 		}
 	}
+	
+	if (!bShouldSetObjects && AssetSelectionUtils::IsAnySurfaceSelected(nullptr))
+	{
+		bShouldSetObjects = true;
+	}
 
 	return bShouldSetObjects;
 }
diff --git a/Engine/Source/Editor/UnrealEd/Private/AssetSelection.cpp b/Engine/Source/Editor/UnrealEd/Private/AssetSelection.cpp
--- a/Engine/Source/Editor/UnrealEd/Private/AssetSelection.cpp
+++ b/Engine/Source/Editor/UnrealEd/Private/AssetSelection.cpp
@@ -317,7 +317,7 @@ namespace AssetSelectionUtils
 
 	int32 GetNumSelectedSurfaces( UWorld* InWorld )
 	{
-		int32 NumSurfs = 0;
+		int32 NumSelectedSurfs = 0;
 		UWorld* World = InWorld;
 		if( !World )
 		{
@@ -325,19 +325,61 @@ namespace AssetSelectionUtils
 		}
 		if( World )
 		{
-			// Count the number of selected surfaces
-			for( int32 Surface=0; Surface<World->GetModel()->Surfs.Num(); ++Surface )
+			const int32 NumLevels = World->GetNumLevels();
+			for (int32 LevelIndex = 0; LevelIndex < NumLevels; LevelIndex++)
 			{
-				FBspSurf *Poly = &World->GetModel()->Surfs[Surface];
+				ULevel* Level = World->GetLevel(LevelIndex);
+				UModel* Model = Level->Model;
+				check(Model);
+				const int32 NumSurfaces = Model->Surfs.Num();
 
-				if( Poly->PolyFlags & PF_Selected )
+				// Count the number of selected surfaces
+				for (int32 Surface = 0; Surface < NumSurfaces; ++Surface)
 				{
-					++NumSurfs;
+					FBspSurf *Poly = &Model->Surfs[Surface];
+
+					if (Poly->PolyFlags & PF_Selected)
+					{
+						++NumSelectedSurfs;
+					}
+				}
+			}
+		}
+
+		return NumSelectedSurfs;
+	}
+
+	bool IsAnySurfaceSelected( UWorld* InWorld )
+	{
+		UWorld* World = InWorld;
+		if (!World)
+		{
+			World = GWorld;	// Fallback to GWorld
+		}
+		if (World)
+		{
+			const int32 NumLevels = World->GetNumLevels();
+			for (int32 LevelIndex = 0; LevelIndex < NumLevels; LevelIndex++)
+			{
+				ULevel* Level = World->GetLevel(LevelIndex);
+				UModel* Model = Level->Model;
+				check(Model);
+				const int32 NumSurfaces = Model->Surfs.Num();
+
+				// Count the number of selected surfaces
+				for (int32 Surface = 0; Surface < NumSurfaces; ++Surface)
+				{
+					FBspSurf *Poly = &Model->Surfs[Surface];
+
+					if (Poly->PolyFlags & PF_Selected)
+					{
+						return true;
+					}
 				}
 			}
 		}
 
-		return NumSurfs;
+		return false;
 	}
 
 	bool IsBuilderBrushSelected()
diff --git a/Engine/Source/Editor/UnrealEd/Public/AssetSelection.h b/Engine/Source/Editor/UnrealEd/Public/AssetSelection.h
--- a/Engine/Source/Editor/UnrealEd/Public/AssetSelection.h
+++ b/Engine/Source/Editor/UnrealEd/Public/AssetSelection.h
@@ -197,6 +197,11 @@ namespace AssetSelectionUtils
 	 */
 	UNREALED_API int32 GetNumSelectedSurfaces( UWorld* InWorld );
 
+	/**
+	 * @return whether any surface is selected
+	 */
+	UNREALED_API bool IsAnySurfaceSelected( UWorld* InWorld );
+
 	/**
 	 * @return true if the builder brush is in the list of selected actors; false, otherwise. 
 	 */
