Fixed a crash after renaming local variables when other functions have local variables of the same name.

Renaming was bleeding into other functions because the rename was not scoped to the graph.

diff --git a/Engine/Source/Editor/UnrealEd/Private/Kismet2/BlueprintEditorUtils.cpp b/Engine/Source/Editor/UnrealEd/Private/Kismet2/BlueprintEditorUtils.cpp
--- a/Engine/Source/Editor/UnrealEd/Private/Kismet2/BlueprintEditorUtils.cpp
+++ b/Engine/Source/Editor/UnrealEd/Private/Kismet2/BlueprintEditorUtils.cpp
@@ -220,6 +220,57 @@ static void PromoteInterfaceImplementationToOverride(FBPInterfaceDescription con
 	}
 }
 
+/**
+ * Looks through the specified graph for any references to the specified 
+ * variable, and renames them accordingly.
+ * 
+ * @param  InBlueprint		The blueprint that you want to search through.
+ * @param  InVariableClass	The class that owns the variable that we're renaming
+ * @param  InGraph			Graph to scope the rename to
+ * @param  InOldVarName		The current name of the variable we want to replace
+ * @param  InNewVarName		The name that we wish to change all references to
+ */
+static void RenameVariableReferencesInGraph(UBlueprint* InBlueprint, UClass* InVariableClass, UEdGraph* InGraph, const FName& InOldVarName, const FName& InNewVarName)
+{
+	for(UEdGraphNode* GraphNode : InGraph->Nodes)
+	{
+		if(UK2Node_Variable* const VariableNode = Cast<UK2Node_Variable>(GraphNode))
+		{
+			UClass* const NodeRefClass = VariableNode->VariableReference.GetMemberParentClass(InBlueprint->GeneratedClass);
+			if(NodeRefClass && NodeRefClass->IsChildOf(InVariableClass) && InOldVarName == VariableNode->GetVarName())
+			{
+				VariableNode->Modify();
+
+				if(VariableNode->VariableReference.IsLocalScope())
+				{
+					VariableNode->VariableReference.SetLocalMember(InNewVarName, VariableNode->VariableReference.GetMemberScopeName(), VariableNode->VariableReference.GetMemberGuid());
+				}
+				else if(VariableNode->VariableReference.IsSelfContext())
+				{
+					VariableNode->VariableReference.SetSelfMember(InNewVarName);
+				}
+				else
+				{
+					VariableNode->VariableReference.SetExternalMember(InNewVarName, NodeRefClass);
+				}
+
+				VariableNode->RenameUserDefinedPin(InOldVarName.ToString(), InNewVarName.ToString());
+			}
+			continue;
+		}
+
+		if(UK2Node_ComponentBoundEvent* const ComponentBoundEventNode = Cast<UK2Node_ComponentBoundEvent>(GraphNode))
+		{
+			if(InOldVarName == ComponentBoundEventNode->ComponentPropertyName)
+			{
+				ComponentBoundEventNode->Modify();
+				ComponentBoundEventNode->ComponentPropertyName = InNewVarName;
+			}
+			continue;
+		}
+	}
+}
+
 /**
  * Looks through the specified blueprint for any references to the specified 
  * variable, and renames them accordingly.
@@ -237,43 +288,7 @@ static void RenameVariableReferences(UBlueprint* Blueprint, UClass* VariableClas
 	// Update any graph nodes that reference the old variable name to instead reference the new name
 	for(UEdGraph* CurrentGraph : AllGraphs)
 	{
-		for(UEdGraphNode* GraphNode : CurrentGraph->Nodes)
-		{
-			if(UK2Node_Variable* const VariableNode = Cast<UK2Node_Variable>(GraphNode))
-			{
-				UClass* const NodeRefClass = VariableNode->VariableReference.GetMemberParentClass(Blueprint->GeneratedClass);
-				if(NodeRefClass && NodeRefClass->IsChildOf(VariableClass) && OldVarName == VariableNode->GetVarName())
-				{
-					VariableNode->Modify();
-
-					if(VariableNode->VariableReference.IsLocalScope())
-					{
-						VariableNode->VariableReference.SetLocalMember(NewVarName, VariableNode->VariableReference.GetMemberScopeName(), VariableNode->VariableReference.GetMemberGuid());
-					}
-					else if(VariableNode->VariableReference.IsSelfContext())
-					{
-						VariableNode->VariableReference.SetSelfMember(NewVarName);
-					}
-					else
-					{
-						VariableNode->VariableReference.SetExternalMember(NewVarName, NodeRefClass);
-					}
-
-					VariableNode->RenameUserDefinedPin(OldVarName.ToString(), NewVarName.ToString());
-				}
-				continue;
-			}
-			
-			if(UK2Node_ComponentBoundEvent* const ComponentBoundEventNode = Cast<UK2Node_ComponentBoundEvent>(GraphNode))
-			{
-				if(OldVarName == ComponentBoundEventNode->ComponentPropertyName)
-				{
-					ComponentBoundEventNode->Modify();
-					ComponentBoundEventNode->ComponentPropertyName = NewVarName;
-				}
-				continue;
-			}
-		}
+		RenameVariableReferencesInGraph(Blueprint, VariableClass, CurrentGraph, OldVarName, NewVarName);
 	}
 }
 
@@ -4704,7 +4719,7 @@ void FBlueprintEditorUtils::RenameLocalVariable(UBlueprint* InBlueprint, const U
 			LocalVariable->FriendlyName = FName::NameToDisplayString( InNewName.ToString(), LocalVariable->VarType.PinCategory == K2Schema->PC_Boolean );
 
 			// Update any existing references to the old name
-			FBlueprintEditorUtils::ReplaceVariableReferences(InBlueprint, InOldName, InNewName);
+			RenameVariableReferencesInGraph(InBlueprint, InBlueprint->GeneratedClass, FindScopeGraph(InBlueprint, InScope), InOldName, InNewName);
 
 			// Validate child blueprints and adjust variable names to avoid a potential name collision
 			FBlueprintEditorUtils::ValidateBlueprintChildVariables(InBlueprint, InNewName);
