Merging

//depot/UE4-Orion/Engine/Source/Editor/UnrealEd/Private/Fbx/FbxMainImport.cpp

to //depot/UE4/Engine/Source/Editor/UnrealEd/Private/Fbx/FbxMainImport.cpp

diff --git a/Engine/Source/Editor/UnrealEd/Private/Fbx/FbxMainImport.cpp b/Engine/Source/Editor/UnrealEd/Private/Fbx/FbxMainImport.cpp
--- a/Engine/Source/Editor/UnrealEd/Private/Fbx/FbxMainImport.cpp
+++ b/Engine/Source/Editor/UnrealEd/Private/Fbx/FbxMainImport.cpp
@@ -415,8 +415,9 @@ int32 FFbxImporter::GetImportType(const FString& InFilename)
 				Result = 0;
 			}
 			// Check for animation data. It can be overwritten by Geometry or Deformer
-			else if ( (ItemName == "AnimationStack" || ItemName == "AnimationLayer" || ItemName == "AnimationCurve" || ItemName == "AnimationCurveNode") && ItemCount > 0 )
+			else if ( (ItemName == "AnimationCurve" || ItemName == "AnimationCurveNode") && ItemCount > 0 )
 			{
+				Result = 1;
 				bHasAnimation = true;
 			}
 		}
