Merging:
2557015
 dispatch a delegate when in pie for when the viewport is constructed, also wrap all of the UI in a safe zone out of game. don't initialzie UI until viewport exists (should prevent race condition)

diff --git a/Engine/Source/Editor/UnrealEd/Private/PlayLevel.cpp b/Engine/Source/Editor/UnrealEd/Private/PlayLevel.cpp
--- a/Engine/Source/Editor/UnrealEd/Private/PlayLevel.cpp
+++ b/Engine/Source/Editor/UnrealEd/Private/PlayLevel.cpp
@@ -3005,6 +3005,8 @@ UGameInstance* UEditorEngine::CreatePIEGameInstance(int32 PIEInstance, bool bInS
 					}
 				}
 			}
+
+			UGameViewportClient::OnViewportCreated().Broadcast();
 		}
 	}
 
