UE-16930 - Fixed an issue where setting the game mode to a BP type could result in a null game mode due to the BP never being loaded

diff --git a/Engine/Source/Editor/UnrealEd/Public/GameModeInfoCustomizer.h b/Engine/Source/Editor/UnrealEd/Public/GameModeInfoCustomizer.h
--- a/Engine/Source/Editor/UnrealEd/Public/GameModeInfoCustomizer.h
+++ b/Engine/Source/Editor/UnrealEd/Public/GameModeInfoCustomizer.h
@@ -4,6 +4,7 @@
 
 #include "IDocumentation.h"
 #include "KismetEditorUtilities.h"
+#include "EditorClassUtils.h"
 #include "GameFramework/GameMode.h"
 #include "Engine/BlueprintGeneratedClass.h"
 
@@ -168,23 +169,14 @@ public:
 		FString ClassName;
 		DefaultGameModeClassHandle->GetValueAsFormattedString(ClassName);
 
-		if (ClassName.IsEmpty() || ClassName == "None")
+		// Do we have a valid cached class pointer? (Note: We can't search for the class while a save is happening)
+		const UClass* GameModeClass = CachedGameModeClass.Get();
+		if ((!GameModeClass || GameModeClass->GetPathName() != ClassName) && !GIsSavingPackage)
 		{
-			CachedGameModeClass = nullptr;
-			return nullptr;
+			GameModeClass = FEditorClassUtils::GetClassFromString(ClassName);
+			CachedGameModeClass = GameModeClass;
 		}
-
-		FString StrippedClassName = ClassName;
-		ConstructorHelpers::StripObjectClass(StrippedClassName);
-
-		// Just returned the cached value if it's already up-to-date, or if we're saving packages since we can't search for the object while a save is happening
-		if ((CachedGameModeClass && CachedGameModeClass->GetPathName() == StrippedClassName) || GIsSavingPackage)
-		{
-			return CachedGameModeClass;
-		}
-
-		CachedGameModeClass = FindObject<UClass>(ANY_PACKAGE, *StrippedClassName);
-		return CachedGameModeClass;
+		return GameModeClass;
 	}
 
 	void SetCurrentGameModeClass(const UClass* NewGameModeClass)
@@ -346,7 +338,7 @@ private:
 	/** Handle to the DefaultGameMode property */
 	TSharedPtr<IPropertyHandle> DefaultGameModeClassHandle;
 	/** Cached class pointer from the DefaultGameModeClassHandle */
-	mutable const UClass* CachedGameModeClass;
+	mutable TWeakObjectPtr<UClass> CachedGameModeClass;
 };
 
 #undef LOCTEXT_NAMESPACE
