Added a boolean to the play advanced settings window. If true PIE window set as as child, drawing in front of the parent window.

UE-15774

diff --git a/Engine/Source/Editor/UnrealEd/Classes/Settings/LevelEditorPlaySettings.h b/Engine/Source/Editor/UnrealEd/Classes/Settings/LevelEditorPlaySettings.h
--- a/Engine/Source/Editor/UnrealEd/Classes/Settings/LevelEditorPlaySettings.h
+++ b/Engine/Source/Editor/UnrealEd/Classes/Settings/LevelEditorPlaySettings.h
@@ -169,11 +169,15 @@ public:
 	UPROPERTY(config)
 	uint32 bOnlyLoadVisibleLevelsInPIE:1;
 
-	UPROPERTY(config, EditAnywhere, Category = PlayInEditor, meta = (DisplayName="Stream Sub-Levels during Play in Editor", ToolTip = "Prefer to stream sub-levels from the disk instead of duplicating editor sub-levels"))
+	UPROPERTY(config, EditAnywhere, Category = PlayInEditor, meta = (DisplayName="Stream Sub-Levels during Play in Editor", ToolTip="Prefer to stream sub-levels from the disk instead of duplicating editor sub-levels"))
 	uint32 bPreferToStreamLevelsInPIE:1;
 
 public:
 
+	/** The width of the new view port window in pixels (0 = use the desktop's screen resolution). */
+	UPROPERTY(config, EditAnywhere, Category = PlayInNewWindow, meta = (ToolTip="Always have the PIE window on top of the parent windows."))
+	bool PIEAlwaysOnTop;
+
 	/** The width of the new view port window in pixels (0 = use the desktop's screen resolution). */
 	UPROPERTY(config, EditAnywhere, Category=PlayInNewWindow)
 	int32 NewWindowWidth;
diff --git a/Engine/Source/Editor/UnrealEd/Private/PlayLevel.cpp b/Engine/Source/Editor/UnrealEd/Private/PlayLevel.cpp
--- a/Engine/Source/Editor/UnrealEd/Private/PlayLevel.cpp
+++ b/Engine/Source/Editor/UnrealEd/Private/PlayLevel.cpp
@@ -2898,7 +2898,10 @@ UGameInstance* UEditorEngine::CreatePIEGameInstance(int32 PIEInstance, bool bInS
 				FOnSwitchWorldHack OnWorldSwitch = FOnSwitchWorldHack::CreateUObject( this, &UEditorEngine::OnSwitchWorldForSlatePieWindow );
 				PieWindow->SetOnWorldSwitchHack( OnWorldSwitch );
 
-				FSlateApplication::Get().AddWindow( PieWindow );
+				if (PlayInSettings->PIEAlwaysOnTop && FSlateApplication::Get().GetActiveTopLevelWindow().IsValid())
+					FSlateApplication::Get().AddWindowAsNativeChild(PieWindow, FSlateApplication::Get().GetActiveTopLevelWindow().ToSharedRef(), true);
+				else
+					FSlateApplication::Get().AddWindow(PieWindow);
 
 				TSharedRef<SOverlay> ViewportOverlayWidgetRef = SNew(SOverlay);
 
diff --git a/Engine/Source/Editor/UnrealEd/Private/Settings/LevelEditorPlaySettingsCustomization.h b/Engine/Source/Editor/UnrealEd/Private/Settings/LevelEditorPlaySettingsCustomization.h
--- a/Engine/Source/Editor/UnrealEd/Private/Settings/LevelEditorPlaySettingsCustomization.h
+++ b/Engine/Source/Editor/UnrealEd/Private/Settings/LevelEditorPlaySettingsCustomization.h
@@ -343,6 +343,9 @@ public:
 		// play in new window settings
 		IDetailCategoryBuilder& PlayInNewWindowCategory = LayoutBuilder.EditCategory("PlayInNewWindow");
 		{
+			PlayInNewWindowCategory.AddProperty("PIEAlwaysOnTop")
+				.DisplayName(LOCTEXT("PIEAlwaysOnTop", "Always On Top"));
+
 			// new window size
 			TSharedRef<IPropertyHandle> WindowHeightHandle = LayoutBuilder.GetProperty(GET_MEMBER_NAME_CHECKED(ULevelEditorPlaySettings, NewWindowHeight));
 			TSharedRef<IPropertyHandle> WindowWidthHandle = LayoutBuilder.GetProperty(GET_MEMBER_NAME_CHECKED(ULevelEditorPlaySettings, NewWindowWidth));
