Allows viewport to move when piloting for focus, move to, etc.
The piloting was overwriting the viewport transition code resulting in no movement.

UE-7376

diff --git a/Engine/Source/Editor/UnrealEd/Private/EditorViewportClient.cpp b/Engine/Source/Editor/UnrealEd/Private/EditorViewportClient.cpp
--- a/Engine/Source/Editor/UnrealEd/Private/EditorViewportClient.cpp
+++ b/Engine/Source/Editor/UnrealEd/Private/EditorViewportClient.cpp
@@ -173,6 +173,11 @@ FMatrix FViewportCameraTransform::ComputeOrbitMatrix() const
 	return Transform.ToMatrixNoScale() * FInverseRotationMatrix( FRotator(0,90.f,0) );
 }
 
+bool FViewportCameraTransform::IsPlaying()
+{
+	return TransitionCurve->IsPlaying();
+}
+
 /**The Maximum Mouse/Camera Speeds Setting supported */
 const uint32 FEditorViewportClient::MaxCameraSpeeds = 8;
 
diff --git a/Engine/Source/Editor/UnrealEd/Private/LevelEditorViewport.cpp b/Engine/Source/Editor/UnrealEd/Private/LevelEditorViewport.cpp
--- a/Engine/Source/Editor/UnrealEd/Private/LevelEditorViewport.cpp
+++ b/Engine/Source/Editor/UnrealEd/Private/LevelEditorViewport.cpp
@@ -2023,7 +2023,6 @@ void FLevelEditorViewportClient::Tick(float DeltaTime)
 	UpdateViewForLockedActor();
 }
 
-
 void FLevelEditorViewportClient::UpdateViewForLockedActor()
 {
 	// We can't be locked to a matinee actor if this viewport doesn't allow matinee control
@@ -2038,35 +2037,45 @@ void FLevelEditorViewportClient::UpdateViewForLockedActor()
 	const AActor* Actor = ActorLockedByMatinee.IsValid() ? ActorLockedByMatinee.Get() : ActorLockedToCamera.Get();
 	if( Actor != NULL )
 	{
-		// Update transform
-		if( Actor->GetAttachParentActor() != NULL )
+		// Check if the viewport is transitioning
+		FViewportCameraTransform& ViewTransform = GetViewTransform();
+		if (ViewTransform.IsPlaying())
 		{
-			// Actor is parented, so use the actor to world matrix for translation and rotation information.
-			SetViewLocation( Actor->GetActorLocation() );
-			SetViewRotation( Actor->GetActorRotation() );				
+			// Move actor to the transitioned viewport
+			PerspectiveCameraMoved();
 		}
-		else if( Actor->GetRootComponent() != NULL )
+		else
 		{
-			// No attachment, so just use the relative location, so that we don't need to
-			// convert from a quaternion, which loses winding information.
-			SetViewLocation( Actor->GetRootComponent()->RelativeLocation );
-			SetViewRotation( Actor->GetRootComponent()->RelativeRotation );
-		}
+			// Update transform
+			if (Actor->GetAttachParentActor() != NULL)
+			{
+				// Actor is parented, so use the actor to world matrix for translation and rotation information.
+				SetViewLocation(Actor->GetActorLocation());
+				SetViewRotation(Actor->GetActorRotation());
+			}
+			else if (Actor->GetRootComponent() != NULL)
+			{
+				// No attachment, so just use the relative location, so that we don't need to
+				// convert from a quaternion, which loses winding information.
+				SetViewLocation(Actor->GetRootComponent()->RelativeLocation);
+				SetViewRotation(Actor->GetRootComponent()->RelativeRotation);
+			}
 
-		if( bLockedCameraView )
-		{
-			// If this is a camera actor, then inherit some other settings
-			UCameraComponent* CameraComponent = Actor->FindComponentByClass<UCameraComponent>();
-			if( CameraComponent != NULL )
+			if (bLockedCameraView)
 			{
-				bUseControllingActorViewInfo = true;
-				CameraComponent->GetCameraView(0.0f, ControllingActorViewInfo);
-
-				// Post processing is handled by OverridePostProcessingSettings
-				ViewFOV = ControllingActorViewInfo.FOV;
-				AspectRatio = ControllingActorViewInfo.AspectRatio;
-				SetViewLocation(ControllingActorViewInfo.Location);
-				SetViewRotation(ControllingActorViewInfo.Rotation);
+				// If this is a camera actor, then inherit some other settings
+				UCameraComponent* CameraComponent = Actor->FindComponentByClass<UCameraComponent>();
+				if (CameraComponent != NULL)
+				{
+					bUseControllingActorViewInfo = true;
+					CameraComponent->GetCameraView(0.0f, ControllingActorViewInfo);
+
+					// Post processing is handled by OverridePostProcessingSettings
+					ViewFOV = ControllingActorViewInfo.FOV;
+					AspectRatio = ControllingActorViewInfo.AspectRatio;
+					SetViewLocation(ControllingActorViewInfo.Location);
+					SetViewRotation(ControllingActorViewInfo.Rotation);
+				}
 			}
 		}
 	}
diff --git a/Engine/Source/Editor/UnrealEd/Public/EditorViewportClient.h b/Engine/Source/Editor/UnrealEd/Public/EditorViewportClient.h
--- a/Engine/Source/Editor/UnrealEd/Public/EditorViewportClient.h
+++ b/Engine/Source/Editor/UnrealEd/Public/EditorViewportClient.h
@@ -189,6 +189,8 @@ public:
 		OrthoZoom = InOrthoZoom;
 	}
 
+	/** Check if transition curve is playing. */
+	bool IsPlaying();
 
 	/** @return The transform's location */
 	FORCEINLINE const FVector& GetLocation() const { return ViewLocation; }
@@ -221,10 +223,11 @@ public:
 	 * Computes a matrix to use for viewport location and rotation when orbiting
 	 */
 	FMatrix ComputeOrbitMatrix() const;
+
 private:
 	/** The time when a transition to the desired location began */
 	//double TransitionStartTime;
-	
+
 	/** Curve for animating between locations */
 	TSharedPtr<struct FCurveSequence> TransitionCurve;
 	/** Current viewport Position. */
