Local variables in Blueprint functions will now be set to the default values of the user defined struct their type represents.

diff --git a/Engine/Source/Editor/UnrealEd/Private/Kismet2/BlueprintEditorUtils.cpp b/Engine/Source/Editor/UnrealEd/Private/Kismet2/BlueprintEditorUtils.cpp
--- a/Engine/Source/Editor/UnrealEd/Private/Kismet2/BlueprintEditorUtils.cpp
+++ b/Engine/Source/Editor/UnrealEd/Private/Kismet2/BlueprintEditorUtils.cpp
@@ -4564,6 +4564,30 @@ FBPVariableDescription FBlueprintEditorUtils::DuplicateVariableDescription(UBlue
 	return NewVar;
 }
 
+void FBlueprintEditorUtils::SetDefaultValueOnUserDefinedStructProperty(UBlueprint* InBlueprint, FString InScopeName, FBPVariableDescription* InOutVariableDescription)
+{
+	if (InOutVariableDescription->VarType.PinSubCategoryObject.IsValid() && InOutVariableDescription->VarType.PinSubCategoryObject->IsA(UUserDefinedStruct::StaticClass()))
+	{
+		// Create a variable reference for the variable, so we can resolve it and use it to generate the default value
+		FMemberReference VariableReference;
+		VariableReference.SetLocalMember(InOutVariableDescription->VarName, InScopeName, InOutVariableDescription->VarGuid);
+		UStruct* FunctionScope = VariableReference.GetMemberScope(InBlueprint->SkeletonGeneratedClass);
+		UProperty* VariableProperty = VariableReference.ResolveMember<UProperty>(InBlueprint->SkeletonGeneratedClass);
+
+		TSharedPtr<FStructOnScope> StructData = MakeShareable(new FStructOnScope(FunctionScope));
+
+		// Create the default value for the property
+		FStructureEditorUtils::Fill_MakeStructureDefaultValue(VariableProperty, StructData->GetStructMemory());
+
+		// Export the default value as a string so we can store it with the variable description
+		FString DefaultValueString;
+		FBlueprintEditorUtils::PropertyValueToString(VariableProperty, StructData->GetStructMemory(), DefaultValueString);
+
+		// Set the default value
+		InOutVariableDescription->DefaultValue = DefaultValueString;
+	}
+}
+
 bool FBlueprintEditorUtils::AddLocalVariable(UBlueprint* Blueprint, UEdGraph* InTargetGraph, const FName InNewVarName, const FEdGraphPinType& InNewVarType, const FString& DefaultValue/*= FString()*/)
 {
 	if(InTargetGraph != NULL && InTargetGraph->GetSchema()->GetGraphType(InTargetGraph) == GT_Function)
@@ -4594,6 +4618,13 @@ bool FBlueprintEditorUtils::AddLocalVariable(UBlueprint* Blueprint, UEdGraph* In
 
 		FBlueprintEditorUtils::MarkBlueprintAsStructurallyModified(Blueprint);
 
+		if (NewVar.DefaultValue.IsEmpty())
+		{
+			// Pull the description of the new variable, we will update it's default value with the default value from the user defined struct.
+			FBPVariableDescription* VariableDescription = &FunctionEntryNodes[0]->LocalVariables[FunctionEntryNodes[0]->LocalVariables.Num() - 1];
+			SetDefaultValueOnUserDefinedStructProperty(Blueprint, InTargetGraph->GetName(), VariableDescription);
+		}
+
 		return true;
 	}
 	return false;
@@ -4887,6 +4918,10 @@ void FBlueprintEditorUtils::ChangeLocalVariableType(UBlueprint* InBlueprint, con
 						}
 					}
 				}
+				else
+				{
+					SetDefaultValueOnUserDefinedStructProperty(InBlueprint, FunctionEntry->GetGraph()->GetName(), VariablePtr);
+				}
 
 				// Reconstruct all local variables referencing the modified one
 				for(UK2Node_Variable* VariableNode : VariableNodes)
diff --git a/Engine/Source/Editor/UnrealEd/Public/Kismet2/BlueprintEditorUtils.h b/Engine/Source/Editor/UnrealEd/Public/Kismet2/BlueprintEditorUtils.h
--- a/Engine/Source/Editor/UnrealEd/Public/Kismet2/BlueprintEditorUtils.h
+++ b/Engine/Source/Editor/UnrealEd/Public/Kismet2/BlueprintEditorUtils.h
@@ -701,6 +701,15 @@ public:
 	 */
 	static UEdGraph* FindScopeGraph(const UBlueprint* InBlueprint, const UStruct* InScope);
 
+	/**
+	 * Helper function to set the default value of a user defined struct property's local variable description
+	 *
+	 * @param	InBlueprint					The Blueprint the local variable can be found in
+	 * @param	InScope	Name				Local variable's scope name
+	 * @param	InOutVariableDescription	Description of the variable, will be modified with the default value
+	 */
+	static void SetDefaultValueOnUserDefinedStructProperty(UBlueprint* InBlueprint, FString InScopeName, FBPVariableDescription* InOutVariableDescription);
+
 	/**
 	 * Adds a local variable to the function graph.  It cannot mask a member variable or a variable in any superclass.
 	 *
