fixed typo

diff --git a/Engine/Source/Editor/UnrealEd/Private/Factories/IESLoader.cpp b/Engine/Source/Editor/UnrealEd/Private/Factories/IESLoader.cpp
--- a/Engine/Source/Editor/UnrealEd/Private/Factories/IESLoader.cpp
+++ b/Engine/Source/Editor/UnrealEd/Private/Factories/IESLoader.cpp
@@ -215,7 +215,7 @@ void FIESLoadHelper::Load(const uint8* Buffer)
 
 	if(HAnglesNum < 0)
 	{
-		Error = "VAnglesNum is not valid";
+		Error = "HAnglesNum is not valid";
 		return;
 	}
 
