UEdGraphNode::ShouldMergeChildGraphs - Mechanism to prevent merging subgraphs.  It's needed for MathExpression optimization.

diff --git a/Engine/Source/Editor/KismetCompiler/Private/KismetCompiler.cpp b/Engine/Source/Editor/KismetCompiler/Private/KismetCompiler.cpp
--- a/Engine/Source/Editor/KismetCompiler/Private/KismetCompiler.cpp
+++ b/Engine/Source/Editor/KismetCompiler/Private/KismetCompiler.cpp
@@ -2823,6 +2823,12 @@ void FKismetCompilerContext::ExpandTunnelsAndMacros(UEdGraph* SourceGraph)
 	// Collapse any remaining tunnels
 	for (TArray<UEdGraphNode*>::TIterator NodeIt(SourceGraph->Nodes); NodeIt; ++NodeIt)
 	{
+		UEdGraphNode* CurrentNode = *NodeIt;
+		if (!CurrentNode || !CurrentNode->ShouldMergeChildGraphs())
+		{
+			continue;
+		}
+
 		if (UK2Node_MacroInstance* MacroInstanceNode = Cast<UK2Node_MacroInstance>(*NodeIt))
 		{
 			UEdGraph* MacroGraph = MacroInstanceNode->GetMacroGraph();
diff --git a/Engine/Source/Editor/UnrealEd/Private/EdGraphUtilities.cpp b/Engine/Source/Editor/UnrealEd/Private/EdGraphUtilities.cpp
--- a/Engine/Source/Editor/UnrealEd/Private/EdGraphUtilities.cpp
+++ b/Engine/Source/Editor/UnrealEd/Private/EdGraphUtilities.cpp
@@ -218,9 +218,13 @@ void FEdGraphUtilities::MergeChildrenGraphsIn(UEdGraph* MergeTarget, UEdGraph* P
 	{
 		UEdGraph* ChildGraph = ParentGraph->SubGraphs[Index];
 
+		auto NodeOwner = Cast<const UEdGraphNode>(ChildGraph->GetOuter());
+		const bool bNonVirtualGraph = NodeOwner ? NodeOwner->ShouldMergeChildGraphs() : true;
+
 		// Only merge children in with the same schema as the parent
 		const bool bSchemaMatches = ChildGraph->GetSchema()->GetClass()->IsChildOf(MergeTarget->GetSchema()->GetClass());
-		if (!bRequireSchemaMatch || bSchemaMatches)
+		const bool bDoMerge = bNonVirtualGraph && (!bRequireSchemaMatch || bSchemaMatches);
+		if (bDoMerge)
 		{
 			// Even if we don't require a match to recurse, we do to actually copy the nodes
 			if (bSchemaMatches)
diff --git a/Engine/Source/Runtime/Engine/Classes/EdGraph/EdGraphNode.h b/Engine/Source/Runtime/Engine/Classes/EdGraph/EdGraphNode.h
--- a/Engine/Source/Runtime/Engine/Classes/EdGraph/EdGraphNode.h
+++ b/Engine/Source/Runtime/Engine/Classes/EdGraph/EdGraphNode.h
@@ -436,6 +436,9 @@ public:
 
 	/** Return the matching "pass-through" pin for the given pin (if supported by this node) */
 	virtual UEdGraphPin* GetPassThroughPin(const UEdGraphPin* FromPin) const { return nullptr; }
+
+	/** If the node has a subgraph, should they be merged into the main graph? */
+	virtual bool ShouldMergeChildGraphs() const { return true; }
 #endif // WITH_EDITOR
 
 };
