Added an option to not use MikkTSpace for tangent generation on import.

diff --git a/Engine/Source/Developer/MeshUtilities/Private/MeshUtilities.cpp b/Engine/Source/Developer/MeshUtilities/Private/MeshUtilities.cpp
--- a/Engine/Source/Developer/MeshUtilities/Private/MeshUtilities.cpp
+++ b/Engine/Source/Developer/MeshUtilities/Private/MeshUtilities.cpp
@@ -2780,8 +2780,8 @@ public:
 				FindOverlappingCorners(OverlappingCorners, RawMesh, ComparisonThreshold);
 
 				// Figure out if we should recompute normals and tangents.
-				bool bRecomputeNormals = SrcModel.BuildSettings.bRecomputeNormals || RawMesh.WedgeTangentZ.Num() == 0;
-				bool bRecomputeTangents = SrcModel.BuildSettings.bRecomputeTangents || RawMesh.WedgeTangentX.Num() == 0 || RawMesh.WedgeTangentY.Num() == 0;
+				bool bRecomputeNormals = SrcModel.BuildSettings.bRecomputeNormals || RawMesh.WedgeTangentZ.Num() != NumWedges;
+				bool bRecomputeTangents = SrcModel.BuildSettings.bRecomputeTangents || RawMesh.WedgeTangentX.Num() != NumWedges || RawMesh.WedgeTangentY.Num() != NumWedges;
 
 				// Dump normals and tangents if we are recomputing them.
 				if (bRecomputeTangents)
diff --git a/Engine/Source/Editor/UnrealEd/Classes/Factories/FbxMeshImportData.h b/Engine/Source/Editor/UnrealEd/Classes/Factories/FbxMeshImportData.h
--- a/Engine/Source/Editor/UnrealEd/Classes/Factories/FbxMeshImportData.h
+++ b/Engine/Source/Editor/UnrealEd/Classes/Factories/FbxMeshImportData.h
@@ -12,6 +12,20 @@ enum EFBXNormalImportMethod
 	FBXNIM_MAX,
 };
 
+UENUM()
+namespace EFBXNormalGenerationMethod
+{
+	enum Type
+	{
+		/** Use the legacy built in method to generate normals (faster in some cases) */
+		BuiltIn,
+		/** Use MikkTSpace to generate normals and tangents */
+		MikkTSpace,
+		MAX,
+	};
+}
+
+
 /**
  * Import data and options used when importing any mesh from FBX
  */
@@ -28,5 +42,9 @@ class UFbxMeshImportData : public UFbxAssetImportData
 	UPROPERTY(EditAnywhere, AdvancedDisplay, config, Category=ImportSettings, meta=(ImportType="Mesh"))
 	TEnumAsByte<enum EFBXNormalImportMethod> NormalImportMethod;
 
+	/** Use the MikkTSpace tangent space generator for generating normals and tangents on the mesh */
+	UPROPERTY(EditAnywhere, config, AdvancedDisplay, Category = ImportSettings, meta = (ImportType="StaticMesh"))
+	TEnumAsByte<enum EFBXNormalGenerationMethod::Type> NormalGenerationMethod;
+
 	bool CanEditChange( const UProperty* InProperty ) const override;
 };
diff --git a/Engine/Source/Editor/UnrealEd/Private/Fbx/FbxMainImport.cpp b/Engine/Source/Editor/UnrealEd/Private/Fbx/FbxMainImport.cpp
--- a/Engine/Source/Editor/UnrealEd/Private/Fbx/FbxMainImport.cpp
+++ b/Engine/Source/Editor/UnrealEd/Private/Fbx/FbxMainImport.cpp
@@ -189,6 +189,7 @@ void ApplyImportUIToImportOptions(UFbxImportUI* ImportUI, FBXImportOptions& InOu
 	{
 		UFbxStaticMeshImportData* StaticMeshData	= ImportUI->StaticMeshImportData;
 		InOutImportOptions.NormalImportMethod		= StaticMeshData->NormalImportMethod;
+		InOutImportOptions.NormalGenerationMethod	= StaticMeshData->NormalGenerationMethod;
 		InOutImportOptions.ImportTranslation		= StaticMeshData->ImportTranslation;
 		InOutImportOptions.ImportRotation			= StaticMeshData->ImportRotation;
 		InOutImportOptions.ImportUniformScale		= StaticMeshData->ImportUniformScale;
@@ -197,6 +198,7 @@ void ApplyImportUIToImportOptions(UFbxImportUI* ImportUI, FBXImportOptions& InOu
 	{
 		UFbxSkeletalMeshImportData* SkeletalMeshData	= ImportUI->SkeletalMeshImportData;
 		InOutImportOptions.NormalImportMethod			= SkeletalMeshData->NormalImportMethod;
+		InOutImportOptions.NormalGenerationMethod		= SkeletalMeshData->NormalGenerationMethod;
 		InOutImportOptions.ImportTranslation			= SkeletalMeshData->ImportTranslation;
 		InOutImportOptions.ImportRotation				= SkeletalMeshData->ImportRotation;
 		InOutImportOptions.ImportUniformScale			= SkeletalMeshData->ImportUniformScale;
diff --git a/Engine/Source/Editor/UnrealEd/Private/Fbx/FbxMeshImportData.cpp b/Engine/Source/Editor/UnrealEd/Private/Fbx/FbxMeshImportData.cpp
--- a/Engine/Source/Editor/UnrealEd/Private/Fbx/FbxMeshImportData.cpp
+++ b/Engine/Source/Editor/UnrealEd/Private/Fbx/FbxMeshImportData.cpp
@@ -1,11 +1,13 @@
 // Copyright 1998-2015 Epic Games, Inc. All Rights Reserved.
 
 #include "UnrealEd.h"
+#include "../Classes/Factories/FbxMeshImportData.h"
 
 UFbxMeshImportData::UFbxMeshImportData(const FObjectInitializer& ObjectInitializer)
 	: Super(ObjectInitializer)
 {
 	NormalImportMethod = FBXNIM_ComputeNormals;
+	NormalGenerationMethod = EFBXNormalGenerationMethod::MikkTSpace;
 }
 
 bool UFbxMeshImportData::CanEditChange(const UProperty* InProperty) const
@@ -17,5 +19,14 @@ bool UFbxMeshImportData::CanEditChange(const UProperty* InProperty) const
 		// Let the parent object handle the editability of our properties
 		bMutable = Outer->CanEditChange(InProperty);
 	}
+
+	static FName NormalGenerationMethod("NormalGenerationMethod");
+	if( bMutable && InProperty->GetFName() == NormalGenerationMethod )
+	{
+		// Normal generation method is ignored if we are importing both tangents and normals
+		return NormalImportMethod == FBXNIM_ImportNormalsAndTangents ? false : true;
+	}
+
+
 	return bMutable;
 }
\ No newline at end of file
diff --git a/Engine/Source/Editor/UnrealEd/Private/Fbx/FbxStaticMeshImport.cpp b/Engine/Source/Editor/UnrealEd/Private/Fbx/FbxStaticMeshImport.cpp
--- a/Engine/Source/Editor/UnrealEd/Private/Fbx/FbxStaticMeshImport.cpp
+++ b/Engine/Source/Editor/UnrealEd/Private/Fbx/FbxStaticMeshImport.cpp
@@ -245,7 +245,11 @@ bool UnFbx::FFbxImporter::BuildStaticMeshFromGeometry(FbxMesh* Mesh, UStaticMesh
 	int32 LayerSmoothingCount = Mesh->GetLayerCount(FbxLayerElement::eSmoothing);
 	for(int32 i = 0; i < LayerSmoothingCount; i++)
 	{
-		GeometryConverter->ComputePolygonSmoothingFromEdgeSmoothing (Mesh, i);
+		FbxLayerElementSmoothing const* SmoothingInfo = Mesh->GetLayer(0)->GetSmoothing();
+		if (SmoothingInfo && SmoothingInfo->GetMappingMode() != FbxLayerElement::eByPolygon)
+		{
+			GeometryConverter->ComputePolygonSmoothingFromEdgeSmoothing (Mesh, i);
+		}
 	}
 
 	if (!Mesh->IsTriangleMesh())
@@ -332,14 +336,6 @@ bool UnFbx::FFbxImporter::BuildStaticMeshFromGeometry(FbxMesh* Mesh, UStaticMesh
 	FbxLayerElement::EMappingMode SmoothingMappingMode(FbxLayerElement::eByEdge);
 	if (SmoothingInfo)
 	{
-		if( SmoothingInfo->GetMappingMode() == FbxLayerElement::eByEdge )
-		{
-			if (!GeometryConverter->ComputePolygonSmoothingFromEdgeSmoothing(Mesh))
-			{
-				AddTokenizedErrorMessage(FTokenizedMessage::Create(EMessageSeverity::Warning, FText::Format(LOCTEXT("Error_FailedtoConvertSmoothingGroup", "Unable to fully convert the smoothing groups for mesh '{0}'"), FText::FromString(Mesh->GetName()))), FFbxErrors::Generic_Mesh_ConvertSmoothingGroupFailed);
-				bSmoothingAvailable = false;
-			}
-		}
 
 		if( SmoothingInfo->GetMappingMode() == FbxLayerElement::eByPolygon )
 		{
@@ -1214,7 +1210,7 @@ UStaticMesh* UnFbx::FFbxImporter::ImportStaticMeshAsSingle(UObject* InParent, TA
 		SrcModel.BuildSettings.bRemoveDegenerates = ImportOptions->bRemoveDegenerates;
 		SrcModel.BuildSettings.bRecomputeNormals = ImportOptions->NormalImportMethod == FBXNIM_ComputeNormals;
 		SrcModel.BuildSettings.bRecomputeTangents = ImportOptions->NormalImportMethod != FBXNIM_ImportNormalsAndTangents;
-
+		SrcModel.BuildSettings.bUseMikkTSpace = ImportOptions->NormalGenerationMethod == EFBXNormalGenerationMethod::MikkTSpace;
 		if( ImportOptions->bGenerateLightmapUVs )
 		{
 			SrcModel.BuildSettings.bGenerateLightmapUVs = true;
diff --git a/Engine/Source/Editor/UnrealEd/Private/FbxImporter.h b/Engine/Source/Editor/UnrealEd/Private/FbxImporter.h
--- a/Engine/Source/Editor/UnrealEd/Private/FbxImporter.h
+++ b/Engine/Source/Editor/UnrealEd/Private/FbxImporter.h
@@ -118,6 +118,7 @@ struct FBXImportOptions
 	FRotator ImportRotation;
 	float ImportUniformScale;
 	EFBXNormalImportMethod NormalImportMethod;
+	EFBXNormalGenerationMethod::Type NormalGenerationMethod;
 	// Static Mesh options
 	bool bCombineToSingle;
 	EVertexColorImportOption::Type VertexColorImportOption;
@@ -126,6 +127,7 @@ struct FBXImportOptions
 	bool bGenerateLightmapUVs;
 	bool bOneConvexHullPerUCX;
 	bool bAutoGenerateCollision;
+	
 	FName StaticMeshLODGroup;
 	// Skeletal Mesh options
 	bool bImportMorph;
