Fix erroneous "MAX" option appering in FBX normal generation method dropdown

diff --git a/Engine/Source/Editor/UnrealEd/Classes/Factories/FbxMeshImportData.h b/Engine/Source/Editor/UnrealEd/Classes/Factories/FbxMeshImportData.h
--- a/Engine/Source/Editor/UnrealEd/Classes/Factories/FbxMeshImportData.h
+++ b/Engine/Source/Editor/UnrealEd/Classes/Factories/FbxMeshImportData.h
@@ -21,7 +21,6 @@ namespace EFBXNormalGenerationMethod
 		BuiltIn,
 		/** Use MikkTSpace to generate normals and tangents */
 		MikkTSpace,
-		MAX,
 	};
 }
 
