Fix existing collision not being restored if there is no collision in the FBX file

diff --git a/Engine/Source/Editor/UnrealEd/Private/StaticMeshEdit.cpp b/Engine/Source/Editor/UnrealEd/Private/StaticMeshEdit.cpp
--- a/Engine/Source/Editor/UnrealEd/Private/StaticMeshEdit.cpp
+++ b/Engine/Source/Editor/UnrealEd/Private/StaticMeshEdit.cpp
@@ -981,7 +981,7 @@ void RestoreExistingMeshData(struct ExistingStaticMeshData* ExistingMeshDataPtr,
 		{
 			// If we didn't import anything, always keep collision.
 			bool bKeepCollision;
-			if(!NewMesh->BodySetup)
+			if(!NewMesh->BodySetup || NewMesh->BodySetup->AggGeom.GetElementCount() == 0)
 			{
 				bKeepCollision = true;
 			}
