Fixes to allow Starter Content assets to be deleted cleanly without bogus references being held. Also deleted stale Starter Content packages.

diff --git a/Engine/Plugins/Runtime/SoundMod/Source/SoundMod/Private/SoundNodeModPlayer.cpp b/Engine/Plugins/Runtime/SoundMod/Source/SoundMod/Private/SoundNodeModPlayer.cpp
--- a/Engine/Plugins/Runtime/SoundMod/Source/SoundMod/Private/SoundNodeModPlayer.cpp
+++ b/Engine/Plugins/Runtime/SoundMod/Source/SoundMod/Private/SoundNodeModPlayer.cpp
@@ -6,14 +6,6 @@
 
 #define LOCTEXT_NAMESPACE "SoundNodeModPlayer"
 
-void USoundNodeModPlayer::AddReferencedObjects(UObject* InThis, FReferenceCollector& Collector)
-{
-	USoundNodeModPlayer* This = CastChecked<USoundNodeModPlayer>(InThis);
-	Super::AddReferencedObjects(This, Collector);
-
-	Collector.AddReferencedObject(This->SoundMod);
-}
-
 void USoundNodeModPlayer::LoadAsset()
 {
 	if (IsAsyncLoading())
diff --git a/Engine/Plugins/Runtime/SoundMod/Source/SoundMod/Public/SoundNodeModPlayer.h b/Engine/Plugins/Runtime/SoundMod/Source/SoundMod/Public/SoundNodeModPlayer.h
--- a/Engine/Plugins/Runtime/SoundMod/Source/SoundMod/Public/SoundNodeModPlayer.h
+++ b/Engine/Plugins/Runtime/SoundMod/Source/SoundMod/Public/SoundNodeModPlayer.h
@@ -17,6 +17,7 @@ private:
 	UPROPERTY(EditAnywhere, Category=ModPlayer, meta=(DisplayName="Sound Mod"))
 	TAssetPtr<USoundMod> SoundModAssetPtr;
 
+	UPROPERTY(transient)
 	USoundMod* SoundMod;
 
 	UPROPERTY(EditAnywhere, Category=ModPlayer)
@@ -29,9 +30,6 @@ public:
 	USoundMod* GetSoundMod() const { return SoundMod; }
 	void SetSoundMod(USoundMod* SoundMod);
 
-	// Begin UObject Interface
-	static void AddReferencedObjects(UObject* InThis, FReferenceCollector& Collector);
-
 #if WITH_EDITOR
 	virtual void PostEditChangeProperty(FPropertyChangedEvent& PropertyChangedEvent) override;
 #endif
diff --git a/Engine/Source/Editor/UnrealEd/Private/ObjectTools.cpp b/Engine/Source/Editor/UnrealEd/Private/ObjectTools.cpp
--- a/Engine/Source/Editor/UnrealEd/Private/ObjectTools.cpp
+++ b/Engine/Source/Editor/UnrealEd/Private/ObjectTools.cpp
@@ -1555,6 +1555,23 @@ namespace ObjectTools
 		// Allows deleting of sounds after they have been previewed
 		GEditor->ClearPreviewComponents();
 
+		// Ensure the audio manager is not holding on to any sounds
+		FAudioDeviceManager* AudioDeviceManager = GEditor->GetAudioDeviceManager();
+		if (AudioDeviceManager != nullptr)
+		{
+			AudioDeviceManager->UpdateActiveAudioDevices(false);
+
+			const int32 NumAudioDevices = AudioDeviceManager->GetNumActiveAudioDevices();
+			for (int32 DeviceIndex = 0; DeviceIndex < NumAudioDevices; DeviceIndex++)
+			{
+				FAudioDevice* AudioDevice = AudioDeviceManager->GetAudioDevice(DeviceIndex);
+				if (AudioDevice != nullptr)
+				{
+					AudioDevice->StopAllSounds();
+				}
+			}
+		}
+
 		const FScopedBusyCursor BusyCursor;
 
 		// Make sure packages being saved are fully loaded.
diff --git a/Engine/Source/Runtime/AssetRegistry/Private/AssetRegistry.cpp b/Engine/Source/Runtime/AssetRegistry/Private/AssetRegistry.cpp
--- a/Engine/Source/Runtime/AssetRegistry/Private/AssetRegistry.cpp
+++ b/Engine/Source/Runtime/AssetRegistry/Private/AssetRegistry.cpp
@@ -1161,9 +1161,9 @@ void FAssetRegistry::AssetDeleted(UObject* DeletedAsset)
 				AddEmptyPackage( DeletedObjectPackage->GetFName() );
 
 				// If there is a package metadata object, clear the standalone flag so the package can be truly emptied upon GC
-				if ( DeletedObjectPackage->MetaData != nullptr )
+				if ( UMetaData* MetaData = DeletedObjectPackage->GetMetaData() )
 				{
-					DeletedObjectPackage->MetaData->ClearFlags( RF_Standalone );
+					MetaData->ClearFlags( RF_Standalone );
 				}
 			}
 		}
diff --git a/Engine/Source/Runtime/Engine/Classes/Sound/SoundNodeWavePlayer.h b/Engine/Source/Runtime/Engine/Classes/Sound/SoundNodeWavePlayer.h
--- a/Engine/Source/Runtime/Engine/Classes/Sound/SoundNodeWavePlayer.h
+++ b/Engine/Source/Runtime/Engine/Classes/Sound/SoundNodeWavePlayer.h
@@ -17,6 +17,7 @@ private:
 	UPROPERTY(EditAnywhere, Category=WavePlayer, meta=(DisplayName="Sound Wave"))
 	TAssetPtr<USoundWave> SoundWaveAssetPtr;
 
+	UPROPERTY(transient)
 	USoundWave* SoundWave;
 
 	void OnSoundWaveLoaded(const FName& PackageName, UPackage * Package, EAsyncLoadingResult::Type Result);
@@ -30,8 +31,6 @@ public:
 	ENGINE_API void SetSoundWave(USoundWave* SoundWave);
 
 	// Begin UObject Interface
-	static void AddReferencedObjects(UObject* InThis, FReferenceCollector& Collector);
-
 #if WITH_EDITOR
 	virtual void PostEditChangeProperty(FPropertyChangedEvent& PropertyChangedEvent) override;
 #endif
diff --git a/Engine/Source/Runtime/Engine/Private/SoundNodeWavePlayer.cpp b/Engine/Source/Runtime/Engine/Private/SoundNodeWavePlayer.cpp
--- a/Engine/Source/Runtime/Engine/Private/SoundNodeWavePlayer.cpp
+++ b/Engine/Source/Runtime/Engine/Private/SoundNodeWavePlayer.cpp
@@ -7,14 +7,6 @@
 
 #define LOCTEXT_NAMESPACE "SoundNodeWavePlayer"
 
-void USoundNodeWavePlayer::AddReferencedObjects(UObject* InThis, FReferenceCollector& Collector)
-{
-	USoundNodeWavePlayer* This = CastChecked<USoundNodeWavePlayer>(InThis);
-	Super::AddReferencedObjects(This, Collector);
-
-	Collector.AddReferencedObject(This->SoundWave);
-}
-
 void USoundNodeWavePlayer::LoadAsset()
 {
 	if (IsAsyncLoading())
