Added initial support for multiprocess cooking.

diff --git a/Engine/Source/Developer/TargetPlatform/Private/TargetPlatformManagerModule.cpp b/Engine/Source/Developer/TargetPlatform/Private/TargetPlatformManagerModule.cpp
--- a/Engine/Source/Developer/TargetPlatform/Private/TargetPlatformManagerModule.cpp
+++ b/Engine/Source/Developer/TargetPlatform/Private/TargetPlatformManagerModule.cpp
@@ -587,12 +587,32 @@ protected:
 				if (Platform != nullptr)
 				{
 					// would like to move this check to GetActiveTargetPlatforms, but too many things cache this result
-					// this setup will become faster after TTP 341897 is complete.					
+					// this setup will become faster after TTP 341897 is complete.
+RETRY_SETUPANDVALIDATE:
 					if (SetupAndValidateAutoSDK(Platform->GetPlatformInfo().AutoSDKPath))
 					{
 						UE_LOG(LogTemp, Display, TEXT("Loaded TP %s"), *Modules[Index].ToString());
 						Platforms.Add(Platform);
 					}
+					else
+					{
+						// this hack is here because if you try and setup and validate autosdk some times it will fail because shared files are in use by another child cooker
+						static bool bIsChildCooker = FParse::Param(FCommandLine::Get(), TEXT("cookchild"));
+						if (bIsChildCooker)
+						{
+							static int Counter = 0;
+							++Counter;
+							if (Counter < 10)
+							{
+								goto RETRY_SETUPANDVALIDATE;
+							}
+						}
+						UE_LOG(LogTemp, Display, TEXT("Failed to SetupAndValidateAutoSDK for platform %s"), *Modules[Index].ToString());
+					}
+				}
+				else
+				{
+					UE_LOG(LogTemp, Display, TEXT("Failed to get target platform %s"), *Modules[Index].ToString());
 				}
 			}
 		}
diff --git a/Engine/Source/Editor/UnrealEd/Classes/CookOnTheSide/CookOnTheFlyServer.h b/Engine/Source/Editor/UnrealEd/Classes/CookOnTheSide/CookOnTheFlyServer.h
--- a/Engine/Source/Editor/UnrealEd/Classes/CookOnTheSide/CookOnTheFlyServer.h
+++ b/Engine/Source/Editor/UnrealEd/Classes/CookOnTheSide/CookOnTheFlyServer.h
@@ -421,6 +421,7 @@ private:
 		TMap<FName, TArray<FName>> PlatformList;
 		mutable FCriticalSection SynchronizationObject;
 	public:
+		const TArray<FName> &GetQueue() const { return Queue;  }
 		void EnqueueUnique(const FFilePlatformRequest& Request, bool ForceEnqueFront = false)
 		{
 			FScopeLock ScopeLock( &SynchronizationObject );
@@ -588,6 +589,17 @@ private:
 	FString OutputDirectoryOverride;
 	//////////////////////////////////////////////////////////////////////////
 	// Cook by the book options
+	struct FChildCooker
+	{
+		FChildCooker() : ReadPipe(nullptr), ReturnCode(-1), bFinished(false) { }
+
+		FProcHandle ProcessHandle;
+		FString ResponseFileName;
+		FString BaseResponseFileName;
+		void* ReadPipe;
+		int32 ReturnCode;
+		bool bFinished;
+	};
 	struct FCookByTheBookOptions
 	{
 	public:
@@ -599,7 +611,8 @@ private:
 			CookStartTime( 0.0 ), 
 			bErrorOnEngineContentUse(false),
 			bForceEnableCompressedPackages(false),
-			bForceDisableCompressedPackages(false)
+			bForceDisableCompressedPackages(false),
+			bIsChildCooker(false)
 		{ }
 
 		/** Should we test for UObject leaks */
@@ -634,6 +647,8 @@ private:
 		/** force this cook by the book to enable \ disable package compression */
 		bool bForceEnableCompressedPackages;
 		bool bForceDisableCompressedPackages;
+		bool bIsChildCooker;
+		TArray<FChildCooker> ChildCookers;
 	};
 	FCookByTheBookOptions* CookByTheBookOptions;
 	
@@ -766,16 +781,18 @@ public:
 		FString DLCName;
 		FString CreateReleaseVersion;
 		FString BasedOnReleaseVersion;
+		FString ChildCookFileName; // if we are the child cooker 
 		bool bGenerateStreamingInstallManifests; 
 		bool bGenerateDependenciesForMaps; 
 		bool bErrorOnEngineContentUse; // this is a flag for dlc, will cause the cooker to error if the dlc references engine content
-
+		int32 NumProcesses;
 		FCookByTheBookStartupOptions() :
 			CookOptions(ECookByTheBookOptions::None),
 			DLCName(FString()),
 			bGenerateStreamingInstallManifests(false),
 			bGenerateDependenciesForMaps(false),
-			bErrorOnEngineContentUse(false)
+			bErrorOnEngineContentUse(false),
+			NumProcesses(0)
 		{ }
 	};
 
@@ -897,6 +914,14 @@ public:
 	 */
 	void MarkPackageDirtyForCooker( UPackage *Package );
 
+	/**
+	 * MaybeMarkPackageAsAlreadyLoaded
+	 * Mark the package as already loaded if we have already cooked the package for all requested target platforms
+	 * this hints to the objects on load that we don't need to load all our bulk data
+	 * 
+	 * @param Package to mark as not requiring reload
+	 */
+	void MaybeMarkPackageAsAlreadyLoaded(UPackage* Package);
 
 
 private:
@@ -908,7 +933,7 @@ private:
 	 */
 	void CollectFilesToCook(TArray<FName>& FilesInPath, 
 		const TArray<FString>& CookMaps, const TArray<FString>& CookDirectories, const TArray<FString>& CookCultures, 
-		const TArray<FString>& IniMapSections, bool bCookAll, bool bMapsOnly, bool bNoDev);
+		const TArray<FString>& IniMapSections, bool bCookAll, bool bMapsOnly, bool bNoDev, const FString& ChildCookFilename);
 
 	/**
 	 * AddFileToCook add file to cook list 
@@ -920,6 +945,22 @@ private:
 	 */
 	void CookByTheBookFinished();
 
+	/**
+	 * StartChildCookers to help out with cooking
+	 * only valid in cook by the book (not from the editor)
+	 * 
+	 * @param NumChildCookersToSpawn, number of child cookers we want to use
+	 */
+	void StartChildCookers(int32 NumChildCookersToSpawn, const TArray<FName>& TargetPlatformNames );
+
+	/**
+	 * TickChildCookers
+	 * output the information form the child cookers to the main cooker output
+	 * 
+	 * @return return true if all child cookers are finished
+	 */
+	bool TickChildCookers();
+
 	/**
 	 * Get all the packages which are listed in asset registry passed in.  
 	 *
@@ -929,6 +970,20 @@ private:
 	 */
 	bool GetAllPackagesFromAssetRegistry( const FString& AssetRegistryPath, TArray<FName>& OutPackageNames ) const;
 
+	/**
+	 * IsChildCooker, 
+	 * returns if this cooker is a sue chef for some other master chef.
+	 */
+	bool IsChildCooker() const
+	{
+		if (IsCookByTheBookMode())
+		{
+			return CookByTheBookOptions->bIsChildCooker;
+		}
+		return false;
+	}
+
+
 
 	//////////////////////////////////////////////////////////////////////////
 	// cook on the fly specific functions
@@ -983,12 +1038,21 @@ private:
 
 	/**
 	 * GetDependencies
+	 * get package dependencies according to the asset registry
 	 * 
 	 * @param Packages List of packages to use as the root set for dependency checking
 	 * @param Found return value, all objects which package is dependent on
 	 */
 	void GetDependentPackages( const TSet<UPackage*>& Packages, TSet<FName>& Found);
 
+	/**
+	 * GetDependencies
+	 * get package dependencies according to the asset registry
+	 *
+	 * @param Root set of packages to use when looking for dependencies
+	 * @param FoundPackages list of packages which were found
+	 */
+	void GetDependentPackages(const TSet<FName>& RootPackages, TSet<FName>& FoundPackages);
 	/**
 	 * GenerateManifestInfo
 	 * generate the manifest information for a given package
@@ -997,6 +1061,35 @@ private:
 	 */
 	void GenerateManifestInfo( UPackage* Package, const TArray<FName>& TargetPlatformNames );
 
+
+	/**
+	 * GenerateManifestInfo
+	 * generate manfiest information for the given package and add it to the manifest
+	 * this version is teh same as GenerateManifestInfo which takes in a UPackageObject except that doesn't require the package to be loaded
+	 *
+	 * @param Package name of the package to generate the manifest information for
+	 * @param TargetPlatformNames to add the manifest information to
+	 */
+	void GenerateManifestInfo(const FName& Package, const TArray<FName>& TargetPlatformNames);
+
+	/**
+	 * ContainsWorld
+	 * use the asset registry to determine if a Package contains a UWorld or ULevel object
+	 * 
+	 * @param PackageName to return if it contains the a UWorld object or a ULevel
+	 * @return true if the Package contains a UWorld or ULevel false otherwise
+	 */
+	bool ContainsMap(const FName& PackageName) const;
+
+	/**
+	 * AddDependenciesToManifest
+	 * Add the dependencies of this package to the cooked package manifest
+	 * use information from the asset registry to generate this information
+	 *
+	 * @params StandardFilename of the package which we want dependencies to be copied over
+	 */
+	void AddDependenciesToManifest(const FName& StandardFilename, const FName& LastLoadedMapName, const TSet<FName>& Dependencies, const TArray<FName>& TargetPlatformNames);
+
 	/**
 	 * GetCurrentIniVersionStrings gets the current ini version strings for compare against previous cook
 	 * 
diff --git a/Engine/Source/Editor/UnrealEd/Private/Commandlets/CookCommandlet.cpp b/Engine/Source/Editor/UnrealEd/Private/Commandlets/CookCommandlet.cpp
--- a/Engine/Source/Editor/UnrealEd/Private/Commandlets/CookCommandlet.cpp
+++ b/Engine/Source/Editor/UnrealEd/Private/Commandlets/CookCommandlet.cpp
@@ -1053,6 +1053,12 @@ bool UCookCommandlet::NewCook( const TArray<ITargetPlatform*>& Platforms, TArray
 	FString DLCName;
 	FParse::Value( *Params, TEXT("DLCNAME="), DLCName);
 
+	FString ChildCookFile;
+	FParse::Value(*Params, TEXT("cookchild="), ChildCookFile);
+
+	int32 NumProcesses = 0;
+	FParse::Value(*Params, TEXT("multiprocesscooker="), NumProcesses);
+
 	FString BasedOnReleaseVersion;
 	FParse::Value( *Params, TEXT("BasedOnReleaseVersion="), BasedOnReleaseVersion);
 
@@ -1175,7 +1181,8 @@ bool UCookCommandlet::NewCook( const TArray<ITargetPlatform*>& Platforms, TArray
 	StartupOptions.bErrorOnEngineContentUse = bErrorOnEngineContentUse;
 	StartupOptions.bGenerateDependenciesForMaps = Switches.Contains(TEXT("GenerateDependenciesForMaps"));
 	StartupOptions.bGenerateStreamingInstallManifests = bGenerateStreamingInstallManifests;
-
+	StartupOptions.ChildCookFileName = ChildCookFile;
+	StartupOptions.NumProcesses = NumProcesses;
 
 	CookOnTheFlyServer->StartCookByTheBook( StartupOptions );
 
diff --git a/Engine/Source/Editor/UnrealEd/Private/CookOnTheFlyServer.cpp b/Engine/Source/Editor/UnrealEd/Private/CookOnTheFlyServer.cpp
--- a/Engine/Source/Editor/UnrealEd/Private/CookOnTheFlyServer.cpp
+++ b/Engine/Source/Editor/UnrealEd/Private/CookOnTheFlyServer.cpp
@@ -135,6 +135,22 @@ public:
 	}
 };
 
+static TMap<FName, int32> IntStats;
+void IncIntStat(const FName& Name, const int32 Amount)
+{
+	int32* Value = IntStats.Find(Name);
+
+	if (Value == NULL)
+	{
+		IntStats.Add(Name, Amount);
+	}
+	else
+	{
+		*Value += Amount;
+	}
+}
+
+
 static FHierarchicalTimerInfo RootTimerInfo(MoveTemp(FString(TEXT("Root"))));
 static FHierarchicalTimerInfo* CurrentTimerInfo = &RootTimerInfo;
 #endif
@@ -317,7 +333,7 @@ void OutputTimers()
 	GTimerInfo.Empty();
 #endif
 }
-
+#if HEIRARCHICAL_TIMER
 void OutputHierarchyTimers(const FHierarchicalTimerInfo* TimerInfo, FString& Output, int32& Depth)
 {
 	// put our line into the output buffer
@@ -326,7 +342,7 @@ void OutputHierarchyTimers(const FHierarchicalTimerInfo* TimerInfo, FString& Out
 		Output += TEXT("\t");
 	}*/
 	
-	Output += FString::Printf(TEXT("%s: %fms\n"), *TimerInfo->Name, TimerInfo->Length * 1000);
+	Output += FString::Printf(TEXT("%s: %fms\r\n"), *TimerInfo->Name, TimerInfo->Length * 1000);
 
 	++Depth;
 	for (const auto& ChildInfo : TimerInfo->Children)
@@ -344,12 +360,17 @@ void OutputHierarchyTimers()
 	FString Output = TEXT("Hierarchy timer information\nName:  Length(ms)");
 	OutputHierarchyTimers(TimerInfo, Output, Depth);
 
+	Output += FString::Printf(TEXT("IntStats\n"));
+	for (const auto& IntStat : IntStats)
+	{
+		Output += FString::Printf(TEXT("%s=%d\r\n"), *IntStat.Key.ToString(), IntStat.Value);
+	}
 
 	UE_LOG(LogCook, Display, TEXT("%s"), *Output);
 
 	RootTimerInfo.ClearChildren();
 }
-
+#endif
 
 #define CREATE_TIMER(name, incrementScope) FScopeTimer ScopeTimer##name(#name, incrementScope); 
 
@@ -361,6 +382,11 @@ void OutputHierarchyTimers()
 #define ACCUMULATE_TIMER_SCOPE(name) FScopeTimer ScopeTimerInner##name(ScopeTimer##name); ScopeTimerInner##name.Start();
 #define ACCUMULATE_TIMER_START(name) ScopeTimer##name.Start();
 #define ACCUMULATE_TIMER_STOP(name) ScopeTimer##name.Stop();
+#if HEIRARCHICAL_TIMER
+#define INC_INT_STAT( name, amount ) { static FName StaticName##name(#name); IncIntStat(StaticName##name, amount); }
+#else
+#define INC_INT_STAT( name, amount ) 
+#endif
 
 #define OUTPUT_TIMERS() OutputTimers();
 #define OUTPUT_HIERARCHYTIMERS() OutputHierarchyTimers();
@@ -375,6 +401,7 @@ void OutputHierarchyTimers()
 #define ACCUMULATE_TIMER_SCOPE(name) 
 #define ACCUMULATE_TIMER_START(name) 
 #define ACCUMULATE_TIMER_STOP(name) 
+#define INC_INT_STAT( name, amount ) 
 
 #define OUTPUT_TIMERS()
 #define OUTPUT_HIERARCHYTIMERS()
@@ -797,21 +824,41 @@ public:
 	virtual FString GetArchiveName() const override { return TEXT("FArchiveFindReferences"); }
 };
 
+void UCookOnTheFlyServer::GetDependentPackages(const TSet<UPackage*>& RootPackages, TSet<FName>& FoundPackages)
+{
+	check(IsChildCooker() == false);
+
+	static const FName AssetRegistryName("AssetRegistry");
+	FAssetRegistryModule& AssetRegistryModule = FModuleManager::LoadModuleChecked<FAssetRegistryModule>(AssetRegistryName);
+	IAssetRegistry& AssetRegistry = AssetRegistryModule.Get();
+
+	TSet<FName> RootPackageFNames;
+	for (const UPackage* RootPackage : RootPackages)
+	{
+		RootPackageFNames.Add(RootPackage->GetFName());
+	}
+
+
+	GetDependentPackages(RootPackageFNames, FoundPackages);
+
+}
+
 
-void UCookOnTheFlyServer::GetDependentPackages( const TSet<UPackage*>& RootPackages, TSet<FName>& FoundPackages )
+void UCookOnTheFlyServer::GetDependentPackages( const TSet<FName>& RootPackages, TSet<FName>& FoundPackages )
 {
+	check(IsChildCooker() == false);
+
 	static const FName AssetRegistryName("AssetRegistry");
 	FAssetRegistryModule& AssetRegistryModule = FModuleManager::LoadModuleChecked<FAssetRegistryModule>(AssetRegistryName);
 	IAssetRegistry& AssetRegistry = AssetRegistryModule.Get();
 
 	TArray<FName> FoundPackagesArray;
-	for ( const UPackage* RootPackage : RootPackages)
+	for (const auto& RootPackage : RootPackages)
 	{
-		FoundPackagesArray.Add( RootPackage->GetFName() );
-		FoundPackages.Add( RootPackage->GetFName() );
+		FoundPackagesArray.Add(RootPackage);
+		FoundPackages.Add(RootPackage);
 	}
 
-
 	int FoundPackagesCounter = 0;
 	while ( FoundPackagesCounter < FoundPackagesArray.Num() )
 	{
@@ -918,14 +965,62 @@ void UCookOnTheFlyServer::GetDependencies( const TSet<UPackage*>& Packages, TSet
 	FArchiveFindReferences ArFindReferences( RootSet, Found, Exclude );
 }
 
+bool UCookOnTheFlyServer::ContainsMap(const FName& PackageName) const
+{
+	static const FName AssetRegistryName("AssetRegistry");
+	FAssetRegistryModule& AssetRegistryModule = FModuleManager::LoadModuleChecked<FAssetRegistryModule>(AssetRegistryName);
+	IAssetRegistry& AssetRegistry = AssetRegistryModule.Get();
+
+
+	TArray<FAssetData> Assets;
+	ensure( AssetRegistry.GetAssetsByPackageName(PackageName, Assets) );
+
+	for (const auto& Asset : Assets)
+	{
+		if (Asset.GetClass()->IsChildOf(UWorld::StaticClass()) || Asset.GetClass()->IsChildOf(ULevel::StaticClass()))
+		{
+			return true;
+		}
+	}
+	return false;
+}
+
+
+// don't need the package loaded to generate the manifest infomation
+void UCookOnTheFlyServer::GenerateManifestInfo(const FName& Package, const TArray<FName>& TargetPlatformNames)
+{
+	if (IsCookByTheBookMode())
+		return;
+	
+	check(IsChildCooker() == false);
+
+	TSet<FName> RootPackages;
+	RootPackages.Add(Package);
+
+	// see if teh package contains a world
+	bool bIsWorldPackage = ContainsMap(Package);
+
+	TSet<FName> FoundPackages;
+	GetDependentPackages(RootPackages, FoundPackages);
+
+
+	FName StandardFilename = GetCachedStandardPackageFileFName(Package);
+	FName MapName = NAME_None;
+	if (bIsWorldPackage)
+	{
+		MapName = Package;
+	}
+	AddDependenciesToManifest(StandardFilename, MapName, FoundPackages, TargetPlatformNames);
+}
 
 void UCookOnTheFlyServer::GenerateManifestInfo(UPackage* Package, const TArray<FName>& TargetPlatformNames)
 {
 	if (!CookByTheBookOptions)
 		return;
 
-	// generate dependency information for this package
+	check(IsChildCooker() == false);
 
+	// generate dependency information for this package
 
 	TSet<UPackage*> RootPackages;
 	RootPackages.Add(Package);
@@ -935,6 +1030,11 @@ void UCookOnTheFlyServer::GenerateManifestInfo(UPackage* Package, const TArray<F
 	// load sublevels
 	UWorld* World = UWorld::FindWorldInPackage(Package);
 
+#define VERIFY_LEVELS_IN_DEPENDENCIES_LIST 1
+#if VERIFY_LEVELS_IN_DEPENDENCIES_LIST
+	TArray<FString> LevelPackageNames;
+#endif
+
 	if (World)
 	{
 		for (const auto& StreamingLevel : World->StreamingLevels)
@@ -942,6 +1042,9 @@ void UCookOnTheFlyServer::GenerateManifestInfo(UPackage* Package, const TArray<F
 			if (StreamingLevel->GetLoadedLevel())
 			{
 				RootPackages.Add(StreamingLevel->GetLoadedLevel()->GetOutermost());
+#if VERIFY_LEVELS_IN_DEPENDENCIES_LIST
+				LevelPackageNames.AddUnique(StreamingLevel->GetLoadedLevel()->GetOutermost()->GetName());
+#endif
 			}
 		}
 
@@ -953,6 +1056,13 @@ void UCookOnTheFlyServer::GenerateManifestInfo(UPackage* Package, const TArray<F
 			World->WorldComposition->CollectTilesToCook(NewPackagesToCook);
 		}
 
+#if VERIFY_LEVELS_IN_DEPENDENCIES_LIST
+		for (const auto& PackageToCook : NewPackagesToCook)
+		{
+			LevelPackageNames.AddUnique(PackageToCook);
+		}
+#endif
+
 		for (const auto& PackageName : NewPackagesToCook)
 		{
 			UPackage* PackageToCook = LoadPackage(NULL, *PackageName, LOAD_None);
@@ -977,32 +1087,20 @@ void UCookOnTheFlyServer::GenerateManifestInfo(UPackage* Package, const TArray<F
 
 	auto GetPackageFName = [](const FName& Name) { return Name; };
 
-#if VERIFY_GETDEPENDENTPACKAGES
-
-	TSet<UObject*> Dependencies; 
-	GetDependencies( RootPackages, Dependencies );
 
-	TSet<UPackage*> VerifyPackages;
-	for ( const auto& Object : Dependencies )
-	{
-		VerifyPackages.Add( Object->GetOutermost() );
-	}
-	check( VerifyPackages.Num() == Packages.Num() );
-
-
-
-	for ( const auto& VerifyPackage : VerifyPackages )
+#if VERIFY_LEVELS_IN_DEPENDENCIES_LIST
+	if (World)
 	{
-		check( Packages.Contains( VerifyPackage->GetFName() ) );
-
-		TArray<FName> PackagePaths;
-		AssetRegistry.GetPackagePathsByPackageName(VerifyPackage->GetFName(), PackagePaths);
-		check( PackagePaths.Num() == 0 );
+		TSet<UPackage*> WorldPackages;
+		WorldPackages.Add(Package);
+		TSet<FName> DependentWorldPackages;
+		GetDependentPackages(WorldPackages, DependentWorldPackages);
 
-		check( PackagePaths[0].ToString() == VerifyPackage->GetPathName() );
+		for (const auto& LevelPackageName : LevelPackageNames)
+		{
+			check(DependentWorldPackages.Contains(FName(*LevelPackageName)));
+		}
 	}
-
-
 #endif
 
 #else
@@ -1019,21 +1117,30 @@ void UCookOnTheFlyServer::GenerateManifestInfo(UPackage* Package, const TArray<F
 #endif
 
 
-
-
+	
 	FName StandardFilename = GetCachedStandardPackageFileFName(Package);
+	FName MapName = NAME_None;
+	if (Package->ContainsMap())
+	{
+		MapName = Package->GetFName();
+	}
+	AddDependenciesToManifest(StandardFilename, MapName, Packages, TargetPlatformNames);
+	
+}
 
+void UCookOnTheFlyServer::AddDependenciesToManifest(const FName& StandardFilename, const FName& MapName, const TSet<FName>& Dependencies, const TArray<FName>& TargetPlatformNames)
+{
+	FString LastLoadedMapName;
+	if (MapName != NAME_None)
+	{
+		LastLoadedMapName = MapName.ToString();
+	}
+	// Add package dependencies to the cooked package manifest
 	if (CookByTheBookOptions &&
 		CookByTheBookOptions->bGenerateDependenciesForMaps &&
-		Package->ContainsMap() &&
-		(World != nullptr))
+		(MapName != NAME_None))
 	{
-		TSet <FName> Names;
-		for (auto Object : Packages)
-		{
-			Names.Add(GetPackageFName(Object));
-		}
-		CookByTheBookOptions->MapDependencyGraph.Add(Package->GetFName(), Names);
+		CookByTheBookOptions->MapDependencyGraph.Add(MapName, Dependencies);
 	}
 
 	// update the manifests with generated dependencies
@@ -1046,12 +1153,12 @@ void UCookOnTheFlyServer::GenerateManifestInfo(UPackage* Package, const TArray<F
 			ManifestGenerator->PrepareToLoadNewPackage(StandardFilename.ToString());
 		}
 
-		for (const auto& DependentPackage : Packages)
+		for (const auto& DependentPackage : Dependencies)
 		{
 			ManifestGenerator->OnLastPackageLoaded(DependentPackage);
 		}
 
-		for (const auto& DependentPackage : Packages)
+		for (const auto& DependentPackage : Dependencies)
 		{
 			FString Filename = GetCachedPackageFilename(DependentPackage);
 			if (!Filename.IsEmpty())
@@ -1101,6 +1208,52 @@ FString UCookOnTheFlyServer::GetDLCContentPath()
 
 COREUOBJECT_API extern bool GOutputCookingWarnings;
 
+bool UCookOnTheFlyServer::TickChildCookers()
+{
+	bool bChildCookersFinished = true;
+	if (IsCookByTheBookMode())
+	{
+		for (auto& ChildCooker : CookByTheBookOptions->ChildCookers)
+		{
+			if (ChildCooker.bFinished == false)
+			{
+				bool bIsCookerFinished = false;
+				int32 ReturnCode = 0;
+				if (FPlatformProcess::GetProcReturnCode(ChildCooker.ProcessHandle, &ReturnCode))
+				{
+					UE_LOG(LogCook, Display, TEXT("ChildCooker returned error code %d"), ReturnCode);
+					ChildCooker.ReturnCode = ReturnCode;
+					ChildCooker.bFinished = true;
+					if (ReturnCode != 0)
+					{
+						UE_LOG(LogCook, Fatal, TEXT("ChildCooker returned error code %d"), ReturnCode);
+					}
+				}
+				else
+				{
+					bChildCookersFinished = false;
+				}
+
+				// process the log output from the child cooker even if we just finished to ensure that we get the end of the log
+				FString PipeContents = FPlatformProcess::ReadPipe(ChildCooker.ReadPipe);
+				if (PipeContents.Len())
+				{
+					TArray<FString> PipeLines;
+					PipeContents.ParseIntoArrayLines(PipeLines);
+
+					for (const auto& Line : PipeLines)
+					{
+						UE_LOG(LogCook, Display, TEXT("Cooker output %s: %s"), *ChildCooker.BaseResponseFileName, *Line);
+					}
+				}
+			}
+
+		}
+	}
+	return bChildCookersFinished;
+}
+
+
 uint32 UCookOnTheFlyServer::TickCookOnTheSide( const float TimeSlice, uint32 &CookedPackageCount )
 {
 	struct FCookerTimer
@@ -1145,19 +1298,26 @@ uint32 UCookOnTheFlyServer::TickCookOnTheSide( const float TimeSlice, uint32 &Co
 
 	uint32 Result = 0;
 
-	static const FName AssetRegistryName("AssetRegistry");
-	FAssetRegistryModule& AssetRegistryModule = FModuleManager::LoadModuleChecked<FAssetRegistryModule>(AssetRegistryName);
-	IAssetRegistry& AssetRegistry = AssetRegistryModule.Get();
-	if ( AssetRegistry.IsLoadingAssets() )
+	if (IsChildCooker() == false)
 	{
-		// early out
-		return Result;
+		static const FName AssetRegistryName("AssetRegistry");
+		FAssetRegistryModule& AssetRegistryModule = FModuleManager::LoadModuleChecked<FAssetRegistryModule>(AssetRegistryName);
+		IAssetRegistry& AssetRegistry = AssetRegistryModule.Get();
+		if (AssetRegistry.IsLoadingAssets())
+		{
+			// early out
+			return Result;
+		}
 	}
 
 	// This is all the target platforms which we needed to process requests for this iteration
 	// we use this in the unsolicited packages processing below
 	TArray<FName> AllTargetPlatformNames;
 
+	{
+		SCOPE_TIMER(TickChildCookersOutput);
+		TickChildCookers();
+	}
 	while (!GIsRequestingExit || CurrentCookMode == ECookMode::CookByTheBook)
 	{
 		// if we just cooked a map then don't process anything the rest of this tick
@@ -1297,6 +1457,7 @@ uint32 UCookOnTheFlyServer::TickCookOnTheSide( const float TimeSlice, uint32 &Co
 
 				SCOPE_TIMER(LoadPackage);
 				Package = LoadPackage( NULL, *BuildFilename, LOAD_None );
+				INC_INT_STAT(LoadPackage, 1);
 			}
 #if DEBUG_COOKONTHEFLY
 			else
@@ -1318,6 +1479,8 @@ uint32 UCookOnTheFlyServer::TickCookOnTheSide( const float TimeSlice, uint32 &Co
 
 				if (Package->ContainsMap())
 				{
+					// child cookers can't process maps
+					// check(IsChildCooker() == false);
 					if ( IsCookByTheBookMode() )
 					{
 						// load sublevels
@@ -1343,9 +1506,10 @@ uint32 UCookOnTheFlyServer::TickCookOnTheSide( const float TimeSlice, uint32 &Co
 							FString Filename;
 							if ( FPackageName::DoesPackageExist(PackageName, NULL, &Filename) )
 							{
+								check(IsChildCooker() == false);
+
 								FString StandardFilename = FPaths::ConvertRelativePathToFull(Filename);
 								FPaths::MakeStandardFilename(StandardFilename);
-
 								CookRequests.EnqueueUnique(FFilePlatformRequest(MoveTemp(FName(*StandardFilename)), MoveTemp(TargetPlatformNames)));
 							}
 						}
@@ -1504,6 +1668,8 @@ uint32 UCookOnTheFlyServer::TickCookOnTheSide( const float TimeSlice, uint32 &Co
 		int32 FirstUnsolicitedPackage = PackagesToSave.Num();
 
 		// generate a list of other packages which were loaded with this one
+		// if we are a child cooker then we don't have unsolicited packages
+		if ( IsChildCooker() == false)
 		{
 			SCOPE_TIMER(UnsolicitedMarkup);
 
@@ -1647,12 +1813,6 @@ uint32 UCookOnTheFlyServer::TickCookOnTheSide( const float TimeSlice, uint32 &Co
 					}
 				}
 
-				{
-					SCOPE_TIMER(GenerateManifestInfo);
-					// update manifest with cooked package info
-					GenerateManifestInfo( PackagesToSave[I], AllTargetPlatformNames );
-				}
-
 				FName PackageFName = GetCachedStandardPackageFileFName(Package);
 				TArray<FName> SaveTargetPlatformNames = AllTargetPlatformNames;
 				TArray<FName> CookedTargetPlatforms;
@@ -1670,6 +1830,13 @@ uint32 UCookOnTheFlyServer::TickCookOnTheSide( const float TimeSlice, uint32 &Co
 					continue;
 				}
 
+				if (IsChildCooker()==false)
+				{
+					SCOPE_TIMER(GenerateManifestInfo);
+					// update manifest with cooked package info
+					GenerateManifestInfo(PackagesToSave[I], AllTargetPlatformNames);
+				}
+
 				// precache platform data for next package 
 				UPackage *NextPackage = PackagesToSave[FMath::Min( PackagesToSave.Num()-1, I + 1 )];
 				UPackage *NextNextPackage = PackagesToSave[FMath::Min( PackagesToSave.Num()-1, I + 2 )];
@@ -2213,6 +2380,7 @@ bool UCookOnTheFlyServer::SaveCookedPackage( UPackage* Package, uint32 SaveFlags
 				{
 					SCOPE_TIMER(GEditorSavePackage);
 					bSavedCorrectly &= GEditor->SavePackage( Package, World, Flags, *PlatFilename, GError, NULL, bSwap, false, SaveFlags, Target, FDateTime::MinValue(), false );
+					INC_INT_STAT(SavedPackage, 1);
 				}
 				
 				bOutWasUpToDate = false;
@@ -2243,6 +2411,18 @@ void UCookOnTheFlyServer::Initialize( ECookMode::Type DesiredCookMode, ECookInit
 	CurrentCookMode = DesiredCookMode;
 	CookFlags = InCookFlags;
 
+	if (IsCookByTheBookMode() && !IsCookingInEditor())
+	{
+		FCoreUObjectDelegates::PackageCreatedForLoad.AddUObject(this, &UCookOnTheFlyServer::MaybeMarkPackageAsAlreadyLoaded);
+	}
+
+	if (IsCookingInEditor())
+	{
+		FCoreUObjectDelegates::OnObjectPropertyChanged.AddUObject(this, &UCookOnTheFlyServer::OnObjectPropertyChanged);
+		FCoreUObjectDelegates::OnObjectModified.AddUObject(this, &UCookOnTheFlyServer::OnObjectModified);
+		FCoreUObjectDelegates::OnObjectSaved.AddUObject(this, &UCookOnTheFlyServer::OnObjectSaved);
+	}
+
 	bool bUseFullGCAssetClassNames = true;
 	GConfig->GetBool(TEXT("CookSettings"), TEXT("bUseFullGCAssetClassNames"), bUseFullGCAssetClassNames, GEditorIni);
 	
@@ -2557,6 +2737,7 @@ bool UCookOnTheFlyServer::IniSettingsOutOfDate( const ITargetPlatform* TargetPla
 
 bool UCookOnTheFlyServer::SaveCurrentIniSettings( const ITargetPlatform* TargetPlatform ) const
 {
+	check(IsChildCooker()==false);
 	const FString CookedIni = FPaths::GameDir() / TEXT("CookedIniVersion.txt");
 	const FString SandboxCookedIni = ConvertToFullSandboxPath(CookedIni, true);
 
@@ -2577,6 +2758,8 @@ bool UCookOnTheFlyServer::SaveCurrentIniSettings( const ITargetPlatform* TargetP
 
 void UCookOnTheFlyServer::PopulateCookedPackagesFromDisk( const TArray<ITargetPlatform*>& Platforms )
 {
+	check(IsChildCooker() == false);
+
 	// check each package to see if it's out of date
 	FPackageDependencyInfoModule& PDInfoModule = FModuleManager::LoadModuleChecked<FPackageDependencyInfoModule>("PackageDependencyInfo");
 
@@ -2651,6 +2834,9 @@ void UCookOnTheFlyServer::PopulateCookedPackagesFromDisk( const TArray<ITargetPl
 
 void UCookOnTheFlyServer::CleanSandbox( const bool bIterative )
 {
+	//child cookers shouldn't clean sandbox we would be deleting the master cookers / other cookers data
+	check(IsChildCooker() == false);
+
 	ITargetPlatformManagerModule& TPM = GetTargetPlatformManagerRef();
 	const TArray<ITargetPlatform*>& Platforms = TPM.GetCookingTargetPlatforms();	
 
@@ -2814,6 +3000,12 @@ void UCookOnTheFlyServer::CleanSandbox( const bool bIterative )
 
 void UCookOnTheFlyServer::GenerateAssetRegistry(const TArray<ITargetPlatform*>& Platforms)
 {
+	if (IsChildCooker())
+	{
+		// don't generate the asset registry
+		return;
+	}
+
 	if (IsCookingInEditor() == false)
 	{
 		// we want to register the temporary save directory if we are cooking outside the editor.  
@@ -2933,10 +3125,26 @@ void UCookOnTheFlyServer::AddFileToCook( TArray<FName>& InOutFilesToCook, const
 	}
 }
 
-void UCookOnTheFlyServer::CollectFilesToCook(TArray<FName>& FilesInPath, const TArray<FString> &CookMaps, const TArray<FString> &InCookDirectories, const TArray<FString> &CookCultures, const TArray<FString> &IniMapSections, bool bCookAll, bool bMapsOnly, bool bNoDev)
+void UCookOnTheFlyServer::CollectFilesToCook(TArray<FName>& FilesInPath, const TArray<FString> &CookMaps, const TArray<FString> &InCookDirectories, const TArray<FString> &CookCultures, const TArray<FString> &IniMapSections, bool bCookAll, bool bMapsOnly, bool bNoDev, const FString& ChildCookFilename )
 {
-	TArray<FString> CookDirectories = InCookDirectories;
+	if (CookByTheBookOptions->bIsChildCooker)
+	{
+		check(ChildCookFilename.Len());
+		FString ChildCookString;
+		ensure(FFileHelper::LoadFileToString(ChildCookString, *ChildCookFilename));
 	
+		TArray<FString> ChildCookArray;
+		ChildCookString.ParseIntoArrayLines(ChildCookArray);
+
+		for (const FString& ChildFile : ChildCookArray)
+		{
+			AddFileToCook(FilesInPath, ChildFile);
+		}
+		// if we are a child cooker then just add the files for the child cooker and early out 
+		return;
+	}
+
+	TArray<FString> CookDirectories = InCookDirectories;
 	
 	if (!IsCookingDLC())
 	{
@@ -3269,9 +3477,25 @@ void UCookOnTheFlyServer::CookByTheBookFinished()
 	UPackage::WaitForAsyncFileWrites();
 
 	GetDerivedDataCacheRef().WaitForQuiescence(true);
-
 	
+
 	{
+		SCOPE_TIMER(WaitForChildCookersToFinish);
+		// wait for child cookers to finish
+		bool CookersFinished = false;
+		while (TickChildCookers() == false)
+		{
+			// wait for cookers to finish
+			FPlatformProcess::Sleep(0.01f);
+		}
+
+		CookByTheBookOptions->ChildCookers.Empty();
+		
+	}
+
+	if (IsChildCooker() == false)
+	{
+		SCOPE_TIMER(SavingAssetRegistry);
 		// Save modified asset registry with all streaming chunk info generated during cook
 		const FString& SandboxRegistryFilename = GetSandboxAssetRegistryFilename();
 		// the registry filename will be modified when we call save asset registry
@@ -3315,7 +3539,7 @@ void UCookOnTheFlyServer::CookByTheBookFinished()
 		}
 	}
 
-	if (CookByTheBookOptions->bGenerateDependenciesForMaps)
+	if (CookByTheBookOptions->bGenerateDependenciesForMaps && IsChildCooker() == false)
 	{
 		auto& MapDependencyGraph = CookByTheBookOptions->MapDependencyGraph;
 
@@ -3447,12 +3671,15 @@ void UCookOnTheFlyServer::InitializeSandbox()
 
 		CreateSandboxFile();
 
-		if ( IsCookFlagSet(ECookInitializationFlags::Iterative) )
+		if (!IsChildCooker())
 		{
-			PopulateCookedPackagesFromDisk(TargetPlatforms);
-		}
+			if (IsCookFlagSet(ECookInitializationFlags::Iterative))
+			{
+				PopulateCookedPackagesFromDisk(TargetPlatforms);
+			}
 
-		CleanSandbox(IsCookFlagSet(ECookInitializationFlags::Iterative));
+			CleanSandbox(IsCookFlagSet(ECookInitializationFlags::Iterative));
+		}
 	}
 }
 
@@ -3480,34 +3707,39 @@ void UCookOnTheFlyServer::StartCookByTheBook( const FCookByTheBookStartupOptions
 	check( IsInGameThread() );
 	check( IsCookByTheBookMode() );
 
-	if ( CookByTheBookOptions->DlcName != DLCName )
-	{
-		// we are going to change the state of dlc we need to clean out our package filename cache (the generated filename cache is dependent on this key)
-		CookByTheBookOptions->DlcName = DLCName;
-		check( OutputDirectoryOverride.Len() == 0 );
-
-		TermSandbox();
-	}
-
-	InitializeSandbox();
-
-	// iteratively clean any old files out of the sandbox (check if ini settings are out of date and clean out any stale files)
-	CleanSandbox(true);
-
 	CookByTheBookOptions->bRunning = true;
 	CookByTheBookOptions->bCancel = false;
 	CookByTheBookOptions->CookTime = 0.0f;
 	CookByTheBookOptions->CookStartTime = FPlatformTime::Seconds();
 	CookByTheBookOptions->bGenerateStreamingInstallManifests = CookByTheBookStartupOptions.bGenerateStreamingInstallManifests;
 	CookByTheBookOptions->bGenerateDependenciesForMaps = CookByTheBookStartupOptions.bGenerateDependenciesForMaps;
-	CookByTheBookOptions->CreateReleaseVersion = CreateReleaseVersion;	
+	CookByTheBookOptions->CreateReleaseVersion = CreateReleaseVersion;
 	CookByTheBookOptions->bForceEnableCompressedPackages = !!(CookOptions & ECookByTheBookOptions::ForceEnableCompressed);
 	CookByTheBookOptions->bForceDisableCompressedPackages = !!(CookOptions & ECookByTheBookOptions::ForceDisableCompressed);
+	CookByTheBookOptions->bIsChildCooker = CookByTheBookStartupOptions.ChildCookFileName.Len() > 0 ? true : false;
 
 	bool bCookAll = (CookOptions & ECookByTheBookOptions::CookAll) != ECookByTheBookOptions::None;
 	bool bMapsOnly = (CookOptions & ECookByTheBookOptions::MapsOnly) != ECookByTheBookOptions::None;
 	bool bNoDev = (CookOptions & ECookByTheBookOptions::NoDevContent) != ECookByTheBookOptions::None;
 
+
+	if ( CookByTheBookOptions->DlcName != DLCName )
+	{
+		// we are going to change the state of dlc we need to clean out our package filename cache (the generated filename cache is dependent on this key)
+		CookByTheBookOptions->DlcName = DLCName;
+		check( OutputDirectoryOverride.Len() == 0 );
+
+		TermSandbox();
+	}
+
+	InitializeSandbox();
+
+	if (IsChildCooker() == false)
+	{
+		// iteratively clean any old files out of the sandbox (check if ini settings are out of date and clean out any stale files)
+		CleanSandbox(true);
+	}
+	
 	CookByTheBookOptions->bLeakTest = (CookOptions & ECookByTheBookOptions::LeakTest) != ECookByTheBookOptions::None; // this won't work from the editor this needs to be standalone
 	check( !CookByTheBookOptions->bLeakTest || CurrentCookMode == ECookMode::CookByTheBook );
 
@@ -3520,36 +3752,46 @@ void UCookOnTheFlyServer::StartCookByTheBook( const FCookByTheBookStartupOptions
 		}
 	}
 
+
 	TArray<FName> TargetPlatformNames;
-	for( const auto &Platform : TargetPlatforms )
+	for (const auto &Platform : TargetPlatforms)
 	{
 		FName PlatformName = FName(*Platform->PlatformName());
-		TargetPlatformNames.Add( PlatformName ); // build list of all target platform names
+		TargetPlatformNames.Add(PlatformName); // build list of all target platform names
+	}
 
-		// make sure we have a manifest for all the platforms 
-		// we want a seperate manifest for each platform because they can all be in different states of cooked content
-		FChunkManifestGenerator* ManifestGenerator = CookByTheBookOptions->ManifestGenerators.FindRef( PlatformName );
-		if ( ManifestGenerator == NULL )
-		{
-			TArray<ITargetPlatform*> Platforms;
-			Platforms.Add( Platform );
-			ManifestGenerator = new FChunkManifestGenerator(Platforms);
-			ManifestGenerator->CleanManifestDirectories();
-			ManifestGenerator->Initialize( CookByTheBookOptions->bGenerateStreamingInstallManifests );
 
+	if (!IsChildCooker())
+	{
+		TArray<FName> TargetPlatformNames;
+		for (const auto &Platform : TargetPlatforms)
+		{
+			FName PlatformName = FName(*Platform->PlatformName());
 
-			if (IsCookFlagSet(ECookInitializationFlags::Iterative) && (PackagesKeptFromPreviousCook.Num() > 0) )
+			// make sure we have a manifest for all the platforms 
+			// we want a seperate manifest for each platform because they can all be in different states of cooked content
+			FChunkManifestGenerator* ManifestGenerator = CookByTheBookOptions->ManifestGenerators.FindRef(PlatformName);
+			if (ManifestGenerator == NULL)
 			{
-				const FString& SandboxRegistryFilename = GetSandboxAssetRegistryFilename();
+				TArray<ITargetPlatform*> Platforms;
+				Platforms.Add(Platform);
+				ManifestGenerator = new FChunkManifestGenerator(Platforms);
+				ManifestGenerator->CleanManifestDirectories();
+				ManifestGenerator->Initialize(CookByTheBookOptions->bGenerateStreamingInstallManifests);
 
-				ManifestGenerator->LoadAssetRegistry(GetCookedAssetRegistryFilename(PlatformName.ToString()), &PackagesKeptFromPreviousCook);
-			}
 
+				if (IsCookFlagSet(ECookInitializationFlags::Iterative) && (PackagesKeptFromPreviousCook.Num() > 0))
+				{
+					const FString& SandboxRegistryFilename = GetSandboxAssetRegistryFilename();
+
+					ManifestGenerator->LoadAssetRegistry(GetCookedAssetRegistryFilename(PlatformName.ToString()), &PackagesKeptFromPreviousCook);
+				}
 
-			CookByTheBookOptions->ManifestGenerators.Add(PlatformName, ManifestGenerator);
+
+				CookByTheBookOptions->ManifestGenerators.Add(PlatformName, ManifestGenerator);
+			}
 		}
 	}
-
 	if ( IsCookingDLC() )
 	{
 		// if we are cooking dlc we must be based of a release version cook
@@ -3577,24 +3819,28 @@ void UCookOnTheFlyServer::StartCookByTheBook( const FCookByTheBookStartupOptions
 		}
 	}
 	
-	// allow the game to fill out the asset registry, as well as get a list of objects to always cook
-	TArray<FString> FilesInPathStrings;
-	FGameDelegates::Get().GetCookModificationDelegate().ExecuteIfBound(FilesInPathStrings);
-
 	TArray<FName> FilesInPath;
-	for (const auto& FileString : FilesInPathStrings)
+
+	if (!IsChildCooker())
 	{
-		FilesInPath.Add(FName(*FileString));
+		// allow the game to fill out the asset registry, as well as get a list of objects to always cook
+		TArray<FString> FilesInPathStrings;
+		FGameDelegates::Get().GetCookModificationDelegate().ExecuteIfBound(FilesInPathStrings);
+
+		for (const auto& FileString : FilesInPathStrings)
+		{
+			FilesInPath.Add(FName(*FileString));
+		}
 	}
 
 	// don't resave the global shader map files in dlc
-	if ( !IsCookingDLC() )
+	if ( !IsCookingDLC() && !IsChildCooker())
 	{
 		SaveGlobalShaderMapFiles(TargetPlatforms);
 	}
 	
 	
-	CollectFilesToCook(FilesInPath, CookMaps, CookDirectories, CookCultures, IniMapSections, bCookAll, bMapsOnly, bNoDev );
+	CollectFilesToCook(FilesInPath, CookMaps, CookDirectories, CookCultures, IniMapSections, bCookAll, bMapsOnly, bNoDev, CookByTheBookStartupOptions.ChildCookFileName );
 	if (FilesInPath.Num() == 0)
 	{
 		LogCookerMessage( FString::Printf(TEXT("No files found to cook.")), EMessageSeverity::Warning );
@@ -3626,7 +3872,7 @@ void UCookOnTheFlyServer::StartCookByTheBook( const FCookByTheBookStartupOptions
 	}
 
 
-	if ( !IsCookingDLC() )
+	if ( !IsCookingDLC() && !IsChildCooker() )
 	{
 		// if we are not cooking dlc then basedOfRelease version just needs to make sure that we cook all the packages which are in the previous release (as well as the new ones)
 		if ( !BasedOnReleaseVersion.IsEmpty() )
@@ -3666,6 +3912,205 @@ void UCookOnTheFlyServer::StartCookByTheBook( const FCookByTheBookStartupOptions
 	}
 	CookByTheBookOptions->PreviousCookRequests.Empty();
 
+	if (CookByTheBookStartupOptions.NumProcesses)
+	{
+		StartChildCookers(CookByTheBookStartupOptions.NumProcesses, TargetPlatformNames);
+	}
+}
+
+
+
+
+// sue chefs away!!
+void UCookOnTheFlyServer::StartChildCookers(int32 NumCookersToSpawn, const TArray<FName>& TargetPlatformNames)
+{
+	SCOPE_TIMER(StartingChildCookers);
+	// create a comprehensive list of all the files we need to cook
+	// then get the packages with least dependencies and give them to some sue chefs to handle
+	
+	check(!IsChildCooker());
+
+	TArray<FName> PackageNames;
+	PackageNames.Empty(CookRequests.Num());
+	
+
+	// TArray<FFilePlatformRequest>;
+
+	for (const auto& CookRequest : CookRequests.GetQueue())
+	{
+		FString LongPackageName = FPackageName::FilenameToLongPackageName(CookRequest.ToString());
+
+		PackageNames.Add(FName(*LongPackageName));
+	}
+
+
+	static const FName AssetRegistryName("AssetRegistry");
+	FAssetRegistryModule& AssetRegistryModule = FModuleManager::LoadModuleChecked<FAssetRegistryModule>(AssetRegistryName);
+	IAssetRegistry& AssetRegistry = AssetRegistryModule.Get();
+
+
+
+	TArray<FName> DistributeCandidates;
+
+	int32 PackageCounter = 0;
+	while (PackageCounter < PackageNames.Num())
+	{
+		const auto& PackageName = PackageNames[PackageCounter];
+		++PackageCounter;
+		TArray<FName> UnfilteredDependencies;
+		AssetRegistry.GetDependencies(PackageName, UnfilteredDependencies);
+
+		TArray<FName> Dependencies;
+		for (const auto& PackageName : UnfilteredDependencies)
+		{
+			if (FPackageName::IsScriptPackage(PackageName.ToString()) == false)
+			{
+				Dependencies.Add(PackageName);
+			}
+		}
+
+		if (Dependencies.Num() == 0)
+		{
+			DistributeCandidates.AddUnique(PackageName);
+		}
+		else
+		{
+			for (const auto& Dependency : Dependencies)
+			{
+				PackageNames.AddUnique(Dependency);
+			}
+		}
+	}
+
+	TArray<FName> DistributeStandardFilenames;
+	for (const auto& DistributeCandidate : DistributeCandidates)
+	{
+		FText OutReason;
+		FString LongPackageName = DistributeCandidate.ToString();
+		if (!FPackageName::IsValidLongPackageName(LongPackageName, true, &OutReason))
+		{
+			const FText FailMessage = FText::Format(LOCTEXT("UnableToGeneratePackageName", "Unable to generate long package name for {0}. {1}"),
+				FText::FromString(LongPackageName), OutReason);
+
+			LogCookerMessage(FailMessage.ToString(), EMessageSeverity::Warning);
+			UE_LOG(LogCook, Warning, TEXT("%s"), *(FailMessage.ToString()));
+			continue;
+		}
+		else if (FPackageName::IsScriptPackage(LongPackageName))
+		{
+			continue;
+		}
+		DistributeStandardFilenames.Add(FName(*LongPackageName));
+	}
+
+	UE_LOG(LogCook, Display, TEXT("Distributing %d packages to %d cookers for processing"), DistributeStandardFilenames.Num(), NumCookersToSpawn);
+
+	FString TargetPlatformString;
+	for (const auto& TargetPlatformName : TargetPlatformNames)
+	{
+		if (TargetPlatformString.Len() != 0)
+		{
+			TargetPlatformString += TEXT("+");
+		}
+		TargetPlatformString += TargetPlatformName.ToString();
+	}
+
+	// start the child cookers and give them each some distribution candidates
+	for (int32 CookerCounter = 0; CookerCounter < NumCookersToSpawn; ++CookerCounter )
+	{
+		// don't spawn a cooker unless it has atleast 5 files to do
+		int32 NumFilesForCooker = DistributeStandardFilenames.Num() / (NumCookersToSpawn - CookerCounter);
+		if (NumFilesForCooker < 5)
+		{
+			continue;
+		}
+
+
+		int32 ChildCookerIndex = CookByTheBookOptions->ChildCookers.AddDefaulted(1);
+		auto& ChildCooker = CookByTheBookOptions->ChildCookers[ChildCookerIndex];
+		// starting child cookers
+		//FProcHandle ProcessHandle = FPlatformProcess::CreateProc(*ExecutablePath, *CommandLine, false, true, true, NULL, 0, NULL, WritePipe);
+		// ChildCooker.
+
+		ChildCooker.ResponseFileName = FPaths::CreateTempFilename(*(FPaths::GameSavedDir() / TEXT("CookingTemp")));
+		ChildCooker.BaseResponseFileName = FPaths::GetBaseFilename(ChildCooker.ResponseFileName);
+		// FArchive* ResponseFile = IFileManager::CreateFileWriter(ChildCooker.UniqueTempName);
+		FString ResponseFileText;
+
+
+
+		for (int32 I = 0; I < NumFilesForCooker; ++I)
+		{
+			FName PackageFName = DistributeStandardFilenames[I];
+			FString PackageName = PackageFName.ToString();
+
+			ResponseFileText += FString::Printf(TEXT("%s%s"), *PackageName, LINE_TERMINATOR);
+
+			// these are long package names
+			FName StandardPackageName = GetCachedStandardPackageFileFName(PackageFName);
+			if (StandardPackageName == NAME_None)
+				continue;
+#if 0 // validation code can be disabled
+			PackageName = FPackageName::LongPackageNameToFilename(PackageName,FPaths::GetExtension(PackageName));
+			FPaths::MakeStandardFilename(PackageName);
+			FName PackageFName = FName(*PackageName);
+			check(PackageFName== StandardPackageName);
+#endif
+
+			/*PackageName = FPackageName::LongPackageNameToFilename(PackageName, FPaths::GetExtension(PackageName));
+			FPaths::MakeStandardFilename(PackageName);
+			FName StandardPackageName = FName(*PackageName);*/
+			
+			CookedPackages.Add(FFilePlatformRequest(StandardPackageName, TargetPlatformNames));
+
+
+			
+			/*TArray<FName> Dependencies;
+			AssetRegistry.GetDependencies(PackageFName, Dependencies);
+			TSet<FName> DependenciesSet;
+			for (const auto& Dependency : Dependencies)
+			{
+				DependenciesSet.Add(Dependency);
+			}*/
+
+			TSet<FName> RootPackages;
+			RootPackages.Add(PackageFName);
+			TSet<FName> FoundPackages;
+			GetDependentPackages(RootPackages, FoundPackages);
+
+
+			FName MapName = NAME_None; // to do figure out which map depends on this package for now this package will get put into chunk 0
+			AddDependenciesToManifest(StandardPackageName, MapName, FoundPackages, TargetPlatformNames);
+		}
+		DistributeStandardFilenames.RemoveAt(0, NumFilesForCooker);
+
+		UE_LOG(LogCook, Display, TEXT("Child cooker %d working on %d files"), CookerCounter, NumFilesForCooker);
+
+		FFileHelper::SaveStringToFile(ResponseFileText, *ChildCooker.ResponseFileName);
+
+		FString CommandLine = FString::Printf(TEXT("%s -run=cook -targetplatform=%s -cookchild=%s -abslog=%sLog.txt"), *FPaths::GetProjectFilePath(), *TargetPlatformString, *ChildCooker.ResponseFileName, *ChildCooker.ResponseFileName);
+
+		FString DDCCommandline;
+		if (FParse::Value(FCommandLine::Get(), TEXT("ddc="), DDCCommandline))
+		{
+			CommandLine += TEXT(" -ddc=");
+			CommandLine += DDCCommandline;
+		}
+
+		FString ExecutablePath = FPlatformProcess::ExecutableName(true);
+		// FString ExeFileName = FPaths::EngineDir() / TEXT("Binaries") / PlatformConfig / FString(FPlatformProcess::ExecutableName()) + TEXT(".exe");
+
+		UE_LOG(LogCook, Display, TEXT("Launching cooker using commandline %s %s"), *ExecutablePath, *CommandLine);
+
+		void* ReadPipe = NULL;
+		void* WritePipe = NULL;
+		FPlatformProcess::CreatePipe(ReadPipe, WritePipe);
+
+		//ChildCooker.ProcessHandle = FPlatformProcess::CreateProc(*ExecutablePath, *CommandLine, false, true, true, NULL, 0, NULL, WritePipe);
+		ChildCooker.ProcessHandle = FPlatformProcess::CreateProc(*ExecutablePath, *CommandLine, false, true, true, NULL, 0, NULL, NULL);
+		ChildCooker.ReadPipe = ReadPipe;
+	}
+
 }
 
 
@@ -3675,6 +4120,40 @@ void UCookOnTheFlyServer::StartCookByTheBook( const FCookByTheBookStartupOptions
 /* UCookOnTheFlyServer callbacks
  *****************************************************************************/
 
+void UCookOnTheFlyServer::MaybeMarkPackageAsAlreadyLoaded(UPackage *Package)
+{
+	// can't use this optimization while cooking in editor
+	check(IsCookingInEditor()==false);
+	check(IsCookByTheBookMode());
+
+
+	FName StandardName = GetCachedStandardPackageFileFName(Package);
+
+
+	TArray<FName> CookedPlatforms;
+	if (CookedPackages.GetCookedPlatforms(StandardName, CookedPlatforms))
+	{
+		FString Platforms;
+		for (const auto& CookedPlatform : CookedPlatforms)
+		{
+			Platforms += TEXT(" ");
+			Platforms += CookedPlatform.ToString();
+		}
+		UE_LOG(LogCook, Display, TEXT("Marking %s as reloading for cooker because it's been cooked for platforms%s."), *StandardName.ToString(), *Platforms);
+		Package->PackageFlags |= PKG_ReloadingForCooker;
+	}
+
+
+	/*FString Name = Package->GetName();
+
+
+	if (PackagesToNotReload.Contains(Name))
+	{
+		UE_LOG(LogCookCommandlet, Verbose, TEXT("Marking %s already loaded."), *Name);
+		Package->PackageFlags |= PKG_ReloadingForCooker;
+	}*/
+}
+
 
 void UCookOnTheFlyServer::HandleNetworkFileServerFileRequest( const FString& Filename, const FString &Platformname, TArray<FString>& UnsolicitedFiles )
 {
diff --git a/Engine/Source/Editor/UnrealEd/Private/UnrealEdEngine.cpp b/Engine/Source/Editor/UnrealEd/Private/UnrealEdEngine.cpp
--- a/Engine/Source/Editor/UnrealEd/Private/UnrealEdEngine.cpp
+++ b/Engine/Source/Editor/UnrealEd/Private/UnrealEdEngine.cpp
@@ -126,31 +126,26 @@ void UUnrealEdEngine::Init(IEngineLoop* InEngineLoop)
 		PropertyModule.RegisterCustomClassLayout("ProjectPackagingSettings", FOnGetDetailCustomizationInstance::CreateStatic(&FProjectPackagingSettingsCustomization::MakeInstance));
 	}
 
-	UEditorExperimentalSettings const* ExperimentalSettings =  GetDefault<UEditorExperimentalSettings>();
-	ECookInitializationFlags BaseCookingFlags = ECookInitializationFlags::AutoTick | ECookInitializationFlags::AsyncSave | ECookInitializationFlags::Compressed;
-	BaseCookingFlags |= ExperimentalSettings->bIterativeCookingForLaunchOn ? ECookInitializationFlags::Iterative : ECookInitializationFlags::None;
-
-	bool bEnableCookOnTheSide = false;
-	GConfig->GetBool(TEXT("/Script/UnrealEd.CookerSettings"), TEXT("bEnableCookOnTheSide"), bEnableCookOnTheSide, GEngineIni);
-
-	if ( bEnableCookOnTheSide )
+	if (!IsRunningCommandlet())
 	{
-		CookServer = NewObject<UCookOnTheFlyServer>();
-		CookServer->Initialize( ECookMode::CookOnTheFlyFromTheEditor, BaseCookingFlags );
-		CookServer->StartNetworkFileServer( false );
+		UEditorExperimentalSettings const* ExperimentalSettings = GetDefault<UEditorExperimentalSettings>();
+		ECookInitializationFlags BaseCookingFlags = ECookInitializationFlags::AutoTick | ECookInitializationFlags::AsyncSave | ECookInitializationFlags::Compressed;
+		BaseCookingFlags |= ExperimentalSettings->bIterativeCookingForLaunchOn ? ECookInitializationFlags::Iterative : ECookInitializationFlags::None;
 
-		FCoreUObjectDelegates::OnObjectPropertyChanged.AddUObject(CookServer, &UCookOnTheFlyServer::OnObjectPropertyChanged);
-		FCoreUObjectDelegates::OnObjectModified.AddUObject(CookServer, &UCookOnTheFlyServer::OnObjectModified);
-		FCoreUObjectDelegates::OnObjectSaved.AddUObject( CookServer, &UCookOnTheFlyServer::OnObjectSaved );
-	}
-	else if ( !ExperimentalSettings->bDisableCookInEditor)
-	{
-		CookServer = NewObject<UCookOnTheFlyServer>();
-		CookServer->Initialize( ECookMode::CookByTheBookFromTheEditor, BaseCookingFlags );
+		bool bEnableCookOnTheSide = false;
+		GConfig->GetBool(TEXT("/Script/UnrealEd.CookerSettings"), TEXT("bEnableCookOnTheSide"), bEnableCookOnTheSide, GEngineIni);
 
-		FCoreUObjectDelegates::OnObjectPropertyChanged.AddUObject(CookServer, &UCookOnTheFlyServer::OnObjectPropertyChanged);
-		FCoreUObjectDelegates::OnObjectModified.AddUObject(CookServer, &UCookOnTheFlyServer::OnObjectModified);
-		FCoreUObjectDelegates::OnObjectSaved.AddUObject( CookServer, &UCookOnTheFlyServer::OnObjectSaved );
+		if (bEnableCookOnTheSide)
+		{
+			CookServer = NewObject<UCookOnTheFlyServer>();
+			CookServer->Initialize(ECookMode::CookOnTheFlyFromTheEditor, BaseCookingFlags);
+			CookServer->StartNetworkFileServer(false);
+		}
+		else if (!ExperimentalSettings->bDisableCookInEditor)
+		{
+			CookServer = NewObject<UCookOnTheFlyServer>();
+			CookServer->Initialize(ECookMode::CookByTheBookFromTheEditor, BaseCookingFlags);
+		}
 	}
 }
 
