Fix shadowed variables issues.

diff --git a/Engine/Source/Editor/UnrealEd/Private/CookOnTheFlyServer.cpp b/Engine/Source/Editor/UnrealEd/Private/CookOnTheFlyServer.cpp
--- a/Engine/Source/Editor/UnrealEd/Private/CookOnTheFlyServer.cpp
+++ b/Engine/Source/Editor/UnrealEd/Private/CookOnTheFlyServer.cpp
@@ -3763,7 +3763,6 @@ void UCookOnTheFlyServer::StartCookByTheBook( const FCookByTheBookStartupOptions
 
 	if (!IsChildCooker())
 	{
-		TArray<FName> TargetPlatformNames;
 		for (const auto &Platform : TargetPlatforms)
 		{
 			FName PlatformName = FName(*Platform->PlatformName());
@@ -3961,11 +3960,11 @@ void UCookOnTheFlyServer::StartChildCookers(int32 NumCookersToSpawn, const TArra
 		AssetRegistry.GetDependencies(PackageName, UnfilteredDependencies);
 
 		TArray<FName> Dependencies;
-		for (const auto& PackageName : UnfilteredDependencies)
+		for (const auto& DependencyName : UnfilteredDependencies)
 		{
-			if (FPackageName::IsScriptPackage(PackageName.ToString()) == false)
+			if (FPackageName::IsScriptPackage(DependencyName.ToString()) == false)
 			{
-				Dependencies.Add(PackageName);
+				Dependencies.Add(DependencyName);
 			}
 		}
 
