Editor: Added an explicit factory for UCurveFloat to allow the "New Asset" option to be displayed in the context menu of UCurveFloat* properties

diff --git a/Engine/Source/Editor/UnrealEd/Classes/Factories/CurveFactory.h b/Engine/Source/Editor/UnrealEd/Classes/Factories/CurveFactory.h
--- a/Engine/Source/Editor/UnrealEd/Classes/Factories/CurveFactory.h
+++ b/Engine/Source/Editor/UnrealEd/Classes/Factories/CurveFactory.h
@@ -1,24 +1,14 @@
 // Copyright 1998-2015 Epic Games, Inc. All Rights Reserved.
 
-/**
- *
- */
-
 #pragma once
+
 #include "CurveFactory.generated.h"
 
-/** Enum indicating the type of curve to make */
-UENUM()
-enum ECurveType
-{
-	/** Float keyframed */
-	ECT_FloatKeyframed,
-	/** Vector keyframed */
-	ECT_VectorKeyframed,
-	ECT_MAX,
-};
+/**
+ * Factory that creates curve assets, prompting to pick the kind of curve at creation time
+ */
 
-UCLASS(hidecategories=Object)
+UCLASS()
 class UCurveFactory : public UFactory
 {
 	GENERATED_UCLASS_BODY()
@@ -27,11 +17,22 @@ class UCurveFactory : public UFactory
 	UPROPERTY(EditAnywhere, Category=CurveFactory)
 	TSubclassOf<UCurveBase> CurveClass;
 
-	// Begin UFactory Interface
+	// UFactory Interface
 	virtual bool ConfigureProperties() override;
 	virtual UObject* FactoryCreateNew(UClass* Class,UObject* InParent,FName Name,EObjectFlags Flags,UObject* Context,FFeedbackContext* Warn) override;
-	// Begin UFactory Interface
+	// End of UFactory interface
 };
 
+/**
+ * Factory that creates float curve assets
+ */
 
+UCLASS()
+class UCurveFloatFactory : public UCurveFactory
+{
+	GENERATED_UCLASS_BODY()
 
+	// UFactory interface
+	virtual bool ConfigureProperties() override;
+	// End of UFactory interface
+};
diff --git a/Engine/Source/Editor/UnrealEd/Private/Factories/EditorFactories.cpp b/Engine/Source/Editor/UnrealEd/Private/Factories/EditorFactories.cpp
--- a/Engine/Source/Editor/UnrealEd/Private/Factories/EditorFactories.cpp
+++ b/Engine/Source/Editor/UnrealEd/Private/Factories/EditorFactories.cpp
@@ -6480,6 +6480,24 @@ UObject* UCurveFactory::FactoryCreateNew(UClass* Class, UObject* InParent, FName
 	return NewCurve;
 }
 
+/*------------------------------------------------------------------------------
+	UCurveFloatFactory implementation.
+------------------------------------------------------------------------------*/
+
+UCurveFloatFactory::UCurveFloatFactory(const FObjectInitializer& ObjectInitializer)
+	: Super(ObjectInitializer)
+{
+	bCreateNew = true;
+	bEditAfterNew = true;
+	SupportedClass = UCurveFloat::StaticClass();
+	CurveClass = UCurveFloat::StaticClass();
+}
+
+bool UCurveFloatFactory::ConfigureProperties()
+{
+	return true;
+}
+
 /*------------------------------------------------------------------------------
 	UCurveImportFactory implementation.
 ------------------------------------------------------------------------------*/
