Fix cook on the fly server garbage collecting too many objects causing the editor to crash because objects it is using are removed underneath it.

#coderevew daniel.lamb

diff --git a/Engine/Source/Editor/UnrealEd/Private/CookOnTheFlyServer.cpp b/Engine/Source/Editor/UnrealEd/Private/CookOnTheFlyServer.cpp
--- a/Engine/Source/Editor/UnrealEd/Private/CookOnTheFlyServer.cpp
+++ b/Engine/Source/Editor/UnrealEd/Private/CookOnTheFlyServer.cpp
@@ -1955,7 +1955,7 @@ uint32 UCookOnTheFlyServer::TickCookOnTheSide( const float TimeSlice, uint32 &Co
 			check( CurrentCookMode == ECookMode::CookByTheBook);
 			UE_LOG(LogCook, Display, TEXT("Full GC..."));
 
-			CollectGarbage( RF_Native );
+			CollectGarbage( GARBAGE_COLLECTION_KEEPFLAGS );
 			for (FObjectIterator It; It; ++It)
 			{
 				if (!CookByTheBookOptions->LastGCItems.Contains(FWeakObjectPtr(*It)))
@@ -2990,7 +2990,7 @@ void UCookOnTheFlyServer::CleanSandbox( const bool bIterative )
 			}
 
 			// Collect garbage to ensure we don't have any packages hanging around from dependent time stamp determination
-			CollectGarbage(RF_Native);
+			CollectGarbage(GARBAGE_COLLECTION_KEEPFLAGS);
 		}
 	}
 #if OUTPUT_TIMING
