Disabled validation code, causes crash on kite demo.

diff --git a/Engine/Source/Editor/UnrealEd/Private/CookOnTheFlyServer.cpp b/Engine/Source/Editor/UnrealEd/Private/CookOnTheFlyServer.cpp
--- a/Engine/Source/Editor/UnrealEd/Private/CookOnTheFlyServer.cpp
+++ b/Engine/Source/Editor/UnrealEd/Private/CookOnTheFlyServer.cpp
@@ -1030,7 +1030,7 @@ void UCookOnTheFlyServer::GenerateManifestInfo(UPackage* Package, const TArray<F
 	// load sublevels
 	UWorld* World = UWorld::FindWorldInPackage(Package);
 
-#define VERIFY_LEVELS_IN_DEPENDENCIES_LIST 1
+#define VERIFY_LEVELS_IN_DEPENDENCIES_LIST 0
 #if VERIFY_LEVELS_IN_DEPENDENCIES_LIST
 	TArray<FString> LevelPackageNames;
 #endif
