Fixed missing include

diff --git a/Engine/Source/Editor/UnrealEd/Private/LevelEditorViewport.cpp b/Engine/Source/Editor/UnrealEd/Private/LevelEditorViewport.cpp
--- a/Engine/Source/Editor/UnrealEd/Private/LevelEditorViewport.cpp
+++ b/Engine/Source/Editor/UnrealEd/Private/LevelEditorViewport.cpp
@@ -13,7 +13,7 @@
 #include "BSPOps.h"
 #include "LevelUtils.h"
 #include "Layers/ILayers.h"
-#include "../Private/StaticLightingSystem/StaticLightingPrivate.h"
+#include "StaticLightingSystem/StaticLightingPrivate.h"
 #include "EditorLevelUtils.h"
 #include "Engine.h"
 #include "Editor/LevelEditor/Public/LevelEditor.h"
diff --git a/Engine/Source/Editor/UnrealEd/Private/LevelViewportClickHandlers.cpp b/Engine/Source/Editor/UnrealEd/Private/LevelViewportClickHandlers.cpp
--- a/Engine/Source/Editor/UnrealEd/Private/LevelViewportClickHandlers.cpp
+++ b/Engine/Source/Editor/UnrealEd/Private/LevelViewportClickHandlers.cpp
@@ -18,6 +18,7 @@
 #include "Engine/StaticMeshActor.h"
 #include "Engine/TargetPoint.h"
 #include "EngineUtils.h"
+#include "StaticLightingSystem/StaticLightingPrivate.h"
 
 #define LOCTEXT_NAMESPACE "ClickHandlers"
 
