Viewport: Add begin/end camera broadcast messages when changing the camera view.

diff --git a/Engine/Source/Editor/UnrealEd/Classes/Editor/EditorEngine.h b/Engine/Source/Editor/UnrealEd/Classes/Editor/EditorEngine.h
--- a/Engine/Source/Editor/UnrealEd/Classes/Editor/EditorEngine.h
+++ b/Engine/Source/Editor/UnrealEd/Classes/Editor/EditorEngine.h
@@ -645,6 +645,14 @@ public:
 	DECLARE_EVENT_OneParam( UEditorEngine, FOnEndTransformObject, UObject& );
 	FOnEndTransformObject& OnEndObjectMovement() { return OnEndObjectTransformEvent; }
 
+	/** Editor-only event triggered before the camera viewed through the viewport is moved by an editor system */
+	DECLARE_EVENT_OneParam( UEditorEngine, FOnBeginTransformCamera, UObject& );
+	FOnBeginTransformCamera& OnBeginCameraMovement() { return OnBeginCameraTransformEvent; }
+
+	/** Editor-only event triggered after the camera viewed through the viewport has been moved by an editor system */
+	DECLARE_EVENT_OneParam( UEditorEngine, FOnEndTransformCamera, UObject& );
+	FOnEndTransformCamera& OnEndCameraMovement() { return OnEndCameraTransformEvent; }
+
 	/** Delegate broadcast by the engine every tick when PIE/SIE is active, to check to see whether we need to
 		be able to capture state for simulating actor (for Sequencer recording features).  The single bool parameter
 		should be set to true if recording features are needed. */
@@ -665,6 +673,20 @@ public:
 	*/
 	void BroadcastEndObjectMovement(UObject& Object) const { OnEndObjectTransformEvent.Broadcast(Object); }
 
+	/**
+	* Called before the camera viewed through the viewport is moved by the editor
+	*
+	* @param Object	The camera that will be moved
+	*/
+	void BroadcastBeginCameraMovement(UObject& Object) const { OnBeginCameraTransformEvent.Broadcast(Object); }
+
+	/**
+	* Called when the camera viewed through the viewport has been moved by the editor
+	*
+	* @param Object	The camera that moved
+	*/
+	void BroadcastEndCameraMovement(UObject& Object) const { OnEndCameraTransformEvent.Broadcast(Object); }
+
 	/**	Broadcasts that an object has been reimported. THIS SHOULD NOT BE PUBLIC */
 	void BroadcastObjectReimported(UObject* InObject);
 
@@ -2549,6 +2571,12 @@ private:
 	/** Delegate broadcast when an actor or component has been moved, rotated, or scaled */
 	FOnEndTransformObject OnEndObjectTransformEvent;
 
+	/** Delegate broadcast when the camera viewed through the viewport is about to be moved */
+	FOnBeginTransformCamera OnBeginCameraTransformEvent;
+
+	/** Delegate broadcast when the camera viewed through the viewport has been moved */
+	FOnEndTransformCamera OnEndCameraTransformEvent;
+
 	/** Delegate broadcast by the engine every tick when PIE/SIE is active, to check to see whether we need to
 		be able to capture state for simulating actor (for Sequencer recording features) */
 	FGetActorRecordingState GetActorRecordingStateEvent;
diff --git a/Engine/Source/Editor/UnrealEd/Private/EditorViewportClient.cpp b/Engine/Source/Editor/UnrealEd/Private/EditorViewportClient.cpp
--- a/Engine/Source/Editor/UnrealEd/Private/EditorViewportClient.cpp
+++ b/Engine/Source/Editor/UnrealEd/Private/EditorViewportClient.cpp
@@ -272,6 +272,8 @@ FEditorViewportClient::FEditorViewportClient(FEditorModeTools* InModeTools, FPre
 	, bUsesDrawHelper(true)
 	, bIsSimulateInEditorViewport(false)
 	, bCameraLock(false)
+	, bIsCameraMoving(false)
+	, bIsCameraMovingOnTick(false)
 	, EditorViewportWidget(InEditorViewportWidget)
 	, PreviewScene(InPreviewScene)
 	, MovingPreviewLightSavedScreenPos(ForceInitToZero)
@@ -936,12 +938,16 @@ void FEditorViewportClient::Tick(float DeltaTime)
 
 	if( !bIsAnimating )
 	{
+		bIsCameraMovingOnTick = bIsCameraMoving;
+
 		// Update any real-time camera movement
 		UpdateCameraMovement( DeltaTime );
 
 		UpdateMouseDelta();
 		
 		UpdateGestureDelta();
+
+		EndCameraMovement();
 	}
 
 	if ( bIsTracking )
@@ -1439,6 +1445,26 @@ void FEditorViewportClient::UpdateCameraMovement( float DeltaTime )
 		EditorMovementDeltaUpperBound = .15f;
 #endif
 
+		// Check whether the camera is being moved by the mouse or keyboard
+		bool bHasMovement = bIsTracking;
+
+		if ((*CameraUserImpulseData).RotateYawVelocityModifier != 0.0f ||
+			(*CameraUserImpulseData).RotatePitchVelocityModifier != 0.0f ||
+			(*CameraUserImpulseData).RotateRollVelocityModifier != 0.0f ||
+			(*CameraUserImpulseData).MoveForwardBackwardImpulse != 0.0f ||
+			(*CameraUserImpulseData).MoveRightLeftImpulse != 0.0f ||
+			(*CameraUserImpulseData).MoveUpDownImpulse != 0.0f ||
+			(*CameraUserImpulseData).ZoomOutInImpulse != 0.0f ||
+			(*CameraUserImpulseData).RotateYawImpulse != 0.0f ||
+			(*CameraUserImpulseData).RotatePitchImpulse != 0.0f ||
+			(*CameraUserImpulseData).RotateRollImpulse != 0.0f
+			)
+		{
+			bHasMovement = true;
+		}
+
+		BeginCameraMovement(bHasMovement);
+
 		CameraController->UpdateSimulation(
 			*CameraUserImpulseData,
 			FMath::Min(DeltaTime, EditorMovementDeltaUpperBound),
@@ -1783,6 +1809,10 @@ void FEditorViewportClient::UpdateMouseDelta()
 			{
 				if ( ShouldOrbitCamera() )
 				{
+					bool bHasMovement = !DragDelta.IsNearlyZero();
+
+					BeginCameraMovement(bHasMovement);
+
 					FVector TempDrag;
 					FRotator TempRot;
 					InputAxisForOrbit( Viewport, DragDelta, TempDrag, TempRot );
@@ -1795,6 +1825,10 @@ void FEditorViewportClient::UpdateMouseDelta()
 
 					if ( ShouldPanOrDollyCamera() )
 					{
+						bool bHasMovement = !Drag.IsNearlyZero() || !Rot.IsNearlyZero();
+
+						BeginCameraMovement(bHasMovement);
+
 						if( !IsOrtho())
 						{
 							const float CameraSpeed = GetCameraSpeed();
diff --git a/Engine/Source/Editor/UnrealEd/Private/LevelEditorViewport.cpp b/Engine/Source/Editor/UnrealEd/Private/LevelEditorViewport.cpp
--- a/Engine/Source/Editor/UnrealEd/Private/LevelEditorViewport.cpp
+++ b/Engine/Source/Editor/UnrealEd/Private/LevelEditorViewport.cpp
@@ -1711,6 +1711,38 @@ bool FLevelEditorViewportClient::ShouldLockPitch() const
 	return FEditorViewportClient::ShouldLockPitch() || !ModeTools->GetActiveMode(FBuiltinEditorModes::EM_InterpEdit) ;
 }
 
+void FLevelEditorViewportClient::BeginCameraMovement(bool bHasMovement)
+{
+	// If there's new movement broadcast it
+	if (bHasMovement)
+	{
+		if (!bIsCameraMoving)
+		{
+			if (!bIsCameraMovingOnTick)
+			{
+				GEditor->BroadcastBeginCameraMovement(*GetActiveActorLock());
+			}
+			bIsCameraMoving = true;
+		}
+	}
+	else
+	{
+		bIsCameraMoving = false;
+	}
+}
+
+void FLevelEditorViewportClient::EndCameraMovement()
+{
+	// If there was movement and it has now stopped, broadcast it
+	if (bIsCameraMovingOnTick)
+	{
+		if (!bIsCameraMoving)
+		{
+			GEditor->BroadcastEndCameraMovement(*GetActiveActorLock());
+		}
+	}
+}
+
 void FLevelEditorViewportClient::PerspectiveCameraMoved()
 {
 	// Update the locked actor (if any) from the camera
diff --git a/Engine/Source/Editor/UnrealEd/Public/EditorViewportClient.h b/Engine/Source/Editor/UnrealEd/Public/EditorViewportClient.h
--- a/Engine/Source/Editor/UnrealEd/Public/EditorViewportClient.h
+++ b/Engine/Source/Editor/UnrealEd/Public/EditorViewportClient.h
@@ -1203,6 +1203,10 @@ private:
 	/** Delegate handler for when all stats are disabled in a viewport */
 	void HandleViewportStatDisableAll(const bool bInAnyViewport);
 
+	/** Handle the camera about to be moved or stopped **/
+	virtual void BeginCameraMovement(bool bIsCameraMoving) {}
+	virtual void EndCameraMovement() {}
+
 public:
 	static const uint32 MaxCameraSpeeds;
 
@@ -1380,6 +1384,12 @@ protected:
 	/** Camera Lock or not **/
 	bool bCameraLock;
 
+	/** Is the camera moving? **/
+	bool bIsCameraMoving;
+
+	/** Is the camera moving at the beginning of the tick? **/
+	bool bIsCameraMovingOnTick;
+
 	/** Draw helper for rendering common editor functionality like the grid */
 	FEditorCommonDrawHelper DrawHelper;
 
diff --git a/Engine/Source/Editor/UnrealEd/Public/LevelEditorViewport.h b/Engine/Source/Editor/UnrealEd/Public/LevelEditorViewport.h
--- a/Engine/Source/Editor/UnrealEd/Public/LevelEditorViewport.h
+++ b/Engine/Source/Editor/UnrealEd/Public/LevelEditorViewport.h
@@ -156,6 +156,8 @@ public:
 	virtual void ReceivedFocus(FViewport* Viewport) override;
 	virtual void ProcessClick(FSceneView& View, HHitProxy* HitProxy, FKey Key, EInputEvent Event, uint32 HitX, uint32 HitY) override;
 	virtual UWorld* GetWorld() const override;
+	virtual void BeginCameraMovement(bool bIsCameraMoving) override;
+	virtual void EndCameraMovement() override;
 
 	virtual bool OverrideHighResScreenshotCaptureRegion(FIntRect& OutCaptureRegion) override;
 
