SCurveEditor - Make adding keys more predictable.
+ Keys are only added under the mouse when there is only one curve to edit.
+ When multiple curves are available clicking on a curve will add a key inline on that curve, and clicking off of the curve will add keys too all curves.
+ Changed the text of the menu item to make it clear what's going to happen.

Jira: UE-15095

diff --git a/Engine/Source/Editor/UnrealEd/Private/SCurveEditor.cpp b/Engine/Source/Editor/UnrealEd/Private/SCurveEditor.cpp
--- a/Engine/Source/Editor/UnrealEd/Private/SCurveEditor.cpp
+++ b/Engine/Source/Editor/UnrealEd/Private/SCurveEditor.cpp
@@ -475,27 +475,9 @@ void SCurveEditor::PushWarningMenu( FVector2D Position, const FText& Message )
 		FPopupTransitionEffect( FPopupTransitionEffect::ContextMenu));
 }
 
-
 void SCurveEditor::PushKeyMenu(const FGeometry& InMyGeometry, const FPointerEvent& InMouseEvent)
 {
-	FVector2D Position = InMouseEvent.GetScreenSpacePosition();
-
 	FMenuBuilder MenuBuilder(true, Commands.ToSharedRef());
-
-	if (SelectedKeys.Num() == 0)
-	{
-		MenuBuilder.BeginSection("EditCurveEditorActions", LOCTEXT("Actions", "Actions"));
-		{
-			FUIAction Action = FUIAction(FExecuteAction::CreateSP(this, &SCurveEditor::AddNewKey, InMyGeometry, Position));
-			MenuBuilder.AddMenuEntry(
-				LOCTEXT("AddKey", "Add Key"),
-				LOCTEXT("AddKey_ToolTip", "Create a new point in the curve.  You can also use (Shift + Left Mouse Button)"),
-				FSlateIcon(),
-				Action);
-		}
-		MenuBuilder.EndSection();
-	}
-
 	MenuBuilder.BeginSection("CurveEditorInterpolation", LOCTEXT("KeyInterpolationMode", "Key Interpolation"));
 	{
 		MenuBuilder.AddMenuEntry(FRichCurveEditorCommands::Get().InterpolationCubicAuto);
@@ -507,7 +489,7 @@ void SCurveEditor::PushKeyMenu(const FGeometry& InMyGeometry, const FPointerEven
 	MenuBuilder.EndSection(); //CurveEditorInterpolation
 
 	FWidgetPath WidgetPath = InMouseEvent.GetEventPath() != nullptr ? *InMouseEvent.GetEventPath() : FWidgetPath();
-
+	FVector2D Position = InMouseEvent.GetScreenSpacePosition();
 	FSlateApplication::Get().PushMenu(
 		SharedThis( this ),
 		WidgetPath,
@@ -1140,12 +1122,12 @@ FReply SCurveEditor::OnMouseButtonDown( const FGeometry& InMyGeometry, const FPo
 	return FReply::Unhandled();
 }
 
-void SCurveEditor::AddNewKey(FGeometry InMyGeometry, FVector2D ScreenPosition)
+void SCurveEditor::AddNewKey(FGeometry InMyGeometry, FVector2D ScreenPosition, TSharedPtr<TArray<TSharedPtr<FCurveViewModel>>> CurvesToAddKeysTo, bool bAddKeysInline)
 {
 	const FScopedTransaction Transaction(LOCTEXT("CurveEditor_AddKey", "Add Key(s)"));
 	CurveOwner->ModifyOwner();
 	TArray<FRichCurveEditInfo> ChangedCurveEditInfos;
-	for (auto CurveViewModel : CurveViewModels)
+	for (TSharedPtr<FCurveViewModel> CurveViewModel : *CurvesToAddKeysTo)
 	{
 		if (!CurveViewModel->bIsLocked)
 		{
@@ -1156,9 +1138,17 @@ void SCurveEditor::AddNewKey(FGeometry InMyGeometry, FVector2D ScreenPosition)
 
 				FVector2D LocalClickPos = InMyGeometry.AbsoluteToLocal(ScreenPosition);
 
-				FVector2D NewKeyLocation = SnapLocation(FVector2D(
-					ScaleInfo.LocalXToInput(LocalClickPos.X),
-					ScaleInfo.LocalYToOutput(LocalClickPos.Y)));
+				float Input = ScaleInfo.LocalXToInput(LocalClickPos.X);
+				float Output;
+				if (bAddKeysInline)
+				{
+					Output = SelectedCurve->Eval(Input);
+				}
+				else
+				{
+					Output = ScaleInfo.LocalYToOutput(LocalClickPos.Y);
+				}
+				FVector2D NewKeyLocation = SnapLocation(FVector2D(Input, Output));
 				FKeyHandle NewKeyHandle = SelectedCurve->AddKey(NewKeyLocation.X, NewKeyLocation.Y);
 
 				EmptySelection();
@@ -1524,7 +1514,28 @@ void SCurveEditor::ProcessClick(const FGeometry& InMyGeometry, const FPointerEve
 			// If the user didn't click a key, add a new one if shift is held down, or try to select a curve.
 			if (bShiftDown)
 			{
-				AddNewKey(InMyGeometry, InMouseEvent.GetScreenSpacePosition());
+				TSharedPtr<TArray<TSharedPtr<FCurveViewModel>>> CurvesToAddKeysTo = MakeShareable(new TArray<TSharedPtr<FCurveViewModel>>());
+				TSharedPtr<FCurveViewModel> HoveredCurve = HitTestCurves(InMyGeometry, InMouseEvent);
+				bool bAddKeysInline;
+				if (HoveredCurve.IsValid())
+				{
+					CurvesToAddKeysTo->Add(HoveredCurve);
+					bAddKeysInline = false;
+				}
+				else
+				{
+					if (CurveViewModels.Num() == 1)
+					{
+						CurvesToAddKeysTo->Add(CurveViewModels[0]);
+						bAddKeysInline = true;
+					}
+					else
+					{
+						CurvesToAddKeysTo->Append(CurveViewModels);
+						bAddKeysInline = false;
+					}
+				}
+				AddNewKey(InMyGeometry, InMouseEvent.GetScreenSpacePosition(), CurvesToAddKeysTo, bAddKeysInline);
 			}
 			else
 			{
@@ -2033,14 +2044,47 @@ void SCurveEditor::CreateContextMenu(const FGeometry& InMyGeometry, const FPoint
 
 	MenuBuilder.BeginSection("EditCurveEditorActions", LOCTEXT("Actions", "Actions"));
 	{
-		FUIAction Action = FUIAction(FExecuteAction::CreateSP(this, &SCurveEditor::AddNewKey, InMyGeometry, ScreenPosition));
-		MenuBuilder.AddMenuEntry
-		(
-			LOCTEXT("AddKey", "Add Key"),
-			LOCTEXT("AddKey_ToolTip", "Create a new point in the curve.  You can also use (Shift + Left Mouse Button)"),
+		FText MenuItemLabel;
+		FText MenuItemToolTip;
+		TSharedPtr<TArray<TSharedPtr<FCurveViewModel>>> CurvesToAddKeysTo = MakeShareable(new TArray<TSharedPtr<FCurveViewModel>>());
+		bool bAddKeysInline;
+
+		FText AddKeyToCurveLabelFormat = LOCTEXT("AddKeyToCurveLabelFormat", "Add key to {0}");
+		FText AddKeyToCurveToolTipFormat = LOCTEXT("AddKeyToCurveToolTipFormat", "Add a new key at the hovered time to the {0} curve.  Keys can also be added with Shift + Click.");
+
+		TSharedPtr<FCurveViewModel> HoveredCurve = HitTestCurves(InMyGeometry, InMouseEvent);
+		if (HoveredCurve.IsValid())
+		{
+			MenuItemLabel = FText::Format(AddKeyToCurveLabelFormat, FText::FromName(HoveredCurve->CurveInfo.CurveName));
+			MenuItemToolTip = FText::Format(AddKeyToCurveToolTipFormat, FText::FromName(HoveredCurve->CurveInfo.CurveName));
+			CurvesToAddKeysTo->Add(HoveredCurve);
+			bAddKeysInline = false;
+		}
+		else
+		{
+			if (CurveViewModels.Num() == 1)
+			{
+				MenuItemLabel = FText::Format(AddKeyToCurveLabelFormat, FText::FromName(CurveViewModels[0]->CurveInfo.CurveName));
+				MenuItemToolTip = FText::Format(AddKeyToCurveToolTipFormat, FText::FromName(CurveViewModels[0]->CurveInfo.CurveName));
+				CurvesToAddKeysTo->Add(CurveViewModels[0]);
+				bAddKeysInline = true;
+			}
+			else
+			{
+				MenuItemLabel = LOCTEXT("AddKeyToAllCurves", "Add key to all curves");
+				MenuItemToolTip = LOCTEXT("AddKeyToAllCurveToolTip", "Adds a key at the hovered time to all curves.  Keys can also be added with Shift + Click.");
+				CurvesToAddKeysTo->Append(CurveViewModels);
+				bAddKeysInline = false;
+			}
+		}
+
+		FVector2D Position = InMouseEvent.GetScreenSpacePosition();
+		FUIAction Action = FUIAction(FExecuteAction::CreateSP(this, &SCurveEditor::AddNewKey, InMyGeometry, Position, CurvesToAddKeysTo, bAddKeysInline));
+		MenuBuilder.AddMenuEntry(
+			MenuItemLabel,
+			MenuItemToolTip,
 			FSlateIcon(),
-			Action
-		);
+			Action);
 	}
 	MenuBuilder.EndSection();
 
diff --git a/Engine/Source/Editor/UnrealEd/Public/SCurveEditor.h b/Engine/Source/Editor/UnrealEd/Public/SCurveEditor.h
--- a/Engine/Source/Editor/UnrealEd/Public/SCurveEditor.h
+++ b/Engine/Source/Editor/UnrealEd/Public/SCurveEditor.h
@@ -268,7 +268,7 @@ private:
 
 private:
 	/** Adds a new key to the curve. */
-	void AddNewKey(FGeometry InMyGeometry, FVector2D ScreenPosition);
+	void AddNewKey(FGeometry InMyGeometry, FVector2D ScreenPosition, TSharedPtr<TArray<TSharedPtr<FCurveViewModel>>> CurvesToAddKeysTo, bool bAddKeysInline);
 
 	/** Test if the curve is exists, and if it being displayed on this widget */
 	bool		IsValidCurve(FRichCurve* Curve) const;
