Viewport: Fix compile error.

diff --git a/Engine/Source/Editor/UnrealEd/Public/EditorViewportClient.h b/Engine/Source/Editor/UnrealEd/Public/EditorViewportClient.h
--- a/Engine/Source/Editor/UnrealEd/Public/EditorViewportClient.h
+++ b/Engine/Source/Editor/UnrealEd/Public/EditorViewportClient.h
@@ -1204,7 +1204,7 @@ private:
 	void HandleViewportStatDisableAll(const bool bInAnyViewport);
 
 	/** Handle the camera about to be moved or stopped **/
-	virtual void BeginCameraMovement(bool bIsCameraMoving) {}
+	virtual void BeginCameraMovement(bool bHasMovement) {}
 	virtual void EndCameraMovement() {}
 
 public:
diff --git a/Engine/Source/Editor/UnrealEd/Public/LevelEditorViewport.h b/Engine/Source/Editor/UnrealEd/Public/LevelEditorViewport.h
--- a/Engine/Source/Editor/UnrealEd/Public/LevelEditorViewport.h
+++ b/Engine/Source/Editor/UnrealEd/Public/LevelEditorViewport.h
@@ -156,7 +156,7 @@ public:
 	virtual void ReceivedFocus(FViewport* Viewport) override;
 	virtual void ProcessClick(FSceneView& View, HHitProxy* HitProxy, FKey Key, EInputEvent Event, uint32 HitX, uint32 HitY) override;
 	virtual UWorld* GetWorld() const override;
-	virtual void BeginCameraMovement(bool bIsCameraMoving) override;
+	virtual void BeginCameraMovement(bool bHasMovement) override;
 	virtual void EndCameraMovement() override;
 
 	virtual bool OverrideHighResScreenshotCaptureRegion(FIntRect& OutCaptureRegion) override;
