Fixed geometry surface alignment tools.

diff --git a/Engine/Source/Editor/UnrealEd/Private/TexAlignTools.cpp b/Engine/Source/Editor/UnrealEd/Private/TexAlignTools.cpp
--- a/Engine/Source/Editor/UnrealEd/Private/TexAlignTools.cpp
+++ b/Engine/Source/Editor/UnrealEd/Private/TexAlignTools.cpp
@@ -79,8 +79,8 @@ void UTexAligner::Align( UWorld* InWorld, ETexAlign InTexAlignType, UModel* InMo
 	for( int32 i = 0 ; i < InModel->Surfs.Num() ; i++ )
 	{
 		FBspSurf* Surf = &InModel->Surfs[i];
-		GEditor->polyFindMaster( InModel, i, EdPoly );
-		FVector Normal = InModel->Vectors[ Surf->vNormal ];
+//		GEditor->polyFindMaster( InModel, i, EdPoly );
+//		FVector Normal = InModel->Vectors[ Surf->vNormal ];
 
 		if( Surf->PolyFlags & PF_Selected )
 		{
@@ -101,8 +101,8 @@ void UTexAligner::Align( UWorld* InWorld, ETexAlign InTexAlignType, UModel* InMo
 	for( int32 i = 0 ; i < InitialSurfList.Num() ; i++ )
 	{
 		FBspSurfIdx* Surf = &InitialSurfList[i];
-		Normal = InModel->Vectors[ Surf->Surf->vNormal ];
-		GEditor->polyFindMaster( InModel, Surf->Idx, EdPoly );
+//		Normal = InModel->Vectors[ Surf->Surf->vNormal ];
+//		GEditor->polyFindMaster( InModel, Surf->Idx, EdPoly );
 
 		bool bOK = 1;
 		/*
@@ -211,7 +211,7 @@ void UTexAlignerDefault::PostInitProperties()
 
 void UTexAlignerDefault::AlignSurf( ETexAlign InTexAlignType, UModel* InModel, FBspSurfIdx* InSurfIdx, FPoly* InPoly, FVector* InNormal )
 {
-	InPoly->Base = FVector::ZeroVector;
+	InPoly->Base = InPoly->Vertices[0];
 	InPoly->TextureU = FVector::ZeroVector;
 	InPoly->TextureV = FVector::ZeroVector;
 	InPoly->Finalize( NULL, 0 );
@@ -219,9 +219,20 @@ void UTexAlignerDefault::AlignSurf( ETexAlign InTexAlignType, UModel* InModel, F
 	InPoly->TextureU *= UTile;
 	InPoly->TextureV *= VTile;
 
-	InSurfIdx->Surf->vTextureU = FBSPOps::bspAddVector( InModel, &InPoly->TextureU, 0);
-	InSurfIdx->Surf->vTextureV = FBSPOps::bspAddVector( InModel, &InPoly->TextureV, 0);
+	ABrush* Actor = InSurfIdx->Surf->Actor;
+	const FVector PrePivot = Actor->GetPrePivot();
+	const FVector Location = Actor->GetActorLocation();
+	const FRotator Rotation = Actor->GetActorRotation();
+	const FVector Scale = Actor->GetActorScale();
+	const FRotationMatrix RotMatrix(Rotation);
+
+	FVector Base = RotMatrix.TransformVector((InPoly->Base - PrePivot) * Scale) + Location;
+	FVector TextureU = RotMatrix.TransformVector(InPoly->TextureU / Scale);
+	FVector TextureV = RotMatrix.TransformVector(InPoly->TextureV / Scale);
 
+	InSurfIdx->Surf->pBase = FBSPOps::bspAddPoint(InModel, &Base, 0);
+	InSurfIdx->Surf->vTextureU = FBSPOps::bspAddVector( InModel, &TextureU, 0);
+	InSurfIdx->Surf->vTextureV = FBSPOps::bspAddVector( InModel, &TextureV, 0);
 }
 
 UTexAlignerBox::UTexAlignerBox(const FObjectInitializer& ObjectInitializer)
