Fixed enum ident in new feature pack mount code

diff --git a/Engine/Source/Editor/UnrealEd/Private/UnrealEdMisc.cpp b/Engine/Source/Editor/UnrealEd/Private/UnrealEdMisc.cpp
--- a/Engine/Source/Editor/UnrealEd/Private/UnrealEdMisc.cpp
+++ b/Engine/Source/Editor/UnrealEd/Private/UnrealEdMisc.cpp
@@ -1662,7 +1662,7 @@ void FUnrealEdMisc::MountTemplateSharedPaths()
 				if (EachPack.DetailLevels[iDetail] == EachEditDetail)
 				{
 					FString DetailString;
-					UEnum::GetValueAsString(TEXT("/Script/FeaturePackContentSource.EPackDetailLevel"), EachEditDetail, DetailString);
+					UEnum::GetValueAsString(TEXT("/Script/AddContentDialog.EFeaturePackDetailLevel"), EachEditDetail, DetailString);
 					FString ShareMountName = EachPack.MountName;
 					if (AddedMountSources.Find(ShareMountName) == INDEX_NONE)
 					{
