Move PortableObjectFormatDOM out of InternationalizationExportCommandlet

diff --git a/Engine/Source/Editor/UnrealEd/Classes/Commandlets/InternationalizationExportCommandlet.h b/Engine/Source/Editor/UnrealEd/Classes/Commandlets/InternationalizationExportCommandlet.h
--- a/Engine/Source/Editor/UnrealEd/Classes/Commandlets/InternationalizationExportCommandlet.h
+++ b/Engine/Source/Editor/UnrealEd/Classes/Commandlets/InternationalizationExportCommandlet.h
@@ -5,378 +5,6 @@
 #include "Commandlets/Commandlet.h"
 #include "InternationalizationExportCommandlet.generated.h"
 
-class FCulture;
-
-/**
-* Class for handling language codes used in the Portable Object format.
-*/
-class FPortableObjectCulture
-{
-public:
-	FPortableObjectCulture()
-	{}
-
-	FPortableObjectCulture( const FString& LangCode, const FString& PluralForms = TEXT("") );
-
-	FPortableObjectCulture( const FPortableObjectCulture& Other );
-
-	FPortableObjectCulture& operator=( const FPortableObjectCulture& Other )
-	{
-		if (this == &Other)
-		{
-			return *this;
-		}
-		SetLanguageCode( Other.GetLanguageCode() );
-		SetPluralForms( Other.GetPluralForms() );
-		return *this;
-	}
-
-	/**
-	 * Checks to see if we have a language code and if we were able to match a culture to that code
-	 */
-	bool IsValid() const { return !LanguageCode.IsEmpty() && Culture.IsValid(); }
-
-	/**
-	 * Sets the language code
-	 */
-	void SetLanguageCode( const FString& LangCode );
-
-	/**
-	 * Retrieves the language code
-	 */
-	const FString& GetLanguageCode() const { return LanguageCode; }
-
-	/**
-	 * Sets the language plural forms.  This is only required if we wish to override the default plural forms associated with a language.
-	 */
-	void SetPluralForms( const FString& InPluralForms ) { LanguagePluralForms = InPluralForms; }
-
-	/**
-	 * Retrieves the language plural forms.
-	 */
-	FString GetPluralForms() const;
-
-	/**
-	 * Retrieves the two or three letter language code for the culture
-	 */
-	FString Language() const;
-
-	/**
-	 * Retrieves the country info of the culture
-	 */
-	FString Country() const;
-
-	/**
-	 * Retrieves the variant info for the culture
-	 */
-	FString Variant() const;
-
-	/**
-	 * Retrieves the culture equivalent that matches the set language code.
-	 */
-	FCulturePtr GetCulture() const;
-
-	/**
-	 * Retrieves the display name for the language.
-	 */
-	FString DisplayName() const;
-
-	/**
-	 * Retrieves the English name for the language.
-	 */
-	FString EnglishName() const;
-
-
-private:
-	/**
-	 * Retrieves the language default plural forms.  If no plural forms are provided for this language we use this to pull from a set of defaults if available.
-	 */
-	FString GetDefaultPluralForms() const;
-
-	// The language code as it appears in the PO file.
-	FString LanguageCode;
-	FString LanguagePluralForms;
-
-	FCulturePtr Culture;
-};
-
-
-/**
-* Class for representing entries in a Portable Object file(.po) or a Portable Object Template file(.pot).
-*/
-class FPortableObjectEntry
-{
-public: 
-
-	FPortableObjectEntry()
-	{
-	}
-
-	FPortableObjectEntry( const FPortableObjectEntry& Other )
-		: MsgId( Other.MsgId )
-		, MsgIdPlural( Other.MsgIdPlural )
-		, MsgCtxt( Other.MsgCtxt )
-		, MsgStr( Other.MsgStr )
-		, UnknownElements( Other.UnknownElements )
-	{
-	}
-
-	FPortableObjectEntry& operator=( const FPortableObjectEntry& Other )
-	{
-		if (this == &Other)
-		{
-			return *this;
-		}
-		this->MsgId				= Other.MsgId;
-		this->MsgIdPlural		= Other.MsgIdPlural;
-		this->MsgCtxt			= Other.MsgCtxt;
-		this->MsgStr			= Other.MsgStr;
-		this->UnknownElements	= Other.UnknownElements;
-		return *this;
-	}
-
-	
-	bool operator==( const FPortableObjectEntry& Other ) 
-	{ 
-		return ( MsgId.Equals(Other.MsgId, ESearchCase::CaseSensitive) &&
-				MsgIdPlural.Equals(Other.MsgIdPlural, ESearchCase::CaseSensitive) &&
-				MsgCtxt.Equals(Other.MsgCtxt, ESearchCase::IgnoreCase) ); 
-	}
-	
-	/**
-	 * Helper function that adds to the extracted comments.
-	 *
-	 * @param InReference	String representing extracted comment entries of the following form: "File/path/file.cpp:20 File/path/file.cpp:21 File/path/file2.cpp:5".
-	 */
-	void AddExtractedComment( const FString& InComment );
-
-	/**
-	 * Helper function that adds to the reference comments.
-	 *
-	 * @param InReference	String representing reference comment entries.
-	 */
-	void AddReference( const FString& InReference );
-
-	/**
-	 * Helper function that adds to the reference comments.
-	 *
-	 * @param InReferences	String array representing reference comment entries of the following form: "File/path/file.cpp:20 File/path/file.cpp:21 File/path/file2.cpp:5".
-	 */
-	void AddExtractedComments( const TArray<FString>& InComments );
-
-	/**
-	 * Helper function that adds to the extracted comments.
-	 *
-	 * @param InReferences	String array representing extracted comment entries
-	 */
-	void AddReferences( const TArray<FString>& InReferences );
-
-	/**
-	 * Function to convert the entry to a string.
-	 *
-	 * @return	String representing the entry that can be written directly to the .po file.
-	 */
-	FString ToString() const;
-
-public:
-	/* Represents the original source text(also called the id or context).  Stored here are the msgid values from the Portable Object entries. */
-	FString MsgId;
-
-	/* Represents the plural form of the source text.  Stored here are the msgid_plural values from the Portable Object file entries. */
-	FString MsgIdPlural;
-
-	/* Represents the disambiguating  context for the source text.  If used, will prevent two identical source strings from getting collapsed into one entry. */
-	FString MsgCtxt;
-
-	/* Represents the translated text.  This stores the msgstr, msgstr[0], msgstr[1], etc values from Portable Object entries. */
-	TArray< FString > MsgStr;
-
-	/* Stores extracted comments. Lines starting with #. above the msgid.  
-	
-		#. TRANSLATORS: A test phrase with all letters of the English alphabet.
-		#. Replace it with a sample text in your language, such that it is
-		#. representative of language's writing system.
-		msgid "The Quick Brown Fox Jumps Over The Lazy Dog"
-		msgstr ""
-	*/
-	TArray<FString> ExtractedComments;
-
-	/* Stores the translator comments.  Lines starting with # (hash and space), followed by any text whatsoever.
-	
-		# Wikipedia says that 'etrurski' is our name for this script.
-		msgid "Old Italic"
-		msgstr "etrurski"
-	*/
-	TArray<FString> TranslatorComments;
-
-	/* Stores a reference comments. Lines starting with #: above the msgid.  
-	    
-		#: /Engine/Source/Runtime/Engine/Private/Actor.cpp:2306
-		#: /Engine/Source/Runtime/Engine/Private/Actor.cpp:2307 /Engine/Source/Runtime/Engine/Private/Actor.cpp:2308
-		msgid "The Quick Brown Fox Jumps Over The Lazy Dog"
-		msgstr ""
-	*/
-	TArray<FString> ReferenceComments;
-
-	/* Stores flags.  Lines starting with #,  
-	    
-		#: /Engine/Source/Runtime/Engine/Private/Actor.cpp:2306
-		#: /Engine/Source/Runtime/Engine/Private/Actor.cpp:2307 /Engine/Source/Runtime/Engine/Private/Actor.cpp:2308
-		msgid "The Quick Brown Fox Jumps Over The Lazy Dog"
-		msgstr ""
-	*/
-	TArray<FString> Flags;
-
-	/* Stores any unknown elements we may encounter when processing a Portable Object file.  */
-	TArray< FString > UnknownElements;
-};
-
-/**
-* Class that stores and manipulates PO and POT file header info.
-*/
-class FPortableObjectHeader
-{
-public:
-	typedef TPair<FString, FString> FPOHeaderEntry;
-	typedef TArray< FPOHeaderEntry > FPOHeaderData;
-
-	FPortableObjectHeader()
-	{
-	}
-
-	/**
-	 * Creates a string representation of the Portable Object header.
-	 *
-	 * @return	String representing the header that can be written directly to the .po file.
-	 */
-	FString ToString() const;
-
-	/**
-	 * Parses out header key/value pair entries.
-	 *
-	 * @param InStr	String representing a .po or .pot file header
-	 * @return	Returns true if successful, false otherwise.
-	 */
-	bool FromString( const FString& InStr ) {return false;}
-
-	/**
-	 * Parses out header info from a FPortableObjectEntry with an empty ID.
-	 *
-	 * @param LocEntry	The Portable Object entry with an empty ID
-	 * @return	Returns true if successful, false otherwise.
-	 */
-	bool FromLocPOEntry( const TSharedRef<const FPortableObjectEntry> LocEntry );
-	
-	/** Checks if a header entry exists with the given key */
-	bool HasEntry ( const FString& EntryKey );
-
-	/** Gets the header entry value for the given key or the empty string if the key does not exist */
-	FString GetEntryValue ( const FString& EntryKey );
-
-	/** Sets a header entry value */
-	void SetEntryValue( const FString& EntryKey, const FString& EntryValue );
-	
-	/**
-	 * Puts the current time into POT-Creation-Date and PO-Revision-Date entries
-	 */
-	void TimeStamp();
-
-	/**
-	 * Clears the header entries.
-	 */
-	void Clear() { HeaderEntries.Empty(); Comments.Empty(); }
-
-private:
-	FPortableObjectHeader::FPOHeaderEntry* FindEntry( const FString& EntryKey );
-
-public:
-
-	/* Stores the header comment block */
-	TArray<FString> Comments;
-	/* Stores all the header key/value pairs*/
-	FPOHeaderData HeaderEntries;
-};
-
-/**
-* Contains all the info we need to represent files in the Portable Object format.
-*/
-class FPortableObjectFormatDOM
-{
-public:
-	FPortableObjectFormatDOM()
-	{}
-
-	FPortableObjectFormatDOM( const FString& LanguageCode )
-		: Language( LanguageCode )
-	{}
-
-	/**
-	 * Creates a string representation of the Portable Object.
-	 *
-	 * @return	String representing the Portable Object that can be written directly to a file.
-	 */
-	FString ToString();
-
-	/**
-	 * Parses Portable Object elements from a string.
-	 *
-	 * @param InStr	String representing a Portable Object file(.PO) or Portable Object Template file(.POT).
-	 * @return	Returns true if successful, false otherwise.
-	 */
-	bool FromString( const FString& InStr );
-
-	/** Creates a header entry based on the project and language info. */
-	void CreateNewHeader();
-
-	/**
-	 * Sets the language.
-	 *
-	 * @param LanguageCode	String representing a Portable Object language code.
-	 * @param LangPluralForms	Optional plural forms for the language.  If not provided, defaults will be used if found.
-	 * @return	Returns true if the function was able to successfully set the language code and pair it with a known culture, false otherwise.
-	 */
-	bool SetLanguage( const FString& LanguageCode, const FString& LangPluralForms = TEXT("") );
-
-	/**
-	 * Adds a translation entry to the Portable Object.
-	 *
-	 * @param LocEntry	The LocPOEntry to add.
-	 * @return	Returns true if successful, false otherwise.
-	 */
-	bool AddEntry( const TSharedRef<FPortableObjectEntry> LocEntry );
-
-	/* Returns the Portable Object entry that matches the passed in entry or NULL if not found */
-	TSharedPtr<FPortableObjectEntry> FindEntry( const TSharedRef<const FPortableObjectEntry> LocEntry );
-
-	/* Returns the Portable Object file entry that matches the passed in parameters or NULL if not found */
-	TSharedPtr<FPortableObjectEntry> FindEntry( const FString& MsgId, const FString& MsgIdPlural, const FString& MsgCtxt );
-
-	/* Sets the project name that will appear in the Project-Id-Version header entry */
-	void SetProjectName( const FString& ProjName ){ ProjectName = ProjName; }
-
-	/* Sets the project name from the Project-Id-Version header entry */
-	FString GetProjectName() { return ProjectName; }
-
-	TArray< TSharedPtr< FPortableObjectEntry > >::TConstIterator GetEntriesIterator()
-	{
-		return Entries.CreateConstIterator();
-	}
-
-	void SortEntries();
-
-private:
-	/** Copying is not supported */
-	FPortableObjectFormatDOM( const FPortableObjectFormatDOM& ); 
-	FPortableObjectFormatDOM& operator=( const FPortableObjectFormatDOM& );
-
-private:
-	FPortableObjectCulture Language;
-	FPortableObjectHeader Header;
-	FString ProjectName;
-	TArray< TSharedPtr< FPortableObjectEntry > > Entries;
-};
-
-
 /**
  *	UInternationalizationExportCommandlet: Commandlet used to export internationalization data to various standard formats. 
  */
diff --git a/Engine/Source/Editor/UnrealEd/Private/Commandlets/InternationalizationExportCommandlet.cpp b/Engine/Source/Editor/UnrealEd/Private/Commandlets/InternationalizationExportCommandlet.cpp
--- a/Engine/Source/Editor/UnrealEd/Private/Commandlets/InternationalizationExportCommandlet.cpp
+++ b/Engine/Source/Editor/UnrealEd/Private/Commandlets/InternationalizationExportCommandlet.cpp
@@ -4,187 +4,22 @@
 #include "Culture.h"
 #include "JsonInternationalizationManifestSerializer.h"
 #include "JsonInternationalizationArchiveSerializer.h"
+#include "PortableObjectFormatDOM.h"
 
 DEFINE_LOG_CATEGORY_STATIC(LogInternationalizationExportCommandlet, Log, All);
 
 static const TCHAR* NewLineDelimiter = TEXT("\n");
 
-/* Default culture plural rules.  Culture names are in the following format: Language_Country@Variant
-  See references:	http://www.unicode.org/cldr/charts/latest/supplemental/language_plural_rules.html  
-					http://docs.translatehouse.org/projects/localization-guide/en/latest/l10n/pluralforms.html
-*/
-static const TMap< FString, FString > POCulturePluralForms = TMapBuilder< FString, FString >()
-	.Add( TEXT( "ach" )	, TEXT( "nplurals=2; plural=(n > 1);" )  )
-	.Add( TEXT( "af" )	, TEXT( "nplurals=2; plural=(n != 1);" )  )
-	.Add( TEXT( "ak" )	, TEXT( "nplurals=2; plural=(n > 1);" )  )
-	.Add( TEXT( "aln" )	, TEXT( "nplurals=2; plural=(n != 1);" )  )
-	.Add( TEXT( "am" )	, TEXT( "nplurals=2; plural=(n > 1);" )  )
-	.Add( TEXT( "am_ET"), TEXT( "nplurals=2; plural=(n != 1);" )  )
-	.Add( TEXT( "an" )	, TEXT( "nplurals=2; plural=(n != 1);" )  )
-	.Add( TEXT( "ar" )	, TEXT( "nplurals=6; plural=(n==0 ? 0 : n==1 ? 1 : n==2 ? 2 : n%100>=3 && n%100<=10 ? 3 : n%100>=11 && n%100<=99 ? 4 : 5);" )  )
-	.Add( TEXT( "ar_SA"), TEXT( "nplurals=2; plural=(n != 1);" )  )
-	.Add( TEXT( "arn" )	, TEXT( "nplurals=2; plural=(n > 1);" )  )
-	.Add( TEXT( "as" )	, TEXT( "nplurals=2; plural=(n != 1);" )  )
-	.Add( TEXT( "ast" )	, TEXT( "nplurals=2; plural=(n != 1);" )  )
-	.Add( TEXT( "ay" )	, TEXT( "nplurals=1; plural=0;" )  )
-	.Add( TEXT( "az" )	, TEXT( "nplurals=2; plural=(n != 1);" )  )
-	.Add( TEXT( "bal" )	, TEXT( "nplurals=2; plural=(n != 1);" )  )
-	.Add( TEXT( "be" )	, TEXT( "nplurals=3; plural=(n%10==1 && n%100!=11 ? 0 : n%10>=2 && n%10<=4 && (n%100<10 || n%100>=20) ? 1 : 2);" )  )
-	.Add( TEXT( "bg" )	, TEXT( "nplurals=2; plural=(n != 1);" )  )
-	.Add( TEXT( "bn" )	, TEXT( "nplurals=2; plural=(n != 1);" )  )
-	.Add( TEXT( "bo" )	, TEXT( "nplurals=1; plural=0;" )  )
-	.Add( TEXT( "br" )	, TEXT( "nplurals=2; plural=(n > 1);" )  )
-	.Add( TEXT( "brx" )	, TEXT( "nplurals=2; plural=(n != 1);" )  )
-	.Add( TEXT( "bs" )	, TEXT( "nplurals=3; plural=(n%10==1 && n%100!=11 ? 0 : n%10>=2 && n%10<=4 && (n%100<10 || n%100>=20) ? 1 : 2);" )  )
-	.Add( TEXT( "ca" )	, TEXT( "nplurals=2; plural=(n != 1);" )  )
-	.Add( TEXT( "cgg" )	, TEXT( "nplurals=1; plural=0;" )  )
-	.Add( TEXT( "crh" )	, TEXT( "nplurals=1; plural=0;" )  )
-	.Add( TEXT( "cs" )	, TEXT( "nplurals=3; plural=(n==1) ? 0 : (n>=2 && n<=4) ? 1 : 2;" )  )
-	.Add( TEXT( "csb" )	, TEXT( "nplurals=3; plural=(n==1 ? 0 : n%10>=2 && n%10<=4 && (n%100<10 || n%100>=20) ? 1 : 2;" )  )
-	.Add( TEXT( "cy" )	, TEXT( "nplurals=4; plural=(n==1) ? 0 : (n==2) ? 1 : (n != 8 && n != 11) ? 2 : 3;" )  )
-	.Add( TEXT( "da" )	, TEXT( "nplurals=2; plural=(n != 1);" )  )
-	.Add( TEXT( "de" )	, TEXT( "nplurals=2; plural=(n != 1);" )  )
-	.Add( TEXT( "doi" )	, TEXT( "nplurals=2; plural=(n != 1);" )  )
-	.Add( TEXT( "dz" )	, TEXT( "nplurals=1; plural=0;" )  )
-	.Add( TEXT( "el" )	, TEXT( "nplurals=2; plural=(n != 1);" )  )
-	.Add( TEXT( "en" )	, TEXT( "nplurals=2; plural=(n != 1);" )  )
-	.Add( TEXT( "eo" )	, TEXT( "nplurals=2; plural=(n != 1);" )  )
-	.Add( TEXT( "es" )	, TEXT( "nplurals=2; plural=(n != 1);" )  )
-	.Add( TEXT( "es_ar"), TEXT( "nplurals=2; plural=(n != 1);" )  )
-	.Add( TEXT( "et" )	, TEXT( "nplurals=2; plural=(n != 1);" )  )
-	.Add( TEXT( "eu" )	, TEXT( "nplurals=2; plural=(n != 1);" )  )
-	.Add( TEXT( "fa" )	, TEXT( "nplurals=1; plural=0;" )  )
-	.Add( TEXT( "fi" )	, TEXT( "nplurals=2; plural=(n != 1);" )  )
-	.Add( TEXT( "fil" )	, TEXT( "nplurals=2; plural=(n > 1);" )  )
-	.Add( TEXT( "fo" )	, TEXT( "nplurals=2; plural=(n != 1);" )  )
-	.Add( TEXT( "fr" )	, TEXT( "nplurals=2; plural=(n > 1);" )  )
-	.Add( TEXT( "frp" )	, TEXT( "nplurals=2; plural=(n > 1);" )  )
-	.Add( TEXT( "fur" )	, TEXT( "nplurals=2; plural=(n != 1);" )  )
-	.Add( TEXT( "fy" )	, TEXT( "nplurals=2; plural=(n != 1);" )  )
-	.Add( TEXT( "ga" )	, TEXT( "nplurals=5; plural=(n==1 ? 0 : n==2 ? 1 : n<7 ? 2 : n<11 ? 3 : 4);" )  )
-	.Add( TEXT( "gd" )	, TEXT( "nplurals=3; plural=(n < 2 ? 0 : n == 2 ? 1 : 2);" )  )
-	.Add( TEXT( "gl" )	, TEXT( "nplurals=2; plural=(n != 1);" )  )
-	.Add( TEXT( "gu" )	, TEXT( "nplurals=2; plural=(n != 1);" )  )
-	.Add( TEXT( "gun" )	, TEXT( "nplurals=2; plural=(n > 1);" )  )
-	.Add( TEXT( "ha" )	, TEXT( "nplurals=2; plural=(n != 1);" )  )
-	.Add( TEXT( "he" )	, TEXT( "nplurals=2; plural=(n != 1);" )  )
-	.Add( TEXT( "hi" )	, TEXT( "nplurals=2; plural=(n != 1);" )  )
-	.Add( TEXT( "hne" )	, TEXT( "nplurals=2; plural=(n != 1);" )  )
-	.Add( TEXT( "hr" )	, TEXT( "nplurals=3; plural=(n%10==1 && n%100!=11 ? 0 : n%10>=2 && n%10<=4 && (n%100<10 || n%100>=20) ? 1 : 2);" )  )
-	.Add( TEXT( "hsb" )	, TEXT( "nplurals=4; plural=(n%100==1 ? 0 : n%100==2 ? 1 : n%100==3 || n%100==4 ? 2 : 3);" )  )
-	.Add( TEXT( "ht" )	, TEXT( "nplurals=2; plural=(n != 1);" )  )
-	.Add( TEXT( "hu" )	, TEXT( "nplurals=2; plural=(n != 1);" )  )
-	.Add( TEXT( "hy" )	, TEXT( "nplurals=2; plural=(n != 1);" )  )
-	.Add( TEXT( "ia" )	, TEXT( "nplurals=2; plural=(n != 1);" )  )
-	.Add( TEXT( "id" )	, TEXT( "nplurals=1; plural=0;" )  )
-	.Add( TEXT( "ig" )	, TEXT( "nplurals=2; plural=(n != 1);" )  )
-	.Add( TEXT( "ilo" )	, TEXT( "nplurals=2; plural=(n != 1);" )  )
-	.Add( TEXT( "is" )	, TEXT( "nplurals=2; plural=(n != 1);" )  )
-	.Add( TEXT( "it" )	, TEXT( "nplurals=2; plural=(n != 1);" )  )
-	.Add( TEXT( "ja" )	, TEXT( "nplurals=1; plural=0;" )  )
-	.Add( TEXT( "jv" )	, TEXT( "nplurals=2; plural=(n != 1);" )  )
-	.Add( TEXT( "ka" )	, TEXT( "nplurals=1; plural=0;" )  )
-	.Add( TEXT( "kk" )	, TEXT( "nplurals=1; plural=0;" )  )
-	.Add( TEXT( "km" )	, TEXT( "nplurals=1; plural=0;" )  )
-	.Add( TEXT( "kn" )	, TEXT( "nplurals=1; plural=0;" )  )
-	.Add( TEXT( "ko" )	, TEXT( "nplurals=1; plural=0;" )  )
-	.Add( TEXT( "ks" )	, TEXT( "nplurals=2; plural=(n != 1);" )  )
-	.Add( TEXT( "ku" )	, TEXT( "nplurals=2; plural=(n != 1);" )  )
-	.Add( TEXT( "kw" )	, TEXT( "nplurals=4; plural=(n==1) ? 0 : (n==2) ? 1 : (n == 3) ? 2 : 3;" )  )
-	.Add( TEXT( "ky" )	, TEXT( "nplurals=1; plural=0;" )  )
-	.Add( TEXT( "la" )	, TEXT( "nplurals=2; plural=(n != 1);" )  )
-	.Add( TEXT( "lb" )	, TEXT( "nplurals=2; plural=(n != 1);" )  )
-	.Add( TEXT( "li" )	, TEXT( "nplurals=2; plural=(n != 1);" )  )
-	.Add( TEXT( "ln" )	, TEXT( "nplurals=2; plural=(n > 1);" )  )
-	.Add( TEXT( "lo" )	, TEXT( "nplurals=1; plural=0;" )  )
-	.Add( TEXT( "lt" )	, TEXT( "nplurals=3; plural=(n%10==1 && n%100!=11 ? 0 : n%10>=2 && (n%100<10 || n%100>=20) ? 1 : 2);" )  )
-	.Add( TEXT( "lv" )	, TEXT( "nplurals=3; plural=(n%10==1 && n%100!=11 ? 0 : n != 0 ? 1 : 2);" )  )
-	.Add( TEXT( "mai" )	, TEXT( "nplurals=2; plural=(n != 1);" )  )
-	.Add( TEXT( "mg" )	, TEXT( "nplurals=2; plural=(n > 1);" )  )
-	.Add( TEXT( "mi" )	, TEXT( "nplurals=2; plural=(n > 1);" )  )
-	.Add( TEXT( "mk" )	, TEXT( "nplurals=2; plural=(n % 10 == 1 && n % 100 != 11) ? 0 : 1;" )  )
-	.Add( TEXT( "ml" )	, TEXT( "nplurals=2; plural=(n != 1);" )  )
-	.Add( TEXT( "mn" )	, TEXT( "nplurals=2; plural=(n != 1);" )  )
-	.Add( TEXT( "mr" )	, TEXT( "nplurals=2; plural=(n != 1);" )  )
-	.Add( TEXT( "ms" )	, TEXT( "nplurals=1; plural=0;" )  )
-	.Add( TEXT( "mt" )	, TEXT( "nplurals=4; plural=(n==1 ? 0 : n==0 || ( n%100>1 && n%100<11) ? 1 : (n%100>10 && n%100<20 ) ? 2 : 3);" )  )
-	.Add( TEXT( "my" )	, TEXT( "nplurals=1; plural=0;" )  )
-	.Add( TEXT( "nah" )	, TEXT( "nplurals=2; plural=(n != 1);" )  )
-	.Add( TEXT( "nap" )	, TEXT( "nplurals=2; plural=(n != 1);" )  )
-	.Add( TEXT( "nb" )	, TEXT( "nplurals=2; plural=(n != 1);" )  )
-	.Add( TEXT( "nds" )	, TEXT( "nplurals=2; plural=(n != 1);" )  )
-	.Add( TEXT( "ne" )	, TEXT( "nplurals=2; plural=(n != 1);" )  )
-	.Add( TEXT( "nl" )	, TEXT( "nplurals=2; plural=(n != 1);" )  )
-	.Add( TEXT( "nn" )	, TEXT( "nplurals=2; plural=(n != 1);" )  )
-	.Add( TEXT( "no" )	, TEXT( "nplurals=2; plural=(n != 1);" )  )
-	.Add( TEXT( "nr" )	, TEXT( "nplurals=1; plural=0;" )  )
-	.Add( TEXT( "nso" )	, TEXT( "nplurals=2; plural=(n > 1);" )  )
-	.Add( TEXT( "oc" )	, TEXT( "nplurals=2; plural=(n > 1);" )  )
-	.Add( TEXT( "or" )	, TEXT( "nplurals=2; plural=(n != 1);" )  )
-	.Add( TEXT( "pa" )	, TEXT( "nplurals=2; plural=(n != 1);" )  )
-	.Add( TEXT( "pap" )	, TEXT( "nplurals=2; plural=(n != 1);" )  )
-	.Add( TEXT( "pl" )	, TEXT( "nplurals=3; plural=(n==1 ? 0 : n%10>=2 && n%10<=4 && (n%100<10 || n%100>=20) ? 1 : 2);" )  )
-	.Add( TEXT( "pms" )	, TEXT( "nplurals=2; plural=(n != 1);" )  )
-	.Add( TEXT( "ps" )	, TEXT( "nplurals=2; plural=(n != 1);" )  )
-	.Add( TEXT( "pt" )	, TEXT( "nplurals=2; plural=(n != 1);" )  )
-	.Add( TEXT( "pt_BR"), TEXT( "nplurals=2; plural=(n > 1);" )  )
-	.Add( TEXT( "rm" )	, TEXT( "nplurals=2; plural=(n != 1);" )  )
-	.Add( TEXT( "ro" )	, TEXT( "nplurals=3; plural=(n==1?0:(((n%100>19)||((n%100==0)&&(n!=0)))?2:1));" )  )
-	.Add( TEXT( "ru" )	, TEXT( "nplurals=3; plural=(n%10==1 && n%100!=11 ? 0 : n%10>=2 && n%10<=4 && (n%100<10 || n%100>=20) ? 1 : 2);" )  )
-	.Add( TEXT( "rw" )	, TEXT( "nplurals=2; plural=(n != 1);" )  )
-	.Add( TEXT( "sc" )	, TEXT( "nplurals=2; plural=(n != 1);" )  )
-	.Add( TEXT( "sco" )	, TEXT( "nplurals=2; plural=(n != 1);" )  )
-	.Add( TEXT( "se" )	, TEXT( "nplurals=2; plural=(n != 1);" )  )
-	.Add( TEXT( "si" )	, TEXT( "nplurals=2; plural=(n != 1);" )  )
-	.Add( TEXT( "sk" )	, TEXT( "nplurals=3; plural=(n==1) ? 0 : (n>=2 && n<=4) ? 1 : 2;" )  )
-	.Add( TEXT( "sl" )	, TEXT( "nplurals=4; plural=(n%100==1 ? 0 : n%100==2 ? 1 : n%100==3 || n%100==4 ? 2 : 3);" )  )
-	.Add( TEXT( "sm" )	, TEXT( "nplurals=1; plural=0;" )  )
-	.Add( TEXT( "sn" )	, TEXT( "nplurals=2; plural=(n != 1);" )  )
-	.Add( TEXT( "so" )	, TEXT( "nplurals=2; plural=(n != 1);" )  )
-	.Add( TEXT( "son" )	, TEXT( "nplurals=1; plural=0;" )  )
-	.Add( TEXT( "sq" )	, TEXT( "nplurals=2; plural=(n != 1);" )  )
-	.Add( TEXT( "sr" )	, TEXT( "nplurals=3; plural=(n%10==1 && n%100!=11 ? 0 : n%10>=2 && n%10<=4 && (n%100<10 || n%100>=20) ? 1 : 2);" )  )
-	.Add( TEXT( "st" )	, TEXT( "nplurals=2; plural=(n != 1);" )  )
-	.Add( TEXT( "su" )	, TEXT( "nplurals=1; plural=0;" )  )
-	.Add( TEXT( "sv" )	, TEXT( "nplurals=2; plural=(n != 1);" )  )
-	.Add( TEXT( "sw" )	, TEXT( "nplurals=2; plural=(n != 1);" )  )
-	.Add( TEXT( "ta" )	, TEXT( "nplurals=2; plural=(n != 1);" )  )
-	.Add( TEXT( "te" )	, TEXT( "nplurals=2; plural=(n != 1);" )  )
-	.Add( TEXT( "tg" )	, TEXT( "nplurals=2; plural=(n != 1);" )  )
-	.Add( TEXT( "th" )	, TEXT( "nplurals=1; plural=0;" )  )
-	.Add( TEXT( "ti" )	, TEXT( "nplurals=2; plural=(n > 1);" )  )
-	.Add( TEXT( "tk" )	, TEXT( "nplurals=2; plural=(n != 1);" )  )
-	.Add( TEXT( "tl" )	, TEXT( "nplurals=2; plural=(n > 1);" )  )
-	.Add( TEXT( "tlh" )	, TEXT( "nplurals=2; plural=(n > 1);" )  )
-	.Add( TEXT( "to" )	, TEXT( "nplurals=1; plural=0;" )  )
-	.Add( TEXT( "tr" )	, TEXT( "nplurals=1; plural=0;" )  )
-	.Add( TEXT( "tt" )	, TEXT( "nplurals=1; plural=0;" )  )
-	.Add( TEXT( "udm" )	, TEXT( "nplurals=1; plural=0;" )  )
-	.Add( TEXT( "ug" )	, TEXT( "nplurals=1; plural=0;" )  )
-	.Add( TEXT( "uk" )	, TEXT( "nplurals=3; plural=(n%10==1 && n%100!=11 ? 0 : n%10>=2 && n%10<=4 && (n%100<10 || n%100>=20) ? 1 : 2);" )  )
-	.Add( TEXT( "ur" )	, TEXT( "nplurals=2; plural=(n != 1);" )  )
-	.Add( TEXT( "uz" )	, TEXT( "nplurals=1; plural=0;" )  )
-	.Add( TEXT( "ve" )	, TEXT( "nplurals=2; plural=(n != 1);" )  )
-	.Add( TEXT( "vi" )	, TEXT( "nplurals=1; plural=0;" )  )
-	.Add( TEXT( "vls" )	, TEXT( "nplurals=2; plural=(n != 1);" )  )
-	.Add( TEXT( "wa" )	, TEXT( "nplurals=2; plural=(n > 1);" )  )
-	.Add( TEXT( "wo" )	, TEXT( "nplurals=1; plural=0;" )  )
-	.Add( TEXT( "xh" )	, TEXT( "nplurals=2; plural=(n != 1);" )  )
-	.Add( TEXT( "yi" )	, TEXT( "nplurals=2; plural=(n != 1);" )  )
-	.Add( TEXT( "yo" )	, TEXT( "nplurals=2; plural=(n != 1);" )  )
-	.Add( TEXT( "zh" )	, TEXT( "nplurals=1; plural=0;" )  )
-	.Add( TEXT( "zu" )	, TEXT( "nplurals=2; plural=(n != 1);" )  );
-
 
 /**
 *	Helper Functions
 */
-FString ConditionArchiveStrForPo( const FString& InStr )
+FString ConditionArchiveStrForPo(const FString& InStr)
 {
 	FString Result;
 	for (const TCHAR C : InStr)
 	{
-		switch(C)
+		switch (C)
 		{
 		case TEXT('\\'):	Result += TEXT("\\\\");	break;
 		case TEXT('"'):		Result += TEXT("\\\"");	break;
@@ -197,7 +32,7 @@ FString ConditionArchiveStrForPo( const FString& InStr )
 	return Result;
 }
 
-FString ConditionPoStringForArchive( const FString& InStr )
+FString ConditionPoStringForArchive(const FString& InStr)
 {
 	FString Result;
 	FString EscapeSequenceBuffer;
@@ -205,29 +40,29 @@ FString ConditionPoStringForArchive( const FString& InStr )
 	auto HandleEscapeSequenceBuffer = [&]()
 	{
 		// Insert unescaped sequence if needed.
-		if(!EscapeSequenceBuffer.IsEmpty())
+		if (!EscapeSequenceBuffer.IsEmpty())
 		{
 			bool EscapeSequenceIdentified = true;
 
 			// Identify escape sequence
 			TCHAR UnescapedCharacter = 0;
-			if(EscapeSequenceBuffer == TEXT("\\\\"))
+			if (EscapeSequenceBuffer == TEXT("\\\\"))
 			{
 				UnescapedCharacter = '\\';
 			}
-			else if(EscapeSequenceBuffer == TEXT("\\\""))
+			else if (EscapeSequenceBuffer == TEXT("\\\""))
 			{
 				UnescapedCharacter = '"';
 			}
-			else if(EscapeSequenceBuffer == TEXT("\\r"))
+			else if (EscapeSequenceBuffer == TEXT("\\r"))
 			{
 				UnescapedCharacter = '\r';
 			}
-			else if(EscapeSequenceBuffer == TEXT("\\n"))
+			else if (EscapeSequenceBuffer == TEXT("\\n"))
 			{
 				UnescapedCharacter = '\n';
 			}
-			else if(EscapeSequenceBuffer == TEXT("\\t"))
+			else if (EscapeSequenceBuffer == TEXT("\\t"))
 			{
 				UnescapedCharacter = '\t';
 			}
@@ -237,7 +72,7 @@ FString ConditionPoStringForArchive( const FString& InStr )
 			}
 
 			// If identified, append the processed sequence as the unescaped character.
-			if(EscapeSequenceIdentified)
+			if (EscapeSequenceIdentified)
 			{
 				Result += UnescapedCharacter;
 			}
@@ -254,10 +89,10 @@ FString ConditionPoStringForArchive( const FString& InStr )
 	for (const TCHAR C : InStr)
 	{
 		// Not in an escape sequence.
-		if(EscapeSequenceBuffer.IsEmpty())
+		if (EscapeSequenceBuffer.IsEmpty())
 		{
 			// Regular character, just copy over.
-			if(C != TEXT('\\'))
+			if (C != TEXT('\\'))
 			{
 				Result += C;
 			}
@@ -282,803 +117,14 @@ FString ConditionPoStringForArchive( const FString& InStr )
 }
 
 
-FString ConvertSrcLocationToPORef( const FString& InSrcLocation )
+FString ConvertSrcLocationToPORef(const FString& InSrcLocation)
 {
 	// Source location format: /Path1/Path2/file.cpp - line 123
 	// PO Reference format: /Path1/Path2/file.cpp:123
 	// @TODO: Note, we assume the source location format here but it could be arbitrary.
-	return InSrcLocation.Replace( TEXT(" - line "), TEXT(":") );
-}
-
-bool FindDelimitedString( const FString& InStr, const FString& LeftDelim, const FString& RightDelim, FString& Result  )
-{
-	int32 Start = InStr.Find( LeftDelim, ESearchCase::CaseSensitive );
-	int32 End = InStr.Find( RightDelim, ESearchCase::CaseSensitive, ESearchDir::FromEnd );
-	if( Start == INDEX_NONE || !(End>Start))
-	{
-		return false;
-	}
-	Start += LeftDelim.Len();
-	if( Start >= End )
-	{
-		Result = TEXT("");
-	}
-	else
-	{
-		Result = InStr.Mid( Start, End - Start );
-	}
-
-	return true;
-}
-
-
-FPortableObjectCulture::FPortableObjectCulture( const FString& LangCode, const FString& PluralForms )
-	: LanguageCode( LangCode )
-	, LanguagePluralForms( PluralForms )
-	, Culture( FCulture::Create( LangCode ) )
-{
-
-}
-
-FPortableObjectCulture::FPortableObjectCulture( const FPortableObjectCulture& Other )
-	: LanguageCode( Other.LanguageCode )
-	, LanguagePluralForms( Other.LanguagePluralForms )
-	, Culture( FCulture::Create( Other.LanguageCode ) )
-{
-
-}
-
-void FPortableObjectCulture::SetLanguageCode( const FString& LangCode )
-{
-	LanguageCode = LangCode;
-	Culture = FCulture::Create( LangCode );
-}
-
-
-FString FPortableObjectCulture::Language() const
-{
-	FString Result;
-	if( Culture.IsValid() )
-	{
-		// NOTE: The below function name may be a bit misleading because it will return three letter language codes if needed.
-		Result = Culture->GetTwoLetterISOLanguageName();
-	}
-	return Result;
-
-}
-
-FString FPortableObjectCulture::Country() const
-{
-	FString Result;
-	if( Culture.IsValid() )
-	{
-		Result = Culture->GetRegion();
-	}
-	return Result;
-}
-
-FString FPortableObjectCulture::Variant() const
-{
-	FString Result;
-	if( Culture.IsValid() )
-	{
-		Result = Culture->GetVariant();
-	}
-	return Result;
-}
-
-FString FPortableObjectCulture::DisplayName() const
-{
-	FString Result;
-	if( Culture.IsValid() )
-	{
-		Result = Culture->GetDisplayName();
-	}
-	return Result;
-}
-
-FString FPortableObjectCulture::EnglishName() const
-{
-	FString Result;
-	if( Culture.IsValid() )
-	{
-		Result = Culture->GetEnglishName();
-	}
-	return Result;
-}
-
-FString FPortableObjectCulture::GetPluralForms() const
-{
-	if( !LanguagePluralForms.IsEmpty() )
-	{
-		return LanguagePluralForms;
-	}
-	return GetDefaultPluralForms();
-}
-
-FString FPortableObjectCulture::GetDefaultPluralForms() const
-{
-	FString Result;
-	if( LanguageCode.IsEmpty() )
-	{
-		return Result;
-	}
-
-	if( auto* LanguageCodePair = POCulturePluralForms.Find( GetLanguageCode() ) )
-	{
-		Result = *LanguageCodePair;
-	}
-	else if( auto* LangCountryVariantPair = POCulturePluralForms.Find( Language() + TEXT("_") + Country() + TEXT("@") + Variant() ))
-	{
-		Result = *LangCountryVariantPair;
-	}
-	else if( auto* LangCountryPair = POCulturePluralForms.Find( Language() + TEXT("_") + Country() ) )
-	{
-		Result  = *LangCountryPair;
-	} 
-	else if( auto* LangPair = POCulturePluralForms.Find( Language() ) )
-	{
-		Result = *LangPair;
-	}
-	else
-	{
-		auto* Fallback = POCulturePluralForms.Find( TEXT("en") );
-		if( Fallback )
-		{
-			Result = *Fallback;
-		}
-		else
-		{
-			Result = TEXT("nplurals=2; plural=(n != 1);");
-		}
-	}
-
-	return Result;
-}
-
-
-FString FPortableObjectHeader::ToString() const
-{
-	FString Result;
-	for( const FString& Comment : Comments )
-	{
-		Result += FString::Printf(TEXT("# %s%s"), *Comment, NewLineDelimiter);
-	}
-
-	Result += FString::Printf(TEXT("msgid \"\"%s"), NewLineDelimiter);
-	Result += FString::Printf(TEXT("msgstr \"\"%s"), NewLineDelimiter);
-
-	for( auto Entry : HeaderEntries )
-	{
-		const FString& Key = Entry.Key;
-		const FString& Value = Entry.Value;
-		Result += FString::Printf( TEXT("\"%s: %s\\n\"%s"), *Key, *Value, NewLineDelimiter );
-	}
-
-	return Result;
-}
-
-bool FPortableObjectHeader::FromLocPOEntry( const TSharedRef<const FPortableObjectEntry> LocEntry )
-{
-	if( !LocEntry->MsgId.IsEmpty() || LocEntry->MsgStr.Num() != 1 )
-	{
-		return false;
-	}
-	Clear();
-
-	Comments = LocEntry->TranslatorComments;
-
-	// The POEntry would store header info inside the MsgStr[0]
-	TArray<FString> HeaderLinesToProcess;
-	LocEntry->MsgStr[0].ReplaceEscapedCharWithChar().ParseIntoArray( HeaderLinesToProcess, NewLineDelimiter, true );
-
-	for( const FString& PotentialHeaderEntry : HeaderLinesToProcess )
-	{
-		int32 SplitIndex;
-		if( PotentialHeaderEntry.FindChar( TCHAR(':'), SplitIndex ) )
-		{
-			// Looks like a header entry so we add it
-			const FString& Key = PotentialHeaderEntry.LeftChop( PotentialHeaderEntry.Len() - SplitIndex ).Trim().TrimTrailing();
-			FString Value = PotentialHeaderEntry.RightChop( SplitIndex+1 ).Trim().TrimTrailing();
-
-			HeaderEntries.Add( TPairInitializer<FString, FString>( Key, Value ) );
-		}
-	}
-	return true;
-}
-
-FPortableObjectHeader::FPOHeaderEntry* FPortableObjectHeader::FindEntry( const FString& EntryKey )
-{
-	for( auto& Entry : HeaderEntries )
-	{
-		if( Entry.Key == EntryKey )
-		{
-			return &Entry;
-		}
-	}
-	return NULL;
-}
-
-FString FPortableObjectHeader::GetEntryValue( const FString& EntryKey )
-{
-	FString Result;
-	const FPOHeaderEntry* Entry = FindEntry( EntryKey );
-	if( Entry )
-	{
-		Result = Entry->Value;
-	}
-	return Result;
-}
-
-bool FPortableObjectHeader::HasEntry( const FString& EntryKey )
-{
-	return ( FindEntry( EntryKey ) != NULL);
-}
-
-void FPortableObjectHeader::SetEntryValue( const FString& EntryKey, const FString& EntryValue )
-{
-	FPOHeaderEntry* Entry = FindEntry( EntryKey );
-	if( Entry )
-	{
-		Entry->Value = EntryValue;
-	}
-	else
-	{
-		HeaderEntries.Add( TPairInitializer<FString, FString>( EntryKey, EntryValue ) );
-	}
-}
-
-void FPortableObjectHeader::TimeStamp()
-{
-	// @TODO: Time format not exactly correct.  We have something like this: 2014-02-07 20:06  This is what it should be: 2014-02-07 14:12-0600
-	FString Time = *FDateTime::UtcNow().ToString(TEXT("%Y-%m-%d %H:%M"));
-
-	SetEntryValue( TEXT("POT-Creation-Date"), Time );
-	SetEntryValue( TEXT("PO-Revision-Date"), Time );
-}
-
-FString FPortableObjectFormatDOM::ToString()
-{
-	FString Result;
-
-	Header.TimeStamp();
-
-	Result += Header.ToString();
-	Result += NewLineDelimiter;
-
-	for( auto Entry : Entries )
-	{
-		Result += Entry->ToString();
-		Result += NewLineDelimiter;
-	}
-
-	return Result;
-}
-
-
-
-bool FPortableObjectFormatDOM::FromString( const FString& InStr )
-{
-	if( InStr.IsEmpty() )
-	{
-		return false;
-	}
-
-	bool bSuccess = true;
-
-	FString ParseString = InStr.Replace(TEXT("\r\n"), NewLineDelimiter);
-
-	TArray<FString> LinesToProcess;
-	ParseString.ParseIntoArray( LinesToProcess, NewLineDelimiter, false );
-
-	TSharedRef<FPortableObjectEntry> ProcessedEntry = MakeShareable( new FPortableObjectEntry );
-	bool bHasMsgId = false;
-
-	uint32 NumFileLines = LinesToProcess.Num();
-	for( uint32 LineIdx = 0; LineIdx < NumFileLines; ++LineIdx )
-	{
-		const FString& Line = LinesToProcess[LineIdx];
-
-		if( Line.IsEmpty() )
-		{
-			// When we encounter a blank line, we will either ignore it, or consider it the boundary of 
-			// a LocPOEntry or FPortableObjectHeader if we processed any valid data before encountering the blank.
-
-			// If this entry is valid we'll check it and process it further.
-			if( bHasMsgId && ProcessedEntry->MsgStr.Num() > 0 )
-			{
-				// Check if we are dealing with a header entry
-				if( ProcessedEntry->MsgId.IsEmpty() && ProcessedEntry->MsgStr.Num() == 1 )
-				{
-					// This is a header
-					bool bAddedHeader = Header.FromLocPOEntry( ProcessedEntry );
-					if( !bAddedHeader )
-					{
-						return false;
-					}
-					ProjectName = Header.GetEntryValue(TEXT("Project-Id-Version"));
-				}
-				else
-				{
-					bool bAddEntry = AddEntry( ProcessedEntry );
-					if( !bAddEntry )
-					{
-						return false;
-					}
-				}
-			}
-
-			// Now starting a new entry and reset any flags we have been tracking
-			ProcessedEntry = MakeShareable( new FPortableObjectEntry );
-			bHasMsgId = false;
-		}
-		else if( Line.StartsWith( TEXT("#,") ) )
-		{
-			// Flags
-			FString Flag;
-			if( FindDelimitedString(Line, TEXT("#,"), NewLineDelimiter, Flag ) )
-			{
-				if( !Flag.IsEmpty() )
-				{
-					ProcessedEntry->Flags.Add( Flag );
-				}
-			}
-		}
-		else if( Line.StartsWith( TEXT("#.") ) )
-		{
-			// Extracted comments
-			FString Comment;
-			ProcessedEntry->ExtractedComments.Add( Line.RightChop( FString(TEXT("#. ")).Len() ) );
-		}
-		else if( Line.StartsWith( TEXT("#:") ) )
-		{
-			// Reference
-			FString Reference;
-			ProcessedEntry->AddReference( Line.RightChop( FString(TEXT("#: ")).Len() ) );
-		}
-		else if( Line.StartsWith( TEXT(":|") ) )
-		{
-			// This represents previous messages. We just drop this in unknown since we don't handle it
-			ProcessedEntry->UnknownElements.Add( Line );
-		}
-		else if( Line.StartsWith( TEXT("# ") ) || Line.StartsWith( TEXT("#\t") ) )
-		{
-			FString Comment;
-			ProcessedEntry->TranslatorComments.Add( Line.RightChop( FString(TEXT("# ")).Len() ) );
-		}
-		else if( Line == TEXT("#") )
-		{
-			ProcessedEntry->TranslatorComments.Add( TEXT("") );		
-		}
-		else if( Line.StartsWith( TEXT("msgctxt") ) )
-		{
-			FString RawMsgCtxt;
-			if( !FindDelimitedString(Line, TEXT("\""), TEXT("\""), RawMsgCtxt ) )
-			{
-				bSuccess = false;
-				break;
-			}
-			for( uint32 NextLineIdx = LineIdx + 1; NextLineIdx < NumFileLines && LinesToProcess[NextLineIdx].Trim().TrimTrailing().StartsWith(TEXT("\"")); ++NextLineIdx)
-			{
-				FString Tmp;
-				if (FindDelimitedString(Line, TEXT("\""), TEXT("\""), Tmp))
-				{
-					RawMsgCtxt += Tmp;
-				}
-			}
-
-			// If the following line contains more info for this element we continue to parse it out
-			uint32 NextLineIdx = LineIdx + 1;
-			while( NextLineIdx < NumFileLines )
-			{
-				const FString& NextLine = LinesToProcess[NextLineIdx].TrimTrailing().Trim();
-				if( NextLine.StartsWith("\"") && NextLine.EndsWith("\"") )
-				{
-					RawMsgCtxt += NextLine.Mid( 1, NextLine.Len()-2 );
-				}
-				else
-				{
-					break;
-				}
-				LineIdx = NextLineIdx;	
-				NextLineIdx++;
-			}
-
-			ProcessedEntry->MsgCtxt = RawMsgCtxt;
-		}
-		else if( Line.StartsWith( TEXT("msgid") ) )
-		{
-			FString RawMsgId;
-			if( !FindDelimitedString(Line, TEXT("\""), TEXT("\""), RawMsgId ) )
-			{
-				bSuccess = false;
-				break;
-			}
-			// If the following line contains more info for this element we continue to parse it out
-			uint32 NextLineIdx = LineIdx + 1;
-			while( NextLineIdx < NumFileLines )
-			{
-				const FString& NextLine = LinesToProcess[NextLineIdx].TrimTrailing().Trim();
-				if( NextLine.StartsWith("\"") && NextLine.EndsWith("\"") )
-				{
-					RawMsgId += NextLine.Mid( 1, NextLine.Len()-2 );
-				}
-				else
-				{
-					break;
-				}
-				LineIdx = NextLineIdx;	
-				NextLineIdx++;
-			}
-			ProcessedEntry->MsgId = RawMsgId;
-			bHasMsgId = true;
-		}
-		else if( Line.StartsWith( TEXT("msgid_plural") ) )
-		{
-			FString RawMsgIdPlural;
-			if( !FindDelimitedString(Line, TEXT("\""), TEXT("\""), RawMsgIdPlural ) )
-			{
-				bSuccess = false;
-				break;
-			}
-			// If the following line contains more info for this element we continue to parse it out
-			uint32 NextLineIdx = LineIdx + 1;
-			while( NextLineIdx < NumFileLines )
-			{
-				const FString& NextLine = LinesToProcess[NextLineIdx].TrimTrailing().Trim();
-				if( NextLine.StartsWith("\"") && NextLine.EndsWith("\"") )
-				{
-					RawMsgIdPlural += NextLine.Mid( 1, NextLine.Len()-2 );
-				}
-				else
-				{
-					break;
-				}
-				LineIdx = NextLineIdx;	
-				NextLineIdx++;
-			}
-			ProcessedEntry->MsgIdPlural = RawMsgIdPlural;
-		}
-		else if( Line.StartsWith( TEXT("msgstr[") ) )
-		{
-			FString IndexStr;
-			if( !FindDelimitedString(Line, TEXT("["), TEXT("]"), IndexStr ) )
-			{
-				bSuccess = false;
-				break;
-			}
-			int32 Index = -1;
-			TTypeFromString<int32>::FromString(Index, *IndexStr );
-
-			check(Index >= 0);
-
-			FString RawMsgStr;
-			if( !FindDelimitedString(Line, TEXT("\""), TEXT("\""), RawMsgStr ) )
-			{
-				bSuccess = false;
-				break;
-			}
-			// If the following line contains more info for this element we continue to parse it out
-			uint32 NextLineIdx = LineIdx + 1;
-			while( NextLineIdx < NumFileLines )
-			{
-				const FString& NextLine = LinesToProcess[NextLineIdx].TrimTrailing().Trim();
-				if( NextLine.StartsWith("\"") && NextLine.EndsWith("\"") )
-				{
-					RawMsgStr += NextLine.Mid( 1, NextLine.Len()-2 );
-				}
-				else
-				{
-					break;
-				}
-				LineIdx = NextLineIdx;	
-				NextLineIdx++;
-			}
-
-			if( ProcessedEntry->MsgStr.Num() > Index )
-			{
-				ProcessedEntry->MsgStr[Index] = RawMsgStr;
-			}
-			else
-			{
-				ProcessedEntry->MsgStr.Insert( RawMsgStr, Index );
-			}
-		}
-		else if( Line.StartsWith( TEXT("msgstr") ) )
-		{
-			FString RawMsgStr;
-			if( !FindDelimitedString(Line, TEXT("\""), TEXT("\""), RawMsgStr ) )
-			{
-				bSuccess = false;
-				break;
-			}
-			// If the following line contains more info for this element we continue to parse it out
-			uint32 NextLineIdx = LineIdx + 1;
-			while( NextLineIdx < NumFileLines )
-			{
-				const FString& NextLine = LinesToProcess[NextLineIdx].TrimTrailing().Trim();
-				if( NextLine.StartsWith("\"") && NextLine.EndsWith("\"") )
-				{
-					RawMsgStr += NextLine.Mid( 1, NextLine.Len()-2 );
-				}
-				else
-				{
-					break;
-				}
-				LineIdx = NextLineIdx;	
-				NextLineIdx++;
-			}
-
-			FString MsgStr = RawMsgStr;
-			if( ProcessedEntry->MsgStr.Num() > 0 )
-			{
-				ProcessedEntry->MsgStr[0] = MsgStr;
-			}
-			else
-			{
-				ProcessedEntry->MsgStr.Add( MsgStr );
-			}
-		}
-		else
-		{
-			ProcessedEntry->UnknownElements.Add( Line );
-		}
-	}
-
-
-
-	return bSuccess;
+	return InSrcLocation.Replace(TEXT(" - line "), TEXT(":"));
 }
 
-void FPortableObjectFormatDOM::CreateNewHeader()
-{
-	// Reference: http://www.gnu.org/software/gettext/manual/gettext.html#Header-Entry
-	// Reference: http://www.gnu.org/software/gettext/manual/html_node/Header-Entry.html
-
-	//Hard code some header entries for now in the following format
-	/*
-	# Engine English translation
-	# Copyright 1998-2015 Epic Games, Inc. All Rights Reserved.
-	#
-	msgid ""
-	msgstr ""
-	"Project-Id-Version: Engine\n"
-	"POT-Creation-Date: 2014-1-31 04:16+0000\n"
-	"PO-Revision-Date: 2014-1-31 04:16+0000\n"
-	"Language-Team: \n"
-	"Language: en-us\n"
-	"MIME-Version: 1.0\n"
-	"Content-Type: text/plain; charset=UTF-8\n"
-	"Content-Transfer-Encoding: 8bit\n"
-	"Plural-Forms: nplurals=2; plural=(n != 1);\n"
-	*/
-
-	Header.Clear();
-
-	Header.SetEntryValue( TEXT("Project-Id-Version"), *GetProjectName() );
-
-	// Setup header entries
-	Header.TimeStamp();
-	Header.SetEntryValue( TEXT("Language-Team"), TEXT("") );
-	Header.SetEntryValue( TEXT("Language"), Language.GetLanguageCode() );
-	Header.SetEntryValue( TEXT("MIME-Version"), TEXT("1.0") );
-	Header.SetEntryValue( TEXT("Content-Type"), TEXT("text/plain; charset=UTF-8") );
-	Header.SetEntryValue( TEXT("Content-Transfer-Encoding"), TEXT("8bit") );
-	Header.SetEntryValue( TEXT("Plural-Forms"), Language.GetPluralForms() );
-
-	Header.Comments.Add( FString::Printf(TEXT("%s %s translation."), *GetProjectName(), *Language.EnglishName() ) );
-	Header.Comments.Add( TEXT("Copyright 1998-2015 Epic Games, Inc. All Rights Reserved.") );
-	Header.Comments.Add( FString(TEXT("")) );
-
-}
-
-bool FPortableObjectFormatDOM::SetLanguage( const FString& LanguageCode, const FString& LangPluralForms )
-{
-	FPortableObjectCulture NewLang(	LanguageCode, LangPluralForms );
-	if( NewLang.IsValid() )
-	{
-		Language = NewLang;
-		return true;
-	}
-	return false;
-}
-
-bool FPortableObjectFormatDOM::AddEntry( const TSharedRef< FPortableObjectEntry> LocEntry )
-{
-	TSharedPtr<FPortableObjectEntry> ExistingEntry = FindEntry( LocEntry );
-	if( ExistingEntry.IsValid()  )
-	{
-		ExistingEntry->AddReferences( LocEntry->ReferenceComments );
-		ExistingEntry->AddExtractedComments( LocEntry->ExtractedComments );
-
-		// Checks for a situation that we don't know how to handle yet.  ex. What happens when we match all the IDs for an entry but another param differs?
-		check( LocEntry->TranslatorComments == ExistingEntry->TranslatorComments );
-	}
-	else
-	{
-		Entries.Add( LocEntry );
-	}
-
-	return true;
-
-}
-
-
-TSharedPtr<FPortableObjectEntry> FPortableObjectFormatDOM::FindEntry( const TSharedRef<const FPortableObjectEntry> LocEntry )
-{
-	for( auto Entry : Entries )
-	{
-		if( *Entry == *LocEntry )
-		{
-			return Entry;
-		}
-	}
-	return NULL;
-}
-
-
-TSharedPtr<FPortableObjectEntry> FPortableObjectFormatDOM::FindEntry( const FString& MsgId, const FString& MsgIdPlural, const FString& MsgCtxt )
-{
-	TSharedRef<FPortableObjectEntry> TempEntry = MakeShareable( new FPortableObjectEntry );
-	TempEntry->MsgId = MsgId;
-	TempEntry->MsgIdPlural = MsgIdPlural;
-	TempEntry->MsgCtxt = MsgCtxt;
-	return FindEntry( TempEntry );
-}
-
-void FPortableObjectFormatDOM::SortEntries()
-{
-	// Sort keys.
-	for (const TSharedPtr<FPortableObjectEntry>& Entry : Entries)
-	{
-		Entry->ExtractedComments.Sort();
-	}
-
-	// Sort by namespace, then keys, then source text.
-	const auto& SortingPredicate = [](const TSharedPtr<FPortableObjectEntry>& A, const TSharedPtr<FPortableObjectEntry>& B) -> bool
-	{
-		// Compare namespace
-		if (A->MsgCtxt < B->MsgCtxt)
-		{
-			return true;
-		}
-		else if (A->MsgCtxt > B->MsgCtxt)
-		{
-			return false;
-		}
-
-		// Compare keys
-		const int32 ExtractedCommentCount = FMath::Max(A->ExtractedComments.Num(), B->ExtractedComments.Num());
-		for (int32 i = 0; i < ExtractedCommentCount; ++i)
-		{
-			if (A->ExtractedComments[i] < B->ExtractedComments[i] || (!A->ExtractedComments.IsValidIndex(i) && B->ExtractedComments.IsValidIndex(i)))
-			{
-				return true;
-			}
-			else if (A->ExtractedComments[i] > B->ExtractedComments[i] || (A->ExtractedComments.IsValidIndex(i) && !B->ExtractedComments.IsValidIndex(i)))
-			{
-				return false;
-			}
-		}
-
-		// Compare source string
-		if (A->MsgId < B->MsgId)
-		{
-			return true;
-		}
-		else if (A->MsgId > B->MsgId)
-		{
-			return false;
-		}
-
-		return A.Get() < B.Get();
-	};
-	Entries.Sort(SortingPredicate);
-}
-
-void FPortableObjectEntry::AddExtractedComment( const FString& InComment )
-{
-	if(!InComment.IsEmpty())
-	{
-		ExtractedComments.AddUnique(InComment);
-	}
-
-	//// Extracted comments can contain multiple references in a single line so we parse those out.
-	//TArray<FString> CommentsToProcess;
-	//InComment.ParseIntoArray( CommentsToProcess, TEXT(" "), true );
-	//for( const FString& ExtractedComment : CommentsToProcess )
-	//{
-	//	ExtractedComments.AddUnique( ExtractedComment );
-	//}
-}
-
-void FPortableObjectEntry::AddReference( const FString& InReference )
-{
-	if(!InReference.IsEmpty())
-	{
-		ReferenceComments.AddUnique(InReference);
-	}
-	
-	//// Reference comments can contain multiple references in a single line so we parse those out.
-	//TArray<FString> ReferencesToProcess;
-	//InReference.ParseIntoArray( ReferencesToProcess, TEXT(" "), true );
-	//for( const FString& Reference : ReferencesToProcess )
-	//{
-	//	ReferenceComments.AddUnique( Reference );
-	//}
-}
-
-void FPortableObjectEntry::AddExtractedComments( const TArray<FString>& InComments )
-{
-	for( const FString& ExtractedComment : InComments )
-	{
-		AddExtractedComment( ExtractedComment );
-	}
-}
-
-void FPortableObjectEntry::AddReferences( const TArray<FString>& InReferences )
-{
-	for( const FString& Reference : InReferences )
-	{
-		AddReference( Reference );
-	}
-}
-
-FString FPortableObjectEntry::ToString() const
-{
-	check( !MsgId.IsEmpty() );
-
-	FString Result;
-
-	for( const FString& TranslatorComment : TranslatorComments )
-	{
-		Result += FString::Printf( TEXT("# %s%s"), *TranslatorComment, NewLineDelimiter );
-	}
-
-	for( const FString& Comment : ExtractedComments )
-	{
-		if( !Comment.IsEmpty() )
-		{
-			Result += FString::Printf( TEXT("#. %s%s"), *Comment, NewLineDelimiter );
-		}
-		else
-		{
-			Result += FString::Printf( TEXT("#.%s"), NewLineDelimiter );
-		}
-	}
-
-	for( const FString& ReferenceComment : ReferenceComments )
-	{
-		Result += FString::Printf( TEXT("#: %s%s"), *ReferenceComment, NewLineDelimiter );
-	}
-
-	if( !MsgCtxt.IsEmpty() )
-	{
-		Result += FString::Printf(TEXT("msgctxt \"%s\"%s"), *MsgCtxt, NewLineDelimiter);
-	}
-
-	Result += FString::Printf(TEXT("msgid \"%s\"%s"), *MsgId, NewLineDelimiter);
-
-	if( MsgStr.Num() == 0)
-	{
-		Result += FString::Printf(TEXT("msgstr \"\"%s"), NewLineDelimiter);
-	}
-	else if( MsgStr.Num() == 1 )
-	{
-		Result += FString::Printf(TEXT("msgstr \"%s\"%s"), *MsgStr[0], NewLineDelimiter);
-	}
-	else
-	{
-		for( int32 Idx = 0; Idx < MsgStr.Num(); ++Idx )
-		{
-			Result += FString::Printf(TEXT("msgstr[%d] \"%s\"%s"), Idx, *MsgStr[Idx], NewLineDelimiter);
-		}
-	}
-	return Result;
-}
-
-
-
 /**
  *	UInternationalizationExportCommandlet
  */
diff --git a/Engine/Source/Editor/UnrealEd/Private/Localization/PortableObjectFormatDOM.cpp b/Engine/Source/Editor/UnrealEd/Private/Localization/PortableObjectFormatDOM.cpp
new file mode 100644
--- /dev/null
+++ b/Engine/Source/Editor/UnrealEd/Private/Localization/PortableObjectFormatDOM.cpp
@@ -0,0 +1,959 @@
+// Copyright 1998-2015 Epic Games, Inc. All Rights Reserved.
+
+#include "UnrealEd.h"
+#include "Culture.h"
+#include "PortableObjectFormatDOM.h"
+
+static const TCHAR* NewLineDelimiter = TEXT("\n");
+
+/* Default culture plural rules.  Culture names are in the following format: Language_Country@Variant
+  See references:	http://www.unicode.org/cldr/charts/latest/supplemental/language_plural_rules.html  
+					http://docs.translatehouse.org/projects/localization-guide/en/latest/l10n/pluralforms.html
+*/
+static const TMap< FString, FString > POCulturePluralForms = TMapBuilder< FString, FString >()
+	.Add( TEXT( "ach" )	, TEXT( "nplurals=2; plural=(n > 1);" )  )
+	.Add( TEXT( "af" )	, TEXT( "nplurals=2; plural=(n != 1);" )  )
+	.Add( TEXT( "ak" )	, TEXT( "nplurals=2; plural=(n > 1);" )  )
+	.Add( TEXT( "aln" )	, TEXT( "nplurals=2; plural=(n != 1);" )  )
+	.Add( TEXT( "am" )	, TEXT( "nplurals=2; plural=(n > 1);" )  )
+	.Add( TEXT( "am_ET"), TEXT( "nplurals=2; plural=(n != 1);" )  )
+	.Add( TEXT( "an" )	, TEXT( "nplurals=2; plural=(n != 1);" )  )
+	.Add( TEXT( "ar" )	, TEXT( "nplurals=6; plural=(n==0 ? 0 : n==1 ? 1 : n==2 ? 2 : n%100>=3 && n%100<=10 ? 3 : n%100>=11 && n%100<=99 ? 4 : 5);" )  )
+	.Add( TEXT( "ar_SA"), TEXT( "nplurals=2; plural=(n != 1);" )  )
+	.Add( TEXT( "arn" )	, TEXT( "nplurals=2; plural=(n > 1);" )  )
+	.Add( TEXT( "as" )	, TEXT( "nplurals=2; plural=(n != 1);" )  )
+	.Add( TEXT( "ast" )	, TEXT( "nplurals=2; plural=(n != 1);" )  )
+	.Add( TEXT( "ay" )	, TEXT( "nplurals=1; plural=0;" )  )
+	.Add( TEXT( "az" )	, TEXT( "nplurals=2; plural=(n != 1);" )  )
+	.Add( TEXT( "bal" )	, TEXT( "nplurals=2; plural=(n != 1);" )  )
+	.Add( TEXT( "be" )	, TEXT( "nplurals=3; plural=(n%10==1 && n%100!=11 ? 0 : n%10>=2 && n%10<=4 && (n%100<10 || n%100>=20) ? 1 : 2);" )  )
+	.Add( TEXT( "bg" )	, TEXT( "nplurals=2; plural=(n != 1);" )  )
+	.Add( TEXT( "bn" )	, TEXT( "nplurals=2; plural=(n != 1);" )  )
+	.Add( TEXT( "bo" )	, TEXT( "nplurals=1; plural=0;" )  )
+	.Add( TEXT( "br" )	, TEXT( "nplurals=2; plural=(n > 1);" )  )
+	.Add( TEXT( "brx" )	, TEXT( "nplurals=2; plural=(n != 1);" )  )
+	.Add( TEXT( "bs" )	, TEXT( "nplurals=3; plural=(n%10==1 && n%100!=11 ? 0 : n%10>=2 && n%10<=4 && (n%100<10 || n%100>=20) ? 1 : 2);" )  )
+	.Add( TEXT( "ca" )	, TEXT( "nplurals=2; plural=(n != 1);" )  )
+	.Add( TEXT( "cgg" )	, TEXT( "nplurals=1; plural=0;" )  )
+	.Add( TEXT( "crh" )	, TEXT( "nplurals=1; plural=0;" )  )
+	.Add( TEXT( "cs" )	, TEXT( "nplurals=3; plural=(n==1) ? 0 : (n>=2 && n<=4) ? 1 : 2;" )  )
+	.Add( TEXT( "csb" )	, TEXT( "nplurals=3; plural=(n==1 ? 0 : n%10>=2 && n%10<=4 && (n%100<10 || n%100>=20) ? 1 : 2;" )  )
+	.Add( TEXT( "cy" )	, TEXT( "nplurals=4; plural=(n==1) ? 0 : (n==2) ? 1 : (n != 8 && n != 11) ? 2 : 3;" )  )
+	.Add( TEXT( "da" )	, TEXT( "nplurals=2; plural=(n != 1);" )  )
+	.Add( TEXT( "de" )	, TEXT( "nplurals=2; plural=(n != 1);" )  )
+	.Add( TEXT( "doi" )	, TEXT( "nplurals=2; plural=(n != 1);" )  )
+	.Add( TEXT( "dz" )	, TEXT( "nplurals=1; plural=0;" )  )
+	.Add( TEXT( "el" )	, TEXT( "nplurals=2; plural=(n != 1);" )  )
+	.Add( TEXT( "en" )	, TEXT( "nplurals=2; plural=(n != 1);" )  )
+	.Add( TEXT( "eo" )	, TEXT( "nplurals=2; plural=(n != 1);" )  )
+	.Add( TEXT( "es" )	, TEXT( "nplurals=2; plural=(n != 1);" )  )
+	.Add( TEXT( "es_ar"), TEXT( "nplurals=2; plural=(n != 1);" )  )
+	.Add( TEXT( "et" )	, TEXT( "nplurals=2; plural=(n != 1);" )  )
+	.Add( TEXT( "eu" )	, TEXT( "nplurals=2; plural=(n != 1);" )  )
+	.Add( TEXT( "fa" )	, TEXT( "nplurals=1; plural=0;" )  )
+	.Add( TEXT( "fi" )	, TEXT( "nplurals=2; plural=(n != 1);" )  )
+	.Add( TEXT( "fil" )	, TEXT( "nplurals=2; plural=(n > 1);" )  )
+	.Add( TEXT( "fo" )	, TEXT( "nplurals=2; plural=(n != 1);" )  )
+	.Add( TEXT( "fr" )	, TEXT( "nplurals=2; plural=(n > 1);" )  )
+	.Add( TEXT( "frp" )	, TEXT( "nplurals=2; plural=(n > 1);" )  )
+	.Add( TEXT( "fur" )	, TEXT( "nplurals=2; plural=(n != 1);" )  )
+	.Add( TEXT( "fy" )	, TEXT( "nplurals=2; plural=(n != 1);" )  )
+	.Add( TEXT( "ga" )	, TEXT( "nplurals=5; plural=(n==1 ? 0 : n==2 ? 1 : n<7 ? 2 : n<11 ? 3 : 4);" )  )
+	.Add( TEXT( "gd" )	, TEXT( "nplurals=3; plural=(n < 2 ? 0 : n == 2 ? 1 : 2);" )  )
+	.Add( TEXT( "gl" )	, TEXT( "nplurals=2; plural=(n != 1);" )  )
+	.Add( TEXT( "gu" )	, TEXT( "nplurals=2; plural=(n != 1);" )  )
+	.Add( TEXT( "gun" )	, TEXT( "nplurals=2; plural=(n > 1);" )  )
+	.Add( TEXT( "ha" )	, TEXT( "nplurals=2; plural=(n != 1);" )  )
+	.Add( TEXT( "he" )	, TEXT( "nplurals=2; plural=(n != 1);" )  )
+	.Add( TEXT( "hi" )	, TEXT( "nplurals=2; plural=(n != 1);" )  )
+	.Add( TEXT( "hne" )	, TEXT( "nplurals=2; plural=(n != 1);" )  )
+	.Add( TEXT( "hr" )	, TEXT( "nplurals=3; plural=(n%10==1 && n%100!=11 ? 0 : n%10>=2 && n%10<=4 && (n%100<10 || n%100>=20) ? 1 : 2);" )  )
+	.Add( TEXT( "hsb" )	, TEXT( "nplurals=4; plural=(n%100==1 ? 0 : n%100==2 ? 1 : n%100==3 || n%100==4 ? 2 : 3);" )  )
+	.Add( TEXT( "ht" )	, TEXT( "nplurals=2; plural=(n != 1);" )  )
+	.Add( TEXT( "hu" )	, TEXT( "nplurals=2; plural=(n != 1);" )  )
+	.Add( TEXT( "hy" )	, TEXT( "nplurals=2; plural=(n != 1);" )  )
+	.Add( TEXT( "ia" )	, TEXT( "nplurals=2; plural=(n != 1);" )  )
+	.Add( TEXT( "id" )	, TEXT( "nplurals=1; plural=0;" )  )
+	.Add( TEXT( "ig" )	, TEXT( "nplurals=2; plural=(n != 1);" )  )
+	.Add( TEXT( "ilo" )	, TEXT( "nplurals=2; plural=(n != 1);" )  )
+	.Add( TEXT( "is" )	, TEXT( "nplurals=2; plural=(n != 1);" )  )
+	.Add( TEXT( "it" )	, TEXT( "nplurals=2; plural=(n != 1);" )  )
+	.Add( TEXT( "ja" )	, TEXT( "nplurals=1; plural=0;" )  )
+	.Add( TEXT( "jv" )	, TEXT( "nplurals=2; plural=(n != 1);" )  )
+	.Add( TEXT( "ka" )	, TEXT( "nplurals=1; plural=0;" )  )
+	.Add( TEXT( "kk" )	, TEXT( "nplurals=1; plural=0;" )  )
+	.Add( TEXT( "km" )	, TEXT( "nplurals=1; plural=0;" )  )
+	.Add( TEXT( "kn" )	, TEXT( "nplurals=1; plural=0;" )  )
+	.Add( TEXT( "ko" )	, TEXT( "nplurals=1; plural=0;" )  )
+	.Add( TEXT( "ks" )	, TEXT( "nplurals=2; plural=(n != 1);" )  )
+	.Add( TEXT( "ku" )	, TEXT( "nplurals=2; plural=(n != 1);" )  )
+	.Add( TEXT( "kw" )	, TEXT( "nplurals=4; plural=(n==1) ? 0 : (n==2) ? 1 : (n == 3) ? 2 : 3;" )  )
+	.Add( TEXT( "ky" )	, TEXT( "nplurals=1; plural=0;" )  )
+	.Add( TEXT( "la" )	, TEXT( "nplurals=2; plural=(n != 1);" )  )
+	.Add( TEXT( "lb" )	, TEXT( "nplurals=2; plural=(n != 1);" )  )
+	.Add( TEXT( "li" )	, TEXT( "nplurals=2; plural=(n != 1);" )  )
+	.Add( TEXT( "ln" )	, TEXT( "nplurals=2; plural=(n > 1);" )  )
+	.Add( TEXT( "lo" )	, TEXT( "nplurals=1; plural=0;" )  )
+	.Add( TEXT( "lt" )	, TEXT( "nplurals=3; plural=(n%10==1 && n%100!=11 ? 0 : n%10>=2 && (n%100<10 || n%100>=20) ? 1 : 2);" )  )
+	.Add( TEXT( "lv" )	, TEXT( "nplurals=3; plural=(n%10==1 && n%100!=11 ? 0 : n != 0 ? 1 : 2);" )  )
+	.Add( TEXT( "mai" )	, TEXT( "nplurals=2; plural=(n != 1);" )  )
+	.Add( TEXT( "mg" )	, TEXT( "nplurals=2; plural=(n > 1);" )  )
+	.Add( TEXT( "mi" )	, TEXT( "nplurals=2; plural=(n > 1);" )  )
+	.Add( TEXT( "mk" )	, TEXT( "nplurals=2; plural=(n % 10 == 1 && n % 100 != 11) ? 0 : 1;" )  )
+	.Add( TEXT( "ml" )	, TEXT( "nplurals=2; plural=(n != 1);" )  )
+	.Add( TEXT( "mn" )	, TEXT( "nplurals=2; plural=(n != 1);" )  )
+	.Add( TEXT( "mr" )	, TEXT( "nplurals=2; plural=(n != 1);" )  )
+	.Add( TEXT( "ms" )	, TEXT( "nplurals=1; plural=0;" )  )
+	.Add( TEXT( "mt" )	, TEXT( "nplurals=4; plural=(n==1 ? 0 : n==0 || ( n%100>1 && n%100<11) ? 1 : (n%100>10 && n%100<20 ) ? 2 : 3);" )  )
+	.Add( TEXT( "my" )	, TEXT( "nplurals=1; plural=0;" )  )
+	.Add( TEXT( "nah" )	, TEXT( "nplurals=2; plural=(n != 1);" )  )
+	.Add( TEXT( "nap" )	, TEXT( "nplurals=2; plural=(n != 1);" )  )
+	.Add( TEXT( "nb" )	, TEXT( "nplurals=2; plural=(n != 1);" )  )
+	.Add( TEXT( "nds" )	, TEXT( "nplurals=2; plural=(n != 1);" )  )
+	.Add( TEXT( "ne" )	, TEXT( "nplurals=2; plural=(n != 1);" )  )
+	.Add( TEXT( "nl" )	, TEXT( "nplurals=2; plural=(n != 1);" )  )
+	.Add( TEXT( "nn" )	, TEXT( "nplurals=2; plural=(n != 1);" )  )
+	.Add( TEXT( "no" )	, TEXT( "nplurals=2; plural=(n != 1);" )  )
+	.Add( TEXT( "nr" )	, TEXT( "nplurals=1; plural=0;" )  )
+	.Add( TEXT( "nso" )	, TEXT( "nplurals=2; plural=(n > 1);" )  )
+	.Add( TEXT( "oc" )	, TEXT( "nplurals=2; plural=(n > 1);" )  )
+	.Add( TEXT( "or" )	, TEXT( "nplurals=2; plural=(n != 1);" )  )
+	.Add( TEXT( "pa" )	, TEXT( "nplurals=2; plural=(n != 1);" )  )
+	.Add( TEXT( "pap" )	, TEXT( "nplurals=2; plural=(n != 1);" )  )
+	.Add( TEXT( "pl" )	, TEXT( "nplurals=3; plural=(n==1 ? 0 : n%10>=2 && n%10<=4 && (n%100<10 || n%100>=20) ? 1 : 2);" )  )
+	.Add( TEXT( "pms" )	, TEXT( "nplurals=2; plural=(n != 1);" )  )
+	.Add( TEXT( "ps" )	, TEXT( "nplurals=2; plural=(n != 1);" )  )
+	.Add( TEXT( "pt" )	, TEXT( "nplurals=2; plural=(n != 1);" )  )
+	.Add( TEXT( "pt_BR"), TEXT( "nplurals=2; plural=(n > 1);" )  )
+	.Add( TEXT( "rm" )	, TEXT( "nplurals=2; plural=(n != 1);" )  )
+	.Add( TEXT( "ro" )	, TEXT( "nplurals=3; plural=(n==1?0:(((n%100>19)||((n%100==0)&&(n!=0)))?2:1));" )  )
+	.Add( TEXT( "ru" )	, TEXT( "nplurals=3; plural=(n%10==1 && n%100!=11 ? 0 : n%10>=2 && n%10<=4 && (n%100<10 || n%100>=20) ? 1 : 2);" )  )
+	.Add( TEXT( "rw" )	, TEXT( "nplurals=2; plural=(n != 1);" )  )
+	.Add( TEXT( "sc" )	, TEXT( "nplurals=2; plural=(n != 1);" )  )
+	.Add( TEXT( "sco" )	, TEXT( "nplurals=2; plural=(n != 1);" )  )
+	.Add( TEXT( "se" )	, TEXT( "nplurals=2; plural=(n != 1);" )  )
+	.Add( TEXT( "si" )	, TEXT( "nplurals=2; plural=(n != 1);" )  )
+	.Add( TEXT( "sk" )	, TEXT( "nplurals=3; plural=(n==1) ? 0 : (n>=2 && n<=4) ? 1 : 2;" )  )
+	.Add( TEXT( "sl" )	, TEXT( "nplurals=4; plural=(n%100==1 ? 0 : n%100==2 ? 1 : n%100==3 || n%100==4 ? 2 : 3);" )  )
+	.Add( TEXT( "sm" )	, TEXT( "nplurals=1; plural=0;" )  )
+	.Add( TEXT( "sn" )	, TEXT( "nplurals=2; plural=(n != 1);" )  )
+	.Add( TEXT( "so" )	, TEXT( "nplurals=2; plural=(n != 1);" )  )
+	.Add( TEXT( "son" )	, TEXT( "nplurals=1; plural=0;" )  )
+	.Add( TEXT( "sq" )	, TEXT( "nplurals=2; plural=(n != 1);" )  )
+	.Add( TEXT( "sr" )	, TEXT( "nplurals=3; plural=(n%10==1 && n%100!=11 ? 0 : n%10>=2 && n%10<=4 && (n%100<10 || n%100>=20) ? 1 : 2);" )  )
+	.Add( TEXT( "st" )	, TEXT( "nplurals=2; plural=(n != 1);" )  )
+	.Add( TEXT( "su" )	, TEXT( "nplurals=1; plural=0;" )  )
+	.Add( TEXT( "sv" )	, TEXT( "nplurals=2; plural=(n != 1);" )  )
+	.Add( TEXT( "sw" )	, TEXT( "nplurals=2; plural=(n != 1);" )  )
+	.Add( TEXT( "ta" )	, TEXT( "nplurals=2; plural=(n != 1);" )  )
+	.Add( TEXT( "te" )	, TEXT( "nplurals=2; plural=(n != 1);" )  )
+	.Add( TEXT( "tg" )	, TEXT( "nplurals=2; plural=(n != 1);" )  )
+	.Add( TEXT( "th" )	, TEXT( "nplurals=1; plural=0;" )  )
+	.Add( TEXT( "ti" )	, TEXT( "nplurals=2; plural=(n > 1);" )  )
+	.Add( TEXT( "tk" )	, TEXT( "nplurals=2; plural=(n != 1);" )  )
+	.Add( TEXT( "tl" )	, TEXT( "nplurals=2; plural=(n > 1);" )  )
+	.Add( TEXT( "tlh" )	, TEXT( "nplurals=2; plural=(n > 1);" )  )
+	.Add( TEXT( "to" )	, TEXT( "nplurals=1; plural=0;" )  )
+	.Add( TEXT( "tr" )	, TEXT( "nplurals=1; plural=0;" )  )
+	.Add( TEXT( "tt" )	, TEXT( "nplurals=1; plural=0;" )  )
+	.Add( TEXT( "udm" )	, TEXT( "nplurals=1; plural=0;" )  )
+	.Add( TEXT( "ug" )	, TEXT( "nplurals=1; plural=0;" )  )
+	.Add( TEXT( "uk" )	, TEXT( "nplurals=3; plural=(n%10==1 && n%100!=11 ? 0 : n%10>=2 && n%10<=4 && (n%100<10 || n%100>=20) ? 1 : 2);" )  )
+	.Add( TEXT( "ur" )	, TEXT( "nplurals=2; plural=(n != 1);" )  )
+	.Add( TEXT( "uz" )	, TEXT( "nplurals=1; plural=0;" )  )
+	.Add( TEXT( "ve" )	, TEXT( "nplurals=2; plural=(n != 1);" )  )
+	.Add( TEXT( "vi" )	, TEXT( "nplurals=1; plural=0;" )  )
+	.Add( TEXT( "vls" )	, TEXT( "nplurals=2; plural=(n != 1);" )  )
+	.Add( TEXT( "wa" )	, TEXT( "nplurals=2; plural=(n > 1);" )  )
+	.Add( TEXT( "wo" )	, TEXT( "nplurals=1; plural=0;" )  )
+	.Add( TEXT( "xh" )	, TEXT( "nplurals=2; plural=(n != 1);" )  )
+	.Add( TEXT( "yi" )	, TEXT( "nplurals=2; plural=(n != 1);" )  )
+	.Add( TEXT( "yo" )	, TEXT( "nplurals=2; plural=(n != 1);" )  )
+	.Add( TEXT( "zh" )	, TEXT( "nplurals=1; plural=0;" )  )
+	.Add( TEXT( "zu" )	, TEXT( "nplurals=2; plural=(n != 1);" )  );
+
+bool FindDelimitedString(const FString& InStr, const FString& LeftDelim, const FString& RightDelim, FString& Result)
+{
+	int32 Start = InStr.Find(LeftDelim, ESearchCase::CaseSensitive);
+	int32 End = InStr.Find(RightDelim, ESearchCase::CaseSensitive, ESearchDir::FromEnd);
+	if (Start == INDEX_NONE || !(End > Start))
+	{
+		return false;
+	}
+	Start += LeftDelim.Len();
+	if (Start >= End)
+	{
+		Result = TEXT("");
+	}
+	else
+	{
+		Result = InStr.Mid(Start, End - Start);
+	}
+
+	return true;
+}
+
+FPortableObjectCulture::FPortableObjectCulture( const FString& LangCode, const FString& PluralForms )
+	: LanguageCode( LangCode )
+	, LanguagePluralForms( PluralForms )
+	, Culture( FCulture::Create( LangCode ) )
+{
+
+}
+
+FPortableObjectCulture::FPortableObjectCulture( const FPortableObjectCulture& Other )
+	: LanguageCode( Other.LanguageCode )
+	, LanguagePluralForms( Other.LanguagePluralForms )
+	, Culture( FCulture::Create( Other.LanguageCode ) )
+{
+
+}
+
+void FPortableObjectCulture::SetLanguageCode( const FString& LangCode )
+{
+	LanguageCode = LangCode;
+	Culture = FCulture::Create( LangCode );
+}
+
+
+FString FPortableObjectCulture::Language() const
+{
+	FString Result;
+	if( Culture.IsValid() )
+	{
+		// NOTE: The below function name may be a bit misleading because it will return three letter language codes if needed.
+		Result = Culture->GetTwoLetterISOLanguageName();
+	}
+	return Result;
+
+}
+
+FString FPortableObjectCulture::Country() const
+{
+	FString Result;
+	if( Culture.IsValid() )
+	{
+		Result = Culture->GetRegion();
+	}
+	return Result;
+}
+
+FString FPortableObjectCulture::Variant() const
+{
+	FString Result;
+	if( Culture.IsValid() )
+	{
+		Result = Culture->GetVariant();
+	}
+	return Result;
+}
+
+FString FPortableObjectCulture::DisplayName() const
+{
+	FString Result;
+	if( Culture.IsValid() )
+	{
+		Result = Culture->GetDisplayName();
+	}
+	return Result;
+}
+
+FString FPortableObjectCulture::EnglishName() const
+{
+	FString Result;
+	if( Culture.IsValid() )
+	{
+		Result = Culture->GetEnglishName();
+	}
+	return Result;
+}
+
+FString FPortableObjectCulture::GetPluralForms() const
+{
+	if( !LanguagePluralForms.IsEmpty() )
+	{
+		return LanguagePluralForms;
+	}
+	return GetDefaultPluralForms();
+}
+
+FString FPortableObjectCulture::GetDefaultPluralForms() const
+{
+	FString Result;
+	if( LanguageCode.IsEmpty() )
+	{
+		return Result;
+	}
+
+	if( auto* LanguageCodePair = POCulturePluralForms.Find( GetLanguageCode() ) )
+	{
+		Result = *LanguageCodePair;
+	}
+	else if( auto* LangCountryVariantPair = POCulturePluralForms.Find( Language() + TEXT("_") + Country() + TEXT("@") + Variant() ))
+	{
+		Result = *LangCountryVariantPair;
+	}
+	else if( auto* LangCountryPair = POCulturePluralForms.Find( Language() + TEXT("_") + Country() ) )
+	{
+		Result  = *LangCountryPair;
+	} 
+	else if( auto* LangPair = POCulturePluralForms.Find( Language() ) )
+	{
+		Result = *LangPair;
+	}
+	else
+	{
+		auto* Fallback = POCulturePluralForms.Find( TEXT("en") );
+		if( Fallback )
+		{
+			Result = *Fallback;
+		}
+		else
+		{
+			Result = TEXT("nplurals=2; plural=(n != 1);");
+		}
+	}
+
+	return Result;
+}
+
+
+FString FPortableObjectHeader::ToString() const
+{
+	FString Result;
+	for( const FString& Comment : Comments )
+	{
+		Result += FString::Printf(TEXT("# %s%s"), *Comment, NewLineDelimiter);
+	}
+
+	Result += FString::Printf(TEXT("msgid \"\"%s"), NewLineDelimiter);
+	Result += FString::Printf(TEXT("msgstr \"\"%s"), NewLineDelimiter);
+
+	for( auto Entry : HeaderEntries )
+	{
+		const FString& Key = Entry.Key;
+		const FString& Value = Entry.Value;
+		Result += FString::Printf( TEXT("\"%s: %s\\n\"%s"), *Key, *Value, NewLineDelimiter );
+	}
+
+	return Result;
+}
+
+bool FPortableObjectHeader::FromLocPOEntry( const TSharedRef<const FPortableObjectEntry> LocEntry )
+{
+	if( !LocEntry->MsgId.IsEmpty() || LocEntry->MsgStr.Num() != 1 )
+	{
+		return false;
+	}
+	Clear();
+
+	Comments = LocEntry->TranslatorComments;
+
+	// The POEntry would store header info inside the MsgStr[0]
+	TArray<FString> HeaderLinesToProcess;
+	LocEntry->MsgStr[0].ReplaceEscapedCharWithChar().ParseIntoArray( HeaderLinesToProcess, NewLineDelimiter, true );
+
+	for( const FString& PotentialHeaderEntry : HeaderLinesToProcess )
+	{
+		int32 SplitIndex;
+		if( PotentialHeaderEntry.FindChar( TCHAR(':'), SplitIndex ) )
+		{
+			// Looks like a header entry so we add it
+			const FString& Key = PotentialHeaderEntry.LeftChop( PotentialHeaderEntry.Len() - SplitIndex ).Trim().TrimTrailing();
+			FString Value = PotentialHeaderEntry.RightChop( SplitIndex+1 ).Trim().TrimTrailing();
+
+			HeaderEntries.Add( TPairInitializer<FString, FString>( Key, Value ) );
+		}
+	}
+	return true;
+}
+
+FPortableObjectHeader::FPOHeaderEntry* FPortableObjectHeader::FindEntry( const FString& EntryKey )
+{
+	for( auto& Entry : HeaderEntries )
+	{
+		if( Entry.Key == EntryKey )
+		{
+			return &Entry;
+		}
+	}
+	return NULL;
+}
+
+FString FPortableObjectHeader::GetEntryValue( const FString& EntryKey )
+{
+	FString Result;
+	const FPOHeaderEntry* Entry = FindEntry( EntryKey );
+	if( Entry )
+	{
+		Result = Entry->Value;
+	}
+	return Result;
+}
+
+bool FPortableObjectHeader::HasEntry( const FString& EntryKey )
+{
+	return ( FindEntry( EntryKey ) != NULL);
+}
+
+void FPortableObjectHeader::SetEntryValue( const FString& EntryKey, const FString& EntryValue )
+{
+	FPOHeaderEntry* Entry = FindEntry( EntryKey );
+	if( Entry )
+	{
+		Entry->Value = EntryValue;
+	}
+	else
+	{
+		HeaderEntries.Add( TPairInitializer<FString, FString>( EntryKey, EntryValue ) );
+	}
+}
+
+void FPortableObjectHeader::TimeStamp()
+{
+	// @TODO: Time format not exactly correct.  We have something like this: 2014-02-07 20:06  This is what it should be: 2014-02-07 14:12-0600
+	FString Time = *FDateTime::UtcNow().ToString(TEXT("%Y-%m-%d %H:%M"));
+
+	SetEntryValue( TEXT("POT-Creation-Date"), Time );
+	SetEntryValue( TEXT("PO-Revision-Date"), Time );
+}
+
+FString FPortableObjectFormatDOM::ToString()
+{
+	FString Result;
+
+	Header.TimeStamp();
+
+	Result += Header.ToString();
+	Result += NewLineDelimiter;
+
+	for( auto Entry : Entries )
+	{
+		Result += Entry->ToString();
+		Result += NewLineDelimiter;
+	}
+
+	return Result;
+}
+
+
+
+bool FPortableObjectFormatDOM::FromString( const FString& InStr )
+{
+	if( InStr.IsEmpty() )
+	{
+		return false;
+	}
+
+	bool bSuccess = true;
+
+	FString ParseString = InStr.Replace(TEXT("\r\n"), NewLineDelimiter);
+
+	TArray<FString> LinesToProcess;
+	ParseString.ParseIntoArray( LinesToProcess, NewLineDelimiter, false );
+
+	TSharedRef<FPortableObjectEntry> ProcessedEntry = MakeShareable( new FPortableObjectEntry );
+	bool bHasMsgId = false;
+
+	uint32 NumFileLines = LinesToProcess.Num();
+	for( uint32 LineIdx = 0; LineIdx < NumFileLines; ++LineIdx )
+	{
+		const FString& Line = LinesToProcess[LineIdx];
+
+		if( Line.IsEmpty() )
+		{
+			// When we encounter a blank line, we will either ignore it, or consider it the boundary of 
+			// a LocPOEntry or FPortableObjectHeader if we processed any valid data before encountering the blank.
+
+			// If this entry is valid we'll check it and process it further.
+			if( bHasMsgId && ProcessedEntry->MsgStr.Num() > 0 )
+			{
+				// Check if we are dealing with a header entry
+				if( ProcessedEntry->MsgId.IsEmpty() && ProcessedEntry->MsgStr.Num() == 1 )
+				{
+					// This is a header
+					bool bAddedHeader = Header.FromLocPOEntry( ProcessedEntry );
+					if( !bAddedHeader )
+					{
+						return false;
+					}
+					ProjectName = Header.GetEntryValue(TEXT("Project-Id-Version"));
+				}
+				else
+				{
+					bool bAddEntry = AddEntry( ProcessedEntry );
+					if( !bAddEntry )
+					{
+						return false;
+					}
+				}
+			}
+
+			// Now starting a new entry and reset any flags we have been tracking
+			ProcessedEntry = MakeShareable( new FPortableObjectEntry );
+			bHasMsgId = false;
+		}
+		else if( Line.StartsWith( TEXT("#,") ) )
+		{
+			// Flags
+			FString Flag;
+			if( FindDelimitedString(Line, TEXT("#,"), NewLineDelimiter, Flag ) )
+			{
+				if( !Flag.IsEmpty() )
+				{
+					ProcessedEntry->Flags.Add( Flag );
+				}
+			}
+		}
+		else if( Line.StartsWith( TEXT("#.") ) )
+		{
+			// Extracted comments
+			FString Comment;
+			ProcessedEntry->ExtractedComments.Add( Line.RightChop( FString(TEXT("#. ")).Len() ) );
+		}
+		else if( Line.StartsWith( TEXT("#:") ) )
+		{
+			// Reference
+			FString Reference;
+			ProcessedEntry->AddReference( Line.RightChop( FString(TEXT("#: ")).Len() ) );
+		}
+		else if( Line.StartsWith( TEXT(":|") ) )
+		{
+			// This represents previous messages. We just drop this in unknown since we don't handle it
+			ProcessedEntry->UnknownElements.Add( Line );
+		}
+		else if( Line.StartsWith( TEXT("# ") ) || Line.StartsWith( TEXT("#\t") ) )
+		{
+			FString Comment;
+			ProcessedEntry->TranslatorComments.Add( Line.RightChop( FString(TEXT("# ")).Len() ) );
+		}
+		else if( Line == TEXT("#") )
+		{
+			ProcessedEntry->TranslatorComments.Add( TEXT("") );		
+		}
+		else if( Line.StartsWith( TEXT("msgctxt") ) )
+		{
+			FString RawMsgCtxt;
+			if( !FindDelimitedString(Line, TEXT("\""), TEXT("\""), RawMsgCtxt ) )
+			{
+				bSuccess = false;
+				break;
+			}
+			for( uint32 NextLineIdx = LineIdx + 1; NextLineIdx < NumFileLines && LinesToProcess[NextLineIdx].Trim().TrimTrailing().StartsWith(TEXT("\"")); ++NextLineIdx)
+			{
+				FString Tmp;
+				if (FindDelimitedString(Line, TEXT("\""), TEXT("\""), Tmp))
+				{
+					RawMsgCtxt += Tmp;
+				}
+			}
+
+			// If the following line contains more info for this element we continue to parse it out
+			uint32 NextLineIdx = LineIdx + 1;
+			while( NextLineIdx < NumFileLines )
+			{
+				const FString& NextLine = LinesToProcess[NextLineIdx].TrimTrailing().Trim();
+				if( NextLine.StartsWith("\"") && NextLine.EndsWith("\"") )
+				{
+					RawMsgCtxt += NextLine.Mid( 1, NextLine.Len()-2 );
+				}
+				else
+				{
+					break;
+				}
+				LineIdx = NextLineIdx;	
+				NextLineIdx++;
+			}
+
+			ProcessedEntry->MsgCtxt = RawMsgCtxt;
+		}
+		else if( Line.StartsWith( TEXT("msgid") ) )
+		{
+			FString RawMsgId;
+			if( !FindDelimitedString(Line, TEXT("\""), TEXT("\""), RawMsgId ) )
+			{
+				bSuccess = false;
+				break;
+			}
+			// If the following line contains more info for this element we continue to parse it out
+			uint32 NextLineIdx = LineIdx + 1;
+			while( NextLineIdx < NumFileLines )
+			{
+				const FString& NextLine = LinesToProcess[NextLineIdx].TrimTrailing().Trim();
+				if( NextLine.StartsWith("\"") && NextLine.EndsWith("\"") )
+				{
+					RawMsgId += NextLine.Mid( 1, NextLine.Len()-2 );
+				}
+				else
+				{
+					break;
+				}
+				LineIdx = NextLineIdx;	
+				NextLineIdx++;
+			}
+			ProcessedEntry->MsgId = RawMsgId;
+			bHasMsgId = true;
+		}
+		else if( Line.StartsWith( TEXT("msgid_plural") ) )
+		{
+			FString RawMsgIdPlural;
+			if( !FindDelimitedString(Line, TEXT("\""), TEXT("\""), RawMsgIdPlural ) )
+			{
+				bSuccess = false;
+				break;
+			}
+			// If the following line contains more info for this element we continue to parse it out
+			uint32 NextLineIdx = LineIdx + 1;
+			while( NextLineIdx < NumFileLines )
+			{
+				const FString& NextLine = LinesToProcess[NextLineIdx].TrimTrailing().Trim();
+				if( NextLine.StartsWith("\"") && NextLine.EndsWith("\"") )
+				{
+					RawMsgIdPlural += NextLine.Mid( 1, NextLine.Len()-2 );
+				}
+				else
+				{
+					break;
+				}
+				LineIdx = NextLineIdx;	
+				NextLineIdx++;
+			}
+			ProcessedEntry->MsgIdPlural = RawMsgIdPlural;
+		}
+		else if( Line.StartsWith( TEXT("msgstr[") ) )
+		{
+			FString IndexStr;
+			if( !FindDelimitedString(Line, TEXT("["), TEXT("]"), IndexStr ) )
+			{
+				bSuccess = false;
+				break;
+			}
+			int32 Index = -1;
+			TTypeFromString<int32>::FromString(Index, *IndexStr );
+
+			check(Index >= 0);
+
+			FString RawMsgStr;
+			if( !FindDelimitedString(Line, TEXT("\""), TEXT("\""), RawMsgStr ) )
+			{
+				bSuccess = false;
+				break;
+			}
+			// If the following line contains more info for this element we continue to parse it out
+			uint32 NextLineIdx = LineIdx + 1;
+			while( NextLineIdx < NumFileLines )
+			{
+				const FString& NextLine = LinesToProcess[NextLineIdx].TrimTrailing().Trim();
+				if( NextLine.StartsWith("\"") && NextLine.EndsWith("\"") )
+				{
+					RawMsgStr += NextLine.Mid( 1, NextLine.Len()-2 );
+				}
+				else
+				{
+					break;
+				}
+				LineIdx = NextLineIdx;	
+				NextLineIdx++;
+			}
+
+			if( ProcessedEntry->MsgStr.Num() > Index )
+			{
+				ProcessedEntry->MsgStr[Index] = RawMsgStr;
+			}
+			else
+			{
+				ProcessedEntry->MsgStr.Insert( RawMsgStr, Index );
+			}
+		}
+		else if( Line.StartsWith( TEXT("msgstr") ) )
+		{
+			FString RawMsgStr;
+			if( !FindDelimitedString(Line, TEXT("\""), TEXT("\""), RawMsgStr ) )
+			{
+				bSuccess = false;
+				break;
+			}
+			// If the following line contains more info for this element we continue to parse it out
+			uint32 NextLineIdx = LineIdx + 1;
+			while( NextLineIdx < NumFileLines )
+			{
+				const FString& NextLine = LinesToProcess[NextLineIdx].TrimTrailing().Trim();
+				if( NextLine.StartsWith("\"") && NextLine.EndsWith("\"") )
+				{
+					RawMsgStr += NextLine.Mid( 1, NextLine.Len()-2 );
+				}
+				else
+				{
+					break;
+				}
+				LineIdx = NextLineIdx;	
+				NextLineIdx++;
+			}
+
+			FString MsgStr = RawMsgStr;
+			if( ProcessedEntry->MsgStr.Num() > 0 )
+			{
+				ProcessedEntry->MsgStr[0] = MsgStr;
+			}
+			else
+			{
+				ProcessedEntry->MsgStr.Add( MsgStr );
+			}
+		}
+		else
+		{
+			ProcessedEntry->UnknownElements.Add( Line );
+		}
+	}
+
+
+
+	return bSuccess;
+}
+
+void FPortableObjectFormatDOM::CreateNewHeader()
+{
+	// Reference: http://www.gnu.org/software/gettext/manual/gettext.html#Header-Entry
+	// Reference: http://www.gnu.org/software/gettext/manual/html_node/Header-Entry.html
+
+	//Hard code some header entries for now in the following format
+	/*
+	# Engine English translation
+	# Copyright 1998-2015 Epic Games, Inc. All Rights Reserved.
+	#
+	msgid ""
+	msgstr ""
+	"Project-Id-Version: Engine\n"
+	"POT-Creation-Date: 2014-1-31 04:16+0000\n"
+	"PO-Revision-Date: 2014-1-31 04:16+0000\n"
+	"Language-Team: \n"
+	"Language: en-us\n"
+	"MIME-Version: 1.0\n"
+	"Content-Type: text/plain; charset=UTF-8\n"
+	"Content-Transfer-Encoding: 8bit\n"
+	"Plural-Forms: nplurals=2; plural=(n != 1);\n"
+	*/
+
+	Header.Clear();
+
+	Header.SetEntryValue( TEXT("Project-Id-Version"), *GetProjectName() );
+
+	// Setup header entries
+	Header.TimeStamp();
+	Header.SetEntryValue( TEXT("Language-Team"), TEXT("") );
+	Header.SetEntryValue( TEXT("Language"), Language.GetLanguageCode() );
+	Header.SetEntryValue( TEXT("MIME-Version"), TEXT("1.0") );
+	Header.SetEntryValue( TEXT("Content-Type"), TEXT("text/plain; charset=UTF-8") );
+	Header.SetEntryValue( TEXT("Content-Transfer-Encoding"), TEXT("8bit") );
+	Header.SetEntryValue( TEXT("Plural-Forms"), Language.GetPluralForms() );
+
+	Header.Comments.Add( FString::Printf(TEXT("%s %s translation."), *GetProjectName(), *Language.EnglishName() ) );
+	Header.Comments.Add( TEXT("Copyright 1998-2015 Epic Games, Inc. All Rights Reserved.") );
+	Header.Comments.Add( FString(TEXT("")) );
+
+}
+
+bool FPortableObjectFormatDOM::SetLanguage( const FString& LanguageCode, const FString& LangPluralForms )
+{
+	FPortableObjectCulture NewLang(	LanguageCode, LangPluralForms );
+	if( NewLang.IsValid() )
+	{
+		Language = NewLang;
+		return true;
+	}
+	return false;
+}
+
+bool FPortableObjectFormatDOM::AddEntry( const TSharedRef< FPortableObjectEntry> LocEntry )
+{
+	TSharedPtr<FPortableObjectEntry> ExistingEntry = FindEntry( LocEntry );
+	if( ExistingEntry.IsValid()  )
+	{
+		ExistingEntry->AddReferences( LocEntry->ReferenceComments );
+		ExistingEntry->AddExtractedComments( LocEntry->ExtractedComments );
+
+		// Checks for a situation that we don't know how to handle yet.  ex. What happens when we match all the IDs for an entry but another param differs?
+		check( LocEntry->TranslatorComments == ExistingEntry->TranslatorComments );
+	}
+	else
+	{
+		Entries.Add( LocEntry );
+	}
+
+	return true;
+
+}
+
+
+TSharedPtr<FPortableObjectEntry> FPortableObjectFormatDOM::FindEntry( const TSharedRef<const FPortableObjectEntry> LocEntry )
+{
+	for( auto Entry : Entries )
+	{
+		if( *Entry == *LocEntry )
+		{
+			return Entry;
+		}
+	}
+	return NULL;
+}
+
+
+TSharedPtr<FPortableObjectEntry> FPortableObjectFormatDOM::FindEntry( const FString& MsgId, const FString& MsgIdPlural, const FString& MsgCtxt )
+{
+	TSharedRef<FPortableObjectEntry> TempEntry = MakeShareable( new FPortableObjectEntry );
+	TempEntry->MsgId = MsgId;
+	TempEntry->MsgIdPlural = MsgIdPlural;
+	TempEntry->MsgCtxt = MsgCtxt;
+	return FindEntry( TempEntry );
+}
+
+void FPortableObjectFormatDOM::SortEntries()
+{
+	// Sort keys.
+	for (const TSharedPtr<FPortableObjectEntry>& Entry : Entries)
+	{
+		Entry->ExtractedComments.Sort();
+	}
+
+	// Sort by namespace, then keys, then source text.
+	const auto& SortingPredicate = [](const TSharedPtr<FPortableObjectEntry>& A, const TSharedPtr<FPortableObjectEntry>& B) -> bool
+	{
+		// Compare namespace
+		if (A->MsgCtxt < B->MsgCtxt)
+		{
+			return true;
+		}
+		else if (A->MsgCtxt > B->MsgCtxt)
+		{
+			return false;
+		}
+
+		// Compare keys
+		const int32 ExtractedCommentCount = FMath::Max(A->ExtractedComments.Num(), B->ExtractedComments.Num());
+		for (int32 i = 0; i < ExtractedCommentCount; ++i)
+		{
+			if (A->ExtractedComments[i] < B->ExtractedComments[i] || (!A->ExtractedComments.IsValidIndex(i) && B->ExtractedComments.IsValidIndex(i)))
+			{
+				return true;
+			}
+			else if (A->ExtractedComments[i] > B->ExtractedComments[i] || (A->ExtractedComments.IsValidIndex(i) && !B->ExtractedComments.IsValidIndex(i)))
+			{
+				return false;
+			}
+		}
+
+		// Compare source string
+		if (A->MsgId < B->MsgId)
+		{
+			return true;
+		}
+		else if (A->MsgId > B->MsgId)
+		{
+			return false;
+		}
+
+		return A.Get() < B.Get();
+	};
+	Entries.Sort(SortingPredicate);
+}
+
+void FPortableObjectEntry::AddExtractedComment( const FString& InComment )
+{
+	if(!InComment.IsEmpty())
+	{
+		ExtractedComments.AddUnique(InComment);
+	}
+
+	//// Extracted comments can contain multiple references in a single line so we parse those out.
+	//TArray<FString> CommentsToProcess;
+	//InComment.ParseIntoArray( CommentsToProcess, TEXT(" "), true );
+	//for( const FString& ExtractedComment : CommentsToProcess )
+	//{
+	//	ExtractedComments.AddUnique( ExtractedComment );
+	//}
+}
+
+void FPortableObjectEntry::AddReference( const FString& InReference )
+{
+	if(!InReference.IsEmpty())
+	{
+		ReferenceComments.AddUnique(InReference);
+	}
+	
+	//// Reference comments can contain multiple references in a single line so we parse those out.
+	//TArray<FString> ReferencesToProcess;
+	//InReference.ParseIntoArray( ReferencesToProcess, TEXT(" "), true );
+	//for( const FString& Reference : ReferencesToProcess )
+	//{
+	//	ReferenceComments.AddUnique( Reference );
+	//}
+}
+
+void FPortableObjectEntry::AddExtractedComments( const TArray<FString>& InComments )
+{
+	for( const FString& ExtractedComment : InComments )
+	{
+		AddExtractedComment( ExtractedComment );
+	}
+}
+
+void FPortableObjectEntry::AddReferences( const TArray<FString>& InReferences )
+{
+	for( const FString& Reference : InReferences )
+	{
+		AddReference( Reference );
+	}
+}
+
+FString FPortableObjectEntry::ToString() const
+{
+	check( !MsgId.IsEmpty() );
+
+	FString Result;
+
+	for( const FString& TranslatorComment : TranslatorComments )
+	{
+		Result += FString::Printf( TEXT("# %s%s"), *TranslatorComment, NewLineDelimiter );
+	}
+
+	for( const FString& Comment : ExtractedComments )
+	{
+		if( !Comment.IsEmpty() )
+		{
+			Result += FString::Printf( TEXT("#. %s%s"), *Comment, NewLineDelimiter );
+		}
+		else
+		{
+			Result += FString::Printf( TEXT("#.%s"), NewLineDelimiter );
+		}
+	}
+
+	for( const FString& ReferenceComment : ReferenceComments )
+	{
+		Result += FString::Printf( TEXT("#: %s%s"), *ReferenceComment, NewLineDelimiter );
+	}
+
+	if( !MsgCtxt.IsEmpty() )
+	{
+		Result += FString::Printf(TEXT("msgctxt \"%s\"%s"), *MsgCtxt, NewLineDelimiter);
+	}
+
+	Result += FString::Printf(TEXT("msgid \"%s\"%s"), *MsgId, NewLineDelimiter);
+
+	if( MsgStr.Num() == 0)
+	{
+		Result += FString::Printf(TEXT("msgstr \"\"%s"), NewLineDelimiter);
+	}
+	else if( MsgStr.Num() == 1 )
+	{
+		Result += FString::Printf(TEXT("msgstr \"%s\"%s"), *MsgStr[0], NewLineDelimiter);
+	}
+	else
+	{
+		for( int32 Idx = 0; Idx < MsgStr.Num(); ++Idx )
+		{
+			Result += FString::Printf(TEXT("msgstr[%d] \"%s\"%s"), Idx, *MsgStr[Idx], NewLineDelimiter);
+		}
+	}
+	return Result;
+}
\ No newline at end of file
diff --git a/Engine/Source/Editor/UnrealEd/Public/Localization/PortableObjectFormatDOM.h b/Engine/Source/Editor/UnrealEd/Public/Localization/PortableObjectFormatDOM.h
new file mode 100644
--- /dev/null
+++ b/Engine/Source/Editor/UnrealEd/Public/Localization/PortableObjectFormatDOM.h
@@ -0,0 +1,376 @@
+// Copyright 1998-2015 Epic Games, Inc. All Rights Reserved.
+
+#pragma once 
+
+#include "Commandlets/Commandlet.h"
+
+class FCulture;
+
+/**
+* Class for handling language codes used in the Portable Object format.
+*/
+class FPortableObjectCulture
+{
+public:
+	FPortableObjectCulture()
+	{}
+
+	FPortableObjectCulture( const FString& LangCode, const FString& PluralForms = TEXT("") );
+
+	FPortableObjectCulture( const FPortableObjectCulture& Other );
+
+	FPortableObjectCulture& operator=( const FPortableObjectCulture& Other )
+	{
+		if (this == &Other)
+		{
+			return *this;
+		}
+		SetLanguageCode( Other.GetLanguageCode() );
+		SetPluralForms( Other.GetPluralForms() );
+		return *this;
+	}
+
+	/**
+	 * Checks to see if we have a language code and if we were able to match a culture to that code
+	 */
+	bool IsValid() const { return !LanguageCode.IsEmpty() && Culture.IsValid(); }
+
+	/**
+	 * Sets the language code
+	 */
+	void SetLanguageCode( const FString& LangCode );
+
+	/**
+	 * Retrieves the language code
+	 */
+	const FString& GetLanguageCode() const { return LanguageCode; }
+
+	/**
+	 * Sets the language plural forms.  This is only required if we wish to override the default plural forms associated with a language.
+	 */
+	void SetPluralForms( const FString& InPluralForms ) { LanguagePluralForms = InPluralForms; }
+
+	/**
+	 * Retrieves the language plural forms.
+	 */
+	FString GetPluralForms() const;
+
+	/**
+	 * Retrieves the two or three letter language code for the culture
+	 */
+	FString Language() const;
+
+	/**
+	 * Retrieves the country info of the culture
+	 */
+	FString Country() const;
+
+	/**
+	 * Retrieves the variant info for the culture
+	 */
+	FString Variant() const;
+
+	/**
+	 * Retrieves the culture equivalent that matches the set language code.
+	 */
+	FCulturePtr GetCulture() const;
+
+	/**
+	 * Retrieves the display name for the language.
+	 */
+	FString DisplayName() const;
+
+	/**
+	 * Retrieves the English name for the language.
+	 */
+	FString EnglishName() const;
+
+
+private:
+	/**
+	 * Retrieves the language default plural forms.  If no plural forms are provided for this language we use this to pull from a set of defaults if available.
+	 */
+	FString GetDefaultPluralForms() const;
+
+	// The language code as it appears in the PO file.
+	FString LanguageCode;
+	FString LanguagePluralForms;
+
+	FCulturePtr Culture;
+};
+
+
+/**
+* Class for representing entries in a Portable Object file(.po) or a Portable Object Template file(.pot).
+*/
+class UNREALED_API FPortableObjectEntry
+{
+public: 
+
+	FPortableObjectEntry()
+	{
+	}
+
+	FPortableObjectEntry( const FPortableObjectEntry& Other )
+		: MsgId( Other.MsgId )
+		, MsgIdPlural( Other.MsgIdPlural )
+		, MsgCtxt( Other.MsgCtxt )
+		, MsgStr( Other.MsgStr )
+		, UnknownElements( Other.UnknownElements )
+	{
+	}
+
+	FPortableObjectEntry& operator=( const FPortableObjectEntry& Other )
+	{
+		if (this == &Other)
+		{
+			return *this;
+		}
+		this->MsgId				= Other.MsgId;
+		this->MsgIdPlural		= Other.MsgIdPlural;
+		this->MsgCtxt			= Other.MsgCtxt;
+		this->MsgStr			= Other.MsgStr;
+		this->UnknownElements	= Other.UnknownElements;
+		return *this;
+	}
+
+	
+	bool operator==( const FPortableObjectEntry& Other ) 
+	{ 
+		return ( MsgId.Equals(Other.MsgId, ESearchCase::CaseSensitive) &&
+				MsgIdPlural.Equals(Other.MsgIdPlural, ESearchCase::CaseSensitive) &&
+				MsgCtxt.Equals(Other.MsgCtxt, ESearchCase::IgnoreCase) ); 
+	}
+	
+	/**
+	 * Helper function that adds to the extracted comments.
+	 *
+	 * @param InReference	String representing extracted comment entries of the following form: "File/path/file.cpp:20 File/path/file.cpp:21 File/path/file2.cpp:5".
+	 */
+	void AddExtractedComment( const FString& InComment );
+
+	/**
+	 * Helper function that adds to the reference comments.
+	 *
+	 * @param InReference	String representing reference comment entries.
+	 */
+	void AddReference( const FString& InReference );
+
+	/**
+	 * Helper function that adds to the reference comments.
+	 *
+	 * @param InReferences	String array representing reference comment entries of the following form: "File/path/file.cpp:20 File/path/file.cpp:21 File/path/file2.cpp:5".
+	 */
+	void AddExtractedComments( const TArray<FString>& InComments );
+
+	/**
+	 * Helper function that adds to the extracted comments.
+	 *
+	 * @param InReferences	String array representing extracted comment entries
+	 */
+	void AddReferences( const TArray<FString>& InReferences );
+
+	/**
+	 * Function to convert the entry to a string.
+	 *
+	 * @return	String representing the entry that can be written directly to the .po file.
+	 */
+	FString ToString() const;
+
+public:
+	/* Represents the original source text(also called the id or context).  Stored here are the msgid values from the Portable Object entries. */
+	FString MsgId;
+
+	/* Represents the plural form of the source text.  Stored here are the msgid_plural values from the Portable Object file entries. */
+	FString MsgIdPlural;
+
+	/* Represents the disambiguating  context for the source text.  If used, will prevent two identical source strings from getting collapsed into one entry. */
+	FString MsgCtxt;
+
+	/* Represents the translated text.  This stores the msgstr, msgstr[0], msgstr[1], etc values from Portable Object entries. */
+	TArray< FString > MsgStr;
+
+	/* Stores extracted comments. Lines starting with #. above the msgid.  
+	
+		#. TRANSLATORS: A test phrase with all letters of the English alphabet.
+		#. Replace it with a sample text in your language, such that it is
+		#. representative of language's writing system.
+		msgid "The Quick Brown Fox Jumps Over The Lazy Dog"
+		msgstr ""
+	*/
+	TArray<FString> ExtractedComments;
+
+	/* Stores the translator comments.  Lines starting with # (hash and space), followed by any text whatsoever.
+	
+		# Wikipedia says that 'etrurski' is our name for this script.
+		msgid "Old Italic"
+		msgstr "etrurski"
+	*/
+	TArray<FString> TranslatorComments;
+
+	/* Stores a reference comments. Lines starting with #: above the msgid.  
+	    
+		#: /Engine/Source/Runtime/Engine/Private/Actor.cpp:2306
+		#: /Engine/Source/Runtime/Engine/Private/Actor.cpp:2307 /Engine/Source/Runtime/Engine/Private/Actor.cpp:2308
+		msgid "The Quick Brown Fox Jumps Over The Lazy Dog"
+		msgstr ""
+	*/
+	TArray<FString> ReferenceComments;
+
+	/* Stores flags.  Lines starting with #,  
+	    
+		#: /Engine/Source/Runtime/Engine/Private/Actor.cpp:2306
+		#: /Engine/Source/Runtime/Engine/Private/Actor.cpp:2307 /Engine/Source/Runtime/Engine/Private/Actor.cpp:2308
+		msgid "The Quick Brown Fox Jumps Over The Lazy Dog"
+		msgstr ""
+	*/
+	TArray<FString> Flags;
+
+	/* Stores any unknown elements we may encounter when processing a Portable Object file.  */
+	TArray< FString > UnknownElements;
+};
+
+/**
+* Class that stores and manipulates PO and POT file header info.
+*/
+class FPortableObjectHeader
+{
+public:
+	typedef TPair<FString, FString> FPOHeaderEntry;
+	typedef TArray< FPOHeaderEntry > FPOHeaderData;
+
+	FPortableObjectHeader()
+	{
+	}
+
+	/**
+	 * Creates a string representation of the Portable Object header.
+	 *
+	 * @return	String representing the header that can be written directly to the .po file.
+	 */
+	FString ToString() const;
+
+	/**
+	 * Parses out header key/value pair entries.
+	 *
+	 * @param InStr	String representing a .po or .pot file header
+	 * @return	Returns true if successful, false otherwise.
+	 */
+	bool FromString( const FString& InStr ) {return false;}
+
+	/**
+	 * Parses out header info from a FPortableObjectEntry with an empty ID.
+	 *
+	 * @param LocEntry	The Portable Object entry with an empty ID
+	 * @return	Returns true if successful, false otherwise.
+	 */
+	bool FromLocPOEntry( const TSharedRef<const FPortableObjectEntry> LocEntry );
+	
+	/** Checks if a header entry exists with the given key */
+	bool HasEntry ( const FString& EntryKey );
+
+	/** Gets the header entry value for the given key or the empty string if the key does not exist */
+	FString GetEntryValue ( const FString& EntryKey );
+
+	/** Sets a header entry value */
+	void SetEntryValue( const FString& EntryKey, const FString& EntryValue );
+	
+	/**
+	 * Puts the current time into POT-Creation-Date and PO-Revision-Date entries
+	 */
+	void TimeStamp();
+
+	/**
+	 * Clears the header entries.
+	 */
+	void Clear() { HeaderEntries.Empty(); Comments.Empty(); }
+
+private:
+	FPortableObjectHeader::FPOHeaderEntry* FindEntry( const FString& EntryKey );
+
+public:
+
+	/* Stores the header comment block */
+	TArray<FString> Comments;
+	/* Stores all the header key/value pairs*/
+	FPOHeaderData HeaderEntries;
+};
+
+/**
+* Contains all the info we need to represent files in the Portable Object format.
+*/
+class UNREALED_API FPortableObjectFormatDOM
+{
+public:
+	FPortableObjectFormatDOM()
+	{}
+
+	FPortableObjectFormatDOM( const FString& LanguageCode )
+		: Language( LanguageCode )
+	{}
+
+	/**
+	 * Creates a string representation of the Portable Object.
+	 *
+	 * @return	String representing the Portable Object that can be written directly to a file.
+	 */
+	FString ToString();
+
+	/**
+	 * Parses Portable Object elements from a string.
+	 *
+	 * @param InStr	String representing a Portable Object file(.PO) or Portable Object Template file(.POT).
+	 * @return	Returns true if successful, false otherwise.
+	 */
+	bool FromString( const FString& InStr );
+
+	/** Creates a header entry based on the project and language info. */
+	void CreateNewHeader();
+
+	/**
+	 * Sets the language.
+	 *
+	 * @param LanguageCode	String representing a Portable Object language code.
+	 * @param LangPluralForms	Optional plural forms for the language.  If not provided, defaults will be used if found.
+	 * @return	Returns true if the function was able to successfully set the language code and pair it with a known culture, false otherwise.
+	 */
+	bool SetLanguage( const FString& LanguageCode, const FString& LangPluralForms = TEXT("") );
+
+	/**
+	 * Adds a translation entry to the Portable Object.
+	 *
+	 * @param LocEntry	The LocPOEntry to add.
+	 * @return	Returns true if successful, false otherwise.
+	 */
+	bool AddEntry( const TSharedRef<FPortableObjectEntry> LocEntry );
+
+	/* Returns the Portable Object entry that matches the passed in entry or NULL if not found */
+	TSharedPtr<FPortableObjectEntry> FindEntry( const TSharedRef<const FPortableObjectEntry> LocEntry );
+
+	/* Returns the Portable Object file entry that matches the passed in parameters or NULL if not found */
+	TSharedPtr<FPortableObjectEntry> FindEntry( const FString& MsgId, const FString& MsgIdPlural, const FString& MsgCtxt );
+
+	/* Sets the project name that will appear in the Project-Id-Version header entry */
+	void SetProjectName( const FString& ProjName ){ ProjectName = ProjName; }
+
+	/* Sets the project name from the Project-Id-Version header entry */
+	FString GetProjectName() { return ProjectName; }
+
+	TArray< TSharedPtr< FPortableObjectEntry > >::TConstIterator GetEntriesIterator()
+	{
+		return Entries.CreateConstIterator();
+	}
+
+	void SortEntries();
+
+private:
+	/** Copying is not supported */
+	FPortableObjectFormatDOM( const FPortableObjectFormatDOM& ); 
+	FPortableObjectFormatDOM& operator=( const FPortableObjectFormatDOM& );
+
+private:
+	FPortableObjectCulture Language;
+	FPortableObjectHeader Header;
+	FString ProjectName;
+	TArray< TSharedPtr< FPortableObjectEntry > > Entries;
+};
\ No newline at end of file
