Fixed Mac compile error

diff --git a/Engine/Source/Editor/UnrealEd/Private/Factories/TTFontImport.cpp b/Engine/Source/Editor/UnrealEd/Private/Factories/TTFontImport.cpp
--- a/Engine/Source/Editor/UnrealEd/Private/Factories/TTFontImport.cpp
+++ b/Engine/Source/Editor/UnrealEd/Private/Factories/TTFontImport.cpp
@@ -1683,7 +1683,7 @@ UTexture2D* UTrueTypeFontFactory::CreateTextureFromBitmap( UFont* Font, uint8* B
 	Texture->NeverStream = true;
 
 	// Copy bitmap data to texture page.
-	FColor FontColor8Bit( ImportOptions->Data.ForegroundColor );
+	FColor FontColor8Bit( ImportOptions->Data.ForegroundColor.ToFColor(true) );
 
 	// restart progress bar for font bitmap generation since this takes a while
 	int32 TotalProgress = BitmapWidth-1;
