Fixed bad formatting on map load error string

diff --git a/Engine/Source/Editor/UnrealEd/Private/FileHelpers.cpp b/Engine/Source/Editor/UnrealEd/Private/FileHelpers.cpp
--- a/Engine/Source/Editor/UnrealEd/Private/FileHelpers.cpp
+++ b/Engine/Source/Editor/UnrealEd/Private/FileHelpers.cpp
@@ -1985,7 +1985,7 @@ void FEditorFileUtils::LoadMap(const FString& InFilename, bool LoadAsTemplate, b
 	FString LongMapPackageName;
 	if ( !FPackageName::TryConvertFilenameToLongPackageName(Filename, LongMapPackageName) )
 	{
-		FMessageDialog::Open( EAppMsgType::Ok, FText::Format( NSLOCTEXT("Editor", "MapLoad_FriendlyBadFilename", "Map load failed. The filename '%s' is not within the game or engine content folders found in '%s'."), FText::FromString( Filename ), FText::FromString( FPaths::RootDir() ) ) );
+		FMessageDialog::Open( EAppMsgType::Ok, FText::Format( NSLOCTEXT("Editor", "MapLoad_FriendlyBadFilename", "Map load failed. The filename '{0}' is not within the game or engine content folders found in '{1}'."), FText::FromString( Filename ), FText::FromString( FPaths::RootDir() ) ) );
 		return;
 	}
 
