Don't create the checkerboard texture when running a commandlet - it doesn't have any source data, so we can't cook it correctly.

diff --git a/Engine/Source/Editor/UnrealEd/Private/ThumbnailManager.cpp b/Engine/Source/Editor/UnrealEd/Private/ThumbnailManager.cpp
--- a/Engine/Source/Editor/UnrealEd/Private/ThumbnailManager.cpp
+++ b/Engine/Source/Editor/UnrealEd/Private/ThumbnailManager.cpp
@@ -51,9 +51,9 @@ UThumbnailManager::UThumbnailManager(const FObjectInitializer& ObjectInitializer
 		EditorSkySphere = ConstructorStatics.EditorSkySphereMesh.Object;
 		FloorPlaneMaterial = ConstructorStatics.FloorPlaneMaterial.Object;
 		AmbientCubemap = ConstructorStatics.DaylightAmbientCubemap.Object;
-	}
 
-	SetupCheckerboardTexture();
+		SetupCheckerboardTexture();
+	}
 }
 
 void UThumbnailManager::Initialize(void)
