Remove warning when DumpBlueprintsInfoCommandlet doesn't understand an argument - there are lots of engine switches that commandlets are not meant to parse, and generating the blueprint docs specifies quite a lot of them (-stdout -forcelogflush -crashforuat -unattended -allowstdoutlogverbosity -buildmachine). We don't want warnings for those.

diff --git a/Engine/Source/Editor/UnrealEd/Private/Commandlets/DumpBlueprintsInfoCommandlet.cpp b/Engine/Source/Editor/UnrealEd/Private/Commandlets/DumpBlueprintsInfoCommandlet.cpp
--- a/Engine/Source/Editor/UnrealEd/Private/Commandlets/DumpBlueprintsInfoCommandlet.cpp
+++ b/Engine/Source/Editor/UnrealEd/Private/Commandlets/DumpBlueprintsInfoCommandlet.cpp
@@ -628,10 +628,6 @@ DumpBlueprintInfoUtils::CommandletOptions::CommandletOptions(TArray<FString> con
 		{
 			NewDumpFlags |= BPDUMP_ActionDatabaseInfo;
 		}
-		else if (!Switch.StartsWith("run=")) // account for the first switch, which invoked this commandlet
-		{
-			UE_LOG(LogBlueprintInfoDump, Warning, TEXT("Unrecognized command switch '%s', use -help for a listing of all accepted params"), *Switch);
-		}
 	}
 
 	if (NewDumpFlags != 0)
