Fix decals not drag/dropping properly when dropped over meshes

diff --git a/Engine/Source/Editor/UnrealEd/Private/Kismet2/ComponentEditorUtils.cpp b/Engine/Source/Editor/UnrealEd/Private/Kismet2/ComponentEditorUtils.cpp
--- a/Engine/Source/Editor/UnrealEd/Private/Kismet2/ComponentEditorUtils.cpp
+++ b/Engine/Source/Editor/UnrealEd/Private/Kismet2/ComponentEditorUtils.cpp
@@ -687,8 +687,11 @@ bool FComponentEditorUtils::AttemptApplyMaterialToComponent(USceneComponent* Sce
 	auto MeshComponent = Cast<UMeshComponent>(SceneComponent);
 	auto DecalComponent = Cast<UDecalComponent>(SceneComponent);
 
+	UMaterial* BaseMaterial = MaterialToApply->GetBaseMaterial();
+
+	bool bIsPlacable = DecalComponent || ( MeshComponent && BaseMaterial &&  BaseMaterial->MaterialDomain != MD_DeferredDecal && BaseMaterial->MaterialDomain != MD_UI );
 	// We can only apply a material to a mesh or a decal
-	if (MeshComponent || DecalComponent)
+	if (bIsPlacable && (MeshComponent || DecalComponent) )
 	{
 		bResult = true;
 		const FScopedTransaction Transaction(LOCTEXT("DropTarget_UndoSetComponentMaterial", "Assign Material to Component (Drag and Drop)"));
