Made a correction in the wording of a check done to see if a material should be applied or placed as a decal

diff --git a/Engine/Source/Editor/UnrealEd/Private/Kismet2/ComponentEditorUtils.cpp b/Engine/Source/Editor/UnrealEd/Private/Kismet2/ComponentEditorUtils.cpp
--- a/Engine/Source/Editor/UnrealEd/Private/Kismet2/ComponentEditorUtils.cpp
+++ b/Engine/Source/Editor/UnrealEd/Private/Kismet2/ComponentEditorUtils.cpp
@@ -689,9 +689,9 @@ bool FComponentEditorUtils::AttemptApplyMaterialToComponent(USceneComponent* Sce
 
 	UMaterial* BaseMaterial = MaterialToApply->GetBaseMaterial();
 
-	bool bIsPlacable = DecalComponent || ( MeshComponent && BaseMaterial &&  BaseMaterial->MaterialDomain != MD_DeferredDecal && BaseMaterial->MaterialDomain != MD_UI );
+	bool bCanApplyToComponent = DecalComponent || ( MeshComponent && BaseMaterial &&  BaseMaterial->MaterialDomain != MD_DeferredDecal && BaseMaterial->MaterialDomain != MD_UI );
 	// We can only apply a material to a mesh or a decal
-	if (bIsPlacable && (MeshComponent || DecalComponent) )
+	if (bCanApplyToComponent && (MeshComponent || DecalComponent) )
 	{
 		bResult = true;
 		const FScopedTransaction Transaction(LOCTEXT("DropTarget_UndoSetComponentMaterial", "Assign Material to Component (Drag and Drop)"));
