Prevent spamming of checkout notifications when cooking in the editor.

diff --git a/Engine/Source/Editor/UnrealEd/Private/UnrealEdEngine.cpp b/Engine/Source/Editor/UnrealEd/Private/UnrealEdEngine.cpp
--- a/Engine/Source/Editor/UnrealEd/Private/UnrealEdEngine.cpp
+++ b/Engine/Source/Editor/UnrealEd/Private/UnrealEdEngine.cpp
@@ -534,21 +534,22 @@ void UUnrealEdEngine::OnSourceControlStateUpdated(const FSourceControlOperationR
 			}
 			else
 			{
-				if (SourceControlState->CanCheckout() || !SourceControlState->IsCurrent() || SourceControlState->IsCheckedOutOther())
+				// Note when cooking in the editor we ignore package notifications.  The cooker is saving packages in a temp location which generates bogus checkout messages otherwise.
+				if ((SourceControlState->CanCheckout() || !SourceControlState->IsCurrent() || SourceControlState->IsCheckedOutOther()) && (!CookServer || !CookServer->IsCookingInEditor() ) )
 				{
 					// To get here, either "prompt for checkout on asset modification" is set, or "automatically checkout on asset modification"
 					// is set, but it failed.
 
-			// Allow packages that are not checked out to pass through.
-			// Allow packages that are not current or checked out by others pass through.  
-			// The user wont be able to checkout these packages but the checkout dialog will show up with a special icon 
-			// to let the user know they wont be able to checkout the package they are modifying.
+					// Allow packages that are not checked out to pass through.
+					// Allow packages that are not current or checked out by others pass through.  
+					// The user wont be able to checkout these packages but the checkout dialog will show up with a special icon 
+					// to let the user know they wont be able to checkout the package they are modifying.
 
 					PackageToNotifyState.Add(Package, SourceControlState->CanCheckout() ? NS_PendingPrompt : NS_PendingWarning);
-			// We need to prompt since a new package was added
-			bNeedToPromptForCheckout = true;
-		}
-	}
+					// We need to prompt since a new package was added
+					bNeedToPromptForCheckout = true;
+				}
+			}
 		}
 	}
 }
