Fix mac compile errors

diff --git a/Engine/Source/Editor/MaterialEditor/Private/SMaterialEditorViewport.cpp b/Engine/Source/Editor/MaterialEditor/Private/SMaterialEditorViewport.cpp
--- a/Engine/Source/Editor/MaterialEditor/Private/SMaterialEditorViewport.cpp
+++ b/Engine/Source/Editor/MaterialEditor/Private/SMaterialEditorViewport.cpp
@@ -812,4 +812,4 @@ void SMaterialEditorUIPreviewViewport::OnPreviewYCommitted( int32 NewValue, ETex
 	OnPreviewYChanged( NewValue );
 }
 
-#undef LOC_NAMESPACE
\ No newline at end of file
+#undef LOCTEXT_NAMESPACE
diff --git a/Engine/Source/Editor/UnrealEd/Private/ThumbnailHelpers.cpp b/Engine/Source/Editor/UnrealEd/Private/ThumbnailHelpers.cpp
--- a/Engine/Source/Editor/UnrealEd/Private/ThumbnailHelpers.cpp
+++ b/Engine/Source/Editor/UnrealEd/Private/ThumbnailHelpers.cpp
@@ -270,7 +270,7 @@ void FMaterialThumbnailScene::SetMaterialInterface(UMaterialInterface* InMateria
 		// UI material thumbnails always get a 2D plane centered at the camera which is a better representation of the
 		// what the material will look like on UI
 		bIsUIMaterial = BaseMaterial && BaseMaterial->IsUIMaterial();
-		EThumbnailPrimType PrimitiveType = bIsUIMaterial ? TPT_Plane : ThumbnailInfo->PrimitiveType;
+		EThumbnailPrimType PrimitiveType = bIsUIMaterial ? TPT_Plane : ThumbnailInfo->PrimitiveType.GetValue();
 
 		switch( PrimitiveType )
 		{
