Rewrite of the chunk manifest generator code so it doesn't require packages to be loaded.

diff --git a/Engine/Source/Editor/UnrealEd/Classes/CookOnTheSide/CookOnTheFlyServer.h b/Engine/Source/Editor/UnrealEd/Classes/CookOnTheSide/CookOnTheFlyServer.h
--- a/Engine/Source/Editor/UnrealEd/Classes/CookOnTheSide/CookOnTheFlyServer.h
+++ b/Engine/Source/Editor/UnrealEd/Classes/CookOnTheSide/CookOnTheFlyServer.h
@@ -405,7 +405,17 @@ private:
 			return false;
 		}
 
-
+		void GetCookedFilesForPlatform(const FName& PlatformName, TArray<FName>& CookedFiles)
+		{
+			FScopeLock ScopeLock(&SynchronizationObject);
+			for (const auto& CookedFile : FilesProcessed)
+			{
+				if (CookedFile.Value.HasPlatform(PlatformName))
+				{
+					CookedFiles.Add(CookedFile.Value.GetFilename());
+				}
+			}
+		}
 		void Empty(int32 ExpectedNumElements = 0)
 		{
 			FScopeLock ScopeLock( &SynchronizationObject );
@@ -634,7 +644,7 @@ private:
 		/** Map of platform name to manifest generator, manifest is only used in cook by the book however it needs to be maintained across multiple cook by the books. */
 		TMap<FName, FChunkManifestGenerator*> ManifestGenerators;
 		/** Dependency graph of maps as root objects. */
-		TMap< FName, TSet <FName> > MapDependencyGraph; 
+		TMap<FName, TMap< FName, TSet <FName> > > MapDependencyGraphs; 
 		/** If a cook is cancelled next cook will need to resume cooking */ 
 		TArray<FFilePlatformRequest> PreviousCookRequests; 
 		/** If we are based on a release version of the game this is the set of packages which were cooked in that release */
@@ -802,6 +812,12 @@ public:
 	 */
 	void StartCookByTheBook( const FCookByTheBookStartupOptions& CookByTheBookStartupOptions );
 
+	/**
+	 * ValidateCookByTheBookSettings
+	 * look at the cookByTheBookOptions and ensure there isn't any conflicting settings
+	 */
+	void ValidateCookByTheBookSettings() const;
+
 	/**
 	 * Queue a cook by the book cancel (you might want to do this instead of calling cancel directly so that you don't have to be in the game thread when canceling
 	 */
@@ -970,6 +986,22 @@ private:
 	 */
 	bool GetAllPackagesFromAssetRegistry( const FString& AssetRegistryPath, TArray<FName>& OutPackageNames ) const;
 
+	/**
+	 * BuildMapDependencyGraph
+	 * builds a map of dependencies from maps
+	 * 
+	 * @param PlatformName name of the platform we want to build a map for
+	 */
+	void BuildMapDependencyGraph(const FName& PlatformName);
+
+	/**
+	* WriteMapDependencyGraph
+	* write a previously built map dependency graph out to the sandbox directory for a platform
+	*
+	* @param PlatformName name of the platform we want to save out the dependency graph for
+	*/
+	void WriteMapDependencyGraph(const FName& PlatformName);
+
 	/**
 	 * IsChildCooker, 
 	 * returns if this cooker is a sue chef for some other master chef.
@@ -1081,15 +1113,6 @@ private:
 	 */
 	bool ContainsMap(const FName& PackageName) const;
 
-	/**
-	 * AddDependenciesToManifest
-	 * Add the dependencies of this package to the cooked package manifest
-	 * use information from the asset registry to generate this information
-	 *
-	 * @params StandardFilename of the package which we want dependencies to be copied over
-	 */
-	void AddDependenciesToManifest(const FName& StandardFilename, const FName& LastLoadedMapName, const TSet<FName>& Dependencies, const TArray<FName>& TargetPlatformNames);
-
 	/**
 	 * GetCurrentIniVersionStrings gets the current ini version strings for compare against previous cook
 	 * 
diff --git a/Engine/Source/Editor/UnrealEd/Private/Commandlets/ChunkManifestGenerator.cpp b/Engine/Source/Editor/UnrealEd/Private/Commandlets/ChunkManifestGenerator.cpp
--- a/Engine/Source/Editor/UnrealEd/Private/Commandlets/ChunkManifestGenerator.cpp
+++ b/Engine/Source/Editor/UnrealEd/Private/Commandlets/ChunkManifestGenerator.cpp
@@ -43,34 +43,6 @@ FChunkManifestGenerator::~FChunkManifestGenerator()
 	FinalChunkManifests.Empty();
 }
 
-void FChunkManifestGenerator::OnAssetLoaded(UObject* Asset)
-{
-	if (Asset != NULL)
-	{
-		UPackage* AssetPackage = CastChecked<UPackage>(Asset->GetOutermost());
-		if (!AssetsLoadedWithLastPackage.Contains(AssetPackage->GetFName()))
-		{
-			AssetsLoadedWithLastPackage.Add(AssetPackage->GetFName());
-		}
-	}
-}
-
-void FChunkManifestGenerator::OnLastPackageLoaded( UPackage* Package )
-{
-	if ( !AssetsLoadedWithLastPackage.Contains(Package->GetFName()))
-	{
-		AssetsLoadedWithLastPackage.Add(Package->GetFName());
-	}
-}
-
-void FChunkManifestGenerator::OnLastPackageLoaded( const FName& PackageName )
-{
-	if ( !AssetsLoadedWithLastPackage.Contains(PackageName))
-	{
-		AssetsLoadedWithLastPackage.Add(PackageName);
-	}
-}
-
 bool FChunkManifestGenerator::CleanTempPackagingDirectory(const FString& Platform) const
 {
 	FString TmpPackagingDir = GetTempPackagingDirectoryForPlatform(Platform);
@@ -165,56 +137,8 @@ bool FChunkManifestGenerator::GenerateStreamingInstallManifest(const FString& Pl
 	return true;
 }
 
-void FChunkManifestGenerator::Initialize(bool InGenerateChunks)
-{
-	bGenerateChunks = InGenerateChunks;
-
-	// Calculate the largest chunk id used by the registry to get the indices for the default chunks
-	AssetRegistry.GetAllAssets(AssetRegistryData);
-	int32 LargestChunkID = -1;
 
-	for (int32 Index = 0; Index < AssetRegistryData.Num(); ++Index)
-	{
-		auto& AssetData = AssetRegistryData[Index];
-		auto& RegistryChunkIDs = RegistryChunkIDsMap.FindOrAdd(AssetData.PackageName);
-		for (auto ChunkIt = AssetData.ChunkIDs.CreateConstIterator(); ChunkIt; ++ChunkIt)
-		{
-			int32 ChunkID = *ChunkIt;
-			if (ChunkID < 0)
-			{
-				UE_LOG(LogChunkManifestGenerator, Warning, TEXT("Out of range ChunkID: %d"), ChunkID);
-				ChunkID = 0;
-			}
-			else if (ChunkID > LargestChunkID)
-			{
-				LargestChunkID = ChunkID;
-			}
-			if (!RegistryChunkIDs.Contains(ChunkID))
-			{
-				RegistryChunkIDs.Add(ChunkID);
-			}
-			auto* FoundIDList = PackageChunkIDMap.Find(AssetData.PackageName);
-			if (!FoundIDList)
-			{
-				FoundIDList = &PackageChunkIDMap.Add(AssetData.PackageName);
-			}
-			FoundIDList->AddUnique(ChunkID);
-		}
-		// Now clear the original chunk id list. We will fill it with real IDs when cooking.
-		AssetData.ChunkIDs.Empty();
-		// Map asset data to its package (there can be more than one asset data per package).
-		auto& PackageData = PackageToRegistryDataMap.FindOrAdd(AssetData.PackageName);
-		PackageData.Add(Index);
-	}
-
-	// Hook up game delegate
-	if (bGenerateChunks)
-	{
-		FCoreUObjectDelegates::OnAssetLoaded.AddRaw(this, &FChunkManifestGenerator::OnAssetLoaded);
-	}
-}
-
-void FChunkManifestGenerator::AddPackageToChunkManifest(const FName& PackageFName, const FString& PackagePathName, const FString& SandboxFilename, const FString& LastLoadedMapName, FSandboxPlatformFile* InSandboxFile)
+void FChunkManifestGenerator::GenerateChunkManifestForPackage(const FName& PackageFName, const FString& PackagePathName, const FString& SandboxFilename, const FString& LastLoadedMapName, FSandboxPlatformFile* InSandboxFile)
 {
 	TArray<int32> TargetChunks;
 	TArray<int32> ExistingChunkIDs;
@@ -227,16 +151,6 @@ void FChunkManifestGenerator::AddPackageToChunkManifest(const FName& PackageFNam
 
 	if (bGenerateChunks)
 	{
-		// Try to determine if this package has been loaded as a result of loading a map package.
-		FString MapThisAssetWasLoadedWith;
-		if (!LastLoadedMapName.IsEmpty())
-		{
-			if (AssetsLoadedWithLastPackage.Contains(PackageFName))
-			{
-				MapThisAssetWasLoadedWith = LastLoadedMapName;
-			}
-		}
-
 		// Collect all chunk IDs associated with this package from the asset registry
 		TArray<int32> RegistryChunkIDs = GetAssetRegistryChunkAssignments(PackageFName);
 		ExistingChunkIDs = GetExistingPackageChunkAssignments(PackageFName);
@@ -245,7 +159,7 @@ void FChunkManifestGenerator::AddPackageToChunkManifest(const FName& PackageFNam
 		// FString Name = Package->GetPathName();
 		if (FGameDelegates::Get().GetAssignStreamingChunkDelegate().IsBound())
 		{
-			FGameDelegates::Get().GetAssignStreamingChunkDelegate().ExecuteIfBound(PackagePathName, MapThisAssetWasLoadedWith, RegistryChunkIDs, ExistingChunkIDs, TargetChunks);
+			FGameDelegates::Get().GetAssignStreamingChunkDelegate().ExecuteIfBound(PackagePathName, LastLoadedMapName, RegistryChunkIDs, ExistingChunkIDs, TargetChunks);
 		}
 		else
 		{
@@ -255,8 +169,6 @@ void FChunkManifestGenerator::AddPackageToChunkManifest(const FName& PackageFNam
 		}
 	}
 
-	NotifyPackageWasCooked(SandboxFilename, PackageFName);
-
 	bool bAssignedToChunk = false;
 	// if the delegate requested a specific chunk assignment, add them package to it now.
 	for (const auto& PackageChunk : TargetChunks)
@@ -273,86 +185,8 @@ void FChunkManifestGenerator::AddPackageToChunkManifest(const FName& PackageFNam
 		}
 	}
 
-	if (!bAssignedToChunk)
-	{
-		NotifyPackageWasNotAssigned(SandboxFilename, PackageFName);
-	}
 }
 
-void FChunkManifestGenerator::AddPackageToChunkManifest(UPackage* Package, const FString& SandboxFilename, const FString& LastLoadedMapName, FSandboxPlatformFile* InSandboxFile)
-{
-#if 0
-	AddPackageToChunkManifest(Package->GetFName(), Package->GetPathName(), SandboxFilename, LastLoadedMapName, InSandboxFile);
-#else
-	TArray<int32> TargetChunks;
-	TArray<int32> ExistingChunkIDs;
-	
-	if (!bGenerateChunks)
-	{
-		TargetChunks.AddUnique(0);
-		ExistingChunkIDs.AddUnique(0);
-	}
-	
-	auto PackageFName = Package->GetFName();
-	if (bGenerateChunks)
-	{
-		// Try to determine if this package has been loaded as a result of loading a map package.
-		FString MapThisAssetWasLoadedWith;
-		if (!LastLoadedMapName.IsEmpty())
-		{
-			if (AssetsLoadedWithLastPackage.Contains(PackageFName))
-			{
-				MapThisAssetWasLoadedWith = LastLoadedMapName;
-			}
-		}
-
-		// Collect all chunk IDs associated with this package from the asset registry
-		TArray<int32> RegistryChunkIDs = GetAssetRegistryChunkAssignments(Package);
-		ExistingChunkIDs = GetExistingPackageChunkAssignments(PackageFName);
-
-		// Try to call game-specific delegate to determine the target chunk ID
-		FString Name = Package->GetPathName();
-		if (FGameDelegates::Get().GetAssignStreamingChunkDelegate().IsBound())
-		{
-			FGameDelegates::Get().GetAssignStreamingChunkDelegate().ExecuteIfBound(Name, MapThisAssetWasLoadedWith, RegistryChunkIDs, ExistingChunkIDs, TargetChunks);
-		}
-		else
-		{
-			//Take asset registry assignments and existing assignments
-			TargetChunks.Append(RegistryChunkIDs);
-			TargetChunks.Append(ExistingChunkIDs);
-		}
-	}
-
-	NotifyPackageWasCooked(SandboxFilename, PackageFName);
-
-	bool bAssignedToChunk = false;
-	// if the delegate requested a specific chunk assignment, add them package to it now.
-	for (const auto& PackageChunk : TargetChunks)
-	{
-		AddPackageToManifest(SandboxFilename, PackageFName, PackageChunk);
-		bAssignedToChunk = true;
-	}
-	// If the delegate requested to remove the package from any chunk, remove it now
-	for (const auto& PackageChunk : ExistingChunkIDs)
-	{
-		if (!TargetChunks.Contains(PackageChunk))
-		{
-			RemovePackageFromManifest(PackageFName, PackageChunk);
-		}
-	}
-
-	if (!bAssignedToChunk)
-	{
-		NotifyPackageWasNotAssigned(SandboxFilename, PackageFName);
-	}
-#endif
-}
-
-void FChunkManifestGenerator::PrepareToLoadNewPackage(const FString& Filename)
-{
-	AssetsLoadedWithLastPackage.Empty();
-}
 
 void FChunkManifestGenerator::CleanManifestDirectories()
 {
@@ -415,7 +249,7 @@ bool FChunkManifestGenerator::SaveManifests(FSandboxPlatformFile* InSandboxFile)
 
 bool FChunkManifestGenerator::LoadAssetRegistry(const FString& SandboxPath, const TSet<FName>* PackagesToKeep)
 {
-	UE_LOG(LogChunkManifestGenerator, Display, TEXT("Saving asset registry."));
+	UE_LOG(LogChunkManifestGenerator, Display, TEXT("Loading asset registry."));
 
 	// Load generated registry for each platform
 	check(Platforms.Num() == 1);
@@ -470,6 +304,161 @@ bool FChunkManifestGenerator::LoadAssetRegistry(const FString& SandboxPath, cons
 	return true;
 }
 
+bool FChunkManifestGenerator::ContainsMap(const FName& PackageName) const
+{
+	const auto& Assets = PackageToRegistryDataMap.FindChecked(PackageName);
+	
+	for (const auto& AssetIndex : Assets)
+	{
+		const auto& Asset = AssetRegistryData[AssetIndex];
+		if (Asset.GetClass()->IsChildOf(UWorld::StaticClass()) || Asset.GetClass()->IsChildOf(ULevel::StaticClass()))
+		{
+			return true;
+		}
+	}
+	return false;
+}
+
+void FChunkManifestGenerator::Initialize(const TArray<FName> &InStartupPackages)
+{
+	StartupPackages = InStartupPackages;
+}
+
+
+void FChunkManifestGenerator::BuildChunkManifest(const TArray<FName>& CookedPackages, FSandboxPlatformFile* SandboxFile, bool bGenerateStreamingInstallManifest)
+{
+	bGenerateChunks = bGenerateStreamingInstallManifest;
+
+	// initialize LargestChunkId, FoundIDList, PackageChunkIDMap, AssetRegistryData
+
+	// Calculate the largest chunk id used by the registry to get the indices for the default chunks
+	AssetRegistry.GetAllAssets(AssetRegistryData);
+	int32 LargestChunkID = -1;
+
+	for (int32 Index = 0; Index < AssetRegistryData.Num(); ++Index)
+	{
+		auto& AssetData = AssetRegistryData[Index];
+		auto& RegistryChunkIDs = RegistryChunkIDsMap.FindOrAdd(AssetData.PackageName);
+		for (auto ChunkIt = AssetData.ChunkIDs.CreateConstIterator(); ChunkIt; ++ChunkIt)
+		{
+			int32 ChunkID = *ChunkIt;
+			if (ChunkID < 0)
+			{
+				UE_LOG(LogChunkManifestGenerator, Warning, TEXT("Out of range ChunkID: %d"), ChunkID);
+				ChunkID = 0;
+			}
+			else if (ChunkID > LargestChunkID)
+			{
+				LargestChunkID = ChunkID;
+			}
+			if (!RegistryChunkIDs.Contains(ChunkID))
+			{
+				RegistryChunkIDs.Add(ChunkID);
+			}
+			auto* FoundIDList = PackageChunkIDMap.Find(AssetData.PackageName);
+			if (!FoundIDList)
+			{
+				FoundIDList = &PackageChunkIDMap.Add(AssetData.PackageName);
+			}
+			FoundIDList->AddUnique(ChunkID);
+		}
+		// Now clear the original chunk id list. We will fill it with real IDs when cooking.
+		AssetData.ChunkIDs.Empty();
+		// Map asset data to its package (there can be more than one asset data per package).
+		auto& PackageData = PackageToRegistryDataMap.FindOrAdd(AssetData.PackageName);
+		PackageData.Add(Index);
+	}
+
+
+	// add all the packages to the unassigned package list
+	for (const auto& CookedPackage : CookedPackages)
+	{
+		const FString SandboxPath = SandboxFile->ConvertToAbsolutePathForExternalAppForWrite(*FPackageName::LongPackageNameToFilename(CookedPackage.ToString()));
+
+		AllCookedPackages.Add(CookedPackage, SandboxPath);
+		UnassignedPackageSet.Add(CookedPackage, SandboxPath);
+	}
+	
+	for (const auto& CookedPackage : StartupPackages)
+	{
+		const FString SandboxPath = SandboxFile->ConvertToAbsolutePathForExternalAppForWrite(*FPackageName::LongPackageNameToFilename(CookedPackage.ToString()));
+
+		AllCookedPackages.Add(CookedPackage, SandboxPath);
+		AddPackageToManifest(SandboxPath, CookedPackage, 0);
+	}
+
+
+	// assign chunks for all the map packages
+	for (const auto& CookedPackage : UnassignedPackageSet)
+	{
+		// ignore non map packages for now
+		const FName MapFName = CookedPackage.Key;
+
+		// this package could be missing from the map because it didn't get cooked. 
+		// the reason for this might be that it's a redirector therefore we cooked the package which actually contains the asset
+		if (PackageToRegistryDataMap.Find(MapFName) == nullptr)
+			continue;
+
+		if (ContainsMap(MapFName) == false)
+			continue;
+
+		// get all the dependencies for this map
+		TArray<FName> MapDependencies;
+		ensure(GatherAllPackageDependencies(MapFName, MapDependencies));
+
+		
+		for (const auto& RawPackageFName : MapDependencies)
+		{
+			FName PackageFName = RawPackageFName;
+			if ((FPackageName::IsValidLongPackageName(RawPackageFName.ToString()) == false) &&
+				(FPackageName::IsScriptPackage(RawPackageFName.ToString()) == false) )
+			{
+				FString LongPackageName;
+				if (FPackageName::SearchForPackageOnDisk(RawPackageFName.ToString(), &LongPackageName) == false)
+				{
+					continue;
+				}
+				PackageFName = FName(*LongPackageName);
+			}
+			
+			// don't include script packages in dependencies as they are always in memory
+			if (FPackageName::IsScriptPackage(PackageFName.ToString()))
+			{
+				// no one likes script packages
+				continue;
+			}
+			
+
+
+			const FString PackagePathName = PackageFName.ToString();
+			const FString MapName = MapFName.ToString();
+			const FString& SandboxFilename = AllCookedPackages.FindChecked(PackageFName);
+
+
+			GenerateChunkManifestForPackage(PackageFName, PackagePathName, SandboxFilename, MapName, SandboxFile);
+
+
+		}
+
+	}
+
+	// process the remaining unassigned packages, they don't have a map associated with them
+	// probably a good reason for it but maybe not
+	for (const auto& CookedPackage : UnassignedPackageSet)
+	{
+		const FName& PackageFName = CookedPackage.Key;
+		const FString& SandboxFilename = AllCookedPackages.FindChecked(PackageFName);
+		const FString PackagePathName = PackageFName.ToString();
+
+		GenerateChunkManifestForPackage(PackageFName, PackagePathName, SandboxFilename, FString(), SandboxFile);
+	}
+
+
+	// anything that remains in the UnAssignedPackageSet will be put in chunk0 when we save the asset registry
+
+}
+
+
 bool FChunkManifestGenerator::SaveAssetRegistry(const FString& SandboxPath, const TArray<FName>* IgnorePackageList)
 {
 	UE_LOG(LogChunkManifestGenerator, Display, TEXT("Saving asset registry."));
@@ -748,19 +737,27 @@ bool FChunkManifestGenerator::GetPackageDependencies(FName PackageName, TArray<F
 
 bool FChunkManifestGenerator::GatherAllPackageDependencies(FName PackageName, TArray<FName>& DependentPackageNames)
 {
-	TArray<FName> LocalDependentPackages;
-	if (!GetPackageDependencies(PackageName, LocalDependentPackages))
+
+	if (GetPackageDependencies(PackageName, DependentPackageNames) == false)
 	{
 		return false;
 	}
 
-	DependentPackageNames.Append(LocalDependentPackages);
-	for (const auto& DependentPackage : LocalDependentPackages)
+	int32 DependencyCounter = 0;
+	while (DependencyCounter < DependentPackageNames.Num())
 	{
-		if (!GetPackageDependencies(DependentPackage, DependentPackageNames))
+		const FName& ChildPackageName = DependentPackageNames[DependencyCounter];
+		++DependencyCounter;
+		TArray<FName> ChildDependentPackageNames;
+		if (GetPackageDependencies(ChildPackageName, ChildDependentPackageNames) == false)
 		{
 			return false;
 		}
+
+		for (const auto& ChildDependentPackageName : ChildDependentPackageNames)
+		{
+			DependentPackageNames.AddUnique(ChildDependentPackageName);
+		}
 	}
 
 	return true;
@@ -862,10 +859,6 @@ void FChunkManifestGenerator::AddPackageToManifest(const FString& PackageSandbox
 	UnassignedPackageSet.Remove(PackageName);
 }
 
-void FChunkManifestGenerator::AddUnassignedPackageToManifest(UPackage* Package, const FString& PackageSandboxPath )
-{
-	NotifyPackageWasNotAssigned(PackageSandboxPath, Package->GetFName() );
-}
 
 void FChunkManifestGenerator::RemovePackageFromManifest(FName PackageName, int32 ChunkId)
 {
@@ -875,16 +868,6 @@ void FChunkManifestGenerator::RemovePackageFromManifest(FName PackageName, int32
 	}
 }
 
-void FChunkManifestGenerator::NotifyPackageWasNotAssigned(const FString& PackageSandboxPath, FName PackageName)
-{
-	UnassignedPackageSet.Add(PackageName, PackageSandboxPath);
-}
-
-void FChunkManifestGenerator::NotifyPackageWasCooked(const FString& PackageSandboxPath, FName PackageName)
-{
-	AllCookedPackages.Add(PackageName, PackageSandboxPath);
-}
-
 void FChunkManifestGenerator::ResolveChunkDependencyGraph(const FChunkDependencyTreeNode& Node, FChunkPackageSet BaseAssetSet) 
 {
 	if (FinalChunkManifests.Num() > Node.ChunkID && FinalChunkManifests[Node.ChunkID])
@@ -938,6 +921,44 @@ bool FChunkManifestGenerator::CheckChunkAssetsAreNotInChild(const FChunkDependen
 	return true;
 }
 
+void FChunkManifestGenerator::AddPackageAndDependenciesToChunk(FChunkPackageSet* ThisPackageSet, FName InPkgName, const FString& InSandboxFile, int32 ChunkID, FSandboxPlatformFile* SandboxPlatformFile)
+{
+	//Add this asset
+	ThisPackageSet->Add(InPkgName, InSandboxFile);
+
+	//Only gather dependencies if we're chunking
+	if (!bGenerateChunks)
+	{
+		return;
+	}
+
+	//now add any dependencies
+	TArray<FName> DependentPackageNames;
+	if (GatherAllPackageDependencies(InPkgName, DependentPackageNames))
+	{
+		for (const auto& PkgName : DependentPackageNames)
+		{
+			bool bSkip = false;
+			if (ChunkID != 0 && FinalChunkManifests[0])
+			{
+				// Do not add if this asset was assigned to the 0 chunk. These assets always exist on disk
+				bSkip = FinalChunkManifests[0]->Contains(PkgName);
+			}
+			if (!bSkip)
+			{
+				auto DependentPackageLongName = PkgName.ToString();
+				if (FPackageName::IsShortPackageName(PkgName))
+				{
+					DependentPackageLongName = FPackageName::ConvertToLongScriptPackageName(*PkgName.ToString());
+				}
+				FString DependentSandboxFile = SandboxPlatformFile->ConvertToAbsolutePathForExternalAppForWrite(*FPackageName::LongPackageNameToFilename(DependentPackageLongName));
+				ThisPackageSet->Add(PkgName, DependentSandboxFile);
+				UnassignedPackageSet.Remove(PkgName);
+			}
+		}
+	}
+}
+
 void FChunkManifestGenerator::FixupPackageDependenciesForChunks(FSandboxPlatformFile* InSandboxFile)
 {
 	for (int32 ChunkID = 0, MaxChunk = ChunkManifests.Num(); ChunkID < MaxChunk; ++ChunkID)
@@ -1005,44 +1026,6 @@ void FChunkManifestGenerator::FixupPackageDependenciesForChunks(FSandboxPlatform
 	}
 }
 
-void FChunkManifestGenerator::AddPackageAndDependenciesToChunk(FChunkPackageSet* ThisPackageSet, FName InPkgName, const FString& InSandboxFile, int32 ChunkID, FSandboxPlatformFile* SandboxPlatformFile)
-{
-	//Add this asset
-	ThisPackageSet->Add(InPkgName, InSandboxFile);
-
-	//Only gather dependencies if we're chunking
-	if (!bGenerateChunks)
-	{
-		return;
-	}
-
-	//now add any dependencies
-	TArray<FName> DependentPackageNames;
-	if (GatherAllPackageDependencies(InPkgName, DependentPackageNames))
-	{
-		for (const auto& PkgName : DependentPackageNames)
-		{
-			bool bSkip = false;
-			if (ChunkID != 0 && FinalChunkManifests[0])
-			{
-				// Do not add if this asset was assigned to the 0 chunk. These assets always exist on disk
-				bSkip = FinalChunkManifests[0]->Contains(PkgName);
-			}
-			if (!bSkip)
-			{
-				auto DependentPackageLongName = PkgName.ToString();
-				if (FPackageName::IsShortPackageName(PkgName))
-				{
-					DependentPackageLongName = FPackageName::ConvertToLongScriptPackageName(*PkgName.ToString());
-				}
-				FString DependentSandboxFile = SandboxPlatformFile->ConvertToAbsolutePathForExternalAppForWrite(*FPackageName::LongPackageNameToFilename(DependentPackageLongName));
-				ThisPackageSet->Add(PkgName, DependentSandboxFile);
-				UnassignedPackageSet.Remove(PkgName);
-			}
-		}
-	}
-}
-
 
 void FChunkManifestGenerator::FindShortestReferenceChain(TArray<FReferencePair> PackageNames, int32 ChunkID, uint32& OutParentIndex, FString& OutChainPath)
 {
diff --git a/Engine/Source/Editor/UnrealEd/Private/Commandlets/ChunkManifestGenerator.h b/Engine/Source/Editor/UnrealEd/Private/Commandlets/ChunkManifestGenerator.h
--- a/Engine/Source/Editor/UnrealEd/Private/Commandlets/ChunkManifestGenerator.h
+++ b/Engine/Source/Editor/UnrealEd/Private/Commandlets/ChunkManifestGenerator.h
@@ -10,6 +10,7 @@ struct FChunkDependencyTreeNode;
  */
 class FChunkManifestGenerator
 {
+	TArray<FName> StartupPackages;
 	/** Map of Package name to Sandbox Paths */
 	typedef TMap<FName, FString> FChunkPackageSet;
 	/** Holds a reference to asset registry */
@@ -77,22 +78,6 @@ class FChunkManifestGenerator
 	*/
 	void RemovePackageFromManifest(FName PackageName, int32 ChunkId);
 
-	/**
-	 * Adds a package to the unassigned package list for managing later
-	 *
-	 * @param The sandbox filepath of the package
-	 * @param The package name
-	 */
-	void NotifyPackageWasNotAssigned(const FString& PackageSandboxPath, FName PackageName);
-
-	/**
-	 * Adds a package to the global list of cooked packages.
-	 *
-	 * @param The sandbox filepath of the package
-	 * @param The package name
-	 */
-	void NotifyPackageWasCooked(const FString& PackageSandboxPath, FName PackageName);
-
 	/**
 	 * Walks the dependency graph of assets and assigns packages to correct chunks.
 	 * 
@@ -100,22 +85,7 @@ class FChunkManifestGenerator
 	 */
 	void FixupPackageDependenciesForChunks(FSandboxPlatformFile* SandboxFile);
 
-	/**
-	 * Adds a package and all its dependencies to a Chunk
-	 *
-	 * @param The Package set to add to
-	 * @param The Name of the package to add
-	 * @param The sandbox path of the asset
-	 * @param The ID of the Chunk being added to
-	 * @param the SandboxPlatformFile used during cook
-	 */
-	void AddPackageAndDependenciesToChunk(FChunkPackageSet* ChunkPackageSet, FName PkgName, const FString& SandboxFile, int32 ChunkID, FSandboxPlatformFile* SandboxPlatformFile);
-
-	/**
-	 * Callback for FCoreDelegates::FOnAssetLoaded delegate.
-	 * Collects all assets loaded with the last package.
-	 */
-	void OnAssetLoaded(UObject* Asset);
+	void AddPackageAndDependenciesToChunk(FChunkPackageSet* ThisPackageSet, FName InPkgName, const FString& InSandboxFile, int32 ChunkID, FSandboxPlatformFile* SandboxPlatformFile);
 
 	/**
 	 * Returns the path of the temporary packaging directory for the specified platform.
@@ -241,7 +211,29 @@ public:
 	/**
 	 * Initializes manifest generator - creates manifest lists, hooks up delegates.
 	 */
-	void Initialize(bool InGenerateChunks);
+	void Initialize(const TArray<FName> &StartupPackages);
+
+
+	const TArray<ITargetPlatform*>& GetTargetPlatforms() const { return Platforms; }
+
+	/**
+	 * GenerateChunkManifest 
+	 * generate chunk manifest for the packages passed in using the asset registry to determine dependencies
+	 *
+	 * @param StartupPackages list of startup packages which are forced into chunk 0
+	 * @param CookedPackages list of packages which were cooked
+	 * @param bGenerateStreamingInstallManifest should we build a streaming install manifest 
+	 */
+	void BuildChunkManifest(const TArray<FName>& CookedPackages, FSandboxPlatformFile* SandboxFile, bool bGenerateStreamingInstallManifest);
+
+	/**
+	 * ContainsMap
+	 * Does this package contain a map file (determined by finding if this package contains a UWorld / ULevel object)
+	 *
+	 * @param PackageName long package name of the package we want to determine if contains a map 
+	 * @return return if the package contains a UWorld / ULevel object (contains a map)
+	 */
+	bool ContainsMap(const FName& PackageName) const;
 
 	/**
 	 * Adds a package to chunk manifest (just calls the other AddPackageToChunkManifestFunction with more parameters)
@@ -251,10 +243,11 @@ public:
 	 * @param LastLoadedMapName Name of the last loaded map (can be empty)
 	 * @param the SandboxPlatformFile used during cook
 	 */
-	void AddPackageToChunkManifest(UPackage* Package, const FString& SandboxFilename, const FString& LastLoadedMapName, FSandboxPlatformFile* SandboxFile);
 	void AddPackageToChunkManifest(const FName& PackageFName, const FString& PackagePathName, const FString& SandboxFilename, const FString& LastLoadedMapName, FSandboxPlatformFile* InSandboxFile);
 	
 
+	void GenerateChunkManifestForPackage(const FName& PackageFName, const FString& PackagePathName, const FString& SandboxFilename, const FString& LastLoadedMapName, FSandboxPlatformFile* InSandboxFile);
+
 	/**
 	 * Add a package to the manifest but don't assign it to any chunk yet, packages which are not assigned by the end of the cook will be put into chunk 0
 	 * 
@@ -263,27 +256,6 @@ public:
 	 */
 	void AddUnassignedPackageToManifest(UPackage* Package, const FString& PackageSandboxPath );
 
-	/**
-	 * Collects all the packages loaded 
-	 *
-	 * @param Package Package which was loaded 
-	 */
-	void OnLastPackageLoaded( UPackage* Package );
-
-
-	/**
-	 * Collects all the packages loaded 
-	 * Does the same as other overload of OnLastPackageLoaded except takes in the name of the package as a fname 
-	 * 
-	 * @param Package name
-	 */
-	void OnLastPackageLoaded( const FName& PackageName );
-
-	/**
-	 * The cooker is about to load a new package from the list, reset AssetsLoadedWithLastPackage
-	 */
-	void PrepareToLoadNewPackage(const FString& Filename);
-
 	/**
 	 * Deletes temporary manifest directories.
 	 */
diff --git a/Engine/Source/Editor/UnrealEd/Private/CookOnTheFlyServer.cpp b/Engine/Source/Editor/UnrealEd/Private/CookOnTheFlyServer.cpp
--- a/Engine/Source/Editor/UnrealEd/Private/CookOnTheFlyServer.cpp
+++ b/Engine/Source/Editor/UnrealEd/Private/CookOnTheFlyServer.cpp
@@ -863,7 +863,19 @@ void UCookOnTheFlyServer::GetDependentPackages( const TSet<FName>& RootPackages,
 	while ( FoundPackagesCounter < FoundPackagesArray.Num() )
 	{
 		TArray<FName> PackageDependencies;
-		verify( AssetRegistry.GetDependencies(FoundPackagesArray[FoundPackagesCounter], PackageDependencies) );
+		if (AssetRegistry.GetDependencies(FoundPackagesArray[FoundPackagesCounter], PackageDependencies) == false)
+		{
+			// this could happen if we are in the editor and the dependency list is not up to date
+
+			if (IsCookingInEditor() == false)
+			{
+				UE_LOG(LogCook, Fatal, TEXT("Unable to find package %s in asset registry.  Can't generate cooked asset registry"), *FoundPackagesArray[FoundPackagesCounter].ToString());
+			}
+			else
+			{
+				UE_LOG(LogCook, Warning, TEXT("Unable to find package %s in asset registry, cooked asset registry information may be invalid "), *FoundPackagesArray[FoundPackagesCounter].ToString());
+			}
+		}
 		++FoundPackagesCounter;
 		for ( const auto& OriginalPackageDependency : PackageDependencies )
 		{
@@ -985,34 +997,6 @@ bool UCookOnTheFlyServer::ContainsMap(const FName& PackageName) const
 	return false;
 }
 
-
-// don't need the package loaded to generate the manifest infomation
-void UCookOnTheFlyServer::GenerateManifestInfo(const FName& Package, const TArray<FName>& TargetPlatformNames)
-{
-	if (IsCookByTheBookMode())
-		return;
-	
-	check(IsChildCooker() == false);
-
-	TSet<FName> RootPackages;
-	RootPackages.Add(Package);
-
-	// see if teh package contains a world
-	bool bIsWorldPackage = ContainsMap(Package);
-
-	TSet<FName> FoundPackages;
-	GetDependentPackages(RootPackages, FoundPackages);
-
-
-	FName StandardFilename = GetCachedStandardPackageFileFName(Package);
-	FName MapName = NAME_None;
-	if (bIsWorldPackage)
-	{
-		MapName = Package;
-	}
-	AddDependenciesToManifest(StandardFilename, MapName, FoundPackages, TargetPlatformNames);
-}
-
 void UCookOnTheFlyServer::GenerateManifestInfo(UPackage* Package, const TArray<FName>& TargetPlatformNames)
 {
 	if (!CookByTheBookOptions)
@@ -1030,7 +1014,7 @@ void UCookOnTheFlyServer::GenerateManifestInfo(UPackage* Package, const TArray<F
 	// load sublevels
 	UWorld* World = UWorld::FindWorldInPackage(Package);
 
-#define VERIFY_LEVELS_IN_DEPENDENCIES_LIST 0
+#define VERIFY_LEVELS_IN_DEPENDENCIES_LIST 1
 #if VERIFY_LEVELS_IN_DEPENDENCIES_LIST
 	TArray<FString> LevelPackageNames;
 #endif
@@ -1117,70 +1101,7 @@ void UCookOnTheFlyServer::GenerateManifestInfo(UPackage* Package, const TArray<F
 #endif
 
 
-	
-	FName StandardFilename = GetCachedStandardPackageFileFName(Package);
-	FName MapName = NAME_None;
-	if (Package->ContainsMap())
-	{
-		MapName = Package->GetFName();
-	}
-	AddDependenciesToManifest(StandardFilename, MapName, Packages, TargetPlatformNames);
-	
 }
-
-void UCookOnTheFlyServer::AddDependenciesToManifest(const FName& StandardFilename, const FName& MapName, const TSet<FName>& Dependencies, const TArray<FName>& TargetPlatformNames)
-{
-	FString LastLoadedMapName;
-	if (MapName != NAME_None)
-	{
-		LastLoadedMapName = MapName.ToString();
-	}
-	// Add package dependencies to the cooked package manifest
-	if (CookByTheBookOptions &&
-		CookByTheBookOptions->bGenerateDependenciesForMaps &&
-		(MapName != NAME_None))
-	{
-		CookByTheBookOptions->MapDependencyGraph.Add(MapName, Dependencies);
-	}
-
-	// update the manifests with generated dependencies
-	for (const auto& PlatformName : TargetPlatformNames)
-	{
-		FChunkManifestGenerator* ManifestGenerator = CookByTheBookOptions->ManifestGenerators.FindChecked(PlatformName);
-
-		if (CookByTheBookOptions && CookByTheBookOptions->bGenerateStreamingInstallManifests)
-		{
-			ManifestGenerator->PrepareToLoadNewPackage(StandardFilename.ToString());
-		}
-
-		for (const auto& DependentPackage : Dependencies)
-		{
-			ManifestGenerator->OnLastPackageLoaded(DependentPackage);
-		}
-
-		for (const auto& DependentPackage : Dependencies)
-		{
-			FString Filename = GetCachedPackageFilename(DependentPackage);
-			if (!Filename.IsEmpty())
-			{
-				// Populate streaming install manifests
-				FString SandboxFilename = ConvertToFullSandboxPath(*Filename, true);
-				//UE_LOG(LogCook, Display, TEXT("Adding package to manifest %s, %s, %s"), *DependentPackage->GetName(), *SandboxFilename, *LastLoadedMapName);
-#if USEASSETREGISTRYFORDEPENDENTPACKAGES
-				// Determine the package name and path
-				// FString PackageName = FPackageName::FilenameToLongPackageName(PackageFilename);
-				// FString PackagePath = FPackageName::GetLongPackagePath(DependentPackage.ToString());
-				FString PackagePath = DependentPackage.ToString();
-
-				ManifestGenerator->AddPackageToChunkManifest(DependentPackage, PackagePath, SandboxFilename, LastLoadedMapName, SandboxFile.GetOwnedPointer());
-#else
-				ManifestGenerator->AddPackageToChunkManifest(DependentPackage, SandboxFilename, LastLoadedMapName, SandboxFile.GetOwnedPointer());
-#endif
-			}
-		}
-	}
-}
-
 bool UCookOnTheFlyServer::IsCookingInEditor() const
 {
 	return CurrentCookMode == ECookMode::CookByTheBookFromTheEditor || CurrentCookMode == ECookMode::CookOnTheFlyFromTheEditor;;
@@ -1715,24 +1636,6 @@ uint32 UCookOnTheFlyServer::TickCookOnTheSide( const float TimeSlice, uint32 &Co
 							continue;
 						ACCUMULATE_TIMER_STOP(UnsolicitedPackageAlreadyCooked);
 
-						if ( CookByTheBookOptions )
-						{
-							ACCUMULATE_TIMER_START(AddUnassignedPackageToManifest);
-							for ( const auto &TargetPlatform : AllTargetPlatformNames )
-							{
-								FChunkManifestGenerator*& Manifest = CookByTheBookOptions->ManifestGenerators.FindChecked(TargetPlatform);
-
-								FString Filename = GetCachedPackageFilename(Package);
-								if (!Filename.IsEmpty())
-								{
-									// Populate streaming install manifests
-									const FString& SandboxFilename = GetCachedSandboxFilename( Package, SandboxFile );
-									Manifest->AddUnassignedPackageToManifest(Package, SandboxFilename);
-								}
-							}
-							ACCUMULATE_TIMER_STOP(AddUnassignedPackageToManifest);
-						}
-
 						if ( StandardPackageFName != NAME_None ) // if we have name none that means we are in core packages or something...
 						{
 							// check if the package has already been saved
@@ -1830,13 +1733,6 @@ uint32 UCookOnTheFlyServer::TickCookOnTheSide( const float TimeSlice, uint32 &Co
 					continue;
 				}
 
-				if (IsChildCooker()==false)
-				{
-					SCOPE_TIMER(GenerateManifestInfo);
-					// update manifest with cooked package info
-					GenerateManifestInfo(PackagesToSave[I], AllTargetPlatformNames);
-				}
-
 				// precache platform data for next package 
 				UPackage *NextPackage = PackagesToSave[FMath::Min( PackagesToSave.Num()-1, I + 1 )];
 				UPackage *NextNextPackage = PackagesToSave[FMath::Min( PackagesToSave.Num()-1, I + 2 )];
@@ -1898,6 +1794,7 @@ uint32 UCookOnTheFlyServer::TickCookOnTheSide( const float TimeSlice, uint32 &Co
 					Timer.SavedPackage();
 				}
 
+				if ( IsCookingInEditor() == false )
 				{
 					SCOPE_TIMER(ClearAllCachedCookedPlatformData);
 					TArray<UObject*> ObjectsInPackage;
@@ -3369,7 +3266,10 @@ void UCookOnTheFlyServer::CollectFilesToCook(TArray<FName>& FilesInPath, const T
 
 		if (CookByTheBookOptions && CookByTheBookOptions->bGenerateDependenciesForMaps) 
 		{
-			CookByTheBookOptions->MapDependencyGraph.Add(FName(TEXT("ContentDirectoryAssets")), ContentDirectoryAssets);
+			for (auto& MapDependencyGraph : CookByTheBookOptions->MapDependencyGraphs)
+			{
+				MapDependencyGraph.Value.Add(FName(TEXT("ContentDirectoryAssets")), ContentDirectoryAssets);
+			}
 		}
 	}
 }
@@ -3505,6 +3405,21 @@ void UCookOnTheFlyServer::CookByTheBookFinished()
 
 		for ( auto& Manifest : CookByTheBookOptions->ManifestGenerators )
 		{
+			TArray<FName> CookedPackagesForPlatform;
+			CookedPackages.GetCookedFilesForPlatform(Manifest.Key, CookedPackagesForPlatform);
+
+			TArray<FName> LongCookedPackageNames;
+			for (const auto& CookedPackageFileName : CookedPackagesForPlatform)
+			{
+				FString LongPackageName;
+				verify(FPackageName::TryConvertFilenameToLongPackageName(CookedPackageFileName.ToString(), LongPackageName));
+				LongCookedPackageNames.Add(FName(*LongPackageName));
+			}
+
+
+			Manifest.Value->BuildChunkManifest(LongCookedPackageNames, SandboxFile, CookByTheBookOptions->bGenerateStreamingInstallManifests);
+
+
 			// Always try to save the manifests, this is required to make the asset registry work, but doesn't necessarily write a file
 			Manifest.Value->SaveManifests(SandboxFile.GetOwnedPointer());
 
@@ -3541,34 +3456,10 @@ void UCookOnTheFlyServer::CookByTheBookFinished()
 
 	if (CookByTheBookOptions->bGenerateDependenciesForMaps && IsChildCooker() == false)
 	{
-		auto& MapDependencyGraph = CookByTheBookOptions->MapDependencyGraph;
-
-		FString MapDependencyGraphFile = FPaths::GameDir()/ TEXT("MapDependencyGraph.json");
-		// dump dependency graph. 
-		FString DependencyString; 
-		DependencyString += "{";
-		for (auto& Ele : MapDependencyGraph)
-		{
-			TSet <FName>& Deps = Ele.Value;
-			FName MapName = Ele.Key;
-			DependencyString += TEXT("\t\"") + MapName.ToString() + TEXT("\" : \n\t[\n ") ;
-			for (auto& Val : Deps)
-			{
-				DependencyString += TEXT("\t\t\"") + Val.ToString() + TEXT("\",\n");
-			}
-			DependencyString.RemoveFromEnd(TEXT(",\n"));
-			DependencyString += TEXT("\n\t],\n") ;
-		}
-		DependencyString.RemoveFromEnd(TEXT(",\n"));
-		DependencyString += "\n}";
-
-		ITargetPlatformManagerModule& TPM = GetTargetPlatformManagerRef();
-		const TArray<ITargetPlatform*>& Platforms = TPM.GetCookingTargetPlatforms();
-
-		for (ITargetPlatform* Platform : Platforms)
+		for (auto& MapDependencyGraphIt : CookByTheBookOptions->MapDependencyGraphs)
 		{
-			FString CookedMapDependencyGraphFilePlatform =  ConvertToFullSandboxPath(MapDependencyGraphFile, true).Replace(TEXT("[Platform]"), *Platform->PlatformName());
-			FFileHelper::SaveStringToFile(DependencyString, *CookedMapDependencyGraphFilePlatform, FFileHelper::EEncodingOptions::ForceUnicode);
+			BuildMapDependencyGraph(MapDependencyGraphIt.Key);
+			WriteMapDependencyGraph(MapDependencyGraphIt.Key);
 		}
 	}
 
@@ -3583,6 +3474,81 @@ void UCookOnTheFlyServer::CookByTheBookFinished()
 	OUTPUT_HIERARCHYTIMERS();
 }
 
+void UCookOnTheFlyServer::BuildMapDependencyGraph(const FName& PlatformName)
+{
+	auto& MapDependencyGraph = CookByTheBookOptions->MapDependencyGraphs.FindChecked(PlatformName);
+
+
+
+	TArray<FName> PlatformCookedPackages;
+	CookedPackages.GetCookedFilesForPlatform(PlatformName, PlatformCookedPackages);
+
+	static const FName AssetRegistryName("AssetRegistry");
+	FAssetRegistryModule& AssetRegistryModule = FModuleManager::LoadModuleChecked<FAssetRegistryModule>(AssetRegistryName);
+	IAssetRegistry& AssetRegistry = AssetRegistryModule.Get();
+
+
+	// assign chunks for all the map packages
+	for (const auto& CookedPackage : PlatformCookedPackages)
+	{
+		
+		TArray<FAssetData> PackageAssets;
+		ensure(AssetRegistry.GetAssetsByPackageName(CookedPackage, PackageAssets));
+
+		bool bIsMap = false;
+		for (const auto& Asset : PackageAssets)
+		{
+			if (Asset.GetClass()->IsChildOf(UWorld::StaticClass()) ||
+				Asset.GetClass()->IsChildOf(ULevel::StaticClass()))
+			{
+				bIsMap = true;
+			}
+		}
+
+
+		// if we are not a map we don't care about this pacakge
+		if (!bIsMap)
+		{
+			continue;
+		}
+
+		TSet<FName> DependentPackages;
+		for (const auto& Asset : PackageAssets)
+		{
+			DependentPackages.Add(Asset.PackageName);
+		}
+
+		MapDependencyGraph.Add(CookedPackage, DependentPackages);
+	}
+}
+
+void UCookOnTheFlyServer::WriteMapDependencyGraph(const FName& PlatformName)
+{
+	auto& MapDependencyGraph = CookByTheBookOptions->MapDependencyGraphs.FindChecked(PlatformName);
+
+	FString MapDependencyGraphFile = FPaths::GameDir() / TEXT("MapDependencyGraph.json");
+	// dump dependency graph. 
+	FString DependencyString;
+	DependencyString += "{";
+	for (auto& Ele : MapDependencyGraph)
+	{
+		TSet <FName>& Deps = Ele.Value;
+		FName MapName = Ele.Key;
+		DependencyString += TEXT("\t\"") + MapName.ToString() + TEXT("\" : \n\t[\n ");
+		for (auto& Val : Deps)
+		{
+			DependencyString += TEXT("\t\t\"") + Val.ToString() + TEXT("\",\n");
+		}
+		DependencyString.RemoveFromEnd(TEXT(",\n"));
+		DependencyString += TEXT("\n\t],\n");
+	}
+	DependencyString.RemoveFromEnd(TEXT(",\n"));
+	DependencyString += "\n}";
+
+	FString CookedMapDependencyGraphFilePlatform = ConvertToFullSandboxPath(MapDependencyGraphFile, true).Replace(TEXT("[Platform]"), *PlatformName.ToString());
+	FFileHelper::SaveStringToFile(DependencyString, *CookedMapDependencyGraphFilePlatform, FFileHelper::EEncodingOptions::ForceUnicode);
+}
+
 void UCookOnTheFlyServer::QueueCancelCookByTheBook()
 {
 	if ( IsCookByTheBookMode() )
@@ -3690,6 +3656,16 @@ void UCookOnTheFlyServer::TermSandbox()
 	SandboxFile = NULL;
 }
 
+void UCookOnTheFlyServer::ValidateCookByTheBookSettings() const
+{
+	if (IsChildCooker())
+	{
+		// should never be generating dependency maps / streaming install manifests for child cookers
+		check(CookByTheBookOptions->bGenerateDependenciesForMaps == false);
+		check(CookByTheBookOptions->bGenerateStreamingInstallManifests == false);
+	}
+}
+
 void UCookOnTheFlyServer::StartCookByTheBook( const FCookByTheBookStartupOptions& CookByTheBookStartupOptions )
 {
 
@@ -3722,6 +3698,7 @@ void UCookOnTheFlyServer::StartCookByTheBook( const FCookByTheBookStartupOptions
 	bool bMapsOnly = (CookOptions & ECookByTheBookOptions::MapsOnly) != ECookByTheBookOptions::None;
 	bool bNoDev = (CookOptions & ECookByTheBookOptions::NoDevContent) != ECookByTheBookOptions::None;
 
+	ValidateCookByTheBookSettings();
 
 	if ( CookByTheBookOptions->DlcName != DLCName )
 	{
@@ -3776,19 +3753,22 @@ void UCookOnTheFlyServer::StartCookByTheBook( const FCookByTheBookStartupOptions
 				Platforms.Add(Platform);
 				ManifestGenerator = new FChunkManifestGenerator(Platforms);
 				ManifestGenerator->CleanManifestDirectories();
-				ManifestGenerator->Initialize(CookByTheBookOptions->bGenerateStreamingInstallManifests);
 
+				TArray<FName> StartupPackages;
 
-				if (IsCookFlagSet(ECookInitializationFlags::Iterative) && (PackagesKeptFromPreviousCook.Num() > 0))
+				for (TObjectIterator<UPackage> It; It; ++It)
 				{
-					const FString& SandboxRegistryFilename = GetSandboxAssetRegistryFilename();
-
-					ManifestGenerator->LoadAssetRegistry(GetCookedAssetRegistryFilename(PlatformName.ToString()), &PackagesKeptFromPreviousCook);
+					StartupPackages.Add( It->GetFName());
 				}
 
-
+				ManifestGenerator->Initialize(StartupPackages);
 				CookByTheBookOptions->ManifestGenerators.Add(PlatformName, ManifestGenerator);
 			}
+
+			if (CookByTheBookOptions->bGenerateDependenciesForMaps)
+			{
+				CookByTheBookOptions->MapDependencyGraphs.Add(PlatformName);
+			}
 		}
 	}
 	if ( IsCookingDLC() )
@@ -4061,25 +4041,6 @@ void UCookOnTheFlyServer::StartChildCookers(int32 NumCookersToSpawn, const TArra
 			FName StandardPackageName = FName(*PackageName);*/
 			
 			CookedPackages.Add(FFilePlatformRequest(StandardPackageName, TargetPlatformNames));
-
-
-			
-			/*TArray<FName> Dependencies;
-			AssetRegistry.GetDependencies(PackageFName, Dependencies);
-			TSet<FName> DependenciesSet;
-			for (const auto& Dependency : Dependencies)
-			{
-				DependenciesSet.Add(Dependency);
-			}*/
-
-			TSet<FName> RootPackages;
-			RootPackages.Add(PackageFName);
-			TSet<FName> FoundPackages;
-			GetDependentPackages(RootPackages, FoundPackages);
-
-
-			FName MapName = NAME_None; // to do figure out which map depends on this package for now this package will get put into chunk 0
-			AddDependenciesToManifest(StandardPackageName, MapName, FoundPackages, TargetPlatformNames);
 		}
 		DistributeStandardFilenames.RemoveAt(0, NumFilesForCooker);
 
