Removed Manifest generation from the cook commandlet also.

diff --git a/Engine/Source/Editor/UnrealEd/Private/Commandlets/CookCommandlet.cpp b/Engine/Source/Editor/UnrealEd/Private/Commandlets/CookCommandlet.cpp
--- a/Engine/Source/Editor/UnrealEd/Private/Commandlets/CookCommandlet.cpp
+++ b/Engine/Source/Editor/UnrealEd/Private/Commandlets/CookCommandlet.cpp
@@ -1354,7 +1354,7 @@ bool UCookCommandlet::Cook(const TArray<ITargetPlatform*>& Platforms, TArray<FSt
 	FChunkManifestGenerator ManifestGenerator(Platforms);
 	// Always clean manifest directories so that there's no stale data
 	ManifestGenerator.CleanManifestDirectories();
-	ManifestGenerator.Initialize(bGenerateStreamingInstallManifests);
+	// ManifestGenerator.Initialize(bGenerateStreamingInstallManifests);
 
 	for( int32 FileIndex = 0; ; FileIndex++ )
 	{
@@ -1381,7 +1381,7 @@ bool UCookCommandlet::Cook(const TArray<ITargetPlatform*>& Platforms, TArray<FSt
 					// Populate streaming install manifests
 					FString SandboxFilename = SandboxFile->ConvertToAbsolutePathForExternalAppForWrite(*Filename);
 					UE_LOG(LogCookCommandlet, Display, TEXT("Adding package to manifest %s, %s, %s"), *Pkg->GetName(), *SandboxFilename, *LastLoadedMapName);
-					ManifestGenerator.AddPackageToChunkManifest(Pkg, SandboxFilename, LastLoadedMapName, SandboxFile.GetOwnedPointer());
+					//ManifestGenerator.AddPackageToChunkManifest(Pkg, SandboxFilename, LastLoadedMapName, SandboxFile.GetOwnedPointer());
 				}
 					
 				if (!CookedPackages.Contains(Filename))
@@ -1470,12 +1470,6 @@ bool UCookCommandlet::Cook(const TArray<ITargetPlatform*>& Platforms, TArray<FSt
 
 		UE_LOG(LogCookCommandlet, Display, TEXT("Loading %s"), *Filename );
 
-		if (bGenerateStreamingInstallManifests)
-		{
-			UE_LOG(LogCookCommandlet, Display, TEXT("PrepareToLoadNewPackage %s"), *Filename );
-			ManifestGenerator.PrepareToLoadNewPackage(Filename);
-		}
-
 		UPackage* Package = LoadPackage( NULL, *Filename, LOAD_None );
 
 		if( Package == NULL )
