Editor: Fixed shadowed function parameter

diff --git a/Engine/Source/Editor/UnrealEd/Private/Commandlets/ChunkManifestGenerator.cpp b/Engine/Source/Editor/UnrealEd/Private/Commandlets/ChunkManifestGenerator.cpp
--- a/Engine/Source/Editor/UnrealEd/Private/Commandlets/ChunkManifestGenerator.cpp
+++ b/Engine/Source/Editor/UnrealEd/Private/Commandlets/ChunkManifestGenerator.cpp
@@ -325,7 +325,7 @@ void FChunkManifestGenerator::Initialize(const TArray<FName> &InStartupPackages)
 }
 
 
-void FChunkManifestGenerator::BuildChunkManifest(const TArray<FName>& CookedPackages, FSandboxPlatformFile* SandboxFile, bool bGenerateStreamingInstallManifest)
+void FChunkManifestGenerator::BuildChunkManifest(const TArray<FName>& CookedPackages, FSandboxPlatformFile* InSandboxFile, bool bGenerateStreamingInstallManifest)
 {
 	bGenerateChunks = bGenerateStreamingInstallManifest;
 
@@ -373,7 +373,7 @@ void FChunkManifestGenerator::BuildChunkManifest(const TArray<FName>& CookedPack
 	// add all the packages to the unassigned package list
 	for (const auto& CookedPackage : CookedPackages)
 	{
-		const FString SandboxPath = SandboxFile->ConvertToAbsolutePathForExternalAppForWrite(*FPackageName::LongPackageNameToFilename(CookedPackage.ToString()));
+		const FString SandboxPath = InSandboxFile->ConvertToAbsolutePathForExternalAppForWrite(*FPackageName::LongPackageNameToFilename(CookedPackage.ToString()));
 
 		AllCookedPackages.Add(CookedPackage, SandboxPath);
 		UnassignedPackageSet.Add(CookedPackage, SandboxPath);
@@ -381,7 +381,7 @@ void FChunkManifestGenerator::BuildChunkManifest(const TArray<FName>& CookedPack
 	
 	for (const auto& CookedPackage : StartupPackages)
 	{
-		const FString SandboxPath = SandboxFile->ConvertToAbsolutePathForExternalAppForWrite(*FPackageName::LongPackageNameToFilename(CookedPackage.ToString()));
+		const FString SandboxPath = InSandboxFile->ConvertToAbsolutePathForExternalAppForWrite(*FPackageName::LongPackageNameToFilename(CookedPackage.ToString()));
 
 		AllCookedPackages.Add(CookedPackage, SandboxPath);
 		AddPackageToManifest(SandboxPath, CookedPackage, 0);
@@ -435,7 +435,7 @@ void FChunkManifestGenerator::BuildChunkManifest(const TArray<FName>& CookedPack
 			const FString& SandboxFilename = AllCookedPackages.FindChecked(PackageFName);
 
 
-			GenerateChunkManifestForPackage(PackageFName, PackagePathName, SandboxFilename, MapName, SandboxFile);
+			GenerateChunkManifestForPackage(PackageFName, PackagePathName, SandboxFilename, MapName, InSandboxFile);
 
 
 		}
@@ -450,7 +450,7 @@ void FChunkManifestGenerator::BuildChunkManifest(const TArray<FName>& CookedPack
 		const FString& SandboxFilename = AllCookedPackages.FindChecked(PackageFName);
 		const FString PackagePathName = PackageFName.ToString();
 
-		GenerateChunkManifestForPackage(PackageFName, PackagePathName, SandboxFilename, FString(), SandboxFile);
+		GenerateChunkManifestForPackage(PackageFName, PackagePathName, SandboxFilename, FString(), InSandboxFile);
 	}
 
 
diff --git a/Engine/Source/Editor/UnrealEd/Private/Commandlets/ChunkManifestGenerator.h b/Engine/Source/Editor/UnrealEd/Private/Commandlets/ChunkManifestGenerator.h
--- a/Engine/Source/Editor/UnrealEd/Private/Commandlets/ChunkManifestGenerator.h
+++ b/Engine/Source/Editor/UnrealEd/Private/Commandlets/ChunkManifestGenerator.h
@@ -81,9 +81,9 @@ class FChunkManifestGenerator
 	/**
 	 * Walks the dependency graph of assets and assigns packages to correct chunks.
 	 * 
-	 * @param the SandboxPlatformFile used during cook
+	 * @param the InSandboxFile used during cook
 	 */
-	void FixupPackageDependenciesForChunks(FSandboxPlatformFile* SandboxFile);
+	void FixupPackageDependenciesForChunks(FSandboxPlatformFile* InSandboxFile);
 
 	void AddPackageAndDependenciesToChunk(FChunkPackageSet* ThisPackageSet, FName InPkgName, const FString& InSandboxFile, int32 ChunkID, FSandboxPlatformFile* SandboxPlatformFile);
 
@@ -224,7 +224,7 @@ public:
 	 * @param CookedPackages list of packages which were cooked
 	 * @param bGenerateStreamingInstallManifest should we build a streaming install manifest 
 	 */
-	void BuildChunkManifest(const TArray<FName>& CookedPackages, FSandboxPlatformFile* SandboxFile, bool bGenerateStreamingInstallManifest);
+	void BuildChunkManifest(const TArray<FName>& CookedPackages, FSandboxPlatformFile* InSandboxFile, bool bGenerateStreamingInstallManifest);
 
 	/**
 	 * ContainsMap
@@ -264,9 +264,9 @@ public:
 	/**
 	 * Saves all generated manifests for each target platform.
 	 * 
-	 * @param the SandboxPlatformFile used during cook
+	 * @param the InSandboxFile used during cook
 	 */
-	bool SaveManifests(FSandboxPlatformFile* SandboxFile);
+	bool SaveManifests(FSandboxPlatformFile* InSandboxFile);
 
 	/**
 	* Saves generated asset registry data for each platform.
