Fix for editor going black during cook in editor.

diff --git a/Engine/Source/Editor/UnrealEd/Private/CookOnTheFlyServer.cpp b/Engine/Source/Editor/UnrealEd/Private/CookOnTheFlyServer.cpp
--- a/Engine/Source/Editor/UnrealEd/Private/CookOnTheFlyServer.cpp
+++ b/Engine/Source/Editor/UnrealEd/Private/CookOnTheFlyServer.cpp
@@ -2257,11 +2257,6 @@ bool UCookOnTheFlyServer::SaveCookedPackage( UPackage* Package, uint32 SaveFlags
 					Package->PackageFlags &= ~PKG_FilterEditorOnly;
 				}
 
-				if (World)
-				{
-					World->PersistentLevel->OwningWorld = World;
-				}
-
 				// need to subtract 32 because the SavePackage code creates temporary files with longer file names then the one we provide
 				// projects may ignore this restriction if desired
 				static bool bConsiderCompressedPackageFileLengthRequirements = ShouldConsiderCompressedPackageFileLengthRequirements();
@@ -2279,7 +2274,7 @@ bool UCookOnTheFlyServer::SaveCookedPackage( UPackage* Package, uint32 SaveFlags
 					bSavedCorrectly &= GEditor->SavePackage( Package, World, Flags, *PlatFilename, GError, NULL, bSwap, false, SaveFlags, Target, FDateTime::MinValue(), false );
 					INC_INT_STAT(SavedPackage, 1);
 				}
-				
+
 				bOutWasUpToDate = false;
 			}
 			else
diff --git a/Engine/Source/Editor/UnrealEd/Private/EditorEngine.cpp b/Engine/Source/Editor/UnrealEd/Private/EditorEngine.cpp
--- a/Engine/Source/Editor/UnrealEd/Private/EditorEngine.cpp
+++ b/Engine/Source/Editor/UnrealEd/Private/EditorEngine.cpp
@@ -3683,6 +3683,8 @@ bool UEditorEngine::SavePackage( UPackage* InOuter, UObject* InBase, EObjectFlag
 
 	UWorld* World = Cast<UWorld>(Base);
 	bool bInitializedPhysicsSceneForSave = false;
+	
+	UWorld *OriginalOwningWorld = nullptr;
 	if ( World )
 	{
 		// We need a physics scene at save time in case code does traces during onsave events.
@@ -3700,6 +3702,9 @@ bool UEditorEngine::SavePackage( UPackage* InOuter, UObject* InBase, EObjectFlag
 			bHasPhysicsScene = (World->GetPhysicsScene() != nullptr);
 		}
 
+		OriginalOwningWorld = World->PersistentLevel->OwningWorld;
+		World->PersistentLevel->OwningWorld = World;
+
 		// If we didn't find any physics scene we will synthesize one and remove it after save
 		if (!bHasPhysicsScene)
 		{
@@ -3767,6 +3772,11 @@ bool UEditorEngine::SavePackage( UPackage* InOuter, UObject* InBase, EObjectFlag
 	{
 		OnPostSaveWorld(SaveFlags, World, OriginalPackageFlags, bSuccess);
 
+		if (OriginalOwningWorld)
+		{
+			World->PersistentLevel->OwningWorld = OriginalOwningWorld;
+		}
+
 		if (bInitializedPhysicsSceneForSave)
 		{
 			// Make sure we clean up the physics scene here. If we leave too many scenes in memory, undefined behavior occurs when locking a scene for read/write.
