[UE-15077] Offering the option to compile Blueprints before you launch on a device (so pending changes are properly reflected when you cook-in-editor).

diff --git a/Engine/Source/Editor/UnrealEd/Classes/Settings/LevelEditorPlaySettings.h b/Engine/Source/Editor/UnrealEd/Classes/Settings/LevelEditorPlaySettings.h
--- a/Engine/Source/Editor/UnrealEd/Classes/Settings/LevelEditorPlaySettings.h
+++ b/Engine/Source/Editor/UnrealEd/Classes/Settings/LevelEditorPlaySettings.h
@@ -226,6 +226,10 @@ public:
 	UPROPERTY(config, EditAnywhere, Category = PlayOnDevice)
 	TEnumAsByte<EPlayOnBuildMode> BuildGameBeforeLaunch;
 
+	/** Whether to automatically recompile dirty Blueprints before launching */
+	UPROPERTY(config, EditAnywhere, Category=PlayOnDevice)
+	bool bAutoCompileBlueprintsOnLaunch;
+
 private:
 
 	/** NetMode to use for Play In Editor. */
diff --git a/Engine/Source/Editor/UnrealEd/Private/PlayLevel.cpp b/Engine/Source/Editor/UnrealEd/Private/PlayLevel.cpp
--- a/Engine/Source/Editor/UnrealEd/Private/PlayLevel.cpp
+++ b/Engine/Source/Editor/UnrealEd/Private/PlayLevel.cpp
@@ -1479,6 +1479,130 @@ static void HandleHyperlinkNavigate()
 	FGlobalTabmanager::Get()->InvokeTab(FName("OutputLog"));
 }
 
+struct FInternalPlayLevelUtils
+{
+	static int32 ResolveDirtyBlueprints(const bool bPromptForCompile, TArray<UBlueprint*>& ErroredBlueprints)
+	{
+		const bool bAutoCompile = !bPromptForCompile;
+		FString PromtDirtyList;
+
+		TArray<UBlueprint*> InNeedOfRecompile;
+		ErroredBlueprints.Empty();
+
+		double BPRegenStartTime = FPlatformTime::Seconds();
+		for (TObjectIterator<UBlueprint> BlueprintIt; BlueprintIt; ++BlueprintIt)
+		{
+			UBlueprint* Blueprint = *BlueprintIt;
+
+			// do not try to recompile BPs that have not changed since they last failed to compile, so don't check Blueprint->IsUpToDate()
+			const bool bIsDirtyAndShouldBeRecompiled = Blueprint->IsPossiblyDirty();
+			if (!FBlueprintEditorUtils::IsDataOnlyBlueprint(Blueprint)
+				&& (bIsDirtyAndShouldBeRecompiled || FBlueprintEditorUtils::IsLevelScriptBlueprint(Blueprint))
+				&& (Blueprint->Status != BS_Unknown)
+				&& !Blueprint->IsPendingKill())
+			{
+				InNeedOfRecompile.Add(Blueprint);
+
+				if (bPromptForCompile)
+				{
+					PromtDirtyList += FString::Printf(TEXT("\n   %s"), *Blueprint->GetName());
+				}
+			}
+			else if (BS_Error == Blueprint->Status && Blueprint->bDisplayCompilePIEWarning)
+			{
+				ErroredBlueprints.Add(Blueprint);
+			}
+		}
+
+		bool bRunCompilation = bAutoCompile;
+		if (bPromptForCompile)
+		{
+			FFormatNamedArguments Args;
+			Args.Add(TEXT("DirtyBlueprints"), FText::FromString(PromtDirtyList));
+			const FText PromptMsg = FText::Format(NSLOCTEXT("PlayInEditor", "PrePIE_BlueprintsDirty", "One or more blueprints have been modified without being recompiled. Do you want to compile them now? \n{DirtyBlueprints}"), Args);
+
+			EAppReturnType::Type PropmtResponse = FMessageDialog::Open(EAppMsgType::YesNo, PromptMsg);
+			bRunCompilation = (PropmtResponse == EAppReturnType::Yes);
+		}
+		int32 RecompiledCount = 0;
+
+		FMessageLog BlueprintLog("BlueprintLog");
+		if (bRunCompilation && (InNeedOfRecompile.Num() > 0))
+		{
+			const FText LogPageLabel = (bAutoCompile) ? LOCTEXT("BlueprintAutoCompilationPageLabel", "Pre-Play auto-recompile") :
+				LOCTEXT("BlueprintCompilationPageLabel", "Pre-Play recompile");
+			BlueprintLog.NewPage(LogPageLabel);
+
+			// Recompile all necessary blueprints in a single loop, saving GC until the end
+			for (auto BlueprintIt = InNeedOfRecompile.CreateIterator(); BlueprintIt; ++BlueprintIt)
+			{
+				UBlueprint* Blueprint = *BlueprintIt;
+
+				int32 CurrItIndex = BlueprintIt.GetIndex();
+				// gather dependencies so we can ensure that they're getting recompiled as well
+				TArray<UBlueprint*> Dependencies;
+				FBlueprintEditorUtils::GetDependentBlueprints(Blueprint, Dependencies);
+				// if the user made a change, but didn't hit "compile", then dependent blueprints
+				// wouldn't have been marked dirty, so here we make sure to add those dependencies 
+				// to the end of the InNeedOfRecompile array (so we hit them too in this loop)
+				for (auto DependencyIt = Dependencies.CreateIterator(); DependencyIt; ++DependencyIt)
+				{
+					UBlueprint* DependentBp = *DependencyIt;
+
+					int32 ExistingIndex = InNeedOfRecompile.Find(DependentBp);
+					// if this dependent blueprint is already set up to compile 
+					// later in this loop, then there is no need to add it to be recompiled again
+					if (ExistingIndex >= CurrItIndex)
+					{
+						continue;
+					}
+
+					// if this blueprint wasn't slated to  be recompiled
+					if (ExistingIndex == INDEX_NONE)
+					{
+						// we need to make sure this gets recompiled as well 
+						// (since it depends on this other one that is dirty)
+						InNeedOfRecompile.Add(DependentBp);
+					}
+					// else this is a circular dependency... it has previously been compiled
+					// ... is there a case where we'd want to recompile this again?
+				}
+
+				Blueprint->BroadcastChanged();
+
+				UE_LOG(LogPlayLevel, Log, TEXT("[PlayLevel] Compiling %s before play..."), *Blueprint->GetName());
+				FKismetEditorUtilities::CompileBlueprint(Blueprint, /*bIsRegeneratingOnLoad =*/false, /*bSkipGarbageCollection =*/true);
+				const bool bHadError = (!Blueprint->IsUpToDate() && Blueprint->Status != BS_Unknown);
+
+				// Check if the Blueprint has already been added to the error list to prevent it from being added again
+				if (bHadError && ErroredBlueprints.Find(Blueprint) == INDEX_NONE)
+				{
+					ErroredBlueprints.Add(Blueprint);
+
+					FFormatNamedArguments Arguments;
+					Arguments.Add(TEXT("Name"), FText::FromString(Blueprint->GetName()));
+
+					BlueprintLog.Info(FText::Format(LOCTEXT("BlueprintCompileFailed", "Blueprint {Name} failed to compile"), Arguments));
+				}
+				else
+				{
+					++RecompiledCount;
+				}
+
+				UE_LOG(LogPlayLevel, Log, TEXT("PlayLevel: Blueprint regeneration took %d ms (%i blueprints)"), (int32)((FPlatformTime::Seconds() - BPRegenStartTime) * 1000), InNeedOfRecompile.Num());
+			}
+
+			CollectGarbage(GARBAGE_COLLECTION_KEEPFLAGS);
+		}
+		else if (bAutoCompile)
+		{
+			UE_LOG(LogPlayLevel, Log, TEXT("PlayLevel: No blueprints needed recompiling"));
+		}
+
+		return RecompiledCount;
+	}
+};
+
 void UEditorEngine::PlayUsingLauncher()
 {
 	if (!PlayUsingLauncherDeviceId.IsEmpty())
@@ -1497,9 +1621,10 @@ void UEditorEngine::PlayUsingLauncher()
 		bPlayUsingLauncherHasCode = GameProjectModule.Get().ProjectRequiresBuild(FName(*PlayUsingLauncherDeviceId.Left(PlayUsingLauncherDeviceId.Find(TEXT("@")))));
 		bPlayUsingLauncherHasCompiler = FSourceCodeNavigation::IsCompilerAvailable();
 
+		const ULevelEditorPlaySettings* PlayInSettings = GetDefault<ULevelEditorPlaySettings>();
 		// Setup launch profile, keep the setting here to a minimum.
 		ILauncherProfileRef LauncherProfile = LauncherServicesModule.CreateProfile(TEXT("Play On Device"));
-		EPlayOnBuildMode bBuildType = GetDefault<ULevelEditorPlaySettings>()->BuildGameBeforeLaunch;
+		EPlayOnBuildMode bBuildType = PlayInSettings->BuildGameBeforeLaunch;
 		if ((bBuildType == EPlayOnBuildMode::PlayOnBuild_Always) || (bBuildType == PlayOnBuild_Default && (bPlayUsingLauncherHasCode) && bPlayUsingLauncherHasCompiler))
 		{
 			LauncherProfile->SetBuildGame(true);
@@ -1560,6 +1685,9 @@ void UEditorEngine::PlayUsingLauncher()
 			LauncherProfile->SetDeploymentMode(ELauncherProfileDeploymentModes::FileServer);
 		}
 
+		TArray<UBlueprint*> ErroredBlueprints;
+		FInternalPlayLevelUtils::ResolveDirtyBlueprints(!PlayInSettings->bAutoCompileBlueprintsOnLaunch, ErroredBlueprints);
+
 		TArray<FString> MapNames;
 		FWorldContext & EditorContext = GetEditorWorldContext();
 		if (EditorContext.World()->WorldComposition || (LauncherProfile->GetCookMode() == ELauncherProfileCookModes::ByTheBookInEditor) || (LauncherProfile->GetCookMode() == ELauncherProfileCookModes::OnTheFlyInEditor) )
@@ -2002,205 +2130,41 @@ void UEditorEngine::PlayInEditor( UWorld* InWorld, bool bInSimulateInEditor )
 		InWorld->GetNavigationSystem()->OnPIEStart();
 	}
 
-	ULevelEditorPlaySettings const* EditorPlayInSettings = GetDefault<ULevelEditorPlaySettings>();
-	check(EditorPlayInSettings);
+	ULevelEditorPlaySettings* PlayInSettings = GetMutableDefault<ULevelEditorPlaySettings>();
+	check(PlayInSettings);
 
-	// Prompt the user to compile any dirty Blueprints before PIE can occur.
-	TArray< UBlueprint* > ErrorBlueprintList;
-	bool bAnyBlueprintsDirty = false;
-	{
-		FString DirtyBlueprints;
-		FString ErrorBlueprints;
+	TArray<UBlueprint*> ErroredBlueprints;
+	FInternalPlayLevelUtils::ResolveDirtyBlueprints(!PlayInSettings->AutoRecompileBlueprints, ErroredBlueprints);
 
-		TArray<UBlueprint*> BlueprintsToRecompile;
-
-		double BPRegenStartTime = FPlatformTime::Seconds();
-		for (TObjectIterator<UBlueprint> BlueprintIt; BlueprintIt; ++BlueprintIt)
+	if (ErroredBlueprints.Num() && !GIsDemoMode)
+	{
+		FString ErroredBlueprintList;
+		for (auto Blueprint : ErroredBlueprints)
 		{
-			UBlueprint* Blueprint = *BlueprintIt;
-
-			// If the blueprint isn't fresh, try to recompile it automatically
-			if( EditorPlayInSettings->AutoRecompileBlueprints )
-			{
-				// do not try to recompile BPs that have not changed since they last failed to compile, so don't check Blueprint->IsUpToDate()
-				const bool bIsDirtyAndShouldBeRecompiled = Blueprint->IsPossiblyDirty();
-				if( !FBlueprintEditorUtils::IsDataOnlyBlueprint(Blueprint) 
-					&& (bIsDirtyAndShouldBeRecompiled || FBlueprintEditorUtils::IsLevelScriptBlueprint(Blueprint))
-					&& (Blueprint->Status != BS_Unknown)
-					&& !Blueprint->IsPendingKill() )
-				{
-					BlueprintsToRecompile.Add(Blueprint);
-				}
-				else if(BS_Error == Blueprint->Status && Blueprint->bDisplayCompilePIEWarning)
-				{
-					ErrorBlueprintList.Add(Blueprint);
-					ErrorBlueprints += FString::Printf(TEXT("\n   %s"), *Blueprint->GetName());
-				}
-			}
-			else
-			{
-				// Record blueprints that are not fully recompiled or had an error
-				switch (Blueprint->Status)
-				{
-				case BS_Unknown:
-					// Treating unknown as up to date for right now
-					break;
-				case BS_Error:
-					if( Blueprint->bDisplayCompilePIEWarning)
-					{
-						ErrorBlueprintList.Add(Blueprint);
-						ErrorBlueprints += FString::Printf(TEXT("\n   %s"), *Blueprint->GetName());
-					}
-					break;
-				case BS_UpToDate:
-				case BS_UpToDateWithWarnings:
-					break;
-				default:
-				case BS_Dirty:
-					bAnyBlueprintsDirty = true;
-					DirtyBlueprints += FString::Printf(TEXT("\n   %s"), *Blueprint->GetName());
-					break;
-				}
-			}
+			ErroredBlueprintList += FString::Printf(TEXT("\n   %s"), *Blueprint->GetName());
 		}
 
-		FMessageLog BlueprintLog("BlueprintLog");
+		FFormatNamedArguments Args;
+		Args.Add(TEXT("ErrorBlueprints"), FText::FromString(ErroredBlueprintList));
 
-		if( EditorPlayInSettings->AutoRecompileBlueprints )
+		// There was at least one blueprint with an error, make sure the user is OK with that.
+		const bool bContinuePIE = EAppReturnType::Yes == FMessageDialog::Open( EAppMsgType::YesNo, FText::Format( NSLOCTEXT("PlayInEditor", "PrePIE_BlueprintErrors", "One or more blueprints has an unresolved compiler error, are you sure you want to Play in Editor?{ErrorBlueprints}"), Args ) );
+		if ( !bContinuePIE )
 		{
-			if( BlueprintsToRecompile.Num() > 0 )
-			{
-				BlueprintLog.NewPage(LOCTEXT("BlueprintAutoCompilationPageLabel", "Pre-PIE auto-recompile"));
-
-				// Recompile all necessary blueprints in a single loop, saving GC until the end
-				for( auto BlueprintIt = BlueprintsToRecompile.CreateIterator(); BlueprintIt; ++BlueprintIt )
-				{
-					UBlueprint* Blueprint = *BlueprintIt;
-
-					int32 CurrItIndex = BlueprintIt.GetIndex();
-					// gather dependencies so we can ensure that they're getting recompiled as well
-					TArray<UBlueprint*> Dependencies;
-					FBlueprintEditorUtils::GetDependentBlueprints(Blueprint, Dependencies);
-					// if the user made a change, but didn't hit "compile", then dependent blueprints
-					// wouldn't have been marked dirty, so here we make sure to add those dependencies 
-					// to the end of the BlueprintsToRecompile array (so we hit them too in this loop)
-					for (auto DependencyIt = Dependencies.CreateIterator(); DependencyIt; ++DependencyIt)
-					{
-						UBlueprint* DependentBp = *DependencyIt;
-
-						int32 ExistingIndex = BlueprintsToRecompile.Find(DependentBp);
-						// if this dependent blueprint is already set up to compile 
-						// later in this loop, then there is no need to add it to be recompiled again
-						if (ExistingIndex >= CurrItIndex)
-						{
-							continue;
-						}
-
-						// if this blueprint wasn't slated to  be recompiled
-						if (ExistingIndex == INDEX_NONE)
-						{
-							// we need to make sure this gets recompiled as well 
-							// (since it depends on this other one that is dirty)
-							BlueprintsToRecompile.Add(DependentBp);
-						}
-						// else this is a circular dependency... it has previously been compiled
-						// ... is there a case where we'd want to recompile this again?
-					}
-
-					Blueprint->BroadcastChanged();
-
-					UE_LOG(LogPlayLevel, Log, TEXT("[PIE] Compiling %s before PIE..."), *Blueprint->GetName());
-					FKismetEditorUtilities::CompileBlueprint(Blueprint, false, true);
-					const bool bHadError = (!Blueprint->IsUpToDate() && Blueprint->Status != BS_Unknown);
-
-					// Check if the Blueprint has already been added to the error list to prevent it from being added again
-					if (bHadError && ErrorBlueprintList.Find(Blueprint) == INDEX_NONE)
-					{
-						ErrorBlueprintList.Add(Blueprint);
-						ErrorBlueprints += FString::Printf(TEXT("\n   %s"), *Blueprint->GetName());
-
-						FFormatNamedArguments Arguments;
-						Arguments.Add(TEXT("Name"), FText::FromString(Blueprint->GetName()));
-
-						BlueprintLog.Info( FText::Format(LOCTEXT("BlueprintCompileFailed", "Blueprint {Name} failed to compile"), Arguments) );
-					}
-				}
-
-				CollectGarbage( GARBAGE_COLLECTION_KEEPFLAGS );
-
-				UE_LOG(LogPlayLevel, Log, TEXT("PIE:  Blueprint regeneration took %d ms (%i blueprints)"), (int32)((FPlatformTime::Seconds() - BPRegenStartTime) * 1000), BlueprintsToRecompile.Num());
-			}
-			else
+			FEditorDelegates::EndPIE.Broadcast(bInSimulateInEditor);
+			if (InWorld->GetNavigationSystem())
 			{
-				UE_LOG(LogPlayLevel, Log, TEXT("PIE:  No blueprints needed recompiling"));
+				InWorld->GetNavigationSystem()->OnPIEEnd();
 			}
-		}
-		else
-		{
-			if (bAnyBlueprintsDirty)
-			{
-				FFormatNamedArguments Args;
-				Args.Add( TEXT("DirtyBlueprints"), FText::FromString( DirtyBlueprints ) );
-
-				const bool bCompileDirty = EAppReturnType::Yes == FMessageDialog::Open( EAppMsgType::YesNo, FText::Format( NSLOCTEXT("PlayInEditor", "PrePIE_BlueprintsDirty", "One or more blueprints have been modified without being recompiled.  Do you want to compile them now?{DirtyBlueprints}"), Args ) );
-				
-				if ( bCompileDirty)
-				{
-					BlueprintLog.NewPage(LOCTEXT("BlueprintCompilationPageLabel", "Pre-PIE recompile"));
-
-					// Compile all blueprints that aren't up to date
-					for (TObjectIterator<UBlueprint> BlueprintIt; BlueprintIt; ++BlueprintIt)
-					{
-						UBlueprint* Blueprint = *BlueprintIt;
-						// do not try to recompile BPs that have not changed since they last failed to compile, so don't check Blueprint->IsUpToDate()
-						const bool bIsDirtyAndShouldBeRecompiled = Blueprint->IsPossiblyDirty();
-						if (!FBlueprintEditorUtils::IsDataOnlyBlueprint(Blueprint) && bIsDirtyAndShouldBeRecompiled)
-						{
-							// Cache off the dirty flag for the package, so we can restore it later
-							UPackage* Package = Cast<UPackage>(Blueprint->GetOutermost());
-							const bool bIsPackageDirty = Package ? Package->IsDirty() : false;
 
-							FKismetEditorUtilities::CompileBlueprint(Blueprint);
-							if (Blueprint->Status == BS_Error && Blueprint->bDisplayCompilePIEWarning)
-							{
-								ErrorBlueprintList.Add(Blueprint);
-							}
-
-							// Restore the dirty flag
-							if (Package)
-							{
-								Package->SetDirtyFlag(bIsPackageDirty);
-							}
-						}
-					}
-				}
-			}
+			return;
 		}
-
-		if ( ErrorBlueprintList.Num() && !GIsDemoMode )
+		else
 		{
-			FFormatNamedArguments Args;
-			Args.Add( TEXT("ErrorBlueprints"), FText::FromString( ErrorBlueprints ) );
-
-			// There was at least one blueprint with an error, make sure the user is OK with that.
-			const bool bContinuePIE = EAppReturnType::Yes == FMessageDialog::Open( EAppMsgType::YesNo, FText::Format( NSLOCTEXT("PlayInEditor", "PrePIE_BlueprintErrors", "One or more blueprints has an unresolved compiler error, are you sure you want to Play in Editor?{ErrorBlueprints}"), Args ) );
-			if ( !bContinuePIE )
+			// The user wants to ignore the compiler errors, mark the Blueprints and do not warn them again unless the Blueprint attempts to compile
+			for (auto Blueprint : ErroredBlueprints)
 			{
-				FEditorDelegates::EndPIE.Broadcast(bInSimulateInEditor);
-				if (InWorld->GetNavigationSystem())
-				{
-					InWorld->GetNavigationSystem()->OnPIEEnd();
-				}
-
-				return;
-			}
-			else
-			{
-				// The user wants to ignore the compiler errors, mark the Blueprints and do not warn them again unless the Blueprint attempts to compile
-				for( auto Blueprint : ErrorBlueprintList )
-				{
-					Blueprint->bDisplayCompilePIEWarning = false;
-				}
+				Blueprint->bDisplayCompilePIEWarning = false;
 			}
 		}
 	}
@@ -2236,8 +2200,6 @@ void UEditorEngine::PlayInEditor( UWorld* InWorld, bool bInSimulateInEditor )
 	}
 	EditorWorld->bAllowAudioPlayback = false;
 
-	ULevelEditorPlaySettings* PlayInSettings = Cast<ULevelEditorPlaySettings>(ULevelEditorPlaySettings::StaticClass()->GetDefaultObject());
-
 	if (!PlayInSettings->EnableSound && AudioDevice)
 	{
 		AudioDevice->TransientMasterVolume = 0.0f;
@@ -2276,7 +2238,7 @@ void UEditorEngine::PlayInEditor( UWorld* InWorld, bool bInSimulateInEditor )
 	// Can't allow realtime viewports whilst in PIE so disable it for ALL viewports here.
 	DisableRealtimeViewports();
 
-	bool bAnyBlueprintErrors = ErrorBlueprintList.Num()? true : false;
+	bool bAnyBlueprintErrors = ErroredBlueprints.Num() ? true : false;
 	bool bStartInSpectatorMode = false;
 	bool bSupportsOnlinePIE = false;
 	const int32 PlayNumberOfClients = [&PlayInSettings]{ int32 NumberOfClients(0); return (PlayInSettings->GetPlayNumberOfClients(NumberOfClients) ? NumberOfClients : 0); }();
diff --git a/Engine/Source/Editor/UnrealEd/Private/Settings/SettingsClasses.cpp b/Engine/Source/Editor/UnrealEd/Private/Settings/SettingsClasses.cpp
--- a/Engine/Source/Editor/UnrealEd/Private/Settings/SettingsClasses.cpp
+++ b/Engine/Source/Editor/UnrealEd/Private/Settings/SettingsClasses.cpp
@@ -351,6 +351,7 @@ ULevelEditorPlaySettings::ULevelEditorPlaySettings( const FObjectInitializer& Ob
 	RouteGamepadToSecondWindow = false;
 	AutoConnectToServer = true;
 	BuildGameBeforeLaunch = EPlayOnBuildMode::PlayOnBuild_Default;
+	bAutoCompileBlueprintsOnLaunch = true;
 }
 
 void ULevelEditorPlaySettings::PostEditChangeProperty(struct FPropertyChangedEvent& PropertyChangedEvent)
