Fixed crash while cooking when a package is missing from the asset registry.

diff --git a/Engine/Source/Editor/UnrealEd/Private/Commandlets/ChunkManifestGenerator.cpp b/Engine/Source/Editor/UnrealEd/Private/Commandlets/ChunkManifestGenerator.cpp
--- a/Engine/Source/Editor/UnrealEd/Private/Commandlets/ChunkManifestGenerator.cpp
+++ b/Engine/Source/Editor/UnrealEd/Private/Commandlets/ChunkManifestGenerator.cpp
@@ -432,7 +432,17 @@ void FChunkManifestGenerator::BuildChunkManifest(const TArray<FName>& CookedPack
 
 			const FString PackagePathName = PackageFName.ToString();
 			const FString MapName = MapFName.ToString();
-			const FString& SandboxFilename = AllCookedPackages.FindChecked(PackageFName);
+			const FString* SandboxFilenamePtr = AllCookedPackages.Find(PackageFName);
+			if (!SandboxFilenamePtr)
+			{
+				const FString SandboxPath = InSandboxFile->ConvertToAbsolutePathForExternalAppForWrite(*FPackageName::LongPackageNameToFilename(PackagePathName));
+
+				AllCookedPackages.Add(PackageFName, SandboxPath);
+
+				SandboxFilenamePtr = AllCookedPackages.Find(PackageFName);
+				check(SandboxFilenamePtr);
+			}
+			const FString& SandboxFilename = *SandboxFilenamePtr;
 
 
 			GenerateChunkManifestForPackage(PackageFName, PackagePathName, SandboxFilename, MapName, InSandboxFile);
