PR #1307: Do not crash with access violation if UEditorEngine::EditorWorld is NULL when UEngine::LoadMap is called (Contributed by [REDACTED_PERSON])

diff --git a/Engine/Source/Editor/UnrealEd/Private/EditorLevelUtils.cpp b/Engine/Source/Editor/UnrealEd/Private/EditorLevelUtils.cpp
--- a/Engine/Source/Editor/UnrealEd/Private/EditorLevelUtils.cpp
+++ b/Engine/Source/Editor/UnrealEd/Private/EditorLevelUtils.cpp
@@ -842,6 +842,12 @@ namespace EditorLevelUtils
 	void GetWorlds(UWorld* InWorld, TArray<UWorld*>& OutWorlds, bool bIncludeInWorld, bool bOnlyEditorVisible)
 	{
 		OutWorlds.Empty();
+
+		if (!InWorld)
+		{
+			return;
+		}
+
 		if ( bIncludeInWorld )
 		{
 			OutWorlds.AddUnique( InWorld );
