UE4 - Fix tooltip

diff --git a/Engine/Source/Editor/UnrealEd/Classes/Settings/EditorExperimentalSettings.h b/Engine/Source/Editor/UnrealEd/Classes/Settings/EditorExperimentalSettings.h
--- a/Engine/Source/Editor/UnrealEd/Classes/Settings/EditorExperimentalSettings.h
+++ b/Engine/Source/Editor/UnrealEd/Classes/Settings/EditorExperimentalSettings.h
@@ -89,7 +89,7 @@ public:
 	UPROPERTY(EditAnywhere, config, Category = AI, meta = (DisplayName = "Environment Querying System"))
 	bool bEQSEditor;
 
-	/** Enables Environment Queries editor */
+	/** Enables Metal/High-end mobile rendering path previw on Desktop */
 	UPROPERTY(EditAnywhere, config, Category = Rendering, meta = (DisplayName = "Enable Metal/High-end mobile rendering preview"))
 	bool bFeatureLevelES31Preview;
 
