Remove warning when GenerateBlueprintAPI commandlet doesn't understand an argument - there are lots of engine switches that commandlets are not meant to parse, and generating the blueprint docs specifies quite a lot of them (-stdout -forcelogflush -crashforuat -unattended -allowstdoutlogverbosity -buildmachine). We don't want warnings for those.

diff --git a/Engine/Source/Editor/UnrealEd/Private/Commandlets/GenerateBlueprintAPICommandlet.cpp b/Engine/Source/Editor/UnrealEd/Private/Commandlets/GenerateBlueprintAPICommandlet.cpp
--- a/Engine/Source/Editor/UnrealEd/Private/Commandlets/GenerateBlueprintAPICommandlet.cpp
+++ b/Engine/Source/Editor/UnrealEd/Private/Commandlets/GenerateBlueprintAPICommandlet.cpp
@@ -300,10 +300,6 @@ GenerateBlueprintAPIUtils::CommandletOptions::CommandletOptions(TArray<FString>
 			FString PathSwitch;
 			Switch.Split(TEXT("="), &PathSwitch, &SaveDir);
 		}
-		else
-		{
-			UE_LOG(LogBlueprintAPIGenerate, Warning, TEXT("Unrecognized command switch '%s', use -help for a listing of all accepted params"), *Switch);
-		}
 	}
 
 	if (NewDumpFlags != 0)
