Fix issues with SkinXX syntax no longer working

diff --git a/Engine/Source/Editor/UnrealEd/Private/Fbx/FbxSkeletalMeshImport.cpp b/Engine/Source/Editor/UnrealEd/Private/Fbx/FbxSkeletalMeshImport.cpp
--- a/Engine/Source/Editor/UnrealEd/Private/Fbx/FbxSkeletalMeshImport.cpp
+++ b/Engine/Source/Editor/UnrealEd/Private/Fbx/FbxSkeletalMeshImport.cpp
@@ -1163,8 +1163,12 @@ USkeletalMesh* UnFbx::FFbxImporter::ImportSkeletalMesh(UObject* InParent, TArray
 			if( !FbxMaterials.Contains( FbxMaterial ) )
 			{
 				FbxMaterials.Add( FbxMaterial );
+
+				VMaterial NewMaterial;
+
+				NewMaterial.MaterialImportName = UTF8_TO_TCHAR(MakeName(FbxMaterial->GetName()));
 				// Add an entry for each unique material
-				SkelMeshImportDataPtr->Materials.Add( VMaterial() );
+				SkelMeshImportDataPtr->Materials.Add( NewMaterial );
 			}
 		}
 	}
@@ -1664,7 +1668,7 @@ void UnFbx::FFbxImporter::SetMaterialSkinXXOrder(FSkeletalMeshImportData& Import
 				int32 NewIndex = MaterialIndexToSkinIndex[MaterialIndex];
 				if( ExistingMatList.IsValidIndex( NewIndex ) )
 				{
-					ImportData.Materials[MaterialIndex] = ExistingMatList[ NewIndex ];
+					ImportData.Materials[NewIndex] = ExistingMatList[MaterialIndex];
 				}
 			}
 
@@ -1790,9 +1794,8 @@ bool UnFbx::FFbxImporter::FillSkelMeshImporterFromFbx( FSkeletalMeshImportData&
 			// Reuse existing material
 			MaterialMapping[MaterialIndex] = ExistingMatIndex;
 
-			if (ImportOptions->bImportMaterials && Materials.IsValidIndex(MaterialIndex) )
+			if (Materials.IsValidIndex(MaterialIndex) )
 			{
-				ImportData.Materials[ExistingMatIndex].MaterialImportName = UTF8_TO_TCHAR(MakeName(FbxMaterial->GetName()));
 				ImportData.Materials[ExistingMatIndex].Material = Materials[MaterialIndex];
 			}
 		}
