Running the QAGame CompileBlueprint automation tests twice causes a crash, this fixes it. It is valid to not have a generated class.

diff --git a/Engine/Source/Editor/UnrealEd/Private/Kismet2/KismetReinstanceUtilities.cpp b/Engine/Source/Editor/UnrealEd/Private/Kismet2/KismetReinstanceUtilities.cpp
--- a/Engine/Source/Editor/UnrealEd/Private/Kismet2/KismetReinstanceUtilities.cpp
+++ b/Engine/Source/Editor/UnrealEd/Private/Kismet2/KismetReinstanceUtilities.cpp
@@ -612,7 +612,9 @@ void FBlueprintCompileReinstancer::UpdateBytecodeReferences()
 			UClass* BPClass = (*DependentBP)->GeneratedClass;
 
 			// Skip cases where the class is junk, or haven't finished serializing in yet
-			if( (BPClass == ClassToReinstance)
+			// Note that BPClass can be null for blueprints that can no longer be compiled:
+			if (!BPClass
+				|| (BPClass == ClassToReinstance)
 				|| (BPClass->GetOutermost() == GetTransientPackage()) 
 				|| BPClass->HasAnyClassFlags(CLASS_NewerVersionExists)
 				|| (BPClass->ClassGeneratedBy && BPClass->ClassGeneratedBy->HasAnyFlags(RF_NeedLoad|RF_BeingRegenerated)) )
