UE-17236: Redirector Assets fail to Load

diff --git a/Engine/Source/Editor/UnrealEd/Private/EditorEngine.cpp b/Engine/Source/Editor/UnrealEd/Private/EditorEngine.cpp
--- a/Engine/Source/Editor/UnrealEd/Private/EditorEngine.cpp
+++ b/Engine/Source/Editor/UnrealEd/Private/EditorEngine.cpp
@@ -111,6 +111,7 @@
 #include "UnrealEngine.h"
 #include "EngineStats.h"
 #include "Engine/SimpleConstructionScript.h"
+#include "PackageTools.h"
 
 #include "PhysicsPublic.h"
 #include "Engine/CoreSettings.h"
@@ -1809,10 +1810,33 @@ void UEditorEngine::Cleanse( bool ClearSelection, bool Redraw, const FText& Tran
 			RedrawLevelEditingViewports();
 		}
 
-		// Attempt to unload any loaded redirectors. Redirectors should not be referenced in memory and are only used to forward references at load time
+		// Attempt to unload any loaded redirectors. Redirectors should not
+		// be referenced in memory and are only used to forward references
+		// at load time.
+		//
+		// We also have to remove packages that redirectors were contained
+		// in, so they can be loaded again in the future. If we don't do it
+		// loading failure will occur next time someone tries to use it.
+		// This is caused by the fact that the loading routing will check
+		// that already existed, but the object was missing in cache.
+		TArray<UPackage*> PackagesToUnload;
 		for (TObjectIterator<UObjectRedirector> RedirIt; RedirIt; ++RedirIt)
 		{
-			RedirIt->ClearFlags(RF_Standalone | RF_RootSet | RF_Transactional);
+			PackagesToUnload.Add(RedirIt->GetOutermost());
+		}
+
+		for (auto* PackageToUnload : PackagesToUnload)
+		{
+			const EObjectFlags FlagsToClear = RF_Standalone | RF_RootSet | RF_Transactional;
+
+			TArray<UObject*> PackageObjects;
+			GetObjectsWithOuter(PackageToUnload, PackageObjects);
+			for (auto* Object : PackageObjects)
+			{
+				Object->ClearFlags(FlagsToClear);
+			}
+
+			PackageToUnload->ClearFlags(FlagsToClear);
 		}
 
 		// Collect garbage.
