Elminated redundant editor viewport invalidation code

diff --git a/Engine/Source/Editor/UnrealEd/Classes/Editor/EditorEngine.h b/Engine/Source/Editor/UnrealEd/Classes/Editor/EditorEngine.h
--- a/Engine/Source/Editor/UnrealEd/Classes/Editor/EditorEngine.h
+++ b/Engine/Source/Editor/UnrealEd/Classes/Editor/EditorEngine.h
@@ -2663,9 +2663,6 @@ protected:
 	 */
 	void InvalidateAllViewportsAndHitProxies();
 
-	/** Invalidate all level editor viewport client hit proxies immediately */
-	void InvalidateAllLevelEditorViewportClientHitProxies();
-
 	/** Destroy any online subsystems generated by PIE */
 	void CleanupPIEOnlineSessions(TArray<FName> OnlineIdentifiers);
 
diff --git a/Engine/Source/Editor/UnrealEd/Private/EditorEngine.cpp b/Engine/Source/Editor/UnrealEd/Private/EditorEngine.cpp
--- a/Engine/Source/Editor/UnrealEd/Private/EditorEngine.cpp
+++ b/Engine/Source/Editor/UnrealEd/Private/EditorEngine.cpp
@@ -1748,17 +1748,6 @@ void UEditorEngine::InvalidateAllViewportsAndHitProxies()
 	}
 }
 
-void UEditorEngine::InvalidateAllLevelEditorViewportClientHitProxies()
-{
-	for (const auto* LevelViewportClient : LevelViewportClients)
-	{
-		if (LevelViewportClient->Viewport != nullptr)
-		{
-			LevelViewportClient->Viewport->InvalidateHitProxy();
-		}
-	}
-}
-
 void UEditorEngine::PostEditChangeProperty(FPropertyChangedEvent& PropertyChangedEvent)
 {
 	// Propagate the callback up to the superclass.
diff --git a/Engine/Source/Editor/UnrealEd/Private/EditorServer.cpp b/Engine/Source/Editor/UnrealEd/Private/EditorServer.cpp
--- a/Engine/Source/Editor/UnrealEd/Private/EditorServer.cpp
+++ b/Engine/Source/Editor/UnrealEd/Private/EditorServer.cpp
@@ -2071,7 +2071,7 @@ UWorld* UEditorEngine::NewMap()
 	GUnrealEd->ResetTransaction( CleanseText );
 
 	// Invalidate all the level viewport hit proxies
-	InvalidateAllLevelEditorViewportClientHitProxies();
+	RedrawLevelEditingViewports();
 
 	return NewWorld;
 }
@@ -2521,7 +2521,7 @@ bool UEditorEngine::Map_Load(const TCHAR* Str, FOutputDevice& Ar)
 					}
 
 					// Invalidate all the level viewport hit proxies
-					InvalidateAllLevelEditorViewportClientHitProxies();
+					RedrawLevelEditingViewports();
 				}
 			}
 			else
