Fixed missing parameter name in StaticMeshThumbnailRenderer::Draw

diff --git a/Engine/Source/Editor/UnrealEd/Classes/ThumbnailRendering/StaticMeshThumbnailRenderer.h b/Engine/Source/Editor/UnrealEd/Classes/ThumbnailRendering/StaticMeshThumbnailRenderer.h
--- a/Engine/Source/Editor/UnrealEd/Classes/ThumbnailRendering/StaticMeshThumbnailRenderer.h
+++ b/Engine/Source/Editor/UnrealEd/Classes/ThumbnailRendering/StaticMeshThumbnailRenderer.h
@@ -15,7 +15,7 @@ class UStaticMeshThumbnailRenderer : public UDefaultSizedThumbnailRenderer
 
 
 	// Begin UThumbnailRenderer Object
-	UNREALED_API virtual void Draw(UObject* Object, int32 X, int32 Y, uint32 Width, uint32 Height, FRenderTarget*, FCanvas* Canvas) override;
+	UNREALED_API virtual void Draw(UObject* Object, int32 X, int32 Y, uint32 Width, uint32 Height, FRenderTarget* RenderTarget, FCanvas* Canvas) override;
 	// End UThumbnailRenderer Object
 
 	// UObject implementation
