New game windows now default to being positioned centered on the screen.

diff --git a/Engine/Config/BaseEditorPerProjectUserSettings.ini b/Engine/Config/BaseEditorPerProjectUserSettings.ini
--- a/Engine/Config/BaseEditorPerProjectUserSettings.ini
+++ b/Engine/Config/BaseEditorPerProjectUserSettings.ini
@@ -190,8 +190,10 @@ EnableSound=true
 NewWindowWidth=1280
 NewWindowHeight=720
 NewWindowPosition=(X=0,Y=0)
+CenterNewWindow=True
 StandaloneWindowWidth=1280
 StandaloneWindowHeight=720
+CenterStandaloneWindow=True
 ShowMouseControlLabel=True
 AutoRecompileBlueprints=True
 ; True if Play In Editor should only load currently-visible levels in PIE
diff --git a/Engine/Source/Editor/UnrealEd/Private/Settings/SettingsClasses.cpp b/Engine/Source/Editor/UnrealEd/Private/Settings/SettingsClasses.cpp
--- a/Engine/Source/Editor/UnrealEd/Private/Settings/SettingsClasses.cpp
+++ b/Engine/Source/Editor/UnrealEd/Private/Settings/SettingsClasses.cpp
@@ -352,6 +352,8 @@ ULevelEditorPlaySettings::ULevelEditorPlaySettings( const FObjectInitializer& Ob
 	AutoConnectToServer = true;
 	BuildGameBeforeLaunch = EPlayOnBuildMode::PlayOnBuild_Default;
 	bAutoCompileBlueprintsOnLaunch = true;
+	CenterNewWindow = true;
+	CenterStandaloneWindow = true;
 }
 
 void ULevelEditorPlaySettings::PostEditChangeProperty(struct FPropertyChangedEvent& PropertyChangedEvent)
