Fixed caching of the game mode class when it was using a BP class

The BP class potentially had the type info in front of it, so it was never passing the "GameModeClass->GetPathName() != ClassName" check, and was always re-getting the class.

diff --git a/Engine/Source/Editor/UnrealEd/Public/GameModeInfoCustomizer.h b/Engine/Source/Editor/UnrealEd/Public/GameModeInfoCustomizer.h
--- a/Engine/Source/Editor/UnrealEd/Public/GameModeInfoCustomizer.h
+++ b/Engine/Source/Editor/UnrealEd/Public/GameModeInfoCustomizer.h
@@ -169,6 +169,9 @@ public:
 		FString ClassName;
 		DefaultGameModeClassHandle->GetValueAsFormattedString(ClassName);
 
+		// Blueprints may have type information before the class name, so make sure and strip that off now
+		ConstructorHelpers::StripObjectClass(ClassName);
+
 		// Do we have a valid cached class pointer? (Note: We can't search for the class while a save is happening)
 		const UClass* GameModeClass = CachedGameModeClass.Get();
 		if ((!GameModeClass || GameModeClass->GetPathName() != ClassName) && !GIsSavingPackage)
