Always load IntroTutorials and Blutilities even when running commandlets otherwise assets based on types in the module are unknown and cannot be loaded

diff --git a/Engine/Source/Editor/UnrealEd/Private/Commandlets/GenerateBlueprintAPICommandlet.cpp b/Engine/Source/Editor/UnrealEd/Private/Commandlets/GenerateBlueprintAPICommandlet.cpp
--- a/Engine/Source/Editor/UnrealEd/Private/Commandlets/GenerateBlueprintAPICommandlet.cpp
+++ b/Engine/Source/Editor/UnrealEd/Private/Commandlets/GenerateBlueprintAPICommandlet.cpp
@@ -963,8 +963,6 @@ int32 UGenerateBlueprintAPICommandlet::Main(FString const& Params)
 	GenerateBlueprintAPIUtils::CommandOptions = GenerateBlueprintAPIUtils::CommandletOptions(Switches);
 	GenerateBlueprintAPIUtils::CommandletOptions const& CommandOptions = GenerateBlueprintAPIUtils::CommandOptions;
 
-	FModuleManager::Get().LoadModuleChecked(TEXT("Blutility"));
-
 	FString ActiveFilePath;
 	FArchive* FileOut = nullptr;
 
diff --git a/Engine/Source/Editor/UnrealEd/Private/EditorEngine.cpp b/Engine/Source/Editor/UnrealEd/Private/EditorEngine.cpp
--- a/Engine/Source/Editor/UnrealEd/Private/EditorEngine.cpp
+++ b/Engine/Source/Editor/UnrealEd/Private/EditorEngine.cpp
@@ -685,7 +685,6 @@ void UEditorEngine::Init(IEngineLoop* InEngineLoop)
 		if (!IsRunningCommandlet())
 		{
 			FModuleManager::Get().LoadModule(TEXT("EditorLiveStreaming"));
-			FModuleManager::Get().LoadModule(TEXT("IntroTutorials"));
 		}
 
 		if( FParse::Param( FCommandLine::Get(),TEXT( "PListEditor" ) ) )
diff --git a/Engine/Source/Runtime/Launch/Private/LaunchEngineLoop.cpp b/Engine/Source/Runtime/Launch/Private/LaunchEngineLoop.cpp
--- a/Engine/Source/Runtime/Launch/Private/LaunchEngineLoop.cpp
+++ b/Engine/Source/Runtime/Launch/Private/LaunchEngineLoop.cpp
@@ -1931,6 +1931,9 @@ bool FEngineLoop::LoadStartupCoreModules()
 	// We need this for blueprint projects that have online functionality.
 	FModuleManager::Get().LoadModule(TEXT("OnlineBlueprintSupport"));
 
+	FModuleManager::Get().LoadModule(TEXT("IntroTutorials"));
+	FModuleManager::Get().LoadModule(TEXT("Blutility"));
+
 #endif //(WITH_EDITOR && !(UE_BUILD_SHIPPING || UE_BUILD_TEST))
 
 	return bSuccess;
