Updated logic in commandlet for gathering text from assets which was falsely reporting that nothing was gathered because no files were loaded. Now properly reports whether there were any packages found to be processed, rather than if any need to be loaded and processed in a deferred manner.

diff --git a/Engine/Source/Editor/UnrealEd/Private/Commandlets/GatherTextFromAssetsCommandlet.cpp b/Engine/Source/Editor/UnrealEd/Private/Commandlets/GatherTextFromAssetsCommandlet.cpp
--- a/Engine/Source/Editor/UnrealEd/Private/Commandlets/GatherTextFromAssetsCommandlet.cpp
+++ b/Engine/Source/Editor/UnrealEd/Private/Commandlets/GatherTextFromAssetsCommandlet.cpp
@@ -273,8 +273,7 @@ int32 UGatherTextFromAssetsCommandlet::Main(const FString& Params)
 	}
 
 	//The main array of files to work from.
-	TArray< FString > PackageFileNamesToLoad;
-	TSet< FString > LongPackageNamesToProcess;
+	TArray< FString > PackageFileNamesToProcess;
 
 	TArray<FString> PackageFilesNotInIncludePath;
 	TArray<FString> PackageFilesInExcludePath;
@@ -341,44 +340,52 @@ int32 UGatherTextFromAssetsCommandlet::Main(const FString& Params)
 			//If we haven't failed one of the above checks, add it to the array of packages to process.
 			if(!bExclude)
 			{
-				TScopedPointer< FArchive > FileReader( IFileManager::Get().CreateFileReader( *PackageFile ) );
-				if( FileReader )
-				{
-					// Read package file summary from the file
-					FPackageFileSummary PackageFileSummary;
-					(*FileReader) << PackageFileSummary;
+				PackageFileNamesToProcess.Add(PackageFile);
+			}
+		}
+	}
 
-					if (PackageFileSummary.GetFileVersionUE4() >= VER_UE4_SERIALIZE_TEXT_IN_PACKAGES)
-					{
-						TArray<FGatherableTextData> GatherableTextDataArray;
+	if ( PackageFileNamesToProcess.Num() == 0 )
+	{
+		UE_LOG(LogGatherTextFromAssetsCommandlet, Warning, TEXT("No files found or none passed the include/exclude criteria."));
+	}
 
-						if (PackageFileSummary.GatherableTextDataOffset > 0)
-						{
-							FileReader->Seek(PackageFileSummary.GatherableTextDataOffset);
+	TArray< FString > PackageFileNamesToLoad;
+	for (FString& PackageFile : PackageFileNamesToProcess)
+	{
+		TScopedPointer< FArchive > FileReader( IFileManager::Get().CreateFileReader( *PackageFile ) );
+		if( FileReader )
+		{
+			// Read package file summary from the file
+			FPackageFileSummary PackageFileSummary;
+			(*FileReader) << PackageFileSummary;
 
-							for(int32 i = 0; i < PackageFileSummary.GatherableTextDataCount; ++i)
-							{
-								FGatherableTextData* GatherableTextData = new(GatherableTextDataArray) FGatherableTextData;
-								(*FileReader) << *GatherableTextData;
-							}
+			// Package has gatherable text data in its header, process immediately.
+			if (PackageFileSummary.GetFileVersionUE4() >= VER_UE4_SERIALIZE_TEXT_IN_PACKAGES)
+			{
+				TArray<FGatherableTextData> GatherableTextDataArray;
 
-							ProcessGatherableTextDataArray(PackageFile, GatherableTextDataArray);
-						}
-					}
-					else if (PackageFileSummary.PackageFlags & PKG_RequiresLocalizationGather || PackageFileSummary.GetFileVersionUE4() < VER_UE4_PACKAGE_REQUIRES_LOCALIZATION_GATHER_FLAGGING)
+				if (PackageFileSummary.GatherableTextDataOffset > 0)
+				{
+					FileReader->Seek(PackageFileSummary.GatherableTextDataOffset);
+
+					for(int32 i = 0; i < PackageFileSummary.GatherableTextDataCount; ++i)
 					{
-						PackageFileNamesToLoad.Add(PackageFile);
+						FGatherableTextData* GatherableTextData = new(GatherableTextDataArray) FGatherableTextData;
+						(*FileReader) << *GatherableTextData;
 					}
+
+					ProcessGatherableTextDataArray(PackageFile, GatherableTextDataArray);
 				}
 			}
+			// Package not resaved since gatherable text data was added to package headers, defer processing until after loading.
+			else if (PackageFileSummary.PackageFlags & PKG_RequiresLocalizationGather || PackageFileSummary.GetFileVersionUE4() < VER_UE4_PACKAGE_REQUIRES_LOCALIZATION_GATHER_FLAGGING)
+			{
+				PackageFileNamesToLoad.Add(PackageFile);
+			}
 		}
 	}
 
-	if ( PackageFileNamesToLoad.Num() == 0 )
-	{
-		UE_LOG(LogGatherTextFromAssetsCommandlet, Warning, TEXT("No files found. Or none passed the include/exclude criteria."));
-	}
-
 	CollectGarbage( RF_Native );
 
 	//Now go through the remaining packages in the main array and process them in batches.
@@ -386,7 +393,7 @@ int32 UGatherTextFromAssetsCommandlet::Main(const FString& Params)
 	TArray< UPackage* > LoadedPackages;
 	TArray< FString > LoadedPackageFileNames;
 	TArray< FString > FailedPackageFileNames;
-	TArray< UPackage* > PackagesToProcess;
+	TArray< UPackage* > LoadedPackagesToProcess;
 
 	const int32 PackageCount = PackageFileNamesToLoad.Num();
 	const int32 BatchCount = PackageCount / PackagesPerBatchCount + (PackageCount % PackagesPerBatchCount > 0 ? 1 : 0); // Add an extra batch for any remainder if necessary
@@ -415,7 +422,7 @@ int32 UGatherTextFromAssetsCommandlet::Main(const FString& Params)
 				// The loading process should have reflagged said packages so that only true positives will have this flag.
 				if( Package->RequiresLocalizationGather() )
 				{
-					PackagesToProcess.Add( Package );
+					LoadedPackagesToProcess.Add( Package );
 				}
 			}
 			else
@@ -430,8 +437,8 @@ int32 UGatherTextFromAssetsCommandlet::Main(const FString& Params)
 
 		UE_LOG(LogGatherTextFromAssetsCommandlet, Log, TEXT("Loaded %i packages in batch %i of %i. %i failed."), PackagesInThisBatch, BatchIndex + 1, BatchCount, FailuresInThisBatch);
 
-		ProcessPackages(PackagesToProcess);
-		PackagesToProcess.Empty(PackagesPerBatchCount);
+		ProcessPackages(LoadedPackagesToProcess);
+		LoadedPackagesToProcess.Empty(PackagesPerBatchCount);
 
 		if( bFixBroken )
 		{
