cleaner comments

diff --git a/Engine/Source/Editor/UnrealEd/Private/Fbx/FbxMainExport.cpp b/Engine/Source/Editor/UnrealEd/Private/Fbx/FbxMainExport.cpp
--- a/Engine/Source/Editor/UnrealEd/Private/Fbx/FbxMainExport.cpp
+++ b/Engine/Source/Editor/UnrealEd/Private/Fbx/FbxMainExport.cpp
@@ -1819,7 +1819,6 @@ void DetermineVertsToWeld(TArray<int32>& VertRemap, TArray<int32>& UniqueVerts,
 /**
  * Exports a static mesh
  * @param StaticMesh	The static mesh to export
- * @param RenderMesh	The static mesh render data to export
  * @param MeshName		The name of the mesh for the FBX file
  * @param FbxActor		The fbx node representing the mesh
  * @param ExportLOD		The LOD of the mesh to export
diff --git a/Engine/Source/Editor/UnrealEd/Private/FbxExporter.h b/Engine/Source/Editor/UnrealEd/Private/FbxExporter.h
--- a/Engine/Source/Editor/UnrealEd/Private/FbxExporter.h
+++ b/Engine/Source/Editor/UnrealEd/Private/FbxExporter.h
@@ -213,7 +213,6 @@ private:
 	/**
 	 * Exports a spline mesh
 	 * @param SplineMeshComp	The spline mesh component to export
-	 * @param RenderMesh		The mesh render data to export
 	 * @param MeshName		The name of the mesh for the FBX file
 	 * @param FbxActor		The fbx node representing the mesh
 	 */
@@ -222,7 +221,6 @@ private:
 	/**
 	 * Exports an instanced mesh
 	 * @param InstancedMeshComp	The instanced mesh component to export
-	 * @param RenderMesh		The mesh render data to export
 	 * @param MeshName		The name of the mesh for the FBX file
 	 * @param FbxActor		The fbx node representing the mesh
 	 */
