Fix crash force deleting lots of objects

https://jira.ol.epicgames.net/browse/UE-17645?

diff --git a/Engine/Source/Editor/UnrealEd/Private/ObjectTools.cpp b/Engine/Source/Editor/UnrealEd/Private/ObjectTools.cpp
--- a/Engine/Source/Editor/UnrealEd/Private/ObjectTools.cpp
+++ b/Engine/Source/Editor/UnrealEd/Private/ObjectTools.cpp
@@ -47,7 +47,7 @@ DEFINE_LOG_CATEGORY_STATIC(LogObjectTools, Log, All);
 
 // This function should ONLY be needed by ConsolidateObjects and ForceDeleteObjects
 // Use anywhere else could be dangerous as this involves a map transition and GC
-void ReloadEditorWorldForReferenceReplacementIfNecessary(TArray<UObject*>& InOutObjectsToReplace)
+void ReloadEditorWorldForReferenceReplacementIfNecessary(TArray< TWeakObjectPtr<UObject> >& InOutObjectsToReplace)
 {
 	// If we are force-deleting or consolidating the editor world, first transition to an empty map to prevent reference problems.
 	// Then, re-load the world from disk to set it up for delete as an inactive world which isn't attached to the editor engine or other systems.
@@ -765,8 +765,28 @@ namespace ObjectTools
 			// objects they are replacing until the objects are garbage collected
 			TMap<UObjectRedirector*, FName> RedirectorToObjectNameMap;
 
-			// If the current editor world is in this list, transition to a new map and reload the world to finish the consolidate
-			ReloadEditorWorldForReferenceReplacementIfNecessary(ObjectsToConsolidate);
+			{
+				// Note reloading the world via ReloadEditorWorldForReferenceReplacementIfNecessary will cause a garbage collect and potentially cause entries in the ObjectsToConsolidate list to become invalid
+				// We refresh the list here after reloading the editor world
+				TArray< TWeakObjectPtr<UObject> > ObjectsToConsolidateWeakList;
+				for(UObject* Object : ObjectsToConsolidate)
+				{
+					ObjectsToConsolidateWeakList.Add(Object);
+				}
+
+				ObjectsToConsolidate.Empty();
+
+				// If the current editor world is in this list, transition to a new map and reload the world to finish the delete
+				ReloadEditorWorldForReferenceReplacementIfNecessary(ObjectsToConsolidateWeakList);
+
+				for(TWeakObjectPtr<UObject> WeakObject : ObjectsToConsolidateWeakList)
+				{
+					if( WeakObject.IsValid() )
+					{
+						ObjectsToConsolidate.Add(WeakObject.Get());
+					}
+				}
+			}
 
 			FForceReplaceInfo ReplaceInfo;
 			// Scope the reregister context below to complete after object deletion and before garbage collection
@@ -1940,8 +1960,28 @@ namespace ObjectTools
 			GEditor->NoteSelectionChange();
 		}
 
-		// If the current editor world is in this list, transition to a new map and reload the world to finish the delete
-		ReloadEditorWorldForReferenceReplacementIfNecessary(ObjectsToDelete);
+		{
+			// Note reloading the world via ReloadEditorWorldForReferenceReplacementIfNecessary will cause a gabage collect and potentially cause entries in the ObjectsToDelete list to become invalid
+			// We refresh the list here
+			TArray< TWeakObjectPtr<UObject> > ObjectsToDeleteWeakList;
+			for(UObject* Object : ObjectsToDelete)
+			{
+				ObjectsToDeleteWeakList.Add(Object);
+			}
+
+			ObjectsToDelete.Empty();
+
+			// If the current editor world is in this list, transition to a new map and reload the world to finish the delete
+			ReloadEditorWorldForReferenceReplacementIfNecessary(ObjectsToDeleteWeakList);
+
+			for(TWeakObjectPtr<UObject> WeakObject : ObjectsToDeleteWeakList)
+			{
+				if( WeakObject.IsValid() )
+				{
+					ObjectsToDelete.Add(WeakObject.Get());
+				}
+			}
+		}
 
 		{
 			int32 ReplaceableObjectsNum = 0;
