Harden PlayEditorSound against abuse

diff --git a/Engine/Source/Editor/UnrealEd/Private/EditorEngine.cpp b/Engine/Source/Editor/UnrealEd/Private/EditorEngine.cpp
--- a/Engine/Source/Editor/UnrealEd/Private/EditorEngine.cpp
+++ b/Engine/Source/Editor/UnrealEd/Private/EditorEngine.cpp
@@ -1904,7 +1904,7 @@ void UEditorEngine::PlayPreviewSound( USoundBase* Sound,  USoundNode* SoundNode
 void UEditorEngine::PlayEditorSound( const FString& SoundAssetName )
 {
 	// Only play sounds if the user has that feature enabled
-	if( GetDefault<ULevelEditorMiscSettings>()->bEnableEditorSounds && !GIsSavingPackage )
+	if( !GIsSavingPackage && IsInGameThread() && GetDefault<ULevelEditorMiscSettings>()->bEnableEditorSounds )
 	{
 		USoundBase* Sound = Cast<USoundBase>( StaticFindObject( USoundBase::StaticClass(), NULL, *SoundAssetName ) );
 		if( Sound == NULL )
