Added support for distributing all files to child cookers when using multiprocess cooking.

diff --git a/Engine/Source/Editor/UnrealEd/Classes/CookOnTheSide/CookOnTheFlyServer.h b/Engine/Source/Editor/UnrealEd/Classes/CookOnTheSide/CookOnTheFlyServer.h
--- a/Engine/Source/Editor/UnrealEd/Classes/CookOnTheSide/CookOnTheFlyServer.h
+++ b/Engine/Source/Editor/UnrealEd/Classes/CookOnTheSide/CookOnTheFlyServer.h
@@ -503,6 +503,14 @@ private:
 			}
 			return true;
 		}
+
+		bool Exists(const FName& Filename)
+		{
+			const TArray<FName>* Platforms = PlatformList.Find(Filename);
+			if (Platforms == NULL)
+				return false;
+			return true;
+		}
 		
 		bool HasItems() const
 		{
@@ -658,7 +666,11 @@ private:
 		bool bForceEnableCompressedPackages;
 		bool bForceDisableCompressedPackages;
 		bool bIsChildCooker;
+		FString ChildCookFilename;
+		TSet<FName> ChildUnsolicitedPackages;
 		TArray<FChildCooker> ChildCookers;
+		TArray<FName> TargetPlatformNames;
+		
 	};
 	FCookByTheBookOptions* CookByTheBookOptions;
 	
@@ -741,6 +753,7 @@ public:
 		COSR_ErrorLoadingPackage= 0x00000004,
 		COSR_RequiresGC			= 0x00000008,
 		COSR_WaitingOnCache		= 0x00000010,
+		COSR_WaitingOnChildCookers = 0x00000020,
 	};
 
 
@@ -949,7 +962,7 @@ private:
 	 */
 	void CollectFilesToCook(TArray<FName>& FilesInPath, 
 		const TArray<FString>& CookMaps, const TArray<FString>& CookDirectories, const TArray<FString>& CookCultures, 
-		const TArray<FString>& IniMapSections, bool bCookAll, bool bMapsOnly, bool bNoDev, const FString& ChildCookFilename);
+		const TArray<FString>& IniMapSections, bool bCookAll, bool bMapsOnly, bool bNoDev);
 
 	/**
 	 * AddFileToCook add file to cook list 
@@ -975,7 +988,7 @@ private:
 	 * 
 	 * @return return true if all child cookers are finished
 	 */
-	bool TickChildCookers();
+	bool TickChildCookers( const bool bCheckIfFinished );
 
 	/**
 	 * Get all the packages which are listed in asset registry passed in.  
diff --git a/Engine/Source/Editor/UnrealEd/Private/CookOnTheFlyServer.cpp b/Engine/Source/Editor/UnrealEd/Private/CookOnTheFlyServer.cpp
--- a/Engine/Source/Editor/UnrealEd/Private/CookOnTheFlyServer.cpp
+++ b/Engine/Source/Editor/UnrealEd/Private/CookOnTheFlyServer.cpp
@@ -43,7 +43,7 @@ DEFINE_LOG_CATEGORY_STATIC(LogCook, Log, All);
 
 
 #define DEBUG_COOKONTHEFLY 0
-#define OUTPUT_TIMING 0
+#define OUTPUT_TIMING 1
 
 #define USEASSETREGISTRYFORDEPENDENTPACKAGES 1
 #define VERIFY_GETDEPENDENTPACKAGES 0 // verify has false hits because old serialization method for generating dependencies had errors (included transient objects which shouldn't be in asset registry), but you can still use verify to build a list then cross check against transient objects.  
@@ -430,6 +430,13 @@ const FString& GetAssetRegistryPath()
 	return AssetRegistryPath;
 }
 
+
+FString GetChildCookerResultFilename(const FString& ResponseFilename)
+{
+	FString Result = ResponseFilename + TEXT("Result.txt");
+	return Result;
+}
+
 /**
  * Return the release asset registry filename for the release version supplied
  */
@@ -1129,7 +1136,9 @@ FString UCookOnTheFlyServer::GetDLCContentPath()
 
 COREUOBJECT_API extern bool GOutputCookingWarnings;
 
-bool UCookOnTheFlyServer::TickChildCookers()
+
+
+bool UCookOnTheFlyServer::TickChildCookers(const bool bCheckIfFinished)
 {
 	bool bChildCookersFinished = true;
 	if (IsCookByTheBookMode())
@@ -1140,14 +1149,44 @@ bool UCookOnTheFlyServer::TickChildCookers()
 			{
 				bool bIsCookerFinished = false;
 				int32 ReturnCode = 0;
-				if (FPlatformProcess::GetProcReturnCode(ChildCooker.ProcessHandle, &ReturnCode))
+				if (bCheckIfFinished == true)
 				{
-					UE_LOG(LogCook, Display, TEXT("ChildCooker returned error code %d"), ReturnCode);
-					ChildCooker.ReturnCode = ReturnCode;
-					ChildCooker.bFinished = true;
-					if (ReturnCode != 0)
+					if (FPlatformProcess::GetProcReturnCode(ChildCooker.ProcessHandle, &ReturnCode))
 					{
-						UE_LOG(LogCook, Fatal, TEXT("ChildCooker returned error code %d"), ReturnCode);
+						UE_LOG(LogCook, Display, TEXT("ChildCooker returned error code %d"), ReturnCode);
+						ChildCooker.ReturnCode = ReturnCode;
+						ChildCooker.bFinished = true;
+
+						if (ReturnCode != 0)
+						{
+							UE_LOG(LogCook, Error, TEXT("ChildCooker returned error code %d"), ReturnCode);
+						}
+
+
+						// if the child cooker completed successfully then it would have output a list of files which the main cooker needs to process
+						const FString AdditionalPackagesFileName = GetChildCookerResultFilename(ChildCooker.ResponseFileName);
+
+						FString AdditionalPackages;
+						if (FFileHelper::LoadFileToString(AdditionalPackages, *AdditionalPackagesFileName) == false)
+						{
+							UE_LOG(LogCook, Fatal, TEXT("ChildCooker failed to write out additional packages file to location %s"), *AdditionalPackagesFileName);
+						}
+						TArray<FString> AdditionalPackageList;
+						AdditionalPackages.ParseIntoArrayLines(AdditionalPackageList);
+
+						for (const auto& AdditionalPackage : AdditionalPackageList)
+						{
+							FName Filename = FName(*AdditionalPackage);
+							FFilePlatformRequest CheckCooked = FFilePlatformRequest(Filename, CookByTheBookOptions->TargetPlatformNames);
+							if (CookedPackages.Exists(CheckCooked) == false)
+							{
+								CookRequests.EnqueueUnique(CheckCooked);
+							}
+						}
+					}
+					else
+					{
+						bChildCookersFinished = false;
 					}
 				}
 				else
@@ -1234,11 +1273,33 @@ uint32 UCookOnTheFlyServer::TickCookOnTheSide( const float TimeSlice, uint32 &Co
 	// This is all the target platforms which we needed to process requests for this iteration
 	// we use this in the unsolicited packages processing below
 	TArray<FName> AllTargetPlatformNames;
-
+	
+	if (CurrentCookMode == ECookMode::CookByTheBook)
 	{
-		SCOPE_TIMER(TickChildCookersOutput);
-		TickChildCookers();
+		SCOPE_TIMER(TickChildCookers);
+		if (CookRequests.HasItems() == false)
+		{
+			SCOPE_TIMER(WaitingForChildCookersToFinish);
+			// we have nothing left to cook so we are now waiting for child cookers to finish
+			if (TickChildCookers(true) == false)
+			{
+				Result |= COSR_WaitingOnChildCookers;
+				if (CookRequests.HasItems() == false)
+				{
+					FPlatformProcess::Sleep(0.001f);
+				}
+			}
+			
+		}
+		else
+		{
+			if (TickChildCookers(false) == false)
+			{
+				Result |= COSR_WaitingOnChildCookers;
+			}
+		}
 	}
+
 	while (!GIsRequestingExit || CurrentCookMode == ECookMode::CookByTheBook)
 	{
 		// if we just cooked a map then don't process anything the rest of this tick
@@ -1589,8 +1650,6 @@ uint32 UCookOnTheFlyServer::TickCookOnTheSide( const float TimeSlice, uint32 &Co
 		int32 FirstUnsolicitedPackage = PackagesToSave.Num();
 
 		// generate a list of other packages which were loaded with this one
-		// if we are a child cooker then we don't have unsolicited packages
-		if ( IsChildCooker() == false)
 		{
 			SCOPE_TIMER(UnsolicitedMarkup);
 
@@ -1638,8 +1697,16 @@ uint32 UCookOnTheFlyServer::TickCookOnTheSide( const float TimeSlice, uint32 &Co
 
 						if ( StandardPackageFName != NAME_None ) // if we have name none that means we are in core packages or something...
 						{
-							// check if the package has already been saved
-							PackagesToSave.AddUnique( Package );
+							if (IsChildCooker())
+							{
+								// notify the main cooker that it should make sure this package gets cooked
+								CookByTheBookOptions->ChildUnsolicitedPackages.Add(StandardPackageFName);
+							}
+							else
+							{
+								// check if the package has already been saved
+								PackagesToSave.AddUnique(Package);
+							}
 						}
 					}
 				}
@@ -1876,9 +1943,9 @@ uint32 UCookOnTheFlyServer::TickCookOnTheSide( const float TimeSlice, uint32 &Co
 		CookByTheBookOptions->CookTime += Timer.GetTimeTillNow();
 	}
 
-
 	if ( IsCookByTheBookRunning() && 
-		( CookRequests.HasItems() == false ) )
+		( CookRequests.HasItems() == false ) &&
+		(!(Result & COSR_WaitingOnChildCookers)) )
 	{
 		check(IsCookByTheBookMode());
 		// if we are out of stuff and we are in cook by the book from the editor mode then we finish up
@@ -3017,10 +3084,11 @@ void UCookOnTheFlyServer::AddFileToCook( TArray<FName>& InOutFilesToCook, const
 	}
 }
 
-void UCookOnTheFlyServer::CollectFilesToCook(TArray<FName>& FilesInPath, const TArray<FString> &CookMaps, const TArray<FString> &InCookDirectories, const TArray<FString> &CookCultures, const TArray<FString> &IniMapSections, bool bCookAll, bool bMapsOnly, bool bNoDev, const FString& ChildCookFilename )
+void UCookOnTheFlyServer::CollectFilesToCook(TArray<FName>& FilesInPath, const TArray<FString> &CookMaps, const TArray<FString> &InCookDirectories, const TArray<FString> &CookCultures, const TArray<FString> &IniMapSections, bool bCookAll, bool bMapsOnly, bool bNoDev )
 {
 	if (CookByTheBookOptions->bIsChildCooker)
 	{
+		const FString ChildCookFilename = CookByTheBookOptions->ChildCookFilename;
 		check(ChildCookFilename.Len());
 		FString ChildCookString;
 		ensure(FFileHelper::LoadFileToString(ChildCookString, *ChildCookFilename));
@@ -3342,7 +3410,6 @@ FString UCookOnTheFlyServer::ConvertToFullSandboxPath( const FString &FileName,
 	return Result;
 }
 
-
 FString UCookOnTheFlyServer::ConvertToFullSandboxPath( const FString &FileName, bool bForWrite, const FString& PlatformName ) const
 {
 	FString Result = ConvertToFullSandboxPath( FileName, bForWrite );
@@ -3374,18 +3441,19 @@ void UCookOnTheFlyServer::CookByTheBookFinished()
 	GetDerivedDataCacheRef().WaitForQuiescence(true);
 	
 
+	if (IsChildCooker())
 	{
-		SCOPE_TIMER(WaitForChildCookersToFinish);
-		// wait for child cookers to finish
-		bool CookersFinished = false;
-		while (TickChildCookers() == false)
+		// if we are the child cooker create a list of all the packages which we think wes hould have cooked but didn't
+		FString UncookedPackageList;
+		for (const auto& UncookedPackage : CookByTheBookOptions->ChildUnsolicitedPackages)
 		{
-			// wait for cookers to finish
-			FPlatformProcess::Sleep(0.01f);
+			UncookedPackageList.Append(UncookedPackage.ToString() + TEXT("\n\r"));
 		}
-
-		CookByTheBookOptions->ChildCookers.Empty();
-		
+		FFileHelper::SaveStringToFile(UncookedPackageList, *GetChildCookerResultFilename(CookByTheBookOptions->ChildCookFilename));
+	}
+	else
+	{
+		check(CookByTheBookOptions->ChildUnsolicitedPackages.Num() == 0);
 	}
 
 	if (IsChildCooker() == false)
@@ -3449,7 +3517,7 @@ void UCookOnTheFlyServer::CookByTheBookFinished()
 		}
 	}
 
-	if (CookByTheBookOptions->bGenerateDependenciesForMaps && IsChildCooker() == false)
+	if (CookByTheBookOptions->bGenerateDependenciesForMaps && (IsChildCooker() == false))
 	{
 		for (auto& MapDependencyGraphIt : CookByTheBookOptions->MapDependencyGraphs)
 		{
@@ -3688,6 +3756,16 @@ void UCookOnTheFlyServer::StartCookByTheBook( const FCookByTheBookStartupOptions
 	CookByTheBookOptions->bForceEnableCompressedPackages = !!(CookOptions & ECookByTheBookOptions::ForceEnableCompressed);
 	CookByTheBookOptions->bForceDisableCompressedPackages = !!(CookOptions & ECookByTheBookOptions::ForceDisableCompressed);
 	CookByTheBookOptions->bIsChildCooker = CookByTheBookStartupOptions.ChildCookFileName.Len() > 0 ? true : false;
+	CookByTheBookOptions->ChildCookFilename = CookByTheBookStartupOptions.ChildCookFileName;
+
+
+	CookByTheBookOptions->TargetPlatformNames.Empty();
+	for (const auto &Platform : TargetPlatforms)
+	{
+		FName PlatformName = FName(*Platform->PlatformName());
+		CookByTheBookOptions->TargetPlatformNames.Add(PlatformName); // build list of all target platform names
+	}
+	const TArray<FName>& TargetPlatformNames = CookByTheBookOptions->TargetPlatformNames;
 
 	bool bCookAll = (CookOptions & ECookByTheBookOptions::CookAll) != ECookByTheBookOptions::None;
 	bool bMapsOnly = (CookOptions & ECookByTheBookOptions::MapsOnly) != ECookByTheBookOptions::None;
@@ -3706,17 +3784,12 @@ void UCookOnTheFlyServer::StartCookByTheBook( const FCookByTheBookStartupOptions
 
 	InitializeSandbox();
 
-	if (IsChildCooker() == false)
-	{
-		// iteratively clean any old files out of the sandbox (check if ini settings are out of date and clean out any stale files)
-		CleanSandbox(true);
-	}
-	
+
 	CookByTheBookOptions->bLeakTest = (CookOptions & ECookByTheBookOptions::LeakTest) != ECookByTheBookOptions::None; // this won't work from the editor this needs to be standalone
-	check( !CookByTheBookOptions->bLeakTest || CurrentCookMode == ECookMode::CookByTheBook );
+	check(!CookByTheBookOptions->bLeakTest || CurrentCookMode == ECookMode::CookByTheBook);
 
 	CookByTheBookOptions->LastGCItems.Empty();
-	if (CookByTheBookOptions->bLeakTest )
+	if (CookByTheBookOptions->bLeakTest)
 	{
 		for (FObjectIterator It; It; ++It)
 		{
@@ -3724,17 +3797,11 @@ void UCookOnTheFlyServer::StartCookByTheBook( const FCookByTheBookStartupOptions
 		}
 	}
 
-
-	TArray<FName> TargetPlatformNames;
-	for (const auto &Platform : TargetPlatforms)
-	{
-		FName PlatformName = FName(*Platform->PlatformName());
-		TargetPlatformNames.Add(PlatformName); // build list of all target platform names
-	}
-
-
 	if (!IsChildCooker())
 	{
+		// iteratively clean any old files out of the sandbox (check if ini settings are out of date and clean out any stale files)
+		CleanSandbox(true);
+
 		for (const auto &Platform : TargetPlatforms)
 		{
 			FName PlatformName = FName(*Platform->PlatformName());
@@ -3766,6 +3833,7 @@ void UCookOnTheFlyServer::StartCookByTheBook( const FCookByTheBookStartupOptions
 			}
 		}
 	}
+
 	if ( IsCookingDLC() )
 	{
 		// if we are cooking dlc we must be based of a release version cook
@@ -3814,7 +3882,7 @@ void UCookOnTheFlyServer::StartCookByTheBook( const FCookByTheBookStartupOptions
 	}
 	
 	
-	CollectFilesToCook(FilesInPath, CookMaps, CookDirectories, CookCultures, IniMapSections, bCookAll, bMapsOnly, bNoDev, CookByTheBookStartupOptions.ChildCookFileName );
+	CollectFilesToCook(FilesInPath, CookMaps, CookDirectories, CookCultures, IniMapSections, bCookAll, bMapsOnly, bNoDev );
 	if (FilesInPath.Num() == 0)
 	{
 		LogCookerMessage( FString::Printf(TEXT("No files found to cook.")), EMessageSeverity::Warning );
@@ -3904,17 +3972,23 @@ void UCookOnTheFlyServer::StartChildCookers(int32 NumCookersToSpawn, const TArra
 	
 	check(!IsChildCooker());
 
+	// PackageNames contains a sorted list of the packages we want to distribute to child cookers
 	TArray<FName> PackageNames;
+	// PackageNamesSet is a set of the same information for quick lookup into packagenames array
+	TSet<FName> PackageNamesSet;
 	PackageNames.Empty(CookRequests.Num());
-	
 
 	// TArray<FFilePlatformRequest>;
 
+	
+
 	for (const auto& CookRequest : CookRequests.GetQueue())
 	{
 		FString LongPackageName = FPackageName::FilenameToLongPackageName(CookRequest.ToString());
 
-		PackageNames.Add(FName(*LongPackageName));
+		const FName PackageFName = FName(*LongPackageName);
+		PackageNames.Add(PackageFName);
+		PackageNamesSet.Add(PackageFName);
 	}
 
 
@@ -3923,9 +3997,6 @@ void UCookOnTheFlyServer::StartChildCookers(int32 NumCookersToSpawn, const TArra
 	IAssetRegistry& AssetRegistry = AssetRegistryModule.Get();
 
 
-
-	TArray<FName> DistributeCandidates;
-
 	int32 PackageCounter = 0;
 	while (PackageCounter < PackageNames.Num())
 	{
@@ -3934,30 +4005,21 @@ void UCookOnTheFlyServer::StartChildCookers(int32 NumCookersToSpawn, const TArra
 		TArray<FName> UnfilteredDependencies;
 		AssetRegistry.GetDependencies(PackageName, UnfilteredDependencies);
 
-		TArray<FName> Dependencies;
-		for (const auto& DependencyName : UnfilteredDependencies)
-		{
-			if (FPackageName::IsScriptPackage(DependencyName.ToString()) == false)
-			{
-				Dependencies.Add(DependencyName);
-			}
-		}
-
-		if (Dependencies.Num() == 0)
-		{
-			DistributeCandidates.AddUnique(PackageName);
-		}
-		else
+		for (const auto& Dependency : UnfilteredDependencies)
 		{
-			for (const auto& Dependency : Dependencies)
+			if (FPackageName::IsScriptPackage(Dependency.ToString()) == false)
 			{
-				PackageNames.AddUnique(Dependency);
+				if (PackageNamesSet.Contains(Dependency) == false)
+				{
+					PackageNamesSet.Add(Dependency);
+					PackageNames.Insert(Dependency, PackageCounter);
+				}
 			}
 		}
 	}
 
 	TArray<FName> DistributeStandardFilenames;
-	for (const auto& DistributeCandidate : DistributeCandidates)
+	for (const auto& DistributeCandidate : PackageNames)
 	{
 		FText OutReason;
 		FString LongPackageName = DistributeCandidate.ToString();
@@ -3992,8 +4054,10 @@ void UCookOnTheFlyServer::StartChildCookers(int32 NumCookersToSpawn, const TArra
 	// start the child cookers and give them each some distribution candidates
 	for (int32 CookerCounter = 0; CookerCounter < NumCookersToSpawn; ++CookerCounter )
 	{
-		// don't spawn a cooker unless it has atleast 5 files to do
-		int32 NumFilesForCooker = DistributeStandardFilenames.Num() / (NumCookersToSpawn - CookerCounter);
+		// count our selves as a cooker
+		int32 NumFilesForCooker = DistributeStandardFilenames.Num() / ((NumCookersToSpawn+1) - CookerCounter);
+
+		// don't spawn a cooker unless it has a minimum amount of files to do
 		if (NumFilesForCooker < 5)
 		{
 			continue;
@@ -4084,6 +4148,7 @@ void UCookOnTheFlyServer::MaybeMarkPackageAsAlreadyLoaded(UPackage *Package)
 
 	FName StandardName = GetCachedStandardPackageFileFName(Package);
 
+	bool bShouldMarkAsAlreadyProcessed = false;
 
 	TArray<FName> CookedPlatforms;
 	if (CookedPackages.GetCookedPlatforms(StandardName, CookedPlatforms))
@@ -4094,7 +4159,25 @@ void UCookOnTheFlyServer::MaybeMarkPackageAsAlreadyLoaded(UPackage *Package)
 			Platforms += TEXT(" ");
 			Platforms += CookedPlatform.ToString();
 		}
+		bShouldMarkAsAlreadyProcessed = true;
 		UE_LOG(LogCook, Display, TEXT("Marking %s as reloading for cooker because it's been cooked for platforms%s."), *StandardName.ToString(), *Platforms);
+	}
+
+	/*
+	// if we are currently loading a package then the request would already be dequed so this check doesn't work
+	// need to figure out a better way to do this 
+	if (IsChildCooker())
+	{
+		// if we are the child cooker then we will never save anything that isn't in the cook requests
+		if (CookRequests.Exists(StandardName) == false)
+		{
+			bShouldMarkAsAlreadyProcessed = true;
+		}
+	}
+	*/
+
+	if (bShouldMarkAsAlreadyProcessed)
+	{
 		Package->PackageFlags |= PKG_ReloadingForCooker;
 	}
 
