#ANIM: fixed so that default setting not working properly for importing animation

diff --git a/Engine/Source/Editor/UnrealEd/Private/Fbx/FbxMainImport.cpp b/Engine/Source/Editor/UnrealEd/Private/Fbx/FbxMainImport.cpp
--- a/Engine/Source/Editor/UnrealEd/Private/Fbx/FbxMainImport.cpp
+++ b/Engine/Source/Editor/UnrealEd/Private/Fbx/FbxMainImport.cpp
@@ -419,7 +419,6 @@ int32 FFbxImporter::GetImportType(const FString& InFilename)
 			// Check for animation data. It can be overwritten by Geometry or Deformer
 			else if ( (ItemName == "AnimationCurve" || ItemName == "AnimationCurveNode") && ItemCount > 0 )
 			{
-				Result = 1;
 				bHasAnimation = true;
 			}
 		}
@@ -433,6 +432,10 @@ int32 FFbxImporter::GetImportType(const FString& InFilename)
 			// If animation data is found, set the result to 2, otherwise default to static mesh
 			Result = bHasAnimation ? 2 : 0;
 		}
+		else if ( Result == 0 && bHasAnimation )
+		{
+			Result = 1;
+		}
 	}
 	
 	return Result; 
