Removed warning when local variables can not be found, they are queried often by [REDACTED_PERSON] nodes and it is a non-problem.

diff --git a/Engine/Source/Editor/UnrealEd/Private/Kismet2/BlueprintEditorUtils.cpp b/Engine/Source/Editor/UnrealEd/Private/Kismet2/BlueprintEditorUtils.cpp
--- a/Engine/Source/Editor/UnrealEd/Private/Kismet2/BlueprintEditorUtils.cpp
+++ b/Engine/Source/Editor/UnrealEd/Private/Kismet2/BlueprintEditorUtils.cpp
@@ -4787,11 +4787,6 @@ FBPVariableDescription* FBlueprintEditorUtils::FindLocalVariable(const UBlueprin
 				break;
 			}
 		}
-
-		if(!ReturnVariable)
-		{
-			UE_LOG(LogBlueprint, Warning, TEXT("Could not find local variable '%s'!"), *InVariableName.ToString());
-		}
 	}
 
 	return ReturnVariable;
