Can now edit PlayerStateClass and SpectatorClass for a game mode via the Project Settings menu (UE-15341).

diff --git a/Engine/Source/Editor/UnrealEd/Public/GameModeInfoCustomizer.h b/Engine/Source/Editor/UnrealEd/Public/GameModeInfoCustomizer.h
--- a/Engine/Source/Editor/UnrealEd/Public/GameModeInfoCustomizer.h
+++ b/Engine/Source/Editor/UnrealEd/Public/GameModeInfoCustomizer.h
@@ -155,12 +155,13 @@ public:
 		static FName SelectedGameModeDetailsName(TEXT("SelectedGameModeDetails"));		
 		IDetailGroup& Group = CategoryBuilder.AddGroup(SelectedGameModeDetailsName, LOCTEXT("SelectedGameModeDetails", "Selected GameMode"));
 
-
 		// Then add rows to show key properties and let you edit them
 		CustomizeGameModeDefaultClass(Group, GET_MEMBER_NAME_CHECKED(AGameMode, DefaultPawnClass));
 		CustomizeGameModeDefaultClass(Group, GET_MEMBER_NAME_CHECKED(AGameMode, HUDClass));
 		CustomizeGameModeDefaultClass(Group, GET_MEMBER_NAME_CHECKED(AGameMode, PlayerControllerClass));
 		CustomizeGameModeDefaultClass(Group, GET_MEMBER_NAME_CHECKED(AGameMode, GameStateClass));
+		CustomizeGameModeDefaultClass(Group, GET_MEMBER_NAME_CHECKED(AGameMode, PlayerStateClass));
+		CustomizeGameModeDefaultClass(Group, GET_MEMBER_NAME_CHECKED(AGameMode, SpectatorClass));
 	}
 
 	/** Get the currently set GameMode class */
