Stopped level blueprints from force being recompiled during launch on.

diff --git a/Engine/Source/Editor/UnrealEd/Private/PlayLevel.cpp b/Engine/Source/Editor/UnrealEd/Private/PlayLevel.cpp
--- a/Engine/Source/Editor/UnrealEd/Private/PlayLevel.cpp
+++ b/Engine/Source/Editor/UnrealEd/Private/PlayLevel.cpp
@@ -1481,7 +1481,7 @@ static void HandleHyperlinkNavigate()
 
 struct FInternalPlayLevelUtils
 {
-	static int32 ResolveDirtyBlueprints(const bool bPromptForCompile, TArray<UBlueprint*>& ErroredBlueprints)
+	static int32 ResolveDirtyBlueprints(const bool bPromptForCompile, TArray<UBlueprint*>& ErroredBlueprints, const bool bForceLevelScriptRecompile = true)
 	{
 		const bool bAutoCompile = !bPromptForCompile;
 		FString PromtDirtyList;
@@ -1497,7 +1497,7 @@ struct FInternalPlayLevelUtils
 			// do not try to recompile BPs that have not changed since they last failed to compile, so don't check Blueprint->IsUpToDate()
 			const bool bIsDirtyAndShouldBeRecompiled = Blueprint->IsPossiblyDirty();
 			if (!FBlueprintEditorUtils::IsDataOnlyBlueprint(Blueprint)
-				&& (bIsDirtyAndShouldBeRecompiled || FBlueprintEditorUtils::IsLevelScriptBlueprint(Blueprint))
+				&& (bIsDirtyAndShouldBeRecompiled || (FBlueprintEditorUtils::IsLevelScriptBlueprint(Blueprint) && bForceLevelScriptRecompile))
 				&& (Blueprint->Status != BS_Unknown)
 				&& !Blueprint->IsPendingKill())
 			{
@@ -1686,7 +1686,7 @@ void UEditorEngine::PlayUsingLauncher()
 		}
 
 		TArray<UBlueprint*> ErroredBlueprints;
-		FInternalPlayLevelUtils::ResolveDirtyBlueprints(!PlayInSettings->bAutoCompileBlueprintsOnLaunch, ErroredBlueprints);
+		FInternalPlayLevelUtils::ResolveDirtyBlueprints(!PlayInSettings->bAutoCompileBlueprintsOnLaunch, ErroredBlueprints, false);
 
 		TArray<FString> MapNames;
 		FWorldContext & EditorContext = GetEditorWorldContext();
