Fixed issue where sometimes redirector fixup needed to be run twice in order to complete successfully.

Added referencing assets' paths into the list in FAssetFixUpRedirectors::UpdateAssetReferencers so that the asset references lists are updated correctly.
Also fixed a problem leading to duplicate entries in the PendingDeletes list in FAssetDeleteModel.

diff --git a/Engine/Source/Developer/AssetTools/Private/AssetFixUpRedirectors.cpp b/Engine/Source/Developer/AssetTools/Private/AssetFixUpRedirectors.cpp
--- a/Engine/Source/Developer/AssetTools/Private/AssetFixUpRedirectors.cpp
+++ b/Engine/Source/Developer/AssetTools/Private/AssetFixUpRedirectors.cpp
@@ -334,6 +334,11 @@ void FAssetFixUpRedirectors::UpdateAssetReferencers(const TArray<FRedirectorRefs
 	for (const auto& Redirector : RedirectorsToFix)
 	{
 		AssetPaths.AddUnique(FPackageName::GetLongPackagePath(Redirector.RedirectorPackageName.ToString()) / TEXT("")); // Ensure trailing slash
+
+		for (const auto& Referencer : Redirector.ReferencingPackageNames)
+		{
+			AssetPaths.AddUnique(FPackageName::GetLongPackagePath(Referencer.ToString()) / TEXT("")); // Ensure trailing slash
+		}
 	}
 	AssetRegistryModule.Get().ScanPathsSynchronous(AssetPaths, true);
 }
diff --git a/Engine/Source/Editor/UnrealEd/Private/AssetDeleteModel.cpp b/Engine/Source/Editor/UnrealEd/Private/AssetDeleteModel.cpp
--- a/Engine/Source/Editor/UnrealEd/Private/AssetDeleteModel.cpp
+++ b/Engine/Source/Editor/UnrealEd/Private/AssetDeleteModel.cpp
@@ -38,7 +38,10 @@ FAssetDeleteModel::~FAssetDeleteModel()
 
 void FAssetDeleteModel::AddObjectToDelete( UObject* InObject )
 {
-	PendingDeletes.Add(MakeShareable(new FPendingDelete(InObject)));
+	if (!PendingDeletes.ContainsByPredicate([=](const TSharedPtr<FPendingDelete>& A) { return A->GetObject() == InObject; }))
+	{
+		PendingDeletes.Add(MakeShareable(new FPendingDelete(InObject)));
+	}
 
 	PrepareToDelete(InObject);
 
@@ -544,7 +547,10 @@ void FAssetDeleteModel::PrepareToDelete(UObject* InObject)
 			if ( PackageMetaData )
 			{
 				PackageMetaData->RemoveFromRoot();
-				PendingDeletes.AddUnique(MakeShareable(new FPendingDelete(PackageMetaData)));
+				if (!PendingDeletes.ContainsByPredicate([=](const TSharedPtr<FPendingDelete>& A) { return A->GetObject() == PackageMetaData; }))
+				{
+					PendingDeletes.Add(MakeShareable(new FPendingDelete(PackageMetaData)));
+				}
 			}
 		}
 	}
