Fix so that cancelling attempted delete of read only files does not remove the asset from the Content Browser.

diff --git a/Engine/Source/Editor/UnrealEd/Private/ObjectTools.cpp b/Engine/Source/Editor/UnrealEd/Private/ObjectTools.cpp
--- a/Engine/Source/Editor/UnrealEd/Private/ObjectTools.cpp
+++ b/Engine/Source/Editor/UnrealEd/Private/ObjectTools.cpp
@@ -1666,14 +1666,50 @@ namespace ObjectTools
 		return DeleteModel->GetDeletedObjectCount();
 	}
 
-	int32 DeleteObjectsUnchecked( const TArray< UObject* >& ObjectsToDelete )
+	static bool MakeReadOnlyPackageWritable(UObject* ObjectToDelete, bool& bMakeWritable)
+	{
+		// If an object's package is read only, and source control is not enabled, ask the user whether they wish
+		// to make it writable.
+		if (!ISourceControlModule::Get().IsEnabled())
+		{
+			UPackage* ObjectPackage = ObjectToDelete->GetOutermost();
+			check(ObjectPackage != nullptr);
+
+			FString PackageFilename;
+			if (FPackageName::DoesPackageExist(ObjectPackage->GetName(), nullptr, &PackageFilename))
+			{
+				if (IFileManager::Get().IsReadOnly(*PackageFilename))
+				{
+					EAppReturnType::Type ReturnType = EAppReturnType::No;
+					if (!bMakeWritable)
+					{
+						ReturnType = FMessageDialog::Open(EAppMsgType::YesNoYesAll, NSLOCTEXT("ObjectTools", "DeleteReadOnlyWarning", "File is read-only on disk, are you sure you want to delete it?"));
+						bMakeWritable = ReturnType == EAppReturnType::YesAll;
+					}
+
+					if (bMakeWritable || ReturnType == EAppReturnType::Yes)
+					{
+						FPlatformFileManager::Get().GetPlatformFile().SetReadOnly(*PackageFilename, false);
+					}
+					else
+					{
+						return false;
+					}
+				}
+			}
+		}
+
+		return true;
+	}
+
+	int32 DeleteObjectsUnchecked(const TArray< UObject* >& ObjectsToDelete)
 	{
 		GWarn->BeginSlowTask( NSLOCTEXT( "UnrealEd", "Deleting", "Deleting" ), true );
 
 		TArray<UObject*> ObjectsDeletedSuccessfully;
-		TArray<UObject*> ObjectsDeletedUnsuccessfully;
 
-		bool bSawSuccessfulDelete = true;
+		bool bSawSuccessfulDelete = false;
+		bool bMakeWritable = false;
 
 		for ( int32 Index = 0; Index < ObjectsToDelete.Num(); Index++ )
 		{
@@ -1685,16 +1721,18 @@ namespace ObjectTools
 				continue;
 			}
 
+			// Early exclusion for assets contained in read-only packages if the user chooses not to write enable them
+			if (!MakeReadOnlyPackageWritable(ObjectToDelete, bMakeWritable))
+			{
+				continue;
+			}
+
 			// We already know it's not referenced or we wouldn't be performing the safe delete, so don't repeat the reference check.
 			bool bPerformReferenceCheck = false;
 			if ( DeleteSingleObject( ObjectToDelete, bPerformReferenceCheck ) )
 			{
 				ObjectsDeletedSuccessfully.Push( ObjectToDelete );
-			}
-			else
-			{
-				ObjectsDeletedUnsuccessfully.Push( ObjectToDelete );
-				bSawSuccessfulDelete = false;
+				bSawSuccessfulDelete = true;
 			}
 		}
 
@@ -1819,7 +1857,8 @@ namespace ObjectTools
 		TArray<AActor*> ActorsToDelete;
 		TArray<UObject*> ObjectsToDelete;
 		bool bNeedsGarbageCollection = false;
-	
+		bool bMakeWritable = false;
+
 		// Clear audio components to allow previewed sounds to be consolidated
 		GEditor->ClearPreviewComponents();
 
@@ -1829,6 +1868,12 @@ namespace ObjectTools
 
 			GEditor->GetSelectedObjects()->Deselect( CurrentObject );
 
+			// Early exclusion for assets contained in read-only packages if the user chooses not to write enable them
+			if (!MakeReadOnlyPackageWritable(CurrentObject, bMakeWritable))
+			{
+				continue;
+			}
+
 			ObjectsToDelete.Add( CurrentObject );
 
 			// If the object about to be deleted is a Blueprint asset, make sure that any instances of the Blueprint class get deleted as well
@@ -2102,7 +2147,10 @@ namespace ObjectTools
 			PotentialPackagesToDelete.AddUnique(ObjectsToDelete[ObjIdx]->GetOutermost());
 		}
 
-		CleanupAfterSuccessfulDelete(PotentialPackagesToDelete);
+		if (PotentialPackagesToDelete.Num() > 0)
+		{
+			CleanupAfterSuccessfulDelete(PotentialPackagesToDelete);
+		}
 		ObjectsToDelete.Empty();
 
 		GWarn->EndSlowTask();
