Handle case where levels are being loaded by double clicking in the content browser. This has the effect of loading the level package before LoadMap is called and so we end up calling UnloadPackage on the world before it's been properly set up.

#UE-17205

diff --git a/Engine/Source/Editor/UnrealEd/Private/EditorServer.cpp b/Engine/Source/Editor/UnrealEd/Private/EditorServer.cpp
--- a/Engine/Source/Editor/UnrealEd/Private/EditorServer.cpp
+++ b/Engine/Source/Editor/UnrealEd/Private/EditorServer.cpp
@@ -2254,7 +2254,9 @@ bool UEditorEngine::Map_Load(const TCHAR* Str, FOutputDevice& Ar)
 					for (TObjectIterator<UWorld> It; It; ++It)
 					{
 						UPackage* Package = Cast<UPackage>(It->GetOuter());
-						if (Package && Package != GetTransientPackage())
+						
+
+						if (Package && Package != GetTransientPackage() && Package->GetPathName() != LongTempFname)
 						{
 							WorldPackages.AddUnique(Package);
 						}
