Writing manifest to file now uses full/absolute/rooted path to avoid any ambiguity when debugging.

diff --git a/Engine/Source/Editor/UnrealEd/Private/Commandlets/GenerateGatherManifestCommandlet.cpp b/Engine/Source/Editor/UnrealEd/Private/Commandlets/GenerateGatherManifestCommandlet.cpp
--- a/Engine/Source/Editor/UnrealEd/Private/Commandlets/GenerateGatherManifestCommandlet.cpp
+++ b/Engine/Source/Editor/UnrealEd/Private/Commandlets/GenerateGatherManifestCommandlet.cpp
@@ -83,7 +83,7 @@ int32 UGenerateGatherManifestCommandlet::Main( const FString& Params )
 		ManifestInfo->ApplyManifestDependencies();
 	}
 	
-	const FString ManifestPath = DestinationPath / ManifestName;
+	const FString ManifestPath = FPaths::ConvertRelativePathToFull(DestinationPath) / ManifestName;
 	if( !WriteManifestToFile( ManifestInfo->GetManifest(), ManifestPath ) )
 	{
 		UE_LOG( LogGenerateManifestCommandlet, Error,TEXT("Failed to write manifest to %s."), *ManifestPath );				
