Resolved some edge cases with source control integration in localization JSON file writing. No longer fails if source control is enabled but the output files don't already exist on disk. IE: No longer fails on initial creation of localization JSON files with source control enabled.

diff --git a/Engine/Source/Editor/UnrealEd/Private/Commandlets/GatherTextCommandletBase.cpp b/Engine/Source/Editor/UnrealEd/Private/Commandlets/GatherTextCommandletBase.cpp
--- a/Engine/Source/Editor/UnrealEd/Private/Commandlets/GatherTextCommandletBase.cpp
+++ b/Engine/Source/Editor/UnrealEd/Private/Commandlets/GatherTextCommandletBase.cpp
@@ -63,13 +63,16 @@ bool UGatherTextCommandletBase::WriteJSONToTextFile(TSharedPtr<FJsonObject> Outp
 	// If the user specified a reference file - write the entries read from code to a ref file
 	if ( !Filename.IsEmpty() && Output.IsValid() )
 	{
-		if( SourceControl.IsValid() )
+		const bool DidFileExist = FPaths::FileExists(Filename);
+		if (DidFileExist)
 		{
-			FText SCCErrorText;
-			if (!SourceControl->CheckOutFile(Filename, SCCErrorText))
+			if( SourceControl.IsValid() )
 			{
-				UE_LOG(LogGatherTextCommandletBase, Error, TEXT("Check out of file %s failed: %s"), *Filename, *SCCErrorText.ToString());
-				WasSuccessful = false;
+				FText SCCErrorText;
+				if (!SourceControl->CheckOutFile(Filename, SCCErrorText))
+				{
+					UE_LOG(LogGatherTextCommandletBase, Error, TEXT("Check out of file %s failed: %s"), *Filename, *SCCErrorText.ToString());
+				}
 			}
 		}
 
@@ -86,6 +89,18 @@ bool UGatherTextCommandletBase::WriteJSONToTextFile(TSharedPtr<FJsonObject> Outp
 				WasSuccessful = false;
 			}
 		}
+
+		if (!DidFileExist)
+		{
+			if( SourceControl.IsValid() )
+			{
+				FText SCCErrorText;
+				if (!SourceControl->CheckOutFile(Filename, SCCErrorText))
+				{
+					UE_LOG(LogGatherTextCommandletBase, Error, TEXT("Check out of file %s failed: %s"), *Filename, *SCCErrorText.ToString());
+				}
+			}
+		}
 	}
 	else
 	{
