Avoid updating components in sub-levels more than once when new map is loaded in the editor

diff --git a/Engine/Source/Editor/UnrealEd/Private/EditorServer.cpp b/Engine/Source/Editor/UnrealEd/Private/EditorServer.cpp
--- a/Engine/Source/Editor/UnrealEd/Private/EditorServer.cpp
+++ b/Engine/Source/Editor/UnrealEd/Private/EditorServer.cpp
@@ -2042,6 +2042,8 @@ UWorld* UEditorEngine::NewMap()
 	Context.SetCurrentWorld(NewWorld);
 	GWorld = NewWorld;
 	NewWorld->AddToRoot();
+	// Register components in the persistent level (current)
+	NewWorld->UpdateWorldComponents(true, true);
 
 	NoteSelectionChange();
 
@@ -2414,6 +2416,9 @@ bool UEditorEngine::Map_Load(const TCHAR* Str, FOutputDevice& Ar)
 						Context.World()->GetModel()->Polys->SetFlags( RF_Transactional );
 					}
 
+					// Register components in the persistent level (current)
+					Context.World()->UpdateWorldComponents(true, true);
+
 					// Make sure secondary levels are loaded & visible.
 					Context.World()->FlushLevelStreaming();
 
diff --git a/Engine/Source/Editor/UnrealEd/Private/UnrealEdMisc.cpp b/Engine/Source/Editor/UnrealEd/Private/UnrealEdMisc.cpp
--- a/Engine/Source/Editor/UnrealEd/Private/UnrealEdMisc.cpp
+++ b/Engine/Source/Editor/UnrealEd/Private/UnrealEdMisc.cpp
@@ -920,20 +920,22 @@ void FUnrealEdMisc::CB_MapChange( uint32 InFlags )
 	const FString EmptyString(TEXT(""));
 	GEditor->SetPropertyColorationTarget( World, EmptyString, NULL, NULL, NULL );
 
-	// Rebuild the collision hash if this map change is something major ("new", "open", etc).
-	// Minor things like brush subtraction will set it to "0".
-
-	if( InFlags != MapChangeEventFlags::Default )
+	if (InFlags != MapChangeEventFlags::NewMap)
 	{
-		World->ClearWorldComponents();
+		// Rebuild the collision hash if this map change was rebuilt
+		// Minor things like brush subtraction will set it to "0".
+		if (InFlags != MapChangeEventFlags::Default)
+		{
+			World->ClearWorldComponents();
 
-		// Note: CleanupWorld is being abused here to detach components and some other stuff
-		// CleanupWorld should only be called before destroying the world
-		// So bCleanupResources is being passed as false
-		World->CleanupWorld(true, false);
-	}
+			// Note: CleanupWorld is being abused here to detach components and some other stuff
+			// CleanupWorld should only be called before destroying the world
+			// So bCleanupResources is being passed as false
+			World->CleanupWorld(true, false);
+		}
 
-	GEditor->EditorUpdateComponents();
+		GEditor->EditorUpdateComponents();
+	}
 
 	GLevelEditorModeTools().MapChangeNotify();
 
