Revised code the mounts shared paths. (To try and pacify mac build)

diff --git a/Engine/Source/Editor/UnrealEd/Private/UnrealEdMisc.cpp b/Engine/Source/Editor/UnrealEd/Private/UnrealEdMisc.cpp
--- a/Engine/Source/Editor/UnrealEd/Private/UnrealEdMisc.cpp
+++ b/Engine/Source/Editor/UnrealEd/Private/UnrealEdMisc.cpp
@@ -1642,12 +1642,21 @@ void FUnrealEdMisc::MountTemplateSharedPaths()
 		for (int32 iShared = 0; iShared < TemplateInfo->SharedContentPacks.Num() ; iShared++)
 		{
 			EFeaturePackDetailLevel EachEditDetail = (EFeaturePackDetailLevel)EditDetail;
+			FString DetailString;
+			UEnum::GetValueAsString(TEXT("/Script/AddContentDialog.EFeaturePackDetailLevel"), EachEditDetail, DetailString);
+
 			FFeaturePackLevelSet EachPack = TemplateInfo->SharedContentPacks[iShared];
-			if (EachPack.DetailLevels.Num() == 1)
+			if((EachPack.DetailLevels.Num() == 1) && ( EachEditDetail !=  EachPack.DetailLevels[0] ))
 			{
 				// If theres only only detail level override the requirement with that
 				EachEditDetail = EachPack.DetailLevels[0];
-				UE_LOG(LogUnrealEdMisc, Warning, TEXT("Only 1 detail level defined for %s in %s."), *EachPack.MountName, *TemplateFilename );
+				// Get the name of the level we are falling back to so we can tell the user
+				FString FallbackDetailString;
+				UEnum::GetValueAsString(TEXT("/Script/AddContentDialog.EFeaturePackDetailLevel"), EachEditDetail, FallbackDetailString);
+				UE_LOG(LogUnrealEdMisc, Verbose, TEXT("Only 1 detail level defined for %s in %s. Cannot edit detail level %s. Will fallback to  "), *EachPack.MountName, *TemplateFilename, *DetailString,*FallbackDetailString);
+				// Then correct the string too !
+				DetailString = FallbackDetailString;
+
 			}
 			else if (EachPack.DetailLevels.Num() == 0)
 			{
@@ -1658,22 +1667,22 @@ void FUnrealEdMisc::MountTemplateSharedPaths()
 			for (int32 iDetail = 0; iDetail < EachPack.DetailLevels.Num(); iDetail++)
 			{
 				if (EachPack.DetailLevels[iDetail] == EachEditDetail)
-				{
-					FString DetailString;
-					UEnum::GetValueAsString(TEXT("/Script/AddContentDialog.EFeaturePackDetailLevel"), EachEditDetail, DetailString);
+				{					
 					FString ShareMountName = EachPack.MountName;
 					if (AddedMountSources.Find(ShareMountName) == INDEX_NONE)
 					{
-						FString RelativePath = FString::Printf(TEXT("../../../Templates/TemplateResources/%s/%s/Content"), *DetailString, *ShareMountName);
-						if( FPaths::DirectoryExists(RelativePath))
+						FString ResourcePath = FPaths::Combine(TEXT("Templates"), TEXT("TemplateResources"), *DetailString, *ShareMountName, TEXT("Content"));
+						FString FullPath = FPaths::Combine(*FPaths::RootDir(), *ResourcePath);
+
+						if (FPaths::DirectoryExists(FullPath))
 						{
-							FString MountName = FString::Printf(TEXT("/Game/%s/"), *ShareMountName);						
-							FPackageName::RegisterMountPoint(*MountName, *RelativePath);
+							FString MountName = FString::Printf(TEXT("/Game/%s/"), *ShareMountName);	
+							FPackageName::RegisterMountPoint(*MountName, *FullPath);
 							AddedMountSources.Add(ShareMountName);
 						}
 						else
 						{
-							UE_LOG(LogUnrealEdMisc, Warning, TEXT("Cannot find path %s to mount for %s resource in %s."), *RelativePath, *EachPack.MountName, *TemplateFilename );
+							UE_LOG(LogUnrealEdMisc, Warning, TEXT("Cannot find path %s to mount for %s resource in %s."), *FullPath, *EachPack.MountName, *TemplateFilename);
 						}
 					}
 				}
