Fixed bug where BSP surface textures might become broken when building geometry.

Fixed UEditorEngine::polyUpdateMaster so that the inverse rotation transform is performed with FMatrix::InverseTransformVector, instead of erroneously using a negated rotator.

diff --git a/Engine/Source/Editor/UnrealEd/Private/EditorCsg.cpp b/Engine/Source/Editor/UnrealEd/Private/EditorCsg.cpp
--- a/Engine/Source/Editor/UnrealEd/Private/EditorCsg.cpp
+++ b/Engine/Source/Editor/UnrealEd/Private/EditorCsg.cpp
@@ -375,7 +375,7 @@ void UEditorEngine::polyUpdateMaster
 		ActorLocation = Brush->OwnerLocationWhenLastBuilt;
 		ActorPrePivot = Brush->OwnerPrepivotWhenLastBuilt;
 		ActorScale = Brush->OwnerScaleWhenLastBuilt;
-		ActorRotation = -Brush->OwnerRotationWhenLastBuilt;
+		ActorRotation = Brush->OwnerRotationWhenLastBuilt;
 	}
 	else
 	{
@@ -383,22 +383,24 @@ void UEditorEngine::polyUpdateMaster
 		ActorLocation = Actor->GetActorLocation();
 		ActorPrePivot = Actor->GetPrePivot();
 		ActorScale = Actor->GetActorScale();
-		ActorRotation = -Actor->GetActorRotation();
+		ActorRotation = Actor->GetActorRotation();
 	}
 
-	for( int32 iEdPoly = Surf.iBrushPoly; iEdPoly < Brush->Polys->Element.Num(); iEdPoly++ )
+	const FRotationMatrix RotationMatrix(ActorRotation);
+
+	for (int32 iEdPoly = Surf.iBrushPoly; iEdPoly < Brush->Polys->Element.Num(); iEdPoly++)
 	{
 		FPoly& MasterEdPoly = Brush->Polys->Element[iEdPoly];
-		if( iEdPoly==Surf.iBrushPoly || MasterEdPoly.iLink==Surf.iBrushPoly )
+		if (iEdPoly == Surf.iBrushPoly || MasterEdPoly.iLink == Surf.iBrushPoly)
 		{
-			MasterEdPoly.Material  = Surf.Material;
+			MasterEdPoly.Material = Surf.Material;
 			MasterEdPoly.PolyFlags = Surf.PolyFlags & ~(PF_NoEdit);
 
-			if( UpdateTexCoords )
+			if (UpdateTexCoords)
 			{
-				MasterEdPoly.Base = ActorRotation.RotateVector(Model->Points[Surf.pBase] - ActorLocation) / ActorScale + ActorPrePivot;
-				MasterEdPoly.TextureU = ActorRotation.RotateVector(Model->Vectors[Surf.vTextureU]) * ActorScale;
-				MasterEdPoly.TextureV = ActorRotation.RotateVector(Model->Vectors[Surf.vTextureV]) * ActorScale;
+				MasterEdPoly.Base = RotationMatrix.InverseTransformVector(Model->Points[Surf.pBase] - ActorLocation) / ActorScale + ActorPrePivot;
+				MasterEdPoly.TextureU = RotationMatrix.InverseTransformVector(Model->Vectors[Surf.vTextureU]) * ActorScale;
+				MasterEdPoly.TextureV = RotationMatrix.InverseTransformVector(Model->Vectors[Surf.vTextureV]) * ActorScale;
 			}
 		}
 	}
