Fixed small error which was causing some Material Functions in the library to not appear in the right click menu, despite being exposed.

Changed to using range-based loops so the array name is only referenced by the loop itself, to guard against future problems of this sort.

diff --git a/Engine/Source/Editor/UnrealEd/Private/MaterialGraphSchema.cpp b/Engine/Source/Editor/UnrealEd/Private/MaterialGraphSchema.cpp
--- a/Engine/Source/Editor/UnrealEd/Private/MaterialGraphSchema.cpp
+++ b/Engine/Source/Editor/UnrealEd/Private/MaterialGraphSchema.cpp
@@ -695,10 +695,8 @@ void UMaterialGraphSchema::GetMaterialFunctionActions(FGraphActionMenuBuilder& A
 	TArray<FAssetData> AssetDataList;
 	AssetRegistryModule.Get().GetAssetsByClass(UMaterialFunction::StaticClass()->GetFName(), AssetDataList);
 
-	for (int32 AssetIndex = 0; AssetIndex < AssetDataList.Num(); ++AssetIndex)
+	for (const FAssetData& AssetData : AssetDataList)
 	{
-		const FAssetData& AssetData = AssetDataList[AssetIndex];
-
 		const bool bExposeToLibrary = AssetData.TagsAndValues.FindRef("bExposeToLibrary") == TEXT("TRUE");
 
 		// If this was a function that was selected to be exposed to the library
@@ -760,9 +758,8 @@ void UMaterialGraphSchema::GetMaterialFunctionActions(FGraphActionMenuBuilder& A
 				}
 
 				// For each category the function should belong to...
-				for (int32 CategoryIndex = 0; CategoryIndex < LibraryCategories.Num(); CategoryIndex++)
+				for (const FText& CategoryName : LibraryCategoriesText)
 				{
-					const FText& CategoryName = LibraryCategoriesText[CategoryIndex];
 					TSharedPtr<FMaterialGraphSchemaAction_NewFunctionCall> NewFunctionAction(new FMaterialGraphSchemaAction_NewFunctionCall(
 						CategoryName,
 						FText::FromString(FunctionName),
