Fixed minor typo in localization manifest generation commandlet log output.

diff --git a/Engine/Source/Editor/UnrealEd/Private/Commandlets/GenerateGatherManifestCommandlet.cpp b/Engine/Source/Editor/UnrealEd/Private/Commandlets/GenerateGatherManifestCommandlet.cpp
--- a/Engine/Source/Editor/UnrealEd/Private/Commandlets/GenerateGatherManifestCommandlet.cpp
+++ b/Engine/Source/Editor/UnrealEd/Private/Commandlets/GenerateGatherManifestCommandlet.cpp
@@ -94,7 +94,7 @@ int32 UGenerateGatherManifestCommandlet::Main( const FString& Params )
 
 bool UGenerateGatherManifestCommandlet::WriteManifestToFile( const TSharedPtr<FInternationalizationManifest>& InManifest, const FString& OutputFilePath )
 {
-	UE_LOG(LogGenerateManifestCommandlet, Log, TEXT("Writing archive to %s."), *OutputFilePath);
+	UE_LOG(LogGenerateManifestCommandlet, Log, TEXT("Writing manifest to %s."), *OutputFilePath);
 
 	// We can not continue if the provided manifest is not valid
 	if( !InManifest.IsValid() )
