[UE-17726] Invalidating function exports that we delete from data-only Blueprint packages (to keep other packages from referencing them as imports).

diff --git a/Engine/Source/Editor/UnrealEd/Private/Kismet2/BlueprintEditorUtils.cpp b/Engine/Source/Editor/UnrealEd/Private/Kismet2/BlueprintEditorUtils.cpp
--- a/Engine/Source/Editor/UnrealEd/Private/Kismet2/BlueprintEditorUtils.cpp
+++ b/Engine/Source/Editor/UnrealEd/Private/Kismet2/BlueprintEditorUtils.cpp
@@ -1398,8 +1398,14 @@ static void RemoveStaleFunctions(UBlueprintGeneratedClass* Class, UBlueprint* Bl
 
 		while (Fn)
 		{
-			Class->RemoveFunctionFromFunctionMap(*Fn);
-			Fn->Rename(nullptr, OrphanedClass, RenFlags);
+			UFunction* Function = *Fn;
+			Class->RemoveFunctionFromFunctionMap(Function);
+			Function->Rename(nullptr, OrphanedClass, RenFlags);
+
+			// invalidate this package's reference to this function, so 
+			// subsequent packages that import it will treat it as if it didn't 
+			// exist (because data-only blueprints shouldn't have functions)
+			FLinkerLoad::InvalidateExport(Function); 
 			++Fn;
 		}
 	}
