UE-18012: Crash occurs opening several maps in QAGame.

diff --git a/Engine/Source/Editor/UnrealEd/Private/EditorEngine.cpp b/Engine/Source/Editor/UnrealEd/Private/EditorEngine.cpp
--- a/Engine/Source/Editor/UnrealEd/Private/EditorEngine.cpp
+++ b/Engine/Source/Editor/UnrealEd/Private/EditorEngine.cpp
@@ -1804,20 +1804,43 @@ void UEditorEngine::Cleanse( bool ClearSelection, bool Redraw, const FText& Tran
 		// at load time.
 		//
 		// We also have to remove packages that redirectors were contained
-		// in, so they can be loaded again in the future. If we don't do it
-		// loading failure will occur next time someone tries to use it.
-		// This is caused by the fact that the loading routing will check
-		// that already existed, but the object was missing in cache.
+		// in if those were from redirector-only package, so they can be
+		// loaded again in the future. If we don't do it loading failure
+		// will occur next time someone tries to use it. This is caused by
+		// the fact that the loading routing will check that already
+		// existed, but the object was missing in cache.
+		const EObjectFlags FlagsToClear = RF_Standalone | RF_RootSet | RF_Transactional;
 		TArray<UPackage*> PackagesToUnload;
 		for (TObjectIterator<UObjectRedirector> RedirIt; RedirIt; ++RedirIt)
 		{
-			PackagesToUnload.Add(RedirIt->GetOutermost());
+			UPackage* RedirectorPackage = RedirIt->GetOutermost();
+
+			if (RedirectorPackage == GetTransientPackage())
+			{
+				continue;
+			}
+
+			TArray<UObject*> PackageObjects;
+			GetObjectsWithOuter(RedirectorPackage, PackageObjects);
+
+			if (!PackageObjects.ContainsByPredicate(
+					[](UObject* Object)
+					{
+						return !Object->IsA<UMetaData>() && !Object->IsA<UObjectRedirector>();
+					})
+				)
+			{
+				PackagesToUnload.Add(RedirectorPackage);
+			}
+			else
+			{
+				// In case this isn't redirector-only package, clear just the redirector.
+				RedirIt->ClearFlags(FlagsToClear);
+			}
 		}
 
 		for (auto* PackageToUnload : PackagesToUnload)
 		{
-			const EObjectFlags FlagsToClear = RF_Standalone | RF_RootSet | RF_Transactional;
-
 			TArray<UObject*> PackageObjects;
 			GetObjectsWithOuter(PackageToUnload, PackageObjects);
 			for (auto* Object : PackageObjects)
