Fixed Paper2D JSON files no longer importing as Paper2D assets

Data Tables can now import from JSON, and this was taking priority over the Paper2D importers. There were a few issues here:

1) FAssetTools::ImportAssets wasn't respecting the priority of the import factories, and was just testing them in the order they were found.
    a) UCSVImportFactory needed to be given a lower than normal priority so it gets tested last as it will perform a catch-all match.

2) FAssetTools::ImportAssets will default to the first factory it finds, and then use FactoryCanImport to try and find a better match, however FactoryCanImport returns false by default, so...
    a) UPaperSpriteSheetImportFactory needed to implement FactoryCanImport to avoid sprite sheets being imported as Data Tables.
    b) UCSVImportFactory needed to implement a catch-all FactoryCanImport so it can be used when no other factory is a better match.

diff --git a/Engine/Plugins/2D/Paper2D/Source/PaperSpriteSheetImporter/Private/PaperJsonSpriteSheetImporter.cpp b/Engine/Plugins/2D/Paper2D/Source/PaperSpriteSheetImporter/Private/PaperJsonSpriteSheetImporter.cpp
--- a/Engine/Plugins/2D/Paper2D/Source/PaperSpriteSheetImporter/Private/PaperJsonSpriteSheetImporter.cpp
+++ b/Engine/Plugins/2D/Paper2D/Source/PaperSpriteSheetImporter/Private/PaperJsonSpriteSheetImporter.cpp
@@ -415,6 +415,18 @@ bool FPaperJsonSpriteSheetImporter::Import(TSharedPtr<FJsonObject> SpriteDescrip
 	return bLoadedSuccessfully;
 }
 
+bool FPaperJsonSpriteSheetImporter::CanImportJSON(const FString& FileContents)
+{
+	TSharedPtr<FJsonObject> SpriteDescriptorObject = ParseJSON(FileContents, FString(), /*bSilent=*/ true);
+	if (SpriteDescriptorObject.IsValid())
+	{
+		FString Unused;
+		return ParseMetaBlock(FString(), SpriteDescriptorObject, /*out*/ Unused, /*bSilent=*/ true);
+	}
+	
+	return false;
+}
+
 bool FPaperJsonSpriteSheetImporter::ImportFromString(const FString& FileContents, const FString& NameForErrors, bool bSilent)
 {
 	TSharedPtr<FJsonObject> SpriteDescriptorObject = ParseJSON(FileContents, NameForErrors, bSilent);
diff --git a/Engine/Plugins/2D/Paper2D/Source/PaperSpriteSheetImporter/Private/PaperJsonSpriteSheetImporter.h b/Engine/Plugins/2D/Paper2D/Source/PaperSpriteSheetImporter/Private/PaperJsonSpriteSheetImporter.h
--- a/Engine/Plugins/2D/Paper2D/Source/PaperSpriteSheetImporter/Private/PaperJsonSpriteSheetImporter.h
+++ b/Engine/Plugins/2D/Paper2D/Source/PaperSpriteSheetImporter/Private/PaperJsonSpriteSheetImporter.h
@@ -36,6 +36,9 @@ class FPaperJsonSpriteSheetImporter
 {
 public:
 	FPaperJsonSpriteSheetImporter();
+
+	static bool CanImportJSON(const FString& FileContents);
+
 	bool ImportFromString(const FString& FileContents, const FString& NameForErrors, bool bSilent);
 	bool ImportFromArchive(FArchive* Archive, const FString& NameForErrors, bool bSilent);
 
diff --git a/Engine/Plugins/2D/Paper2D/Source/PaperSpriteSheetImporter/Private/PaperSpriteSheetImportFactory.cpp b/Engine/Plugins/2D/Paper2D/Source/PaperSpriteSheetImporter/Private/PaperSpriteSheetImportFactory.cpp
--- a/Engine/Plugins/2D/Paper2D/Source/PaperSpriteSheetImporter/Private/PaperSpriteSheetImportFactory.cpp
+++ b/Engine/Plugins/2D/Paper2D/Source/PaperSpriteSheetImporter/Private/PaperSpriteSheetImportFactory.cpp
@@ -32,6 +32,17 @@ FText UPaperSpriteSheetImportFactory::GetToolTip() const
 	return NSLOCTEXT("Paper2D", "PaperJsonImporterFactoryDescription", "Sprite sheets exported from Adobe Flash or Texture Packer");
 }
 
+bool UPaperSpriteSheetImportFactory::FactoryCanImport(const FString& Filename)
+{
+	FString FileContent;
+	if (FFileHelper::LoadFileToString(/*out*/ FileContent, *Filename))
+	{
+		return FPaperJsonSpriteSheetImporter::CanImportJSON(FileContent);
+	}
+
+	return false;
+}
+
 UObject* UPaperSpriteSheetImportFactory::FactoryCreateText(UClass* InClass, UObject* InParent, FName InName, EObjectFlags Flags, UObject* Context, const TCHAR* Type, const TCHAR*& Buffer, const TCHAR* BufferEnd, FFeedbackContext* Warn)
 {
 	Flags |= RF_Transactional;
diff --git a/Engine/Plugins/2D/Paper2D/Source/PaperSpriteSheetImporter/Private/PaperSpriteSheetImportFactory.h b/Engine/Plugins/2D/Paper2D/Source/PaperSpriteSheetImporter/Private/PaperSpriteSheetImportFactory.h
--- a/Engine/Plugins/2D/Paper2D/Source/PaperSpriteSheetImporter/Private/PaperSpriteSheetImportFactory.h
+++ b/Engine/Plugins/2D/Paper2D/Source/PaperSpriteSheetImporter/Private/PaperSpriteSheetImportFactory.h
@@ -13,6 +13,7 @@ class UPaperSpriteSheetImportFactory : public UFactory
 
 	// UFactory interface
 	virtual FText GetToolTip() const override;
+	virtual bool FactoryCanImport(const FString& Filename) override;
 	virtual UObject* FactoryCreateText(UClass* InClass, UObject* InParent, FName InName, EObjectFlags Flags, UObject* Context, const TCHAR* Type, const TCHAR*& Buffer, const TCHAR* BufferEnd, FFeedbackContext* Warn) override;
 	// End of UFactory interface
 
diff --git a/Engine/Source/Developer/AssetTools/Private/AssetTools.cpp b/Engine/Source/Developer/AssetTools/Private/AssetTools.cpp
--- a/Engine/Source/Developer/AssetTools/Private/AssetTools.cpp
+++ b/Engine/Source/Developer/AssetTools/Private/AssetTools.cpp
@@ -692,6 +692,11 @@ TArray<UObject*> FAssetTools::ImportAssets(const TArray<FString>& Files, const F
 		}
 	}
 
+	// We need to sort the factories so that they get tested in priority order
+	for (auto& ExtensionToFactories : ExtensionToFactoriesMap)
+	{
+		ExtensionToFactories.Value.Sort(&UFactory::SortFactoriesByPriority);
+	}
 
 	// Some flags to keep track of what the user decided when asked about overwriting or replacing
 	bool bOverwriteAll = false;
diff --git a/Engine/Source/Editor/UnrealEd/Classes/Factories/Factory.h b/Engine/Source/Editor/UnrealEd/Classes/Factories/Factory.h
--- a/Engine/Source/Editor/UnrealEd/Classes/Factories/Factory.h
+++ b/Engine/Source/Editor/UnrealEd/Classes/Factories/Factory.h
@@ -73,6 +73,9 @@ public:
 	static void AddReferencedObjects(UObject* InThis, FReferenceCollector& Collector);
 	// End UObject interface.
 
+	/** Helper function to sort an array of factories by their import priority - use as a predicate for Sort */
+	static bool SortFactoriesByPriority(const UFactory& A, const UFactory& B);
+
 	/** Returns true if this factory should be shown in the New Asset menu (by default calls CanCreateNew). */
 	virtual bool ShouldShowInNewMenu() const;
 
diff --git a/Engine/Source/Editor/UnrealEd/Classes/Factories/ReimportDataTableFactory.h b/Engine/Source/Editor/UnrealEd/Classes/Factories/ReimportDataTableFactory.h
--- a/Engine/Source/Editor/UnrealEd/Classes/Factories/ReimportDataTableFactory.h
+++ b/Engine/Source/Editor/UnrealEd/Classes/Factories/ReimportDataTableFactory.h
@@ -9,6 +9,7 @@ class UReimportDataTableFactory : public UCSVImportFactory, public FReimportHand
 	GENERATED_UCLASS_BODY()
 
 	// Begin FReimportHandler interface
+	virtual bool FactoryCanImport( const FString& Filename ) override;
 	virtual bool CanReimport( UObject* Obj, TArray<FString>& OutFilenames ) override;
 	virtual void SetReimportPaths( UObject* Obj, const TArray<FString>& NewReimportPaths ) override;
 	virtual EReimportResult::Type Reimport( UObject* Obj ) override;
diff --git a/Engine/Source/Editor/UnrealEd/Private/Factories/CSVImportFactory.cpp b/Engine/Source/Editor/UnrealEd/Private/Factories/CSVImportFactory.cpp
--- a/Engine/Source/Editor/UnrealEd/Private/Factories/CSVImportFactory.cpp
+++ b/Engine/Source/Editor/UnrealEd/Private/Factories/CSVImportFactory.cpp
@@ -51,6 +51,9 @@ UCSVImportFactory::UCSVImportFactory(const FObjectInitializer& ObjectInitializer
 	bEditorImport = true;
 	bText = true;
 
+	// Give this factory a lower than normal import priority, as CSV and JSON can be commonly used and we'd like to give the other import factories a shot first
+	--ImportPriority;
+
 	Formats.Add(TEXT("csv;Comma-separated values"));
 }
 
@@ -59,7 +62,6 @@ FText UCSVImportFactory::GetDisplayName() const
 	return LOCTEXT("CSVImportFactoryDescription", "Comma Separated Values");
 }
 
-
 bool UCSVImportFactory::DoesSupportClass(UClass * Class)
 {
 	return (Class == UDataTable::StaticClass() || Class == UCurveTable::StaticClass() || Class == UCurveFloat::StaticClass() || Class == UCurveVector::StaticClass() || Class == UCurveLinearColor::StaticClass() );
@@ -300,6 +302,11 @@ UReimportDataTableFactory::UReimportDataTableFactory(const FObjectInitializer& O
 	Formats.Add(TEXT("json;JavaScript Object Notation"));
 }
 
+bool UReimportDataTableFactory::FactoryCanImport( const FString& Filename )
+{
+	return true;
+}
+
 bool UReimportDataTableFactory::CanReimport( UObject* Obj, TArray<FString>& OutFilenames )
 {	
 	UDataTable* DataTable = Cast<UDataTable>(Obj);
diff --git a/Engine/Source/Editor/UnrealEd/Private/Factories/Factory.cpp b/Engine/Source/Editor/UnrealEd/Private/Factories/Factory.cpp
--- a/Engine/Source/Editor/UnrealEd/Private/Factories/Factory.cpp
+++ b/Engine/Source/Editor/UnrealEd/Private/Factories/Factory.cpp
@@ -171,6 +171,29 @@ void UFactory::DisplayOverwriteOptionsDialog(const FText& Message)
 	}
 }
 
+bool UFactory::SortFactoriesByPriority(const UFactory& A, const UFactory& B)
+{
+	// First sort so that higher priorities are earlier in the list
+	if( A.ImportPriority > B.ImportPriority )
+	{
+		return true;
+	}
+	else if( A.ImportPriority < B.ImportPriority )
+	{
+		return false;
+	}
+
+	// Then sort so that factories that only create new assets are tried after those that actually import the file data (when they have an equivalent priority)
+	const bool bFactoryAImportsFiles = !A.CanCreateNew();
+	const bool bFactoryBImportsFiles = !B.CanCreateNew();
+	if( bFactoryAImportsFiles && !bFactoryBImportsFiles )
+	{
+		return true;
+	}
+
+	return false;
+}
+
 UObject* UFactory::StaticImportObject
 (
 UClass*				Class,
@@ -235,28 +258,7 @@ UObject* UFactory::StaticImportObject
 			}
 		}
 
-		Factories.Sort([](const UFactory& A, const UFactory& B) -> bool
-		{
-			// First sort so that higher priorities are earlier in the list
-			if( A.ImportPriority > B.ImportPriority )
-			{
-				return true;
-			}
-			else if( A.ImportPriority < B.ImportPriority )
-			{
-				return false;
-			}
-
-			// Then sort so that factories that only create new assets are tried after those that actually import the file data (when they have an equivalent priority)
-			const bool bFactoryAImportsFiles = !A.CanCreateNew();
-			const bool bFactoryBImportsFiles = !B.CanCreateNew();
-			if( bFactoryAImportsFiles && !bFactoryBImportsFiles )
-			{
-				return true;
-			}
-
-			return false;
-		});
+		Factories.Sort(&UFactory::SortFactoriesByPriority);
 	}
 
 	bool bLoadedFile = false;
