Added check for being on mac client, which disables "Always on Top" feature if true.

UE-17765

diff --git a/Engine/Source/Editor/UnrealEd/Private/PlayLevel.cpp b/Engine/Source/Editor/UnrealEd/Private/PlayLevel.cpp
--- a/Engine/Source/Editor/UnrealEd/Private/PlayLevel.cpp
+++ b/Engine/Source/Editor/UnrealEd/Private/PlayLevel.cpp
@@ -2851,10 +2851,15 @@ UGameInstance* UEditorEngine::CreatePIEGameInstance(int32 PIEInstance, bool bInS
 				FOnSwitchWorldHack OnWorldSwitch = FOnSwitchWorldHack::CreateUObject( this, &UEditorEngine::OnSwitchWorldForSlatePieWindow );
 				PieWindow->SetOnWorldSwitchHack( OnWorldSwitch );
 
+				// Mac does not support parenting, do not keep on top
+#if PLATFORM_MAC
+				FSlateApplication::Get().AddWindow(PieWindow);
+#else
 				if (PlayInSettings->PIEAlwaysOnTop && FSlateApplication::Get().GetActiveTopLevelWindow().IsValid())
 					FSlateApplication::Get().AddWindowAsNativeChild(PieWindow, FSlateApplication::Get().GetActiveTopLevelWindow().ToSharedRef(), true);
 				else
 					FSlateApplication::Get().AddWindow(PieWindow);
+#endif
 
 				TSharedRef<SOverlay> ViewportOverlayWidgetRef = SNew(SOverlay);
 
diff --git a/Engine/Source/Editor/UnrealEd/Private/Settings/LevelEditorPlaySettingsCustomization.h b/Engine/Source/Editor/UnrealEd/Private/Settings/LevelEditorPlaySettingsCustomization.h
--- a/Engine/Source/Editor/UnrealEd/Private/Settings/LevelEditorPlaySettingsCustomization.h
+++ b/Engine/Source/Editor/UnrealEd/Private/Settings/LevelEditorPlaySettingsCustomization.h
@@ -343,8 +343,15 @@ public:
 		// play in new window settings
 		IDetailCategoryBuilder& PlayInNewWindowCategory = LayoutBuilder.EditCategory("PlayInNewWindow");
 		{
+		// Mac does not support parenting, do not show
+#if PLATFORM_MAC
+			PlayInNewWindowCategory.AddProperty("PIEAlwaysOnTop")
+				.DisplayName(LOCTEXT("PIEAlwaysOnTop", "Always On Top"))
+				.IsEnabled(false);
+#else
 			PlayInNewWindowCategory.AddProperty("PIEAlwaysOnTop")
 				.DisplayName(LOCTEXT("PIEAlwaysOnTop", "Always On Top"));
+#endif
 
 			// new window size
 			TSharedRef<IPropertyHandle> WindowHeightHandle = LayoutBuilder.GetProperty(GET_MEMBER_NAME_CHECKED(ULevelEditorPlaySettings, NewWindowHeight));
