nicer / more clear UI naming

diff --git a/Engine/Source/Editor/UnrealEd/Private/Kismet2/DebuggerCommands.cpp b/Engine/Source/Editor/UnrealEd/Private/Kismet2/DebuggerCommands.cpp
--- a/Engine/Source/Editor/UnrealEd/Private/Kismet2/DebuggerCommands.cpp
+++ b/Engine/Source/Editor/UnrealEd/Private/Kismet2/DebuggerCommands.cpp
@@ -227,10 +227,10 @@ void FPlayWorldCommands::RegisterCommands()
 
 	// PIE
 	UI_COMMAND( RepeatLastPlay, "Play", "Launches a game preview session in the same mode as the last game preview session launched from the Game Preview Modes dropdown next to the Play button on the level editor toolbar", EUserInterfaceActionType::Button, FInputChord( EKeys::P, EModifierKey::Alt ) )
-	UI_COMMAND( PlayInViewport, "Selected Viewport", "Play this level in the active level editor viewport", EUserInterfaceActionType::Check, FInputChord() );
-	UI_COMMAND( PlayInEditorFloating, "New Editor Window", "Play this level in a new window", EUserInterfaceActionType::Check, FInputChord() );
+	UI_COMMAND( PlayInViewport, "Selected Viewport (PIE)", "Play this level in the active level editor viewport", EUserInterfaceActionType::Check, FInputChord() );
+	UI_COMMAND( PlayInEditorFloating, "New Editor Window (PIE)", "Play this level in a new window", EUserInterfaceActionType::Check, FInputChord() );
 	UI_COMMAND( PlayInVR, "VR Preview", "Play this level in VR", EUserInterfaceActionType::Check, FInputChord() );
-	UI_COMMAND( PlayInMobilePreview, "Mobile Preview", "Play this level as a mobile device preview (runs in its own process)", EUserInterfaceActionType::Check, FInputChord() );
+	UI_COMMAND( PlayInMobilePreview, "Mobile Preview (PIE)", "Play this level as a mobile device preview (runs in its own process)", EUserInterfaceActionType::Check, FInputChord() );
 	UI_COMMAND( PlayInNewProcess, "Standalone Game", "Play this level in a new window that runs in its own process", EUserInterfaceActionType::Check, FInputChord() );
 	UI_COMMAND( PlayInCameraLocation, "Current Camera Location", "Spawn the player at the current camera location", EUserInterfaceActionType::RadioButton, FInputChord() );
 	UI_COMMAND( PlayInDefaultPlayerStart, "Default Player Start", "Spawn the player at the map's default player start", EUserInterfaceActionType::RadioButton, FInputChord() );
@@ -437,8 +437,8 @@ void FPlayWorldCommands::BuildToolbar( FToolBarBuilder& ToolbarBuilder, bool bIn
 	ToolbarBuilder.AddComboButton(
 		SpecialPIEOptionsMenuAction,
 		FOnGetContent::CreateStatic( &GeneratePlayMenuContent, GlobalPlayWorldActions.ToSharedRef() ),
-		LOCTEXT( "PlayCombo_Label", "Play Options" ),
-		LOCTEXT( "PIEComboToolTip", "Options for Play in Editor (PIE)" ),
+		LOCTEXT( "PlayCombo_Label", "Active Play Mode" ),
+		LOCTEXT( "PIEComboToolTip", "Change Play Mode and Play Settings" ),
 		FSlateIcon(FEditorStyle::GetStyleSetName(), "LevelEditor.RepeatLastPlay"),
 		true
 	);
