Fix missing header that was causing Mac compile failure.

diff --git a/Engine/Source/Editor/UnrealEd/Private/Editor.cpp b/Engine/Source/Editor/UnrealEd/Private/Editor.cpp
--- a/Engine/Source/Editor/UnrealEd/Private/Editor.cpp
+++ b/Engine/Source/Editor/UnrealEd/Private/Editor.cpp
@@ -109,6 +109,7 @@
 #include "UnrealEngine.h"
 #include "EngineStats.h"
 #include "Engine/SimpleConstructionScript.h"
+#include "K2Node_AddComponent.h"
 
 #include "Settings/EditorSettings.h"
 
