EditorFactories: Get rid of ForceValidActorRefs function used when pasting actors into Level.  Was sub optimal performance to say the least, and invalid actor references are already null since UObjectPropertyBase::FindImportedObject() will set them to null if it can't find the reference in question.
Also fixed a loop that was looping over all existing actors in the world, but only doing something if they were in a certain map of actors.  Just iterating over the map of actors instead.

Should make pasting objects into levels with a lot of actors much faster.

diff --git a/Engine/Source/Editor/UnrealEd/Private/Factories/EditorFactories.cpp b/Engine/Source/Editor/UnrealEd/Private/Factories/EditorFactories.cpp
--- a/Engine/Source/Editor/UnrealEd/Private/Factories/EditorFactories.cpp
+++ b/Engine/Source/Editor/UnrealEd/Private/Factories/EditorFactories.cpp
@@ -269,58 +269,6 @@ UObject* UMaterialParameterCollectionFactoryNew::FactoryCreateNew(UClass* Class,
 	ULevelFactory.
 ------------------------------------------------------------------------------*/
 
-/**
- * Iterates over an object's properties making sure that any UObjectProperty properties
- * that refer to non-nullptr actors refer to valid actors.
- *
- * @return		false if no object references were nullptr'd out, true otherwise.
- */
-static bool ForceValidActorRefs(UStruct* Struct, uint8* Data)
-{
-	bool bChangedObjectPointer = false;
-
-	//@todo DB: Optimize this!!
-	for( TFieldIterator<UProperty> PropertyIt(Struct); PropertyIt; ++PropertyIt )
-	{
-		for( int32 i=0; i<PropertyIt->ArrayDim; i++ )
-		{
-			uint8* Value = PropertyIt->ContainerPtrToValuePtr<uint8>(Data, i);
-			UObjectPropertyBase* Prop = Cast<UObjectPropertyBase>(*PropertyIt);
-			if(Prop)
-			{
-				UObject* Obj = Prop->GetObjectPropertyValue(Value);
-				AActor* SearchActor = Cast<AActor>(Obj);
-				if( SearchActor && !Obj->HasAnyFlags(RF_ArchetypeObject|RF_ClassDefaultObject) )
-				{
-					bool bFound = false;
-					for( FActorIterator ActorIt(SearchActor->GetWorld()); ActorIt; ++ActorIt )
-					{
-						AActor* Actor = *ActorIt;
-						if( Actor == SearchActor )
-						{
-							bFound = true;
-							break;
-						}
-					}
-					
-					if( !bFound )
-					{
-						UE_LOG(LogEditorFactories, Log,  TEXT("Usurped %s"), *Obj->GetClass()->GetName() );
-						Prop->SetObjectPropertyValue(Value, nullptr);
-						bChangedObjectPointer = true;
-					}
-				}
-			}
-			else if( Cast<UStructProperty>(*PropertyIt) )
-			{
-				bChangedObjectPointer |= ForceValidActorRefs( ((UStructProperty*)*PropertyIt)->Struct, Value );
-			}
-		}
-	}
-
-	return bChangedObjectPointer;
-}
-
 ULevelFactory::ULevelFactory(const FObjectInitializer& ObjectInitializer)
 	: Super(ObjectInitializer)
 {
@@ -827,13 +775,13 @@ UObject* ULevelFactory::FactoryCreateText
 
 	// Pass 1: Sort out all the properties on the individual actors
 	bool bIsMoveToStreamingLevel =(FCString::Stricmp(Type, TEXT("move")) == 0);
-	for( FActorIterator It(World); It; ++It )
+	for (auto& ActorMapElement : NewActorMap)
 	{
-		AActor* Actor = *It;
+		AActor* Actor = ActorMapElement.Key;
 
 		// Import properties if the new actor is 
 		bool		bActorChanged = false;
-		FString*	PropText = NewActorMap.Find(Actor);
+		FString*	PropText = &(ActorMapElement.Value); 
 		if( PropText )
 		{
 			if ( Actor->ShouldImport(PropText, bIsMoveToStreamingLevel) )
@@ -872,20 +820,6 @@ UObject* ULevelFactory::FactoryCreateText
 			}
 		}
 
-		// Make sure all references to actors are valid.
-		// if they don't belong to same level
-		UWorld * ActorWorld = Actor->GetTypedOuter<UWorld>();
-		if( ActorWorld != World )
-		{
-			const bool bFixedUpObjectRefs = ForceValidActorRefs( Actor->GetClass(), (uint8*)Actor );
-
-			// Aactor references were fixed up, so treat the actor as having been changed.
-			if ( bFixedUpObjectRefs )
-			{
-				bActorChanged = true;
-			}
-		}
-
 		// Copy brushes' model pointers over to their BrushComponent, to keep compatibility with old T3Ds.
 		if( Brush && bActorChanged )
 		{
