Error when a module is both a dependency and dynamically loaded.
Fixes for lots of existing modules by removing all dynamically loaded duplicates (this is probably not the correct solution, but is the safest).

diff --git a/Engine/Plugins/2D/Paper2D/Source/PaperSpriteSheetImporter/PaperSpriteSheetImporter.Build.cs b/Engine/Plugins/2D/Paper2D/Source/PaperSpriteSheetImporter/PaperSpriteSheetImporter.Build.cs
--- a/Engine/Plugins/2D/Paper2D/Source/PaperSpriteSheetImporter/PaperSpriteSheetImporter.Build.cs
+++ b/Engine/Plugins/2D/Paper2D/Source/PaperSpriteSheetImporter/PaperSpriteSheetImporter.Build.cs
@@ -30,7 +30,6 @@ public class PaperSpriteSheetImporter : ModuleRules
 
 		DynamicallyLoadedModuleNames.AddRange(
 			new string[] {
-				"AssetTools",
 				"AssetRegistry"
 			});
 	}
diff --git a/Engine/Plugins/Editor/EpicSurvey/Source/EpicSurvey/EpicSurvey.Build.cs b/Engine/Plugins/Editor/EpicSurvey/Source/EpicSurvey/EpicSurvey.Build.cs
--- a/Engine/Plugins/Editor/EpicSurvey/Source/EpicSurvey/EpicSurvey.Build.cs
+++ b/Engine/Plugins/Editor/EpicSurvey/Source/EpicSurvey/EpicSurvey.Build.cs
@@ -43,8 +43,6 @@ namespace UnrealBuildTool.Rules
             DynamicallyLoadedModuleNames.AddRange(
 				new string[]
 				{
-					"ImageWrapper",
-					"LevelEditor",
 					"MainFrame"
 				}
 			);
diff --git a/Engine/Plugins/Editor/MacGraphicsSwitching/Source/MacGraphicsSwitching/MacGraphicsSwitching.Build.cs b/Engine/Plugins/Editor/MacGraphicsSwitching/Source/MacGraphicsSwitching/MacGraphicsSwitching.Build.cs
--- a/Engine/Plugins/Editor/MacGraphicsSwitching/Source/MacGraphicsSwitching/MacGraphicsSwitching.Build.cs
+++ b/Engine/Plugins/Editor/MacGraphicsSwitching/Source/MacGraphicsSwitching/MacGraphicsSwitching.Build.cs
@@ -31,7 +31,6 @@ namespace UnrealBuildTool.Rules
 			DynamicallyLoadedModuleNames.AddRange(
 				new string[] 
 				{
-					"LevelEditor",
 					"MainFrame",
 					"Settings"
 				}
diff --git a/Engine/Source/Developer/AssetTools/AssetTools.Build.cs b/Engine/Source/Developer/AssetTools/AssetTools.Build.cs
--- a/Engine/Source/Developer/AssetTools/AssetTools.Build.cs
+++ b/Engine/Source/Developer/AssetTools/AssetTools.Build.cs
@@ -55,7 +55,6 @@ public class AssetTools : ModuleRules
 				"AssetRegistry",
 				"ContentBrowser",
 				"CollectionManager",
-                "CurveAssetEditor",
 				"CurveTableEditor",
 				"DataTableEditor",
 				"DesktopPlatform",
diff --git a/Engine/Source/Developer/FriendsAndChat/FriendsAndChat.Build.cs b/Engine/Source/Developer/FriendsAndChat/FriendsAndChat.Build.cs
--- a/Engine/Source/Developer/FriendsAndChat/FriendsAndChat.Build.cs
+++ b/Engine/Source/Developer/FriendsAndChat/FriendsAndChat.Build.cs
@@ -60,7 +60,6 @@ public class FriendsAndChat : ModuleRules
 				"Analytics",
 				"AnalyticsET",
 				"EditorStyle",
-				"OnlineSubsystem"
 			}
 		);
 
@@ -88,13 +87,6 @@ public class FriendsAndChat : ModuleRules
 					"OnlineSubsystemMcp",
 				}
 			);
-		
-			DynamicallyLoadedModuleNames.AddRange(
-				new string[]
-				{
-					"OnlineSubsystemMcp",
-				}
-			);
 		}
 	}
 }
diff --git a/Engine/Source/Developer/FunctionalTesting/FunctionalTesting.Build.cs b/Engine/Source/Developer/FunctionalTesting/FunctionalTesting.Build.cs
--- a/Engine/Source/Developer/FunctionalTesting/FunctionalTesting.Build.cs
+++ b/Engine/Source/Developer/FunctionalTesting/FunctionalTesting.Build.cs
@@ -26,11 +26,5 @@ public class FunctionalTesting : ModuleRules
 				"Developer/FunctionalTesting/Private",
 			}
 			);
-
-		DynamicallyLoadedModuleNames.AddRange(
-			new string[] {
-				"MessageLog",
-			}
-			);
 	}
 }
diff --git a/Engine/Source/Editor/BehaviorTreeEditor/BehaviorTreeEditor.Build.cs b/Engine/Source/Editor/BehaviorTreeEditor/BehaviorTreeEditor.Build.cs
--- a/Engine/Source/Editor/BehaviorTreeEditor/BehaviorTreeEditor.Build.cs
+++ b/Engine/Source/Editor/BehaviorTreeEditor/BehaviorTreeEditor.Build.cs
@@ -53,7 +53,6 @@ public class BehaviorTreeEditor : ModuleRules
 		DynamicallyLoadedModuleNames.AddRange(
             new string[] { 
                 "WorkspaceMenuStructure",
-                "PropertyEditor",
 				"AssetTools",
 				"AssetRegistry",
 				"ContentBrowser"
diff --git a/Engine/Source/Editor/BspMode/BspMode.Build.cs b/Engine/Source/Editor/BspMode/BspMode.Build.cs
--- a/Engine/Source/Editor/BspMode/BspMode.Build.cs
+++ b/Engine/Source/Editor/BspMode/BspMode.Build.cs
@@ -23,7 +23,6 @@ public class BspMode : ModuleRules
         DynamicallyLoadedModuleNames.AddRange(
             new string[] {
 				"PropertyEditor",
-                "LevelEditor",
 			}
         );
 	}
diff --git a/Engine/Source/Editor/ClassViewer/ClassViewer.Build.cs b/Engine/Source/Editor/ClassViewer/ClassViewer.Build.cs
--- a/Engine/Source/Editor/ClassViewer/ClassViewer.Build.cs
+++ b/Engine/Source/Editor/ClassViewer/ClassViewer.Build.cs
@@ -36,7 +36,6 @@ public class ClassViewer : ModuleRules
                 "AssetTools",
 				"EditorWidgets",
 				"GameProjectGeneration",
-				"PropertyEditor",
 			}
 		);
 	}
diff --git a/Engine/Source/Editor/ContentBrowser/ContentBrowser.Build.cs b/Engine/Source/Editor/ContentBrowser/ContentBrowser.Build.cs
--- a/Engine/Source/Editor/ContentBrowser/ContentBrowser.Build.cs
+++ b/Engine/Source/Editor/ContentBrowser/ContentBrowser.Build.cs
@@ -47,7 +47,6 @@ public class ContentBrowser : ModuleRules
 				"AssetRegistry",
 				"AssetTools",
 				"CollectionManager",
-				"EditorWidgets",
 				"GameProjectGeneration",
                 "MainFrame",
                 "ReferenceViewer",
diff --git a/Engine/Source/Editor/DetailCustomizations/DetailCustomizations.Build.cs b/Engine/Source/Editor/DetailCustomizations/DetailCustomizations.Build.cs
--- a/Engine/Source/Editor/DetailCustomizations/DetailCustomizations.Build.cs
+++ b/Engine/Source/Editor/DetailCustomizations/DetailCustomizations.Build.cs
@@ -62,7 +62,6 @@ public class DetailCustomizations : ModuleRules
 		DynamicallyLoadedModuleNames.AddRange(
 			new string[] {
 				"Layers",
-				"PropertyEditor",
 				"GameProjectGeneration",
 			}
 		);
diff --git a/Engine/Source/Editor/DeviceProfileEditor/DeviceProfileEditor.Build.cs b/Engine/Source/Editor/DeviceProfileEditor/DeviceProfileEditor.Build.cs
--- a/Engine/Source/Editor/DeviceProfileEditor/DeviceProfileEditor.Build.cs
+++ b/Engine/Source/Editor/DeviceProfileEditor/DeviceProfileEditor.Build.cs
@@ -41,7 +41,6 @@ public class DeviceProfileEditor : ModuleRules
 		DynamicallyLoadedModuleNames.AddRange(
 			new string[] {
 				"MainFrame",
-				"PropertyEditor",
 			}
 		);
 	}
diff --git a/Engine/Source/Editor/EnvironmentQueryEditor/EnvironmentQueryEditor.Build.cs b/Engine/Source/Editor/EnvironmentQueryEditor/EnvironmentQueryEditor.Build.cs
--- a/Engine/Source/Editor/EnvironmentQueryEditor/EnvironmentQueryEditor.Build.cs
+++ b/Engine/Source/Editor/EnvironmentQueryEditor/EnvironmentQueryEditor.Build.cs
@@ -49,7 +49,6 @@ public class EnvironmentQueryEditor : ModuleRules
         DynamicallyLoadedModuleNames.AddRange(
             new string[] { 
                 "WorkspaceMenuStructure",
-                "PropertyEditor"
             }
 		);
 	}
diff --git a/Engine/Source/Editor/FontEditor/FontEditor.Build.cs b/Engine/Source/Editor/FontEditor/FontEditor.Build.cs
--- a/Engine/Source/Editor/FontEditor/FontEditor.Build.cs
+++ b/Engine/Source/Editor/FontEditor/FontEditor.Build.cs
@@ -35,7 +35,6 @@ public class FontEditor : ModuleRules
 				"MainFrame",
 				"WorkspaceMenuStructure",
 				"MainFrame",
-				"DesktopPlatform",
 				"PropertyEditor"
 			}
 		);
diff --git a/Engine/Source/Editor/GameProjectGeneration/GameProjectGeneration.Build.cs b/Engine/Source/Editor/GameProjectGeneration/GameProjectGeneration.Build.cs
--- a/Engine/Source/Editor/GameProjectGeneration/GameProjectGeneration.Build.cs
+++ b/Engine/Source/Editor/GameProjectGeneration/GameProjectGeneration.Build.cs
@@ -53,7 +53,6 @@ public class GameProjectGeneration : ModuleRules
             new string[] {
 				"AssetRegistry",
 				"ContentBrowser",
-                "DesktopPlatform",
                 "Documentation",
                 "MainFrame",
             }
diff --git a/Engine/Source/Editor/GraphEditor/GraphEditor.Build.cs b/Engine/Source/Editor/GraphEditor/GraphEditor.Build.cs
--- a/Engine/Source/Editor/GraphEditor/GraphEditor.Build.cs
+++ b/Engine/Source/Editor/GraphEditor/GraphEditor.Build.cs
@@ -56,7 +56,6 @@ public class GraphEditor : ModuleRules
 		DynamicallyLoadedModuleNames.AddRange(
 			new string[] {
 				"ContentBrowser",
-				"Documentation",
 				"EditorWidgets"
 			}
 		);
diff --git a/Engine/Source/Editor/IntroTutorials/IntroTutorials.Build.cs b/Engine/Source/Editor/IntroTutorials/IntroTutorials.Build.cs
--- a/Engine/Source/Editor/IntroTutorials/IntroTutorials.Build.cs
+++ b/Engine/Source/Editor/IntroTutorials/IntroTutorials.Build.cs
@@ -56,7 +56,6 @@ namespace UnrealBuildTool.Rules
 
 			DynamicallyLoadedModuleNames.AddRange(
 				new string[] {
-                    "Documentation",
 					"MainFrame",
 					"LauncherServices",
 				}
diff --git a/Engine/Source/Editor/LandscapeEditor/LandscapeEditor.Build.cs b/Engine/Source/Editor/LandscapeEditor/LandscapeEditor.Build.cs
--- a/Engine/Source/Editor/LandscapeEditor/LandscapeEditor.Build.cs
+++ b/Engine/Source/Editor/LandscapeEditor/LandscapeEditor.Build.cs
@@ -38,7 +38,6 @@ public class LandscapeEditor : ModuleRules
 			new string[] {
 				"MainFrame",
 				"DesktopPlatform",
-				"ImageWrapper",
 			}
 			);
 
diff --git a/Engine/Source/Editor/Layers/Layers.Build.cs b/Engine/Source/Editor/Layers/Layers.Build.cs
--- a/Engine/Source/Editor/Layers/Layers.Build.cs
+++ b/Engine/Source/Editor/Layers/Layers.Build.cs
@@ -23,7 +23,6 @@ public class Layers : ModuleRules
 		DynamicallyLoadedModuleNames.AddRange(
 			new string[] {
 				"PropertyEditor",
-				"SceneOutliner",
 			}
 		);
 	}
diff --git a/Engine/Source/Editor/LevelEditor/LevelEditor.Build.cs b/Engine/Source/Editor/LevelEditor/LevelEditor.Build.cs
--- a/Engine/Source/Editor/LevelEditor/LevelEditor.Build.cs
+++ b/Engine/Source/Editor/LevelEditor/LevelEditor.Build.cs
@@ -66,8 +66,6 @@ public class LevelEditor : ModuleRules
 		DynamicallyLoadedModuleNames.AddRange(
 			new string[] {
 				"MainFrame",
-				"PropertyEditor",
-				"SceneOutliner",
 				"ClassViewer",
 				"DeviceManager",
 				"SettingsEditor",
@@ -80,7 +78,6 @@ public class LevelEditor : ModuleRules
 				"WorkspaceMenuStructure",
 				"NewLevelDialog",
 				"DeviceProfileEditor",
-				"DeviceProfileServices",
                 "PlacementMode",
 				"UserFeedback",
 				"ReferenceViewer",
diff --git a/Engine/Source/Editor/MainFrame/MainFrame.Build.cs b/Engine/Source/Editor/MainFrame/MainFrame.Build.cs
--- a/Engine/Source/Editor/MainFrame/MainFrame.Build.cs
+++ b/Engine/Source/Editor/MainFrame/MainFrame.Build.cs
@@ -69,14 +69,11 @@ public class MainFrame : ModuleRules
 		DynamicallyLoadedModuleNames.AddRange(
 			new string[] {
 				"AssetTools",
-				"DesktopPlatform",
-                "Documentation",
 				"GameProjectGeneration",
 				"ProjectTargetPlatformEditor",
 				"LevelEditor",
                 "OutputLog",
 				"SuperSearch",
-                "TranslationEditor",
 				"SourceCodeAccess",
 				"EditorLiveStreaming",
 				"HotReload",
diff --git a/Engine/Source/Editor/MaterialEditor/MaterialEditor.Build.cs b/Engine/Source/Editor/MaterialEditor/MaterialEditor.Build.cs
--- a/Engine/Source/Editor/MaterialEditor/MaterialEditor.Build.cs
+++ b/Engine/Source/Editor/MaterialEditor/MaterialEditor.Build.cs
@@ -39,7 +39,6 @@ public class MaterialEditor : ModuleRules
 			new string[] {
 				"AssetTools",
 				"MainFrame",
-				"PropertyEditor",
 				"SceneOutliner",
 				"ClassViewer",
 				"ContentBrowser",
diff --git a/Engine/Source/Editor/Matinee/Matinee.Build.cs b/Engine/Source/Editor/Matinee/Matinee.Build.cs
--- a/Engine/Source/Editor/Matinee/Matinee.Build.cs
+++ b/Engine/Source/Editor/Matinee/Matinee.Build.cs
@@ -49,7 +49,6 @@ public class Matinee : ModuleRules
 		DynamicallyLoadedModuleNames.AddRange(
 			new string[] {
 				"PropertyEditor",
-                "ContentBrowser"
 				}
 			);
 	}
diff --git a/Engine/Source/Editor/MovieSceneTools/MovieSceneTools.Build.cs b/Engine/Source/Editor/MovieSceneTools/MovieSceneTools.Build.cs
--- a/Engine/Source/Editor/MovieSceneTools/MovieSceneTools.Build.cs
+++ b/Engine/Source/Editor/MovieSceneTools/MovieSceneTools.Build.cs
@@ -49,7 +49,6 @@ public class MovieSceneTools : ModuleRules
 		DynamicallyLoadedModuleNames.AddRange(
 			new string[] {
 				"AssetTools",
-				"Sequencer",
 				"PropertyEditor"
 			}
 		);
diff --git a/Engine/Source/Editor/NiagaraEditor/NiagaraEditor.Build.cs b/Engine/Source/Editor/NiagaraEditor/NiagaraEditor.Build.cs
--- a/Engine/Source/Editor/NiagaraEditor/NiagaraEditor.Build.cs
+++ b/Engine/Source/Editor/NiagaraEditor/NiagaraEditor.Build.cs
@@ -46,7 +46,6 @@ public class NiagaraEditor : ModuleRules
 		DynamicallyLoadedModuleNames.AddRange(
             new string[] {
                 "WorkspaceMenuStructure",
-				"PropertyEditor",
                 }
             );
 	}
diff --git a/Engine/Source/Editor/Persona/Persona.Build.cs b/Engine/Source/Editor/Persona/Persona.Build.cs
--- a/Engine/Source/Editor/Persona/Persona.Build.cs
+++ b/Engine/Source/Editor/Persona/Persona.Build.cs
@@ -44,7 +44,6 @@ public class Persona : ModuleRules
 				"Documentation",
 				"MainFrame",
 				"DesktopPlatform",
-                "PropertyEditor"
 			}
 		);
 	}
diff --git a/Engine/Source/Editor/PlacementMode/PlacementMode.Build.cs b/Engine/Source/Editor/PlacementMode/PlacementMode.Build.cs
--- a/Engine/Source/Editor/PlacementMode/PlacementMode.Build.cs
+++ b/Engine/Source/Editor/PlacementMode/PlacementMode.Build.cs
@@ -25,12 +25,5 @@ public class PlacementMode : ModuleRules
 				"BspMode"
             } 
         );
-
-        DynamicallyLoadedModuleNames.AddRange(
-            new string[] { 
-                "AssetTools",
-                "LevelEditor"
-            }
-        );
     }
 }
diff --git a/Engine/Source/Editor/PropertyEditor/PropertyEditor.Build.cs b/Engine/Source/Editor/PropertyEditor/PropertyEditor.Build.cs
--- a/Engine/Source/Editor/PropertyEditor/PropertyEditor.Build.cs
+++ b/Engine/Source/Editor/PropertyEditor/PropertyEditor.Build.cs
@@ -68,7 +68,6 @@ public class PropertyEditor : ModuleRules
                 "AssetTools",
 				"ClassViewer",
 				"ContentBrowser",
-				"Documentation",
 				"MainFrame",
 			}
 		);
diff --git a/Engine/Source/Editor/StaticMeshEditor/StaticMeshEditor.Build.cs b/Engine/Source/Editor/StaticMeshEditor/StaticMeshEditor.Build.cs
--- a/Engine/Source/Editor/StaticMeshEditor/StaticMeshEditor.Build.cs
+++ b/Engine/Source/Editor/StaticMeshEditor/StaticMeshEditor.Build.cs
@@ -39,13 +39,10 @@ public class StaticMeshEditor : ModuleRules
 
 		DynamicallyLoadedModuleNames.AddRange(
 			new string[] {
-				"PropertyEditor",
 				"SceneOutliner",
 				"ClassViewer",
 				"ContentBrowser",
 				"WorkspaceMenuStructure",
-				"MeshUtilities",
-				"TargetPlatform",
 			}
 		);
 
diff --git a/Engine/Source/Editor/TranslationEditor/TranslationEditor.Build.cs b/Engine/Source/Editor/TranslationEditor/TranslationEditor.Build.cs
--- a/Engine/Source/Editor/TranslationEditor/TranslationEditor.Build.cs
+++ b/Engine/Source/Editor/TranslationEditor/TranslationEditor.Build.cs
@@ -52,8 +52,6 @@ public class TranslationEditor : ModuleRules
 			new string[] {
 				"WorkspaceMenuStructure",
 				"DesktopPlatform",
-				"SourceControl",
-                "LocalizationService",
 			}
 		);
 	}
diff --git a/Engine/Source/Editor/UMGEditor/UMGEditor.Build.cs b/Engine/Source/Editor/UMGEditor/UMGEditor.Build.cs
--- a/Engine/Source/Editor/UMGEditor/UMGEditor.Build.cs
+++ b/Engine/Source/Editor/UMGEditor/UMGEditor.Build.cs
@@ -55,14 +55,5 @@ public class UMGEditor : ModuleRules
 				"RenderCore",
 			}
 			);
-
-		DynamicallyLoadedModuleNames.AddRange(
-			new string[] {
-				"PropertyEditor",
-				"AssetTools",
-				"UMG",
-				"Sequencer",
-			}
-			);
 	}
 }
diff --git a/Engine/Source/Editor/UnrealEd/UnrealEd.Build.cs b/Engine/Source/Editor/UnrealEd/UnrealEd.Build.cs
--- a/Engine/Source/Editor/UnrealEd/UnrealEd.Build.cs
+++ b/Engine/Source/Editor/UnrealEd/UnrealEd.Build.cs
@@ -140,7 +140,6 @@ public class UnrealEd : ModuleRules
 			new string[] 
 			{
 				"CrashTracker",
-				"MaterialEditor",
 				"FontEditor",
 				"StaticMeshEditor",
 				"TextureEditor",
@@ -162,7 +161,6 @@ public class UnrealEd : ModuleRules
 				"ComponentVisualizers",
 				"MainFrame",
 				"LevelEditor",
-				"InputBindingEditor",
 				"PackagesDialog",
 				"Persona",
 				"PhAT",
@@ -171,7 +169,6 @@ public class UnrealEd : ModuleRules
 				"SettingsEditor",
 				"SessionFrontend",
 				"Sequencer",
-				"SoundClassEditor",
 				"GeometryMode",
 				"TextureAlignMode",
 				"FoliageEdit",
@@ -180,18 +177,14 @@ public class UnrealEd : ModuleRules
 				"Blutility",
 				"IntroTutorials",
                 "SuperSearch",
-				"DesktopPlatform",
 				"WorkspaceMenuStructure",
-				"BspMode",
 				"PlacementMode",
 				"NiagaraEditor",
 				"MeshUtilities",
                 "MergeActors",
-				"GameProjectGeneration",
 				"ProjectSettingsViewer",
 				"ProjectTargetPlatformEditor",
 				"PListEditor",
-                "Documentation",
                 "BehaviorTreeEditor",
                 "EnvironmentQueryEditor",
 				"ViewportSnapping",
diff --git a/Engine/Source/Editor/WorldBrowser/WorldBrowser.Build.cs b/Engine/Source/Editor/WorldBrowser/WorldBrowser.Build.cs
--- a/Engine/Source/Editor/WorldBrowser/WorldBrowser.Build.cs
+++ b/Engine/Source/Editor/WorldBrowser/WorldBrowser.Build.cs
@@ -55,7 +55,6 @@ public class WorldBrowser : ModuleRules
 				"SceneOutliner",
                 "MeshUtilities",
                 "ContentBrowser",
-                "ImageWrapper",
 			}
 		);
     }
diff --git a/Engine/Source/Programs/UnrealBuildTool/System/RulesCompiler.cs b/Engine/Source/Programs/UnrealBuildTool/System/RulesCompiler.cs
--- a/Engine/Source/Programs/UnrealBuildTool/System/RulesCompiler.cs
+++ b/Engine/Source/Programs/UnrealBuildTool/System/RulesCompiler.cs
@@ -1590,7 +1590,13 @@ namespace UnrealBuildTool
 				{
 					if (RulesObject.PrivateAssemblyReferences.Count > 0)
 					{
-						throw new BuildException("Module rules for '{0}' may not specify PrivateAssemblyReferences unless it is a CPlusPlusCLR module type.", AssemblyFileName);
+						throw new BuildException("Module rules for '{0}' may not specify PrivateAssemblyReferences unless it is a CPlusPlusCLR module type.", ModuleName);
+					}
+
+					var InvalidDependencies = RulesObject.DynamicallyLoadedModuleNames.Intersect(RulesObject.PublicDependencyModuleNames.Concat(RulesObject.PrivateDependencyModuleNames)).ToList();
+					if (InvalidDependencies.Count != 0)
+					{
+						throw new BuildException("Module rules for '{0}' should not be dependent on modules which are also dynamically loaded: {1}", ModuleName, String.Join(", ", InvalidDependencies));
 					}
 
 					// Choose code optimization options based on module type (game/engine) if
diff --git a/Engine/Source/Runtime/Launch/Launch.Build.cs b/Engine/Source/Runtime/Launch/Launch.Build.cs
--- a/Engine/Source/Runtime/Launch/Launch.Build.cs
+++ b/Engine/Source/Runtime/Launch/Launch.Build.cs
@@ -71,7 +71,6 @@ public class Launch : ModuleRules
             DynamicallyLoadedModuleNames.AddRange(
                 new string[] {
 			        "SlateRHIRenderer",
-					"HeadMountedDisplay"
 		        }
             );
         }
@@ -129,7 +128,6 @@ public class Launch : ModuleRules
 			DynamicallyLoadedModuleNames.AddRange(
 				new string[] {
 					"AutomationController",
-					"AutomationWorker",
 					"AutomationWindow",
 					"ProfilerClient",
 					"Toolbox",
