Added support for multiprocess cooking into unreal frontend.
Fixes for multiprocesscooking.

diff --git a/Engine/Source/Developer/LauncherServices/Private/Launcher/LauncherWorker.cpp b/Engine/Source/Developer/LauncherServices/Private/Launcher/LauncherWorker.cpp
--- a/Engine/Source/Developer/LauncherServices/Private/Launcher/LauncherWorker.cpp
+++ b/Engine/Source/Developer/LauncherServices/Private/Launcher/LauncherWorker.cpp
@@ -416,6 +416,11 @@ FString FLauncherWorker::CreateUATCommand( const ILauncherProfileRef& InProfile,
 				UATCommand += Cmd;
 			}
 
+			if (InProfile->GetNumCookersToSpawn() > 0)
+			{
+				UATCommand += FString::Printf(TEXT(" -NumCookersToSpawn=%d"), InProfile->GetNumCookersToSpawn());
+			}
+
 			FCommandDesc Desc;
 			FText Command = FText::Format(LOCTEXT("LauncherCookDesc", "Cook content for {0}"), FText::FromString(Platforms.RightChop(1)));
 			Desc.Name = "Cook Task";
@@ -468,6 +473,10 @@ FString FLauncherWorker::CreateUATCommand( const ILauncherProfileRef& InProfile,
 		UATCommand += TEXT(" -iterativecooking");
 	}
 
+	if (InProfile->GetSkipCookingEditorContent())
+	{
+		UATCommand += TEXT(" -SkipCookingEditorContent");
+	}
 
 	if ( InProfile->IsCompressed() )
 	{
diff --git a/Engine/Source/Developer/LauncherServices/Private/Profiles/LauncherProfile.h b/Engine/Source/Developer/LauncherServices/Private/Profiles/LauncherProfile.h
--- a/Engine/Source/Developer/LauncherServices/Private/Profiles/LauncherProfile.h
+++ b/Engine/Source/Developer/LauncherServices/Private/Profiles/LauncherProfile.h
@@ -8,6 +8,8 @@
 #define LAUNCHERSERVICES_REMOVEDPATCHSOURCECONTENTPATH 14
 #define LAUNCHERSERVICES_ADDEDDLCINCLUDEENGINECONTENT 15
 #define LAUNCHERSERVICES_ADDEDGENERATECHUNKS 16
+#define LAUNCHERSERVICES_ADDEDNUMCOOKERSTOSPAWN 17
+#define LAUNCHERSERVICES_ADDEDSKIPCOOKINGEDITORCONTENT 18
 
 /**
 * Implements a simple profile which controls the desired output of the Launcher for simple
@@ -130,7 +132,9 @@ public:
 	FLauncherProfile(ILauncherProfileManagerRef ProfileManager)
 		: LauncherProfileManager(ProfileManager)
 		, DefaultLaunchRole(MakeShareable(new FLauncherProfileLaunchRole()))
-	{ }
+	{ 
+		SetDefaults();
+	}
 
 	/**
 	 * Creates and initializes a new instance.
@@ -258,6 +262,16 @@ public:
 		return CookedCultures;
 	}
 
+	virtual const int32 GetNumCookersToSpawn() const override
+	{
+		return NumCookersToSpawn;
+	}
+
+	virtual const bool GetSkipCookingEditorContent() const override
+	{
+		return bSkipCookingEditorContent;
+	}
+
 	virtual const TArray<FString>& GetCookedMaps( ) const override
 	{
 		return CookedMaps;
@@ -582,7 +596,7 @@ public:
 
 	virtual bool Serialize( FArchive& Archive ) override
 	{
-		int32 Version = LAUNCHERSERVICES_PROFILEVERSION;
+		int32 Version = LAUNCHERSERVICES_ADDEDSKIPCOOKINGEDITORCONTENT;
 
 		Archive	<< Version;
 
@@ -630,11 +644,18 @@ public:
                 << ForceClose
                 << Timeout;
 
+		if (Version >= LAUNCHERSERVICES_ADDEDNUMCOOKERSTOSPAWN)
+		{
+			Archive << NumCookersToSpawn;
+		}
+		if (Version >= LAUNCHERSERVICES_ADDEDSKIPCOOKINGEDITORCONTENT)
+		{
+			Archive << bSkipCookingEditorContent;
+		}
 		if (Version >= LAUNCHERSERVICES_ADDEDINCREMENTALDEPLOYVERSION)
 		{
 			Archive << DeployIncremental;
 		}
-
 		if ( Version >= LAUNCHERSERVICES_REMOVEDPATCHSOURCECONTENTPATH )
 		{
 			Archive << GeneratePatch;
@@ -645,21 +666,11 @@ public:
 			Archive << Temp;
 			Archive << GeneratePatch;
 		}
-		else if ( Archive.IsLoading() )
-		{
-			GeneratePatch = false;
-		}
-
 		if (Version >= LAUNCHERSERVICES_ADDEDDLCINCLUDEENGINECONTENT)
 		{
 			Archive << DLCIncludeEngineContent;
 		}
-		else if (Archive.IsLoading())
-		{
-			DLCIncludeEngineContent = false;
-		}
 		
-
 		if ( Version >= LAUNCHERSERVICES_ADDEDRELEASEVERSION )
 		{
 			Archive << CreateReleaseVersion;
@@ -669,12 +680,6 @@ public:
 			Archive << CreateDLC;
 			Archive << DLCName;
 		}
-		else if ( Archive.IsLoading() )
-		{
-			CreateReleaseVersion = false;
-			CreateDLC = false;
-		}
-
 		if (Version >= LAUNCHERSERVICES_ADDEDGENERATECHUNKS)
 		{
 			Archive << bGenerateChunks;
@@ -682,14 +687,7 @@ public:
 			Archive << HttpChunkDataDirectory;
 			Archive << HttpChunkDataReleaseName;
 		}
-		else if (Archive.IsLoading())
-		{
-			bGenerateChunks = false;
-			bGenerateHttpChunkData = false;
-			HttpChunkDataDirectory = TEXT("");
-			HttpChunkDataReleaseName = TEXT("");
-		}
-
+		
 		DefaultLaunchRole->Serialize(Archive);
 
 		// serialize launch roles
@@ -762,8 +760,10 @@ public:
 		CookedCultures.Add(I18N.GetCurrentCulture()->GetName());
 		CookedMaps.Reset();
 		CookedPlatforms.Reset();
+		bSkipCookingEditorContent = true;
         ForceClose = true;
         Timeout = 60;
+		NumCookersToSpawn = 0;
 
 /*		if (GetTargetPlatformManager()->GetRunningTargetPlatform() != NULL)
 		{
@@ -868,6 +868,24 @@ public:
 		}
 	}
 
+	virtual void SetNumCookersToSpawn(const int32 InNumCookersToSpawn) override
+	{
+		if (NumCookersToSpawn != InNumCookersToSpawn)
+		{
+			NumCookersToSpawn = InNumCookersToSpawn;
+			Validate();
+		}
+	}
+
+	virtual void SetSkipCookingEditorContent(const bool InSkipCookingEditorContent) override
+	{
+		if (bSkipCookingEditorContent != InSkipCookingEditorContent)
+		{
+			bSkipCookingEditorContent = InSkipCookingEditorContent;
+			Validate();
+		}
+	}
+
 	virtual void SetDeployWithUnrealPak( bool UseUnrealPak ) override
 	{
 		if (DeployWithUnrealPak != UseUnrealPak)
@@ -1375,6 +1393,11 @@ private:
 	// Holds a flag indicating whether packages should be saved without a version.
 	bool CookUnversioned;
 
+	// num cookers we want to spawn during cooking
+	int32 NumCookersToSpawn;
+
+	bool bSkipCookingEditorContent;
+
 	// This setting is used only if cooking by the book (only used if cooking by the book).
 	TArray<FString> CookedCultures;
 
diff --git a/Engine/Source/Developer/LauncherServices/Public/Interfaces/ILauncherProfile.h b/Engine/Source/Developer/LauncherServices/Public/Interfaces/ILauncherProfile.h
--- a/Engine/Source/Developer/LauncherServices/Public/Interfaces/ILauncherProfile.h
+++ b/Engine/Source/Developer/LauncherServices/Public/Interfaces/ILauncherProfile.h
@@ -438,6 +438,34 @@ public:
 	 */
 	virtual const FString& GetCookOptions( ) const = 0;
 
+
+	/**
+	 * Get the number of cookers we want to spawn during cooking
+	 *
+	 * @return number of cookers we want to spawn in addition to the master cooker
+	 */
+	virtual const int32 GetNumCookersToSpawn() const = 0;
+
+	/**
+	 * Set the number of cookers we want to spawn during cooking
+	 * 
+	 * @param InNumCookersToSpawn number of cookers we want to spawn in addition to the master cooker
+	 */
+	virtual void SetNumCookersToSpawn(const int32 InNumCookersToSpawn) = 0; 
+
+
+
+	virtual const bool GetSkipCookingEditorContent() const = 0; 
+
+	/**
+	 * Skip editor content while cooking, 
+	 * This will strip editor content from final builds
+	 * 
+	 * @param InSkipCookingEditorContent
+	 */
+	virtual void SetSkipCookingEditorContent(const bool InSkipCookingEditorContent) = 0;
+
+
 	/**
 	 * Gets the list of cooked culture.
 	 *
diff --git a/Engine/Source/Developer/ProjectLauncher/Private/Widgets/Cook/SProjectLauncherCookByTheBookSettings.cpp b/Engine/Source/Developer/ProjectLauncher/Private/Widgets/Cook/SProjectLauncherCookByTheBookSettings.cpp
--- a/Engine/Source/Developer/ProjectLauncher/Private/Widgets/Cook/SProjectLauncherCookByTheBookSettings.cpp
+++ b/Engine/Source/Developer/ProjectLauncher/Private/Widgets/Cook/SProjectLauncherCookByTheBookSettings.cpp
@@ -563,6 +563,27 @@ TSharedRef<SWidget> SProjectLauncherCookByTheBookSettings::MakeComplexWidget()
 						]
 					]
 
+				// multiprocess cooking options
+				+ SVerticalBox::Slot()
+					.AutoHeight()
+					.Padding(0.0f, 8.0f, 0.0f, 0.0f)
+					[
+						SNew(SProjectLauncherFormLabel)
+						.LabelText(LOCTEXT("MultiProcessCookerTextBoxLabel", "Num cookers to spawn:"))
+					]
+
+				+ SVerticalBox::Slot()
+					.AutoHeight()
+					.Padding(0.0f, 4.0f, 0.0f, 0.0f)
+					[
+						// cooker command line options
+						SNew(SEditableTextBox)
+						.ToolTipText(LOCTEXT("MultiProcessCookerTextBoxTooltip", "The number of cookers to spawn when we do a cook by the book."))
+						.Text(this, &SProjectLauncherCookByTheBookSettings::HandleMultiProcessCookerTextBlockText)
+						.OnTextCommitted(this, &SProjectLauncherCookByTheBookSettings::HandleMultiProcessCookerCommitted)
+					]
+
+				// unreal pak check box
 				+ SVerticalBox::Slot()
 					.AutoHeight()
 					.Padding(0.0f, 4.0f, 0.0f, 0.0f)
@@ -580,11 +601,12 @@ TSharedRef<SWidget> SProjectLauncherCookByTheBookSettings::MakeComplexWidget()
 						]
 					]
 
+				// generate chunks check box
 				+ SVerticalBox::Slot()
 					.AutoHeight()
 					.Padding(0.0f, 4.0f, 0.0f, 0.0f)
 					[
-						// unreal pak check box
+						
 						SNew(SCheckBox)
 						.IsChecked(this, &SProjectLauncherCookByTheBookSettings::HandleGenerateChunksCheckBoxIsChecked)
 						.OnCheckStateChanged(this, &SProjectLauncherCookByTheBookSettings::HandleGenerateChunksCheckBoxCheckStateChanged)
@@ -597,6 +619,25 @@ TSharedRef<SWidget> SProjectLauncherCookByTheBookSettings::MakeComplexWidget()
 						]
 					]
 
+				// don't include editor content
+				+ SVerticalBox::Slot()
+					.AutoHeight()
+					.Padding(0.0f, 4.0f, 0.0f, 0.0f)
+					[
+
+						SNew(SCheckBox)
+						.IsChecked(this, &SProjectLauncherCookByTheBookSettings::HandleDontIncludeEditorContentCheckBoxIsChecked)
+						.OnCheckStateChanged(this, &SProjectLauncherCookByTheBookSettings::HandleDontIncludeEditorContentCheckBoxCheckStateChanged)
+						.Padding(FMargin(4.0f, 0.0f))
+						.ToolTipText(LOCTEXT("DontIncludeEditorContentCheckBoxTooltip", "If checked the cooker will skip editor content and not include it in the build."))
+						.Content()
+						[
+							SNew(STextBlock)
+							.Text(LOCTEXT("DontIncludeEditorContentCheckBoxText", "Don't Include editor content in the build"))
+						]
+					]
+
+
 				+ SVerticalBox::Slot()
 					.AutoHeight()
 					.Padding(0.0f, 4.0f, 0.0f, 0.0f)
@@ -703,6 +744,8 @@ TSharedRef<SWidget> SProjectLauncherCookByTheBookSettings::MakeComplexWidget()
 						.ToolTipText(LOCTEXT("CookConfigurationToolTipText", "Sets the build configuration to use for the cooker commandlet."))
 					]
 
+
+				// additional cooker options text box
 				+ SVerticalBox::Slot()
 					.AutoHeight()
 					.Padding(0.0f, 8.0f, 0.0f, 0.0f)
@@ -721,6 +764,8 @@ TSharedRef<SWidget> SProjectLauncherCookByTheBookSettings::MakeComplexWidget()
 						.Text(this, &SProjectLauncherCookByTheBookSettings::HandleCookOptionsTextBlockText)
 						.OnTextCommitted(this, &SProjectLauncherCookByTheBookSettings::HandleCookerOptionsCommitted)
 					]
+
+
 			]
 		];
 
@@ -1376,6 +1421,45 @@ void SProjectLauncherCookByTheBookSettings::HandleCookerOptionsCommitted(const F
 	}
 }
 
+
+
+FText SProjectLauncherCookByTheBookSettings::HandleMultiProcessCookerTextBlockText() const
+{
+	ILauncherProfilePtr SelectedProfile = Model->GetSelectedProfile();
+
+	FText result;
+
+	if (SelectedProfile.IsValid())
+	{
+		result = FText::FromString(FString::FromInt(SelectedProfile->GetNumCookersToSpawn()));
+	}
+
+	return result;
+}
+
+void SProjectLauncherCookByTheBookSettings::HandleMultiProcessCookerCommitted(const FText& NewText, ETextCommit::Type CommitType)
+{
+	ILauncherProfilePtr SelectedProfile = Model->GetSelectedProfile();
+
+	if (SelectedProfile.IsValid())
+	{
+		int32 NumCookersToSpawn = FCString::Atoi(*NewText.ToString());
+		switch (CommitType)
+		{
+		case ETextCommit::Default:
+		case ETextCommit::OnCleared:
+			NumCookersToSpawn = 0;
+			break;
+		default:
+			break;
+		}
+		SelectedProfile->SetNumCookersToSpawn(NumCookersToSpawn);
+	}
+}
+
+
+
+
 void SProjectLauncherCookByTheBookSettings::HandleUnrealPakCheckBoxCheckStateChanged( ECheckBoxState NewState )
 {
 	ILauncherProfilePtr SelectedProfile = Model->GetSelectedProfile();
@@ -1714,4 +1798,28 @@ ECheckBoxState SProjectLauncherCookByTheBookSettings::HandleGenerateChunksCheckB
 	return ECheckBoxState::Unchecked;
 }
 
+
+
+
+
+void SProjectLauncherCookByTheBookSettings::HandleDontIncludeEditorContentCheckBoxCheckStateChanged(ECheckBoxState NewState)
+{
+	ILauncherProfilePtr SelectedProfile = Model->GetSelectedProfile();
+
+	if (SelectedProfile.IsValid())
+	{
+		SelectedProfile->SetSkipCookingEditorContent(NewState == ECheckBoxState::Checked);
+	}
+}
+
+ECheckBoxState SProjectLauncherCookByTheBookSettings::HandleDontIncludeEditorContentCheckBoxIsChecked() const
+{
+	ILauncherProfilePtr SelectedProfile = Model->GetSelectedProfile();
+	if (SelectedProfile.IsValid())
+	{
+		return SelectedProfile->GetSkipCookingEditorContent() ? ECheckBoxState::Checked : ECheckBoxState::Unchecked;
+	}
+	return ECheckBoxState::Unchecked;
+}
+
 #undef LOCTEXT_NAMESPACE
diff --git a/Engine/Source/Developer/ProjectLauncher/Private/Widgets/Cook/SProjectLauncherCookByTheBookSettings.h b/Engine/Source/Developer/ProjectLauncher/Private/Widgets/Cook/SProjectLauncherCookByTheBookSettings.h
--- a/Engine/Source/Developer/ProjectLauncher/Private/Widgets/Cook/SProjectLauncherCookByTheBookSettings.h
+++ b/Engine/Source/Developer/ProjectLauncher/Private/Widgets/Cook/SProjectLauncherCookByTheBookSettings.h
@@ -112,10 +112,14 @@ private:
 
 	/** Callback for getting the cookers additional options. */
 	FText HandleCookOptionsTextBlockText() const;
-
 	/** Callback for changing the cookers additional options */
 	void HandleCookerOptionsCommitted(const FText& NewText, ETextCommit::Type CommitType);
 
+	/** Callback for when the number of cookers need to be refreshed */
+	FText HandleMultiProcessCookerTextBlockText() const;
+	/** callback for when the number of cookers changes */
+	void HandleMultiProcessCookerCommitted(const FText& NewText, ETextCommit::Type CommitType);
+
 	/** Callback for updating any settings after the selected project has changed in the profile. */
 	void HandleProfileProjectChanged();
 
@@ -128,6 +132,11 @@ private:
 	void HandleGenerateChunksCheckBoxCheckStateChanged(ECheckBoxState NewState);
 	ECheckBoxState HandleGenerateChunksCheckBoxIsChecked() const;
 
+	// callbacks for including editor content checkbox
+	void HandleDontIncludeEditorContentCheckBoxCheckStateChanged(ECheckBoxState NewState);
+	ECheckBoxState HandleDontIncludeEditorContentCheckBoxIsChecked() const;
+
+
 
 	//////////////////////////////////////////////////////////////////////////
 	// creating release version related functions
diff --git a/Engine/Source/Developer/TargetPlatform/Private/TargetPlatformManagerModule.cpp b/Engine/Source/Developer/TargetPlatform/Private/TargetPlatformManagerModule.cpp
--- a/Engine/Source/Developer/TargetPlatform/Private/TargetPlatformManagerModule.cpp
+++ b/Engine/Source/Developer/TargetPlatform/Private/TargetPlatformManagerModule.cpp
@@ -40,12 +40,15 @@ public:
 			DECLARE_SCOPE_CYCLE_COUNTER( TEXT( "FTargetPlatformManagerModule.StartAutoSDK" ), STAT_FTargetPlatformManagerModule_StartAutoSDK, STATGROUP_TargetPlatform );
 
 			// amortize UBT cost by calling it once for all platforms, rather than once per platform.
-			FString UBTParams(TEXT("-autosdkonly"));
-			int32 UBTReturnCode = -1;
-			FString UBTOutput;
-			if (!FDesktopPlatformModule::Get()->InvokeUnrealBuildToolSync(UBTParams, *GLog, true, UBTReturnCode, UBTOutput))
+			if (FParse::Param(FCommandLine::Get(), TEXT("Multiprocess"))==false)
 			{
-				UE_LOG(LogTargetPlatformManager, Fatal, TEXT("Failed to run UBT to check SDK status!"));
+				FString UBTParams(TEXT("-autosdkonly"));
+				int32 UBTReturnCode = -1;
+				FString UBTOutput;
+				if (!FDesktopPlatformModule::Get()->InvokeUnrealBuildToolSync(UBTParams, *GLog, true, UBTReturnCode, UBTOutput))
+				{
+					UE_LOG(LogTargetPlatformManager, Fatal, TEXT("Failed to run UBT to check SDK status!"));
+				}
 			}
 
 			// we have to setup our local environment according to AutoSDKs or the ITargetPlatform's IsSDkInstalled calls may fail
diff --git a/Engine/Source/Editor/UnrealEd/Classes/CookOnTheSide/CookOnTheFlyServer.h b/Engine/Source/Editor/UnrealEd/Classes/CookOnTheSide/CookOnTheFlyServer.h
--- a/Engine/Source/Editor/UnrealEd/Classes/CookOnTheSide/CookOnTheFlyServer.h
+++ b/Engine/Source/Editor/UnrealEd/Classes/CookOnTheSide/CookOnTheFlyServer.h
@@ -22,7 +22,7 @@ enum class ECookInitializationFlags
 	AutoTick = 0x10,				// enable ticking (only works in the editor)
 	AsyncSave = 0x20,				// save packages async
 	IncludeServerMaps = 0x80,		// should we include the server maps when cooking
-	UseSerializationForPackageDependencies = 0x100, // should we use the serialization code path for generating package dependencies (old method will be depricated)
+	UseSerializationForPackageDependencies = 0x100, // should we use the serialization code path for generating package dependencies (old method will be deprecated)
 };
 ENUM_CLASS_FLAGS(ECookInitializationFlags);
 
diff --git a/Engine/Source/Editor/UnrealEd/Private/Commandlets/ChunkManifestGenerator.cpp b/Engine/Source/Editor/UnrealEd/Private/Commandlets/ChunkManifestGenerator.cpp
--- a/Engine/Source/Editor/UnrealEd/Private/Commandlets/ChunkManifestGenerator.cpp
+++ b/Engine/Source/Editor/UnrealEd/Private/Commandlets/ChunkManifestGenerator.cpp
@@ -21,6 +21,48 @@
 
 DEFINE_LOG_CATEGORY_STATIC(LogChunkManifestGenerator, Log, All);
 
+#define LOCTEXT_NAMESPACE "ChunkManifestGenerator"
+
+//////////////////////////////////////////////////////////////////////////
+// Static functions
+FName GetPackageNameFromDependencyPackageName(const FName RawPackageFName)
+{
+	FName PackageFName = RawPackageFName;
+	if ((FPackageName::IsValidLongPackageName(RawPackageFName.ToString()) == false) &&
+		(FPackageName::IsScriptPackage(RawPackageFName.ToString()) == false))
+	{
+		FText OutReason;
+		if (!FPackageName::IsValidLongPackageName(RawPackageFName.ToString(), true, &OutReason))
+		{
+			const FText FailMessage = FText::Format(LOCTEXT("UnableToGeneratePackageName", "Unable to generate long package name for {0}. {1}"),
+				FText::FromString(RawPackageFName.ToString()), OutReason);
+
+			UE_LOG(LogChunkManifestGenerator, Warning, TEXT("%s"), *(FailMessage.ToString()));
+			return NAME_None;
+		}
+
+
+		FString LongPackageName;
+		if (FPackageName::SearchForPackageOnDisk(RawPackageFName.ToString(), &LongPackageName) == false)
+		{
+			return NAME_None;
+		}
+		PackageFName = FName(*LongPackageName);
+	}
+
+	// don't include script packages in dependencies as they are always in memory
+	if (FPackageName::IsScriptPackage(PackageFName.ToString()))
+	{
+		// no one likes script packages
+		return NAME_None;
+	}
+	return PackageFName;
+}
+
+
+//////////////////////////////////////////////////////////////////////////
+// FChunkManifestGenerator
+
 FChunkManifestGenerator::FChunkManifestGenerator(const TArray<ITargetPlatform*>& InPlatforms)
 	: AssetRegistry(FModuleManager::LoadModuleChecked<FAssetRegistryModule>(TEXT("AssetRegistry")).Get())
 	, Platforms(InPlatforms)
@@ -410,27 +452,13 @@ void FChunkManifestGenerator::BuildChunkManifest(const TArray<FName>& CookedPack
 		
 		for (const auto& RawPackageFName : MapDependencies)
 		{
-			FName PackageFName = RawPackageFName;
-			if ((FPackageName::IsValidLongPackageName(RawPackageFName.ToString()) == false) &&
-				(FPackageName::IsScriptPackage(RawPackageFName.ToString()) == false) )
-			{
-				FString LongPackageName;
-				if (FPackageName::SearchForPackageOnDisk(RawPackageFName.ToString(), &LongPackageName) == false)
-				{
-					continue;
-				}
-				PackageFName = FName(*LongPackageName);
-			}
+			const FName PackageFName = GetPackageNameFromDependencyPackageName(RawPackageFName);
 			
-			// don't include script packages in dependencies as they are always in memory
-			if (FPackageName::IsScriptPackage(PackageFName.ToString()))
+			if (PackageFName == NAME_None)
 			{
-				// no one likes script packages
 				continue;
 			}
 			
-
-
 			const FString PackagePathName = PackageFName.ToString();
 			const FString MapName = MapFName.ToString();
 			const FString* SandboxFilenamePtr = AllCookedPackages.Find(PackageFName);
@@ -507,7 +535,7 @@ bool FChunkManifestGenerator::SaveAssetRegistry(const FString& SandboxPath, cons
 			IgnorePackageSet.Add(IgnorePackage);
 		}
 	}
-
+	
 
 	// Create asset registry data
 	TArray<FName> MapList;
@@ -529,9 +557,9 @@ bool FChunkManifestGenerator::SaveAssetRegistry(const FString& SandboxPath, cons
 			if (ContainsMap(AssetData.PackageName))
 			{
 				MapList.Add(AssetData.PackageName);
-			}
 		}
 	}
+	}
 
 	AssetRegistry.SaveRegistryData(SerializedAssetRegistry, GeneratedAssetRegistryData, &MapList);
 	UE_LOG(LogChunkManifestGenerator, Display, TEXT("Generated asset registry num assets %d, size is %5.2fkb"), GeneratedAssetRegistryData.Num(), (float)SerializedAssetRegistry.Num() / 1024.f);
@@ -1076,13 +1104,13 @@ void FChunkManifestGenerator::AddPackageAndDependenciesToChunk(FChunkPackageSet*
 			}
 			if (!bSkip)
 			{
-				auto DependentPackageLongName = PkgName.ToString();
-				if (FPackageName::IsShortPackageName(PkgName))
+				const FName FilteredPackageName = GetPackageNameFromDependencyPackageName(PkgName);
+				if (FilteredPackageName == NAME_None)
 				{
-					DependentPackageLongName = FPackageName::ConvertToLongScriptPackageName(*PkgName.ToString());
+					continue;
 				}
-				FString DependentSandboxFile = SandboxPlatformFile->ConvertToAbsolutePathForExternalAppForWrite(*FPackageName::LongPackageNameToFilename(DependentPackageLongName));
-				ThisPackageSet->Add(PkgName, DependentSandboxFile);
+				FString DependentSandboxFile = SandboxPlatformFile->ConvertToAbsolutePathForExternalAppForWrite(*FPackageName::LongPackageNameToFilename(*FilteredPackageName.ToString()));
+				ThisPackageSet->Add(FilteredPackageName, DependentSandboxFile);
 				UnassignedPackageSet.Remove(PkgName);
 			}
 		}
diff --git a/Engine/Source/Editor/UnrealEd/Private/Commandlets/CookCommandlet.cpp b/Engine/Source/Editor/UnrealEd/Private/Commandlets/CookCommandlet.cpp
--- a/Engine/Source/Editor/UnrealEd/Private/Commandlets/CookCommandlet.cpp
+++ b/Engine/Source/Editor/UnrealEd/Private/Commandlets/CookCommandlet.cpp
@@ -1057,7 +1057,7 @@ bool UCookCommandlet::NewCook( const TArray<ITargetPlatform*>& Platforms, TArray
 	FParse::Value(*Params, TEXT("cookchild="), ChildCookFile);
 
 	int32 NumProcesses = 0;
-	FParse::Value(*Params, TEXT("multiprocesscooker="), NumProcesses);
+	FParse::Value(*Params, TEXT("numcookerstospawn="), NumProcesses);
 
 	FString BasedOnReleaseVersion;
 	FParse::Value( *Params, TEXT("BasedOnReleaseVersion="), BasedOnReleaseVersion);
@@ -1250,7 +1250,6 @@ bool UCookCommandlet::NewCook( const TArray<ITargetPlatform*>& Platforms, TArray
 		}
 
 
-
 		ProcessDeferredCommands();
 	}
 
diff --git a/Engine/Source/Editor/UnrealEd/Private/CookOnTheFlyServer.cpp b/Engine/Source/Editor/UnrealEd/Private/CookOnTheFlyServer.cpp
--- a/Engine/Source/Editor/UnrealEd/Private/CookOnTheFlyServer.cpp
+++ b/Engine/Source/Editor/UnrealEd/Private/CookOnTheFlyServer.cpp
@@ -43,7 +43,7 @@ DEFINE_LOG_CATEGORY_STATIC(LogCook, Log, All);
 
 
 #define DEBUG_COOKONTHEFLY 0
-#define OUTPUT_TIMING 0
+#define OUTPUT_TIMING 1
 
 #define USEASSETREGISTRYFORDEPENDENTPACKAGES 1
 #define VERIFY_GETDEPENDENTPACKAGES 0 // verify has false hits because old serialization method for generating dependencies had errors (included transient objects which shouldn't be in asset registry), but you can still use verify to build a list then cross check against transient objects.  
@@ -1140,6 +1140,7 @@ COREUOBJECT_API extern bool GOutputCookingWarnings;
 
 bool UCookOnTheFlyServer::TickChildCookers(const bool bCheckIfFinished)
 {
+	SCOPE_TIMER(TickChildCookers);
 	bool bChildCookersFinished = true;
 	if (IsCookByTheBookMode())
 	{
@@ -1151,15 +1152,16 @@ bool UCookOnTheFlyServer::TickChildCookers(const bool bCheckIfFinished)
 				int32 ReturnCode = 0;
 				if (bCheckIfFinished == true)
 				{
+					SCOPE_TIMER(CheckIfChildCookerIsFinished);
 					if (FPlatformProcess::GetProcReturnCode(ChildCooker.ProcessHandle, &ReturnCode))
 					{
-						UE_LOG(LogCook, Display, TEXT("ChildCooker returned error code %d"), ReturnCode);
+						UE_LOG(LogCook, Display, TEXT("Child cooker %s returned error code %d"), *ChildCooker.BaseResponseFileName, ReturnCode);
 						ChildCooker.ReturnCode = ReturnCode;
 						ChildCooker.bFinished = true;
 
 						if (ReturnCode != 0)
 						{
-							UE_LOG(LogCook, Error, TEXT("ChildCooker returned error code %d"), ReturnCode);
+							UE_LOG(LogCook, Error, TEXT("Child cooker %s returned error code %d"), *ChildCooker.BaseResponseFileName, ReturnCode);
 						}
 
 
@@ -1194,6 +1196,7 @@ bool UCookOnTheFlyServer::TickChildCookers(const bool bCheckIfFinished)
 					bChildCookersFinished = false;
 				}
 
+				SCOPE_TIMER(ProcessLogOutput);
 				// process the log output from the child cooker even if we just finished to ensure that we get the end of the log
 				FString PipeContents = FPlatformProcess::ReadPipe(ChildCooker.ReadPipe);
 				if (PipeContents.Len())
@@ -1276,18 +1279,13 @@ uint32 UCookOnTheFlyServer::TickCookOnTheSide( const float TimeSlice, uint32 &Co
 	
 	if (CurrentCookMode == ECookMode::CookByTheBook)
 	{
-		SCOPE_TIMER(TickChildCookers);
 		if (CookRequests.HasItems() == false)
 		{
-			SCOPE_TIMER(WaitingForChildCookersToFinish);
+			SCOPE_TIMER(WaitingForChildCookers);
 			// we have nothing left to cook so we are now waiting for child cookers to finish
 			if (TickChildCookers(true) == false)
 			{
 				Result |= COSR_WaitingOnChildCookers;
-				if (CookRequests.HasItems() == false)
-				{
-					FPlatformProcess::Sleep(0.001f);
-				}
 			}
 			
 		}
@@ -4107,17 +4105,46 @@ void UCookOnTheFlyServer::StartChildCookers(int32 NumCookersToSpawn, const TArra
 
 		FFileHelper::SaveStringToFile(ResponseFileText, *ChildCooker.ResponseFileName);
 
-		FString CommandLine = FString::Printf(TEXT("%s -run=cook -targetplatform=%s -cookchild=%s -abslog=%sLog.txt"), *FPaths::GetProjectFilePath(), *TargetPlatformString, *ChildCooker.ResponseFileName, *ChildCooker.ResponseFileName);
+		// default commands
+		// multiprocess tells unreal in general we shouldn't do things like save ddc, clean shader working directory, and other various multiprocess unsafe things
+		FString CommandLine = FString::Printf(TEXT("\"%s\" -run=cook -multiprocess -targetplatform=%s -cookchild=\"%s\" -abslog=\"%sLog.txt\""), *FPaths::GetProjectFilePath(), *TargetPlatformString, *ChildCooker.ResponseFileName, *ChildCooker.ResponseFileName);
 
-		FString DDCCommandline;
-		if (FParse::Value(FCommandLine::Get(), TEXT("ddc="), DDCCommandline))
+		auto KeepCommandlineValue = [&](const TCHAR* CommandlineToKeep)
 		{
-			CommandLine += TEXT(" -ddc=");
-			CommandLine += DDCCommandline;
-		}
+			FString CommandlineValue;
+			if (FParse::Value(FCommandLine::Get(), CommandlineToKeep, CommandlineValue ))
+			{
+				CommandLine += TEXT(" -");
+				CommandLine += CommandlineToKeep;
+				CommandLine += TEXT("=");
+				CommandLine += *CommandlineValue;
+			}
+		};
+
+		auto KeepCommandlineParam = [&](const TCHAR* CommandlineToKeep)
+		{
+			if (FParse::Param(FCommandLine::Get(), CommandlineToKeep))
+			{
+				CommandLine += TEXT(" -");
+				CommandLine += CommandlineToKeep;
+			}
+		};
+
+
+
+		KeepCommandlineValue(TEXT("ddc="));
+		KeepCommandlineParam(TEXT("SkipEditorContent"));
+		KeepCommandlineParam(TEXT("compressed"));
+		KeepCommandlineParam(TEXT("Unversioned"));
+		KeepCommandlineParam(TEXT("buildmachine"));
+		KeepCommandlineParam(TEXT("fileopenlog"));
+		KeepCommandlineParam(TEXT("stdout"));
+		KeepCommandlineParam(TEXT("FORCELOGFLUSH"));
+		KeepCommandlineParam(TEXT("CrashForUAT"));
+		KeepCommandlineParam(TEXT("AllowStdOutLogVerbosity"));
+		KeepCommandlineParam(TEXT("UTF8Output"));
 
 		FString ExecutablePath = FPlatformProcess::ExecutableName(true);
-		// FString ExeFileName = FPaths::EngineDir() / TEXT("Binaries") / PlatformConfig / FString(FPlatformProcess::ExecutableName()) + TEXT(".exe");
 
 		UE_LOG(LogCook, Display, TEXT("Launching cooker using commandline %s %s"), *ExecutablePath, *CommandLine);
 
@@ -4125,8 +4152,7 @@ void UCookOnTheFlyServer::StartChildCookers(int32 NumCookersToSpawn, const TArra
 		void* WritePipe = NULL;
 		FPlatformProcess::CreatePipe(ReadPipe, WritePipe);
 
-		//ChildCooker.ProcessHandle = FPlatformProcess::CreateProc(*ExecutablePath, *CommandLine, false, true, true, NULL, 0, NULL, WritePipe);
-		ChildCooker.ProcessHandle = FPlatformProcess::CreateProc(*ExecutablePath, *CommandLine, false, true, true, NULL, 0, NULL, NULL);
+		ChildCooker.ProcessHandle = FPlatformProcess::CreateProc(*ExecutablePath, *CommandLine, false, true, true, NULL, 0, NULL, WritePipe);
 		ChildCooker.ReadPipe = ReadPipe;
 	}
 
diff --git a/Engine/Source/Programs/AutomationTool/AutomationUtils/ProjectParams.cs b/Engine/Source/Programs/AutomationTool/AutomationUtils/ProjectParams.cs
--- a/Engine/Source/Programs/AutomationTool/AutomationUtils/ProjectParams.cs
+++ b/Engine/Source/Programs/AutomationTool/AutomationUtils/ProjectParams.cs
@@ -233,6 +233,8 @@ namespace AutomationTool
 			this.CookOnTheFly = InParams.CookOnTheFly;
             this.CookOnTheFlyStreaming = InParams.CookOnTheFlyStreaming;
             this.UnversionedCookedContent = InParams.UnversionedCookedContent;
+            this.SkipCookingEditorContent = InParams.SkipCookingEditorContent;
+            this.NumCookersToSpawn = InParams.NumCookersToSpawn;
 			this.FileServer = InParams.FileServer;
 			this.DedicatedServer = InParams.DedicatedServer;
 			this.Client = InParams.Client;
@@ -344,6 +346,8 @@ namespace AutomationTool
             bool? CookOnTheFly = null,
             bool? CookOnTheFlyStreaming = null,
             bool? UnversionedCookedContent = null,
+            bool? SkipCookingEditorContent = null,
+            int? NumCookersToSpawn = null,
             string AdditionalCookerOptions = null,
             string BasedOnReleaseVersion = null,
             string CreateReleaseVersion = null,
@@ -478,6 +482,15 @@ namespace AutomationTool
             }
             this.CookOnTheFlyStreaming = GetParamValueIfNotSpecified(Command, CookOnTheFlyStreaming, this.CookOnTheFlyStreaming, "cookontheflystreaming");
             this.UnversionedCookedContent = GetParamValueIfNotSpecified(Command, UnversionedCookedContent, this.UnversionedCookedContent, "UnversionedCookedContent");
+            this.SkipCookingEditorContent = GetParamValueIfNotSpecified(Command, SkipCookingEditorContent, this.SkipCookingEditorContent, "SkipCookingEditorContent");
+            if (NumCookersToSpawn.HasValue)
+            {
+                this.NumCookersToSpawn = NumCookersToSpawn.Value;
+            }
+            else if (Command != null)
+            {
+                this.NumCookersToSpawn = Command.ParseParamInt("NumCookersToSpawn");
+            }
             this.Compressed = GetParamValueIfNotSpecified(Command, Compressed, this.Compressed, "compressed");
             this.UseDebugParamForEditorExe = GetParamValueIfNotSpecified(Command, UseDebugParamForEditorExe, this.UseDebugParamForEditorExe, "UseDebugParamForEditorExe");
             this.IterativeCooking = GetParamValueIfNotSpecified(Command, IterativeCooking, this.IterativeCooking, new string[] { "iterativecooking", "iterate" } );
@@ -1074,7 +1087,7 @@ namespace AutomationTool
         public bool UseDebugParamForEditorExe;
 
         /// <summary>
-        /// 
+        /// Cook: Do not include a version number in the cooked content
         /// </summary>
         public bool UnversionedCookedContent;
 
@@ -1098,6 +1111,16 @@ namespace AutomationTool
         public bool CookAll;
 
 
+        /// <summary>
+        /// Cook: Skip cooking editor content 
+        /// </summary>
+        public bool SkipCookingEditorContent;
+
+        /// <summary>
+        /// Cook: number of additional cookers to spawn while cooking
+        /// </summary>
+        public int NumCookersToSpawn;
+
 		/// <summary>
 		/// Cook: Uses the iterative deploy, command line: -iterativedeploy or -iterate
 		/// </summary>
@@ -2019,6 +2042,8 @@ namespace AutomationTool
 				CommandUtils.Log("CookOnTheFly={0}", CookOnTheFly);
                 CommandUtils.Log("CookOnTheFlyStreaming={0}", CookOnTheFlyStreaming);
                 CommandUtils.Log("UnversionedCookedContent={0}", UnversionedCookedContent);
+                CommandUtils.Log("SkipCookingEditorContent={0}", SkipCookingEditorContent);
+                CommandUtils.Log("NumCookersToSpawn={0}", NumCookersToSpawn);
                 CommandUtils.Log("GeneratePatch={0}", GeneratePatch);
                 CommandUtils.Log("CreateReleaseVersion={0}", CreateReleaseVersion);
                 CommandUtils.Log("BasedOnReleaseVersion={0}", BasedOnReleaseVersion);
diff --git a/Engine/Source/Programs/AutomationTool/Scripts/CookCommand.Automation.cs b/Engine/Source/Programs/AutomationTool/Scripts/CookCommand.Automation.cs
--- a/Engine/Source/Programs/AutomationTool/Scripts/CookCommand.Automation.cs
+++ b/Engine/Source/Programs/AutomationTool/Scripts/CookCommand.Automation.cs
@@ -175,6 +175,15 @@ public partial class Project : CommandUtils
                 {
                     CommandletParams += " -createreleaseversion=" + Params.CreateReleaseVersion;
                 }
+                if ( Params.SkipCookingEditorContent)
+                {
+                    CommandletParams += " -skipeditorcontent";
+                }
+                if ( Params.NumCookersToSpawn != 0)
+                {
+                    CommandletParams += " -numcookerstospawn=" + Params.NumCookersToSpawn;
+
+                }
                 if (Params.HasDLCName)
                 {
                     CommandletParams += " -dlcname=" + Params.DLCName;
