UE-15527 - Reimporting a texture causes all materials which reference it to render incorrectly until the material is recompiled

Copied the RHI texure reference from the existing material for the new material.
All materials using the texture now still have a valid reference.
Wrapped this up in a helper class FTextureReferenceReplacer. Mainly for aesthetics rather than reusablilty.

Tagging Nick and Matt as I'm not sure who owns the texure importing code atm.

diff --git a/Engine/Source/Editor/UnrealEd/Private/Factories/EditorFactories.cpp b/Engine/Source/Editor/UnrealEd/Private/Factories/EditorFactories.cpp
--- a/Engine/Source/Editor/UnrealEd/Private/Factories/EditorFactories.cpp
+++ b/Engine/Source/Editor/UnrealEd/Private/Factories/EditorFactories.cpp
@@ -4291,12 +4291,8 @@ UObject* UTextureFactory::FactoryCreateBinary
 		// Update with new settings, which should disable streaming...
 		ExistingTexture2D->UpdateResource();
 	}
-
-	if (ExistingTexture)
-	{
-		// Release the existing resource so the new texture can get a fresh one. 
-		ExistingTexture->ReleaseResource();
-	}
+	
+	FTextureReferenceReplacer RefReplacer(ExistingTexture);
 
 	UTexture* Texture = ImportTexture(Class, InParent, Name, Flags, Type, Buffer, BufferEnd, Warn);
 
@@ -4313,6 +4309,9 @@ UObject* UTextureFactory::FactoryCreateBinary
 		return nullptr;
 	}
 
+	//Replace the reference for the new texture with the existing one so that all current users still have valid references.
+	RefReplacer.Replace(Texture);
+
 	// Start with the value that the loader suggests.
 	CompressionSettings = Texture->CompressionSettings;
 
diff --git a/Engine/Source/Runtime/Engine/Classes/Engine/Texture.h b/Engine/Source/Runtime/Engine/Classes/Engine/Texture.h
--- a/Engine/Source/Runtime/Engine/Classes/Engine/Texture.h
+++ b/Engine/Source/Runtime/Engine/Classes/Engine/Texture.h
@@ -781,5 +781,32 @@ public:
 
 };
 
+/** 
+* Replaces the RHI reference of one texture with another.
+* Allows one texture to be replaced with another at runtime and have all existing references to it remain valid.
+*/
+struct FTextureReferenceReplacer
+{
+	FTextureReferenceRHIRef OriginalRef;
 
+	FTextureReferenceReplacer(UTexture* OriginalTexture)
+	{
+		if (OriginalTexture)
+		{
+			OriginalTexture->ReleaseResource();
+			OriginalRef = OriginalTexture->TextureReference.TextureReferenceRHI;
+		}
+		else
+		{
+			OriginalRef = nullptr;
+		}
+	}
 
+	void Replace(UTexture* NewTexture)
+	{
+		if (OriginalRef)
+		{
+			NewTexture->TextureReference.TextureReferenceRHI = OriginalRef;
+		}
+	}
+};
