Fix for FBX import failure / crash due to incorrect factories attempting to import.

diff --git a/Engine/Source/Editor/UnrealEd/Classes/Factories/FbxFactory.h b/Engine/Source/Editor/UnrealEd/Classes/Factories/FbxFactory.h
--- a/Engine/Source/Editor/UnrealEd/Classes/Factories/FbxFactory.h
+++ b/Engine/Source/Editor/UnrealEd/Classes/Factories/FbxFactory.h
@@ -26,6 +26,7 @@ class UNREALED_API UFbxFactory : public UFactory
 	virtual bool DoesSupportClass(UClass * Class) override;
 	virtual UClass* ResolveSupportedClass() override;
 	virtual UObject* FactoryCreateBinary(UClass* InClass, UObject* InParent, FName InName, EObjectFlags Flags, UObject* Context, const TCHAR* Type, const uint8*& Buffer, const uint8* BufferEnd, FFeedbackContext* Warn, bool& bOutOperationCanceled) override;
+	virtual bool FactoryCanImport(const FString& Filename) override;
 	// End UFactory Interface
 	
 	/**
diff --git a/Engine/Source/Editor/UnrealEd/Classes/Factories/ReimportFbxAnimSequenceFactory.h b/Engine/Source/Editor/UnrealEd/Classes/Factories/ReimportFbxAnimSequenceFactory.h
--- a/Engine/Source/Editor/UnrealEd/Classes/Factories/ReimportFbxAnimSequenceFactory.h
+++ b/Engine/Source/Editor/UnrealEd/Classes/Factories/ReimportFbxAnimSequenceFactory.h
@@ -19,6 +19,10 @@ class UReimportFbxAnimSequenceFactory : public UFbxFactory, public FReimportHand
 	virtual EReimportResult::Type Reimport( UObject* Obj ) override;
 	virtual int32 GetPriority() const override;
 	// End FReimportHandler interface
+
+	// Begin UFactory Interface
+	virtual bool FactoryCanImport(const FString& Filename) override;
+	// End UFactory Interface
 };
 
 
diff --git a/Engine/Source/Editor/UnrealEd/Classes/Factories/ReimportFbxSkeletalMeshFactory.h b/Engine/Source/Editor/UnrealEd/Classes/Factories/ReimportFbxSkeletalMeshFactory.h
--- a/Engine/Source/Editor/UnrealEd/Classes/Factories/ReimportFbxSkeletalMeshFactory.h
+++ b/Engine/Source/Editor/UnrealEd/Classes/Factories/ReimportFbxSkeletalMeshFactory.h
@@ -19,6 +19,10 @@ class UReimportFbxSkeletalMeshFactory : public UFbxFactory, public FReimportHand
 	virtual EReimportResult::Type Reimport( UObject* Obj ) override;
 	virtual int32 GetPriority() const override;
 	// End FReimportHandler interface
+
+	// Begin UFactory Interface
+	virtual bool FactoryCanImport(const FString& Filename) override;
+	// End UFactory Interface
 };
 
 
diff --git a/Engine/Source/Editor/UnrealEd/Classes/Factories/ReimportFbxStaticMeshFactory.h b/Engine/Source/Editor/UnrealEd/Classes/Factories/ReimportFbxStaticMeshFactory.h
--- a/Engine/Source/Editor/UnrealEd/Classes/Factories/ReimportFbxStaticMeshFactory.h
+++ b/Engine/Source/Editor/UnrealEd/Classes/Factories/ReimportFbxStaticMeshFactory.h
@@ -19,6 +19,10 @@ class UReimportFbxStaticMeshFactory : public UFbxFactory, public FReimportHandle
 	virtual EReimportResult::Type Reimport( UObject* Obj ) override;
 	virtual int32 GetPriority() const override;
 	// End FReimportHandler interface
+
+	// Begin UFactory Interface
+	virtual bool FactoryCanImport(const FString& Filename) override;
+	// End UFactory Interface
 };
 
 
diff --git a/Engine/Source/Editor/UnrealEd/Private/Factories/EditorFactories.cpp b/Engine/Source/Editor/UnrealEd/Private/Factories/EditorFactories.cpp
--- a/Engine/Source/Editor/UnrealEd/Private/Factories/EditorFactories.cpp
+++ b/Engine/Source/Editor/UnrealEd/Private/Factories/EditorFactories.cpp
@@ -5547,6 +5547,12 @@ UReimportFbxStaticMeshFactory::UReimportFbxStaticMeshFactory(const FObjectInitia
 	ImportPriority = DefaultImportPriority - 1;
 }
 
+bool UReimportFbxStaticMeshFactory::FactoryCanImport(const FString& Filename)
+{
+	// Return false, we are a reimport only factory
+	return false;
+}
+
 bool UReimportFbxStaticMeshFactory::CanReimport( UObject* Obj, TArray<FString>& OutFilenames )
 {	
 	UStaticMesh* Mesh = Cast<UStaticMesh>(Obj);
@@ -5737,6 +5743,12 @@ UReimportFbxSkeletalMeshFactory::UReimportFbxSkeletalMeshFactory(const FObjectIn
 	bText = false;
 }
 
+bool UReimportFbxSkeletalMeshFactory::FactoryCanImport(const FString& Filename)
+{
+	// Return false, we are a reimport only factory
+	return false;
+}
+
 bool UReimportFbxSkeletalMeshFactory::CanReimport( UObject* Obj, TArray<FString>& OutFilenames )
 {	
 	USkeletalMesh* SkeletalMesh = Cast<USkeletalMesh>(Obj);
@@ -5909,6 +5921,12 @@ UReimportFbxAnimSequenceFactory::UReimportFbxAnimSequenceFactory(const FObjectIn
 	bText = false;
 }
 
+bool UReimportFbxAnimSequenceFactory::FactoryCanImport(const FString& Filename)
+{
+	// Return false, we are a reimport only factory
+	return false;
+}
+
 bool UReimportFbxAnimSequenceFactory::CanReimport( UObject* Obj, TArray<FString>& OutFilenames )
 {	
 	UAnimSequence* AnimSequence = Cast<UAnimSequence>(Obj);
diff --git a/Engine/Source/Editor/UnrealEd/Private/Fbx/FbxFactory.cpp b/Engine/Source/Editor/UnrealEd/Private/Fbx/FbxFactory.cpp
--- a/Engine/Source/Editor/UnrealEd/Private/Fbx/FbxFactory.cpp
+++ b/Engine/Source/Editor/UnrealEd/Private/Fbx/FbxFactory.cpp
@@ -581,6 +581,15 @@ void UFbxFactory::CleanUp()
 	}
 }
 
+bool UFbxFactory::FactoryCanImport(const FString& Filename)
+{
+	if(FPaths::GetExtension(Filename) == TEXT("fbx"))
+	{
+		return true;
+	}
+	return false;
+}
+
 UFbxImportUI::UFbxImportUI(const FObjectInitializer& ObjectInitializer)
 	: Super(ObjectInitializer)
 {
