Make FBlueprintCompileReinstancer::ReplaceInstancesOfClass unregister the old Component and re-register the new one if it was previously registered

diff --git a/Engine/Source/Editor/UnrealEd/Private/Kismet2/KismetReinstanceUtilities.cpp b/Engine/Source/Editor/UnrealEd/Private/Kismet2/KismetReinstanceUtilities.cpp
--- a/Engine/Source/Editor/UnrealEd/Private/Kismet2/KismetReinstanceUtilities.cpp
+++ b/Engine/Source/Editor/UnrealEd/Private/Kismet2/KismetReinstanceUtilities.cpp
@@ -1081,6 +1081,17 @@ void FBlueprintCompileReinstancer::ReplaceInstancesOfClass(UClass* OldClass, UCl
 					}
 				}
 
+				bool bWasRegistered = false;
+				if (bIsComponent)
+				{
+					UActorComponent* OldComponent = CastChecked<UActorComponent>(OldObject);
+					if (OldComponent->IsRegistered())
+					{
+						bWasRegistered = true;
+						OldComponent->UnregisterComponent();
+					}
+				}
+
 				OldObject->RemoveFromRoot();
 				OldObject->MarkPendingKill();
 
@@ -1088,7 +1099,7 @@ void FBlueprintCompileReinstancer::ReplaceInstancesOfClass(UClass* OldClass, UCl
 
 				if (bIsComponent)
 				{
-					UActorComponent* Component = Cast<UActorComponent>(NewUObject);
+					UActorComponent* Component = CastChecked<UActorComponent>(NewUObject);
 					AActor* OwningActor = Component->GetOwner();
 					if (OwningActor)
 					{
@@ -1110,6 +1121,11 @@ void FBlueprintCompileReinstancer::ReplaceInstancesOfClass(UClass* OldClass, UCl
 							OwnersToRerunConstructionScript.Add(OwningActor);
 						}
 					}
+
+					if (bWasRegistered)
+					{
+						Component->RegisterComponent();
+					}
 				}
 			}
 
