Fixed: Cannot place static mesh in location if foliage mesh is already placed in same location UE-15453

diff --git a/Engine/Source/Editor/UnrealEd/Private/LevelEditorViewport.cpp b/Engine/Source/Editor/UnrealEd/Private/LevelEditorViewport.cpp
--- a/Engine/Source/Editor/UnrealEd/Private/LevelEditorViewport.cpp
+++ b/Engine/Source/Editor/UnrealEd/Private/LevelEditorViewport.cpp
@@ -1244,7 +1244,7 @@ bool FLevelEditorViewportClient::DropObjectsAtCoordinates(int32 MouseX, int32 Mo
 		FSnappingUtils::SnapPointToGrid(GEditor->ClickLocation, FVector::ZeroVector);
 
 		EObjectFlags ObjectFlags = bCreateDropPreview ? RF_Transient : RF_Transactional;
-		if ( HitProxy == nullptr )
+		if (HitProxy == nullptr || HitProxy->IsA(HInstancedStaticMeshInstance::StaticGetType()))
 		{
 			bResult = DropObjectsOnBackground(Cursor, DroppedObjects, ObjectFlags, OutNewActors, bCreateDropPreview, SelectActors, FactoryToUse);
 		}
