Deprecated BrushComponent PrePivot.  Replacement is a generic pivot offset property on AActor which is used when positioning the pivot for selected actors in the editor.  Improved viewport click handling.

diff --git a/Engine/Source/Editor/GeometryMode/Private/GeometryModifiers.cpp b/Engine/Source/Editor/GeometryMode/Private/GeometryModifiers.cpp
--- a/Engine/Source/Editor/GeometryMode/Private/GeometryModifiers.cpp
+++ b/Engine/Source/Editor/GeometryMode/Private/GeometryModifiers.cpp
@@ -912,7 +912,7 @@ void UGeomModifier_Lathe::Apply( int32 InTotalSegments, int32 InSegments, EAxis:
 	ABrush* BuilderBrush = GeomMode->GetWorld()->GetDefaultBrush();
 
 	BuilderBrush->SetActorLocation(GeomMode->GetWidgetLocation(), false);
-	BuilderBrush->SetPrePivot(FVector::ZeroVector);
+	BuilderBrush->SetPivotOffset(FVector::ZeroVector);
 	BuilderBrush->SetFlags( RF_Transactional );
 	BuilderBrush->Brush->Polys->Element.Empty();
 
@@ -975,7 +975,7 @@ void UGeomModifier_Lathe::Apply( int32 InTotalSegments, int32 InSegments, EAxis:
 
 					for( int32 e = 0 ; e < WindingVerts->Num() ; ++e )
 					{
-						FVector vtx = (*WindingVerts)[e] - delta - BrushShape->GetPrePivot();
+						FVector vtx = (*WindingVerts)[e] - delta - BrushShape->GetPivotOffset();
 
 						vtx = RotationMatrix.TransformPosition( vtx );
 
@@ -1043,7 +1043,7 @@ void UGeomModifier_Lathe::Apply( int32 InTotalSegments, int32 InSegments, EAxis:
 
 					for( int32 v = 0 ; v < WindingVerts->Num() ; ++v )
 					{
-						Poly.Vertices.Add( (*WindingVerts)[v] - delta - BrushShape->GetPrePivot() );
+						Poly.Vertices.Add((*WindingVerts)[v] - delta - BrushShape->GetPivotOffset());
 					}
 
 					Poly.Finalize( BuilderBrush, 1 );
@@ -1089,7 +1089,7 @@ void UGeomModifier_Lathe::Apply( int32 InTotalSegments, int32 InSegments, EAxis:
 
 					for( int32 v = 0 ; v < WindingVerts->Num() ; ++v )
 					{
-						Poly.Vertices.Add( RotationMatrix.TransformPosition( (*WindingVerts)[v] - delta - BrushShape->GetPrePivot() ) );
+						Poly.Vertices.Add(RotationMatrix.TransformPosition((*WindingVerts)[v] - delta - BrushShape->GetPivotOffset()));
 					}
 
 					Poly.Finalize( BuilderBrush, 1 );
@@ -1238,7 +1238,7 @@ void UGeomModifier_Pen::Apply()
 		FlushRenderingCommands();
 		
 		ResultingBrush->SetActorLocation(BaseLocation, false);
-		ResultingBrush->SetPrePivot(FVector::ZeroVector);
+		ResultingBrush->SetPivotOffset(FVector::ZeroVector);
 		ResultingBrush->SetFlags( RF_Transactional );
 		ResultingBrush->Brush->Polys->Element.Empty();
 
@@ -1688,7 +1688,7 @@ namespace GeometryClipping {
 static void BuildGiantAlignedBrush( ABrush& OutGiantBrush, const FPlane& InPlane )
 {
 	OutGiantBrush.SetActorLocation(FVector::ZeroVector, false);
-	OutGiantBrush.SetPrePivot(FVector::ZeroVector);
+	OutGiantBrush.SetPivotOffset(FVector::ZeroVector);
 
 	verify( OutGiantBrush.Brush );
 	verify( OutGiantBrush.Brush->Polys );
diff --git a/Engine/Source/Editor/LevelEditor/Private/LevelEditorActions.cpp b/Engine/Source/Editor/LevelEditor/Private/LevelEditorActions.cpp
--- a/Engine/Source/Editor/LevelEditor/Private/LevelEditorActions.cpp
+++ b/Engine/Source/Editor/LevelEditor/Private/LevelEditorActions.cpp
@@ -2973,8 +2973,8 @@ void FLevelEditorCommands::RegisterCommands()
 	UI_COMMAND( HoldToEnableVertexSnapping, "Hold to Enable Vertex Snapping", "When the key binding is pressed and held vertex snapping will be enabled", EUserInterfaceActionType::ToggleButton, FInputChord(EKeys::V) );
 
 	//@ todo Slate better tooltips for pivot options
-	UI_COMMAND( SavePivotToPrePivot, "Save Pivot to Pre-Pivot", "Saves the pivot to the pre-pivot", EUserInterfaceActionType::Button, FInputChord() );
-	UI_COMMAND( ResetPivot, "Reset Pivot", "Resets the pivot", EUserInterfaceActionType::Button, FInputChord() );
+	UI_COMMAND( SavePivotToPrePivot, "Set as Pivot Offset", "Sets the current pivot location as the pivot offset for this actor", EUserInterfaceActionType::Button, FInputChord() );
+	UI_COMMAND( ResetPivot, "Reset Pivot Offset", "Resets the pivot offset for this actor", EUserInterfaceActionType::Button, FInputChord() );
 	UI_COMMAND( ResetPrePivot, "Reset Pre-Pivot", "Resets the pre-pivot", EUserInterfaceActionType::Button, FInputChord() );
 	UI_COMMAND( MovePivotHere, "Move Here", "Moves the pivot to the clicked location", EUserInterfaceActionType::Button, FInputChord() );
 	UI_COMMAND( MovePivotHereSnapped, "Move Here (Snapped)", "Moves the pivot to the nearest grid point to the clicked location", EUserInterfaceActionType::Button, FInputChord() );
diff --git a/Engine/Source/Editor/LevelEditor/Private/LevelEditorContextMenu.cpp b/Engine/Source/Editor/LevelEditor/Private/LevelEditorContextMenu.cpp
--- a/Engine/Source/Editor/LevelEditor/Private/LevelEditorContextMenu.cpp
+++ b/Engine/Source/Editor/LevelEditor/Private/LevelEditorContextMenu.cpp
@@ -410,23 +410,11 @@ void FLevelEditorContextMenu::FillMenu( FMenuBuilder& MenuBuilder, TWeakPtr<SLev
 					FNewMenuDelegate::CreateStatic(&FLevelEditorContextMenuImpl::FillTransformMenu));
 			}
 
-			// @todo UE4: The current pivot options only work for brushes
-			if (SelectionInfo.bHaveBrush)
-			{
-				// You can only move the pivot in ortho viewports, but you can reset it in any viewport
-				if (GCurrentLevelEditingViewportClient->ViewportType != LVT_Perspective)
-				{
-					// Add a sub-menu for "Pivot"
-					MenuBuilder.AddSubMenu(
-						LOCTEXT("PivotSubMenu", "Pivot"),
-						LOCTEXT("PivotSubMenu_ToolTip", "Actor pivoting utils"),
-						FNewMenuDelegate::CreateStatic(&FLevelEditorContextMenuImpl::FillPivotMenu));
-				}
-				else
-				{
-					MenuBuilder.AddMenuEntry(FLevelEditorCommands::Get().ResetPivot);
-				}
-			}
+			// Add a sub-menu for "Pivot"
+			MenuBuilder.AddSubMenu(
+				LOCTEXT("PivotSubMenu", "Pivot"),
+				LOCTEXT("PivotSubMenu_ToolTip", "Actor pivoting utils"),
+				FNewMenuDelegate::CreateStatic(&FLevelEditorContextMenuImpl::FillPivotMenu));
 		}
 		MenuBuilder.EndSection();
 
@@ -928,7 +916,6 @@ void FLevelEditorContextMenuImpl::FillPivotMenu( FMenuBuilder& MenuBuilder )
 	MenuBuilder.BeginSection("SaveResetPivot");
 	{
 		MenuBuilder.AddMenuEntry( FLevelEditorCommands::Get().SavePivotToPrePivot );
-		MenuBuilder.AddMenuEntry( FLevelEditorCommands::Get().ResetPrePivot );
 		MenuBuilder.AddMenuEntry( FLevelEditorCommands::Get().ResetPivot );
 	}
 	MenuBuilder.EndSection();
diff --git a/Engine/Source/Editor/UnrealEd/Classes/Editor/EditorEngine.h b/Engine/Source/Editor/UnrealEd/Classes/Editor/EditorEngine.h
--- a/Engine/Source/Editor/UnrealEd/Classes/Editor/EditorEngine.h
+++ b/Engine/Source/Editor/UnrealEd/Classes/Editor/EditorEngine.h
@@ -1344,6 +1344,11 @@ public:
 	 */
 	virtual void SelectNone(bool bNoteSelectionChange, bool bDeselectBSPSurfs, bool WarnAboutManyActors=true) {}
 
+	/**
+	 * Deselect all surfaces.
+	 */
+	virtual void DeselectAllSurfaces() {}
+
 	// Bsp Poly selection virtuals from EditorCsg.cpp.
 	virtual void polySelectAll ( UModel* Model );
 	virtual void polySelectMatchingGroups( UModel* Model );
diff --git a/Engine/Source/Editor/UnrealEd/Classes/Editor/UnrealEdEngine.h b/Engine/Source/Editor/UnrealEd/Classes/Editor/UnrealEdEngine.h
--- a/Engine/Source/Editor/UnrealEd/Classes/Editor/UnrealEdEngine.h
+++ b/Engine/Source/Editor/UnrealEd/Classes/Editor/UnrealEdEngine.h
@@ -186,6 +186,7 @@ public:
 	virtual void SelectComponent(class UActorComponent* Component, bool bInSelected, bool bNotify, bool bSelectEvenIfHidden = false) override;
 	virtual void SelectBSPSurf(UModel* InModel, int32 iSurf, bool bSelected, bool bNoteSelectionChange) override;
 	virtual void SelectNone(bool bNoteSelectionChange, bool bDeselectBSPSurfs, bool WarnAboutManyActors=true) override;
+	virtual void DeselectAllSurfaces() override;
 	virtual void NoteSelectionChange() override;
 	virtual void NoteActorMovement() override;
 	virtual void FinishAllSnaps() override;
diff --git a/Engine/Source/Editor/UnrealEd/Private/BSPOps.cpp b/Engine/Source/Editor/UnrealEd/Private/BSPOps.cpp
--- a/Engine/Source/Editor/UnrealEd/Private/BSPOps.cpp
+++ b/Engine/Source/Editor/UnrealEd/Private/BSPOps.cpp
@@ -1259,9 +1259,9 @@ void FBSPOps::RotateBrushVerts(ABrush* Brush, const FRotator& Rotation, bool bCl
 			const FRotationMatrix RotMatrix( Rotation );
 			for( int32 vertex = 0 ; vertex < Poly->Vertices.Num() ; vertex++ )
 			{
-				Poly->Vertices[vertex] = Brush->GetPrePivot() + RotMatrix.TransformVector( Poly->Vertices[vertex] - Brush->GetPrePivot() );
+				Poly->Vertices[vertex] = Brush->GetPivotOffset() + RotMatrix.TransformVector(Poly->Vertices[vertex] - Brush->GetPivotOffset());
 			}
-			Poly->Base = Brush->GetPrePivot() + RotMatrix.TransformVector( Poly->Base - Brush->GetPrePivot() );
+			Poly->Base = Brush->GetPivotOffset() + RotMatrix.TransformVector(Poly->Base - Brush->GetPivotOffset());
 
 			// Rotate the texture vectors.
 			Poly->TextureU = RotMatrix.TransformVector( Poly->TextureU );
diff --git a/Engine/Source/Editor/UnrealEd/Private/EditorBrushBuilder.cpp b/Engine/Source/Editor/UnrealEd/Private/EditorBrushBuilder.cpp
--- a/Engine/Source/Editor/UnrealEd/Private/EditorBrushBuilder.cpp
+++ b/Engine/Source/Editor/UnrealEd/Private/EditorBrushBuilder.cpp
@@ -67,7 +67,7 @@ bool UEditorBrushBuilder::EndBrush( UWorld* InWorld, ABrush* InBrush )
 	FRotator Temp(0.0f,0.0f,0.0f);
 	FSnappingUtils::SnapToBSPVertex( Location, FVector::ZeroVector, Temp );
 	BuilderBrush->SetActorLocation(Location, false);
-	BuilderBrush->SetPrePivot( FVector::ZeroVector );
+	BuilderBrush->SetPivotOffset( FVector::ZeroVector );
 
 	// Try and maintain the materials assigned to the surfaces. 
 	TArray<FPoly> CachedPolys;
diff --git a/Engine/Source/Editor/UnrealEd/Private/EditorBsp.cpp b/Engine/Source/Editor/UnrealEd/Private/EditorBsp.cpp
--- a/Engine/Source/Editor/UnrealEd/Private/EditorBsp.cpp
+++ b/Engine/Source/Editor/UnrealEd/Private/EditorBsp.cpp
@@ -1140,7 +1140,6 @@ int UEditorEngine::bspBrushCSG
 
 	UMaterialInterface* SelectedMaterialInstance = GetSelectedObjects()->GetTop<UMaterialInterface>();
 
-	const FVector PrePivot = Actor->GetPrePivot();
 	const FVector Scale = Actor->GetActorScale();
 	const FRotator Rotation = Actor->GetActorRotation();
 	const FVector Location = Actor->GetActorLocation();
@@ -1149,7 +1148,6 @@ int UEditorEngine::bspBrushCSG
 	Brush->OwnerLocationWhenLastBuilt = Location;
 	Brush->OwnerRotationWhenLastBuilt = Rotation;
 	Brush->OwnerScaleWhenLastBuilt = Scale;
-	Brush->OwnerPrepivotWhenLastBuilt = PrePivot;
 	Brush->bCachedOwnerTransformValid = true;
 
 	for( i=0; i<Brush->Polys->Element.Num(); i++ )
@@ -1182,9 +1180,9 @@ int UEditorEngine::bspBrushCSG
 			DestEdPoly.iLink = i;
 
 		// Transform it.
-		DestEdPoly.Scale( PrePivot, Scale );
-		DestEdPoly.Rotate( PrePivot, Rotation );
-		DestEdPoly.Transform( PrePivot, Location );
+		DestEdPoly.Scale( Scale );
+		DestEdPoly.Rotate( Rotation );
+		DestEdPoly.Transform( Location );
 
 		// Add poly to the temp model.
 		new(TempModel->Polys->Element)FPoly( DestEdPoly );
@@ -1303,7 +1301,9 @@ int UEditorEngine::bspBrushCSG
 		for( i=0; i<Brush->Polys->Element.Num(); i++ )
 		{
 			FPoly *DestEdPoly = &Brush->Polys->Element[i];
-			DestEdPoly->Transform( Actor->GetActorLocation(), Actor->GetPrePivot() );
+			DestEdPoly->Transform(-Location);
+			DestEdPoly->Rotate(Rotation.GetInverse());
+			DestEdPoly->Scale(FVector(1.0f) / Scale);
 			DestEdPoly->Fix();
 			DestEdPoly->Actor		= NULL;
 			DestEdPoly->iBrushPoly	= i;
diff --git a/Engine/Source/Editor/UnrealEd/Private/EditorCsg.cpp b/Engine/Source/Editor/UnrealEd/Private/EditorCsg.cpp
--- a/Engine/Source/Editor/UnrealEd/Private/EditorCsg.cpp
+++ b/Engine/Source/Editor/UnrealEd/Private/EditorCsg.cpp
@@ -373,7 +373,6 @@ void UEditorEngine::polyUpdateMaster
 		// Use transform cached when the geometry was last built, in case the current Actor transform has changed since then
 		// (e.g. because Auto Update BSP is disabled)
 		ActorLocation = Brush->OwnerLocationWhenLastBuilt;
-		ActorPrePivot = Brush->OwnerPrepivotWhenLastBuilt;
 		ActorScale = Brush->OwnerScaleWhenLastBuilt;
 		ActorRotation = Brush->OwnerRotationWhenLastBuilt;
 	}
@@ -381,7 +380,6 @@ void UEditorEngine::polyUpdateMaster
 	{
 		// No cached owner transform, so use the current one
 		ActorLocation = Actor->GetActorLocation();
-		ActorPrePivot = Actor->GetPrePivot();
 		ActorScale = Actor->GetActorScale();
 		ActorRotation = Actor->GetActorRotation();
 	}
@@ -398,7 +396,7 @@ void UEditorEngine::polyUpdateMaster
 
 			if (UpdateTexCoords)
 			{
-				MasterEdPoly.Base = RotationMatrix.InverseTransformVector(Model->Points[Surf.pBase] - ActorLocation) / ActorScale + ActorPrePivot;
+				MasterEdPoly.Base = RotationMatrix.InverseTransformVector(Model->Points[Surf.pBase] - ActorLocation) / ActorScale;
 				MasterEdPoly.TextureU = RotationMatrix.InverseTransformVector(Model->Vectors[Surf.vTextureU]) * ActorScale;
 				MasterEdPoly.TextureV = RotationMatrix.InverseTransformVector(Model->Vectors[Surf.vTextureV]) * ActorScale;
 			}
diff --git a/Engine/Source/Editor/UnrealEd/Private/EditorSelectUtils.cpp b/Engine/Source/Editor/UnrealEd/Private/EditorSelectUtils.cpp
--- a/Engine/Source/Editor/UnrealEd/Private/EditorSelectUtils.cpp
+++ b/Engine/Source/Editor/UnrealEd/Private/EditorSelectUtils.cpp
@@ -149,9 +149,10 @@ void UUnrealEdEngine::SetPivot( FVector NewPivot, bool bSnapPivotToGrid, bool bI
 	if( !bIgnoreAxis )
 	{
 		// Don't stomp on orthonormal axis.
-		if( NewPivot.X==0 ) NewPivot.X=EditorModeTools.PivotLocation.X;
-		if( NewPivot.Y==0 ) NewPivot.Y=EditorModeTools.PivotLocation.Y;
-		if( NewPivot.Z==0 ) NewPivot.Z=EditorModeTools.PivotLocation.Z;
+		// TODO: this breaks if there is genuinely a need to set the pivot to a coordinate containing a zero component
+ 		if( NewPivot.X==0 ) NewPivot.X=EditorModeTools.PivotLocation.X;
+ 		if( NewPivot.Y==0 ) NewPivot.Y=EditorModeTools.PivotLocation.Y;
+ 		if( NewPivot.Z==0 ) NewPivot.Z=EditorModeTools.PivotLocation.Z;
 	}
 
 	// Set the pivot.
@@ -335,8 +336,8 @@ void UUnrealEdEngine::UpdatePivotLocationForSelection( bool bOnChange )
 		FEditorModeTools& Tools = GLevelEditorModeTools();
 		if( Tools.IsModeActive(FBuiltinEditorModes::EM_Geometry) == false || bOnChange == true )
 		{
-			// Set pivot point to the actor's location
-			FVector PivotPoint = SingleActor->GetActorLocation();
+			// Set pivot point to the actor's location, accounting for any set pivot offset
+			FVector PivotPoint = SingleActor->GetTransform().TransformPosition(SingleActor->GetPivotOffset());
 
 			// If grouping is active, see if this actor is part of a locked group and use that pivot instead
 			if(GEditor->bGroupingActive)
@@ -568,6 +569,7 @@ void UUnrealEdEngine::SelectActor(AActor* Actor, bool bInSelected, bool bNotify,
 					GetSelectedComponents()->Modify();
 				}
 
+				GetSelectedComponents()->BeginBatchSelectOperation();
 				for (UActorComponent* Component : Actor->GetComponents())
 				{
 					GetSelectedComponents()->Deselect( Component );
@@ -576,6 +578,7 @@ void UUnrealEdEngine::SelectActor(AActor* Actor, bool bInSelected, bool bNotify,
 					auto SceneComponent = Cast<USceneComponent>(Component);
 					FComponentEditorUtils::BindComponentSelectionOverride(SceneComponent, false);
 				}
+				GetSelectedComponents()->EndBatchSelectOperation(false);
 			}
 			else
 			{
@@ -725,6 +728,27 @@ static uint32 DeselectAllSurfacesForLevel(ULevel* Level)
 }
 
 
+void UUnrealEdEngine::DeselectAllSurfaces()
+{
+	UWorld* World = GWorld;
+	if (World)
+	{
+		DeselectAllSurfacesForLevel(World->PersistentLevel);
+		for (int32 LevelIndex = 0; LevelIndex < World->StreamingLevels.Num(); ++LevelIndex)
+		{
+			ULevelStreaming* StreamingLevel = World->StreamingLevels[LevelIndex];
+			if (StreamingLevel)
+			{
+				ULevel* Level = StreamingLevel->GetLoadedLevel();
+				if (Level != NULL)
+				{
+					DeselectAllSurfacesForLevel(Level);
+				}
+			}
+		}
+	}
+}
+
 void UUnrealEdEngine::SelectNone(bool bNoteSelectionChange, bool bDeselectBSPSurfs, bool WarnAboutManyActors)
 {
 	if( GEdSelectionLock )
@@ -807,8 +831,6 @@ void UUnrealEdEngine::SelectNone(bool bNoteSelectionChange, bool bDeselectBSPSur
 	//prevents clicking on background multiple times spamming selection changes
 	if (ActorsToDeselect.Num() || NumDeselectSurfaces)
 	{
-		GetSelectedActors()->DeselectAll();
-
 		if( bNoteSelectionChange )
 		{
 			NoteSelectionChange();
diff --git a/Engine/Source/Editor/UnrealEd/Private/EditorServer.cpp b/Engine/Source/Editor/UnrealEd/Private/EditorServer.cpp
--- a/Engine/Source/Editor/UnrealEd/Private/EditorServer.cpp
+++ b/Engine/Source/Editor/UnrealEd/Private/EditorServer.cpp
@@ -571,13 +571,13 @@ bool UEditorEngine::Exec_Brush( UWorld* InWorld, const TCHAR* Str, FOutputDevice
 			{
 				DefaultBrush->Brush->Modify();
 				SnapLocation = DefaultBrush->GetActorLocation();
-				PrePivot = DefaultBrush->GetPrePivot();
+				PrePivot = DefaultBrush->GetPivotOffset();
 			}
 			
 			FSnappingUtils::SnapToBSPVertex( SnapLocation, FVector::ZeroVector, Temp );
 
 			WorldBrush->SetActorLocation(SnapLocation - PrePivot, false);
-			WorldBrush->SetPrePivot( FVector::ZeroVector );
+			WorldBrush->SetPivotOffset( FVector::ZeroVector );
 			WorldBrush->Brush->Polys->Element.Empty();
 			UPolysFactory* It = NewObject<UPolysFactory>();
 			It->FactoryCreateText( UPolys::StaticClass(), WorldBrush->Brush->Polys->GetOuter(), *WorldBrush->Brush->Polys->GetName(), RF_NoFlags, WorldBrush->Brush->Polys, TEXT("t3d"), GStream, GStream+FCString::Strlen(GStream), GWarn );
@@ -628,11 +628,11 @@ bool UEditorEngine::Exec_Brush( UWorld* InWorld, const TCHAR* Str, FOutputDevice
 
 						Poly->TextureU *= InvScale;
 						Poly->TextureV *= InvScale;
-						Poly->Base = ((Poly->Base - Brush->GetPrePivot()) * Scale) + Brush->GetPrePivot();
+						Poly->Base = ((Poly->Base - Brush->GetPivotOffset()) * Scale) + Brush->GetPivotOffset();
 
 						for( int32 vtx = 0 ; vtx < Poly->Vertices.Num() ; vtx++ )
 						{
-							Poly->Vertices[vtx] = ((Poly->Vertices[vtx] - Brush->GetPrePivot()) * Scale) + Brush->GetPrePivot();
+							Poly->Vertices[vtx] = ((Poly->Vertices[vtx] - Brush->GetPivotOffset()) * Scale) + Brush->GetPivotOffset();
 						}
 
 						Poly->CalcNormal();
@@ -3811,11 +3811,11 @@ bool UEditorEngine::Map_Scale( UWorld* InWorld, const TCHAR* Str, FOutputDevice&
 
 					Poly->TextureU /= Factor;
 					Poly->TextureV /= Factor;
-					Poly->Base = ((Poly->Base - Brush->GetPrePivot()) * Factor) + Brush->GetPrePivot();
+					Poly->Base = ((Poly->Base - Brush->GetPivotOffset()) * Factor) + Brush->GetPivotOffset();
 
 					for( int32 vtx = 0 ; vtx < Poly->Vertices.Num() ; vtx++ )
 					{
-						Poly->Vertices[vtx] = ((Poly->Vertices[vtx] - Brush->GetPrePivot()) * Factor) + Brush->GetPrePivot();
+						Poly->Vertices[vtx] = ((Poly->Vertices[vtx] - Brush->GetPivotOffset()) * Factor) + Brush->GetPivotOffset();
 					}
 
 					Poly->CalcNormal();
diff --git a/Engine/Source/Editor/UnrealEd/Private/Factories/EditorFactories.cpp b/Engine/Source/Editor/UnrealEd/Private/Factories/EditorFactories.cpp
--- a/Engine/Source/Editor/UnrealEd/Private/Factories/EditorFactories.cpp
+++ b/Engine/Source/Editor/UnrealEd/Private/Factories/EditorFactories.cpp
@@ -1439,7 +1439,7 @@ UObject* UModelFactory::FactoryCreateText
 			{
 				FVector TempPrePivot(0.f);
 				GetFVECTOR 	(StrPtr,TempPrePivot);
-				TempOwner->SetPrePivot(TempPrePivot);
+				TempOwner->SetPivotOffset(TempPrePivot);
 			}
 			else if (FParse::Command(&StrPtr,TEXT("LOCATION"	))) 
 			{
diff --git a/Engine/Source/Editor/UnrealEd/Private/LevelEditorViewport.cpp b/Engine/Source/Editor/UnrealEd/Private/LevelEditorViewport.cpp
--- a/Engine/Source/Editor/UnrealEd/Private/LevelEditorViewport.cpp
+++ b/Engine/Source/Editor/UnrealEd/Private/LevelEditorViewport.cpp
@@ -1530,6 +1530,7 @@ FLevelEditorViewportClient::FLevelEditorViewportClient(const TSharedPtr<SLevelVi
 	, bDuplicateOnNextDrag( false )
 	, bDuplicateActorsInProgress( false )
 	, bIsTrackingBrushModification( false )
+	, bOnlyMovedPivot(false)
 	, bLockedCameraView(true)
 	, bReceivedFocusRecently(false)
 	, SpriteCategoryVisibility()
@@ -1855,9 +1856,6 @@ void FLevelEditorViewportClient::ReceivedFocus(FViewport* InViewport)
 //
 void FLevelEditorViewportClient::ProcessClick(FSceneView& View, HHitProxy* HitProxy, FKey Key, EInputEvent Event, uint32 HitX, uint32 HitY)
 {
-	// We clicked, allow the pivot to reposition itself.
-	bPivotMovedIndependently = false;
-
 	static FName ProcessClickTrace = FName(TEXT("ProcessClickTrace"));
 
 	const FViewportClick Click(&View,this,Key,Event,HitX,HitY);
@@ -1936,6 +1934,9 @@ void FLevelEditorViewportClient::ProcessClick(FSceneView& View, HHitProxy* HitPr
 			{
 				ClickHandlers::ClickActor(this, ActorHitProxy->Actor, Click, true);
 			}
+
+			// We clicked an actor, allow the pivot to reposition itself.
+			bPivotMovedIndependently = false;
 		}
 		else if (HitProxy->IsA(HInstancedStaticMeshInstance::StaticGetType()))
 		{
@@ -2336,6 +2337,7 @@ bool FLevelEditorViewportClient::InputWidgetDelta(FViewport* Viewport, EAxisList
 				{
 					FSnappingUtils::SnapDragLocationToNearestVertex( ModeTools->PivotLocation, Drag, this );
 					bPivotMovedIndependently = true;
+					bOnlyMovedPivot = true;
 				}
 
 				ModeTools->PivotLocation += Drag;
@@ -2540,6 +2542,8 @@ void FLevelEditorViewportClient::TrackingStarted( const FInputEventState& InInpu
 		}
 	}
 
+	bOnlyMovedPivot = false;
+
 	const bool bIsDraggingComponents = GEditor->GetSelectedComponentCount() > 0;
 	PreDragActorTransforms.Empty();
 	if (bIsDraggingComponents)
@@ -2682,7 +2686,7 @@ void FLevelEditorViewportClient::TrackingStopped()
 	// Finish tracking a brush transform and update the Bsp
 	if (bIsTrackingBrushModification)
 	{
-		bDidAnythingActuallyChange = HaveSelectedObjectsBeenChanged();
+		bDidAnythingActuallyChange = HaveSelectedObjectsBeenChanged() && !bOnlyMovedPivot;
 
 		bIsTrackingBrushModification = false;
 		if ( bDidAnythingActuallyChange && bWidgetAxisControlledByDrag )
diff --git a/Engine/Source/Editor/UnrealEd/Private/LevelViewportClickHandlers.cpp b/Engine/Source/Editor/UnrealEd/Private/LevelViewportClickHandlers.cpp
--- a/Engine/Source/Editor/UnrealEd/Private/LevelViewportClickHandlers.cpp
+++ b/Engine/Source/Editor/UnrealEd/Private/LevelViewportClickHandlers.cpp
@@ -864,9 +864,13 @@ namespace ClickHandlers
 			check( Model );
 
 			bool bNeedViewportRefresh = false;
+			bool bIsActorAlreadySelected = Surf.Actor && Surf.Actor->IsSelected();
 			{
 				const FScopedTransaction Transaction( NSLOCTEXT("UnrealEd", "SelectSurfaces", "Select Surfaces") );
 
+				USelection* SelectedActors = GEditor->GetSelectedActors();
+				SelectedActors->BeginBatchSelectOperation();
+
 				// We only need to unselect surfaces if the surface the user clicked on was not already selected
 				if( !( Surf.PolyFlags & PF_Selected ) )
 				{
@@ -878,7 +882,13 @@ namespace ClickHandlers
 				Model->ModifySurf( iSurf, false );
 				Surf.PolyFlags |= PF_Selected;
 
-				GEditor->NoteSelectionChange();
+				GEditor->SelectActor(Surf.Actor, true, false);
+				SelectedActors->EndBatchSelectOperation(false);
+
+				if (!bIsActorAlreadySelected)
+				{
+					GEditor->NoteSelectionChange();
+				}
 			}
 
 			if( bNeedViewportRefresh )
@@ -912,7 +922,11 @@ namespace ClickHandlers
 		{	
 			const FScopedTransaction Transaction( NSLOCTEXT("UnrealEd", "SelectBrushSurface", "Select Brush Surface") );
 			bool bDeselectAlreadyHandled = false;
-			
+			bool bIsActorAlreadySelected = Surf.Actor && Surf.Actor->IsSelected();
+
+			USelection* SelectedActors = GEditor->GetSelectedActors();
+			SelectedActors->BeginBatchSelectOperation();
+
 			// We are going to handle the notification ourselves
 			const bool bNotify = false;
 			if(GetDefault<ULevelEditorViewportSettings>()->bClickBSPSelectsBrush)
@@ -927,7 +941,6 @@ namespace ClickHandlers
 						bDeselectAlreadyHandled = true;
 					}
 					// If the builder brush is selected, first deselect it.
-					USelection* SelectedActors = GEditor->GetSelectedActors();
 					for(FSelectionIterator It(*SelectedActors); It; ++It)
 					{
 						ABrush* Brush = Cast<ABrush>(*It);
@@ -955,12 +968,16 @@ namespace ClickHandlers
 				if (!Model->HasSelectedSurfaces())
 				{
 					GEditor->SelectActor(Surf.Actor, false, bNotify);
+					bIsActorAlreadySelected = false;
 				}
 			}
 
+			SelectedActors->EndBatchSelectOperation(false);
 
-
-			GEditor->NoteSelectionChange();
+			if (!bIsActorAlreadySelected)
+			{
+				GEditor->NoteSelectionChange();
+			}
 		}
 	}
 
diff --git a/Engine/Source/Editor/UnrealEd/Private/StaticMeshEdit.cpp b/Engine/Source/Editor/UnrealEd/Private/StaticMeshEdit.cpp
--- a/Engine/Source/Editor/UnrealEd/Private/StaticMeshEdit.cpp
+++ b/Engine/Source/Editor/UnrealEd/Private/StaticMeshEdit.cpp
@@ -620,7 +620,7 @@ void GetBrushMesh(ABrush* Brush,UModel* Model,struct FRawMesh& OutMesh,TArray<UM
 
 		// Cache the texture coordinate system for this polygon.
 
-		FVector	TextureBase = Polygon.Base - (Brush ? Brush->GetPrePivot() : FVector::ZeroVector),
+		FVector	TextureBase = Polygon.Base - (Brush ? Brush->GetPivotOffset() : FVector::ZeroVector),
 				TextureX = Polygon.TextureU / UModel::GetGlobalBSPTexelScale(),
 				TextureY = Polygon.TextureV / UModel::GetGlobalBSPTexelScale();
 
diff --git a/Engine/Source/Editor/UnrealEd/Private/TexAlignTools.cpp b/Engine/Source/Editor/UnrealEd/Private/TexAlignTools.cpp
--- a/Engine/Source/Editor/UnrealEd/Private/TexAlignTools.cpp
+++ b/Engine/Source/Editor/UnrealEd/Private/TexAlignTools.cpp
@@ -220,7 +220,7 @@ void UTexAlignerDefault::AlignSurf( ETexAlign InTexAlignType, UModel* InModel, F
 	InPoly->TextureV *= VTile;
 
 	ABrush* Actor = InSurfIdx->Surf->Actor;
-	const FVector PrePivot = Actor->GetPrePivot();
+	const FVector PrePivot = Actor->GetPivotOffset();
 	const FVector Location = Actor->GetActorLocation();
 	const FRotator Rotation = Actor->GetActorRotation();
 	const FVector Scale = Actor->GetActorScale();
diff --git a/Engine/Source/Editor/UnrealEd/Private/UnrealEdSrv.cpp b/Engine/Source/Editor/UnrealEd/Private/UnrealEdSrv.cpp
--- a/Engine/Source/Editor/UnrealEd/Private/UnrealEdSrv.cpp
+++ b/Engine/Source/Editor/UnrealEd/Private/UnrealEdSrv.cpp
@@ -937,11 +937,11 @@ bool UUnrealEdEngine::Exec( UWorld* InWorld, const TCHAR* Stream, FOutputDevice&
 
 								Poly->TextureU /= ScaleVec;
 								Poly->TextureV /= ScaleVec;
-								Poly->Base = ((Poly->Base - Brush->GetPrePivot()) * ScaleVec) + Brush->GetPrePivot();
+								Poly->Base = ((Poly->Base - Brush->GetPivotOffset()) * ScaleVec) + Brush->GetPivotOffset();
 
 								for( int32 vtx = 0 ; vtx < Poly->Vertices.Num() ; vtx++ )
 								{
-									Poly->Vertices[vtx] = ((Poly->Vertices[vtx] - Brush->GetPrePivot()) * ScaleVec) + Brush->GetPrePivot();
+									Poly->Vertices[vtx] = ((Poly->Vertices[vtx] - Brush->GetPivotOffset()) * ScaleVec) + Brush->GetPivotOffset();
 
 									// "Then snap the vertices new positions by the specified Snap amount"
 									if ( bSnap )
@@ -2591,7 +2591,7 @@ bool UUnrealEdEngine::Exec_Actor( UWorld* InWorld, const TCHAR* Str, FOutputDevi
 		FScopedLevelDirtied				LevelDirtyCallback;
 		FScopedActorPropertiesChange	ActorPropertiesChangeCallback;
 
-		// Bakes the current pivot position into all selected brushes as their PrePivot
+		// Bakes the current pivot position into all selected actors
 
 		FEditorModeTools& EditorModeTools = GLevelEditorModeTools();
 
@@ -2602,16 +2602,9 @@ bool UUnrealEdEngine::Exec_Actor( UWorld* InWorld, const TCHAR* Str, FOutputDevi
 
 			FVector Delta( EditorModeTools.PivotLocation - Actor->GetActorLocation() );
 
-			ABrush* Brush = Cast<ABrush>(Actor);
-			if( Brush )
-			{
-				Brush->Modify();
-
-				Brush->SetActorLocation(Actor->GetActorLocation() + Delta, false);
-				Brush->SetPrePivot(Brush->GetPrePivot() + Delta);
-
-				Brush->PostEditMove( true );
-			}
+			Actor->Modify();
+			Actor->SetPivotOffset(Actor->GetPivotOffset() + Actor->GetTransform().InverseTransformVector(Delta));
+			Actor->PostEditMove(true);
 		}
 
 		GUnrealEd->NoteSelectionChange();
@@ -2621,7 +2614,7 @@ bool UUnrealEdEngine::Exec_Actor( UWorld* InWorld, const TCHAR* Str, FOutputDevi
 		FScopedLevelDirtied		LevelDirtyCallback;
 		FScopedActorPropertiesChange	ActorPropertiesChangeCallback;
 
-		// Resets the PrePivot of the selected brushes to 0,0,0 while leaving them in the same world location.
+		// Resets the PrePivot of the selected actors to 0,0,0 while leaving them in the same world location.
 
 		FEditorModeTools& EditorModeTools = GLevelEditorModeTools();
 
@@ -2630,18 +2623,9 @@ bool UUnrealEdEngine::Exec_Actor( UWorld* InWorld, const TCHAR* Str, FOutputDevi
 			AActor* Actor = static_cast<AActor*>( *It );
 			checkSlow( Actor->IsA(AActor::StaticClass()) );
 
-			ABrush* Brush = Cast<ABrush>(Actor);
-			if( Brush )
-			{
-				Brush->Modify();
-
-				FVector Delta = Brush->GetPrePivot();
-
-				Brush->SetActorLocation(Actor->GetActorLocation() - Delta, false);
-				Brush->SetPrePivot(FVector::ZeroVector);
-
-				Brush->PostEditMove( true );
-			}
+			Actor->Modify();
+			Actor->SetPivotOffset(FVector::ZeroVector);
+			Actor->PostEditMove(true);
 		}
 
 		GUnrealEd->NoteSelectionChange();
@@ -2703,12 +2687,9 @@ bool UUnrealEdEngine::Exec_Actor( UWorld* InWorld, const TCHAR* Str, FOutputDevi
 				{
 					Actor->SetActorLocation(FVector::ZeroVector, false);
 				}
-				ABrush* Brush = Cast< ABrush >( Actor );
-				if( bPivot && Brush )
+				if( bPivot )
 				{
-					Brush->SetActorLocation(Brush->GetActorLocation() - Brush->GetPrePivot(), false);
-					Brush->SetPrePivot(FVector::ZeroVector);
-					Brush->PostEditChange();
+					Actor->SetPivotOffset(FVector::ZeroVector);
 				}
 
 				if( bScale && Actor->GetRootComponent() != NULL ) 
diff --git a/Engine/Source/Editor/UnrealEd/Public/LevelEditorViewport.h b/Engine/Source/Editor/UnrealEd/Public/LevelEditorViewport.h
--- a/Engine/Source/Editor/UnrealEd/Public/LevelEditorViewport.h
+++ b/Engine/Source/Editor/UnrealEd/Public/LevelEditorViewport.h
@@ -670,6 +670,11 @@ public:
 	 */
 	bool					bIsTrackingBrushModification;
 
+	/**
+	 * true if only the pivot position has been moved
+	 */
+	bool					bOnlyMovedPivot;
+
 	/** True if this viewport is to change its view (aspect ratio, post processing, FOV etc) to match that of the currently locked camera, if applicable */
 	bool					bLockedCameraView;
 
diff --git a/Engine/Source/Runtime/Engine/Classes/Components/BrushComponent.h b/Engine/Source/Runtime/Engine/Classes/Components/BrushComponent.h
--- a/Engine/Source/Runtime/Engine/Classes/Components/BrushComponent.h
+++ b/Engine/Source/Runtime/Engine/Classes/Components/BrushComponent.h
@@ -24,8 +24,8 @@ class UBrushComponent : public UPrimitiveComponent
 	class UBodySetup* BrushBodySetup;
 
 	/** Local space translation */
-	UPROPERTY(EditAnywhere, BlueprintReadOnly, Category=Transform)
-	FVector PrePivot;
+	UPROPERTY()
+	FVector PrePivot_DEPRECATED;
 
 	// Begin UObject interface
 	virtual void PostLoad() override;
@@ -34,15 +34,7 @@ class UBrushComponent : public UPrimitiveComponent
 
 	// Begin USceneComponent interface
 	virtual FBoxSphereBounds CalcBounds(const FTransform& LocalToWorld) const override;
-	virtual FVector GetCustomLocation() const override;
 	virtual bool ShouldCollideWhenPlacing() const override { return true; }
-
-protected:
-	
-	/** Calculate the new ComponentToWorld transform for this component.
-	Parent is optional and can be used for computing ComponentToWorld based on arbitrary USceneComponent.
-	If Parent is not passed in we use the component's AttachParent*/
-	virtual FTransform CalcNewComponentToWorld(const FTransform& NewRelativeTransform, const USceneComponent* Parent = NULL) const override;
 	// End USceneComponent interface
 
 public:
diff --git a/Engine/Source/Runtime/Engine/Classes/Engine/Brush.h b/Engine/Source/Runtime/Engine/Classes/Engine/Brush.h
--- a/Engine/Source/Runtime/Engine/Classes/Engine/Brush.h
+++ b/Engine/Source/Runtime/Engine/Classes/Engine/Brush.h
@@ -227,9 +227,11 @@ public:
 	// ABrush interface.
 
 	/** Returns the prepivot FVector of the RootComponent of this actor */
+	DEPRECATED(4.9, "This method has been superseded by AActor::GetPivotOffset.")
 	FVector GetPrePivot() const;
 
 	/** Set the PrePivot of the root component */
+	DEPRECATED(4.9, "This method has been superseded by AActor::SetPivotOffset.")
 	void SetPrePivot(const FVector& InPrePivot);
 
 	/** Figures out the best color to use for this brushes wireframe drawing.	*/
diff --git a/Engine/Source/Runtime/Engine/Classes/Engine/Polys.h b/Engine/Source/Runtime/Engine/Classes/Engine/Polys.h
--- a/Engine/Source/Runtime/Engine/Classes/Engine/Polys.h
+++ b/Engine/Source/Runtime/Engine/Classes/Engine/Polys.h
@@ -67,20 +67,19 @@ public:
 	ENGINE_API void Reverse();
 
 	/**
-	 * Transform an editor polygon with a coordinate system, a pre-transformation
-	 * addition, and a post-transformation addition.
+	 * Transform an editor polygon with a post-transformation addition.
 	 */
-	ENGINE_API void Transform(const FVector &PreSubtract,const FVector &PostAdd);
+	ENGINE_API void Transform(const FVector &PostAdd);
 
 	/**
-	 * Rotate an editor polygon with a pre-transformation addition.
+	 * Rotate an editor polygon.
 	 */
-	ENGINE_API void Rotate(const FVector &PreSubtract, const FRotator &Rotation);
+	ENGINE_API void Rotate(const FRotator &Rotation);
 
 	/**
-	 * Scale an editor polygon with a pre-transformation addition.
+	 * Scale an editor polygon.
 	 */
-	ENGINE_API void Scale(const FVector &PreSubtract, const FVector &Scale);
+	ENGINE_API void Scale(const FVector &Scale);
 
 	/**
 	 * Fix up an editor poly by deleting vertices that are identical.  Sets
diff --git a/Engine/Source/Runtime/Engine/Classes/GameFramework/Actor.h b/Engine/Source/Runtime/Engine/Classes/GameFramework/Actor.h
--- a/Engine/Source/Runtime/Engine/Classes/GameFramework/Actor.h
+++ b/Engine/Source/Runtime/Engine/Classes/GameFramework/Actor.h
@@ -377,6 +377,12 @@ protected:
 	UPROPERTY()
 	class USceneComponent* RootComponent;
 
+#if WITH_EDITORONLY_DATA
+	/** Local space pivot offset for the actor */
+	UPROPERTY(EditAnywhere, BlueprintReadOnly, AdvancedDisplay, Category = Actor)
+	FVector PivotOffset;
+#endif
+
 	/** The matinee actors that control this actor. */
 	UPROPERTY(transient)
 	TArray<class AMatineeActor*> ControllingMatineeActors;
@@ -1387,6 +1393,20 @@ public:
 		return GetActorQuat(RootComponent);
 	}
 
+#if WITH_EDITOR
+	/** Sets the local space offset added to the actor's pivot as used by the editor */
+	FORCEINLINE void SetPivotOffset(const FVector& InPivotOffset)
+	{
+		PivotOffset = InPivotOffset;
+	}
+
+	/** Gets the local space offset added to the actor's pivot as used by the editor */
+	FORCEINLINE FVector GetPivotOffset() const
+	{
+		return PivotOffset;
+	}
+#endif
+
 /*-----------------------------------------------------------------------------
 	Relations.
 -----------------------------------------------------------------------------*/
diff --git a/Engine/Source/Runtime/Engine/Private/Actor.cpp b/Engine/Source/Runtime/Engine/Private/Actor.cpp
--- a/Engine/Source/Runtime/Engine/Private/Actor.cpp
+++ b/Engine/Source/Runtime/Engine/Private/Actor.cpp
@@ -93,6 +93,7 @@ void AActor::InitializeDefaults()
 	bFindCameraComponentWhenViewTarget = true;
 	bAllowReceiveTickEventOnDedicatedServer = true;
 	bRelevantForNetworkReplays = true;
+	PivotOffset = FVector::ZeroVector;
 }
 
 void FActorTickFunction::ExecuteTick(float DeltaTime, enum ELevelTick TickType, ENamedThreads::Type CurrentThread, const FGraphEventRef& MyCompletionGraphEvent)
diff --git a/Engine/Source/Runtime/Engine/Private/Brush.cpp b/Engine/Source/Runtime/Engine/Private/Brush.cpp
--- a/Engine/Source/Runtime/Engine/Private/Brush.cpp
+++ b/Engine/Source/Runtime/Engine/Private/Brush.cpp
@@ -90,7 +90,7 @@ void ABrush::CopyPosRotScaleFrom( ABrush* Other )
 	SetActorRotation(Other->GetActorRotation());
 	if( GetRootComponent() != NULL )
 	{
-		SetPrePivot( Other->GetPrePivot() );
+		SetPivotOffset(Other->GetPivotOffset());
 	}
 
 	if(Brush)
@@ -107,7 +107,7 @@ void ABrush::InitPosRotScale()
 
 	SetActorLocation(FVector::ZeroVector, false);
 	SetActorRotation(FRotator::ZeroRotator);
-	SetPrePivot( FVector::ZeroVector );
+	SetPivotOffset( FVector::ZeroVector );
 }
 
 void ABrush::SetIsTemporarilyHiddenInEditor( bool bIsHidden )
@@ -147,21 +147,12 @@ void ABrush::SetIsTemporarilyHiddenInEditor( bool bIsHidden )
 
 FVector ABrush::GetPrePivot() const
 {
-	FVector Result(0.f);
-	if( BrushComponent )
-	{
-		Result = BrushComponent->PrePivot;
-	}
-	return Result;
+	return GetPivotOffset();
 }
 
 void ABrush::SetPrePivot( const FVector& InPrePivot )
 {
-	if( BrushComponent )
-	{
-		BrushComponent->PrePivot = InPrePivot;
-		BrushComponent->UpdateComponentToWorld();
-	}
+	SetPivotOffset(InPrePivot);
 }
 
 #if WITH_EDITOR
diff --git a/Engine/Source/Runtime/Engine/Private/Components/BrushComponent.cpp b/Engine/Source/Runtime/Engine/Private/Components/BrushComponent.cpp
--- a/Engine/Source/Runtime/Engine/Private/Components/BrushComponent.cpp
+++ b/Engine/Source/Runtime/Engine/Private/Components/BrushComponent.cpp
@@ -458,7 +458,6 @@ UBrushComponent::UBrushComponent(const FObjectInitializer& ObjectInitializer)
 	AlwaysLoadOnClient = false;
 	AlwaysLoadOnServer = false;
 	bUseAsOccluder = true;
-	bRequiresCustomLocation = true;
 	bUseEditorCompositing = true;
 	bCanEverAffectNavigation = true;
 }
@@ -518,27 +517,6 @@ FBoxSphereBounds UBrushComponent::CalcBounds(const FTransform& LocalToWorld) con
 	}
 }
 
-FVector UBrushComponent::GetCustomLocation() const
-{
-	const FVector LocationWithPivot = ComponentToWorld.GetLocation();
-	const FVector LocationNoPivot = LocationWithPivot + ComponentToWorld.TransformVector(PrePivot);
-
-	return LocationNoPivot;
-}
-
-FTransform UBrushComponent::CalcNewComponentToWorld(const FTransform& NewRelativeTransform, const USceneComponent* Parent) const
-{
-	FTransform CompToWorld = Super::CalcNewComponentToWorld(NewRelativeTransform, Parent);
-
-	const FVector LocationNoPivot = CompToWorld.GetLocation();
-	const FVector LocationWithPivot = LocationNoPivot + CompToWorld.TransformVector(-PrePivot);
-
-	CompToWorld.SetLocation(LocationWithPivot);
-
-	return CompToWorld;
-}
-
-
 void UBrushComponent::GetUsedMaterials( TArray<UMaterialInterface*>& OutMaterials ) const
 {
 #if WITH_EDITOR
@@ -566,6 +544,12 @@ void UBrushComponent::PostLoad()
 	{
 		BrushBodySetup->bGenerateMirroredCollision = false;
 	}
+
+	AActor* Owner = GetOwner();
+	if (Owner)
+	{
+		Owner->SetPivotOffset(PrePivot_DEPRECATED);
+	}
 }
 
 
diff --git a/Engine/Source/Runtime/Engine/Private/Model.cpp b/Engine/Source/Runtime/Engine/Private/Model.cpp
--- a/Engine/Source/Runtime/Engine/Private/Model.cpp
+++ b/Engine/Source/Runtime/Engine/Private/Model.cpp
@@ -358,7 +358,6 @@ void UModel::PostLoad()
 		if (ABrush* Owner = Cast<ABrush>(GetOuter()))
 		{
 			OwnerLocationWhenLastBuilt = Owner->GetActorLocation();
-			OwnerPrepivotWhenLastBuilt = Owner->GetPrePivot();
 			OwnerScaleWhenLastBuilt = Owner->GetActorScale();
 			OwnerRotationWhenLastBuilt = Owner->GetActorRotation();
 			bCachedOwnerTransformValid = true;
@@ -640,7 +639,7 @@ void UModel::Transform( ABrush* Owner )
 	check(Owner);
 
 	for( int32 i=0; i<Polys->Element.Num(); i++ )
-		Polys->Element[ i ].Transform( Owner->GetPrePivot(), Owner->GetActorLocation());
+		Polys->Element[i].Transform(Owner->GetActorLocation());
 
 }
 
diff --git a/Engine/Source/Runtime/Engine/Private/Polygon.cpp b/Engine/Source/Runtime/Engine/Private/Polygon.cpp
--- a/Engine/Source/Runtime/Engine/Private/Polygon.cpp
+++ b/Engine/Source/Runtime/Engine/Private/Polygon.cpp
@@ -457,22 +457,20 @@ int32 FPoly::CalcNormal( bool bSilent )
 
 void FPoly::Transform
 (
-	const FVector&		PreSubtract,
 	const FVector&		PostAdd
 )
 {
 	FVector 	Temp;
 	int32 		i;
 
-	Base = (Base - PreSubtract) + PostAdd;
+	Base += PostAdd;
 	for( i=0; i<Vertices.Num(); i++ )
-		Vertices[i]  = (Vertices[i] - PreSubtract) + PostAdd;
+		Vertices[i] += PostAdd;
 }
 
 
 void FPoly::Rotate
 (
-	const FVector&		PreSubtract,
 	const FRotator&		Rotation
 )
 {
@@ -481,10 +479,10 @@ void FPoly::Rotate
 	// Rotate the vertices.
 	for (int32 Vertex = 0; Vertex < Vertices.Num(); Vertex++)
 	{
-		Vertices[Vertex] = PreSubtract + RotMatrix.TransformVector(Vertices[Vertex] - PreSubtract);
+		Vertices[Vertex] = RotMatrix.TransformVector(Vertices[Vertex]);
 	}
 
-	Base = PreSubtract + RotMatrix.TransformVector(Base - PreSubtract);
+	Base = RotMatrix.TransformVector(Base);
 
 	// Rotate the texture vectors.
 	TextureU = RotMatrix.TransformVector(TextureU);
@@ -498,7 +496,6 @@ void FPoly::Rotate
 
 void FPoly::Scale
 (
-	const FVector&		PreSubtract,
 	const FVector&		Scale
 )
 {
@@ -512,10 +509,10 @@ void FPoly::Scale
 	// Scale the vertices.
 	for (int32 Vertex = 0; Vertex < Vertices.Num(); Vertex++)
 	{
-		Vertices[Vertex] = PreSubtract + (Vertices[Vertex] - PreSubtract) * Scale;
+		Vertices[Vertex] *= Scale;
 	}
 
-	Base = PreSubtract + (Base - PreSubtract) * Scale;
+	Base *= Scale;
 
 	// Scale the texture vectors.
 	TextureU /= Scale;
diff --git a/Engine/Source/Runtime/Engine/Public/Model.h b/Engine/Source/Runtime/Engine/Public/Model.h
--- a/Engine/Source/Runtime/Engine/Public/Model.h
+++ b/Engine/Source/Runtime/Engine/Public/Model.h
@@ -411,7 +411,6 @@ class UModel : public UObject
 	FVector OwnerLocationWhenLastBuilt;
 	FRotator OwnerRotationWhenLastBuilt;
 	FVector OwnerScaleWhenLastBuilt;
-	FVector OwnerPrepivotWhenLastBuilt;
 
 	/** Specifies whether the above cached transform is valid */
 	bool bCachedOwnerTransformValid;
