Stats - Stats for cooking

diff --git a/Engine/Source/Editor/UnrealEd/Private/Commandlets/CookCommandlet.cpp b/Engine/Source/Editor/UnrealEd/Private/Commandlets/CookCommandlet.cpp
--- a/Engine/Source/Editor/UnrealEd/Private/Commandlets/CookCommandlet.cpp
+++ b/Engine/Source/Editor/UnrealEd/Private/Commandlets/CookCommandlet.cpp
@@ -1205,52 +1205,57 @@ bool UCookCommandlet::NewCook( const TArray<ITargetPlatform*>& Platforms, TArray
 	FDateTime LastConnectionTime = FDateTime::UtcNow();
 	bool bHadConnection = false;
 
-	while ( CookOnTheFlyServer->IsCookByTheBookRunning() )
+	while (CookOnTheFlyServer->IsCookByTheBookRunning())
 	{
-		uint32 TickResults = 0;
-		static const float CookOnTheSideTimeSlice = 10.0f;
-		TickResults = CookOnTheFlyServer->TickCookOnTheSide(CookOnTheSideTimeSlice, NonMapPackageCountSinceLastGC);
-
-		if ( TickResults & (UCookOnTheFlyServer::COSR_CookedMap | UCookOnTheFlyServer::COSR_CookedPackage))
+		DECLARE_SCOPE_CYCLE_COUNTER( TEXT( "NewCook.MainLoop" ), STAT_CookOnTheFly_MainLoop, STATGROUP_LoadTime );
 		{
-			LastCookActionTime = FPlatformTime::Seconds();
-		}
+			uint32 TickResults = 0;
+			static const float CookOnTheSideTimeSlice = 10.0f;
+			TickResults = CookOnTheFlyServer->TickCookOnTheSide( CookOnTheSideTimeSlice, NonMapPackageCountSinceLastGC );
 
+			if (TickResults & (UCookOnTheFlyServer::COSR_CookedMap | UCookOnTheFlyServer::COSR_CookedPackage))
+			{
+				LastCookActionTime = FPlatformTime::Seconds();
+			}
 
-		GShaderCompilingManager->ProcessAsyncResults(true, false);
-	
-		if (NonMapPackageCountSinceLastGC > 0)
-		{
-			// We should GC if we have packages to collect and we've been idle for some time.
-			const bool bExceededPackagesPerGC = (PackagesPerGC > 0) && (NonMapPackageCountSinceLastGC > PackagesPerGC);
-			const bool bExceededIdleTimeToGC = (IdleTimeToGC > 0) && ((FPlatformTime::Seconds() - LastCookActionTime) >= IdleTimeToGC);
-			bShouldGC |= bExceededPackagesPerGC || bExceededIdleTimeToGC;
-		}
 
-		bShouldGC |= (TickResults & UCookOnTheFlyServer::COSR_RequiresGC)!=0;
+			GShaderCompilingManager->ProcessAsyncResults( true, false );
 
-		bShouldGC |= HasExceededMaxMemory(MaxMemoryAllowance);
+			if (NonMapPackageCountSinceLastGC > 0)
+			{
+				// We should GC if we have packages to collect and we've been idle for some time.
+				const bool bExceededPackagesPerGC = (PackagesPerGC > 0) && (NonMapPackageCountSinceLastGC > PackagesPerGC);
+				const bool bExceededIdleTimeToGC = (IdleTimeToGC > 0) && ((FPlatformTime::Seconds() - LastCookActionTime) >= IdleTimeToGC);
+				bShouldGC |= bExceededPackagesPerGC || bExceededIdleTimeToGC;
+			}
 
+			bShouldGC |= (TickResults & UCookOnTheFlyServer::COSR_RequiresGC) != 0;
 
-		// don't clean up if we are waiting on cache of cooked data
-		if (bShouldGC && ((TickResults & UCookOnTheFlyServer::COSR_WaitingOnCache) == 0) )
-		{
-			bShouldGC = false;
-			NonMapPackageCountSinceLastGC = 0;
+			bShouldGC |= HasExceededMaxMemory( MaxMemoryAllowance );
 
-			UE_LOG(LogCookCommandlet, Display, TEXT("GC..."));
 
-			CollectGarbage( RF_Native );
-		}
-		else
-		{
-			CookOnTheFlyServer->TickRecompileShaderRequests();
+			// don't clean up if we are waiting on cache of cooked data
+			if (bShouldGC && ((TickResults & UCookOnTheFlyServer::COSR_WaitingOnCache) == 0))
+			{
+				bShouldGC = false;
+				NonMapPackageCountSinceLastGC = 0;
 
-			FPlatformProcess::Sleep(0.0f);
-		}
+				UE_LOG( LogCookCommandlet, Display, TEXT( "GC..." ) );
+
+				CollectGarbage( RF_Native );
+			}
+			else
+			{
+				CookOnTheFlyServer->TickRecompileShaderRequests();
 
+				FPlatformProcess::Sleep( 0.0f );
+			}
+
+
+			ProcessDeferredCommands();
+		}
 
-		ProcessDeferredCommands();
+		FStats::TickCommandletStats();
 	}
 
 	return true;
