CIS fix

diff --git a/Engine/Source/Editor/UnrealEd/Private/AutoReimport/ReimportFeedbackContext.h b/Engine/Source/Editor/UnrealEd/Private/AutoReimport/ReimportFeedbackContext.h
--- a/Engine/Source/Editor/UnrealEd/Private/AutoReimport/ReimportFeedbackContext.h
+++ b/Engine/Source/Editor/UnrealEd/Private/AutoReimport/ReimportFeedbackContext.h
@@ -7,6 +7,7 @@
 #include "INotificationWidget.h"
 #include "MessageLog.h"
 #include "AutoReimportUtilities.h"
+#include "FileCacheUtilities.h"
 
 /** Feedback context that overrides GWarn for import operations to prevent popup spam */
 class FReimportFeedbackContext : public FFeedbackContext, public INotificationWidget, public TSharedFromThis<FReimportFeedbackContext>
