Fixed chained PIE window creation, making the PIE windows as siblings to allow them to be brought to the front.

UE-17766

diff --git a/Engine/Source/Editor/UnrealEd/Private/PlayLevel.cpp b/Engine/Source/Editor/UnrealEd/Private/PlayLevel.cpp
--- a/Engine/Source/Editor/UnrealEd/Private/PlayLevel.cpp
+++ b/Engine/Source/Editor/UnrealEd/Private/PlayLevel.cpp
@@ -2855,8 +2855,11 @@ UGameInstance* UEditorEngine::CreatePIEGameInstance(int32 PIEInstance, bool bInS
 #if PLATFORM_MAC
 				FSlateApplication::Get().AddWindow(PieWindow);
 #else
-				if (PlayInSettings->PIEAlwaysOnTop && FSlateApplication::Get().GetActiveTopLevelWindow().IsValid())
-					FSlateApplication::Get().AddWindowAsNativeChild(PieWindow, FSlateApplication::Get().GetActiveTopLevelWindow().ToSharedRef(), true);
+				TSharedRef<SWindow, ESPMode::NotThreadSafe> MainWindow = FModuleManager::LoadModuleChecked<IMainFrameModule>(TEXT("MainFrame")).GetParentWindow().ToSharedRef();
+				if (PlayInSettings->PIEAlwaysOnTop)
+				{
+					FSlateApplication::Get().AddWindowAsNativeChild(PieWindow, MainWindow, true);
+				}
 				else
 					FSlateApplication::Get().AddWindow(PieWindow);
 #endif
