BP: Don't recompile up-to-date blueprints on "Play" button

https://github.com/EpicGames/UnrealEngine/pull/1320

#github 1320

diff --git a/Engine/Source/Editor/UnrealEd/Private/PlayLevel.cpp b/Engine/Source/Editor/UnrealEd/Private/PlayLevel.cpp
--- a/Engine/Source/Editor/UnrealEd/Private/PlayLevel.cpp
+++ b/Engine/Source/Editor/UnrealEd/Private/PlayLevel.cpp
@@ -1490,6 +1490,12 @@ struct FInternalPlayLevelUtils
 		{
 			UBlueprint* Blueprint = *BlueprintIt;
 
+			// ignore up-to-date BPs
+			if (Blueprint->IsUpToDate())
+			{
+				continue;
+			}
+
 			// do not try to recompile BPs that have not changed since they last failed to compile, so don't check Blueprint->IsUpToDate()
 			const bool bIsDirtyAndShouldBeRecompiled = Blueprint->IsPossiblyDirty();
 			if (!FBlueprintEditorUtils::IsDataOnlyBlueprint(Blueprint)
