Added codes to support tessellation for clothing

diff --git a/Engine/Source/Editor/UnrealEd/Private/ApexClothingUtils.cpp b/Engine/Source/Editor/UnrealEd/Private/ApexClothingUtils.cpp
--- a/Engine/Source/Editor/UnrealEd/Private/ApexClothingUtils.cpp
+++ b/Engine/Source/Editor/UnrealEd/Private/ApexClothingUtils.cpp
@@ -21,6 +21,10 @@ DEFINE_LOG_CATEGORY_STATIC(LogApexClothingUtils, Log, All);
 
 #endif // #if WITH_APEX
 
+#if WITH_EDITOR
+#include "MeshUtilities.h"
+#endif // #if WITH_EDITOR
+
 #define LOCTEXT_NAMESPACE "ApexClothingUtils"
 
 namespace ApexClothingUtils
@@ -907,6 +911,21 @@ bool AssociateClothingAssetWithSkeletalMesh(USkeletalMesh* SkelMesh, int32 LODIn
 
 	LODModel.NumVertices += ClothChunk.GetNumVertices();
 
+	// build adjacency information for this clothing section
+	{
+		TArray<FSoftSkinVertex> Vertices;
+		FMultiSizeIndexContainerData AdjacencyIndexData;
+		IMeshUtilities& MeshUtilities = FModuleManager::Get().LoadModuleChecked<IMeshUtilities>("MeshUtilities");
+
+		UE_LOG(LogApexClothingUtils, Warning, TEXT("Building adjacency information for a clothing section in skeletal mesh '%s'."), *SkelMesh->GetPathName());
+
+		LODModel.GetVertices(Vertices);
+		AdjacencyIndexData.DataTypeSize = LODModel.MultiSizeIndexContainer.GetDataTypeSize();
+		MeshUtilities.BuildSkeletalAdjacencyIndexBuffer(Vertices, LODModel.NumTexCoords, OutIndexBuffer, AdjacencyIndexData.Indices);
+		LODModel.AdjacencyMultiSizeIndexContainer.RebuildIndexBuffer(AdjacencyIndexData);
+	}
+
+
 	// set asset submesh info to new chunk 
 	ClothChunk.SetClothSubmeshIndex(AssetIndex, AssetSubmeshIndex);
 
diff --git a/Engine/Source/Runtime/Core/Public/UObject/ObjectVersion.h b/Engine/Source/Runtime/Core/Public/UObject/ObjectVersion.h
--- a/Engine/Source/Runtime/Core/Public/UObject/ObjectVersion.h
+++ b/Engine/Source/Runtime/Core/Public/UObject/ObjectVersion.h
@@ -558,7 +558,8 @@ enum EUnrealEngineObjectUE4Version
 	VER_UE4_BUILD_MESH_ADJ_BUFFER_FLAG_EXPOSED,
 	// Change the default max angular velocity
 	VER_UE4_MAX_ANGULAR_VELOCITY_DEFAULT,
-
+	// Build Adjacency index buffer for clothing tessellation
+	VER_UE4_APEX_CLOTH_TESSELLATION,
 	// -----<new versions can be added before this line>-------------------------------------------------
 	// - this needs to be the last line (see note below)
 	VER_UE4_AUTOMATIC_VERSION_PLUS_ONE,
diff --git a/Engine/Source/Runtime/Engine/Private/SkeletalMesh.cpp b/Engine/Source/Runtime/Engine/Private/SkeletalMesh.cpp
--- a/Engine/Source/Runtime/Engine/Private/SkeletalMesh.cpp
+++ b/Engine/Source/Runtime/Engine/Private/SkeletalMesh.cpp
@@ -2301,7 +2301,12 @@ void USkeletalMesh::Serialize( FArchive& Ar )
 	for ( int32 LODIndex = 0; LODIndex < ImportedResource->LODModels.Num(); ++LODIndex )
 	{
 		FStaticLODModel& LODModel = ImportedResource->LODModels[ LODIndex ];
-		if ( !LODModel.AdjacencyMultiSizeIndexContainer.IsIndexBufferValid() )
+
+		if ( !LODModel.AdjacencyMultiSizeIndexContainer.IsIndexBufferValid()
+#if WITH_APEX_CLOTHING
+		|| (Ar.IsLoading() && Ar.UE4Ver() < VER_UE4_APEX_CLOTH_TESSELLATION && LODModel.HasApexClothData())
+#endif // WITH_APEX_CLOTHING
+			)
 		{
 			TArray<FSoftSkinVertex> Vertices;
 			FMultiSizeIndexContainerData IndexData;
