Optimizations to setting the material or its texture transform on a BSP surface.

Previously, setting a material or its texture transform on a BSP surface triggered CommitModelSurfaces, which iterated all the ModelComponents in the level and caused (among other things) the physics state to be recreated, a very expensive operation.  This is clearly unnecessary when only the surface material or texture transform has changed, so introduced a new path in which only the Material Index Buffers are updated for the level's UModel.

diff --git a/Engine/Source/Editor/LevelEditor/Private/LevelEditorGenericDetails.cpp b/Engine/Source/Editor/LevelEditor/Private/LevelEditorGenericDetails.cpp
--- a/Engine/Source/Editor/LevelEditor/Private/LevelEditorGenericDetails.cpp
+++ b/Engine/Source/Editor/LevelEditor/Private/LevelEditorGenericDetails.cpp
@@ -57,7 +57,9 @@ void FLevelEditorGenericDetails::OnMaterialChanged( UMaterialInterface* NewMater
 			const int32 SurfaceIndex = It.GetSurfaceIndex();
 
 			Model->Surfs[SurfaceIndex].Material = NewMaterial;
-			GEditor->polyUpdateMaster( Model, SurfaceIndex, 0 );
+			const bool bUpdateTexCoords = false;
+			const bool bOnlyUpdateSurfaceMaterials = true;
+			GEditor->polyUpdateMaster(Model, SurfaceIndex, bUpdateTexCoords, bOnlyUpdateSurfaceMaterials);
 			Model->MarkPackageDirty();
 
 			bModelDirtied = true;
diff --git a/Engine/Source/Editor/TextureAlignMode/Private/TextureAlignEdMode.cpp b/Engine/Source/Editor/TextureAlignMode/Private/TextureAlignEdMode.cpp
--- a/Engine/Source/Editor/TextureAlignMode/Private/TextureAlignEdMode.cpp
+++ b/Engine/Source/Editor/TextureAlignMode/Private/TextureAlignEdMode.cpp
@@ -261,7 +261,9 @@ bool FModeTool_Texture::InputDelta(FEditorViewportClient* InViewportClient,FView
 			Surf->vTextureU = Model->Vectors.Add(TextureU);
 			Surf->vTextureV = Model->Vectors.Add(TextureV);
 
-			GEditor->polyUpdateMaster( Model, It.GetSurfaceIndex(), 1 );
+			const bool bUpdateTexCoords = true;
+			const bool bOnlyRefreshSurfaceMaterials = true;
+			GEditor->polyUpdateMaster(Model, It.GetSurfaceIndex(), bUpdateTexCoords, bOnlyRefreshSurfaceMaterials);
 		}
 	}
 
diff --git a/Engine/Source/Editor/UnrealEd/Classes/Editor/EditorEngine.h b/Engine/Source/Editor/UnrealEd/Classes/Editor/EditorEngine.h
--- a/Engine/Source/Editor/UnrealEd/Classes/Editor/EditorEngine.h
+++ b/Engine/Source/Editor/UnrealEd/Classes/Editor/EditorEngine.h
@@ -1285,7 +1285,7 @@ public:
 	 *
 	 * Doesn't do any transaction tracking.
 	 */
-	virtual void polyUpdateMaster( UModel* Model, int32 iSurf, int32 UpdateTexCoords );
+	virtual void polyUpdateMaster( UModel* Model, int32 iSurf, bool bUpdateTexCoords, bool bOnlyRefreshSurfaceMaterials );
 
 	/**
 	 * Populates a list with all polys that are linked to the specified poly.  The
diff --git a/Engine/Source/Editor/UnrealEd/Private/EditorCsg.cpp b/Engine/Source/Editor/UnrealEd/Private/EditorCsg.cpp
--- a/Engine/Source/Editor/UnrealEd/Private/EditorCsg.cpp
+++ b/Engine/Source/Editor/UnrealEd/Private/EditorCsg.cpp
@@ -325,8 +325,12 @@ void UEditorEngine::polySetAndClearPolyFlags(UModel *Model, uint32 SetBits, uint
 			{
 				Model->ModifySurf( i, UpdateMaster );
 				Poly.PolyFlags = NewFlags;
-				if( UpdateMaster )
-					polyUpdateMaster( Model, i, 0 );
+				if (UpdateMaster)
+				{
+					const bool bUpdateTexCoords = false;
+					const bool bOnlyRefreshSurfaceMaterials = false;
+					polyUpdateMaster(Model, i, bUpdateTexCoords, bOnlyRefreshSurfaceMaterials);
+				}
 			}
 		}
 	}
@@ -352,7 +356,8 @@ void UEditorEngine::polyUpdateMaster
 (
 	UModel*	Model,
 	int32  	iSurf,
-	int32		UpdateTexCoords
+	bool	bUpdateTexCoords,
+	bool	bOnlyRefreshSurfaceMaterials
 )
 {
 	FBspSurf &Surf = Model->Surfs[iSurf];
@@ -394,7 +399,7 @@ void UEditorEngine::polyUpdateMaster
 			MasterEdPoly.Material = Surf.Material;
 			MasterEdPoly.PolyFlags = Surf.PolyFlags & ~(PF_NoEdit);
 
-			if (UpdateTexCoords)
+			if (bUpdateTexCoords)
 			{
 				MasterEdPoly.Base = RotationMatrix.InverseTransformVector(Model->Points[Surf.pBase] - ActorLocation) / ActorScale;
 				MasterEdPoly.TextureU = RotationMatrix.InverseTransformVector(Model->Vectors[Surf.vTextureU]) * ActorScale;
@@ -404,6 +409,11 @@ void UEditorEngine::polyUpdateMaster
 	}
 
 	Model->InvalidSurfaces = true;
+
+	if (bOnlyRefreshSurfaceMaterials)
+	{
+		Model->bOnlyRebuildMaterialIndexBuffers = true;
+	}
 }
 
 
@@ -1460,7 +1470,9 @@ void UEditorEngine::polyTexPan(UModel *Model,int32 PanU,int32 PanV,int32 Absolut
 			Model->Points[Surf.pBase] += PanU * (TextureU / TextureU.SizeSquared());
 			Model->Points[Surf.pBase] += PanV * (TextureV / TextureV.SizeSquared());
 
-			polyUpdateMaster(Model,SurfaceIndex,1);
+			const bool bUpdateTexCoords = true;
+			const bool bOnlyRefreshSurfaceMaterials = true;
+			polyUpdateMaster(Model, SurfaceIndex, bUpdateTexCoords, bOnlyRefreshSurfaceMaterials);
 		}
 	}
 }
@@ -1487,7 +1499,9 @@ void UEditorEngine::polyTexScale( UModel* Model, float UU, float UV, float VU, f
 			Model->Vectors[Poly->vTextureV] = OriginalU * VU + OriginalV * VV;
 
 			// Update generating brush poly.
-			polyUpdateMaster( Model, i, 1 );
+			const bool bUpdateTexCoords = true;
+			const bool bOnlyRefreshSurfaceMaterials = true;
+			polyUpdateMaster(Model, i, bUpdateTexCoords, bOnlyRefreshSurfaceMaterials);
 		}
 	}
 }
diff --git a/Engine/Source/Editor/UnrealEd/Private/EditorServer.cpp b/Engine/Source/Editor/UnrealEd/Private/EditorServer.cpp
--- a/Engine/Source/Editor/UnrealEd/Private/EditorServer.cpp
+++ b/Engine/Source/Editor/UnrealEd/Private/EditorServer.cpp
@@ -4109,7 +4109,9 @@ bool UEditorEngine::Exec_Poly( UWorld* InWorld, const TCHAR* Str, FOutputDevice&
 				const int32 SurfaceIndex = It.GetSurfaceIndex();
 
 				Model->Surfs[SurfaceIndex].Material = SelectedMaterialInstance;
-				polyUpdateMaster( Model, SurfaceIndex, 0 );
+				const bool bUpdateTexCoords = false;
+				const bool bOnlyRefreshSurfaceMaterials = true;
+				polyUpdateMaster(Model, SurfaceIndex, bUpdateTexCoords, bOnlyRefreshSurfaceMaterials);
 				Model->MarkPackageDirty();
 
 				bModelDirtied = true;
@@ -4140,7 +4142,9 @@ bool UEditorEngine::Exec_Poly( UWorld* InWorld, const TCHAR* Str, FOutputDevice&
 				{
 					const int32 SurfaceIndex = It.GetSurfaceIndex();
 					It.GetModel()->Surfs[SurfaceIndex].Material = Material;
-					polyUpdateMaster( It.GetModel(), SurfaceIndex, 0 );
+					const bool bUpdateTexCoords = false;
+					const bool bOnlyRefreshSurfaceMaterials = true;
+					polyUpdateMaster(It.GetModel(), SurfaceIndex, bUpdateTexCoords, bOnlyRefreshSurfaceMaterials);
 				}
 			}
 
diff --git a/Engine/Source/Editor/UnrealEd/Private/Factories/EditorFactories.cpp b/Engine/Source/Editor/UnrealEd/Private/Factories/EditorFactories.cpp
--- a/Engine/Source/Editor/UnrealEd/Private/Factories/EditorFactories.cpp
+++ b/Engine/Source/Editor/UnrealEd/Private/Factories/EditorFactories.cpp
@@ -740,7 +740,9 @@ UObject* ULevelFactory::FactoryCreateText
 
 							CurrentLevel->MarkPackageDirty();
 
-							GEditor->polyUpdateMaster( CurrentLevel->Model, i, 1 );
+							const bool bUpdateTexCoords = true;
+							const bool bOnlyRefreshSurfaceMaterials = false;
+							GEditor->polyUpdateMaster(CurrentLevel->Model, i, bUpdateTexCoords, bOnlyRefreshSurfaceMaterials);
 						}
 					}
 				}
diff --git a/Engine/Source/Editor/UnrealEd/Private/LevelEditorViewport.cpp b/Engine/Source/Editor/UnrealEd/Private/LevelEditorViewport.cpp
--- a/Engine/Source/Editor/UnrealEd/Private/LevelEditorViewport.cpp
+++ b/Engine/Source/Editor/UnrealEd/Private/LevelEditorViewport.cpp
@@ -774,8 +774,9 @@ bool FLevelEditorViewportClient::AttemptApplyObjAsMaterialToSurface( UObject* Ob
 
 				Model->ModifySurf(SelectedSurfIndex, true);
 				Model->Surfs[SelectedSurfIndex].Material = DroppedObjAsMaterial;
-				GEditor->polyUpdateMaster(Model, SelectedSurfIndex, false);
-			
+				const bool bUpdateTexCoords = false;
+				const bool bOnlyRefreshSurfaceMaterials = true;
+				GEditor->polyUpdateMaster(Model, SelectedSurfIndex, bUpdateTexCoords, bOnlyRefreshSurfaceMaterials);
 			}
 
 			bResult = true;
diff --git a/Engine/Source/Editor/UnrealEd/Private/LevelViewportClickHandlers.cpp b/Engine/Source/Editor/UnrealEd/Private/LevelViewportClickHandlers.cpp
--- a/Engine/Source/Editor/UnrealEd/Private/LevelViewportClickHandlers.cpp
+++ b/Engine/Source/Editor/UnrealEd/Private/LevelViewportClickHandlers.cpp
@@ -782,7 +782,9 @@ namespace ClickHandlers
 				{
 					Model->ModifySurf( i, 1 );
 					Model->Surfs[i].Material = SelectedMaterialInstance;
-					GEditor->polyUpdateMaster( Model, i, 0 );
+					const bool bUpdateTexCoords = false;
+					const bool bOnlyRefreshSurfaceMaterials = true;
+					GEditor->polyUpdateMaster(Model, i, bUpdateTexCoords, bOnlyRefreshSurfaceMaterials);
 				}
 			}
 		}
@@ -851,11 +853,15 @@ namespace ClickHandlers
 					UE_LOG(LogEditorViewport, Log, TEXT("WARNING: the texture coordinates were not parallel to the surface.") );
 				}
 				Surf.PolyFlags	= GSaveSurf.PolyFlags;
-				GEditor->polyUpdateMaster( Model, iSurf, 1 );
+				const bool bUpdateTexCoords = true;
+				const bool bOnlyRefreshSurfaceMaterials = true;
+				GEditor->polyUpdateMaster(Model, iSurf, bUpdateTexCoords, bOnlyRefreshSurfaceMaterials);
 			}
 			else
 			{
-				GEditor->polyUpdateMaster( Model, iSurf, 0 );
+				const bool bUpdateTexCoords = false;
+				const bool bOnlyRefreshSurfaceMaterials = true;
+				GEditor->polyUpdateMaster(Model, iSurf, bUpdateTexCoords, bOnlyRefreshSurfaceMaterials);
 			}
 		}
 		else if( Click.GetKey() == EKeys::RightMouseButton && !Click.IsControlDown() )
diff --git a/Engine/Source/Editor/UnrealEd/Private/TexAlignTools.cpp b/Engine/Source/Editor/UnrealEd/Private/TexAlignTools.cpp
--- a/Engine/Source/Editor/UnrealEd/Private/TexAlignTools.cpp
+++ b/Engine/Source/Editor/UnrealEd/Private/TexAlignTools.cpp
@@ -127,7 +127,9 @@ void UTexAligner::Align( UWorld* InWorld, ETexAlign InTexAlignType, UModel* InMo
 
 		AlignSurf( InTexAlignType == TEXALIGN_None ? (ETexAlign)DefTexAlign : InTexAlignType, InModel, Surf, &EdPoly, &Normal );
 
-		GEditor->polyUpdateMaster( InModel, Surf->Idx, 1 );
+		const bool bUpdateTexCoords = true;
+		const bool bOnlyRefreshSurfaceMaterials = true;
+		GEditor->polyUpdateMaster(InModel, Surf->Idx, bUpdateTexCoords, bOnlyRefreshSurfaceMaterials);
 	}
 
 	GEditor->RedrawLevelEditingViewports();
diff --git a/Engine/Source/Runtime/Engine/Private/Level.cpp b/Engine/Source/Runtime/Engine/Private/Level.cpp
--- a/Engine/Source/Runtime/Engine/Private/Level.cpp
+++ b/Engine/Source/Runtime/Engine/Private/Level.cpp
@@ -1187,12 +1187,15 @@ void ULevel::CommitModelSurfaces()
 {
 	if(Model->InvalidSurfaces)
 	{
-		// Unregister model components
-		for(int32 ComponentIndex = 0;ComponentIndex < ModelComponents.Num();ComponentIndex++)
+		if (!Model->bOnlyRebuildMaterialIndexBuffers)
 		{
-			if(ModelComponents[ComponentIndex] && ModelComponents[ComponentIndex]->IsRegistered())
+			// Unregister model components
+			for (int32 ComponentIndex = 0; ComponentIndex < ModelComponents.Num(); ComponentIndex++)
 			{
-				ModelComponents[ComponentIndex]->UnregisterComponent();
+				if (ModelComponents[ComponentIndex] && ModelComponents[ComponentIndex]->IsRegistered())
+				{
+					ModelComponents[ComponentIndex]->UnregisterComponent();
+				}
 			}
 		}
 
@@ -1228,7 +1231,14 @@ void ULevel::CommitModelSurfaces()
 			{
 				if(ModelComponents[ComponentIndex])
 				{
-					ModelComponents[ComponentIndex]->RegisterComponentWithWorld(OwningWorld);
+					if (Model->bOnlyRebuildMaterialIndexBuffers)
+					{
+						ModelComponents[ComponentIndex]->MarkRenderStateDirty();
+					}
+					else
+					{
+						ModelComponents[ComponentIndex]->RegisterComponentWithWorld(OwningWorld);
+					}
 				}
 			}
 		}
@@ -1240,6 +1250,8 @@ void ULevel::CommitModelSurfaces()
 		{
 			BeginInitResource(IndexBufferIt.Value());
 		}
+
+		Model->bOnlyRebuildMaterialIndexBuffers = false;
 	}
 }
 
diff --git a/Engine/Source/Runtime/Engine/Private/Model.cpp b/Engine/Source/Runtime/Engine/Private/Model.cpp
--- a/Engine/Source/Runtime/Engine/Private/Model.cpp
+++ b/Engine/Source/Runtime/Engine/Private/Model.cpp
@@ -560,6 +560,8 @@ UModel::UModel(const FObjectInitializer& ObjectInitializer)
 	, Points()
 	, Surfs()
 	, VertexBuffer(this)
+	, InvalidSurfaces(false)
+	, bOnlyRebuildMaterialIndexBuffers(false)
 #if WITH_EDITOR
 	, bCachedOwnerTransformValid(false)
 #endif
@@ -576,6 +578,8 @@ UModel::UModel(FVTableHelper& Helper)
 	, Points()
 	, Surfs()
 	, VertexBuffer(this)
+	, InvalidSurfaces(false)
+	, bOnlyRebuildMaterialIndexBuffers(false)
 #if WITH_EDITOR
 	, bCachedOwnerTransformValid(false)
 #endif
diff --git a/Engine/Source/Runtime/Engine/Public/Model.h b/Engine/Source/Runtime/Engine/Public/Model.h
--- a/Engine/Source/Runtime/Engine/Public/Model.h
+++ b/Engine/Source/Runtime/Engine/Public/Model.h
@@ -385,6 +385,9 @@ class UModel : public UObject
 	/** True if surfaces in the model have been changed without calling ULevel::CommitModelSurfaces. */
 	bool InvalidSurfaces;
 
+	/** True if only the material index buffers should be rebuilt when committing model surfaces */
+	bool bOnlyRebuildMaterialIndexBuffers;
+
 	/** True if static lighting now can not be validly built for this model */
 	bool bInvalidForStaticLighting;
 
