[UE-16968] Making sure we reference duplicated default sub-objects directly in the transaction buffer, instead of attempting to look them up later, via their outer's class (finding the original instead of the duplicate).

diff --git a/Engine/Source/Editor/UnrealEd/Classes/Editor/Transactor.h b/Engine/Source/Editor/UnrealEd/Classes/Editor/Transactor.h
--- a/Engine/Source/Editor/UnrealEd/Classes/Editor/Transactor.h
+++ b/Engine/Source/Editor/UnrealEd/Classes/Editor/Transactor.h
@@ -108,25 +108,48 @@ protected:
 
 			FPersistentObjectRef(UObject* InObject)
 			{
+				// we want to reference CDOs and default sub-objects in a unique 
+				// way... Blueprints can delete and reconstruct CDOs and their
+				// sub-objects during compilation; when undoing, we want changes 
+				// to be reverted for the most recent version of the CDO/sub-
+				// object (not one that has since been thrown out); therefore, 
+				// we record the CDO's class (which remains static) and sub-
+				// objects' names so we can look them up later in Get() 
 				const bool bIsCDO = InObject && InObject->HasAllFlags(RF_ClassDefaultObject);
-				UObject* CDO = nullptr;
+				UObject* CDO = bIsCDO ? InObject : nullptr;
 				const bool bIsSubobjectOfCDO = OuterIsCDO(InObject, CDO);
-				if (bIsCDO )
+				
+				// we have to be careful though, Blueprints also duplicate CDOs 
+				// and their sub-objects; we don't want changes to the 
+				// duplicated CDO/sub-object to be applied back to the original 
+				// (the original would most likely be destroyed when we attempt 
+				// to undo the duplication)... here we check that the class
+				// recognizes this CDO as its own (if not, then we're most 
+				// likely in the middle of a duplicate)
+				const bool bIsClassCDO = (CDO != nullptr) ? (CDO->GetClass()->ClassDefaultObject == CDO) : false;
+				const bool bRefObjectsByClass = (bIsCDO || bIsSubobjectOfCDO) && bIsClassCDO;
+
+				if (bRefObjectsByClass)
 				{
 					Object = nullptr;
-					SourceCDO = InObject->GetClass();
-				}
-				else if (bIsSubobjectOfCDO)
-				{
-					Object = nullptr;
-					SubObjectId = InObject->GetFName();
 					SourceCDO = CDO->GetClass();
+
+					if (bIsSubobjectOfCDO)
+					{
+						SubObjectId = InObject->GetFName();
+					}
 				}
 				else
 				{
+					// @TODO: if bIsCDO/bIsSubobjectOfCDO is true, but bRefObjectsByClass is not,
+					//        then we end up here and the transaction buffer ends up most likely 
+					//        referencing an intermediate REINST/TRASH class (keeping it from being GC'd)     
 					Object = InObject;
 					SourceCDO = nullptr;
 				}
+
+				// can be used to verify that this struct is referencing the object correctly
+				//ensure(Get() == InObject);
 			}
 
 			UObject* Get() const
@@ -147,9 +170,15 @@ protected:
 				return Object;
 			}
 
+			/** Determines if the object referenced by this struct, needs to be kept from garbage collection */
 			bool ShouldAddReference() const
 			{
-				return !(SourceCDO != nullptr && SubObjectId != FName());
+				// if the object is being referenced through SourceCDO (instead 
+				// of a direct Object pointer), then we don't need keep it from 
+				// garbage collection, this will continue to reference its 
+				// replacement post GC... we only need to keep hard Object 
+				// references from getting GC'd
+				return (Object != nullptr);
 			}
 
 			UObject* operator->() const
