Removed erroneous use of FFeedbackContextEditor

diff --git a/Engine/Source/Editor/UnrealEd/Private/Fbx/FbxMainImport.cpp b/Engine/Source/Editor/UnrealEd/Private/Fbx/FbxMainImport.cpp
--- a/Engine/Source/Editor/UnrealEd/Private/Fbx/FbxMainImport.cpp
+++ b/Engine/Source/Editor/UnrealEd/Private/Fbx/FbxMainImport.cpp
@@ -33,7 +33,6 @@
 =============================================================================*/
 
 #include "UnrealEd.h"
-#include "FeedbackContextEditor.h"
 
 #include "Factories.h"
 #include "Engine.h"
@@ -445,8 +444,7 @@ int32 FFbxImporter::GetImportType(const FString& InFilename)
 bool FFbxImporter::GetSceneInfo(FString Filename, FbxSceneInfo& SceneInfo)
 {
 	bool Result = true;
-	FFeedbackContextEditor FbxImportWarn;
-	FbxImportWarn.BeginSlowTask( NSLOCTEXT("FbxImporter", "BeginGetSceneInfoTask", "Parse FBX file to get scene info"), true );
+	GWarn->BeginSlowTask( NSLOCTEXT("FbxImporter", "BeginGetSceneInfoTask", "Parse FBX file to get scene info"), true );
 	
 	bool bSceneInfo = true;
 	switch (CurPhase)
@@ -457,14 +455,14 @@ bool FFbxImporter::GetSceneInfo(FString Filename, FbxSceneInfo& SceneInfo)
 			Result = false;
 			break;
 		}
-		FbxImportWarn.UpdateProgress( 40, 100 );
+		GWarn->UpdateProgress( 40, 100 );
 	case FILEOPENED:
 		if (!ImportFile(Filename))
 		{
 			Result = false;
 			break;
 		}
-		FbxImportWarn.UpdateProgress( 90, 100 );
+		GWarn->UpdateProgress( 90, 100 );
 	case IMPORTED:
 	
 	default:
@@ -586,7 +584,7 @@ bool FFbxImporter::GetSceneInfo(FString Filename, FbxSceneInfo& SceneInfo)
 		}
 	}
 	
-	FbxImportWarn.EndSlowTask();
+	GWarn->EndSlowTask();
 	return Result;
 }
 
