minor polish on UI for FoV

diff --git a/Engine/Source/Editor/LevelEditor/Private/SLevelViewportToolBar.cpp b/Engine/Source/Editor/LevelEditor/Private/SLevelViewportToolBar.cpp
--- a/Engine/Source/Editor/LevelEditor/Private/SLevelViewportToolBar.cpp
+++ b/Engine/Source/Editor/LevelEditor/Private/SLevelViewportToolBar.cpp
@@ -516,7 +516,7 @@ TSharedRef<SWidget> SLevelViewportToolBar::GenerateOptionsMenu() const
 
 			if( bIsPerspective )
 			{
-				OptionsMenuBuilder.AddWidget( GenerateFOVMenu(), LOCTEXT("FOVAngle", "Field of View") );
+				OptionsMenuBuilder.AddWidget( GenerateFOVMenu(), LOCTEXT("FOVAngle", "Field of View (H)") );
 				OptionsMenuBuilder.AddWidget( GenerateFarViewPlaneMenu(), LOCTEXT("FarViewPlane", "Far View Plane") );
 			}
 		}
diff --git a/Engine/Source/Editor/UnrealEd/Private/SCommonEditorViewportToolbarBase.cpp b/Engine/Source/Editor/UnrealEd/Private/SCommonEditorViewportToolbarBase.cpp
--- a/Engine/Source/Editor/UnrealEd/Private/SCommonEditorViewportToolbarBase.cpp
+++ b/Engine/Source/Editor/UnrealEd/Private/SCommonEditorViewportToolbarBase.cpp
@@ -205,7 +205,7 @@ TSharedRef<SWidget> SCommonEditorViewportToolbarBase::GenerateOptionsMenu() cons
 
 			if (bIsPerspective)
 			{
-				OptionsMenuBuilder.AddWidget( GenerateFOVMenu(), LOCTEXT("FOVAngle", "Field of View") );
+				OptionsMenuBuilder.AddWidget( GenerateFOVMenu(), LOCTEXT("FOVAngle", "Field of View (H)") );
 				OptionsMenuBuilder.AddWidget( GenerateFarViewPlaneMenu(), LOCTEXT("FarViewPlane", "Far View Plane") );
 			}
 		}
