Sequencer: Curve editor enhancements

- Support for multi select tangents
- Middle mouse also moves selected keyframes or selected tangents
- Fix selection so that shift adds to selected keyframes or selected tangents and control toggles selection of keyframes/tangents
- Shift to restrict movement to one axis
- alt-middle mouse to pan
- alt-right mouse to zoom
- tangent display options - always show tangents, show tangents for selected points, never show tangents
- pre-post infinity extrapolation options for curves
- add auto frame option to frame on curves when they are selected
- flatten and straight tangents buttons
- move curve editor toolbar button settings from sequencer settings to curve editor settings

diff --git a/Engine/Source/Editor/EditorStyle/Private/SlateEditorStyle.cpp b/Engine/Source/Editor/EditorStyle/Private/SlateEditorStyle.cpp
--- a/Engine/Source/Editor/EditorStyle/Private/SlateEditorStyle.cpp
+++ b/Engine/Source/Editor/EditorStyle/Private/SlateEditorStyle.cpp
@@ -931,7 +931,9 @@ void FSlateEditorStyle::FStyle::SetupGeneralStyles()
 		Set("CurveEd.InfoFont", TTF_CORE_FONT("Fonts/Roboto-Regular", 8));
 		Set("CurveEd.LabelFont", TTF_CORE_FONT("Fonts/Roboto-Bold", 10));
 		Set("CurveEd.Tangent", new IMAGE_BRUSH("Common/Tangent", FVector2D(7.0f, 7.0f), FLinearColor(0.0f, 0.66f, 0.7f)));
+		Set("CurveEd.TangentSelected", new IMAGE_BRUSH("Common/Tangent", FVector2D(7.0f, 7.0f), FLinearColor(1.0f, 1.0f, 0.0f)));
 		Set("CurveEd.TangentColor", FLinearColor(0.0f, 0.66f, 0.7f));
+		Set("CurveEd.TangentColorSelected", FLinearColor(1.0f, 1.0f, 0.0f));
 		Set("CurveEd.Visible", new IMAGE_BRUSH("Icons/icon_levels_visible_16px", Icon16x16));
 		Set("CurveEd.VisibleHighlight", new IMAGE_BRUSH("Icons/icon_levels_visible_hi_16px", Icon16x16));
 		Set("CurveEd.Invisible", new IMAGE_BRUSH("Icons/icon_levels_invisible_16px", Icon16x16));
@@ -6068,6 +6070,11 @@ void FSlateEditorStyle::FStyle::SetupToolkitStyles()
 		Set("RichCurveEditor.InterpolationLinear.Small", new IMAGE_BRUSH("Icons/icon_CurveEditor_Linear_40x", Icon20x20));
 		Set("RichCurveEditor.InterpolationConstant", new IMAGE_BRUSH("Icons/icon_CurveEditor_Constant_40x", Icon40x40));
 		Set("RichCurveEditor.InterpolationConstant.Small", new IMAGE_BRUSH("Icons/icon_CurveEditor_Constant_40x", Icon20x20));
+
+		Set("RichCurveEditor.FlattenTangents", new IMAGE_BRUSH("Icons/icon_CurveEditor_Flatten_40x", Icon40x40));
+		Set("RichCurveEditor.FlattenTangents.Small", new IMAGE_BRUSH("Icons/icon_CurveEditor_Flatten_40x", Icon20x20));
+		Set("RichCurveEditor.StraightenTangents", new IMAGE_BRUSH("Icons/icon_CurveEditor_Straighten_40x", Icon40x40));
+		Set("RichCurveEditor.StraightenTangents.Small", new IMAGE_BRUSH("Icons/icon_CurveEditor_Straighten_40x", Icon20x20));
 	}
 
 	// PhAT
diff --git a/Engine/Source/Editor/Sequencer/Private/SSequencer.cpp b/Engine/Source/Editor/Sequencer/Private/SSequencer.cpp
--- a/Engine/Source/Editor/Sequencer/Private/SSequencer.cpp
+++ b/Engine/Source/Editor/Sequencer/Private/SSequencer.cpp
@@ -193,8 +193,6 @@ void SSequencer::Construct( const FArguments& InArgs, TSharedRef< class FSequenc
 	USelection::SelectionChangedEvent.AddSP(this, &SSequencer::OnActorSelectionChanged);
 	Settings = InSequencer->GetSettings();
 
-	Settings->GetOnShowCurveEditorChanged().AddSP(this, &SSequencer::OnCurveEditorVisibilityChanged);
-
 	// Create a node tree which contains a tree of movie scene data to display in the sequence
 	SequencerNodeTree = MakeShareable( new FSequencerNodeTree( InSequencer.Get() ) );
 
@@ -587,13 +585,6 @@ TSharedRef<SWidget> SSequencer::MakeToolBar()
 		ToolBarBuilder.SetLabelVisibility( EVisibility::Visible );
 		ToolBarBuilder.AddToolBarButton( FSequencerCommands::Get().ToggleShowCurveEditor );
 		ToolBarBuilder.SetLabelVisibility( EVisibility::Collapsed );
-		ToolBarBuilder.AddComboButton(
-			FUIAction(),
-			FOnGetContent::CreateSP( this, &SSequencer::MakeCurveEditorMenu ),
-			LOCTEXT( "CurveEditorOptions", "Options" ),
-			LOCTEXT( "CurveEditorOptionsToolTip", "Curve Editor Options" ),
-			TAttribute<FSlateIcon>(),
-			true );
 	}
 
 	return ToolBarBuilder.MakeWidget();
@@ -638,27 +629,6 @@ TSharedRef<SWidget> SSequencer::MakeSnapMenu()
 	return MenuBuilder.MakeWidget();
 }
 
-TSharedRef<SWidget> SSequencer::MakeCurveEditorMenu()
-{
-	FMenuBuilder MenuBuilder( true, Sequencer.Pin()->GetCommandBindings() );
-
-	MenuBuilder.BeginSection( "CurveVisibility", LOCTEXT( "CurveEditorMenuCurveVisibilityHeader", "Curve Visibility" ) );
-	{
-		MenuBuilder.AddMenuEntry( FSequencerCommands::Get().SetAllCurveVisibility );
-		MenuBuilder.AddMenuEntry( FSequencerCommands::Get().SetSelectedCurveVisibility );
-		MenuBuilder.AddMenuEntry( FSequencerCommands::Get().SetAnimatedCurveVisibility );
-	}
-	MenuBuilder.EndSection();
-
-	MenuBuilder.BeginSection( "Options", LOCTEXT( "CurveEditorMenuOptionsHeader", "Editor Options" ) );
-	{
-		MenuBuilder.AddMenuEntry( FSequencerCommands::Get().ToggleShowCurveEditorCurveToolTips );
-	}
-	MenuBuilder.EndSection();
-
-	return MenuBuilder.MakeWidget();
-}
-
 TSharedRef<SWidget> SSequencer::MakeTransportControls()
 {
 	FEditorWidgetsModule& EditorWidgetsModule = FModuleManager::Get().LoadModuleChecked<FEditorWidgetsModule>( "EditorWidgets" );
@@ -680,7 +650,6 @@ TSharedRef<SWidget> SSequencer::MakeTransportControls()
 SSequencer::~SSequencer()
 {
 	USelection::SelectionChangedEvent.RemoveAll(this);
-	Settings->GetOnShowCurveEditorChanged().RemoveAll(this);
 }
 
 void SSequencer::RegisterActiveTimerForPlayback()
@@ -740,9 +709,7 @@ void SSequencer::UpdateLayoutTree()
 	CurveEditor->SetSequencerNodeTree(SequencerNodeTree);
 
 	// Restore the selection state.
-	RestoreSelectionState(SequencerNodeTree->GetRootNodes(), SelectedPathNames, Sequencer.Pin()->GetSelection());
-
-	// Update to actor selection.
+	RestoreSelectionState(SequencerNodeTree->GetRootNodes(), SelectedPathNames, Sequencer.Pin()->GetSelection());	// Update to actor selection.
 	OnActorSelectionChanged(NULL);
 
 	SequencerNodeTree->UpdateCachedVisibilityBasedOnShotFiltersChanged();
@@ -1247,13 +1214,4 @@ EVisibility SSequencer::GetCurveEditorVisibility() const
 	return Settings->GetShowCurveEditor() ? EVisibility::Visible : EVisibility::Collapsed;
 }
 
-void SSequencer::OnCurveEditorVisibilityChanged()
-{
-	if (CurveEditor.IsValid())
-	{
-		// Only zoom horizontally if the editor is visible
-		CurveEditor->SetZoomToFit(true, Settings->GetShowCurveEditor());
-	}
-}
-
 #undef LOCTEXT_NAMESPACE
diff --git a/Engine/Source/Editor/Sequencer/Private/SSequencer.h b/Engine/Source/Editor/Sequencer/Private/SSequencer.h
--- a/Engine/Source/Editor/Sequencer/Private/SSequencer.h
+++ b/Engine/Source/Editor/Sequencer/Private/SSequencer.h
@@ -120,9 +120,6 @@ private:
 	/** Makes the snapping menu for the toolbar. */
 	TSharedRef<SWidget> MakeSnapMenu();
 
-	/** Makes the curve editor menu for the toolbar. */
-	TSharedRef<SWidget> MakeCurveEditorMenu();
-
 	/** Makes and configures a set of the standard UE transport controls. */
 	TSharedRef<SWidget> MakeTransportControls();
 
@@ -229,9 +226,6 @@ private:
 	/** Called when a column fill percentage is changed by a splitter slot. */
 	void OnColumnFillCoefficientChanged(float FillCoefficient, int32 ColumnIndex);
 
-	/** Called when the curve editor is shown or hidden */
-	void OnCurveEditorVisibilityChanged();
-
 private:
 
 	/** Section area widget */
diff --git a/Engine/Source/Editor/Sequencer/Private/SSequencerCurveEditor.cpp b/Engine/Source/Editor/Sequencer/Private/SSequencerCurveEditor.cpp
--- a/Engine/Source/Editor/Sequencer/Private/SSequencerCurveEditor.cpp
+++ b/Engine/Source/Editor/Sequencer/Private/SSequencerCurveEditor.cpp
@@ -10,11 +10,10 @@
 void SSequencerCurveEditor::Construct( const FArguments& InArgs, TSharedRef<FSequencer> InSequencer, TSharedRef<ITimeSliderController> InTimeSliderController )
 {
 	Sequencer = InSequencer;
-	Settings = InSequencer->GetSettings();
+	SequencerSettings = InSequencer->GetSettings();
 
 	TimeSliderController = InTimeSliderController;
 	NodeTreeSelectionChangedHandle = Sequencer.Pin()->GetSelection().GetOnOutlinerNodeSelectionChanged().AddSP(this, &SSequencerCurveEditor::NodeTreeSelectionChanged);
-	Settings->GetOnCurveVisibilityChanged().AddSP( this, &SSequencerCurveEditor::SequencerCurveVisibilityChanged );
 
 	// @todo sequencer: switch to lambda capture expressions when support is introduced?
 	auto ViewRange = InArgs._ViewRange;
@@ -26,7 +25,7 @@ void SSequencerCurveEditor::Construct( const FArguments& InArgs, TSharedRef<FSeq
 		.ViewMaxInput_Lambda( [=]{ return ViewRange.Get().GetUpperBoundValue(); } )
 		.OnSetInputViewRange_Lambda( [=](float InLowerBound, float InUpperBound){ OnViewRangeChanged.ExecuteIfBound(TRange<float>(InLowerBound, InUpperBound), EViewRangeInterpolation::Immediate); } )
 		.HideUI( false )
-		.ZoomToFitHorizontal( Settings->GetShowCurveEditor() )
+		.ZoomToFitHorizontal( true )
 		.ShowCurveSelector( false )
 		.ShowZoomButtons( false )
 		.ShowInputGridNumbers( false )
@@ -34,9 +33,10 @@ void SSequencerCurveEditor::Construct( const FArguments& InArgs, TSharedRef<FSeq
 		.InputSnap( this, &SSequencerCurveEditor::GetCurveTimeSnapInterval )
 		.OutputSnap( this, &SSequencerCurveEditor::GetCurveValueSnapInterval )
 		.TimelineLength( 0.0f )
-		.ShowCurveToolTips( this, &SSequencerCurveEditor::GetShowCurveEditorCurveToolTips )
 		.GridColor( FLinearColor( 0.3f, 0.3f, 0.3f, 0.3f ) )
 	);
+
+	GetSettings()->GetOnCurveEditorCurveVisibilityChanged().AddSP(this, &SSequencerCurveEditor::OnCurveEditorCurveVisibilityChanged);
 }
 
 FReply SSequencerCurveEditor::OnMouseButtonDown( const FGeometry& MyGeometry, const FPointerEvent& MouseEvent )
@@ -110,46 +110,41 @@ void SSequencerCurveEditor::SetSequencerNodeTree( TSharedPtr<FSequencerNodeTree>
 
 void SSequencerCurveEditor::UpdateCurveOwner()
 {
-	CurveOwner = MakeShareable( new FSequencerCurveOwner( SequencerNodeTree ) );
+	CurveOwner = MakeShareable( new FSequencerCurveOwner( SequencerNodeTree, GetSettings()->GetCurveVisibility() ) );
 	SetCurveOwner( CurveOwner.Get() );
 }
 
 bool SSequencerCurveEditor::GetCurveSnapEnabled() const
 {
-	return Settings->GetIsSnapEnabled();
+	return SequencerSettings->GetIsSnapEnabled();
 }
 
 float SSequencerCurveEditor::GetCurveTimeSnapInterval() const
 {
-	return Settings->GetSnapKeyTimesToInterval()
-		? Settings->GetTimeSnapInterval()
+	return SequencerSettings->GetSnapKeyTimesToInterval()
+		? SequencerSettings->GetTimeSnapInterval()
 		: 0;
 }
 
 float SSequencerCurveEditor::GetCurveValueSnapInterval() const
 {
-	return Settings->GetSnapCurveValueToInterval()
-		? Settings->GetCurveValueSnapInterval()
+	return SequencerSettings->GetSnapCurveValueToInterval()
+		? SequencerSettings->GetCurveValueSnapInterval()
 		: 0;
 }
 
-bool SSequencerCurveEditor::GetShowCurveEditorCurveToolTips() const
-{
-	return Settings->GetShowCurveEditorCurveToolTips();
-}
-
 void SSequencerCurveEditor::NodeTreeSelectionChanged()
 {
 	if (SequencerNodeTree.IsValid())
 	{
-		if (Settings->GetCurveVisibility() == ESequencerCurveVisibility::SelectedCurves)
+		if (GetSettings()->GetCurveVisibility() == ECurveEditorCurveVisibility::SelectedCurves)
 		{
 			UpdateCurveOwner();
 		}
 	}
 }
 
-void SSequencerCurveEditor::SequencerCurveVisibilityChanged()
+void SSequencerCurveEditor::OnCurveEditorCurveVisibilityChanged()
 {
 	UpdateCurveOwner();
 }
@@ -160,6 +155,8 @@ SSequencerCurveEditor::~SSequencerCurveEditor()
 	{
 		Sequencer.Pin()->GetSelection().GetOnOutlinerNodeSelectionChanged().Remove( NodeTreeSelectionChangedHandle );
 	}
+
+	GetSettings()->GetOnCurveEditorCurveVisibilityChanged().RemoveAll(this);
 }
 
 #undef LOCTEXT_NAMESPACE
\ No newline at end of file
diff --git a/Engine/Source/Editor/Sequencer/Private/SSequencerCurveEditor.h b/Engine/Source/Editor/Sequencer/Private/SSequencerCurveEditor.h
--- a/Engine/Source/Editor/Sequencer/Private/SSequencerCurveEditor.h
+++ b/Engine/Source/Editor/Sequencer/Private/SSequencerCurveEditor.h
@@ -47,21 +47,18 @@ private:
 	/** Gets the value snapping interval. */
 	float GetCurveValueSnapInterval() const;
 
-	/** Gets whether or not to display curve tooltips. */
-	bool GetShowCurveEditorCurveToolTips() const;
-
 	/** Run whenever the selection on the FSequencerNodeTree changes. */
 	void NodeTreeSelectionChanged();
 
 	/** Run whenever the selected curve visibility changes. */
-	void SequencerCurveVisibilityChanged();
+	void OnCurveEditorCurveVisibilityChanged();
 
 private:
 
 	/** The sequencer which owns this widget. */
 	TWeakPtr<FSequencer> Sequencer;
 	/** Cached user-supplied settings object */
-	USequencerSettings* Settings; 
+	USequencerSettings* SequencerSettings; 
 	/** The class responsible for time sliding on the curve editor */
 	TSharedPtr<class ITimeSliderController> TimeSliderController;
 	/** The sequencer node tree which contains the key area nodes which supply the curves to edit. */
diff --git a/Engine/Source/Editor/Sequencer/Private/SSequencerCurveEditorToolBar.cpp b/Engine/Source/Editor/Sequencer/Private/SSequencerCurveEditorToolBar.cpp
--- a/Engine/Source/Editor/Sequencer/Private/SSequencerCurveEditorToolBar.cpp
+++ b/Engine/Source/Editor/Sequencer/Private/SSequencerCurveEditorToolBar.cpp
@@ -3,9 +3,20 @@
 #include "SequencerPrivatePCH.h"
 #include "RichCurveEditorCommands.h"
 
+#define LOCTEXT_NAMESPACE "CurveEditorToolBar"
+
 void SSequencerCurveEditorToolBar::Construct( const FArguments& InArgs, TSharedPtr<FUICommandList> CurveEditorCommandList )
 {
 	FToolBarBuilder ToolBarBuilder( CurveEditorCommandList, FMultiBoxCustomization::None, TSharedPtr<FExtender>(), Orient_Horizontal, true );
+
+	ToolBarBuilder.AddComboButton(
+		FUIAction(),
+		FOnGetContent::CreateSP( this, &SSequencerCurveEditorToolBar::MakeCurveEditorViewOptionsMenu, CurveEditorCommandList ),
+		LOCTEXT( "CurveEditorViewOptions", "View Options" ),
+		LOCTEXT( "CurveEditorViewOptionsToolTip", "View Options" ),
+		TAttribute<FSlateIcon>(),
+		true );
+
 	ToolBarBuilder.BeginSection( "Curve" );
 	{
 		ToolBarBuilder.AddToolBarButton( FRichCurveEditorCommands::Get().ZoomToFitHorizontal );
@@ -24,8 +35,97 @@ void SSequencerCurveEditorToolBar::Construct( const FArguments& InArgs, TSharedP
 	}
 	ToolBarBuilder.EndSection();
 
+	ToolBarBuilder.BeginSection( "Tangents" );
+	{
+		ToolBarBuilder.AddToolBarButton( FRichCurveEditorCommands::Get().FlattenTangents );
+		ToolBarBuilder.AddToolBarButton( FRichCurveEditorCommands::Get().StraightenTangents );
+	}
+	ToolBarBuilder.EndSection();
+
+	ToolBarBuilder.AddComboButton(
+		FUIAction(),
+		FOnGetContent::CreateSP( this, &SSequencerCurveEditorToolBar::MakeCurveEditorCurveOptionsMenu, CurveEditorCommandList ),
+		LOCTEXT( "CurveEditorCurveOptions", "Curves Options" ),
+		LOCTEXT( "CurveEditorCurveOptionsToolTip", "Curve Options" ),
+		TAttribute<FSlateIcon>(),
+		true );
+
 	ChildSlot
 	[
 		ToolBarBuilder.MakeWidget()
 	];
+}
+
+TSharedRef<SWidget> SSequencerCurveEditorToolBar::MakeCurveEditorViewOptionsMenu(TSharedPtr<FUICommandList> CurveEditorCommandList)
+{
+	FMenuBuilder MenuBuilder( true, CurveEditorCommandList );
+
+	MenuBuilder.BeginSection( "CurveVisibility", LOCTEXT( "CurveEditorMenuCurveVisibilityHeader", "Curve Visibility" ) );
+	{
+		MenuBuilder.AddMenuEntry( FRichCurveEditorCommands::Get().SetAllCurveVisibility );
+		MenuBuilder.AddMenuEntry( FRichCurveEditorCommands::Get().SetSelectedCurveVisibility );
+		MenuBuilder.AddMenuEntry( FRichCurveEditorCommands::Get().SetAnimatedCurveVisibility );
+	}
+	MenuBuilder.EndSection();
+
+	MenuBuilder.BeginSection( "TangentVisibility", LOCTEXT( "CurveEditorMenuTangentVisibilityHeader", "Tangent Visibility" ) );
+	{
+		MenuBuilder.AddMenuEntry( FRichCurveEditorCommands::Get().SetAllTangentsVisibility );
+		MenuBuilder.AddMenuEntry( FRichCurveEditorCommands::Get().SetSelectedKeysTangentVisibility );
+		MenuBuilder.AddMenuEntry( FRichCurveEditorCommands::Get().SetNoTangentsVisibility );
+	}
+	MenuBuilder.EndSection();
+
+	MenuBuilder.AddMenuSeparator();
+
+	MenuBuilder.AddMenuEntry( FRichCurveEditorCommands::Get().ToggleAutoFrameCurveEditor);
+	MenuBuilder.AddMenuEntry( FRichCurveEditorCommands::Get().ToggleShowCurveEditorCurveToolTips);
+
+	return MenuBuilder.MakeWidget();
+}
+
+TSharedRef<SWidget> SSequencerCurveEditorToolBar::MakeCurveEditorCurveOptionsMenu(TSharedPtr<FUICommandList> CurveEditorCommandList)
+{
+	struct FExtrapolationMenus
+	{
+		static void MakePreInfinityExtrapSubMenu(FMenuBuilder& MenuBuilder)
+		{
+			MenuBuilder.BeginSection( "Pre-Infinity Extrapolation", LOCTEXT( "CurveEditorMenuPreInfinityExtrapHeader", "Extrapolation" ) );
+			{
+				MenuBuilder.AddMenuEntry( FRichCurveEditorCommands::Get().SetPreInfinityExtrapCycle);
+				MenuBuilder.AddMenuEntry( FRichCurveEditorCommands::Get().SetPreInfinityExtrapCycleWithOffset);
+				MenuBuilder.AddMenuEntry( FRichCurveEditorCommands::Get().SetPreInfinityExtrapOscillate);
+				MenuBuilder.AddMenuEntry( FRichCurveEditorCommands::Get().SetPreInfinityExtrapLinear);
+				MenuBuilder.AddMenuEntry( FRichCurveEditorCommands::Get().SetPreInfinityExtrapConstant);
+			}
+			MenuBuilder.EndSection();
+		}
+
+		static void MakePostInfinityExtrapSubMenu(FMenuBuilder& MenuBuilder)
+		{
+			MenuBuilder.BeginSection( "Post-Infinity Extrapolation", LOCTEXT( "CurveEditorMenuPostInfinityExtrapHeader", "Extrapolation" ) );
+			{
+				MenuBuilder.AddMenuEntry( FRichCurveEditorCommands::Get().SetPostInfinityExtrapCycle);
+				MenuBuilder.AddMenuEntry( FRichCurveEditorCommands::Get().SetPostInfinityExtrapCycleWithOffset);
+				MenuBuilder.AddMenuEntry( FRichCurveEditorCommands::Get().SetPostInfinityExtrapOscillate);
+				MenuBuilder.AddMenuEntry( FRichCurveEditorCommands::Get().SetPostInfinityExtrapLinear);
+				MenuBuilder.AddMenuEntry( FRichCurveEditorCommands::Get().SetPostInfinityExtrapConstant);
+			}
+			MenuBuilder.EndSection();
+		}
+	};
+
+	FMenuBuilder MenuBuilder( true, CurveEditorCommandList );
+
+	MenuBuilder.AddSubMenu(
+		LOCTEXT( "PreInfinitySubMenu", "Pre-Infinity" ),
+		LOCTEXT( "PreInfinitySubMenuToolTip", "Pre-Infinity Extrapolation" ),
+		FNewMenuDelegate::CreateStatic( &FExtrapolationMenus::MakePreInfinityExtrapSubMenu ) );
+
+	MenuBuilder.AddSubMenu(
+		LOCTEXT( "PostInfinitySubMenu", "Post-Infinity" ),
+		LOCTEXT( "PostInfinitySubMenuToolTip", "Post-Infinity Extrapolation" ),
+		FNewMenuDelegate::CreateStatic( &FExtrapolationMenus::MakePostInfinityExtrapSubMenu ) );
+	
+	return MenuBuilder.MakeWidget();
 }
\ No newline at end of file
diff --git a/Engine/Source/Editor/Sequencer/Private/SSequencerCurveEditorToolBar.h b/Engine/Source/Editor/Sequencer/Private/SSequencerCurveEditorToolBar.h
--- a/Engine/Source/Editor/Sequencer/Private/SSequencerCurveEditorToolBar.h
+++ b/Engine/Source/Editor/Sequencer/Private/SSequencerCurveEditorToolBar.h
@@ -12,4 +12,10 @@ public:
 	SLATE_END_ARGS()
 
 	void Construct( const FArguments& InArgs, TSharedPtr<FUICommandList> CurveEditorCommandList );
+
+private:
+	/** Makes the curve editor view options menu for the toolbar. */
+	TSharedRef<SWidget> MakeCurveEditorViewOptionsMenu(TSharedPtr<FUICommandList> CurveEditorCommandList);
+	/** Makes the curve editor curve options menu for the toolbar. */
+	TSharedRef<SWidget> MakeCurveEditorCurveOptionsMenu(TSharedPtr<FUICommandList> CurveEditorCommandList);
 };
\ No newline at end of file
diff --git a/Engine/Source/Editor/Sequencer/Private/Sequencer.cpp b/Engine/Source/Editor/Sequencer/Private/Sequencer.cpp
--- a/Engine/Source/Editor/Sequencer/Private/Sequencer.cpp
+++ b/Engine/Source/Editor/Sequencer/Private/Sequencer.cpp
@@ -1709,30 +1709,6 @@ void FSequencer::BindSequencerCommands()
 		FCanExecuteAction::CreateLambda( []{ return true; } ),
 		FIsActionChecked::CreateLambda( [this]{ return Settings->GetShowCurveEditor(); } ) );
 
-	SequencerCommandBindings->MapAction(
-		Commands.ToggleShowCurveEditorCurveToolTips,
-		FExecuteAction::CreateLambda( [this]{ Settings->SetShowCurveEditorCurveToolTips( !Settings->GetShowCurveEditorCurveToolTips() ); } ),
-		FCanExecuteAction::CreateLambda( []{ return true; } ),
-		FIsActionChecked::CreateLambda( [this]{ return Settings->GetShowCurveEditorCurveToolTips(); } ) );
-
-	SequencerCommandBindings->MapAction(
-		Commands.SetAllCurveVisibility,
-		FExecuteAction::CreateLambda( [this]{ Settings->SetCurveVisibility( ESequencerCurveVisibility::AllCurves ); } ),
-		FCanExecuteAction::CreateLambda( []{ return true; } ),
-		FIsActionChecked::CreateLambda( [this]{ return Settings->GetCurveVisibility() == ESequencerCurveVisibility::AllCurves; } ) );
-
-	SequencerCommandBindings->MapAction(
-		Commands.SetSelectedCurveVisibility,
-		FExecuteAction::CreateLambda( [this]{ Settings->SetCurveVisibility( ESequencerCurveVisibility::SelectedCurves ); } ),
-		FCanExecuteAction::CreateLambda( []{ return true; } ),
-		FIsActionChecked::CreateLambda( [this]{ return Settings->GetCurveVisibility() == ESequencerCurveVisibility::SelectedCurves; } ) );
-
-	SequencerCommandBindings->MapAction(
-		Commands.SetAnimatedCurveVisibility,
-		FExecuteAction::CreateLambda( [this]{ Settings->SetCurveVisibility( ESequencerCurveVisibility::AnimatedCurves ); } ),
-		FCanExecuteAction::CreateLambda( []{ return true; } ),
-		FIsActionChecked::CreateLambda( [this]{ return Settings->GetCurveVisibility() == ESequencerCurveVisibility::AnimatedCurves; } ) );
-
 	for (int32 i = 0; i < TrackEditors.Num(); ++i)
 	{
 		TrackEditors[i]->BindCommands(SequencerCommandBindings);
diff --git a/Engine/Source/Editor/Sequencer/Private/SequencerCommands.cpp b/Engine/Source/Editor/Sequencer/Private/SequencerCommands.cpp
--- a/Engine/Source/Editor/Sequencer/Private/SequencerCommands.cpp
+++ b/Engine/Source/Editor/Sequencer/Private/SequencerCommands.cpp
@@ -33,15 +33,11 @@ void FSequencerCommands::RegisterCommands()
 	UI_COMMAND( ToggleSnapPlayTimeToInterval, "Snap to the interval while scrubbing", "Snap the play time to the interval while scrubbing.", EUserInterfaceActionType::ToggleButton, FInputChord() );
 	UI_COMMAND( ToggleSnapPlayTimeToDraggedKey, "Snap to the dragged key", "Snap the play time to the dragged key.", EUserInterfaceActionType::ToggleButton, FInputChord() );
 
-	UI_COMMAND(ToggleSnapCurveValueToInterval, "Snap curve key values", "Snap curve keys to the value snapping interval.", EUserInterfaceActionType::ToggleButton, FInputChord() );
+	UI_COMMAND( ToggleSnapCurveValueToInterval, "Snap curve key values", "Snap curve keys to the value snapping interval.", EUserInterfaceActionType::ToggleButton, FInputChord() );
 
 	UI_COMMAND( ToggleCleanView, "Clean View", "Enable 'Clean View' mode which displays only global tracks when no filter is applied.", EUserInterfaceActionType::ToggleButton, FInputChord() );
 	
 	UI_COMMAND( ToggleShowCurveEditor, "Curve Editor", "Show the animation keys in a curve editor.", EUserInterfaceActionType::ToggleButton, FInputChord() );
-	UI_COMMAND( ToggleShowCurveEditorCurveToolTips, "Curve Tool Tips", "Show a tool tip with name and values when hovering over a curve.", EUserInterfaceActionType::ToggleButton, FInputChord() );
-	UI_COMMAND( SetAllCurveVisibility, "All Curves", "Show all curves in the curve editor.", EUserInterfaceActionType::RadioButton, FInputChord() );
-	UI_COMMAND( SetSelectedCurveVisibility, "Selected Curves", "Show only selected curves in the curve editor.", EUserInterfaceActionType::RadioButton, FInputChord() );
-	UI_COMMAND( SetAnimatedCurveVisibility, "Animated Curves", "Show only curves with keyframes in the curve editor.", EUserInterfaceActionType::RadioButton, FInputChord() );
 }
 
 #undef LOCTEXT_NAMESPACE
\ No newline at end of file
diff --git a/Engine/Source/Editor/Sequencer/Private/SequencerCommands.h b/Engine/Source/Editor/Sequencer/Private/SequencerCommands.h
--- a/Engine/Source/Editor/Sequencer/Private/SequencerCommands.h
+++ b/Engine/Source/Editor/Sequencer/Private/SequencerCommands.h
@@ -75,18 +75,6 @@ public:
 	/** Toggles whether or not the curve editor should be shown. */
 	TSharedPtr< FUICommandInfo > ToggleShowCurveEditor;
 
-	/** Toggles whether or not to show tool tips for the curves in the curve editor. */
-	TSharedPtr< FUICommandInfo > ToggleShowCurveEditorCurveToolTips;
-
-	/** Sets the curve visibility to all curves. */
-	TSharedPtr< FUICommandInfo > SetAllCurveVisibility;
-
-	/** Sets the curve visibility to the selected curves. */
-	TSharedPtr< FUICommandInfo > SetSelectedCurveVisibility;
-
-	/** Sets the curve visibility to the animated curves. */
-	TSharedPtr< FUICommandInfo > SetAnimatedCurveVisibility;
-
 	/**
 	 * Initialize commands
 	 */
diff --git a/Engine/Source/Editor/Sequencer/Private/SequencerCurveOwner.cpp b/Engine/Source/Editor/Sequencer/Private/SequencerCurveOwner.cpp
--- a/Engine/Source/Editor/Sequencer/Private/SequencerCurveOwner.cpp
+++ b/Engine/Source/Editor/Sequencer/Private/SequencerCurveOwner.cpp
@@ -42,7 +42,7 @@ FName BuildCurveName( TSharedPtr<FSectionKeyAreaNode> KeyAreaNode )
 	return FName(*FString::Join(NameParts, TEXT(" - ")));
 }
 
-FSequencerCurveOwner::FSequencerCurveOwner( TSharedPtr<FSequencerNodeTree> InSequencerNodeTree )
+FSequencerCurveOwner::FSequencerCurveOwner( TSharedPtr<FSequencerNodeTree> InSequencerNodeTree, ECurveEditorCurveVisibility::Type CurveVisibility )
 {
 	SequencerNodeTree = InSequencerNodeTree;
 
@@ -56,18 +56,19 @@ FSequencerCurveOwner::FSequencerCurveOwner( TSharedPtr<FSequencerNodeTree> InSeq
 			if ( RichCurve != nullptr )
 			{
 				bool bAddCurve = false;
-				switch ( SequencerNodeTree->GetSequencer().GetSettings()->GetCurveVisibility() )
+				switch ( CurveVisibility )
 				{
-				case ESequencerCurveVisibility::AllCurves:
+				case ECurveEditorCurveVisibility::AllCurves:
 					bAddCurve = true;
 					break;
-				case ESequencerCurveVisibility::SelectedCurves:
+				case ECurveEditorCurveVisibility::SelectedCurves:
 					bAddCurve = KeyAreaNode->GetSequencer().GetSelection().IsSelected(KeyAreaNode);
 					break;
-				case ESequencerCurveVisibility::AnimatedCurves:
+				case ECurveEditorCurveVisibility::AnimatedCurves:
 					bAddCurve = RichCurve->GetNumKeys() > 0;
 					break;
 				}
+				
 				if ( bAddCurve )
 				{
 					FName CurveName = BuildCurveName(KeyAreaNode);
diff --git a/Engine/Source/Editor/Sequencer/Private/SequencerCurveOwner.h b/Engine/Source/Editor/Sequencer/Private/SequencerCurveOwner.h
--- a/Engine/Source/Editor/Sequencer/Private/SequencerCurveOwner.h
+++ b/Engine/Source/Editor/Sequencer/Private/SequencerCurveOwner.h
@@ -2,11 +2,13 @@
 
 #pragma once
 
+#include "CurveEditorSettings.h"
+
 /** A curve owner interface for displaying animation curves in sequencer. */
 class FSequencerCurveOwner : public FCurveOwnerInterface
 {
 public:
-	FSequencerCurveOwner(TSharedPtr<FSequencerNodeTree> InSequencerNodeTree);
+	FSequencerCurveOwner(TSharedPtr<FSequencerNodeTree> InSequencerNodeTree, ECurveEditorCurveVisibility::Type CurveVisibility);
 
 	/** FCurveOwnerInterface */
 	virtual TArray<FRichCurveEditInfoConst> GetCurves() const override;
diff --git a/Engine/Source/Editor/Sequencer/Private/SequencerPrivatePCH.h b/Engine/Source/Editor/Sequencer/Private/SequencerPrivatePCH.h
--- a/Engine/Source/Editor/Sequencer/Private/SequencerPrivatePCH.h
+++ b/Engine/Source/Editor/Sequencer/Private/SequencerPrivatePCH.h
@@ -26,6 +26,7 @@
 #include "SequencerCurveOwner.h"
 #include "SSequencerCurveEditor.h"
 #include "SSequencerCurveEditorToolBar.h"
+#include "CurveEditorSettings.h"
 
 DECLARE_LOG_CATEGORY_EXTERN(LogSequencer, Log, All);
 
diff --git a/Engine/Source/Editor/Sequencer/Private/SequencerSettings.cpp b/Engine/Source/Editor/Sequencer/Private/SequencerSettings.cpp
--- a/Engine/Source/Editor/Sequencer/Private/SequencerSettings.cpp
+++ b/Engine/Source/Editor/Sequencer/Private/SequencerSettings.cpp
@@ -40,7 +40,6 @@ USequencerSettings::USequencerSettings( const FObjectInitializer& ObjectInitiali
 	bAutoScrollEnabled = true;
 	bShowCurveEditor = false;
 	bShowCurveEditorCurveToolTips = true;
-	CurveVisibility = ESequencerCurveVisibility::AllCurves;
 }
 
 bool USequencerSettings::GetIsSnapEnabled() const
@@ -222,7 +221,6 @@ void USequencerSettings::SetShowCurveEditor(bool InbShowCurveEditor)
 	if (bShowCurveEditor != InbShowCurveEditor)
 	{
 		bShowCurveEditor = InbShowCurveEditor;
-		OnShowCurveEditorChanged.Broadcast();
 		SaveConfig();
 	}
 }
@@ -241,34 +239,9 @@ void USequencerSettings::SetShowCurveEditorCurveToolTips(bool InbShowCurveEditor
 	}
 }
 
-ESequencerCurveVisibility::Type USequencerSettings::GetCurveVisibility() const
-{
-	return CurveVisibility;
-}
-
-void USequencerSettings::SetCurveVisibility(ESequencerCurveVisibility::Type InCurveVisibility)
-{
-	if (CurveVisibility != InCurveVisibility)
-	{
-		CurveVisibility = InCurveVisibility;
-		OnCurveVisibilityChanged.Broadcast();
-		SaveConfig();
-	}
-}
-
 float USequencerSettings::SnapTimeToInterval( float InTimeValue ) const
 {
 	return TimeSnapInterval > 0
 		? FMath::RoundToInt( InTimeValue / TimeSnapInterval ) * TimeSnapInterval
 		: InTimeValue;
 }
-
-USequencerSettings::FOnCurveVisibilityChanged& USequencerSettings::GetOnCurveVisibilityChanged()
-{
-	return OnCurveVisibilityChanged;
-}
-
-USequencerSettings::FOnShowCurveEditorChanged& USequencerSettings::GetOnShowCurveEditorChanged()
-{
-	return OnShowCurveEditorChanged;
-}
diff --git a/Engine/Source/Editor/Sequencer/Private/SequencerSettings.h b/Engine/Source/Editor/Sequencer/Private/SequencerSettings.h
--- a/Engine/Source/Editor/Sequencer/Private/SequencerSettings.h
+++ b/Engine/Source/Editor/Sequencer/Private/SequencerSettings.h
@@ -4,21 +4,6 @@
 
 #include "SequencerSettings.generated.h"
 
-/** Defines visibility states for the curves in the curve editor. */
-UENUM()
-namespace ESequencerCurveVisibility
-{
-	enum Type
-	{
-		/** All curves should be visible. */
-		AllCurves,
-		/** Only curves from selected nodes should be visible. */
-		SelectedCurves,
-		/** Only curves which have keyframes should be visible. */
-		AnimatedCurves
-	};
-}
-
 /** Empty class used to house multiple named USequencerSettings */
 UCLASS()
 class USequencerSettingsContainer : public UObject
@@ -37,9 +22,6 @@ class USequencerSettings : public UObject
 public:
 	GENERATED_UCLASS_BODY()
 
-	DECLARE_MULTICAST_DELEGATE( FOnCurveVisibilityChanged );
-	DECLARE_MULTICAST_DELEGATE( FOnShowCurveEditorChanged );
-
 	/** Gets whether or not snapping is enabled. */
 	bool GetIsSnapEnabled() const;
 	/** Sets whether or not snapping is enabled. */
@@ -110,20 +92,9 @@ public:
 	/** Sets whether or not to show curve tool tips in the curve editor. */
 	void SetShowCurveEditorCurveToolTips(bool InbShowCurveEditorCurveToolTips);
 
-	/** Gets the current curve visibility. */
-	ESequencerCurveVisibility::Type GetCurveVisibility() const;
-	/** Sets the current curve visibility. */
-	void SetCurveVisibility(ESequencerCurveVisibility::Type InCurveVisibility);
-
 	/** Snaps a time value in seconds to the currently selected interval. */
 	float SnapTimeToInterval(float InTimeValue) const;
 
-	/** Gets the multicast delegate which is run whenever the CurveVisibility value changes. */
-	FOnCurveVisibilityChanged& GetOnCurveVisibilityChanged();
-
-	/** Gets the multicast delegate which is run whenever the curve editor is shown/hidden. */
-	FOnShowCurveEditorChanged& GetOnShowCurveEditorChanged();
-
 protected:
 	UPROPERTY( config )
 	bool bIsSnapEnabled;
@@ -166,10 +137,4 @@ protected:
 
 	UPROPERTY( config )
 	bool bShowCurveEditorCurveToolTips;
-
-	UPROPERTY( config )
-	TEnumAsByte<ESequencerCurveVisibility::Type> CurveVisibility;
-
-	FOnCurveVisibilityChanged OnCurveVisibilityChanged;
-	FOnShowCurveEditorChanged OnShowCurveEditorChanged;
 };
\ No newline at end of file
diff --git a/Engine/Source/Editor/UnrealEd/Private/RichCurveEditorCommands.cpp b/Engine/Source/Editor/UnrealEd/Private/RichCurveEditorCommands.cpp
--- a/Engine/Source/Editor/UnrealEd/Private/RichCurveEditorCommands.cpp
+++ b/Engine/Source/Editor/UnrealEd/Private/RichCurveEditorCommands.cpp
@@ -18,6 +18,31 @@ void FRichCurveEditorCommands::RegisterCommands()
 	UI_COMMAND(InterpolationCubicAuto, "Auto", "Cubic interpolation - Automatic tangents", EUserInterfaceActionType::ToggleButton, FInputChord());
 	UI_COMMAND(InterpolationCubicUser, "User", "Cubic interpolation - User flat tangents", EUserInterfaceActionType::ToggleButton, FInputChord());
 	UI_COMMAND(InterpolationCubicBreak, "Break", "Cubic interpolation - User broken tangents", EUserInterfaceActionType::ToggleButton, FInputChord());
+
+	UI_COMMAND(FlattenTangents, "Flatten", "Flatten tangents", EUserInterfaceActionType::Button, FInputChord());
+	UI_COMMAND(StraightenTangents, "Straighten", "Straighten tangents", EUserInterfaceActionType::Button, FInputChord());
+
+	UI_COMMAND(SetPreInfinityExtrapCycle, "Cycle", "Set the pre-infinity curve extrapolation to cycle.", EUserInterfaceActionType::RadioButton, FInputChord());
+	UI_COMMAND(SetPreInfinityExtrapCycleWithOffset, "Cycle with Offset", "Set the pre-infinity curve extrapolation to cycle with offset.", EUserInterfaceActionType::RadioButton, FInputChord());
+	UI_COMMAND(SetPreInfinityExtrapOscillate, "Oscillate", "Set the pre-infinity curve extrapolation to oscillate.", EUserInterfaceActionType::RadioButton, FInputChord());
+	UI_COMMAND(SetPreInfinityExtrapLinear, "Linear", "Set the pre-infinity curve extrapolation to linear.", EUserInterfaceActionType::RadioButton, FInputChord());
+	UI_COMMAND(SetPreInfinityExtrapConstant, "Constant", "Set the pre-infinity curve extrapolation to constant.", EUserInterfaceActionType::RadioButton, FInputChord());
+	UI_COMMAND(SetPostInfinityExtrapCycle, "Cycle", "Set the post-infinity curve extrapolation to cycle.", EUserInterfaceActionType::RadioButton, FInputChord());
+	UI_COMMAND(SetPostInfinityExtrapCycleWithOffset, "Cycle with Offset", "Set the post-infinity curve extrapolation to cycle with offset.", EUserInterfaceActionType::RadioButton, FInputChord());
+	UI_COMMAND(SetPostInfinityExtrapOscillate, "Oscillate", "Set the post-infinity curve extrapolation to oscillate.", EUserInterfaceActionType::RadioButton, FInputChord());
+	UI_COMMAND(SetPostInfinityExtrapLinear, "Linear", "Set the post-infinity curve extrapolation to linear.", EUserInterfaceActionType::RadioButton, FInputChord());
+	UI_COMMAND(SetPostInfinityExtrapConstant, "Constant", "Set the post-infinity curve extrapolation to constant.", EUserInterfaceActionType::RadioButton, FInputChord());
+
+	UI_COMMAND(SetAllCurveVisibility, "All Curves", "Show all curves in the curve editor.", EUserInterfaceActionType::RadioButton, FInputChord());
+	UI_COMMAND(SetSelectedCurveVisibility, "Selected Curves", "Show only selected curves in the curve editor.", EUserInterfaceActionType::RadioButton, FInputChord());
+	UI_COMMAND(SetAnimatedCurveVisibility, "Animated Curves", "Show only curves with keyframes in the curve editor.", EUserInterfaceActionType::RadioButton, FInputChord());
+
+	UI_COMMAND(SetAllTangentsVisibility, "All Tangents", "Show all tangents in the curve editor.", EUserInterfaceActionType::RadioButton, FInputChord());
+	UI_COMMAND(SetSelectedKeysTangentVisibility, "Selected Keys", "Show tangents for selected keys in the curve editor.", EUserInterfaceActionType::RadioButton, FInputChord());
+	UI_COMMAND(SetNoTangentsVisibility, "No Tangents", "Show no tangents in the curve editor.", EUserInterfaceActionType::RadioButton, FInputChord());
+
+	UI_COMMAND(ToggleAutoFrameCurveEditor, "Auto Frame Curves", "Auto frame curves when they are selected.", EUserInterfaceActionType::ToggleButton, FInputChord() );
+	UI_COMMAND(ToggleShowCurveEditorCurveToolTips, "Curve Tool Tips", "Show a tool tip with name and values when hovering over a curve.", EUserInterfaceActionType::ToggleButton, FInputChord() );
 }
 
 #undef LOCTEXT_NAMESPACE
\ No newline at end of file
diff --git a/Engine/Source/Editor/UnrealEd/Private/SCurveEditor.cpp b/Engine/Source/Editor/UnrealEd/Private/SCurveEditor.cpp
--- a/Engine/Source/Editor/UnrealEd/Private/SCurveEditor.cpp
+++ b/Engine/Source/Editor/UnrealEd/Private/SCurveEditor.cpp
@@ -8,6 +8,7 @@
 #include "SColorGradientEditor.h"
 #include "GenericCommands.h"
 #include "SNumericEntryBox.h"
+#include "CurveEditorSettings.h"
 
 #define LOCTEXT_NAMESPACE "SCurveEditor"
 
@@ -29,7 +30,7 @@ void SCurveEditor::Construct(const FArguments& InArgs)
 	CurveFactory = NULL;
 	Commands = TSharedPtr< FUICommandList >(new FUICommandList);
 	CurveOwner = NULL;
-
+	
 	// view input
 	ViewMinInput = InArgs._ViewMinInput;
 	ViewMaxInput = InArgs._ViewMaxInput;
@@ -69,13 +70,14 @@ void SCurveEditor::Construct(const FArguments& InArgs)
 	bShowCurveSelector = InArgs._ShowCurveSelector;
 	bDrawInputGridNumbers = InArgs._ShowInputGridNumbers;
 	bDrawOutputGridNumbers = InArgs._ShowOutputGridNumbers;
-	bShowCurveToolTips = InArgs._ShowCurveToolTips;
 
 	OnCreateAsset = InArgs._OnCreateAsset;
 
 	DragState = EDragState::None;
 	DragThreshold = 4;
 
+	MovementAxisLock = EMovementAxisLock::None;
+
 	//Simple r/g/b for now
 	CurveColors.Add(FLinearColor(1.0f, 0.0f, 0.0f));
 	CurveColors.Add(FLinearColor(0.0f, 1.0f, 0.0f));
@@ -83,6 +85,8 @@ void SCurveEditor::Construct(const FArguments& InArgs)
 
 	TransactionIndex = -1;
 
+	Settings = GetMutableDefault<UCurveEditorSettings>();
+
 	Commands->MapAction(FGenericCommands::Get().Undo,
 		FExecuteAction::CreateSP(this, &SCurveEditor::UndoAction));
 
@@ -103,6 +107,7 @@ void SCurveEditor::Construct(const FArguments& InArgs)
 		FCanExecuteAction(),
 		FIsActionChecked::CreateSP(this, &SCurveEditor::IsSnappingEnabled));
 
+	// Interpolation
 	Commands->MapAction(FRichCurveEditorCommands::Get().InterpolationConstant,
 		FExecuteAction::CreateSP(this, &SCurveEditor::OnSelectInterpolationMode, RCIM_Constant, RCTM_Auto),
 		FCanExecuteAction(),
@@ -128,6 +133,107 @@ void SCurveEditor::Construct(const FArguments& InArgs)
 		FCanExecuteAction(),
 		FIsActionChecked::CreateSP(this, &SCurveEditor::IsInterpolationModeSelected, RCIM_Cubic, RCTM_Break));
 
+	// Tangents
+	Commands->MapAction(FRichCurveEditorCommands::Get().FlattenTangents,
+		FExecuteAction::CreateSP(this, &SCurveEditor::OnFlattenOrStraightenTangents, true));
+
+	Commands->MapAction(FRichCurveEditorCommands::Get().StraightenTangents,
+		FExecuteAction::CreateSP(this, &SCurveEditor::OnFlattenOrStraightenTangents, false));
+
+	// Pre infinity extrapolation
+	Commands->MapAction(FRichCurveEditorCommands::Get().SetPreInfinityExtrapCycle,
+		FExecuteAction::CreateSP(this, &SCurveEditor::OnSelectPreInfinityExtrap, RCCE_Cycle),
+		FCanExecuteAction(),
+		FIsActionChecked::CreateSP(this, &SCurveEditor::IsPreInfinityExtrapSelected, RCCE_Cycle));
+
+	Commands->MapAction(FRichCurveEditorCommands::Get().SetPreInfinityExtrapCycleWithOffset,
+		FExecuteAction::CreateSP(this, &SCurveEditor::OnSelectPreInfinityExtrap, RCCE_CycleWithOffset),
+		FCanExecuteAction(),
+		FIsActionChecked::CreateSP(this, &SCurveEditor::IsPreInfinityExtrapSelected, RCCE_CycleWithOffset));
+
+	Commands->MapAction(FRichCurveEditorCommands::Get().SetPreInfinityExtrapOscillate,
+		FExecuteAction::CreateSP(this, &SCurveEditor::OnSelectPreInfinityExtrap, RCCE_Oscillate),
+		FCanExecuteAction(),
+		FIsActionChecked::CreateSP(this, &SCurveEditor::IsPreInfinityExtrapSelected, RCCE_Oscillate));
+
+	Commands->MapAction(FRichCurveEditorCommands::Get().SetPreInfinityExtrapLinear,
+		FExecuteAction::CreateSP(this, &SCurveEditor::OnSelectPreInfinityExtrap, RCCE_Linear),
+		FCanExecuteAction(),
+		FIsActionChecked::CreateSP(this, &SCurveEditor::IsPreInfinityExtrapSelected, RCCE_Linear));
+
+	Commands->MapAction(FRichCurveEditorCommands::Get().SetPreInfinityExtrapConstant,
+		FExecuteAction::CreateSP(this, &SCurveEditor::OnSelectPreInfinityExtrap, RCCE_Constant),
+		FCanExecuteAction(),
+		FIsActionChecked::CreateSP(this, &SCurveEditor::IsPreInfinityExtrapSelected, RCCE_Constant));
+
+	// Post infinity extrapolation
+	Commands->MapAction(FRichCurveEditorCommands::Get().SetPostInfinityExtrapCycle,
+		FExecuteAction::CreateSP(this, &SCurveEditor::OnSelectPostInfinityExtrap, RCCE_Cycle),
+		FCanExecuteAction(),
+		FIsActionChecked::CreateSP(this, &SCurveEditor::IsPostInfinityExtrapSelected, RCCE_Cycle));
+
+	Commands->MapAction(FRichCurveEditorCommands::Get().SetPostInfinityExtrapCycleWithOffset,
+		FExecuteAction::CreateSP(this, &SCurveEditor::OnSelectPostInfinityExtrap, RCCE_CycleWithOffset),
+		FCanExecuteAction(),
+		FIsActionChecked::CreateSP(this, &SCurveEditor::IsPostInfinityExtrapSelected, RCCE_CycleWithOffset));
+
+	Commands->MapAction(FRichCurveEditorCommands::Get().SetPostInfinityExtrapOscillate,
+		FExecuteAction::CreateSP(this, &SCurveEditor::OnSelectPostInfinityExtrap, RCCE_Oscillate),
+		FCanExecuteAction(),
+		FIsActionChecked::CreateSP(this, &SCurveEditor::IsPostInfinityExtrapSelected, RCCE_Oscillate));
+
+	Commands->MapAction(FRichCurveEditorCommands::Get().SetPostInfinityExtrapLinear,
+		FExecuteAction::CreateSP(this, &SCurveEditor::OnSelectPostInfinityExtrap, RCCE_Linear),
+		FCanExecuteAction(),
+		FIsActionChecked::CreateSP(this, &SCurveEditor::IsPostInfinityExtrapSelected, RCCE_Linear));
+
+	Commands->MapAction(FRichCurveEditorCommands::Get().SetPostInfinityExtrapConstant,
+		FExecuteAction::CreateSP(this, &SCurveEditor::OnSelectPostInfinityExtrap, RCCE_Constant),
+		FCanExecuteAction(),
+		FIsActionChecked::CreateSP(this, &SCurveEditor::IsPostInfinityExtrapSelected, RCCE_Constant));
+
+	// Curve Visibility
+	Commands->MapAction(FRichCurveEditorCommands::Get().SetAllCurveVisibility,
+		FExecuteAction::CreateLambda( [this]{ Settings->SetCurveVisibility( ECurveEditorCurveVisibility::AllCurves ); } ),
+		FCanExecuteAction::CreateLambda( []{ return true; } ),
+		FIsActionChecked::CreateLambda( [this]{ return Settings->GetCurveVisibility() == ECurveEditorCurveVisibility::AllCurves; } ) );
+
+	Commands->MapAction(FRichCurveEditorCommands::Get().SetSelectedCurveVisibility,
+		FExecuteAction::CreateLambda( [this]{ Settings->SetCurveVisibility( ECurveEditorCurveVisibility::SelectedCurves ); } ),
+		FCanExecuteAction::CreateLambda( []{ return true; } ),
+		FIsActionChecked::CreateLambda( [this]{ return Settings->GetCurveVisibility() == ECurveEditorCurveVisibility::SelectedCurves; } ) );
+
+	Commands->MapAction(FRichCurveEditorCommands::Get().SetAnimatedCurveVisibility,
+		FExecuteAction::CreateLambda( [this]{ Settings->SetCurveVisibility( ECurveEditorCurveVisibility::AnimatedCurves ); } ),
+		FCanExecuteAction::CreateLambda( []{ return true; } ),
+		FIsActionChecked::CreateLambda( [this]{ return Settings->GetCurveVisibility() == ECurveEditorCurveVisibility::AnimatedCurves; } ) );
+
+	// Tangent Visibility
+	Commands->MapAction(FRichCurveEditorCommands::Get().SetAllTangentsVisibility,
+		FExecuteAction::CreateLambda( [this]{ Settings->SetTangentVisibility( ECurveEditorTangentVisibility::AllTangents ); } ),
+		FCanExecuteAction::CreateLambda( []{ return true; } ),
+		FIsActionChecked::CreateLambda( [this]{ return Settings->GetTangentVisibility() == ECurveEditorTangentVisibility::AllTangents; } ) );
+
+	Commands->MapAction(FRichCurveEditorCommands::Get().SetSelectedKeysTangentVisibility,
+		FExecuteAction::CreateLambda( [this]{ Settings->SetTangentVisibility( ECurveEditorTangentVisibility::SelectedKeys ); } ),
+		FCanExecuteAction::CreateLambda( []{ return true; } ),
+		FIsActionChecked::CreateLambda( [this]{ return Settings->GetTangentVisibility() == ECurveEditorTangentVisibility::SelectedKeys; } ) );
+
+	Commands->MapAction(FRichCurveEditorCommands::Get().SetNoTangentsVisibility,
+		FExecuteAction::CreateLambda( [this]{ Settings->SetTangentVisibility( ECurveEditorTangentVisibility::NoTangents ); } ),
+		FCanExecuteAction::CreateLambda( []{ return true; } ),
+		FIsActionChecked::CreateLambda( [this]{ return Settings->GetTangentVisibility() == ECurveEditorTangentVisibility::NoTangents; } ) );
+
+	Commands->MapAction(FRichCurveEditorCommands::Get().ToggleAutoFrameCurveEditor,
+		FExecuteAction::CreateLambda( [this]{ Settings->SetAutoFrameCurveEditor( !Settings->GetAutoFrameCurveEditor() ); } ),
+		FCanExecuteAction::CreateLambda( []{ return true; } ),
+		FIsActionChecked::CreateLambda( [this]{ return Settings->GetAutoFrameCurveEditor(); } ) );
+
+	Commands->MapAction(FRichCurveEditorCommands::Get().ToggleShowCurveEditorCurveToolTips,
+		FExecuteAction::CreateLambda( [this]{ Settings->SetShowCurveEditorCurveToolTips( !Settings->GetShowCurveEditorCurveToolTips() ); } ),
+		FCanExecuteAction::CreateLambda( []{ return true; } ),
+		FIsActionChecked::CreateLambda( [this]{ return Settings->GetShowCurveEditorCurveToolTips(); } ) );
+
 	SAssignNew(WarningMessageText, SErrorText);
 
 	TSharedRef<SBox> CurveSelector = SNew(SBox)
@@ -362,7 +468,7 @@ void SCurveEditor::RemoveCurveKeysFromSelection(TSharedPtr<FCurveViewModel> Curv
 	}
 	for (auto KeyToDeselect : SelectedKeysForLockedCurve)
 	{
-		RemoveFromSelection(KeyToDeselect);
+		RemoveFromKeySelection(KeyToDeselect);
 	}
 }
 
@@ -485,6 +591,13 @@ void SCurveEditor::PushKeyMenu(const FGeometry& InMyGeometry, const FPointerEven
 	}
 	MenuBuilder.EndSection(); //CurveEditorInterpolation
 
+	MenuBuilder.BeginSection("CurveEditorTangents", LOCTEXT("Tangents", "Tangents"));
+	{
+		MenuBuilder.AddMenuEntry(FRichCurveEditorCommands::Get().FlattenTangents);
+		MenuBuilder.AddMenuEntry(FRichCurveEditorCommands::Get().StraightenTangents);
+	}
+	MenuBuilder.EndSection(); //CurveEditorTangents
+
 	FWidgetPath WidgetPath = InMouseEvent.GetEventPath() != nullptr ? *InMouseEvent.GetEventPath() : FWidgetPath();
 	FVector2D Position = InMouseEvent.GetScreenSpacePosition();
 	FSlateApplication::Get().PushMenu(
@@ -544,6 +657,11 @@ int32 SCurveEditor::OnPaint(const FPaintArgs& Args, const FGeometry& AllottedGeo
 	// Positioning info
 	FTrackScaleInfo ScaleInfo(ViewMinInput.Get(),  ViewMaxInput.Get(), ViewMinOutput.Get(), ViewMaxOutput.Get(), CurveAreaGeometry.Size);
 
+	if (FMath::IsNearlyEqual(ViewMinInput.Get(), ViewMaxInput.Get()) || FMath::IsNearlyEqual(ViewMinOutput.Get(), ViewMaxOutput.Get()))
+	{
+		return 0;
+	}
+
 	// Draw background to indicate valid timeline area
 	float ZeroInputX = ScaleInfo.InputToLocalX(0.f);
 	float ZeroOutputY = ScaleInfo.OutputToLocalY(0.f);
@@ -608,16 +726,16 @@ int32 SCurveEditor::OnPaint(const FPaintArgs& Args, const FGeometry& AllottedGeo
 		for ( auto CurveViewModel : CurveViewModels )
 		{
 			if (CurveViewModel->bIsVisible)
-				{
+			{
 				PaintCurve(CurveViewModel, CurveAreaGeometry, ScaleInfo, OutDrawElements, CurveViewModel->bIsLocked ? LockedCurveLayerID : CurveLayerId, MyClippingRect, DrawEffects, InWidgetStyle);
-				}
 			}
+		}
 
 		//Paint the keys on top of the curve
 		for (auto CurveViewModel : CurveViewModels)
 		{
 			if (CurveViewModel->bIsVisible)
-		{
+			{
 				PaintKeys(CurveViewModel, ScaleInfo, OutDrawElements, KeyLayerId, SelectedKeyLayerId, CurveAreaGeometry, MyClippingRect, DrawEffects, InWidgetStyle);
 			}
 		}
@@ -784,9 +902,12 @@ void SCurveEditor::PaintKeys(TSharedPtr<FCurveViewModel> CurveViewModel, FTrackS
 			);
 
 		//Handle drawing the tangent controls for curve
-		if(IsSelected && (Curve->GetKeyInterpMode(KeyHandle) == RCIM_Cubic || LastInterpMode == RCIM_Cubic))
+		bool bIsTangentSelected = false;
+		bool bIsArrivalSelected = false;
+		bool bIsLeaveSelected = false;
+		if(IsTangentVisible(Curve, KeyHandle, bIsTangentSelected, bIsArrivalSelected, bIsLeaveSelected) && (Curve->GetKeyInterpMode(KeyHandle) == RCIM_Cubic || LastInterpMode == RCIM_Cubic))
 		{
-			PaintTangent(ScaleInfo, Curve, KeyHandle, KeyLocation, OutDrawElements, LayerId, AllottedGeometry, MyClippingRect, DrawEffects, LayerToUse, InWidgetStyle);
+			PaintTangent(ScaleInfo, Curve, KeyHandle, KeyLocation, OutDrawElements, LayerId, AllottedGeometry, MyClippingRect, DrawEffects, LayerToUse, InWidgetStyle, bIsTangentSelected, bIsArrivalSelected, bIsLeaveSelected);
 		}
 
 		LastInterpMode = Curve->GetKeyInterpMode(KeyHandle);
@@ -794,7 +915,7 @@ void SCurveEditor::PaintKeys(TSharedPtr<FCurveViewModel> CurveViewModel, FTrackS
 }
 
 
-void SCurveEditor::PaintTangent( FTrackScaleInfo &ScaleInfo, FRichCurve* Curve, FKeyHandle KeyHandle, FVector2D KeyLocation, FSlateWindowElementList &OutDrawElements, int32 LayerId, const FGeometry &AllottedGeometry, const FSlateRect& MyClippingRect, ESlateDrawEffect::Type DrawEffects, int32 LayerToUse, const FWidgetStyle &InWidgetStyle ) const
+void SCurveEditor::PaintTangent( FTrackScaleInfo &ScaleInfo, FRichCurve* Curve, FKeyHandle KeyHandle, FVector2D KeyLocation, FSlateWindowElementList &OutDrawElements, int32 LayerId, const FGeometry &AllottedGeometry, const FSlateRect& MyClippingRect, ESlateDrawEffect::Type DrawEffects, int32 LayerToUse, const FWidgetStyle &InWidgetStyle, bool bTangentSelected, bool bIsArrivalSelected, bool bIsLeaveSelected ) const
 {
 	FVector2D ArriveTangentLocation, LeaveTangentLocation;
 	GetTangentPoints(ScaleInfo, FSelectedCurveKey(Curve,KeyHandle), ArriveTangentLocation, LeaveTangentLocation);
@@ -803,7 +924,12 @@ void SCurveEditor::PaintTangent( FTrackScaleInfo &ScaleInfo, FRichCurve* Curve,
 	FVector2D LeaveTangentIconLocation = LeaveTangentLocation - (CONST_TangentSize / 2);
 
 	const FSlateBrush* TangentBrush = FEditorStyle::GetBrush("CurveEd.Tangent");
+	const FSlateBrush* TangentBrushSelected = FEditorStyle::GetBrush("CurveEd.TangentSelected");
 	const FLinearColor TangentColor = FEditorStyle::GetColor("CurveEd.TangentColor");
+	const FLinearColor TangentColorSelected = FEditorStyle::GetColor("CurveEd.TangentColorSelected");
+
+	bool LeaveTangentSelected = bTangentSelected && bIsLeaveSelected;
+	bool ArriveTangentSelected = bTangentSelected && bIsArrivalSelected;
 
 	//Add lines from tangent control point to 'key'
 	TArray<FVector2D> LinePoints;
@@ -816,7 +942,7 @@ void SCurveEditor::PaintTangent( FTrackScaleInfo &ScaleInfo, FRichCurve* Curve,
 		LinePoints,
 		MyClippingRect,
 		DrawEffects,
-		TangentColor
+		ArriveTangentSelected ? TangentColorSelected : TangentColor
 		);
 
 	LinePoints.Empty();
@@ -829,7 +955,7 @@ void SCurveEditor::PaintTangent( FTrackScaleInfo &ScaleInfo, FRichCurve* Curve,
 		LinePoints,
 		MyClippingRect,
 		DrawEffects,
-		TangentColor
+		LeaveTangentSelected ? TangentColorSelected : TangentColor
 		);
 
 	//Arrive tangent control
@@ -837,20 +963,20 @@ void SCurveEditor::PaintTangent( FTrackScaleInfo &ScaleInfo, FRichCurve* Curve,
 		OutDrawElements,
 		LayerToUse,
 		AllottedGeometry.ToPaintGeometry(ArriveTangentIconLocation, CONST_TangentSize ),
-		TangentBrush,
+		ArriveTangentSelected ? TangentBrushSelected : TangentBrush,
 		MyClippingRect,
 		DrawEffects,
-		TangentBrush->GetTint( InWidgetStyle ) * InWidgetStyle.GetColorAndOpacityTint()
+		ArriveTangentSelected ? TangentBrushSelected->GetTint( InWidgetStyle ) : TangentBrush->GetTint( InWidgetStyle ) * InWidgetStyle.GetColorAndOpacityTint()
 		);
 	//Leave tangent control
 	FSlateDrawElement::MakeBox(
 		OutDrawElements,
 		LayerToUse,
 		AllottedGeometry.ToPaintGeometry(LeaveTangentIconLocation, CONST_TangentSize ),
-		TangentBrush,
+		LeaveTangentSelected ? TangentBrushSelected : TangentBrush,
 		MyClippingRect,
 		DrawEffects,
-		TangentBrush->GetTint( InWidgetStyle ) * InWidgetStyle.GetColorAndOpacityTint()
+		LeaveTangentSelected ? TangentBrushSelected->GetTint( InWidgetStyle ) : TangentBrush->GetTint( InWidgetStyle ) * InWidgetStyle.GetColorAndOpacityTint()
 		);
 }
 
@@ -1006,7 +1132,7 @@ void SCurveEditor::SetCurveOwner(FCurveOwnerInterface* InCurveOwner, bool bCanEd
 {
 	if(InCurveOwner != CurveOwner)
 	{
-		EmptySelection();
+		EmptyAllSelection();
 	}
 
 	GradientViewer->SetCurveOwner(InCurveOwner);
@@ -1030,16 +1156,19 @@ void SCurveEditor::SetCurveOwner(FCurveOwnerInterface* InCurveOwner, bool bCanEd
 		CurveOwner->MakeTransactional();
 	}
 
-	SelectedKeys.Empty();
+	ValidateSelection();
 
-	if( bZoomToFitVertical )
+	if (Settings->GetAutoFrameCurveEditor())
 	{
-		ZoomToFitVertical();
-	}
+		if( bZoomToFitVertical )
+		{
+			ZoomToFitVertical();
+		}
 
-	if ( bZoomToFitHorizontal )
-	{
-		ZoomToFitHorizontal();
+		if ( bZoomToFitHorizontal )
+		{
+			ZoomToFitHorizontal();
+		}
 	}
 
 	CurveSelectionWidget.Pin()->SetContent(CreateCurveSelectionWidget());
@@ -1099,10 +1228,13 @@ void SCurveEditor::DeleteSelectedKeys()
 FReply SCurveEditor::OnMouseButtonDown( const FGeometry& InMyGeometry, const FPointerEvent& InMouseEvent)
 {
 	const bool bLeftMouseButton = InMouseEvent.GetEffectingButton() == EKeys::LeftMouseButton;
+	const bool bMiddleMouseButton = InMouseEvent.GetEffectingButton() == EKeys::MiddleMouseButton;
 	const bool bRightMouseButton = InMouseEvent.GetEffectingButton() == EKeys::RightMouseButton;
 
 	DragState = EDragState::PreDrag;
-	if (bLeftMouseButton || bRightMouseButton)
+	MovementAxisLock = EMovementAxisLock::None;
+
+	if (bLeftMouseButton || bMiddleMouseButton || bRightMouseButton)
 	{
 		MouseDownLocation = InMyGeometry.AbsoluteToLocal(InMouseEvent.GetScreenSpacePosition());
 
@@ -1148,8 +1280,8 @@ void SCurveEditor::AddNewKey(FGeometry InMyGeometry, FVector2D ScreenPosition, T
 				FVector2D NewKeyLocation = SnapLocation(FVector2D(Input, Output));
 				FKeyHandle NewKeyHandle = SelectedCurve->AddKey(NewKeyLocation.X, NewKeyLocation.Y);
 
-				EmptySelection();
-				AddToSelection(FSelectedCurveKey(SelectedCurve, NewKeyHandle));
+				EmptyAllSelection();
+				AddToKeySelection(FSelectedCurveKey(SelectedCurve, NewKeyHandle));
 				ChangedCurveEditInfos.Add(CurveViewModel->CurveInfo);
 			}
 		}
@@ -1164,7 +1296,7 @@ void SCurveEditor::AddNewKey(FGeometry InMyGeometry, FVector2D ScreenPosition, T
 void SCurveEditor::OnMouseCaptureLost()
 {
 	// if we began a drag transaction we need to finish it to make sure undo doesn't get out of sync
-	if (DragState == EDragState::DragKey || DragState == EDragState::DragTangent)
+	if (DragState == EDragState::DragKey || DragState == EDragState::FreeDrag || DragState == EDragState::DragTangent)
 	{
 		EndDragTransaction();
 	}
@@ -1231,7 +1363,7 @@ FReply SCurveEditor::OnMouseMove( const FGeometry& InMyGeometry, const FPointerE
 
 void SCurveEditor::UpdateCurveToolTip(const FGeometry& InMyGeometry, const FPointerEvent& InMouseEvent)
 {
-	if (bShowCurveToolTips.Get())
+	if (Settings->GetShowCurveEditorCurveToolTips())
 	{
 		TSharedPtr<FCurveViewModel> HoveredCurve = HitTestCurves(InMyGeometry, InMouseEvent);
 		if (HoveredCurve.IsValid())
@@ -1288,12 +1420,18 @@ void SCurveEditor::UpdateCurveToolTip(const FGeometry& InMyGeometry, const FPoin
 
 FReply SCurveEditor::OnMouseWheel(const FGeometry& MyGeometry, const FPointerEvent& MouseEvent)
 {
-	const float ZoomDelta = -0.1f * MouseEvent.GetWheelDelta();
+	ZoomView(FVector2D(MouseEvent.GetWheelDelta(), MouseEvent.GetWheelDelta()));
+	return FReply::Handled();
+}
+
+void SCurveEditor::ZoomView(FVector2D Delta)
+{
+	const FVector2D ZoomDelta = -0.1f * Delta;
 
 	if (bAllowZoomOutput)
 	{
 		const float OutputViewSize = ViewMaxOutput.Get() - ViewMinOutput.Get();
-		const float OutputChange = OutputViewSize * ZoomDelta;
+		const float OutputChange = OutputViewSize * ZoomDelta.Y;
 
 		const float NewMinOutput = (ViewMinOutput.Get() - (OutputChange * 0.5f));
 		const float NewMaxOutput = (ViewMaxOutput.Get() + (OutputChange * 0.5f));
@@ -1303,15 +1441,13 @@ FReply SCurveEditor::OnMouseWheel(const FGeometry& MyGeometry, const FPointerEve
 
 	{
 		const float InputViewSize = ViewMaxInput.Get() - ViewMinInput.Get();
-		const float InputChange = InputViewSize * ZoomDelta;
+		const float InputChange = InputViewSize * ZoomDelta.X;
 
 		const float NewMinInput = ViewMinInput.Get() - (InputChange * 0.5f);
 		const float NewMaxInput = ViewMaxInput.Get() + (InputChange * 0.5f);
 
 		SetInputMinMax(NewMinInput, NewMaxInput);
 	}
-
-	return FReply::Handled();
 }
 
 FReply SCurveEditor::OnKeyDown( const FGeometry& MyGeometry, const FKeyEvent& InKeyEvent )
@@ -1334,25 +1470,43 @@ FReply SCurveEditor::OnKeyDown( const FGeometry& MyGeometry, const FKeyEvent& In
 void SCurveEditor::TryStartDrag(const FGeometry& InMyGeometry, const FPointerEvent& InMouseEvent)
 {
 	const bool bLeftMouseButton = InMouseEvent.IsMouseButtonDown(EKeys::LeftMouseButton);
+	const bool bMiddleMouseButton = InMouseEvent.IsMouseButtonDown(EKeys::MiddleMouseButton);
 	const bool bRightMouseButton = InMouseEvent.IsMouseButtonDown(EKeys::RightMouseButton);
 	const bool bControlDown = InMouseEvent.IsControlDown();
+	const bool bShiftDown = InMouseEvent.IsShiftDown();
+	const bool bAltDown = InMouseEvent.IsAltDown();
 
-	FVector2D DragVector = InMyGeometry.AbsoluteToLocal(InMouseEvent.GetScreenSpacePosition()) - MouseDownLocation;
+	FVector2D MousePosition = InMyGeometry.AbsoluteToLocal(InMouseEvent.GetScreenSpacePosition());
+	FVector2D DragVector = MousePosition - MouseDownLocation;
 	if (DragVector.Size() >= DragThreshold)
 	{
+		if (bShiftDown)
+		{
+			if(FMath::Abs(MousePosition.X - MouseDownLocation.X) > FMath::Abs(MousePosition.Y - MouseDownLocation.Y))
+			{
+				MovementAxisLock = EMovementAxisLock::AxisLock_Horizontal;
+			}
+			else
+			{
+				MovementAxisLock = EMovementAxisLock::AxisLock_Vertical;
+			}
+		}
+
 		if (bLeftMouseButton)
 		{
 			// Check if we should start dragging keys.
 			FSelectedCurveKey HitKey = HitTestKeys(InMyGeometry, InMyGeometry.LocalToAbsolute(MouseDownLocation));
 			if (HitKey.IsValid())
 			{
-				if (IsKeySelected(HitKey) == false)
+				EmptyTangentSelection();
+
+				if (!IsKeySelected(HitKey))
 				{
 					if (!bControlDown)
 					{
-						EmptySelection();
+						EmptyKeySelection();
 					}
-					AddToSelection(HitKey);
+					AddToKeySelection(HitKey);
 				}
 
 				BeginDragTransaction();
@@ -1374,9 +1528,30 @@ void SCurveEditor::TryStartDrag(const FGeometry& InMyGeometry, const FPointerEve
 				FSelectedTangent Tangent = HitTestCubicTangents(InMyGeometry, InMyGeometry.LocalToAbsolute(MouseDownLocation));
 				if (Tangent.IsValid())
 				{
+					EmptyKeySelection();
+
+					if (!IsTangentSelected(Tangent))
+					{
+						if (!bControlDown)
+						{
+							EmptyTangentSelection();
+						}
+						AddToTangentSelection(Tangent);
+					}
+
 					BeginDragTransaction();
-					SelectedTangent = Tangent;
 					DragState = EDragState::DragTangent;
+					PreDragTangents.Empty();
+					for (auto SelectedTangent : SelectedTangents)
+					{
+						FRichCurve* Curve = SelectedTangent.Key.Curve;
+						FKeyHandle KeyHandle = SelectedTangent.Key.KeyHandle;
+
+						float ArriveTangent = Curve->GetKey(KeyHandle).ArriveTangent;
+						float LeaveTangent = Curve->GetKey(KeyHandle).LeaveTangent;
+
+						PreDragTangents.Add(KeyHandle, FVector2D(ArriveTangent, LeaveTangent));
+					}
 				}
 				else
 				{
@@ -1385,9 +1560,53 @@ void SCurveEditor::TryStartDrag(const FGeometry& InMyGeometry, const FPointerEve
 				}
 			}
 		}
+		else if (bMiddleMouseButton)
+		{
+			if (bAltDown)
+			{
+				DragState = EDragState::Pan;
+			}
+			else if (SelectedTangents.Num())
+			{
+				BeginDragTransaction();
+				DragState = EDragState::DragTangent;
+				PreDragTangents.Empty();
+				for (auto SelectedTangent : SelectedTangents)
+				{
+					FRichCurve* Curve = SelectedTangent.Key.Curve;
+					FKeyHandle KeyHandle = SelectedTangent.Key.KeyHandle;
+
+					float ArriveTangent = Curve->GetKey(KeyHandle).ArriveTangent;
+					float LeaveTangent = Curve->GetKey(KeyHandle).LeaveTangent;
+
+					PreDragTangents.Add(KeyHandle, FVector2D(ArriveTangent, LeaveTangent));
+				}
+			}
+			else if (SelectedKeys.Num())
+			{
+				BeginDragTransaction();
+				DragState = EDragState::FreeDrag;
+				PreDragKeyLocations.Empty();
+				for (auto selectedKey : SelectedKeys)
+				{
+					PreDragKeyLocations.Add(selectedKey.KeyHandle, FVector2D
+					(
+						selectedKey.Curve->GetKeyTime(selectedKey.KeyHandle),
+						selectedKey.Curve->GetKeyValue(selectedKey.KeyHandle)
+					));
+				}
+			}
+		}
 		else if (bRightMouseButton)
 		{
-			DragState = EDragState::Pan;
+			if (bAltDown)
+			{
+				DragState = EDragState::Zoom;
+			}
+			else
+			{
+				DragState = EDragState::Pan;
+			}
 		}
 		else
 		{
@@ -1396,6 +1615,21 @@ void SCurveEditor::TryStartDrag(const FGeometry& InMyGeometry, const FPointerEve
 	}
 }
 
+/* Given a tangent value for a key, calculates the 2D delta vector from that key in curve space */
+static inline FVector2D CalcTangentDir(float Tangent)
+{
+	const float Angle = FMath::Atan(Tangent);
+	return FVector2D( FMath::Cos(Angle), -FMath::Sin(Angle) );
+}
+
+/*Given a 2d delta vector in curve space, calculates a tangent value */
+static inline float CalcTangent(const FVector2D& HandleDelta)
+{
+	// Ensure X is positive and non-zero.
+	// Tangent is gradient of handle.
+	return HandleDelta.Y / FMath::Max<double>(HandleDelta.X, KINDA_SMALL_NUMBER);
+}
+
 void SCurveEditor::ProcessDrag(const FGeometry& InMyGeometry, const FPointerEvent& InMouseEvent)
 {
 	FTrackScaleInfo ScaleInfo(ViewMinInput.Get(), ViewMaxInput.Get(), ViewMinOutput.Get(), ViewMaxOutput.Get(), InMyGeometry.Size);
@@ -1408,17 +1642,39 @@ void SCurveEditor::ProcessDrag(const FGeometry& InMyGeometry, const FPointerEven
 	if (DragState == EDragState::DragKey)
 	{
 		FVector2D MousePosition = InMyGeometry.AbsoluteToLocal(InMouseEvent.GetScreenSpacePosition());
-		MoveSelectedKeys(FVector2D(ScaleInfo.LocalXToInput(MousePosition.X), ScaleInfo.LocalYToOutput(MousePosition.Y)));
+		FVector2D NewLocation = FVector2D(ScaleInfo.LocalXToInput(MousePosition.X), ScaleInfo.LocalYToOutput(MousePosition.Y));
+		FVector2D SnappedNewLocation = SnapLocation(NewLocation);
+		FVector2D Delta = SnappedNewLocation - PreDragKeyLocations[DraggedKeyHandle];
+
+		MoveSelectedKeys(Delta);
+	}
+	else if (DragState == EDragState::FreeDrag)
+	{
+		FVector2D MousePosition = InMyGeometry.AbsoluteToLocal(InMouseEvent.GetScreenSpacePosition());
+		FVector2D NewLocation = FVector2D(ScaleInfo.LocalXToInput(MousePosition.X), ScaleInfo.LocalYToOutput(MousePosition.Y));
+		FVector2D Delta = NewLocation - FVector2D(ScaleInfo.LocalXToInput(MouseDownLocation.X), ScaleInfo.LocalYToOutput(MouseDownLocation.Y));
+
+		MoveSelectedKeys(Delta);
 	}
 	else if (DragState == EDragState::DragTangent)
 	{
 		FVector2D MousePositionScreen = InMouseEvent.GetScreenSpacePosition();
 		MousePositionScreen -= InMyGeometry.AbsolutePosition;
 		FVector2D MousePositionCurve(ScaleInfo.LocalXToInput(MousePositionScreen.X), ScaleInfo.LocalYToOutput(MousePositionScreen.Y));
-		OnMoveTangent(MousePositionCurve);
+		FVector2D Delta = MousePositionCurve;
+		MoveTangents(ScaleInfo, Delta);
 	}
 	else if (DragState == EDragState::Pan)
 	{
+		if (MovementAxisLock == EMovementAxisLock::AxisLock_Horizontal)
+		{
+			InputDelta.Y = 0;
+		}
+		else if (MovementAxisLock == EMovementAxisLock::AxisLock_Vertical)
+		{
+			InputDelta.X = 0;
+		}
+
 		// Output is not clamped.
 		const float NewMinOutput = (ViewMinOutput.Get() - InputDelta.Y);
 		const float NewMaxOutput = (ViewMaxOutput.Get() - InputDelta.Y);
@@ -1432,16 +1688,33 @@ void SCurveEditor::ProcessDrag(const FGeometry& InMyGeometry, const FPointerEven
 
 		SetInputMinMax(NewMinInput, NewMaxInput);
 	}
+	else if (DragState == EDragState::Zoom)
+	{
+		FVector2D Delta = FVector2D(ScreenDelta.X * 0.05f, ScreenDelta.X * 0.05f);
+
+		if (MovementAxisLock == EMovementAxisLock::AxisLock_Horizontal)
+		{
+			Delta.Y = 0;
+		}
+		else if (MovementAxisLock == EMovementAxisLock::AxisLock_Vertical)
+		{
+			Delta.X = 0;
+			Delta.Y = -ScreenDelta.Y * 0.1f;
+		}
+
+		ZoomView(Delta);
+	}
 }
 
 void SCurveEditor::EndDrag(const FGeometry& InMyGeometry, const FPointerEvent& InMouseEvent)
 {
 	const bool bControlDown = InMouseEvent.IsControlDown();
+	const bool bShiftDown = InMouseEvent.IsShiftDown();
 
-	if (DragState == EDragState::DragKey || DragState == EDragState::DragTangent)
+	if (DragState == EDragState::DragKey || DragState == EDragState::FreeDrag || DragState == EDragState::DragTangent)
 	{
-			EndDragTransaction();
-		}
+		EndDragTransaction();
+	}
 	else if (DragState == EDragState::MarqueeSelect)
 	{
 		FVector2D MarqueTopLeft
@@ -1456,27 +1729,51 @@ void SCurveEditor::EndDrag(const FGeometry& InMyGeometry, const FPointerEvent& I
 			FMath::Max(MouseDownLocation.Y, MouseMoveLocation.Y)
 			);
 
+		TArray<FSelectedTangent> SelectedCurveTangents = GetEditableTangentsWithinMarquee(InMyGeometry, MarqueTopLeft, MarqueBottomRight);
+
 		TArray<FSelectedCurveKey> SelectedCurveKeys = GetEditableKeysWithinMarquee(InMyGeometry, MarqueTopLeft, MarqueBottomRight);
 
-		if (!bControlDown)
+		if (!bControlDown && !bShiftDown)
 		{
-			EmptySelection();
+			EmptyAllSelection();
 		}
 
-		for (auto SelectedCurveKey : SelectedCurveKeys)
+		if (SelectedCurveKeys.Num())
 		{
-			
-			if (IsKeySelected(SelectedCurveKey))
-			{
-				RemoveFromSelection(SelectedCurveKey);
+			EmptyTangentSelection();
+
+			for (auto SelectedCurveKey : SelectedCurveKeys)
+			{			
+				if (IsKeySelected(SelectedCurveKey))
+				{
+					RemoveFromKeySelection(SelectedCurveKey);
+				}
+				else
+				{
+					AddToKeySelection(SelectedCurveKey);
+				}
 			}
-			else
-			{
-				AddToSelection(SelectedCurveKey);
+		}
+
+		if (!SelectedCurveKeys.Num())
+		{
+			EmptyKeySelection();
+
+			for (auto SelectedCurveTangent : SelectedCurveTangents)
+			{			
+				if (IsTangentSelected(SelectedCurveTangent))
+				{
+					RemoveFromTangentSelection(SelectedCurveTangent);
+				}
+				else
+				{
+					AddToTangentSelection(SelectedCurveTangent);
+				}
 			}
 		}
 	}
 	DragState = EDragState::None;
+	MovementAxisLock = EMovementAxisLock::None;
 }
 
 void SCurveEditor::ProcessClick(const FGeometry& InMyGeometry, const FPointerEvent& InMouseEvent)
@@ -1487,23 +1784,43 @@ void SCurveEditor::ProcessClick(const FGeometry& InMyGeometry, const FPointerEve
 	const bool bShiftDown = InMouseEvent.IsShiftDown();
 
 	FSelectedCurveKey HitKey = HitTestKeys(InMyGeometry, InMouseEvent.GetScreenSpacePosition());
+	FSelectedTangent HitTangent = HitTestCubicTangents(InMyGeometry, InMouseEvent.GetScreenSpacePosition());
 
 	if (bLeftMouseButton)
 	{
 		// If the user left clicked a key, update selection based on modifier key state.
 		if (HitKey.IsValid())
 		{
+			EmptyTangentSelection();
+
 			if (!IsKeySelected(HitKey))
 			{
-				if (!bControlDown)
+				if (!bControlDown && !bShiftDown)
 				{
-					EmptySelection();
+					EmptyAllSelection();
 				}
-				AddToSelection(HitKey);
+				AddToKeySelection(HitKey);
 			}
 			else if (bControlDown)
 			{
-				RemoveFromSelection(HitKey);
+				RemoveFromKeySelection(HitKey);
+			}
+		}
+		else if (HitTangent.IsValid())
+		{
+			EmptyKeySelection();
+
+			if (!IsTangentSelected(HitTangent))
+			{
+				if (!bControlDown && !bShiftDown)
+				{
+					EmptyAllSelection();
+				}
+				AddToTangentSelection(HitTangent);
+			}
+			else if (bControlDown)
+			{
+				RemoveFromTangentSelection(HitTangent);
 			}
 		}
 		else
@@ -1536,8 +1853,8 @@ void SCurveEditor::ProcessClick(const FGeometry& InMyGeometry, const FPointerEve
 			}
 			else
 			{
-				// clicking on background clears key selection
-				EmptySelection();
+				// clicking on background clears all selection
+				EmptyAllSelection();
 			}
 		}
 	}
@@ -1547,23 +1864,30 @@ void SCurveEditor::ProcessClick(const FGeometry& InMyGeometry, const FPointerEve
 		if (HitKey.IsValid())
 		{
 			// Make sure key is selected in readiness for context menu
+			EmptyTangentSelection();
+
 			if (!IsKeySelected(HitKey))
 			{
-				EmptySelection();
-				AddToSelection(HitKey);
+				EmptyAllSelection();
+				AddToKeySelection(HitKey);
 			}
 			PushKeyMenu(InMyGeometry, InMouseEvent);
 		}
-		else
+		else if (HitTangent.IsValid())
 		{
-			if (!HitTestCubicTangents(InMyGeometry, InMouseEvent.GetScreenSpacePosition()).IsValid())
-			{
-				CreateContextMenu(InMyGeometry, InMouseEvent);
-			}
-			else
+			// Make sure key is selected in readiness for context menu
+			EmptyKeySelection();
+
+			if (!IsTangentSelected(HitTangent))
 			{
-				EmptySelection();
+				EmptyAllSelection();
+				AddToTangentSelection(HitTangent);
 			}
+			PushKeyMenu(InMyGeometry, InMouseEvent);
+		}
+		else
+		{
+			CreateContextMenu(InMyGeometry, InMouseEvent);
 		}
 	}
 }
@@ -1766,7 +2090,7 @@ SCurveEditor::FSelectedCurveKey SCurveEditor::HitTestKeys(const FGeometry& InMyG
 	return SelectedKey;
 }
 
-void SCurveEditor::MoveSelectedKeys(FVector2D InNewLocation)
+void SCurveEditor::MoveSelectedKeys(FVector2D Delta)
 {
 	const FScopedTransaction Transaction( LOCTEXT("CurveEditor_MoveKeys", "Move Keys") );
 	CurveOwner->ModifyOwner();
@@ -1774,11 +2098,11 @@ void SCurveEditor::MoveSelectedKeys(FVector2D InNewLocation)
 	// track all unique curves encountered so their tangents can be updated later
 	TSet<FRichCurve*> UniqueCurves;
 
-	FVector2D SnappedNewLocation = SnapLocation(InNewLocation);
 
 	// The total move distance for all keys is the difference between the current snapped location
 	// and the start location of the key which was actually dragged.
-	FVector2D TotalMoveDistance = SnappedNewLocation - PreDragKeyLocations[DraggedKeyHandle];
+
+	FVector2D TotalMoveDistance = Delta;
 
 	for (int32 i = 0; i < SelectedKeys.Num(); i++)
 	{
@@ -1796,13 +2120,19 @@ void SCurveEditor::MoveSelectedKeys(FVector2D InNewLocation)
 		FVector2D NewLocation = PreDragLocation + TotalMoveDistance;
 
 		// Update the key's value without updating the tangents.
-		Curve->SetKeyValue(OldKeyHandle, NewLocation.Y, false);
+		if (MovementAxisLock != EMovementAxisLock::AxisLock_Horizontal)
+		{
+			Curve->SetKeyValue(OldKeyHandle, NewLocation.Y, false);
+		}
 
 		// Changing the time of a key returns a new handle, so make sure to update existing references.
-		FKeyHandle KeyHandle = Curve->SetKeyTime(OldKeyHandle, NewLocation.X);
-		SelectedKeys[i] = FSelectedCurveKey(Curve, KeyHandle);
-		PreDragKeyLocations.Remove(OldKeyHandle);
-		PreDragKeyLocations.Add(KeyHandle, PreDragLocation);
+		if (MovementAxisLock != EMovementAxisLock::AxisLock_Vertical)
+		{
+			FKeyHandle KeyHandle = Curve->SetKeyTime(OldKeyHandle, NewLocation.X);
+			SelectedKeys[i] = FSelectedCurveKey(Curve, KeyHandle);
+			PreDragKeyLocations.Remove(OldKeyHandle);
+			PreDragKeyLocations.Add(KeyHandle, PreDragLocation);
+		}
 
 		UniqueCurves.Add(Curve);
 	}
@@ -1834,17 +2164,17 @@ TOptional<float> SCurveEditor::GetKeyTime(FSelectedCurveKey Key) const
 	return TOptional<float>();
 }
 
-void SCurveEditor::EmptySelection()
+void SCurveEditor::EmptyKeySelection()
 {
 	SelectedKeys.Empty();
 }
 
-void SCurveEditor::AddToSelection(FSelectedCurveKey Key)
+void SCurveEditor::AddToKeySelection(FSelectedCurveKey Key)
 {
 	SelectedKeys.AddUnique(Key);
 }
 
-void SCurveEditor::RemoveFromSelection(FSelectedCurveKey Key)
+void SCurveEditor::RemoveFromKeySelection(FSelectedCurveKey Key)
 {
 	SelectedKeys.Remove(Key);
 }
@@ -1859,6 +2189,83 @@ bool SCurveEditor::AreKeysSelected() const
 	return SelectedKeys.Num() > 0;
 }
 
+void SCurveEditor::EmptyTangentSelection()
+{
+	SelectedTangents.Empty();
+}
+
+void SCurveEditor::AddToTangentSelection(FSelectedTangent Tangent)
+{
+	SelectedTangents.Add(Tangent);
+}
+
+void SCurveEditor::RemoveFromTangentSelection(FSelectedTangent Tangent)
+{
+	SelectedTangents.Remove(Tangent);
+}
+
+bool SCurveEditor::IsTangentSelected(FSelectedTangent Tangent) const
+{
+	return SelectedTangents.Contains(Tangent);
+}
+
+bool SCurveEditor::AreTangentsSelected() const
+{
+	return SelectedTangents.Num() > 0;
+}
+
+bool SCurveEditor::IsTangentVisible(FRichCurve* Curve, FKeyHandle KeyHandle, bool& bIsTangentSelected, bool& bIsArrivalSelected, bool& bIsLeaveSelected) const
+{
+	bIsTangentSelected = false;
+	bIsArrivalSelected = false;
+	bIsLeaveSelected = false;
+
+	bool IsSelected = IsKeySelected(FSelectedCurveKey(Curve,KeyHandle));
+	for (auto SelectedTangent : SelectedTangents)
+	{
+		if (SelectedTangent.Key.KeyHandle == KeyHandle)
+		{
+			if (SelectedTangent.bIsArrival)
+				bIsArrivalSelected = true;
+			else
+				bIsLeaveSelected = true;
+			bIsTangentSelected = true;
+		}
+	}
+
+	return (IsSelected || bIsTangentSelected || Settings->GetTangentVisibility() == ECurveEditorTangentVisibility::AllTangents) && Settings->GetTangentVisibility() != ECurveEditorTangentVisibility::NoTangents;
+}
+
+void SCurveEditor::EmptyAllSelection()
+{
+	EmptyKeySelection();
+	EmptyTangentSelection();
+}
+
+void SCurveEditor::ValidateSelection()
+{
+	//remove any invalid keys
+	for(int32 i = 0;i<SelectedKeys.Num();++i)
+	{
+		auto Key = SelectedKeys[i];
+		if(!IsValidCurve(Key.Curve) || !Key.IsValid())
+		{
+			SelectedKeys.RemoveAt(i);
+			i--;
+		}
+	}
+
+	// remove any invalid tangents
+	for (int32 i = 0;i<SelectedTangents.Num();++i)
+	{
+		auto Tangent = SelectedTangents[i];
+		if (!IsValidCurve(Tangent.Key.Curve) || !Tangent.Key.IsValid())
+		{
+			SelectedTangents.RemoveAt(i);
+			i--;
+		}
+	}
+}
 
 TArray<FRichCurve*> SCurveEditor::GetCurvesToFit() const
 {
@@ -2261,34 +2668,54 @@ SCurveEditor::FSelectedTangent SCurveEditor::HitTestCubicTangents( const FGeomet
 
 		const FVector2D HitPosition = InMyGeometry.AbsoluteToLocal( HitScreenPosition);
 
-		for(auto It = SelectedKeys.CreateConstIterator();It;++It)
+		for (auto CurveViewModel : CurveViewModels)
 		{
-			FSelectedCurveKey Key = *It;
-			if(Key.IsValid())
+			if (!CurveViewModel->bIsLocked && CurveViewModel->bIsVisible)
 			{
-				float Time		 = ScaleInfo.LocalXToInput(HitPosition.X);
-				float KeyScreenY = ScaleInfo.OutputToLocalY(Key.Curve->Eval(Time));
-
-				FVector2D Arrive, Leave;
-				GetTangentPoints(ScaleInfo, Key, Arrive, Leave);
-
-				if( HitPosition.Y > (Arrive.Y - (0.5f * CONST_CurveSize.Y)) &&
-					HitPosition.Y < (Arrive.Y + (0.5f * CONST_CurveSize.Y)) &&
-					HitPosition.X > (Arrive.X - (0.5f * CONST_TangentSize.X)) && 
-					HitPosition.X < (Arrive.X + (0.5f * CONST_TangentSize.X)))
-				{
-					Tangent.Key = Key;
-					Tangent.bIsArrival = true;
-					break;
-				}
-				if( HitPosition.Y > (Leave.Y - (0.5f * CONST_CurveSize.Y)) &&
-					HitPosition.Y < (Leave.Y  + (0.5f * CONST_CurveSize.Y)) &&
-					HitPosition.X > (Leave.X - (0.5f * CONST_TangentSize.X)) && 
-					HitPosition.X < (Leave.X + (0.5f * CONST_TangentSize.X)))
+				FRichCurve* Curve = CurveViewModel->CurveInfo.CurveToEdit;
+				if (Curve != NULL)
 				{
-					Tangent.Key = Key;
-					Tangent.bIsArrival = false;
-					break;
+					for (auto It(Curve->GetKeyHandleIterator()); It; ++It)
+					{
+						FKeyHandle KeyHandle = It.Key();
+						FSelectedCurveKey SelectedCurveKey(Curve, KeyHandle);
+
+						if(SelectedCurveKey.IsValid())
+						{
+							bool bIsTangentSelected = false;
+							bool bIsArrivalSelected = false;
+							bool bIsLeaveSelected = false;
+							bool bIsTangentVisible = IsTangentVisible(Curve, KeyHandle, bIsTangentSelected, bIsArrivalSelected, bIsLeaveSelected);
+
+							if (bIsTangentVisible)
+							{
+								float Time		 = ScaleInfo.LocalXToInput(HitPosition.X);
+								float KeyScreenY = ScaleInfo.OutputToLocalY(Curve->Eval(Time));
+
+								FVector2D Arrive, Leave;
+								GetTangentPoints(ScaleInfo, SelectedCurveKey, Arrive, Leave);
+
+								if( HitPosition.Y > (Arrive.Y - (0.5f * CONST_CurveSize.Y)) &&
+									HitPosition.Y < (Arrive.Y + (0.5f * CONST_CurveSize.Y)) &&
+									HitPosition.X > (Arrive.X - (0.5f * CONST_TangentSize.X)) && 
+									HitPosition.X < (Arrive.X + (0.5f * CONST_TangentSize.X)))
+								{
+									Tangent.Key = SelectedCurveKey;
+									Tangent.bIsArrival = true;
+									break;
+								}
+								if( HitPosition.Y > (Leave.Y - (0.5f * CONST_CurveSize.Y)) &&
+									HitPosition.Y < (Leave.Y  + (0.5f * CONST_CurveSize.Y)) &&
+									HitPosition.X > (Leave.X - (0.5f * CONST_TangentSize.X)) && 
+									HitPosition.X < (Leave.X + (0.5f * CONST_TangentSize.X)))
+								{
+									Tangent.Key = SelectedCurveKey;
+									Tangent.bIsArrival = false;
+									break;
+								}
+							}
+						}
+					}
 				}
 			}
 		}
@@ -2299,7 +2726,7 @@ SCurveEditor::FSelectedTangent SCurveEditor::HitTestCubicTangents( const FGeomet
 
 void SCurveEditor::OnSelectInterpolationMode(ERichCurveInterpMode InterpMode, ERichCurveTangentMode TangentMode)
 {
-	if(SelectedKeys.Num() > 0)
+	if(SelectedKeys.Num() > 0 || SelectedTangents.Num() > 0)
 	{
 		const FScopedTransaction Transaction(LOCTEXT("CurveEditor_SetInterpolationMode", "Select Interpolation Mode"));
 		CurveOwner->ModifyOwner();
@@ -2313,6 +2740,14 @@ void SCurveEditor::OnSelectInterpolationMode(ERichCurveInterpMode InterpMode, ER
 			Key.Curve->SetKeyTangentMode(Key.KeyHandle,TangentMode );
 		}
 
+		for(auto It = SelectedTangents.CreateIterator();It;++It)
+		{
+			FSelectedTangent& Tangent = *It;
+			check(IsValidCurve(Tangent.Key.Curve));
+			Tangent.Key.Curve->SetKeyInterpMode(Tangent.Key.KeyHandle,InterpMode );
+			Tangent.Key.Curve->SetKeyTangentMode(Tangent.Key.KeyHandle,TangentMode );
+		}
+
 		TArray<FRichCurveEditInfo> ChangedCurveEditInfos;
 		for (auto CurveViewModel : CurveViewModels)
 		{
@@ -2338,62 +2773,193 @@ bool SCurveEditor::IsInterpolationModeSelected(ERichCurveInterpMode InterpMode,
 		}
 		return true;
 	}
-	else
+	else if (SelectedTangents.Num() > 0)
 	{
-					return false;
-				}
+		for (auto SelectedTangent : SelectedTangents)
+		{
+			if (SelectedTangent.Key.Curve->GetKeyInterpMode(SelectedTangent.Key.KeyHandle) != InterpMode || SelectedTangent.Key.Curve->GetKeyTangentMode(SelectedTangent.Key.KeyHandle) != TangentMode)
+			{
+				return false;
 			}
+		}
+		return true;
+	}
+	else
+	{
+		return false;
+	}
+}
 
-/* Given a tangent value for a key, calculates the 2D delta vector from that key in curve space */
-static inline FVector2D CalcTangentDir(float Tangent)
+void SCurveEditor::OnFlattenOrStraightenTangents(bool bFlattenTangents)
 {
-	const float Angle = FMath::Atan(Tangent);
-	return FVector2D( FMath::Cos(Angle), -FMath::Sin(Angle) );
+	if(SelectedKeys.Num() > 0 || SelectedTangents.Num() > 0)
+	{
+		const FScopedTransaction Transaction(LOCTEXT("CurveEditor_FlattenTangents", "Flatten Tangents"));
+		CurveOwner->ModifyOwner();
+		TSet<FRichCurve*> ChangedCurves;
+
+		for(auto It = SelectedKeys.CreateIterator();It;++It)
+		{
+			FSelectedCurveKey& Key = *It;
+			check(IsValidCurve(Key.Curve));
+
+			float LeaveTangent = Key.Curve->GetKey(Key.KeyHandle).LeaveTangent;
+			float ArriveTangent = Key.Curve->GetKey(Key.KeyHandle).ArriveTangent;
+
+			if (bFlattenTangents)
+			{
+				LeaveTangent = 0;
+				ArriveTangent = 0;
+			}
+			else
+			{
+				LeaveTangent = (LeaveTangent + ArriveTangent) * 0.5f;
+				ArriveTangent = LeaveTangent;
+			}
+
+			Key.Curve->GetKey(Key.KeyHandle).LeaveTangent = LeaveTangent;
+			Key.Curve->GetKey(Key.KeyHandle).ArriveTangent = ArriveTangent;
+		}
+				
+		for(auto It = SelectedTangents.CreateIterator();It;++It)
+		{
+			FSelectedTangent& Tangent = *It;
+			check(IsValidCurve(Tangent.Key.Curve));
+
+			float LeaveTangent = Tangent.Key.Curve->GetKey(Tangent.Key.KeyHandle).LeaveTangent;
+			float ArriveTangent = Tangent.Key.Curve->GetKey(Tangent.Key.KeyHandle).ArriveTangent;
+
+			if (bFlattenTangents)
+			{
+				LeaveTangent = 0;
+				ArriveTangent = 0;
+			}
+			else
+			{
+				LeaveTangent = (LeaveTangent + ArriveTangent) * 0.5f;
+				ArriveTangent = LeaveTangent;
+			}
+
+			Tangent.Key.Curve->GetKey(Tangent.Key.KeyHandle).LeaveTangent = LeaveTangent;
+			Tangent.Key.Curve->GetKey(Tangent.Key.KeyHandle).ArriveTangent = ArriveTangent;
+		}
+
+		TArray<FRichCurveEditInfo> ChangedCurveEditInfos;
+		for (auto CurveViewModel : CurveViewModels)
+		{
+			if (ChangedCurves.Contains(CurveViewModel->CurveInfo.CurveToEdit))
+			{
+				ChangedCurveEditInfos.Add(CurveViewModel->CurveInfo);
+			}
+		}
+		CurveOwner->OnCurveChanged(ChangedCurveEditInfos);
+	}
 }
 
-/*Given a 2d delta vector in curve space, calculates a tangent value */
-static inline float CalcTangent(const FVector2D& HandleDelta)
+void SCurveEditor::OnSelectPreInfinityExtrap(ERichCurveExtrapolation Extrapolation)
 {
-	// Ensure X is positive and non-zero.
-	// Tangent is gradient of handle.
-	return HandleDelta.Y / FMath::Max<double>(HandleDelta.X, KINDA_SMALL_NUMBER);
+	const FScopedTransaction Transaction(LOCTEXT("CurveEditor_SetPreInfinityExtrapolation", "Set Pre-Infinity Extrapolation"));
+	CurveOwner->ModifyOwner();
+
+	TArray<FRichCurveEditInfo> ChangedCurveEditInfos;
+	for (auto CurveViewModel : CurveViewModels)
+	{
+		if (CurveViewModel->CurveInfo.CurveToEdit->PreInfinityExtrap != Extrapolation)
+		{
+			CurveViewModel->CurveInfo.CurveToEdit->PreInfinityExtrap = Extrapolation;
+			ChangedCurveEditInfos.Add(CurveViewModel->CurveInfo);
+		}
+	}
+
+	if (ChangedCurveEditInfos.Num())
+	{
+		CurveOwner->OnCurveChanged(ChangedCurveEditInfos);
+	}
 }
 
-void SCurveEditor::OnMoveTangent(FVector2D MouseCurvePosition)
+bool SCurveEditor::IsPreInfinityExtrapSelected(ERichCurveExtrapolation Extrapolation)
 {
-	auto& RichKey = SelectedTangent.Key.Curve->GetKey(SelectedTangent.Key.KeyHandle);
+	for (auto CurveViewModel : CurveViewModels)
+	{
+		if (CurveViewModel->CurveInfo.CurveToEdit->PreInfinityExtrap != Extrapolation)
+		{
+			return false;
+		}
+	}
+	return CurveViewModels.Num() > 0;
+}
 
-	const FSelectedCurveKey &Key = SelectedTangent.Key;
+void SCurveEditor::OnSelectPostInfinityExtrap(ERichCurveExtrapolation Extrapolation)
+{
+	const FScopedTransaction Transaction(LOCTEXT("CurveEditor_SetPostInfinityExtrapolation", "Set Post-Infinity Extrapolation"));
+	CurveOwner->ModifyOwner();
 
-	FVector2D KeyPosition( Key.Curve->GetKeyTime(Key.KeyHandle),Key.Curve->GetKeyValue(Key.KeyHandle) );
+	TArray<FRichCurveEditInfo> ChangedCurveEditInfos;
+	for (auto CurveViewModel : CurveViewModels)
+	{
+		if (CurveViewModel->CurveInfo.CurveToEdit->PostInfinityExtrap != Extrapolation)
+		{
+			CurveViewModel->CurveInfo.CurveToEdit->PostInfinityExtrap = Extrapolation;
+			ChangedCurveEditInfos.Add(CurveViewModel->CurveInfo);
+		}
+	}
 
-	FVector2D Movement = MouseCurvePosition - KeyPosition;
-	if(SelectedTangent.bIsArrival)
+	if (ChangedCurveEditInfos.Num())
 	{
-		Movement *= -1.0f;
+		CurveOwner->OnCurveChanged(ChangedCurveEditInfos);
 	}
-	float Tangent = CalcTangent(Movement);
+}
 
-	if(RichKey.TangentMode  != RCTM_Break)
+bool SCurveEditor::IsPostInfinityExtrapSelected(ERichCurveExtrapolation Extrapolation)
+{
+	for (auto CurveViewModel : CurveViewModels)
 	{
-		RichKey.ArriveTangent = Tangent;
-		RichKey.LeaveTangent  = Tangent;
-		
-		RichKey.TangentMode = RCTM_User;
+		if (CurveViewModel->CurveInfo.CurveToEdit->PostInfinityExtrap != Extrapolation)
+		{
+			return false;
+		}
 	}
-	else
+	return CurveViewModels.Num() > 0;
+}
+
+void SCurveEditor::MoveTangents(FTrackScaleInfo ScaleInfo, FVector2D Delta)
+{
+	for (auto SelectedTangent : SelectedTangents)
 	{
+		auto& RichKey = SelectedTangent.Key.Curve->GetKey(SelectedTangent.Key.KeyHandle);
+
+		const FSelectedCurveKey &Key = SelectedTangent.Key;
+
+		FVector2D KeyPosition( Key.Curve->GetKeyTime(Key.KeyHandle),Key.Curve->GetKeyValue(Key.KeyHandle) );
+
+		FVector2D Movement = Delta - KeyPosition;
 		if(SelectedTangent.bIsArrival)
+		{
+			Movement *= -1.0f;
+		}
+		float Tangent = CalcTangent(Movement);
+
+		if(RichKey.TangentMode  != RCTM_Break)
 		{
 			RichKey.ArriveTangent = Tangent;
+			RichKey.LeaveTangent  = Tangent;
+		
+			RichKey.TangentMode = RCTM_User;
 		}
 		else
 		{
-			RichKey.LeaveTangent = Tangent;
-		}
-	}	
-
+			if(SelectedTangent.bIsArrival)
+			{
+				RichKey.ArriveTangent = Tangent;
+			}
+			else
+			{
+				RichKey.LeaveTangent = Tangent;
+			}
+		}	
+	}
 }
+
 void SCurveEditor::GetTangentPoints(  FTrackScaleInfo &ScaleInfo, const FSelectedCurveKey &Key, FVector2D& Arrive, FVector2D& Leave ) const
 {
 	FVector2D ArriveTangentDir = CalcTangentDir( Key.Curve->GetKey(Key.KeyHandle).ArriveTangent);
@@ -2457,6 +3023,59 @@ TArray<SCurveEditor::FSelectedCurveKey> SCurveEditor::GetEditableKeysWithinMarqu
 	return KeysWithinMarquee;
 }
 
+TArray<SCurveEditor::FSelectedTangent> SCurveEditor::GetEditableTangentsWithinMarquee(const FGeometry& InMyGeometry, FVector2D MarqueeTopLeft, FVector2D MarqueeBottomRight) const
+{
+	FBox MarqueeBox;
+	MarqueeBox.Min = FVector(MarqueeTopLeft.X, MarqueeTopLeft.Y, 0);
+	MarqueeBox.Max = FVector(MarqueeBottomRight.X, MarqueeBottomRight.Y, 0);
+
+	TArray<FSelectedTangent> TangentsWithinMarquee;
+	if (AreCurvesVisible())
+	{
+		FTrackScaleInfo ScaleInfo(ViewMinInput.Get(), ViewMaxInput.Get(), ViewMinOutput.Get(), ViewMaxOutput.Get(), InMyGeometry.Size);
+		for (auto CurveViewModel : CurveViewModels)
+		{
+			if (!CurveViewModel->bIsLocked && CurveViewModel->bIsVisible)
+			{
+				FRichCurve* Curve = CurveViewModel->CurveInfo.CurveToEdit;
+				if (Curve != NULL)
+				{
+					for (auto It(Curve->GetKeyHandleIterator()); It; ++It)
+					{
+						FKeyHandle KeyHandle = It.Key();
+						FSelectedCurveKey SelectedCurveKey(Curve, KeyHandle);
+
+						if(SelectedCurveKey.IsValid())
+						{
+							bool bIsTangentSelected = false;
+							bool bIsArrivalSelected = false;
+							bool bIsLeaveSelected = false;
+							bool bIsTangentVisible = IsTangentVisible(Curve, KeyHandle, bIsTangentSelected, bIsArrivalSelected, bIsLeaveSelected);
+
+							if (bIsTangentVisible)
+							{
+								FVector2D Arrive, Leave;
+								GetTangentPoints(ScaleInfo, SelectedCurveKey, Arrive, Leave);
+
+								bool bArriveInside = MarqueeBox.IsInsideOrOn(FVector(Arrive.X, Arrive.Y, 0));
+								bool bLeaveInside = MarqueeBox.IsInsideOrOn(FVector(Leave.X, Leave.Y, 0));
+
+								if (bArriveInside || bLeaveInside)
+								{
+									FSelectedTangent SelectedTangent(SelectedCurveKey);
+									SelectedTangent.bIsArrival = bArriveInside;
+									TangentsWithinMarquee.Add(SelectedTangent);
+								}	
+							}
+						}
+					}
+				}
+			}
+		}
+	}
+	return TangentsWithinMarquee;
+}
+
 void SCurveEditor::BeginDragTransaction()
 {
 	TransactionIndex = GEditor->BeginTransaction( LOCTEXT("CurveEditor_Drag", "Mouse Drag") );
@@ -2495,16 +3114,7 @@ void SCurveEditor::RedoAction()
 
 void SCurveEditor::PostUndo(bool bSuccess)
 {
-	//remove any invalid keys
-	for(int32 i = 0;i<SelectedKeys.Num();++i)
-	{
-		auto Key = SelectedKeys[i];
-		if(!IsValidCurve(Key.Curve) || !Key.IsValid())
-		{
-			SelectedKeys.RemoveAt(i);
-			i--;
-		}
-	}
+	ValidateSelection();
 }
 
 bool SCurveEditor::IsLinearColorCurve() const 
diff --git a/Engine/Source/Editor/UnrealEd/Public/CurveEditorSettings.cpp b/Engine/Source/Editor/UnrealEd/Public/CurveEditorSettings.cpp
new file mode 100644
--- /dev/null
+++ b/Engine/Source/Editor/UnrealEd/Public/CurveEditorSettings.cpp
@@ -0,0 +1,75 @@
+// Copyright 1998-2015 Epic Games, Inc. All Rights Reserved.
+
+#include "UnrealEd.h"
+#include "CurveEditorSettings.h"
+
+UCurveEditorSettings::UCurveEditorSettings( const FObjectInitializer& ObjectInitializer )
+	: Super( ObjectInitializer )
+{
+	bAutoFrameCurveEditor = true;
+	bShowCurveEditorCurveToolTips = true;
+	CurveVisibility = ECurveEditorCurveVisibility::AllCurves;
+	TangentVisibility = ECurveEditorTangentVisibility::SelectedKeys;
+}
+
+bool UCurveEditorSettings::GetAutoFrameCurveEditor() const
+{
+	return bAutoFrameCurveEditor;
+}
+
+void UCurveEditorSettings::SetAutoFrameCurveEditor(bool InbAutoFrameCurveEditor)
+{
+	if (bAutoFrameCurveEditor != InbAutoFrameCurveEditor)
+	{
+		bAutoFrameCurveEditor = InbAutoFrameCurveEditor;
+		SaveConfig();
+	}
+}
+
+bool UCurveEditorSettings::GetShowCurveEditorCurveToolTips() const
+{
+	return bShowCurveEditorCurveToolTips;
+}
+
+void UCurveEditorSettings::SetShowCurveEditorCurveToolTips(bool InbShowCurveEditorCurveToolTips)
+{
+	if (bShowCurveEditorCurveToolTips != InbShowCurveEditorCurveToolTips)
+	{
+		bShowCurveEditorCurveToolTips = InbShowCurveEditorCurveToolTips;
+		SaveConfig();
+	}
+}
+
+ECurveEditorCurveVisibility::Type UCurveEditorSettings::GetCurveVisibility() const
+{
+	return CurveVisibility;
+}
+
+void UCurveEditorSettings::SetCurveVisibility(ECurveEditorCurveVisibility::Type InCurveVisibility)
+{
+	if (CurveVisibility != InCurveVisibility)
+	{
+		CurveVisibility = InCurveVisibility;
+		OnCurveEditorCurveVisibilityChanged.Broadcast();
+		SaveConfig();
+	}
+}
+
+ECurveEditorTangentVisibility::Type UCurveEditorSettings::GetTangentVisibility() const
+{
+	return TangentVisibility;
+}
+
+void UCurveEditorSettings::SetTangentVisibility(ECurveEditorTangentVisibility::Type InTangentVisibility)
+{
+	if (TangentVisibility != InTangentVisibility)
+	{
+		TangentVisibility = InTangentVisibility;
+		SaveConfig();
+	}
+}
+
+UCurveEditorSettings::FOnCurveEditorCurveVisibilityChanged& UCurveEditorSettings::GetOnCurveEditorCurveVisibilityChanged()
+{
+	return OnCurveEditorCurveVisibilityChanged;
+}
diff --git a/Engine/Source/Editor/UnrealEd/Public/CurveEditorSettings.h b/Engine/Source/Editor/UnrealEd/Public/CurveEditorSettings.h
new file mode 100644
--- /dev/null
+++ b/Engine/Source/Editor/UnrealEd/Public/CurveEditorSettings.h
@@ -0,0 +1,83 @@
+// Copyright 1998-2015 Epic Games, Inc. All Rights Reserved.
+
+#pragma once
+
+#include "CurveEditorSettings.generated.h"
+
+/** Defines visibility states for the curves in the curve editor. */
+UENUM()
+namespace ECurveEditorCurveVisibility
+{
+	enum Type
+	{
+		/** All curves should be visible. */
+		AllCurves,
+		/** Only curves from selected nodes should be visible. */
+		SelectedCurves,
+		/** Only curves which have keyframes should be visible. */
+		AnimatedCurves
+	};
+}
+
+/** Defines visibility states for the tangents in the curve editor. */
+UENUM()
+namespace ECurveEditorTangentVisibility
+{
+	enum Type
+	{
+		/** All tangents should be visible. */
+		AllTangents,
+		/** Only tangents from selected keys should be visible. */
+		SelectedKeys,
+		/** Don't display tangents. */
+		NoTangents
+	};
+}
+
+/** Serializable options for curve editor. */
+UCLASS(config=EditorPerProjectUserSettings)
+class UNREALED_API UCurveEditorSettings : public UObject
+{
+public:
+	GENERATED_UCLASS_BODY()
+
+	DECLARE_MULTICAST_DELEGATE( FOnCurveEditorCurveVisibilityChanged );
+
+	/** Gets whether or not the curve editor auto frames the selected curves. */
+	bool GetAutoFrameCurveEditor() const;
+	/** Sets whether or not the curve editor auto frames the selected curves. */
+	void SetAutoFrameCurveEditor(bool InbAutoFrameCurveEditor);
+
+	/** Gets whether or not to show curve tool tips in the curve editor. */
+	bool GetShowCurveEditorCurveToolTips() const;
+	/** Sets whether or not to show curve tool tips in the curve editor. */
+	void SetShowCurveEditorCurveToolTips(bool InbShowCurveEditorCurveToolTips);
+
+	/** Gets the current curve visibility. */
+	ECurveEditorCurveVisibility::Type GetCurveVisibility() const;
+	/** Sets the current curve visibility. */
+	void SetCurveVisibility(ECurveEditorCurveVisibility::Type InCurveVisibility);
+
+	/** Gets the current tangent visibility. */
+	ECurveEditorTangentVisibility::Type GetTangentVisibility() const;
+	/** Sets the current tangent visibility. */
+	void SetTangentVisibility(ECurveEditorTangentVisibility::Type InTangentVisibility);
+
+	/** Gets the multicast delegate which is run whenever the curve editor curve visibility changes. */
+	FOnCurveEditorCurveVisibilityChanged& GetOnCurveEditorCurveVisibilityChanged();
+
+protected:
+	UPROPERTY( config )
+	bool bAutoFrameCurveEditor;
+
+	UPROPERTY( config )
+	bool bShowCurveEditorCurveToolTips;
+
+	UPROPERTY( config )
+	TEnumAsByte<ECurveEditorCurveVisibility::Type> CurveVisibility;
+
+	UPROPERTY( config )
+	TEnumAsByte<ECurveEditorTangentVisibility::Type> TangentVisibility;
+
+	FOnCurveEditorCurveVisibilityChanged OnCurveEditorCurveVisibilityChanged;
+};
diff --git a/Engine/Source/Editor/UnrealEd/Public/RichCurveEditorCommands.h b/Engine/Source/Editor/UnrealEd/Public/RichCurveEditorCommands.h
--- a/Engine/Source/Editor/UnrealEd/Public/RichCurveEditorCommands.h
+++ b/Engine/Source/Editor/UnrealEd/Public/RichCurveEditorCommands.h
@@ -23,12 +23,39 @@ public:
 	TSharedPtr<FUICommandInfo> ZoomToFitVertical;
 	TSharedPtr<FUICommandInfo> ZoomToFit;
 	TSharedPtr<FUICommandInfo> ToggleSnapping;
+
 	TSharedPtr<FUICommandInfo> InterpolationConstant;
 	TSharedPtr<FUICommandInfo> InterpolationLinear;
 	TSharedPtr<FUICommandInfo> InterpolationCubicAuto;
 	TSharedPtr<FUICommandInfo> InterpolationCubicUser;
 	TSharedPtr<FUICommandInfo> InterpolationCubicBreak;
 
+	TSharedPtr<FUICommandInfo> FlattenTangents;
+	TSharedPtr<FUICommandInfo> StraightenTangents;
+
+	TSharedPtr<FUICommandInfo> SetPreInfinityExtrapCycle;
+	TSharedPtr<FUICommandInfo> SetPreInfinityExtrapCycleWithOffset;
+	TSharedPtr<FUICommandInfo> SetPreInfinityExtrapOscillate;
+	TSharedPtr<FUICommandInfo> SetPreInfinityExtrapLinear;
+	TSharedPtr<FUICommandInfo> SetPreInfinityExtrapConstant;
+
+	TSharedPtr<FUICommandInfo> SetPostInfinityExtrapCycle;
+	TSharedPtr<FUICommandInfo> SetPostInfinityExtrapCycleWithOffset;
+	TSharedPtr<FUICommandInfo> SetPostInfinityExtrapOscillate;
+	TSharedPtr<FUICommandInfo> SetPostInfinityExtrapLinear;
+	TSharedPtr<FUICommandInfo> SetPostInfinityExtrapConstant;
+
+	TSharedPtr<FUICommandInfo> SetAllCurveVisibility;
+	TSharedPtr<FUICommandInfo> SetSelectedCurveVisibility;
+	TSharedPtr<FUICommandInfo> SetAnimatedCurveVisibility;
+
+	TSharedPtr<FUICommandInfo> SetAllTangentsVisibility;
+	TSharedPtr<FUICommandInfo> SetSelectedKeysTangentVisibility;
+	TSharedPtr<FUICommandInfo> SetNoTangentsVisibility;
+
+	TSharedPtr<FUICommandInfo> ToggleAutoFrameCurveEditor;
+	TSharedPtr<FUICommandInfo> ToggleShowCurveEditorCurveToolTips;
+
 public:
 	virtual void RegisterCommands() override;
 };
\ No newline at end of file
diff --git a/Engine/Source/Editor/UnrealEd/Public/SCurveEditor.h b/Engine/Source/Editor/UnrealEd/Public/SCurveEditor.h
--- a/Engine/Source/Editor/UnrealEd/Public/SCurveEditor.h
+++ b/Engine/Source/Editor/UnrealEd/Public/SCurveEditor.h
@@ -4,6 +4,7 @@
 #ifndef __SCurveEditor_h__
 #define __SCurveEditor_h__
 
+#include "CurveEditorSettings.h"
 #include "EditorUndoClient.h"
 
 //////////////////////////////////////////////////////////////////////////
@@ -115,16 +116,30 @@ private:
 		PreDrag,
 		/** The user is dragging the selected keys. */
 		DragKey,
+		/** The user is free dragging the selected keys. */
+		FreeDrag,
 		/** The user is dragging a selected tangent handle. */
 		DragTangent,
 		/** The user is performing a marquee selection of keys. */
 		MarqueeSelect,
 		/** The user is panning the curve view. */
 		Pan,
+		/** The user is zooming the curve view. */
+		Zoom,
 		/** There is no active drag operation. */
 		None
 	};
 
+	enum class EMovementAxisLock
+	{
+		/** Lock movement to horizontal axis */
+		AxisLock_Horizontal,
+		/** Lock movement to vertical axis */
+		AxisLock_Vertical,
+		/** Don't lock movement */
+		None
+	};
+
 public:
 	SLATE_BEGIN_ARGS( SCurveEditor )
 		: _ViewMinInput(0.0f)
@@ -136,7 +151,6 @@ public:
 		, _SnappingEnabled(false)
 		, _TimelineLength(5.0f)
 		, _DesiredSize(FVector2D::ZeroVector)
-		, _ShowCurveToolTips(false)
 		, _DrawCurve(true)
 		, _HideUI(true)
 		, _AllowZoomOutput(true)
@@ -163,7 +177,6 @@ public:
 		SLATE_ATTRIBUTE( bool, SnappingEnabled )
 		SLATE_ATTRIBUTE( float, TimelineLength )
 		SLATE_ATTRIBUTE( FVector2D, DesiredSize )
-		SLATE_ATTRIBUTE( bool, ShowCurveToolTips )
 		SLATE_ARGUMENT( bool, DrawCurve )
 		SLATE_ARGUMENT( bool, HideUI )
 		SLATE_ARGUMENT( bool, AllowZoomOutput )
@@ -256,6 +269,16 @@ private:
 		{
 		}
 
+		FSelectedTangent(FSelectedCurveKey InKey) : Key(InKey)
+		{
+		}
+
+		/** Does the key and the arrival match ?*/
+		bool			operator == (const FSelectedTangent& Other) const 
+		{
+			return (Key == Other.Key) && (bIsArrival == Other.bIsArrival);
+		}
+
 		/** If this is a valid Curve/Key */
 		bool			IsValid() const;
 
@@ -297,24 +320,47 @@ private:
 	/** Get the set of keys within a rectangle in local space */
 	TArray<FSelectedCurveKey> GetEditableKeysWithinMarquee(const FGeometry& InMyGeometry, FVector2D MarqueeTopLeft, FVector2D MarqueeBottomRight) const;
 
-	/** Empy key selecttion set */
-	void EmptySelection();
+	/** Get the set of tangents within a rectangle in local space */
+	TArray<FSelectedTangent> GetEditableTangentsWithinMarquee(const FGeometry& InMyGeometry, FVector2D MarqueeTopLeft, FVector2D MarqueeBottomRight) const;
+
+	/** Empty key selection set */
+	void EmptyKeySelection();
 	/** Add a key to the selection set */
-	void AddToSelection(FSelectedCurveKey Key);
+	void AddToKeySelection(FSelectedCurveKey Key);
 	/** Remove a key from the selection set */
-	void RemoveFromSelection(FSelectedCurveKey Key);
+	void RemoveFromKeySelection(FSelectedCurveKey Key);
 	/** See if a key is currently selected */
 	bool IsKeySelected(FSelectedCurveKey Key) const;
 	/** See if any keys are selected */
 	bool AreKeysSelected() const;
 
+	/** Empty tangent selection set */
+	void EmptyTangentSelection();
+	/** Add a tangent to the selection set */
+	void AddToTangentSelection(FSelectedTangent Tangent);
+	/** Remove a tangent from the selection set */
+	void RemoveFromTangentSelection(FSelectedTangent Tangent);
+	/** See if a tangent is currently selected */
+	bool IsTangentSelected(FSelectedTangent Tangent) const;
+	/** See if any tangents are selected */
+	bool AreTangentsSelected() const;
+
+	/** Is the tangent visible? */
+	bool IsTangentVisible(FRichCurve*, FKeyHandle, bool& IsTangentSelected, bool& IsArrivalSelected, bool& IsLeaveSelected) const;
+
+	/** Empty key and tangent selection set */
+	void EmptyAllSelection();
+
+	/** Ensure that selected keys and tangents are still valid */
+	void ValidateSelection();
+
 	/** Get the value of the desired key as text */
 	TOptional<float> GetKeyValue(FSelectedCurveKey Key) const;
 	/** Get the time of the desired key */
 	TOptional<float> GetKeyTime(FSelectedCurveKey Key) const;
 
 	/** Move the selected keys */
-	void MoveSelectedKeys(FVector2D NewLocation);
+	void MoveSelectedKeys(FVector2D Delta);
 
 	/** Function to check whether the current track is editable */
 	bool IsEditingEnabled() const;
@@ -378,7 +424,7 @@ private:
 
 	/** Paint the tangent for a key with cubic curves */
 	void PaintTangent( FTrackScaleInfo &ScaleInfo, FRichCurve* Curve, FKeyHandle KeyHandle, FVector2D KeyLocation, FSlateWindowElementList &OutDrawElements, int32 LayerId, 
-					   const FGeometry &AllottedGeometry, const FSlateRect& MyClippingRect, ESlateDrawEffect::Type DrawEffects, int32 LayerToUse, const FWidgetStyle &InWidgetStyle ) const;
+					   const FGeometry &AllottedGeometry, const FSlateRect& MyClippingRect, ESlateDrawEffect::Type DrawEffects, int32 LayerToUse, const FWidgetStyle &InWidgetStyle, bool bTangentSelected, bool bIsArrivalSelected, bool bIsLeaveSelected ) const;
 
 	/** Paint Grid lines, these make it easier to visualize relative distance */
 	void PaintGridLines( const FGeometry &AllottedGeometry, FTrackScaleInfo &ScaleInfo, FSlateWindowElementList &OutDrawElements, int32 LayerId, const FSlateRect& MyClippingRect, ESlateDrawEffect::Type DrawEffects )const;
@@ -413,21 +459,24 @@ private:
 		return true;
 	}
 
+	/** Zoom the view. */
+	void ZoomView(FVector2D Delta);
+
 	/* Generates the line(s) for rendering between KeyIndex and the following key. */
 	void CreateLinesForSegment( FRichCurve* Curve, const FRichCurveKey& Key1, const FRichCurveKey& Key2, TArray<FVector2D>& Points, FTrackScaleInfo &ScaleInfo) const;
 	
 	/** Detect if user is clicking on a curve */
 	TSharedPtr<FCurveViewModel> HitTestCurves(const FGeometry& InMyGeometry, const FPointerEvent& InMouseEvent);
 
-	/* user is moving the tangent for a key */
-	void OnMoveTangent(FVector2D MouseCurvePosition);
+	/* user is moving the tangents */
+	void MoveTangents(FTrackScaleInfo ScaleInfo, FVector2D Delta);
 
 	//Curve Selection interface 
 
-	/** Construct widget that allows user to select which curve to edit if there are multiple */
+	/** Construct widget that allows user to select which curve to edit if tthere are multiple */
 	TSharedRef<SWidget> CreateCurveSelectionWidget() const;
 
-	/** Create ontext Menu for waring menu*/
+	/** Create context Menu for waring menu*/
 	void	PushWarningMenu(FVector2D Position, const FText& Message);
 
 	/** Create context Menu for key interpolation settings*/
@@ -438,6 +487,15 @@ private:
 
 	bool IsInterpolationModeSelected(ERichCurveInterpMode InterpMode, ERichCurveTangentMode TangentMode);
 
+	/** Flatten or straighten tangents */
+	void OnFlattenOrStraightenTangents(bool bFlattenTangents);
+
+	/** Called when the user selects the extrapolation type */
+	void OnSelectPreInfinityExtrap(ERichCurveExtrapolation Extrapolation);
+	bool IsPreInfinityExtrapSelected(ERichCurveExtrapolation Extrapolation);
+	void OnSelectPostInfinityExtrap(ERichCurveExtrapolation Extrapolation);
+	bool IsPostInfinityExtrapSelected(ERichCurveExtrapolation Extrapolation);
+
 	/** Begin a transaction for dragging a key or tangent */
 	void	BeginDragTransaction();
 
@@ -496,8 +554,14 @@ protected:
 	UNREALED_API void SetInputMinMax(float NewMin, float NewMax);
 	UNREALED_API void SetOutputMinMax(float NewMin, float NewMax);
 
+	/** Access the user-supplied settings object */
+	UCurveEditorSettings* GetSettings() const { return Settings; }
+
 private:
 
+	/** User-supplied object for this curve editor */
+	UCurveEditorSettings* Settings;
+
 	/** Curve selection*/
 	TWeakPtr<SBox>		CurveSelectionWidget;
 
@@ -526,8 +590,8 @@ private:
 	/** Array of selected keys*/
 	TArray<FSelectedCurveKey> SelectedKeys;
 
-	/** Currently selected tangent */
-	FSelectedTangent	SelectedTangent;
+	/** Array of selected tangents */
+	TArray<FSelectedTangent> SelectedTangents;
 
 	/** Minimum input of data range  */
 	TAttribute< TOptional<float> >	DataMinInput;
@@ -587,14 +651,11 @@ protected:
 	/** Whether or not snapping is enabled. */
 	TAttribute<bool> bSnappingEnabled;
 
-	/** Whether or not to show tool tips for curves. */
-	TAttribute<bool> bShowCurveToolTips;
-
 	/** True if you want the curve editor to fit to zoom **/
-	bool				bZoomToFitVertical;
+	bool bZoomToFitVertical;
 
 	/** True if you want the curve editor to fit to zoom **/
-	bool				bZoomToFitHorizontal;
+	bool bZoomToFitHorizontal;
 
 	/** True if the sliders are being used to adjust point values **/
 	bool bIsUsingSlider;
@@ -614,6 +675,9 @@ protected:
 	/** The state of the current drag operation happening with the widget, if any. */
 	EDragState DragState;
 
+	/** The movement axis lock state */
+	EMovementAxisLock MovementAxisLock;
+
 	/** The number of pixels which the mouse must move before a drag operation starts. */
 	float DragThreshold;
 
@@ -623,6 +687,9 @@ protected:
 	/** A map of selected key handle to their starting locations at the beginning of a drag operation. */
 	TMap<FKeyHandle, FVector2D> PreDragKeyLocations;
 
+	/** A map of selected key handles to their tangent values at the beginning of a drag operation. */
+	TMap<FKeyHandle, FVector2D> PreDragTangents;
+
 	/** The text to display for the input axis. */
 	FText InputAxisName;
 	/** The text to display for the output axis. */
diff --git a/Engine/Source/Runtime/Engine/Classes/Curves/CurveBase.h b/Engine/Source/Runtime/Engine/Classes/Curves/CurveBase.h
--- a/Engine/Source/Runtime/Engine/Classes/Curves/CurveBase.h
+++ b/Engine/Source/Runtime/Engine/Classes/Curves/CurveBase.h
@@ -156,6 +156,17 @@ enum ERichCurveTangentWeightMode
 	RCTWM_WeightedBoth
 };
 
+/** Enum to indicate extrapolation option. */
+UENUM()
+enum ERichCurveExtrapolation
+{
+	RCCE_Cycle,
+	RCCE_CycleWithOffset,
+	RCCE_Oscillate,
+	RCCE_Linear,
+	RCCE_Constant
+};
+
 /** One key in a rich, editable float curve */
 USTRUCT()
 struct ENGINE_API FRichCurveKey
@@ -271,6 +282,13 @@ struct ENGINE_API FRichCurve : public FIndexedCurve
 	GENERATED_USTRUCT_BODY()
 
 public:
+	FRichCurve() 
+	: FIndexedCurve()
+	, PreInfinityExtrap(RCCE_Constant)
+	, PostInfinityExtrap(RCCE_Constant)
+	{
+	}
+
 	virtual ~FRichCurve()
 	{
 	}
@@ -356,6 +374,9 @@ public:
 	/** Clear all keys. */
 	void Reset();
 
+	/** Remap InTime based on pre and post infinity extrapolation values */
+	void RemapTimeValue(float& InTime, float& CycleValueOffset) const;
+
 	/** Evaluate this rich curve at the specified time */
 	float Eval(float InTime, float DefaultValue = 0.0f) const;
 
@@ -367,6 +388,15 @@ public:
 
 	/** Determine if two RichCurves are the same */
 	bool operator == (const FRichCurve& Curve) const;
+
+	/** Pre-infinity extrapolation state */
+	UPROPERTY()
+	TEnumAsByte<ERichCurveExtrapolation> PreInfinityExtrap;
+
+	/** Post-infinity extrapolation state */
+	UPROPERTY()
+	TEnumAsByte<ERichCurveExtrapolation> PostInfinityExtrap;
+
 private:
 	/** Sorted array of keys */
 	UPROPERTY()
diff --git a/Engine/Source/Runtime/Engine/Private/CurveBase.cpp b/Engine/Source/Runtime/Engine/Private/CurveBase.cpp
--- a/Engine/Source/Runtime/Engine/Private/CurveBase.cpp
+++ b/Engine/Source/Runtime/Engine/Private/CurveBase.cpp
@@ -750,8 +750,98 @@ static float BezierInterp2(float P0, float Y1, float Y2, float P3, float mu)
 	return Result;
 }
 
-float FRichCurve::Eval(const float InTime, float DefaultValue) const
+static void CycleTime(float MinTime, float MaxTime, float& InTime, int& CycleCount)
+{
+	float InitTime = InTime;
+	float Duration = MaxTime - MinTime;
+
+	if (InTime > MaxTime)
+	{
+		CycleCount = FMath::FloorToInt((MaxTime-InTime)/Duration);
+		InTime = InTime + Duration*CycleCount;
+	}
+	else if (InTime < MinTime)
+	{
+		CycleCount = FMath::FloorToInt((InTime-MinTime)/Duration);
+		InTime = InTime - Duration*CycleCount;
+	}
+
+	if (InTime == MaxTime && InitTime < MinTime)
+	{
+		InTime = MinTime;
+	}
+
+	if (InTime == MinTime && InitTime > MaxTime)
+	{
+		InTime = MaxTime;
+	}
+
+	CycleCount = FMath::Abs(CycleCount);
+}
+
+void FRichCurve::RemapTimeValue(float& InTime, float& CycleValueOffset) const
 {
+	const int32 NumKeys = Keys.Num();
+	if (NumKeys < 2)
+	{
+		return;
+	} 
+	else if (InTime <= Keys[0].Time)
+	{
+		if (PreInfinityExtrap != RCCE_Linear && PreInfinityExtrap != RCCE_Constant)
+		{
+			float MinTime = Keys[0].Time;
+			float MaxTime = Keys[NumKeys - 1].Time;
+
+			int CycleCount = 0;
+			CycleTime(MinTime, MaxTime, InTime, CycleCount);
+
+			if (PreInfinityExtrap == RCCE_CycleWithOffset)
+			{
+				float DV = Keys[0].Value - Keys[NumKeys - 1].Value;
+				CycleValueOffset = DV * CycleCount;
+			}
+			else if (PreInfinityExtrap == RCCE_Oscillate)
+			{
+				if (CycleCount % 2 == 1)
+				{
+					InTime = MinTime + (MaxTime - InTime);
+				}
+			}
+		}
+	}
+	else if (InTime >= Keys[NumKeys - 1].Time)
+	{
+		if (PostInfinityExtrap != RCCE_Linear && PostInfinityExtrap != RCCE_Constant)
+		{
+			float MinTime = Keys[0].Time;
+			float MaxTime = Keys[NumKeys - 1].Time;
+
+			int CycleCount = 0; 
+			CycleTime(MinTime, MaxTime, InTime, CycleCount);
+
+			if (PostInfinityExtrap == RCCE_CycleWithOffset)
+			{
+				float DV = Keys[NumKeys - 1].Value - Keys[0].Value;
+				CycleValueOffset = DV * CycleCount;
+			}
+			else if (PostInfinityExtrap == RCCE_Oscillate)
+			{
+				if (CycleCount % 2 == 1)
+				{
+					InTime = MinTime + (MaxTime - InTime);
+				}
+			}
+		}
+	}
+}
+
+float FRichCurve::Eval(float InTime, float DefaultValue) const
+{
+	// Remap time if extrapolation is present and compute offset value to use if cycling 
+	float CycleValueOffset = 0;
+	RemapTimeValue(InTime, CycleValueOffset);
+
 	const int32 NumKeys = Keys.Num();
 	float InterpVal = DefaultValue;
 
@@ -762,8 +852,26 @@ float FRichCurve::Eval(const float InTime, float DefaultValue) const
 	} 
 	else if (NumKeys < 2 || (InTime <= Keys[0].Time))
 	{
-		// If only one point, or before the first point in the curve, return the first points value.
-		InterpVal = Keys[0].Value;
+		if (PreInfinityExtrap == RCCE_Linear && NumKeys > 1)
+		{
+			float DT = Keys[1].Time - Keys[0].Time;
+			
+			if (FMath::IsNearlyZero(DT))
+			{
+				InterpVal = Keys[0].Value;
+			}
+			else
+			{
+				float DV = Keys[1].Value - Keys[0].Value;
+				float Slope = DV / DT;
+				InterpVal = Slope * (InTime - Keys[0].Time) + Keys[0].Value;
+			}
+		}
+		else
+		{
+			// Otherwise if constant or in a cycle or oscillate, always use the first key value
+			InterpVal = Keys[0].Value;
+		}
 	}
 	else if (InTime < Keys[NumKeys - 1].Time)
 	{
@@ -817,11 +925,29 @@ float FRichCurve::Eval(const float InTime, float DefaultValue) const
 	}
 	else
 	{
-		// If beyond the last point in the curve, return its value.
-		InterpVal = Keys[NumKeys - 1].Value;
+		if (PostInfinityExtrap == RCCE_Linear)
+		{
+			float DT = Keys[NumKeys - 2].Time - Keys[NumKeys - 1].Time;
+			
+			if (FMath::IsNearlyZero(DT))
+			{
+				InterpVal = Keys[NumKeys - 1].Value;
+			}
+			else
+			{
+				float DV = Keys[NumKeys - 2].Value - Keys[NumKeys - 1].Value;
+				float Slope = DV / DT;
+				InterpVal = Slope * (InTime - Keys[NumKeys - 1].Time) + Keys[NumKeys - 1].Value;
+			}
+		}
+		else
+		{
+			// Otherwise if constant or in a cycle or oscillate, always use the last key value
+			InterpVal = Keys[NumKeys - 1].Value;
+		}
 	}
 
-	return InterpVal;
+	return InterpVal+CycleValueOffset;
 }
 
 
@@ -838,6 +964,12 @@ bool FRichCurve::operator==(const FRichCurve& Curve) const
 			return false;
 		}
 	}
+
+	if (PreInfinityExtrap != Curve.PreInfinityExtrap || PostInfinityExtrap != Curve.PostInfinityExtrap)
+	{
+		return false;
+	}
+
 	return true;
 }
 
