Optimization to FTransaction::FObjectRecord::FReader::Preload so that it can early-out if it knows the given object is not referenced in any of its records.
Based on GitHub #840, but made safe in the case that an inner transaction is cancelled and a subsequent inner transaction then adds new records.

diff --git a/Engine/Source/Editor/UnrealEd/Classes/Editor/Transactor.h b/Engine/Source/Editor/UnrealEd/Classes/Editor/Transactor.h
--- a/Engine/Source/Editor/UnrealEd/Classes/Editor/Transactor.h
+++ b/Engine/Source/Editor/UnrealEd/Classes/Editor/Transactor.h
@@ -309,11 +309,14 @@ protected:
 			{
 				if( Owner )
 				{
-					for( int32 i=0; i<Owner->Records.Num(); i++ )
+					if ( Owner->ObjectMap.Contains( InObject ) )
 					{
-						if( Owner->Records[i].Object.Get()==InObject )
+						for (int32 i = 0; i < Owner->Records.Num(); i++)
 						{
-							Owner->Records[i].Restore( Owner );
+							if (Owner->Records[i].Object.Get() == InObject)
+							{
+								Owner->Records[i].Restore(Owner);
+							}
 						}
 					}
 				}
diff --git a/Engine/Source/Editor/UnrealEd/Private/EditorTransaction.cpp b/Engine/Source/Editor/UnrealEd/Private/EditorTransaction.cpp
--- a/Engine/Source/Editor/UnrealEd/Private/EditorTransaction.cpp
+++ b/Engine/Source/Editor/UnrealEd/Private/EditorTransaction.cpp
@@ -128,10 +128,13 @@ void FTransaction::RemoveRecords( int32 Count /* = 1  */ )
 {
 	if ( Count > 0 && Records.Num() >= Count )
 	{
-		Records.RemoveAt( Records.Num() - Count, Count );
+		// Remove anything from the ObjectMap which is about to be removed from the Records array
+		for (int32 Index = 0; Index < Count; Index++)
+		{
+			ObjectMap.Remove( Records[Records.Num() - Count + Index].Object.Get() );
+		}
 
-		// Kill our object maps that are used to track redundant saves
-		ObjectMap.Empty();
+		Records.RemoveAt( Records.Num() - Count, Count );
 	}
 }
 
