Change the way import setting works

- fixes issue of https://jira.ol.epicgames.net/browse/UE-18214

#code reivew: Matt.Kulehnschmidt

diff --git a/Engine/Source/Editor/UnrealEd/Private/Fbx/FbxMainImport.cpp b/Engine/Source/Editor/UnrealEd/Private/Fbx/FbxMainImport.cpp
--- a/Engine/Source/Editor/UnrealEd/Private/Fbx/FbxMainImport.cpp
+++ b/Engine/Source/Editor/UnrealEd/Private/Fbx/FbxMainImport.cpp
@@ -399,30 +399,55 @@ int32 FFbxImporter::GetImportType(const FString& InFilename)
 			UE_LOG(LogFbx, Log, TEXT("ItemName: %s, ItemCount : %d"), *NameBuffer, ItemCount);
 		}
 
-		for ( ItemIndex = 0; ItemIndex < Statistics.GetNbItems(); ItemIndex++ )
+		FbxSceneInfo SceneInfo;
+		if (GetSceneInfo(Filename, SceneInfo))
 		{
-			Statistics.GetItemPair(ItemIndex, ItemName, ItemCount);
-			const char* NameBuffer = ItemName.Buffer();
-			if ( ItemName == "Deformer" && ItemCount > 0 )
+			if (SceneInfo.SkinnedMeshNum > 0)
 			{
-				// if SkeletalMesh is found, just return
 				Result = 1;
-				break;
 			}
-			// if Geometry is found, sets it, but it can be overwritten by Deformer
-			else if ( ItemName == "Geometry" && ItemCount > 0)
+			else if (SceneInfo.TotalGeometryNum > 0)
 			{
-				// let it still loop through even if Geometry is found
-				// Deformer can overwrite this information
 				Result = 0;
 			}
-			// Check for animation data. It can be overwritten by Geometry or Deformer
-			else if ( (ItemName == "AnimationCurve" || ItemName == "AnimationCurveNode") && ItemCount > 0 )
+
+			if (SceneInfo.TakeName != nullptr)
 			{
 				bHasAnimation = true;
 			}
 		}
-		Importer->Destroy();
+		else
+		{
+			for(ItemIndex = 0; ItemIndex < Statistics.GetNbItems(); ItemIndex++)
+			{
+				Statistics.GetItemPair(ItemIndex, ItemName, ItemCount);
+				const char* NameBuffer = ItemName.Buffer();
+				if(ItemName == "Deformer" && ItemCount > 0)
+				{
+					// if SkeletalMesh is found, just return
+					Result = 1;
+					break;
+				}
+				// if Geometry is found, sets it, but it can be overwritten by Deformer
+				else if(ItemName == "Geometry" && ItemCount > 0)
+				{
+					// let it still loop through even if Geometry is found
+					// Deformer can overwrite this information
+					Result = 0;
+				}
+				// Check for animation data. It can be overwritten by Geometry or Deformer
+				else if((ItemName == "AnimationCurve" || ItemName == "AnimationCurveNode") && ItemCount > 0)
+				{
+					bHasAnimation = true;
+				}
+			}
+		}
+
+		if (Importer)
+		{
+			Importer->Destroy();
+		}
+
 		Importer = NULL;
 		CurPhase = NOTSTARTED;
 
