Move Deserialize calls out of check macros so test doesn't crash with DO_CHECK 0

diff --git a/Engine/Source/Editor/UnrealEd/Private/Tests/EditorAutomationTests.cpp b/Engine/Source/Editor/UnrealEd/Private/Tests/EditorAutomationTests.cpp
--- a/Engine/Source/Editor/UnrealEd/Private/Tests/EditorAutomationTests.cpp
+++ b/Engine/Source/Editor/UnrealEd/Private/Tests/EditorAutomationTests.cpp
@@ -261,7 +261,8 @@ bool FJsonAutomationTest::RunTest(const FString& Parameters)
 		TSharedRef< TJsonReader<> > Reader = TJsonReaderFactory<>::Create( InputString );
 
 		TArray< TSharedPtr<FJsonValue> > Array;
-		check( FJsonSerializer::Deserialize( Reader, Array ) );
+		bool bSuccessful = FJsonSerializer::Deserialize(Reader, Array);
+		check(bSuccessful);
 		check( Array.Num() == 1 );
 		check( Array[0].IsValid() )
 
@@ -293,7 +294,9 @@ bool FJsonAutomationTest::RunTest(const FString& Parameters)
 		TSharedRef< TJsonReader<> > Reader = TJsonReaderFactory<>::Create(InputString);
 
 		TArray< TSharedPtr<FJsonValue> > Array;
-		check(FJsonSerializer::Deserialize(Reader, Array));
+
+		bool bSuccessful = FJsonSerializer::Deserialize(Reader, Array);
+		check(bSuccessful);
 		check(Array.Num() == 3);
 		check(Array[0].IsValid());
 		check(Array[1].IsValid());
@@ -330,7 +333,8 @@ bool FJsonAutomationTest::RunTest(const FString& Parameters)
 		TSharedRef< TJsonReader<> > Reader = TJsonReaderFactory<>::Create(InputString);
 
 		TArray< TSharedPtr<FJsonValue> > Array;
-		check(FJsonSerializer::Deserialize(Reader, Array));
+		bool bSuccessful = FJsonSerializer::Deserialize(Reader, Array);
+		check(bSuccessful);
 		check(Array.Num() == 4);
 		check(Array[0].IsValid());
 		check(Array[1].IsValid());
@@ -368,7 +372,8 @@ bool FJsonAutomationTest::RunTest(const FString& Parameters)
 		TSharedRef< TJsonReader<> > Reader = TJsonReaderFactory<>::Create(InputString);
 
 		TArray< TSharedPtr<FJsonValue> > Array;
-		check(FJsonSerializer::Deserialize(Reader, Array));
+		bool bSuccessful = FJsonSerializer::Deserialize(Reader, Array);
+		check(bSuccessful);
 		check(Array.Num() == 4);
 		check(Array[0].IsValid());
 		check(Array[1].IsValid());
@@ -413,7 +418,8 @@ bool FJsonAutomationTest::RunTest(const FString& Parameters)
 		TSharedRef< TJsonReader<> > Reader = TJsonReaderFactory<>::Create(InputString);
 
 		TArray< TSharedPtr<FJsonValue> > Array;
-		check(FJsonSerializer::Deserialize(Reader, Array));
+		bool bSuccessful = FJsonSerializer::Deserialize(Reader, Array);
+		check(bSuccessful);
 		check(Array.Num() == 6);
 		check(Array[0].IsValid());
 		check(Array[1].IsValid());
@@ -463,7 +469,8 @@ bool FJsonAutomationTest::RunTest(const FString& Parameters)
 		TSharedRef< TJsonReader<> > Reader = TJsonReaderFactory<>::Create( InputString );
 
 		TSharedPtr<FJsonObject> Object;
-		check( FJsonSerializer::Deserialize( Reader, Object ) );
+		bool bSuccessful = FJsonSerializer::Deserialize(Reader, Object);
+		check(bSuccessful);
 		check( Object.IsValid() );
 
 		const TSharedPtr<FJsonValue>* Value = Object->Values.Find(TEXT("Value"));
@@ -495,7 +502,8 @@ bool FJsonAutomationTest::RunTest(const FString& Parameters)
 		TSharedRef< TJsonReader<> > Reader = TJsonReaderFactory<>::Create( InputString );
 
 		TSharedPtr<FJsonObject> Object;
-		check( FJsonSerializer::Deserialize( Reader, Object ) );
+		bool bSuccessful = FJsonSerializer::Deserialize(Reader, Object);
+		check(bSuccessful);
 		check( Object.IsValid() );
 
 		double TestValues[] = {2.544e+15, -0.544e-2, 251e3, -0.0, 843};
@@ -537,7 +545,8 @@ bool FJsonAutomationTest::RunTest(const FString& Parameters)
 		TSharedRef< TJsonReader<> > Reader = TJsonReaderFactory<>::Create( InputString );
 
 		TSharedPtr<FJsonObject> Object;
-		check( FJsonSerializer::Deserialize( Reader, Object ) );
+		bool bSuccessful = FJsonSerializer::Deserialize(Reader, Object);
+		check(bSuccessful);
 		check( Object.IsValid() );
 
 		bool TestValues[] = {true, true, false};
@@ -598,7 +607,8 @@ bool FJsonAutomationTest::RunTest(const FString& Parameters)
 		TSharedRef< TJsonReader<> > Reader = TJsonReaderFactory<>::Create( InputStringWithExtraWhitespace );
 
 		TSharedPtr<FJsonObject> Object;
-		check( FJsonSerializer::Deserialize( Reader, Object ) );
+		bool bSuccessful = FJsonSerializer::Deserialize(Reader, Object);
+		check(bSuccessful);
 		check( Object.IsValid() );
 
 		const TSharedPtr<FJsonValue>* InnerValueFail = Object->Values.Find(TEXT("InnerValue"));
@@ -655,7 +665,8 @@ bool FJsonAutomationTest::RunTest(const FString& Parameters)
 		TSharedRef< TJsonReader<> > Reader = TJsonReaderFactory<>::Create( InputString );
 
 		TSharedPtr<FJsonObject> Object;
-		check( FJsonSerializer::Deserialize( Reader, Object ) );
+		bool bSuccessful = FJsonSerializer::Deserialize(Reader, Object);
+		check(bSuccessful);
 		check( Object.IsValid() );
 
 		const TSharedPtr<FJsonValue>* InnerValueFail = Object->Values.Find(TEXT("InnerValue"));
