Fix adding one extra frame all the time.

https://jira.ol.epicgames.net/browse/UE-16027

diff --git a/Engine/Source/Editor/UnrealEd/Private/SkeletalMeshEdit.cpp b/Engine/Source/Editor/UnrealEd/Private/SkeletalMeshEdit.cpp
--- a/Engine/Source/Editor/UnrealEd/Private/SkeletalMeshEdit.cpp
+++ b/Engine/Source/Editor/UnrealEd/Private/SkeletalMeshEdit.cpp
@@ -1147,7 +1147,7 @@ bool UnFbx::FFbxImporter::ImportAnimation(USkeleton* Skeleton, UAnimSequence * D
 		FMatrix AddedMatrix = Converter.ConvertMatrix(FbxAddedMatrix);
 
 		const int32 NumSamplingKeys = FMath::FloorToInt(AnimTimeSpan.GetDuration().GetSecondDouble() * ResampleRate);
-		const FbxTime TimeIncrement = AnimTimeSpan.GetDuration() / FMath::Max(NumSamplingKeys,1);
+		const FbxTime TimeIncrement = (NumSamplingKeys > 1)? AnimTimeSpan.GetDuration() / (NumSamplingKeys - 1) : AnimTimeSpan.GetDuration();
 		for(int32 SourceTrackIdx = 0; SourceTrackIdx < FbxRawBoneNames.Num(); ++SourceTrackIdx)
 		{
 			int32 NumKeysForTrack = 0;
