Sequencer: Fix tangent manipulation so that it's relative to the mouse motion. This allows you to do multi tangent manipulation and also to middle click drag anywhere.

diff --git a/Engine/Source/Editor/UnrealEd/Private/SCurveEditor.cpp b/Engine/Source/Editor/UnrealEd/Private/SCurveEditor.cpp
--- a/Engine/Source/Editor/UnrealEd/Private/SCurveEditor.cpp
+++ b/Engine/Source/Editor/UnrealEd/Private/SCurveEditor.cpp
@@ -1658,11 +1658,9 @@ void SCurveEditor::ProcessDrag(const FGeometry& InMyGeometry, const FPointerEven
 	}
 	else if (DragState == EDragState::DragTangent)
 	{
-		FVector2D MousePositionScreen = InMouseEvent.GetScreenSpacePosition();
-		MousePositionScreen -= InMyGeometry.AbsolutePosition;
-		FVector2D MousePositionCurve(ScaleInfo.LocalXToInput(MousePositionScreen.X), ScaleInfo.LocalYToOutput(MousePositionScreen.Y));
-		FVector2D Delta = MousePositionCurve;
-		MoveTangents(ScaleInfo, Delta);
+		FVector2D MousePositionScreen = InMyGeometry.AbsoluteToLocal(InMouseEvent.GetScreenSpacePosition());
+		FVector2D MouseDownPositionScreen = MouseDownLocation;
+		MoveTangents(ScaleInfo, MousePositionScreen - MouseDownPositionScreen);
 	}
 	else if (DragState == EDragState::Pan)
 	{
@@ -2929,20 +2927,60 @@ void SCurveEditor::MoveTangents(FTrackScaleInfo ScaleInfo, FVector2D Delta)
 		auto& RichKey = SelectedTangent.Key.Curve->GetKey(SelectedTangent.Key.KeyHandle);
 
 		const FSelectedCurveKey &Key = SelectedTangent.Key;
+		float PreDragArriveTangent = PreDragTangents[SelectedTangent.Key.KeyHandle][0];
+		float PreDragLeaveTangent = PreDragTangents[SelectedTangent.Key.KeyHandle][1];
+
+		// Get tangent points in screen space
+		FVector2D ArriveTangentDir = CalcTangentDir( PreDragArriveTangent );
+		FVector2D LeaveTangentDir = CalcTangentDir( PreDragLeaveTangent );
 
 		FVector2D KeyPosition( Key.Curve->GetKeyTime(Key.KeyHandle),Key.Curve->GetKeyValue(Key.KeyHandle) );
 
-		FVector2D Movement = Delta - KeyPosition;
-		if(SelectedTangent.bIsArrival)
+		ArriveTangentDir.Y *= -1.0f;
+		LeaveTangentDir.Y *= -1.0f;
+		FVector2D ArrivePosition = -ArriveTangentDir + KeyPosition;
+
+		FVector2D LeavePosition = LeaveTangentDir + KeyPosition;
+
+		FVector2D Arrive = FVector2D(ScaleInfo.InputToLocalX(ArrivePosition.X), ScaleInfo.OutputToLocalY(ArrivePosition.Y));
+		FVector2D Leave = FVector2D(ScaleInfo.InputToLocalX(LeavePosition.X), ScaleInfo.OutputToLocalY(LeavePosition.Y));
+
+		FVector2D KeyScreenPosition = FVector2D(ScaleInfo.InputToLocalX(KeyPosition.X), ScaleInfo.OutputToLocalY(KeyPosition.Y));
+
+		FVector2D ToArrive = Arrive - KeyScreenPosition;
+		ToArrive.Normalize();
+
+		Arrive = KeyScreenPosition + ToArrive*CONST_KeyTangentOffset;
+
+		FVector2D ToLeave = Leave - KeyScreenPosition;
+		ToLeave.Normalize();
+
+		Leave = KeyScreenPosition + ToLeave*CONST_KeyTangentOffset;
+
+		// New arrive and leave directions in screen space
+		if (SelectedTangent.bIsArrival)
 		{
-			Movement *= -1.0f;
+			Arrive += Delta;
+			Leave -= Delta;
 		}
-		float Tangent = CalcTangent(Movement);
+		else
+		{
+			Arrive -= Delta;
+			Leave += Delta;
+		}
+
+		// Convert back to input/output space
+		FVector2D NewArriveDir(ScaleInfo.LocalXToInput(Arrive.X), ScaleInfo.LocalYToOutput(Arrive.Y));
+		FVector2D NewLeaveDir(ScaleInfo.LocalXToInput(Leave.X), ScaleInfo.LocalYToOutput(Leave.Y));
+
+		// Compute tangents
+		float NewArriveTangent = CalcTangent(-1.f*(NewArriveDir - KeyPosition));
+		float NewLeaveTangent = CalcTangent(NewLeaveDir - KeyPosition);
 
-		if(RichKey.TangentMode  != RCTM_Break)
+		if(RichKey.TangentMode != RCTM_Break)
 		{
-			RichKey.ArriveTangent = Tangent;
-			RichKey.LeaveTangent  = Tangent;
+			RichKey.ArriveTangent = NewArriveTangent;
+			RichKey.LeaveTangent = NewLeaveTangent;
 		
 			RichKey.TangentMode = RCTM_User;
 		}
@@ -2950,11 +2988,11 @@ void SCurveEditor::MoveTangents(FTrackScaleInfo ScaleInfo, FVector2D Delta)
 		{
 			if(SelectedTangent.bIsArrival)
 			{
-				RichKey.ArriveTangent = Tangent;
+				RichKey.ArriveTangent = NewArriveTangent;
 			}
 			else
 			{
-				RichKey.LeaveTangent = Tangent;
+				RichKey.LeaveTangent = NewLeaveTangent;
 			}
 		}	
 	}
