Sequencer: Fix a reference.

diff --git a/Engine/Source/Editor/UnrealEd/Private/SCurveEditor.cpp b/Engine/Source/Editor/UnrealEd/Private/SCurveEditor.cpp
--- a/Engine/Source/Editor/UnrealEd/Private/SCurveEditor.cpp
+++ b/Engine/Source/Editor/UnrealEd/Private/SCurveEditor.cpp
@@ -2920,7 +2920,7 @@ bool SCurveEditor::IsPostInfinityExtrapSelected(ERichCurveExtrapolation Extrapol
 	return CurveViewModels.Num() > 0;
 }
 
-void SCurveEditor::MoveTangents(FTrackScaleInfo ScaleInfo, FVector2D Delta)
+void SCurveEditor::MoveTangents(FTrackScaleInfo& ScaleInfo, FVector2D Delta)
 {
 	for (auto SelectedTangent : SelectedTangents)
 	{
diff --git a/Engine/Source/Editor/UnrealEd/Public/SCurveEditor.h b/Engine/Source/Editor/UnrealEd/Public/SCurveEditor.h
--- a/Engine/Source/Editor/UnrealEd/Public/SCurveEditor.h
+++ b/Engine/Source/Editor/UnrealEd/Public/SCurveEditor.h
@@ -469,7 +469,7 @@ private:
 	TSharedPtr<FCurveViewModel> HitTestCurves(const FGeometry& InMyGeometry, const FPointerEvent& InMouseEvent);
 
 	/* user is moving the tangents */
-	void MoveTangents(FTrackScaleInfo ScaleInfo, FVector2D Delta);
+	void MoveTangents(FTrackScaleInfo& ScaleInfo, FVector2D Delta);
 
 	//Curve Selection interface 
 
