Fix for importing static mesh
- Make sure you have animation curve node with animating

diff --git a/Engine/Source/Editor/UnrealEd/Private/Fbx/FbxMainImport.cpp b/Engine/Source/Editor/UnrealEd/Private/Fbx/FbxMainImport.cpp
--- a/Engine/Source/Editor/UnrealEd/Private/Fbx/FbxMainImport.cpp
+++ b/Engine/Source/Editor/UnrealEd/Private/Fbx/FbxMainImport.cpp
@@ -411,10 +411,7 @@ int32 FFbxImporter::GetImportType(const FString& InFilename)
 				Result = 0;
 			}
 
-			if (SceneInfo.TakeName != nullptr)
-			{
-				bHasAnimation = true;
-			}
+			bHasAnimation = SceneInfo.bHasAnimation;
 		}
 		else
 		{
@@ -586,17 +583,17 @@ bool FFbxImporter::GetSceneInfo(FString Filename, FbxSceneInfo& SceneInfo)
 			}
 		}
 		
-		// TODO: display multiple anim stack
-		SceneInfo.TakeName = NULL;
-		for (int32 AnimStackIndex = 0; AnimStackIndex < Scene->GetSrcObjectCount<FbxAnimStack>(); AnimStackIndex++)
+		SceneInfo.bHasAnimation = false;
+		for (int32 AnimCurveNodeIndex = 0; AnimCurveNodeIndex < Scene->GetSrcObjectCount<FbxAnimCurveNode>(); AnimCurveNodeIndex++)
 		{
-			FbxAnimStack* CurAnimStack = Scene->GetSrcObject<FbxAnimStack>(0);
-			// TODO: skip empty anim stack
-			const char* AnimStackName = CurAnimStack->GetName();
-			SceneInfo.TakeName = new char[FCStringAnsi::Strlen(AnimStackName) + 1];
-
-			FCStringAnsi::Strcpy(SceneInfo.TakeName, FCStringAnsi::Strlen(AnimStackName) + 1, AnimStackName);
+			FbxAnimCurveNode* CurAnimCruveNode = Scene->GetSrcObject<FbxAnimCurveNode>(AnimCurveNodeIndex);
+			if (CurAnimCruveNode->IsAnimated(true))
+			{
+				SceneInfo.bHasAnimation = true;
+				break;
+			}
 		}
+
 		SceneInfo.FrameRate = FbxTime::GetFrameRate(Scene->GetGlobalSettings().GetTimeMode());
 		
 		if ( GlobalTimeSpan.GetDirection() == FBXSDK_TIME_FORWARD)
diff --git a/Engine/Source/Editor/UnrealEd/Private/FbxImporter.h b/Engine/Source/Editor/UnrealEd/Private/FbxImporter.h
--- a/Engine/Source/Editor/UnrealEd/Private/FbxImporter.h
+++ b/Engine/Source/Editor/UnrealEd/Private/FbxImporter.h
@@ -198,8 +198,8 @@ struct FbxSceneInfo
 	
 	TArray<FbxMeshInfo> MeshInfo;
 	
-	// only one take supported currently
-	char* TakeName;
+	/* true if it has animation */
+	bool bHasAnimation;
 	double FrameRate;
 	double TotalTime;
 
