ADDED HierarchicalLodVolume, this volume defines a single cluster containing X num of actors for manual usage
ADDED UE-14356 HLOD: Print user friendly message if you try build and meshes don't have have LODs
ADDED MergeStaticMeshComponets to MeshUtilities needed for next CHANGE
ADDED HLOD preview build step for maps, this only builds the clusters for previewing the settings
ADDED SphereComponent now takes into account min and max drawing distances
ADDED Global forward declares and removed local ones from HierarchicalLOD.h
ADDED Visualization to LODActor class, uses USphereComponent to render its bounds
ADDED IsPreviewActor flag to LODActor class, if true during cluster generation DrawSphereComponents bounds will be used instead of LODActor's
ADDED vertex duplication removal on import
ADDED vertex/index buffer cache optimization on import
CHANGED World.h now forward declares HierarchicalLODBuilder to minimize (re)compile time
CHANGED While building/merging actors HLOD system now takes the orginal staticmesh from previously merged LODActors, this to enable retrieving the correct LOD per mesh (related to UE-15398)
CHANGED HLOD Cluster generation process, now takes into account HierarchicalLODVolumes to exclude actors from cluster generation
CHANGED Renamed CalculateRawMeshTangents to ExtractMeshDataForGeometryCache
FIXED Crash where a nullptr actor was added to actors list for > LOD0 clusters
FIXED Bound creation for actor within a cluster used FVector.Size(), now used FVector.AbsMax()
MOVED FLODCLuster structure into seperate header and cpp file

diff --git a/Engine/Source/Developer/MeshUtilities/Private/MeshUtilities.cpp b/Engine/Source/Developer/MeshUtilities/Private/MeshUtilities.cpp
--- a/Engine/Source/Developer/MeshUtilities/Private/MeshUtilities.cpp
+++ b/Engine/Source/Developer/MeshUtilities/Private/MeshUtilities.cpp
@@ -130,11 +130,24 @@ private:
 		const TArray<AActor*>& SourceActors,
 		const FMeshMergingSettings& InSettings,
 		UPackage* InOuter,
-		const FString& BasePackageName,
+		const FString& InBasePackageName,
 		int32 UseLOD, // does not build all LODs but only use this LOD to create base mesh
-		TArray<UObject*>& OutAssetsToSync, 
-		FVector& OutMergedActorLocation, 
-		bool bSilent=false) const override;
+		TArray<UObject*>& OutAssetsToSync,
+		FVector& OutMergedActorLocation,
+		bool bSilent = false) const override;
+
+
+	virtual void MergeStaticMeshComponents(
+		const TArray<UStaticMeshComponent*>& ComponentsToMerge,
+		UWorld* World,
+		const FMeshMergingSettings& InSettings,
+		UPackage* InOuter,
+		const FString& InBasePackageName,
+		int32 UseLOD, // does not build all LODs but only use this LOD to create base mesh
+		TArray<UObject*>& OutAssetsToSync,
+		FVector& OutMergedActorLocation,
+		float ViewDistance,
+		bool bSilent = false) const override;
 	
 	virtual void CreateProxyMesh( 
 		const TArray<AActor*>& Actors, 
@@ -4496,44 +4509,51 @@ static void MergeMaterials(UWorld* InWorld, const TArray<UMaterialInterface*>& I
 		AtlasTargetPos = FIntPoint(AtlasColIdx*ExportTextureSize.X, AtlasRowIdx*ExportTextureSize.Y);
 	}
 }
-
 void FMeshUtilities::MergeActors(
-	const TArray<AActor*>& SourceActors, 
-	const FMeshMergingSettings& InSettings, 
+	const TArray<AActor*>& SourceActors,
+	const FMeshMergingSettings& InSettings,
 	UPackage* InOuter,
 	const FString& InBasePackageName,
 	int32 UseLOD, // does not build all LODs but only use this LOD to create base mesh
-	TArray<UObject*>& OutAssetsToSync, 
-	FVector& OutMergedActorLocation, 
+	TArray<UObject*>& OutAssetsToSync,
+	FVector& OutMergedActorLocation,
 	bool bSilent) const
 {
 	TArray<UStaticMeshComponent*> ComponentsToMerge;
 	ComponentsToMerge.Reserve(SourceActors.Num());
-	
+
 	// Collect static mesh components
 	for (AActor* Actor : SourceActors)
 	{
 		TInlineComponentArray<UStaticMeshComponent*> Components;
 		Actor->GetComponents<UStaticMeshComponent>(Components);
+
 		// Filter out bad components
 		for (UStaticMeshComponent* MeshComponent : Components)
 		{
-			if (MeshComponent->StaticMesh != nullptr && 
+			if (MeshComponent->StaticMesh != nullptr &&
 				MeshComponent->StaticMesh->SourceModels.Num() > 0)
 			{
 				ComponentsToMerge.Add(MeshComponent);
 			}
 		}
 	}
-	
+
+	UWorld* World = SourceActors[0]->GetWorld();
+
+	MergeStaticMeshComponents(ComponentsToMerge, World, InSettings, InOuter, InBasePackageName, UseLOD, OutAssetsToSync, OutMergedActorLocation, bSilent);
+}
+
+void FMeshUtilities::MergeStaticMeshComponents(const TArray<UStaticMeshComponent*>& ComponentsToMerge, UWorld* World, const FMeshMergingSettings& InSettings, UPackage* InOuter, const FString& InBasePackageName, int32 UseLOD, /* does not build all LODs but only use this LOD to create base mesh */ TArray<UObject*>& OutAssetsToSync, FVector& OutMergedActorLocation, float ViewDistance, bool bSilent /*= false*/) const
+{
 	typedef FIntPoint FMeshIdAndLOD;
-	
+
 	TArray<UMaterialInterface*>						UniqueMaterials;
 	TMap<FMeshIdAndLOD, TArray<int32>>				MaterialMap;
 	TArray<FRawMeshExt>								SourceMeshes;
 	bool											bWithVertexColors[MAX_STATIC_MESH_LODS] = {};
 	bool											bOcuppiedUVChannels[MAX_STATIC_MESH_LODS][MAX_MESH_TEXTURE_COORDS] = {};
-		
+
 	SourceMeshes.Reserve(ComponentsToMerge.Num());
 	SourceMeshes.SetNum(ComponentsToMerge.Num());
 
@@ -4548,10 +4568,13 @@ void FMeshUtilities::MergeActors(
 		SourceMeshes[MeshId].Pivot = MeshComponent->ComponentToWorld.GetLocation();
 		// Source mesh asset package name
 		SourceMeshes[MeshId].AssetPackageName = MeshComponent->StaticMesh->GetOutermost()->GetName();
+
+
+		bool check = true;
 	}
 
 	NumMaxLOD = FMath::Min(NumMaxLOD, MAX_STATIC_MESH_LODS);
-		
+
 	int32 StartLODIndex = 0;
 	if (UseLOD >= 0)
 	{
@@ -4559,11 +4582,11 @@ void FMeshUtilities::MergeActors(
 		StartLODIndex = FMath::Min(UseLOD, NumMaxLOD - 1);
 		NumMaxLOD = StartLODIndex + 1;
 	}
-		
+
 	int32 RawMeshLODIdx = 0;
-	for(int32 LODIndex = StartLODIndex; LODIndex < NumMaxLOD; ++LODIndex, ++RawMeshLODIdx)
+	for (int32 LODIndex = StartLODIndex; LODIndex < NumMaxLOD; ++LODIndex, ++RawMeshLODIdx)
 	{
-		for(int32 MeshId = 0; MeshId < ComponentsToMerge.Num(); ++MeshId)
+		for (int32 MeshId = 0; MeshId < ComponentsToMerge.Num(); ++MeshId)
 		{
 			UStaticMeshComponent* MeshComponent = ComponentsToMerge[MeshId];
 
@@ -4571,25 +4594,25 @@ void FMeshUtilities::MergeActors(
 			FRawMesh& RawMeshLOD = SourceMeshes[MeshId].MeshLOD[RawMeshLODIdx];
 
 			// We duplicate lower LOD in case this mesh has no LOD we want
-			int32 ExportLODIndex = FMath::Min(LODIndex, MeshComponent->StaticMesh->SourceModels.Num()-1);
+			int32 ExportLODIndex = FMath::Min(LODIndex, MeshComponent->StaticMesh->SourceModels.Num() - 1);
 
-			if(ConstructRawMesh(MeshComponent, ExportLODIndex, RawMeshLOD, UniqueMaterials, MeshMaterialMap))
+			if (ConstructRawMesh(MeshComponent, ExportLODIndex, RawMeshLOD, UniqueMaterials, MeshMaterialMap))
 			{
 				MaterialMap.Add(FMeshIdAndLOD(MeshId, RawMeshLODIdx), MeshMaterialMap);
 
 				// Should we use vertex colors?
-				if(InSettings.bImportVertexColors)
+				if (InSettings.bImportVertexColors)
 				{
 					// Propagate painted vertex colors into our raw mesh
 					PropagatePaintedColorsToRawMesh(MeshComponent, ExportLODIndex, RawMeshLOD);
 					// Whether at least one of the meshes has vertex colors
-					bWithVertexColors[RawMeshLODIdx]|= (RawMeshLOD.WedgeColors.Num() != 0);
+					bWithVertexColors[RawMeshLODIdx] |= (RawMeshLOD.WedgeColors.Num() != 0);
 				}
 
 				// Which UV channels has data at least in one mesh
-				for(int32 ChannelIdx = 0; ChannelIdx < MAX_MESH_TEXTURE_COORDS; ++ChannelIdx)
+				for (int32 ChannelIdx = 0; ChannelIdx < MAX_MESH_TEXTURE_COORDS; ++ChannelIdx)
 				{
-					bOcuppiedUVChannels[RawMeshLODIdx][ChannelIdx]|= (RawMeshLOD.WedgeTexCoords[ChannelIdx].Num() != 0);
+					bOcuppiedUVChannels[RawMeshLODIdx][ChannelIdx] |= (RawMeshLOD.WedgeTexCoords[ChannelIdx].Num() != 0);
 				}
 			}
 		}
@@ -4597,9 +4620,9 @@ void FMeshUtilities::MergeActors(
 
 	if (SourceMeshes.Num() == 0)
 	{
-		return;	
+		return;
 	}
-	
+
 	// For each raw mesh, re-map the material indices according to the MaterialMap
 	for (int32 MeshIndex = 0; MeshIndex < SourceMeshes.Num(); ++MeshIndex)
 	{
@@ -4622,18 +4645,17 @@ void FMeshUtilities::MergeActors(
 	// Use first mesh for naming and pivot
 	MergedMesh.AssetPackageName = SourceMeshes[0].AssetPackageName;
 	MergedMesh.Pivot = InSettings.bPivotPointAtZero ? FVector::ZeroVector : SourceMeshes[0].Pivot;
-			
+
 	if (InSettings.bMergeMaterials)
 	{
-		FIntPoint AtlasTextureSize			= FIntPoint(InSettings.MergedMaterialAtlasResolution, InSettings.MergedMaterialAtlasResolution);
+		FIntPoint AtlasTextureSize = FIntPoint(InSettings.MergedMaterialAtlasResolution, InSettings.MergedMaterialAtlasResolution);
 		FFlattenMaterial MergedFlatMaterial;
-		MergedFlatMaterial.DiffuseSize		= AtlasTextureSize;
-		MergedFlatMaterial.NormalSize		= InSettings.bExportNormalMap ? AtlasTextureSize : FIntPoint::ZeroValue;
-		MergedFlatMaterial.MetallicSize		= InSettings.bExportMetallicMap ? AtlasTextureSize : FIntPoint::ZeroValue;
-		MergedFlatMaterial.RoughnessSize	= InSettings.bExportRoughnessMap ? AtlasTextureSize : FIntPoint::ZeroValue;
-		MergedFlatMaterial.SpecularSize		= InSettings.bExportSpecularMap ? AtlasTextureSize : FIntPoint::ZeroValue;
+		MergedFlatMaterial.DiffuseSize = AtlasTextureSize;
+		MergedFlatMaterial.NormalSize = InSettings.bExportNormalMap ? AtlasTextureSize : FIntPoint::ZeroValue;
+		MergedFlatMaterial.MetallicSize = InSettings.bExportMetallicMap ? AtlasTextureSize : FIntPoint::ZeroValue;
+		MergedFlatMaterial.RoughnessSize = InSettings.bExportRoughnessMap ? AtlasTextureSize : FIntPoint::ZeroValue;
+		MergedFlatMaterial.SpecularSize = InSettings.bExportSpecularMap ? AtlasTextureSize : FIntPoint::ZeroValue;
 		TArray<FRawMeshUVTransform> UVTransforms;
-		UWorld* World = SourceActors[0]->GetWorld();
 
 		MergeMaterials(World, UniqueMaterials, MergedFlatMaterial, UVTransforms);
 
@@ -4665,25 +4687,25 @@ void FMeshUtilities::MergeActors(
 					TArray<FVector2D>& UVs = RawMesh.WedgeTexCoords[UVChannelIdx];
 					if (UVs.Num() > 0)
 					{
-						int32 UVIdx = 0; 
+						int32 UVIdx = 0;
 						for (int32 FaceMaterialIndex : RawMesh.FaceMaterialIndices)
 						{
 							const FRawMeshUVTransform& UVTransform = UVTransforms[FaceMaterialIndex];
 							if (UVTransform.IsValid())
 							{
-								FVector2D UV0 = GetValidUV(UVs[UVIdx+0]);
-								FVector2D UV1 = GetValidUV(UVs[UVIdx+1]);
-								FVector2D UV2 = GetValidUV(UVs[UVIdx+2]);
-								UVs[UVIdx+0] = UV0 * UVTransform.Scale + UVTransform.Offset;
-								UVs[UVIdx+1] = UV1 * UVTransform.Scale + UVTransform.Offset;
-								UVs[UVIdx+2] = UV2 * UVTransform.Scale + UVTransform.Offset;
+								FVector2D UV0 = GetValidUV(UVs[UVIdx + 0]);
+								FVector2D UV1 = GetValidUV(UVs[UVIdx + 1]);
+								FVector2D UV2 = GetValidUV(UVs[UVIdx + 2]);
+								UVs[UVIdx + 0] = UV0 * UVTransform.Scale + UVTransform.Offset;
+								UVs[UVIdx + 1] = UV1 * UVTransform.Scale + UVTransform.Offset;
+								UVs[UVIdx + 2] = UV2 * UVTransform.Scale + UVTransform.Offset;
 							}
-						
-							UVIdx+=3;
+
+							UVIdx += 3;
 						}
 					}
 				}
-				
+
 				// Remap material indexes
 				for (int32& FaceMaterialIndex : RawMesh.FaceMaterialIndices)
 				{
@@ -4708,7 +4730,7 @@ void FMeshUtilities::MergeActors(
 		}
 
 		UPackage* MaterialPackage = InOuter;
-		if(MaterialPackage == nullptr)
+		if (MaterialPackage == nullptr)
 		{
 			MaterialPackage = CreatePackage(nullptr, *MaterialPackageName);
 			check(MaterialPackage);
@@ -4716,7 +4738,7 @@ void FMeshUtilities::MergeActors(
 			MaterialPackage->Modify();
 		}
 
-		UMaterial* MergedMaterial = FMaterialUtilities::CreateMaterial(MergedFlatMaterial, MaterialPackage, MaterialAssetName, RF_Public|RF_Standalone, OutAssetsToSync);
+		UMaterial* MergedMaterial = FMaterialUtilities::CreateMaterial(MergedFlatMaterial, MaterialPackage, MaterialAssetName, RF_Public | RF_Standalone, OutAssetsToSync);
 
 		// Set material static lighting usage flag if project has static lighting enabled
 		static const auto AllowStaticLightingVar = IConsoleManager::Get().FindTConsoleVariableDataInt(TEXT("r.AllowStaticLighting"));
@@ -4729,7 +4751,7 @@ void FMeshUtilities::MergeActors(
 
 		UniqueMaterials.Add(MergedMaterial);
 	}
-		
+
 	// Merge meshes into single mesh
 	for (int32 SourceMeshIdx = 0; SourceMeshIdx < SourceMeshes.Num(); ++SourceMeshIdx)
 	{
@@ -4738,12 +4760,12 @@ void FMeshUtilities::MergeActors(
 			// Merge vertex data from source mesh list into single mesh
 			FRawMesh& TargetRawMesh = MergedMesh.MeshLOD[LODIndex];
 			const FRawMesh& SourceRawMesh = SourceMeshes[SourceMeshIdx].MeshLOD[LODIndex];
-			
+
 			if (SourceRawMesh.VertexPositions.Num() == 0)
 			{
 				continue;
 			}
-									
+
 			TargetRawMesh.FaceSmoothingMasks.Append(SourceRawMesh.FaceSmoothingMasks);
 			TargetRawMesh.FaceMaterialIndices.Append(SourceRawMesh.FaceMaterialIndices);
 
@@ -4758,11 +4780,11 @@ void FMeshUtilities::MergeActors(
 			{
 				TargetRawMesh.VertexPositions.Add(VertexPos - MergedMesh.Pivot);
 			}
-						
+
 			TargetRawMesh.WedgeTangentX.Append(SourceRawMesh.WedgeTangentX);
 			TargetRawMesh.WedgeTangentY.Append(SourceRawMesh.WedgeTangentY);
 			TargetRawMesh.WedgeTangentZ.Append(SourceRawMesh.WedgeTangentZ);
-		
+
 			// Deal with vertex colors
 			// Some meshes may have it, in this case merged mesh will be forced to have vertex colors as well
 			if (bWithVertexColors[LODIndex])
@@ -4780,7 +4802,7 @@ void FMeshUtilities::MergeActors(
 					FMemory::Memset(&TargetRawMesh.WedgeColors[ColorsOffset], 0xFF, ColorsNum*TargetRawMesh.WedgeColors.GetTypeSize());
 				}
 			}
-		
+
 			// Merge all other UV channels 
 			for (int32 ChannelIdx = 0; ChannelIdx < MAX_MESH_TEXTURE_COORDS; ++ChannelIdx)
 			{
@@ -4818,7 +4840,7 @@ void FMeshUtilities::MergeActors(
 		for (; ChannelIdx < MAX_MESH_TEXTURE_COORDS && bOcuppiedUVChannels[LODIndex][ChannelIdx]; ++ChannelIdx);
 		TargetLightMapUVChannel[LODIndex] = FMath::Min(InSettings.TargetLightMapUVChannel, ChannelIdx);
 	}
-	
+
 	//
 	//Create merged mesh asset
 	//
@@ -4837,7 +4859,7 @@ void FMeshUtilities::MergeActors(
 		}
 
 		UPackage* Package = InOuter;
-		if(Package == nullptr)
+		if (Package == nullptr)
 		{
 			Package = CreatePackage(NULL, *PackageName);
 			check(Package);
@@ -4847,7 +4869,7 @@ void FMeshUtilities::MergeActors(
 
 		auto StaticMesh = NewObject<UStaticMesh>(Package, *AssetName, RF_Public | RF_Standalone);
 		StaticMesh->InitResources();
-		
+
 		FString OutputPath = StaticMesh->GetPathName();
 
 		// make sure it has a new lighting guid
@@ -4857,7 +4879,7 @@ void FMeshUtilities::MergeActors(
 			StaticMesh->LightMapResolution = InSettings.TargetLightMapResolution;
 			StaticMesh->LightMapCoordinateIndex = TargetLightMapUVChannel[0];
 		}
-		
+
 		for (int32 LODIndex = 0; LODIndex < NumMaxLOD; ++LODIndex)
 		{
 			FRawMesh& MergedMeshLOD = MergedMesh.MeshLOD[LODIndex];
@@ -4877,7 +4899,7 @@ void FMeshUtilities::MergeActors(
 				SrcModel->RawMeshBulkData->SaveRawMesh(MergedMeshLOD);
 			}
 		}
-			
+
 		// Assign materials
 		for (UMaterialInterface* Material : UniqueMaterials)
 		{
@@ -4888,12 +4910,12 @@ void FMeshUtilities::MergeActors(
 
 			StaticMesh->Materials.Add(Material);
 		}
-		
+
 		StaticMesh->Build(bSilent);
 		StaticMesh->PostEditChange();
 
 		OutAssetsToSync.Add(StaticMesh);
-		
+
 		//
 		OutMergedActorLocation = MergedMesh.Pivot;
 	}
diff --git a/Engine/Source/Developer/MeshUtilities/Public/MeshUtilities.h b/Engine/Source/Developer/MeshUtilities/Public/MeshUtilities.h
--- a/Engine/Source/Developer/MeshUtilities/Public/MeshUtilities.h
+++ b/Engine/Source/Developer/MeshUtilities/Public/MeshUtilities.h
@@ -189,6 +189,34 @@ public:
 		TArray<UObject*>& OutAssetsToSync, 
 		FVector& OutMergedActorLocation, 
 		bool bSilent=false) const = 0;
+
+	/**
+	* MergeStaticMeshComponents
+	*
+	* @param ComponentsToMerge - Components to merge
+	* @param World - World in which the component reside
+	* @param InSettings	- Settings to use
+	* @param InOuter - Outer if required
+	* @param InBasePackageName - Destination package name for a generated assets. Used if Outer is null.
+	* @param UseLOD	-1 if you'd like to build for all LODs. If you specify, that LOD mesh for source meshes will be used to merge the mesh
+	*									This is used by hierarchical building LODs
+	* @param OutAssetsToSync Merged mesh assets
+	* @param OutMergedActorLocation	World position of merged mesh
+	* @param ViewDistance Distance for LOD determination
+	* @param bSilent Non-verbose flag
+	* @return void
+	*/
+	virtual void MergeStaticMeshComponents(
+		const TArray<UStaticMeshComponent*>& ComponentsToMerge,
+		UWorld* World,
+		const FMeshMergingSettings& InSettings,
+		UPackage* InOuter,
+		const FString& InBasePackageName,
+		int32 UseLOD, /* does not build all LODs but only use this LOD to create base mesh */
+		TArray<UObject*>& OutAssetsToSync,
+		FVector& OutMergedActorLocation,
+		float ViewDistance,
+		bool bSilent /*= false*/) const = 0;
 	
 	/**
 	 *	Merges list of actors into single proxy mesh
diff --git a/Engine/Source/Editor/LevelEditor/Private/LevelEditor.cpp b/Engine/Source/Editor/LevelEditor/Private/LevelEditor.cpp
--- a/Engine/Source/Editor/LevelEditor/Private/LevelEditor.cpp
+++ b/Engine/Source/Editor/LevelEditor/Private/LevelEditor.cpp
@@ -1379,6 +1379,9 @@ void FLevelEditorModule::BindGlobalLevelEditorCommands()
 	ActionList.MapAction(Commands.BuildLODsOnly,
 		FExecuteAction::CreateStatic(&FLevelEditorActionCallbacks::BuildLODsOnly_Execute));
 
+	ActionList.MapAction(Commands.PreviewHLODClustersOnly,
+		FExecuteAction::CreateStatic(&FLevelEditorActionCallbacks::PreviewHLODClustersOnly_Execute));
+
 	ActionList.MapAction( 
 		Commands.LightingQuality_Production, 
 		FExecuteAction::CreateStatic( &FLevelEditorActionCallbacks::SetLightingQuality, (ELightingBuildQuality)Quality_Production ),
diff --git a/Engine/Source/Editor/LevelEditor/Private/LevelEditorActions.cpp b/Engine/Source/Editor/LevelEditor/Private/LevelEditorActions.cpp
--- a/Engine/Source/Editor/LevelEditor/Private/LevelEditorActions.cpp
+++ b/Engine/Source/Editor/LevelEditor/Private/LevelEditorActions.cpp
@@ -629,6 +629,12 @@ void FLevelEditorActionCallbacks::BuildLODsOnly_Execute()
 	FEditorBuildUtils::EditorBuild(GetWorld(), EBuildOptions::BuildHierarchicalLOD);
 }
 
+void FLevelEditorActionCallbacks::PreviewHLODClustersOnly_Execute()
+{
+	// Preview HLOD Clusters
+	FEditorBuildUtils::EditorBuild(GetWorld(), EBuildOptions::PreviewHierarchicalLOD);
+}
+
 bool FLevelEditorActionCallbacks::IsLightingQualityChecked( ELightingBuildQuality TestQuality )
 {
 	int32 CurrentQualityLevel;
@@ -2905,6 +2911,7 @@ void FLevelEditorCommands::RegisterCommands()
 	UI_COMMAND( BuildGeometryOnly_OnlyCurrentLevel, "Build Geometry (Current Level)", "Builds geometry, only for the current level", EUserInterfaceActionType::Button, FInputChord() );
 	UI_COMMAND( BuildPathsOnly, "Build Paths", "Only builds paths (all levels.)", EUserInterfaceActionType::Button, FInputChord() );
 	UI_COMMAND( BuildLODsOnly, "Build LODs", "Only builds LODs (all levels.)", EUserInterfaceActionType::Button, FInputChord() );
+	UI_COMMAND( PreviewHLODClustersOnly, "Preview HLOD clusters", "Preview builds LODs and shows clusters (all levels.)", EUserInterfaceActionType::Button, FInputChord());
 	UI_COMMAND( LightingQuality_Production, "Production", "Sets precomputed lighting quality to highest possible quality (slowest computation time.)", EUserInterfaceActionType::RadioButton, FInputChord() );
 	UI_COMMAND( LightingQuality_High, "High", "Sets precomputed lighting quality to high quality", EUserInterfaceActionType::RadioButton, FInputChord() );
 	UI_COMMAND( LightingQuality_Medium, "Medium", "Sets precomputed lighting quality to medium quality", EUserInterfaceActionType::RadioButton, FInputChord() );
diff --git a/Engine/Source/Editor/LevelEditor/Private/LevelEditorToolBar.cpp b/Engine/Source/Editor/LevelEditor/Private/LevelEditorToolBar.cpp
--- a/Engine/Source/Editor/LevelEditor/Private/LevelEditorToolBar.cpp
+++ b/Engine/Source/Editor/LevelEditor/Private/LevelEditorToolBar.cpp
@@ -1562,9 +1562,10 @@ TSharedRef< SWidget > FLevelEditorToolBar::GenerateBuildMenuContent( TSharedRef<
 	}
 	MenuBuilder.EndSection();
 
-	MenuBuilder.BeginSection("LevelEditorLOD", LOCTEXT("LODHeading", "Hiearchical LOD"));
+	MenuBuilder.BeginSection("LevelEditorLOD", LOCTEXT("LODHeading", "Hierarchical LOD"));
 	{
 		MenuBuilder.AddMenuEntry(FLevelEditorCommands::Get().BuildLODsOnly);
+		MenuBuilder.AddMenuEntry(FLevelEditorCommands::Get().PreviewHLODClustersOnly);
 	}
 	MenuBuilder.EndSection();
 
diff --git a/Engine/Source/Editor/LevelEditor/Public/LevelEditorActions.h b/Engine/Source/Editor/LevelEditor/Public/LevelEditorActions.h
--- a/Engine/Source/Editor/LevelEditor/Public/LevelEditorActions.h
+++ b/Engine/Source/Editor/LevelEditor/Public/LevelEditorActions.h
@@ -71,6 +71,7 @@ public:
 	TSharedPtr< FUICommandInfo > BuildGeometryOnly_OnlyCurrentLevel;
 	TSharedPtr< FUICommandInfo > BuildPathsOnly;
 	TSharedPtr< FUICommandInfo > BuildLODsOnly;
+	TSharedPtr< FUICommandInfo > PreviewHLODClustersOnly;
 	TSharedPtr< FUICommandInfo > LightingQuality_Production;
 	TSharedPtr< FUICommandInfo > LightingQuality_High;
 	TSharedPtr< FUICommandInfo > LightingQuality_Medium;
@@ -701,6 +702,7 @@ public:
 	static void BuildGeometryOnly_OnlyCurrentLevel_Execute();
 	static void BuildPathsOnly_Execute();
 	static void BuildLODsOnly_Execute();
+	static void PreviewHLODClustersOnly_Execute();
 	static void SetLightingQuality( ELightingBuildQuality NewQuality );
 	static bool IsLightingQualityChecked( ELightingBuildQuality TestQuality );
 	static float GetLightingDensityIdeal();
diff --git a/Engine/Source/Editor/UnrealEd/Private/EditorBuildUtils.cpp b/Engine/Source/Editor/UnrealEd/Private/EditorBuildUtils.cpp
--- a/Engine/Source/Editor/UnrealEd/Private/EditorBuildUtils.cpp
+++ b/Engine/Source/Editor/UnrealEd/Private/EditorBuildUtils.cpp
@@ -20,6 +20,7 @@
 #include "Engine/LevelStreaming.h"
 #include "GameFramework/WorldSettings.h"
 #include "AI/Navigation/NavigationSystem.h"
+#include "HierarchicalLOD.h"
 
 DEFINE_LOG_CATEGORY_STATIC(LogEditorBuildUtils, Log, All);
 
@@ -216,6 +217,7 @@ bool FEditorBuildUtils::EditorBuild( UWorld* InWorld, EBuildOptions::Type Id, co
 		BuildType = SBuildProgressWidget::BUILDTYPE_Paths;
 		break;
 	case EBuildOptions::BuildHierarchicalLOD:
+	case EBuildOptions::PreviewHierarchicalLOD:
 		BuildType = SBuildProgressWidget::BUILDTYPE_LODs;
 		break;
 	default:
@@ -317,6 +319,23 @@ bool FEditorBuildUtils::EditorBuild( UWorld* InWorld, EBuildOptions::Type Id, co
 			break;
 		}
 
+	case  EBuildOptions::PreviewHierarchicalLOD:
+		{
+			bDoBuild = GEditor->WarnAboutHiddenLevels(InWorld, false);
+			if (bDoBuild)
+			{
+				GEditor->ResetTransaction(NSLOCTEXT("UnrealEd", "RebuildLOD", "Rebuilding HierarchicalLOD"));
+
+				// We can't set the busy cursor for all windows, because lighting
+				// needs a cursor for the lighting options dialog.
+				const FScopedBusyCursor BusyCursor;
+
+				TriggerHierarchicalLODBuilder(InWorld, Id);
+			}
+
+			break;
+		}
+
 	case EBuildOptions::BuildAll:
 	case EBuildOptions::BuildAllSubmit:
 		{
@@ -798,7 +817,15 @@ void FEditorBuildUtils::TriggerNavigationBuilder(UWorld* InWorld, EBuildOptions:
 
 void FEditorBuildUtils::TriggerHierarchicalLODBuilder(UWorld* InWorld, EBuildOptions::Type Id)
 {
-	// Invoke HLOD generator
-	InWorld->HierarchicalLODBuilder.Build();
+	// Invoke HLOD generator, with either preview or full build
+	if (Id == EBuildOptions::PreviewHierarchicalLOD)
+	{
+		InWorld->HierarchicalLODBuilder->PreviewBuild();
+	}
+	else
+	{
+		InWorld->HierarchicalLODBuilder->Build();
+	}
 }
+
 #undef LOCTEXT_NAMESPACE
\ No newline at end of file
diff --git a/Engine/Source/Editor/UnrealEd/Private/HierarchicalLOD.cpp b/Engine/Source/Editor/UnrealEd/Private/HierarchicalLOD.cpp
--- a/Engine/Source/Editor/UnrealEd/Private/HierarchicalLOD.cpp
+++ b/Engine/Source/Editor/UnrealEd/Private/HierarchicalLOD.cpp
@@ -1,113 +1,32 @@
 // Copyright 1998-2015 Epic Games, Inc. All Rights Reserved.
 
-/*=============================================================================
-	World.cpp: UWorld implementation
-=============================================================================*/
-
 #include "UnrealEd.h"
 #include "HierarchicalLOD.h"
 
+#include "MessageLog.h"
+#include "UObjectToken.h"
+#include "MapErrors.h"
+
 #if WITH_EDITOR
 #include "Editor/UnrealEd/Classes/Editor/EditorEngine.h"
 #include "Editor/UnrealEd/Public/Editor.h"
 #include "Engine/LODActor.h"
 #include "Editor/UnrealEd/Classes/Factories/Factory.h"
-#include "MeshUtilities.h"
 #include "ObjectTools.h"
-
 #endif // WITH_EDITOR
+
 #include "GameFramework/WorldSettings.h"
 #include "Components/InstancedStaticMeshComponent.h"
 
+#include "HierarchicalLODVolume.h"
+
 DEFINE_LOG_CATEGORY_STATIC(LogLODGenerator, Warning, All);
 
 #define LOCTEXT_NAMESPACE "HierarchicalLOD"
 #define CM_TO_METER		0.01f
+#define METER_TO_CM		100.0f
 
-/*-----------------------------------------------------------------------------
-	HierarchicalLOD implementation.
------------------------------------------------------------------------------*/
-
-////////////////////////////////////////////////////////////////////////////////////////////
-// Hierarchical LOD System implementation 
-////////////////////////////////////////////////////////////////////////////////////////////
-
-// utility function to calculate overlap of two spheres
-float CalculateOverlap(const FSphere& ASphere, const float AFillingFactor, const FSphere& BSphere, const float BFillingFactor)
-{
-	// if it doesn't intersect, return zero 
-	if (!ASphere.Intersects(BSphere))
-	{
-		return 0.f;
-	}
-
-	if (ASphere.IsInside(BSphere, 1.f))
-	{
-		return ASphere.GetVolume();
-	}
-
-	if(BSphere.IsInside(ASphere, 1.f))
-	{
-		return BSphere.GetVolume();
-	}
-
-	if (ASphere.Equals(BSphere, 1.f))
-	{
-		return ASphere.GetVolume();
-	}
-
-	float Distance = (ASphere.Center-BSphere.Center).Size();
-	check (Distance > 0.f);
-
-	float  ARadius = ASphere.W;
-	float  BRadius = BSphere.W;
-
-	float  ACapHeight = (BRadius*BRadius - (ARadius - Distance)*(ARadius - Distance)) / (2*Distance);
-	float  BCapHeight = (ARadius*ARadius - (BRadius - Distance)*(BRadius - Distance)) / (2*Distance);
-
-	if (ACapHeight<=0.f || BCapHeight<=0.f)
-	{
-		// it's possible to get cap height to be less than 0 
-		// since when we do check intersect, we do have regular tolerance
-		return 0.f;		
-	}
-
-	float  OverlapRadius1 = ((ARadius+BRadius)*(ARadius+BRadius) - Distance*Distance) * (Distance*Distance - (ARadius-BRadius)*(ARadius-BRadius));
-	float  OverlapRadius2 = 2*Distance;
-	float  OverlapRedius = FMath::Sqrt(OverlapRadius1) / OverlapRadius2;
-	float  OverlapRediusSq = OverlapRedius*OverlapRedius;
-
-	check (OverlapRadius1 >= 0.f);
-
-	float ConstPI = PI/6.0f;
-	float AVolume = ConstPI*(3*OverlapRediusSq + ACapHeight*ACapHeight) * ACapHeight;
-	float BVolume = ConstPI*(3*OverlapRediusSq + BCapHeight*BCapHeight) * BCapHeight;
-
-	check (AVolume > 0.f &&  BVolume > 0.f);
-
-	float TotalVolume = AFillingFactor*AVolume + BFillingFactor*BVolume;
-	return TotalVolume;
-}
-
-// utility functions that calculates filling factor
-float CalculateFillingFactor(const FSphere& ASphere, const float AFillingFactor, const FSphere& BSphere, const float BFillingFactor)
-{
-	float OverlapVolume = CalculateOverlap( ASphere, AFillingFactor, BSphere, BFillingFactor);
-	FSphere UnionSphere = ASphere + BSphere;
-	// it shouldn't be zero or it should be checked outside
-	ensure(UnionSphere.W != 0.f);
-
-	// http://deim.urv.cat/~rivi/pub/3d/icra04b.pdf
-	// cost is calculated based on r^3 / filling factor
-	// since it subtract by AFillingFactor * 1/2 overlap volume + BfillingFactor * 1/2 overlap volume
-	return FMath::Max(0.f, (AFillingFactor * ASphere.GetVolume() + BFillingFactor * BSphere.GetVolume() - OverlapVolume) / UnionSphere.GetVolume());
-}
-
-////////////////////////////////////////////////////////////////////////////////////////////
-// Hierarchical LOD Builder System implementation 
-////////////////////////////////////////////////////////////////////////////////////////////
-
-FHierarchicalLODBuilder::FHierarchicalLODBuilder(class UWorld* InWorld)
+FHierarchicalLODBuilder::FHierarchicalLODBuilder(UWorld* InWorld)
 :	World(InWorld)
 {}
 
@@ -119,22 +38,30 @@ FHierarchicalLODBuilder::FHierarchicalLODBuilder()
 }
 #endif // WITH_HOT_RELOAD_CTORS
 
-// build hierarchical cluster
 void FHierarchicalLODBuilder::Build()
 {
-	check (World);
-
-	const TArray<class ULevel*>& Levels = World->GetLevels();
-
+	check (World)
+	const TArray<ULevel*>& Levels = World->GetLevels();
 	for (const auto& LevelIter : Levels)
 	{
-		BuildClusters(LevelIter);
+		BuildClusters(LevelIter, true);
 	}
 }
 
-// for now it only builds per level, it turns to one actor at the end
-void FHierarchicalLODBuilder::BuildClusters(class ULevel* InLevel)
+void FHierarchicalLODBuilder::PreviewBuild()
 {
+	check(World);
+	const TArray<ULevel*>& Levels = World->GetLevels();
+	for (const auto& LevelIter : Levels)
+	{
+		BuildClusters(LevelIter, false);
+	}
+}
+
+void FHierarchicalLODBuilder::BuildClusters(ULevel* InLevel, const bool bCreateMeshes)
+{	
+	SCOPE_LOG_TIME(TEXT("STAT_HLOD_BuildClusters"), nullptr);
+
 	// I'm using stack mem within this scope of the function
 	// so we need this
 	FMemMark Mark(FMemStack::Get());
@@ -172,18 +99,21 @@ void FHierarchicalLODBuilder::BuildClusters(class ULevel* InLevel)
 	// only build if it's enabled
 	if(InLevel->GetWorld()->GetWorldSettings()->bEnableHierarchicalLODSystem)
 	{
-		const int32 TotalNumLOD = InLevel->GetWorld()->GetWorldSettings()->HierarchicalLODSetup.Num();
+		// Handle HierachicalLOD volumes first
+		HandleHLODVolumes(InLevel);
 
+		AWorldSettings* WorldSetting = InLevel->GetWorld()->GetWorldSettings();		
+		const int32 TotalNumLOD = InLevel->GetWorld()->GetWorldSettings()->HierarchicalLODSetup.Num();
 		for(int32 LODId=0; LODId<TotalNumLOD; ++LODId)
 		{
-			AWorldSettings* WorldSetting = InLevel->GetWorld()->GetWorldSettings();
+			
 			// we use meter for bound. Otherwise it's very easy to get to overflow and have problem with filling ratio because
 			// bound is too huge
-			float DesiredBoundRadius = WorldSetting->HierarchicalLODSetup[LODId].DesiredBoundRadius * CM_TO_METER;
-			float DesiredFillingRatio = WorldSetting->HierarchicalLODSetup[LODId].DesiredFillingPercentage * 0.01f;
+			const float DesiredBoundRadius = WorldSetting->HierarchicalLODSetup[LODId].DesiredBoundRadius * CM_TO_METER;
+			const float DesiredFillingRatio = WorldSetting->HierarchicalLODSetup[LODId].DesiredFillingPercentage * 0.01f;
 			ensure(DesiredFillingRatio!=0.f);
-			float HighestCost = FMath::Pow(DesiredBoundRadius, 3)/(DesiredFillingRatio);
-			int32 MinNumActors = WorldSetting->HierarchicalLODSetup[LODId].MinNumberOfActorsToBuild;
+			const float HighestCost = FMath::Pow(DesiredBoundRadius, 3) / (DesiredFillingRatio);
+			const int32 MinNumActors = WorldSetting->HierarchicalLODSetup[LODId].MinNumberOfActorsToBuild;
 			check (MinNumActors > 0);
 			// test parameter I was playing with to cull adding to the array
 			// intialization can have too many elements, decided to cull
@@ -212,17 +142,132 @@ void FHierarchicalLODBuilder::BuildClusters(class ULevel* InLevel)
 			}
 
 			// now we have to calculate merge clusters and build actors
-			MergeClustersAndBuildActors(InLevel, LODId, HighestCost, MinNumActors);
+			MergeClustersAndBuildActors(InLevel, LODId, HighestCost, MinNumActors, bCreateMeshes);
 		}
 	}
+	else
+	{
+		// Fire map check warnings if HLOD System is not enabled
+		FMessageLog MapCheck("MapCheck");
+		MapCheck.Warning()
+			->AddToken(FUObjectToken::Create(InLevel->GetWorld()->GetWorldSettings()))
+			->AddToken(FTextToken::Create(LOCTEXT("MapCheck_Message_HLODSystemNotEnabled", "Hierarchical LOD System is disabled, unable to build LOD actors.")))
+			->AddToken(FMapErrorToken::Create(FMapErrors::HLODSystemNotEnabled));
+	}
 
 	// Clear Clusters. It is using stack mem, so it won't be good after this
 	Clusters.Empty();
 	Clusters.Shrink();
 }
 
+void FHierarchicalLODBuilder::InitializeClusters(ULevel* InLevel, const int32 LODIdx, float CullCost)
+{
+	SCOPE_LOG_TIME(TEXT("STAT_HLOD_InitializeClusters"), nullptr);
+	if (InLevel->Actors.Num() > 0)
+	{
+		if (LODIdx == 0)
+		{
+			Clusters.Empty();
+
+			TArray<AActor*> GenerationActors;			
+			for (int32 ActorId = 0; ActorId < InLevel->Actors.Num(); ++ActorId)
+			{
+				AActor* Actor = InLevel->Actors[ActorId];
+				if (ShouldGenerateCluster(Actor))
+				{
+					// Check whether or not this actor falls within a HierarchicalLODVolume, if so add to the Volume's cluster and exclude from normal process
+					bool bAdded = false;
+					for (auto& Cluster : HLODVolumeClusters)
+					{
+						if (Cluster.Key->EncompassesPoint(Actor->GetActorLocation(), 0.0f, nullptr))
+						{
+							FLODCluster ActorCluster(Actor);
+							Cluster.Value += ActorCluster;
+							bAdded = true;
+							break;
+						}
+					}
+
+					if (!bAdded)
+					{
+						GenerationActors.Add(Actor);
+					}
+				}
+			}
+			
+			// Create clusters using actor pairs
+			for (int32 ActorId = 0; ActorId<GenerationActors.Num(); ++ActorId)
+			{
+				AActor* Actor1 = GenerationActors[ActorId];
+
+				for (int32 SubActorId = ActorId + 1; SubActorId<GenerationActors.Num(); ++SubActorId)
+				{
+					AActor* Actor2 = GenerationActors[SubActorId];
+
+					FLODCluster NewClusterCandidate = FLODCluster(Actor1, Actor2);
+					float NewClusterCost = NewClusterCandidate.GetCost();
+
+					if ( NewClusterCost <= CullCost)
+					{
+						Clusters.Add(NewClusterCandidate);
+					}
+				}
+			}
+		}
+		else // at this point we only care for LODActors
+		{
+			Clusters.Empty();
+
+			// we filter the LOD index first
+			TArray<AActor*> Actors;
+			for(int32 ActorId=0; ActorId<InLevel->Actors.Num(); ++ActorId)
+			{
+				AActor* Actor = (InLevel->Actors[ActorId]);
+
+				if (Actor)
+				{
+					if (Actor->IsA(ALODActor::StaticClass()))
+					{
+						ALODActor* LODActor = CastChecked<ALODActor>(Actor);
+						if (LODActor->LODLevel == LODIdx)
+						{
+							Actors.Add(Actor);
+						}
+					}
+					else if (ShouldGenerateCluster(Actor))
+					{
+						Actors.Add(Actor);
+					}
+				}				
+			}
+			
+			// first we generate graph with 2 pair nodes
+			// this is very expensive when we have so many actors
+			// so we'll need to optimize later @todo
+			for(int32 ActorId=0; ActorId<Actors.Num(); ++ActorId)
+			{
+				AActor* Actor1 = (Actors[ActorId]);
+				for(int32 SubActorId=ActorId+1; SubActorId<Actors.Num(); ++SubActorId)
+				{
+					AActor* Actor2 = Actors[SubActorId];
+
+					// create new cluster
+					FLODCluster NewClusterCandidate = FLODCluster(Actor1, Actor2);
+					Clusters.Add(NewClusterCandidate);
+				}
+			}
+
+			// shrink after adding actors
+			// LOD 0 has lots of actors, and subsequence LODs tend to have a lot less actors
+			// so this should save a lot more. 
+			Clusters.Shrink();
+		}
+	}
+}
+
 void FHierarchicalLODBuilder::FindMST() 
 {
+	SCOPE_LOG_TIME(TEXT("STAT_HLOD_FindMST"), nullptr);
 	if (Clusters.Num() > 0)
 	{
 		// now sort edge in the order of weight
@@ -234,12 +279,77 @@ void FHierarchicalLODBuilder::FindMST()
 			}
 		};
 
-		Clusters.Sort(FCompareCluster());
+		Clusters.HeapSort(FCompareCluster());
 	}
 }
 
-void FHierarchicalLODBuilder::MergeClustersAndBuildActors(class ULevel* InLevel, const int32 LODIdx, float HighestCost, int32 MinNumActors)
+void FHierarchicalLODBuilder::HandleHLODVolumes(ULevel* InLevel)
+{	
+	for (int32 ActorId = 0; ActorId < InLevel->Actors.Num(); ++ActorId)
+	{
+		AActor* Actor = InLevel->Actors[ActorId];
+		if (Actor && Actor->IsA( AHierarchicalLODVolume::StaticClass()))
+		{
+			// Came across a HLOD volume			
+			FLODCluster& NewCluster = HLODVolumeClusters.Add(CastChecked<AHierarchicalLODVolume>(Actor));
+
+			//FVector Origin, Extent;
+			//Actor->GetActorBounds(false, Origin, Extent);			
+			//NewCluster.Bound = FSphere(Origin * CM_TO_METER, Extent.GetAbsMax() * CM_TO_METER);
+		}
+	}
+}
+
+bool FHierarchicalLODBuilder::ShouldGenerateCluster(AActor* Actor)
 {
+	if (!Actor)
+	{
+		return false;
+	}
+
+	if (!Actor->bEnableAutoLODGeneration)
+	{
+		return false;
+	}
+
+	FVector Origin, Extent;
+	Actor->GetActorBounds(false, Origin, Extent);
+	if (Extent.SizeSquared() <= 0.1)
+	{
+		return false;
+	}
+
+	// for now only consider staticmesh - I don't think skel mesh would work with simplygon merge right now @fixme
+	TArray<UStaticMeshComponent*> Components;
+	Actor->GetComponents<UStaticMeshComponent>(Components);
+	// TODO: support instanced static meshes
+	Components.RemoveAll([](UStaticMeshComponent* Val){ return Val->IsA(UInstancedStaticMeshComponent::StaticClass()); });
+
+	int32 ValidComponentCount = 0;
+	// now make sure you check parent primitive, so that we don't build for the actor that already has built. 
+	if (Components.Num() > 0)
+	{
+		for (auto& ComponentIter : Components)
+		{
+			if (ComponentIter->GetLODParentPrimitive())
+			{
+				return false;
+			}
+
+			// see if we should generate it
+			if (ComponentIter->ShouldGenerateAutoLOD())
+			{
+				++ValidComponentCount;
+				break;
+			}
+		}
+	}
+
+	return (ValidComponentCount > 0);
+}
+
+void FHierarchicalLODBuilder::MergeClustersAndBuildActors(ULevel* InLevel, const int32 LODIdx, float HighestCost, int32 MinNumActors, const bool bCreateMeshes)
+{	
 	if (Clusters.Num() > 0)
 	{
 		FString LevelName = FPackageName::GetShortName(InLevel->GetOutermost()->GetName());
@@ -248,6 +358,7 @@ void FHierarchicalLODBuilder::MergeClustersAndBuildActors(class ULevel* InLevel,
 		Arguments.Add(TEXT("LevelName"), FText::FromString(LevelName));
 		// merge clusters first
 		{
+			SCOPE_LOG_TIME(TEXT("HLOD_MergeClusters"), nullptr);
 			static int32 TotalIteration=3;
 			const int32 TotalCluster = Clusters.Num();
 
@@ -305,11 +416,17 @@ void FHierarchicalLODBuilder::MergeClustersAndBuildActors(class ULevel* InLevel,
 			}
 		}
 
+		for (auto& Cluster : HLODVolumeClusters)
+		{
+			Clusters.Add(Cluster.Value);
+		}
+
 		// debug flag, so that I can just see clustered data since no visualization in the editor yet
 		//static bool bBuildActor=true;
 
-		//if (bBuildActor)
+		//if (bBuildActors)
 		{
+			SCOPE_LOG_TIME(TEXT("HLOD_BuildActors"), nullptr);
 			// print data
 			int32 TotalValidCluster=0;
 			for(auto& Cluster: Clusters)
@@ -331,414 +448,11 @@ void FHierarchicalLODBuilder::MergeClustersAndBuildActors(class ULevel* InLevel,
 
 					if (Cluster.Actors.Num() >= MinNumActors)
 					{
-						Cluster.BuildActor(InLevel, LODIdx);
-					}
-				}
-			}
-		}
-	}
-}
-
-bool ShouldGenerateCluster(class AActor* Actor)
-{
-	if (!Actor)
-	{
-		return false;
-	}
-
-	if (!Actor->bEnableAutoLODGeneration)
-	{
-		return false;
-	}
-
-	// make sure we don't generate for LODActor for Level 0
-	if (Actor->IsA(ALODActor::StaticClass()))
-	{
-		return false;
-	}
-
-	FVector Origin, Extent;
-	Actor->GetActorBounds(false, Origin, Extent);
-	if(Extent.SizeSquared() <= 0.1)
-	{
-		return false;
-	}
-
-	// for now only consider staticmesh - I don't think skel mesh would work with simplygon merge right now @fixme
-	TArray<UStaticMeshComponent*> Components;
-	Actor->GetComponents<UStaticMeshComponent>(Components);
-	// TODO: support instanced static meshes
-	Components.RemoveAll([](UStaticMeshComponent* Val){ return Val->IsA(UInstancedStaticMeshComponent::StaticClass()); });
-
-	int32 ValidComponentCount=0;
-	// now make sure you check parent primitive, so that we don't build for the actor that already has built. 
-	if (Components.Num() > 0)
-	{
-		for (auto& ComponentIter: Components)
-		{
-			if (ComponentIter->GetLODParentPrimitive())
-			{
-				return false;
-			}
-
-			// see if we should generate it
-			if (ComponentIter->ShouldGenerateAutoLOD())
-			{
-				++ValidComponentCount;
-			}
-		}
-	}
-
-	return (ValidComponentCount>0);
-}
-
-void FHierarchicalLODBuilder::InitializeClusters(class ULevel* InLevel, const int32 LODIdx, float CullCost)
-{
-	if (InLevel->Actors.Num() > 0)
-	{
-		if (LODIdx == 0)
-		{
-			Clusters.Empty();
-
-			// first we generate graph with 2 pair nodes
-			// this is very expensive when we have so many actors
-			// so we'll need to optimize later @todo
-			for(int32 ActorId=0; ActorId<InLevel->Actors.Num(); ++ActorId)
-			{
-				AActor* Actor1 = InLevel->Actors[ActorId];
-				// make sure Actor has bound
-				if(ShouldGenerateCluster(Actor1))
-				{
-					for(int32 SubActorId=ActorId+1; SubActorId<InLevel->Actors.Num(); ++SubActorId)
-					{
-						AActor* Actor2 = InLevel->Actors[SubActorId];
-						if(ShouldGenerateCluster(Actor2))
-						{
-							FLODCluster NewClusterCandidate = FLODCluster(Actor1, Actor2);
-							float NewClusterCost = NewClusterCandidate.GetCost();
-							// @todo optimization : filter first - this might create disconnected tree - 
-							//@Todo debug but it's too slow to have full tree 
-							if ( NewClusterCost <= CullCost)
-							{
-								Clusters.Add(NewClusterCandidate);
-							}
-						}
-					}
-				}
-			}
-		}
-		else // at this point we only care for LODActors
-		{
-			Clusters.Empty();
-
-			// we filter the LOD index first
-			TArray<AActor*> Actors;
-			for(int32 ActorId=0; ActorId<InLevel->Actors.Num(); ++ActorId)
-			{
-				AActor* Actor = (InLevel->Actors[ActorId]);
-				if (Actor && Actor->IsA(ALODActor::StaticClass()))
-				{
-					ALODActor* LODActor = CastChecked<ALODActor>(Actor);
-					if (LODActor->LODLevel == LODIdx)
-					{
-						Actors.Add(Actor);
-					}
-				}
-				else if(ShouldGenerateCluster(Actor))
-				{
-					Actors.Add(Actor);
-				}
-			}
-			// first we generate graph with 2 pair nodes
-			// this is very expensive when we have so many actors
-			// so we'll need to optimize later @todo
-			for(int32 ActorId=0; ActorId<Actors.Num(); ++ActorId)
-			{
-				AActor* Actor1 = (Actors[ActorId]);
-				for(int32 SubActorId=ActorId+1; SubActorId<Actors.Num(); ++SubActorId)
-				{
-					AActor* Actor2 = Actors[SubActorId];
-
-					// create new cluster
-					FLODCluster NewClusterCandidate = FLODCluster(Actor1, Actor2);
-					Clusters.Add(NewClusterCandidate);
-				}
-			}
-
-			// shrink after adding actors
-			// LOD 0 has lots of actors, and subsequence LODs tend to have a lot less actors
-			// so this should save a lot more. 
-			Clusters.Shrink();
-		}
-	}
-}
-
-////////////////////////////////////////////////////////////////////////////////////////////
-// LODCluster
-////////////////////////////////////////////////////////////////////////////////////////////
-
-
-FLODCluster::FLODCluster(const FLODCluster& Other)
-: Actors(Other.Actors)
-, Bound(Other.Bound)
-, FillingFactor(Other.FillingFactor)
-, bValid(Other.bValid)
-{
-}
-
-FLODCluster::FLODCluster(class AActor* Actor1)
-: Bound(ForceInit)
-, bValid(true)
-{
-	AddActor(Actor1);
-	// calculate new filling factor
-	FillingFactor = 1.f;
-}
-
-FLODCluster::FLODCluster(class AActor* Actor1, class AActor* Actor2)
-: Bound(ForceInit)
-, bValid(true)
-{
-	FSphere Actor1Bound = AddActor(Actor1);
-	FSphere Actor2Bound = AddActor(Actor2);
-
-	// calculate new filling factor
-	FillingFactor = CalculateFillingFactor(Actor1Bound, 1.f, Actor2Bound, 1.f);
-}
-
-FSphere FLODCluster::AddActor(class AActor* NewActor)
-{
-	ensure (Actors.Contains(NewActor) == false);
-	Actors.Add(NewActor);
-	FVector Origin, Extent;
-	NewActor->GetActorBounds(false, Origin, Extent);
-	
-	// scale 0.01 (change to meter from centermeter)
-	FSphere NewBound = FSphere(Origin*CM_TO_METER, Extent.Size()*CM_TO_METER);
-	Bound += NewBound;
-
-	return NewBound;
-}
-
-FLODCluster FLODCluster::operator+(const FLODCluster& Other) const
-{
-	FLODCluster UnionCluster(*this);
-	UnionCluster.MergeClusters(Other);
-	return UnionCluster;
-}
-
-FLODCluster FLODCluster::operator+=(const FLODCluster& Other)
-{
-	MergeClusters(Other);
-	return *this;
-}
-
-FLODCluster FLODCluster::operator-(const FLODCluster& Other) const
-{
-	FLODCluster Cluster(*this);
-	Cluster.SubtractCluster(Other);
-	return *this;
-}
-
-FLODCluster FLODCluster::operator-=(const FLODCluster& Other)
-{
-	SubtractCluster(Other);
-	return *this;
-}
-
-FLODCluster& FLODCluster::operator=(const FLODCluster& Other)
-{
-	this->bValid		= Other.bValid;
-	this->Actors		= Other.Actors;
-	this->Bound			= Other.Bound;
-	this->FillingFactor = Other.FillingFactor;
-
-	return *this;
-}
-// Merge two clusters
-void FLODCluster::MergeClusters(const FLODCluster& Other)
-{
-	// please note that when merge, we merge two boxes from each cluster, not exactly all actors' bound
-	// have to recalculate filling factor and bound based on cluster data
-	FillingFactor = CalculateFillingFactor(Bound, FillingFactor, Other.Bound, Other.FillingFactor);
-	Bound += Other.Bound;
-
-	for (auto& Actor: Other.Actors)
-	{
-		Actors.AddUnique(Actor);
-	}
-}
-
-void FLODCluster::SubtractCluster(const FLODCluster& Other)
-{
-	for(int32 ActorId=0; ActorId<Actors.Num(); ++ActorId)
-	{
-		if (Other.Actors.Contains(Actors[ActorId]))
-		{
-			Actors.RemoveAt(ActorId);
-			--ActorId;
-		}
-	}
-
-	TArray<AActor*> NewActors = Actors;
-	Actors.Empty();
-	// need to recalculate parameter
-	if (NewActors.Num() == 0)
-	{
-		Invalidate();
-	}
-	else if (NewActors.Num() == 1)
-	{
-		Bound = FSphere(ForceInitToZero);
-		AddActor(NewActors[0]);
-		FillingFactor = 1.f;
-	}
-	else if (NewActors.Num() >= 2)
-	{
-		Bound = FSphere(ForceInit);
-
-		FSphere Actor1Bound = AddActor(NewActors[0]);
-		FSphere Actor2Bound = AddActor(NewActors[1]);
-
-		// calculate new filling factor
-		FillingFactor = CalculateFillingFactor(Actor1Bound, 1.f, Actor2Bound, 1.f);
-
-		// if more actors, we add them manually
-		for (int32 ActorId=2; ActorId<NewActors.Num(); ++ActorId)
-		{
-			// if not contained, it shouldn't be
-			check (!Actors.Contains(NewActors[ActorId]));
-
-			FSphere NewBound = AddActor(NewActors[ActorId]);
-			FillingFactor = CalculateFillingFactor(NewBound, 1.f, Bound, FillingFactor);
-			Bound += NewBound;
-		}
-	}
-}
-
-// if criteria matches, build new LODActor and replace current Actors with that. We don't need 
-// this clears previous actors and sets to this new actor
-// this is required when new LOD is created from these actors, this will be replaced
-// to save memory and to reduce memory increase during this process, we discard previous actors and replace with this actor
-void FLODCluster::BuildActor(class ULevel* InLevel, const int32 LODIdx)
-{
-	// do big size
-	if (InLevel && InLevel->GetWorld())
-	{
-		////////////////////////////////////////////////////////////////////////////////////
-		// create asset using Actors
-		const FHierarchicalSimplification& LODSetup = InLevel->GetWorld()->GetWorldSettings()->HierarchicalLODSetup[LODIdx];
-
-		// Where generated assets will be stored
-		UPackage* AssetsOuter = InLevel->GetOutermost(); // this asset is going to save with map, this means, I'll have to delete with it
-		if (AssetsOuter)
-		{
-			TArray<UStaticMeshComponent*> AllComponents;
-
-			for (auto& Actor: Actors)
-			{
-				TArray<UStaticMeshComponent*> Components;
-				Actor->GetComponents<UStaticMeshComponent>(Components);
-				// TODO: support instanced static meshes
-				Components.RemoveAll([](UStaticMeshComponent* Val){ return Val->IsA(UInstancedStaticMeshComponent::StaticClass()); });
-				AllComponents.Append(Components);
-			}
-
-			// it shouldn't even have come here if it didn't have any staticmesh
-			if (ensure(AllComponents.Num() > 0))
-			{
-				// In case we don't have outer generated assets should have same path as LOD level
-				
-				const FString AssetsPath = AssetsOuter->GetName() + TEXT("/");
-				class AActor* FirstActor = Actors[0];
-
-				TArray<UObject*> OutAssets;
-				FVector OutProxyLocation = FVector::ZeroVector;
-				// Generate proxy mesh and proxy material assets
-				IMeshUtilities& MeshUtilities = FModuleManager::Get().LoadModuleChecked<IMeshUtilities>("MeshUtilities");
-				// should give unique name, so use level + actor name
-				const FString PackageName = FString::Printf(TEXT("LOD_%s"), *FirstActor->GetName());
-				if (MeshUtilities.GetMeshMergingInterface() && LODSetup.bSimplifyMesh)
-				{
-					MeshUtilities.CreateProxyMesh(Actors, LODSetup.ProxySetting, AssetsOuter, PackageName, OutAssets, OutProxyLocation);
-				}
-				else
-				{
-					MeshUtilities.MergeActors(Actors, LODSetup.MergeSetting, AssetsOuter, PackageName, LODIdx+1, OutAssets, OutProxyLocation, true );
-				}
-
-				// we make it private, so it can't be used by outside of map since it's useless, and then remove standalone
-				for (auto& AssetIter : OutAssets)
-				{
-					AssetIter->ClearFlags(RF_Public | RF_Standalone);
-				}
-
-				UStaticMesh* MainMesh=NULL;
-				// set staticmesh
-				for(auto& Asset: OutAssets)
-				{
-					class UStaticMesh* StaticMesh = Cast<UStaticMesh>(Asset);
-
-					if(StaticMesh)
-					{
-						MainMesh = StaticMesh;
+						Cluster.BuildActor(InLevel, LODIdx, bCreateMeshes);
 					}
 				}
-
-				if (MainMesh)
-				{
-					UWorld* LevelWorld = Cast<UWorld>(InLevel->GetOuter());
-
-					check (LevelWorld);
-					////////////////////////////////////////////////////////////////////////////////////
-					// create LODActors using the current Actors
-					class ALODActor* NewActor = Cast<ALODActor>(LevelWorld->SpawnActor(ALODActor::StaticClass(), &OutProxyLocation, &FRotator::ZeroRotator));
-
-					NewActor->SubObjects = OutAssets;
-					NewActor->SubActors = Actors;
-					NewActor->LODLevel = LODIdx+1; 
-					float DrawDistance = LODSetup.DrawDistance;
-					NewActor->LODDrawDistance = DrawDistance;
-
-					NewActor->GetStaticMeshComponent()->StaticMesh = MainMesh;
-
-					// now set as parent
-					for(auto& Actor: Actors)
-					{
-						Actor->SetLODParent(NewActor->GetStaticMeshComponent(), DrawDistance);
-					}
-				}
-			}
-		}
-	}
-}
-
-bool FLODCluster::Contains(FLODCluster& Other) const
-{
-	if (IsValid() && Other.IsValid())
-	{
-		for(auto& Actor: Other.Actors)
-		{
-			if(Actors.Contains(Actor))
-			{
-				return true;
 			}
 		}
 	}
-
-	return false;
-}
-
-FString FLODCluster::ToString() const
-{
-	FString ActorList;
-	for (auto& Actor: Actors)
-	{
-		ActorList += Actor->GetActorLabel();
-		ActorList += ", ";
-	}
-
-	return FString::Printf(TEXT("ActorNum(%d), Actor List (%s)"), Actors.Num(), *ActorList);
 }
-
 #undef LOCTEXT_NAMESPACE 
diff --git a/Engine/Source/Editor/UnrealEd/Private/HierarchicalLODVolume.cpp b/Engine/Source/Editor/UnrealEd/Private/HierarchicalLODVolume.cpp
new file mode 100644
--- /dev/null
+++ b/Engine/Source/Editor/UnrealEd/Private/HierarchicalLODVolume.cpp
@@ -0,0 +1,14 @@
+// Copyright 1998-2015 Epic Games, Inc. All Rights Reserved.
+
+#include "UnrealEd.h"
+#include "HierarchicalLODVolume.h"
+
+AHierarchicalLODVolume::AHierarchicalLODVolume(const FObjectInitializer& ObjectInitializer)
+: Super(ObjectInitializer)
+{
+	bColored = true;
+	BrushColor.R = 255;
+	BrushColor.G = 100;
+	BrushColor.B = 255;
+	BrushColor.A = 255;
+}
diff --git a/Engine/Source/Editor/UnrealEd/Private/LODCluster.cpp b/Engine/Source/Editor/UnrealEd/Private/LODCluster.cpp
new file mode 100644
--- /dev/null
+++ b/Engine/Source/Editor/UnrealEd/Private/LODCluster.cpp
@@ -0,0 +1,441 @@
+// Copyright 1998-2015 Epic Games, Inc. All Rights Reserved.
+
+#include "UnrealEd.h"
+
+#include "LODCluster.h"
+
+#if WITH_EDITOR
+#include "Editor/UnrealEd/Classes/Editor/EditorEngine.h"
+#include "Editor/UnrealEd/Public/Editor.h"
+#include "Engine/LODActor.h"
+#include "Editor/UnrealEd/Classes/Factories/Factory.h"
+#include "MeshUtilities.h"
+#include "ObjectTools.h"
+#endif // WITH_EDITOR
+
+#include "GameFramework/WorldSettings.h"
+#include "Components/InstancedStaticMeshComponent.h"
+
+#include "HierarchicalLODVolume.h"
+
+#define LOCTEXT_NAMESPACE "LODCluster"
+#define CM_TO_METER		0.01f
+#define METER_TO_CM		100.0f
+
+/** Utility function to calculate overlap of two spheres */
+const float CalculateOverlap(const FSphere& ASphere, const float AFillingFactor, const FSphere& BSphere, const float BFillingFactor)
+{
+	// if it doesn't intersect, return zero 
+	if (!ASphere.Intersects(BSphere))
+	{
+		return 0.f;
+	}
+
+	if (ASphere.IsInside(BSphere, 1.f))
+	{
+		return ASphere.GetVolume();
+	}
+
+	if(BSphere.IsInside(ASphere, 1.f))
+	{
+		return BSphere.GetVolume();
+	}
+
+	if (ASphere.Equals(BSphere, 1.f))
+	{
+		return ASphere.GetVolume();
+	}
+
+	float Distance = (ASphere.Center-BSphere.Center).Size();
+	check (Distance > 0.f);
+
+	float ARadius = ASphere.W;
+	float BRadius = BSphere.W;
+
+	float ACapHeight = (BRadius*BRadius - (ARadius - Distance)*(ARadius - Distance)) / (2*Distance);
+	float BCapHeight = (ARadius*ARadius - (BRadius - Distance)*(BRadius - Distance)) / (2*Distance);
+
+	if (ACapHeight<=0.f || BCapHeight<=0.f)
+	{
+		// it's possible to get cap height to be less than 0 
+		// since when we do check intersect, we do have regular tolerance
+		return 0.f;		
+	}
+
+	float OverlapRadius1 = ((ARadius + BRadius)*(ARadius + BRadius) - Distance*Distance) * (Distance*Distance - (ARadius - BRadius)*(ARadius - BRadius));
+	float OverlapRadius2 = 2 * Distance;
+	float OverlapRedius = FMath::Sqrt(OverlapRadius1) / OverlapRadius2;
+	float OverlapRediusSq = OverlapRedius*OverlapRedius;
+
+	check (OverlapRadius1 >= 0.f);
+
+	float ConstPI = PI/6.0f;
+	float AVolume = ConstPI*(3*OverlapRediusSq + ACapHeight*ACapHeight) * ACapHeight;
+	float BVolume = ConstPI*(3*OverlapRediusSq + BCapHeight*BCapHeight) * BCapHeight;
+
+	check (AVolume > 0.f &&  BVolume > 0.f);
+
+	float TotalVolume = AFillingFactor*AVolume + BFillingFactor*BVolume;
+	return TotalVolume;
+}
+
+/** Utility function that calculates filling factor */
+const float CalculateFillingFactor(const FSphere& ASphere, const float AFillingFactor, const FSphere& BSphere, const float BFillingFactor)
+{
+	const float OverlapVolume = CalculateOverlap( ASphere, AFillingFactor, BSphere, BFillingFactor);
+	FSphere UnionSphere = ASphere + BSphere;
+	// it shouldn't be zero or it should be checked outside
+	ensure(UnionSphere.W != 0.f);
+
+	// http://deim.urv.cat/~rivi/pub/3d/icra04b.pdf
+	// cost is calculated based on r^3 / filling factor
+	// since it subtract by AFillingFactor * 1/2 overlap volume + BfillingFactor * 1/2 overlap volume
+	return FMath::Max(0.f, (AFillingFactor * ASphere.GetVolume() + BFillingFactor * BSphere.GetVolume() - OverlapVolume) / UnionSphere.GetVolume());
+}
+
+FLODCluster::FLODCluster(const FLODCluster& Other)
+: Actors(Other.Actors)
+, Bound(Other.Bound)
+, FillingFactor(Other.FillingFactor)
+, ClusterCost( FMath::Pow(Bound.W, 3) / FillingFactor )
+, bValid(Other.bValid)
+{
+}
+
+FLODCluster::FLODCluster(AActor* Actor1)
+: Bound(ForceInit)
+, bValid(true)
+{
+	AddActor(Actor1);
+	// calculate new filling factor
+	FillingFactor = 1.f;	
+	ClusterCost = FMath::Pow(Bound.W, 3) / FillingFactor;
+}
+
+FLODCluster::FLODCluster(AActor* Actor1, AActor* Actor2)
+: Bound(ForceInit)
+, bValid(true)
+{
+	FSphere Actor1Bound = AddActor(Actor1);
+	FSphere Actor2Bound = AddActor(Actor2);
+	
+	// calculate new filling factor
+	FillingFactor = CalculateFillingFactor(Actor1Bound, 1.f, Actor2Bound, 1.f);	
+	ClusterCost = FMath::Pow(Bound.W, 3) / FillingFactor;
+}
+
+FLODCluster::FLODCluster()
+: Bound(ForceInit)
+, bValid(false)
+{
+	// calculate new filling factor
+	FillingFactor = 1.f;
+	ClusterCost = FMath::Pow(Bound.W, 3) / FillingFactor;
+}
+
+FSphere FLODCluster::AddActor(AActor* NewActor)
+{
+	bValid = true;
+	ensure (Actors.Contains(NewActor) == false);
+	Actors.Add(NewActor);
+	FVector Origin, Extent;
+	if (NewActor->IsA<ALODActor>())
+	{
+		ALODActor* LODActor = CastChecked<ALODActor>(NewActor);
+		if (LODActor->IsPreviewActorInstance())
+		{
+			Origin = LODActor->GetDrawSphereComponent()->GetComponentLocation();
+			Extent = FVector(LODActor->GetDrawSphereComponent()->GetScaledSphereRadius());
+		}
+		else
+		{
+			NewActor->GetActorBounds(false, Origin, Extent);
+		}
+	}
+	else
+	{
+		NewActor->GetActorBounds(false, Origin, Extent);
+	}
+
+	// scale 0.01 (change to meter from centimeter) QQ Extens.GetSize	
+	FSphere NewBound = FSphere(Origin*CM_TO_METER, Extent.Size()*CM_TO_METER);
+	Bound += NewBound;
+
+	return NewBound;
+}
+
+FLODCluster FLODCluster::operator+(const FLODCluster& Other) const
+{
+	FLODCluster UnionCluster(*this);
+	UnionCluster.MergeClusters(Other);
+	return UnionCluster;
+}
+
+FLODCluster FLODCluster::operator+=(const FLODCluster& Other)
+{
+	MergeClusters(Other);
+	return *this;
+}
+
+FLODCluster FLODCluster::operator-(const FLODCluster& Other) const
+{
+	FLODCluster Cluster(*this);
+	Cluster.SubtractCluster(Other);
+	return *this;
+}
+
+FLODCluster FLODCluster::operator-=(const FLODCluster& Other)
+{
+	SubtractCluster(Other);
+	return *this;
+}
+
+FLODCluster& FLODCluster::operator=(const FLODCluster& Other)
+{
+	this->bValid		= Other.bValid;
+	this->Actors		= Other.Actors;
+	this->Bound			= Other.Bound;
+	this->FillingFactor = Other.FillingFactor;	
+	this->ClusterCost = Other.ClusterCost;
+
+	return *this;
+}
+
+void FLODCluster::MergeClusters(const FLODCluster& Other)
+{
+	// please note that when merge, we merge two boxes from each cluster, not exactly all actors' bound
+	// have to recalculate filling factor and bound based on cluster data
+	FillingFactor = CalculateFillingFactor(Bound, FillingFactor, Other.Bound, Other.FillingFactor);
+	Bound += Other.Bound;	
+	ClusterCost = FMath::Pow(Bound.W, 3) / FillingFactor;
+
+	for (auto& Actor: Other.Actors)
+	{
+		Actors.AddUnique(Actor);
+	}
+
+	if (Actors.Num() > 0)
+	{
+		bValid = true;
+	}
+}
+
+void FLODCluster::SubtractCluster(const FLODCluster& Other)
+{
+	for(int32 ActorId=0; ActorId<Actors.Num(); ++ActorId)
+	{
+		if (Other.Actors.Contains(Actors[ActorId]))
+		{
+			Actors.RemoveAt(ActorId);
+			--ActorId;
+		}
+	}
+
+	TArray<AActor*> NewActors = Actors;
+	Actors.Empty();
+	// need to recalculate parameter
+	if (NewActors.Num() == 0)
+	{
+		Invalidate();
+	}
+	else if (NewActors.Num() == 1)
+	{
+		Bound = FSphere(ForceInitToZero);
+		AddActor(NewActors[0]);
+		FillingFactor = 1.f;
+		ClusterCost = FMath::Pow(Bound.W, 3) / FillingFactor;
+	}
+	else if (NewActors.Num() >= 2)
+	{
+		Bound = FSphere(ForceInit);
+
+		FSphere Actor1Bound = AddActor(NewActors[0]);
+		FSphere Actor2Bound = AddActor(NewActors[1]);
+
+		// calculate new filling factor
+		FillingFactor = CalculateFillingFactor(Actor1Bound, 1.f, Actor2Bound, 1.f);
+
+		// if more actors, we add them manually
+		for (int32 ActorId=2; ActorId<NewActors.Num(); ++ActorId)
+		{
+			// if not contained, it shouldn't be
+			check (!Actors.Contains(NewActors[ActorId]));
+
+			FSphere NewBound = AddActor(NewActors[ActorId]);
+			FillingFactor = CalculateFillingFactor(NewBound, 1.f, Bound, FillingFactor);
+			Bound += NewBound;
+		}
+
+		ClusterCost = FMath::Pow(Bound.W, 3) / FillingFactor;
+	}
+}
+
+void FLODCluster::BuildActor(ULevel* InLevel, const int32 LODIdx, const bool bCreateMeshes)
+{
+	FColor Colours[8] = { FColor::Cyan, FColor::Red, FColor::Green, FColor::Blue, FColor::Yellow, FColor::Magenta, FColor::White, FColor::Black };
+	// do big size
+	if (InLevel && InLevel->GetWorld())
+	{
+		// create asset using Actors
+		const FHierarchicalSimplification& LODSetup = InLevel->GetWorld()->GetWorldSettings()->HierarchicalLODSetup[LODIdx];
+
+		// Retrieve draw distance for current and next LOD level
+		const float DrawDistance = LODSetup.DrawDistance;		
+		const int32 LODCount = InLevel->GetWorld()->GetWorldSettings()->HierarchicalLODSetup.Num();
+		const float NextDrawDistance = (LODIdx < (LODCount - 1)) ? InLevel->GetWorld()->GetWorldSettings()->HierarchicalLODSetup[LODIdx + 1].DrawDistance : 0.0f;
+
+		// Where generated assets will be stored
+		UPackage* AssetsOuter = InLevel->GetOutermost(); // this asset is going to save with map, this means, I'll have to delete with it
+		if (AssetsOuter)
+		{
+			TArray<UStaticMeshComponent*> AllComponents;
+
+			for (auto& Actor: Actors)
+			{
+				TArray<UStaticMeshComponent*> Components;
+				
+				if (Actor->IsA<ALODActor>())
+				{
+					ExtractStaticMeshComponentsFromLODActor(Actor, Components);
+				}
+				else
+				{
+					Actor->GetComponents<UStaticMeshComponent>(Components);
+				}
+
+				// TODO: support instanced static meshes
+				Components.RemoveAll([](UStaticMeshComponent* Val){ return Val->IsA(UInstancedStaticMeshComponent::StaticClass()); });
+
+				AllComponents.Append(Components);
+			}
+
+			// it shouldn't even have come here if it didn't have any staticmesh
+			if (ensure(AllComponents.Num() > 0))
+			{
+				// In case we don't have outer generated assets should have same path as LOD level
+
+				const FString AssetsPath = AssetsOuter->GetName() + TEXT("/");
+				AActor* FirstActor = Actors[0];
+
+				TArray<UObject*> OutAssets;
+				FVector OutProxyLocation = FVector::ZeroVector;
+				UStaticMesh* MainMesh = nullptr;
+				if (bCreateMeshes)
+				{
+					// Generate proxy mesh and proxy material assets
+					IMeshUtilities& MeshUtilities = FModuleManager::Get().LoadModuleChecked<IMeshUtilities>("MeshUtilities");
+					// should give unique name, so use level + actor name
+					const FString PackageName = FString::Printf(TEXT("LOD_%s"), *FirstActor->GetName());
+					if (MeshUtilities.GetMeshMergingInterface() && LODSetup.bSimplifyMesh)
+					{
+						MeshUtilities.CreateProxyMesh(Actors, LODSetup.ProxySetting, AssetsOuter, PackageName, OutAssets, OutProxyLocation);
+					}
+					else
+					{						
+						MeshUtilities.MergeStaticMeshComponents(AllComponents, FirstActor->GetWorld(), LODSetup.MergeSetting, AssetsOuter, PackageName, LODIdx, OutAssets, OutProxyLocation, LODSetup.DrawDistance, true);
+					}
+
+					// we make it private, so it can't be used by outside of map since it's useless, and then remove standalone
+					for (auto& AssetIter : OutAssets)
+					{
+						AssetIter->ClearFlags(RF_Public | RF_Standalone);
+					}
+
+					
+					// set staticmesh
+					for (auto& Asset : OutAssets)
+					{
+						UStaticMesh* StaticMesh = Cast<UStaticMesh>(Asset);
+
+						if (StaticMesh)
+						{
+							MainMesh = StaticMesh;
+						}
+					}
+				}
+							
+
+				if (MainMesh || !bCreateMeshes)
+				{
+					UWorld* LevelWorld = Cast<UWorld>(InLevel->GetOuter());
+
+					check (LevelWorld);
+										
+					// create LODActors using the current Actors
+					ALODActor* NewActor = Cast<ALODActor>(LevelWorld->SpawnActor(ALODActor::StaticClass(), &OutProxyLocation, &FRotator::ZeroRotator));
+					NewActor->SubObjects = OutAssets;
+					NewActor->SubActors = Actors;
+					NewActor->LODLevel = LODIdx+1; 					
+					NewActor->LODDrawDistance = DrawDistance;
+					NewActor->SetStaticMesh( MainMesh );
+					NewActor->SetFolderPath("/HLODActors");											
+
+#if WITH_EDITOR						
+					// Setup sphere drawing visualization
+					UDrawSphereComponent* SphereComponent = NewActor->GetDrawSphereComponent();
+					SphereComponent->SetSphereRadius(Bound.W * METER_TO_CM);
+					SphereComponent->SetWorldLocation(Bound.Center * METER_TO_CM);					
+					SphereComponent->ShapeColor = Colours[LODIdx % 7];
+					SphereComponent->MinDrawDistance = DrawDistance;
+					SphereComponent->LDMaxDrawDistance = NextDrawDistance;
+					SphereComponent->CachedMaxDrawDistance = NextDrawDistance;	
+#endif								
+					// now set as parent
+					for (auto& Actor : Actors)
+					{
+						Actor->SetLODParent(NewActor->GetStaticMeshComponent(), DrawDistance);
+					}					
+				}
+			}
+		}
+	}
+}
+
+void FLODCluster::ExtractStaticMeshComponentsFromLODActor(AActor* Actor, TArray<UStaticMeshComponent*> &InOutComponents)
+{
+	ALODActor* LODActor = CastChecked<ALODActor>(Actor);
+	for (auto& ChildActor : LODActor->SubActors)
+	{
+		TArray<UStaticMeshComponent*> ChildComponents;
+		if (ChildActor->IsA<ALODActor>())
+		{
+			ExtractStaticMeshComponentsFromLODActor(ChildActor, ChildComponents);
+		}
+		else
+		{
+			ChildActor->GetComponents<UStaticMeshComponent>(ChildComponents);
+		}
+		
+		InOutComponents.Append(ChildComponents);
+	}
+}
+
+bool FLODCluster::Contains(FLODCluster& Other) const
+{
+	if (IsValid() && Other.IsValid())
+	{
+		for(auto& Actor: Other.Actors)
+		{
+			if(Actors.Contains(Actor))
+			{
+				return true;
+			}
+		}
+	}
+
+	return false;
+}
+
+FString FLODCluster::ToString() const
+{
+	FString ActorList;
+	for (auto& Actor: Actors)
+	{
+		ActorList += Actor->GetActorLabel();
+		ActorList += ", ";
+	}
+
+	return FString::Printf(TEXT("ActorNum(%d), Actor List (%s)"), Actors.Num(), *ActorList);
+}
+
+#undef LOCTEXT_NAMESPACE 
diff --git a/Engine/Source/Editor/UnrealEd/Public/EditorBuildUtils.h b/Engine/Source/Editor/UnrealEd/Public/EditorBuildUtils.h
--- a/Engine/Source/Editor/UnrealEd/Public/EditorBuildUtils.h
+++ b/Engine/Source/Editor/UnrealEd/Public/EditorBuildUtils.h
@@ -29,6 +29,8 @@ namespace EBuildOptions
 		BuildAllOnlySelectedPaths,
 		/** Build Hierarchical LOD system - need WorldSetting setup*/
 		BuildHierarchicalLOD, 
+		/** Preview Hierarchical LOD system clusters - need WorldSetting setup*/
+		PreviewHierarchicalLOD,
 	};
 }
 
diff --git a/Engine/Source/Editor/UnrealEd/Public/HierarchicalLOD.h b/Engine/Source/Editor/UnrealEd/Public/HierarchicalLOD.h
--- a/Engine/Source/Editor/UnrealEd/Public/HierarchicalLOD.h
+++ b/Engine/Source/Editor/UnrealEd/Public/HierarchicalLOD.h
@@ -2,76 +2,16 @@
 
 #pragma once
 
+#include "LODCluster.h"
+
 /*=============================================================================
 	HierarchicalLOD.h: Hierarchical LOD definition.
 =============================================================================*/
 
-/**
- *
- *	This is a LOD cluster struct that holds list of actors with relevant information
- *
- *	http://deim.urv.cat/~rivi/pub/3d/icra04b.pdf
- *
- *	This is used by Hierarchical LOD Builder to generates list of clusters 
- *	that are together in vincinity and build as one actor
- *
- **/
-
-struct FLODCluster
-{
-	// constructors
-	FLODCluster(const FLODCluster& Other);
-	FLODCluster(class AActor* Actor1);
-	FLODCluster(class AActor* Actor1, class AActor* Actor2);
-
-	FLODCluster operator+( const FLODCluster& Other ) const;
-	FLODCluster operator+=( const FLODCluster& Other );
-	FLODCluster operator-(const FLODCluster& Other) const;
-	FLODCluster operator-=(const FLODCluster& Other);
-	FLODCluster& operator=(const FLODCluster & Other);
-
-	// Invalidate current cluster
-	void Invalidate() { bValid = false; }
-	// return true if valid
-	bool IsValid() const {	return bValid; }
-
-	// return cost of the cluster, lower is better
-	const float GetCost() const
-	{
-		return (FMath::Pow(Bound.W, 3) / FillingFactor);
-	}
-
-	// return true if this cluster contains ANY of actors of Other
-	bool Contains(FLODCluster& Other) const;
-	// return string of data
-	FString ToString() const;
-	
-	// member variable
-	// list of actors
-	TArray<class AActor*>	Actors;
-	// bound of this cluster
-	FSphere					Bound;
-	// filling factor - higher means filled more
-	float					FillingFactor;
-
-	// if criteria matches, build new LODActor and replace current Actors with that. We don't need 
-	// this clears previous actors and sets to this new actor
-	// this is required when new LOD is created from these actors, this will be replaced
-	// to save memory and to reduce memory increase during this process, we discard previous actors and replace with this actor
-	void BuildActor(class ULevel* InLevel, const int32 LODIdx);
-
-private:
-	// cluster operations
-	void MergeClusters(const FLODCluster& Other);
-	void SubtractCluster(const FLODCluster& Other);
-
-	// add new actor, this doesn't recalculate filling factor
-	// the filling factor has to be calculated outside
-	FSphere AddActor(class AActor* NewActor);
-
-	// whether it's valid or not
-	bool bValid;
-};
+class ULevel;
+class AActor;
+class AHierarchicalLODVolume;
+class UWorld;
 
 /**
  *
@@ -79,35 +19,79 @@ private:
  *
  * This builds list of clusters and make sure it's sorted in the order of lower cost to high and merge clusters
  **/
-
 struct UNREALED_API FHierarchicalLODBuilder
 {
-	FHierarchicalLODBuilder(class UWorld* InWorld);
+	FHierarchicalLODBuilder(UWorld* InWorld);
 
 #if WITH_HOT_RELOAD_CTORS
 	/** DO NOT USE. This constructor is for internal usage only for hot-reload purposes. */
 	FHierarchicalLODBuilder();
 #endif // WITH_HOT_RELOAD_CTORS
-
-	// build hierarchical cluster
+		
+	/**
+	* Build, Builds the clusters and spawn LODActors with their merged Static Meshes
+	*/
 	void Build();
+	
+	/**
+	* PreviewBuild, Builds the clusters and spawns LODActors but without actually creating/merging new StaticMeshes
+	*/
+	void PreviewBuild();
 
 private:
-	// data structure - this is only valid within scope since mem stack allocator
-	TArray<FLODCluster, TMemStackAllocator<>> Clusters;
-
-	// owner world
-	class UWorld*		World;
-
-	// for now it only builds per level, it turns to one actor at the end
-	void BuildClusters(class ULevel* InLevel);
+	/**
+	* Builds the clusters (HLODs) for InLevel, and will create the new/merged StaticMeshes if bCreateMeshes is true
+	*
+	* @param InLevel - Level for which the HLODs are currently being build
+	* @param bCreateMeshes - Whether or not to create/merge the StaticMeshes (false if builing preview only)
+	*/
+	void BuildClusters(ULevel* InLevel, const bool bCreateMeshes);
+
+	/**
+	* Initializes the clusters, creating one for each actor within the level eligble for HLOD-ing
+	*
+	* @param InLevel - Level for which the HLODs are currently being build
+	* @param LODIdx - LOD index we are building
+	* @param CullCost - Test variable for tweaking HighestCost
+	*/
+	void InitializeClusters(ULevel* InLevel, const int32 LODIdx, float CullCost);
+
+	/**
+	* Merges clusters and builds actors for resulting (valid) clusters
+	*
+	* @param InLevel - Level for which the HLODs are currently being build
+	* @param LODIdx - LOD index we are building, used for determining which StaticMesh LOD to use
+	* @param HighestCost - Allowed HighestCost for this LOD level
+	* @param MinNumActors - Minimum number of actors for this LOD level
+	* @param bCreateMeshes - Whether or not to create/merge the StaticMeshes (false if builing preview only)
+	*/
+	void MergeClustersAndBuildActors(ULevel* InLevel, const int32 LODIdx, float HighestCost, int32 MinNumActors, const bool bCreateMeshes);
+
+	/**
+	* Finds the minimal spanning tree MST for the clusters by sorting on their cost ( Lower == better )
+	*/
+	void FindMST();
 
-	// initialize Clusters variable - each Actor becomes Cluster
-	void InitializeClusters(class ULevel* InLevel, const int32 LODIdx, float CullCost);
+	/* Retrieves HierarchicalLODVolumes and creates a cluster for each individual one
+	*
+	* @param InLevel - Level for which the HLODs are currently being build
+	*/
+	void HandleHLODVolumes(ULevel* InLevel);
+
+	/**
+	* Determine whether or not this actor is eligble for HLOD creation
+	*
+	* @param Actor - Actor to test
+	* @return bool
+	*/
+	bool ShouldGenerateCluster(AActor* Actor);
+	
+	/** Array of LOD Clusters - this is only valid within scope since mem stack allocator */
+	TArray<FLODCluster, TMemStackAllocator<>> Clusters;
 
-	// merge clusters
-	void MergeClustersAndBuildActors(class ULevel* InLevel, const int32 LODIdx, float HighestCost, int32 MinNumActors);
+	/** Owning world HLODs are created for */
+	UWorld*	World;
 
-	// find minmal spanning tree of clusters
-	void FindMST();
+	/** Array of LOD clusters created for the HierachicalLODVolumes found within the level */
+	TMap<AHierarchicalLODVolume*, FLODCluster> HLODVolumeClusters;	
 };
diff --git a/Engine/Source/Editor/UnrealEd/Public/HierarchicalLODVolume.h b/Engine/Source/Editor/UnrealEd/Public/HierarchicalLODVolume.h
new file mode 100644
--- /dev/null
+++ b/Engine/Source/Editor/UnrealEd/Public/HierarchicalLODVolume.h
@@ -0,0 +1,13 @@
+// Copyright 1998-2015 Epic Games, Inc. All Rights Reserved.
+
+#pragma once
+
+#include "GameFramework/Volume.h"
+#include "HierarchicalLODVolume.generated.h"
+
+/** An invisible volume used to manually define/create a HLOD cluster. */
+UCLASS()
+class AHierarchicalLODVolume : public AVolume
+{
+	GENERATED_UCLASS_BODY()
+};
diff --git a/Engine/Source/Editor/UnrealEd/Public/LODCluster.h b/Engine/Source/Editor/UnrealEd/Public/LODCluster.h
new file mode 100644
--- /dev/null
+++ b/Engine/Source/Editor/UnrealEd/Public/LODCluster.h
@@ -0,0 +1,107 @@
+// Copyright 1998-2015 Epic Games, Inc. All Rights Reserved.
+
+#pragma once
+
+class ULevel;
+class AActor;
+
+/**
+ *
+ *	This is a LOD cluster struct that holds list of actors with relevant information
+ *
+ *	http://deim.urv.cat/~rivi/pub/3d/icra04b.pdf
+ *
+ *	This is used by Hierarchical LOD Builder to generates list of clusters 
+ *	that are together in vicinity and build as one actor
+ *
+ **/
+struct FLODCluster
+{
+	// constructors
+	FLODCluster(const FLODCluster& Other);
+	FLODCluster(AActor* Actor1);
+	FLODCluster(AActor* Actor1, AActor* Actor2);
+	FLODCluster();
+
+	/** Cluster operators */
+	FLODCluster operator+( const FLODCluster& Other ) const;
+	FLODCluster operator+=( const FLODCluster& Other );
+	FLODCluster operator-(const FLODCluster& Other) const;
+	FLODCluster operator-=(const FLODCluster& Other);
+	FLODCluster& operator=(const FLODCluster & Other);
+
+	/** Invalidates this cluster */
+	void Invalidate() { bValid = false; }
+	
+	/** Returns whether or not this cluster is valid */
+	bool IsValid() const {	return bValid; }
+
+	/** Return cost of the cluster, lower is better */
+	const float GetCost() const
+	{
+		return ClusterCost;
+	}
+
+	/** Compare clusters and returns true when this contains any of Other's actors */
+	bool Contains(FLODCluster& Other) const;
+	
+	/** Returns data/info for this Cluster as a string */
+	FString ToString() const;
+	
+	// member variable
+	/** List of Actors that this cluster contains */
+	TArray<AActor*> Actors;
+	/** Cluster bounds */
+	FSphere	Bound;
+	/** Filling factor for this cluster, determines how much of the cluster's bounds/area is occupied by the contained actors*/
+	float FillingFactor;
+	/** Cached cluster cost, FMath::Pow(Bound.W, 3) / FillingFactor */
+	float ClusterCost;
+
+	/**
+	* if criteria matches, build new LODActor and replace current Actors with that. We don't need 
+	* this clears previous actors and sets to this new actor
+	* this is required when new LOD is created from these actors, this will be replaced
+	* to save memory and to reduce memory increase during this process, we discard previous actors and replace with this actor
+	*
+	* @param InLevel - Level for which currently the HLODs are being build
+	* @param LODIdx - LOD index to build
+	* @param bCreateMeshes - Whether or not the new StaticMeshes should be created
+	* @return void
+	*/
+	void BuildActor(ULevel* InLevel, const int32 LODIdx, const bool bCreateMeshes);
+
+	/**
+	* Recursively retrieves StaticMeshComponents from a LODActor and its child LODActors 
+	*
+	* @param Actor - LODActor instance
+	* @param InOutComponents - Will hold the StaticMeshComponents	
+	*/
+	void ExtractStaticMeshComponentsFromLODActor(AActor* Actor, TArray<UStaticMeshComponent*>& InOutComponents);
+
+private:	
+	/**
+	* Merges this cluster with Other by combining the actor arrays and updating the bounds, filling factor and cluster cost
+	*
+	* @param Other - Other Cluster to merge with	
+	*/
+	void MergeClusters(const FLODCluster& Other);
+
+	/**
+	* Subtracts the Other cluster from this cluster by removing Others actors and updating the bounds, filling factor and cluster cost
+	* This will invalidate the cluster if Actors.Num result to be 0
+	*
+	* @param Other - Other cluster to subtract
+	*/
+	void SubtractCluster(const FLODCluster& Other);
+
+	/**
+	* Adds a new actor to this cluster and updates the bounds accordingly, the filling factor is NOT updated
+	*
+	* @param NewActor - Actor to be added to the cluster	
+	*/
+	FSphere AddActor(AActor* NewActor);
+
+	/** Bool flag whether or not this cluster is valid */
+	bool bValid;
+};
diff --git a/Engine/Source/Runtime/Core/Private/Misc/MapErrors.cpp b/Engine/Source/Runtime/Core/Private/Misc/MapErrors.cpp
--- a/Engine/Source/Runtime/Core/Private/Misc/MapErrors.cpp
+++ b/Engine/Source/Runtime/Core/Private/Misc/MapErrors.cpp
@@ -57,6 +57,7 @@ FName FMapErrors::RepairedPaintedVertexColors(TEXT("RepairedPaintedVertexColors"
 FName FMapErrors::LODActorMissingStaticMesh(TEXT("LODActorMissingStaticMesh"));
 FName FMapErrors::LODActorMissingActor(TEXT("LODActorMissingActor"));
 FName FMapErrors::LODActorNoActorFound(TEXT("LODActorNoActor"));
+FName FMapErrors::HLODSystemNotEnabled(TEXT("HLODSystemNotEnabled"));
 
 static const FString MapErrorsPath = TEXT("Shared/Editor/MapErrors");
 
diff --git a/Engine/Source/Runtime/Core/Public/Misc/MapErrors.h b/Engine/Source/Runtime/Core/Public/Misc/MapErrors.h
--- a/Engine/Source/Runtime/Core/Public/Misc/MapErrors.h
+++ b/Engine/Source/Runtime/Core/Public/Misc/MapErrors.h
@@ -211,6 +211,9 @@ struct CORE_API FMapErrors
 
 	/** {LODActor} : NoActor is assigned. We recommend to delete this actor. */
 	static FName LODActorNoActorFound;
+
+	/** Hierarchical LOD System is disabled, unable to build LOD actors */
+	static FName HLODSystemNotEnabled;
 };
 
 /**
diff --git a/Engine/Source/Runtime/Engine/Classes/Engine/LODActor.h b/Engine/Source/Runtime/Engine/Classes/Engine/LODActor.h
--- a/Engine/Source/Runtime/Engine/Classes/Engine/LODActor.h
+++ b/Engine/Source/Runtime/Engine/Classes/Engine/LODActor.h
@@ -27,7 +27,7 @@ public:
 	/** what distance do you want this to show up instead of SubActors */
 	UPROPERTY(Category=LODActor, EditAnywhere, meta=(ClampMin = "0", UIMin = "0"))
 	float LODDrawDistance;
-
+	
 	/** what distance do you want this to show up instead of SubActors */
 	UPROPERTY(Category=LODActor, VisibleAnywhere)
 	int32 LODLevel;
@@ -43,6 +43,12 @@ public:
 	// End AActor Interface
 #endif // WITH_EDITOR	
 
+	/** Sets StaticMesh and IsPreviewActor to true if InStaticMesh equals nullptr */
+	void SetStaticMesh(class UStaticMesh* InStaticMesh);
+	
+	/** Returns whether or not this is a preview actor instance, meaning it has no StaticMesh */
+	const bool IsPreviewActorInstance() { return bIsPreviewActor;  }
+	
 protected:
 	// Begin UObject interface.
 	virtual FString GetDetailedInfoInternal() const override;
@@ -56,11 +62,21 @@ protected:
 #endif // WITH_EDITOR	
 	virtual void PostRegisterAllComponents() override;
 	// End UObject interface.
+		
+protected:
+#if WITH_EDITORONLY_DATA
+	/** Visualization component for rendering the LODActors bounds (cluster bounds)*/
+	UPROPERTY()
+	class UDrawSphereComponent* DrawSphereComponent;
+#endif // WITH_EDITORONLY_DATA
+
+	/** Whether or not this is a preview actor, will not fire StaticMesh map errors if so*/
+	bool bIsPreviewActor;
 
 public:
 	/** Returns StaticMeshComponent subobject **/
 	class UStaticMeshComponent* GetStaticMeshComponent() const { return StaticMeshComponent; }
-};
-
-
 
+	/** Returns StaticMeshComponent subobject **/
+	class UDrawSphereComponent* GetDrawSphereComponent() const { return DrawSphereComponent; }
+}; 
diff --git a/Engine/Source/Runtime/Engine/Classes/Engine/World.h b/Engine/Source/Runtime/Engine/Classes/Engine/World.h
--- a/Engine/Source/Runtime/Engine/Classes/Engine/World.h
+++ b/Engine/Source/Runtime/Engine/Classes/Engine/World.h
@@ -14,10 +14,6 @@
 #include "GameFramework/Actor.h"
 #include "GameInstance.h"
 
-#if WITH_EDITOR
-	#include "Editor/UnrealEd/Public/HierarchicalLOD.h"
-#endif // WITH_EDITOR
-
 #include "World.generated.h"
 
 class FPhysScene;
@@ -678,7 +674,7 @@ public:
 
 #if WITH_EDITOR
 	/** Hierarchical LOD System. Used when WorldSetting.bEnableHierarchicalLODSystem is true */
-	FHierarchicalLODBuilder						HierarchicalLODBuilder;
+	struct FHierarchicalLODBuilder*						HierarchicalLODBuilder;
 #endif // WITH_EDITOR
 
 private:
diff --git a/Engine/Source/Runtime/Engine/Private/Components/SphereComponent.cpp b/Engine/Source/Runtime/Engine/Private/Components/SphereComponent.cpp
--- a/Engine/Source/Runtime/Engine/Private/Components/SphereComponent.cpp
+++ b/Engine/Source/Runtime/Engine/Private/Components/SphereComponent.cpp
@@ -105,6 +105,13 @@ FPrimitiveSceneProxy* USphereComponent::CreateSceneProxy()
 					const FMatrix& LocalToWorld = GetLocalToWorld();
 					const FLinearColor DrawSphereColor = GetViewSelectionColor(SphereColor, *View, IsSelected(), IsHovered(), false, IsIndividuallySelected() );
 
+					// Taking into account the min and maximum drawing distance
+					const float Distance = (View->ViewMatrices.ViewOrigin - LocalToWorld.GetOrigin()).Size();
+					if (Distance < GetMinDrawDistance() || Distance > GetMaxDrawDistance() )
+					{
+						continue;
+					}
+					
 					float AbsScaleX = LocalToWorld.GetScaledAxis(EAxis::X).Size();
 					float AbsScaleY = LocalToWorld.GetScaledAxis(EAxis::Y).Size();
 					float AbsScaleZ = LocalToWorld.GetScaledAxis(EAxis::Z).Size();
diff --git a/Engine/Source/Runtime/Engine/Private/LODActor.cpp b/Engine/Source/Runtime/Engine/Private/LODActor.cpp
--- a/Engine/Source/Runtime/Engine/Private/LODActor.cpp
+++ b/Engine/Source/Runtime/Engine/Private/LODActor.cpp
@@ -10,13 +10,18 @@
 #include "MessageLog.h"
 #include "UObjectToken.h"
 
+#if WITH_EDITOR	
+#include "Components/DrawSphereComponent.h"
+#endif // WITH_EDITOR
+
 #define LOCTEXT_NAMESPACE "LODActor"
 
 ALODActor::ALODActor(const FObjectInitializer& ObjectInitializer)
 	: Super(ObjectInitializer)
-	, LODDrawDistance(5000)
+	, LODDrawDistance(5000)	
 {
 	bCanBeDamaged = false;
+	bIsPreviewActor = false;
 
 	StaticMeshComponent = CreateDefaultSubobject<UStaticMeshComponent>(TEXT("StaticMeshComponent0"));
 	StaticMeshComponent->SetCollisionProfileName(UCollisionProfile::NoCollision_ProfileName);
@@ -27,6 +32,11 @@ ALODActor::ALODActor(const FObjectInitializer& ObjectInitializer)
 	StaticMeshComponent->CastShadow = false;
 
 	RootComponent = StaticMeshComponent;
+	
+#if WITH_EDITORONLY_DATA	
+	DrawSphereComponent = CreateEditorOnlyDefaultSubobject<UDrawSphereComponent>(TEXT("VisualizeComponent0"));
+	DrawSphereComponent->SetSphereRadius(0.0f);
+#endif // WITH_EDITORONLY_DATA
 }
 
 FString ALODActor::GetDetailedInfoInternal() const
@@ -38,13 +48,25 @@ void ALODActor::PostRegisterAllComponents()
 {
 	Super::PostRegisterAllComponents();
 #if WITH_EDITOR
-	if (StaticMeshComponent->SceneProxy)
+	if (!bIsPreviewActor && StaticMeshComponent->SceneProxy)
 	{
 		ensure (LODLevel >= 1);
 		(StaticMeshComponent->SceneProxy)->SetHierarchicalLOD_GameThread(LODLevel);
 	}
 #endif
 }
+
+void ALODActor::SetStaticMesh(class UStaticMesh* InStaticMesh)
+{
+	if (StaticMeshComponent)
+	{
+		StaticMeshComponent->StaticMesh = InStaticMesh;
+	}
+
+	// If given NULL for static mesh then this is a preview actor (only rendering the bounds)
+	bIsPreviewActor = (InStaticMesh == nullptr);
+}
+
 #if WITH_EDITOR
 
 void ALODActor::PostEditChangeProperty(FPropertyChangedEvent& PropertyChangedEvent)
@@ -85,27 +107,31 @@ bool ALODActor::GetReferencedContentObjects( TArray<UObject*>& Objects ) const
 
 void ALODActor::CheckForErrors()
 {
-	Super::CheckForErrors();
-
 	FMessageLog MapCheck("MapCheck");
 
-	if( !StaticMeshComponent )
+	// Only check when this is not a preview actor and actually has a static mesh
+	if (!bIsPreviewActor)
 	{
-		MapCheck.Warning()
-			->AddToken(FUObjectToken::Create(this))
-			->AddToken(FTextToken::Create(LOCTEXT( "MapCheck_Message_StaticMeshComponent", "Static mesh actor has NULL StaticMeshComponent property - please delete" ) ))
-			->AddToken(FMapErrorToken::Create(FMapErrors::StaticMeshComponent));
-	}
+		Super::CheckForErrors();
+		if (!StaticMeshComponent)
+		{
+			MapCheck.Warning()
+				->AddToken(FUObjectToken::Create(this))
+				->AddToken(FTextToken::Create(LOCTEXT("MapCheck_Message_StaticMeshComponent", "Static mesh actor has NULL StaticMeshComponent property - please delete")))
+				->AddToken(FMapErrorToken::Create(FMapErrors::StaticMeshComponent));
+		}
 
-	if(StaticMeshComponent && StaticMeshComponent->StaticMesh == NULL)
-	{
-		FFormatNamedArguments Arguments;
-		Arguments.Add(TEXT("ActorName"), FText::FromString(GetName()));
-		FMessageLog("MapCheck").Error()
-			->AddToken(FUObjectToken::Create(this))
-			->AddToken(FTextToken::Create(FText::Format(LOCTEXT("MapCheck_Message_InvalidLODActorMissingMesh", "{ActorName} : Static mesh is missing for the built LODActor.  Did you remove the asset? Please delete it and build LOD again. "), Arguments)))
-			->AddToken(FMapErrorToken::Create(FMapErrors::LODActorMissingStaticMesh));
+		if (StaticMeshComponent && StaticMeshComponent->StaticMesh == NULL)
+		{
+			FFormatNamedArguments Arguments;
+			Arguments.Add(TEXT("ActorName"), FText::FromString(GetName()));
+			FMessageLog("MapCheck").Error()
+				->AddToken(FUObjectToken::Create(this))
+				->AddToken(FTextToken::Create(FText::Format(LOCTEXT("MapCheck_Message_InvalidLODActorMissingMesh", "{ActorName} : Static mesh is missing for the built LODActor.  Did you remove the asset? Please delete it and build LOD again. "), Arguments)))
+				->AddToken(FMapErrorToken::Create(FMapErrors::LODActorMissingStaticMesh));
+		}
 	}
+	
 
 	if (SubActors.Num() == 0)
 	{
diff --git a/Engine/Source/Runtime/Engine/Private/World.cpp b/Engine/Source/Runtime/Engine/Private/World.cpp
--- a/Engine/Source/Runtime/Engine/Private/World.cpp
+++ b/Engine/Source/Runtime/Engine/Private/World.cpp
@@ -52,6 +52,7 @@
 	#include "Editor/Kismet/Public/FindInBlueprintManager.h"
 	#include "Editor/UnrealEd/Classes/Editor/UnrealEdTypes.h"
 	#include "Editor/UnrealEd/Classes/Settings/LevelEditorPlaySettings.h"
+	#include "Editor/UnrealEd/Public/HierarchicalLOD.h"
 #endif
 
 #include "MallocProfiler.h"
@@ -114,9 +115,9 @@ FWorldDelegates::FRefreshLevelScriptActionsEvent FWorldDelegates::RefreshLevelSc
 #endif // WITH_EDITOR
 
 UWorld::UWorld( const FObjectInitializer& ObjectInitializer )
-:	UObject(ObjectInitializer)
+: UObject(ObjectInitializer)
 #if WITH_EDITOR
-,	HierarchicalLODBuilder(this)
+, HierarchicalLODBuilder(new FHierarchicalLODBuilder(this))
 #endif
 ,	FeatureLevel(GMaxRHIFeatureLevel)
 , URL(FURL(NULL))
@@ -141,7 +142,7 @@ UWorld::~UWorld()
 	while (AsyncPreRegisterLevelStreamingTasks.GetValue())
 	{
 		FPlatformProcess::Sleep(0.0f);
-	}
+	}	
 }
 
 void UWorld::Serialize( FArchive& Ar )
@@ -516,6 +517,11 @@ void UWorld::FinishDestroy()
 		delete TimerManager;
 	}
 
+	if (HierarchicalLODBuilder)
+	{
+		delete HierarchicalLODBuilder;
+	}
+
 	// Remove the PKG_ContainsMap flag from packages that no longer contain a world
 	{
 		UPackage* WorldPackage = GetOutermost();
