Changed another EDITORONLY define to EDITORONLY_DATA

diff --git a/Engine/Source/Editor/UnrealEd/Private/LODCluster.cpp b/Engine/Source/Editor/UnrealEd/Private/LODCluster.cpp
--- a/Engine/Source/Editor/UnrealEd/Private/LODCluster.cpp
+++ b/Engine/Source/Editor/UnrealEd/Private/LODCluster.cpp
@@ -139,9 +139,11 @@ FSphere FLODCluster::AddActor(AActor* NewActor)
 	ensure (Actors.Contains(NewActor) == false);
 	Actors.Add(NewActor);
 	FVector Origin, Extent;
+#ifdef WITH_EDITORONLY_DATA
 	if (NewActor->IsA<ALODActor>())
 	{
 		ALODActor* LODActor = CastChecked<ALODActor>(NewActor);
+
 		if (LODActor->IsPreviewActorInstance())
 		{
 			Origin = LODActor->GetDrawSphereComponent()->GetComponentLocation();
@@ -153,6 +155,7 @@ FSphere FLODCluster::AddActor(AActor* NewActor)
 		}
 	}
 	else
+#endif // WITH_EDITORONLY_DATA
 	{
 		NewActor->GetActorBounds(false, Origin, Extent);
 	}
@@ -370,7 +373,7 @@ void FLODCluster::BuildActor(ULevel* InLevel, const int32 LODIdx, const bool bCr
 					NewActor->SetStaticMesh( MainMesh );
 					NewActor->SetFolderPath("/HLODActors");											
 
-#if WITH_EDITOR						
+#if WITH_EDITORONLY_DATA						
 					// Setup sphere drawing visualization
 					UDrawSphereComponent* SphereComponent = NewActor->GetDrawSphereComponent();
 					SphereComponent->SetSphereRadius(Bound.W * METER_TO_CM);
@@ -379,7 +382,7 @@ void FLODCluster::BuildActor(ULevel* InLevel, const int32 LODIdx, const bool bCr
 					SphereComponent->MinDrawDistance = DrawDistance;
 					SphereComponent->LDMaxDrawDistance = NextDrawDistance;
 					SphereComponent->CachedMaxDrawDistance = NextDrawDistance;	
-#endif								
+#endif // WITH_EDITORONLY_DATA						
 					// now set as parent
 					for (auto& Actor : Actors)
 					{
