Changed how HLODVolume impacts HLOD cluster generation, now only has an influence on LOD0

diff --git a/Engine/Source/Editor/UnrealEd/Private/HierarchicalLOD.cpp b/Engine/Source/Editor/UnrealEd/Private/HierarchicalLOD.cpp
--- a/Engine/Source/Editor/UnrealEd/Private/HierarchicalLOD.cpp
+++ b/Engine/Source/Editor/UnrealEd/Private/HierarchicalLOD.cpp
@@ -416,10 +416,15 @@ void FHierarchicalLODBuilder::MergeClustersAndBuildActors(ULevel* InLevel, const
 			}
 		}
 
-		for (auto& Cluster : HLODVolumeClusters)
+
+		if (LODIdx == 0)
 		{
-			Clusters.Add(Cluster.Value);
+			for (auto& Cluster : HLODVolumeClusters)
+			{
+				Clusters.Add(Cluster.Value);
+			}
 		}
+		
 
 		// debug flag, so that I can just see clustered data since no visualization in the editor yet
 		//static bool bBuildActor=true;
