Added NoAll, YesAll and Cancel options to the "Failed to Merge Bones" dialog on skeletal mesh import

diff --git a/Engine/Source/Editor/UnrealEd/Private/Fbx/FbxFactory.cpp b/Engine/Source/Editor/UnrealEd/Private/Fbx/FbxFactory.cpp
--- a/Engine/Source/Editor/UnrealEd/Private/Fbx/FbxFactory.cpp
+++ b/Engine/Source/Editor/UnrealEd/Private/Fbx/FbxFactory.cpp
@@ -386,9 +386,18 @@ UObject* UFbxFactory::FactoryCreateBinary
 							{
 								FName OutputName = FbxImporter->MakeNameForMesh(Name.ToString(), SkelMeshNodeArray[0]);
 
-								USkeletalMesh* NewMesh = FbxImporter->ImportSkeletalMesh( InParent, SkelMeshNodeArray, OutputName, Flags, ImportUI->SkeletalMeshImportData );
+								USkeletalMesh* NewMesh = FbxImporter->ImportSkeletalMesh( InParent, SkelMeshNodeArray, OutputName, Flags, ImportUI->SkeletalMeshImportData, &bOperationCanceled );
 								NewObject = NewMesh;
 
+								if(bOperationCanceled)
+								{
+									// User cancelled, clean up and return
+									FbxImporter->ReleaseScene();
+									Warn->EndSlowTask();
+									bOperationCanceled = true;
+									return nullptr;
+								}
+
 								if ( NewMesh && ImportUI->bImportAnimations )
 								{
 									// We need to remove all scaling from the root node before we set up animation data.
@@ -404,8 +413,8 @@ UObject* UFbxFactory::FactoryCreateBinary
 							{
 								USkeletalMesh* BaseSkeletalMesh = Cast<USkeletalMesh>(NewObject);
 								FName LODObjectName = NAME_None;
-								USkeletalMesh *LODObject = FbxImporter->ImportSkeletalMesh( GetTransientPackage(), SkelMeshNodeArray, LODObjectName, RF_NoFlags, ImportUI->SkeletalMeshImportData );
-								bool bImportSucceeded = FbxImporter->ImportSkeletalMeshLOD(LODObject, BaseSkeletalMesh, LODIndex, false);
+								USkeletalMesh *LODObject = FbxImporter->ImportSkeletalMesh( GetTransientPackage(), SkelMeshNodeArray, LODObjectName, RF_NoFlags, ImportUI->SkeletalMeshImportData, &bOperationCanceled );
+								bool bImportSucceeded = !bOperationCanceled && FbxImporter->ImportSkeletalMeshLOD(LODObject, BaseSkeletalMesh, LODIndex, false);
 
 								if (bImportSucceeded)
 								{
diff --git a/Engine/Source/Editor/UnrealEd/Private/Fbx/FbxMainImport.cpp b/Engine/Source/Editor/UnrealEd/Private/Fbx/FbxMainImport.cpp
--- a/Engine/Source/Editor/UnrealEd/Private/Fbx/FbxMainImport.cpp
+++ b/Engine/Source/Editor/UnrealEd/Private/Fbx/FbxMainImport.cpp
@@ -370,6 +370,7 @@ void FFbxImporter::ReleaseScene()
 	CollisionModels.Clear();
 	CurPhase = NOTSTARTED;
 	bFirstMesh = true;
+	LastMergeBonesChoice = EAppReturnType::Ok;
 }
 
 FBXImportOptions* UnFbx::FFbxImporter::GetImportOptions() const
diff --git a/Engine/Source/Editor/UnrealEd/Private/Fbx/FbxSkeletalMeshImport.cpp b/Engine/Source/Editor/UnrealEd/Private/Fbx/FbxSkeletalMeshImport.cpp
--- a/Engine/Source/Editor/UnrealEd/Private/Fbx/FbxSkeletalMeshImport.cpp
+++ b/Engine/Source/Editor/UnrealEd/Private/Fbx/FbxSkeletalMeshImport.cpp
@@ -1058,7 +1058,7 @@ bool UnFbx::FFbxImporter::ImportBone(TArray<FbxNode*>& NodeArray, FSkeletalMeshI
 	return true;
 }
 
-USkeletalMesh* UnFbx::FFbxImporter::ImportSkeletalMesh(UObject* InParent, TArray<FbxNode*>& NodeArray, const FName& Name, EObjectFlags Flags, UFbxSkeletalMeshImportData* TemplateImportData, TArray<FbxShape*> *FbxShapeArray, FSkeletalMeshImportData* OutData, bool bCreateRenderData )
+USkeletalMesh* UnFbx::FFbxImporter::ImportSkeletalMesh(UObject* InParent, TArray<FbxNode*>& NodeArray, const FName& Name, EObjectFlags Flags, UFbxSkeletalMeshImportData* TemplateImportData, bool* bCancelOperation, TArray<FbxShape*> *FbxShapeArray, FSkeletalMeshImportData* OutData, bool bCreateRenderData )
 {
 	if (NodeArray.Num() == 0)
 	{
@@ -1403,14 +1403,33 @@ USkeletalMesh* UnFbx::FFbxImporter::ImportSkeletalMesh(UObject* InParent, TArray
 		// merge bones to the selected skeleton
 		if ( !Skeleton->MergeAllBonesToBoneTree( SkeletalMesh ) )
 		{
-			if ( EAppReturnType::Yes == FMessageDialog::Open( EAppMsgType::YesNo, 
-										LOCTEXT("SkeletonFailed_BoneMerge", "FAILED TO MERGE BONES:\n\n This could happen if significant hierarchical change has been made\n - i.e. inserting bone between nodes\n Would you like to regenerate Skeleton from this mesh? \n\n ***WARNING: THIS WILL REQUIRE RECOMPRESS ALL ANIMATION DATA AND POTENTIALLY INVALIDATE***\n") ) ) 
+			// We should only show the skeleton save toast once, not as many times as we have nodes to import
+			bool bToastSaveMessage = false;
+			if(bFirstMesh || (LastMergeBonesChoice != EAppReturnType::NoAll && LastMergeBonesChoice != EAppReturnType::YesAll))
+			{
+				LastMergeBonesChoice = FMessageDialog::Open(EAppMsgType::YesNoYesAllNoAllCancel,
+															LOCTEXT("SkeletonFailed_BoneMerge", "FAILED TO MERGE BONES:\n\n This could happen if significant hierarchical change has been made\n - i.e. inserting bone between nodes\n Would you like to regenerate Skeleton from this mesh? \n\n ***WARNING: THIS WILL REQUIRE RECOMPRESS ALL ANIMATION DATA AND POTENTIALLY INVALIDATE***\n"));
+				bToastSaveMessage = true;
+			}
+			
+			if(LastMergeBonesChoice == EAppReturnType::Cancel)
+			{
+				// User wants to cancel further importing
+				if(bCancelOperation != nullptr)
 				{
-					if ( Skeleton->RecreateBoneTree( SkeletalMesh ) )
-					{
-						FAssetNotifications::SkeletonNeedsToBeSaved(Skeleton);
-					}
+					// Report back to calling code if we have a flag to set
+					*bCancelOperation = true;
 				}
+				return nullptr;
+			}
+
+			if (LastMergeBonesChoice == EAppReturnType::Yes || LastMergeBonesChoice == EAppReturnType::YesAll) 
+			{
+				if ( Skeleton->RecreateBoneTree( SkeletalMesh ) && bToastSaveMessage)
+				{
+					FAssetNotifications::SkeletonNeedsToBeSaved(Skeleton);
+				}
+			}
 		}
 		else
 		{
@@ -2801,8 +2820,8 @@ void UnFbx::FFbxImporter::ImportMorphTargetsInternal( TArray<FbxNode*>& SkelMesh
 			else
 			{
 				AddTokenizedErrorMessage(FTokenizedMessage::Create(EMessageSeverity::Error, FText::Format(FText::FromString("Could not find the {0} morphtarget for LOD {1}. \
-																																																																																											Make sure the name for morphtarget matches with LOD 0"), FText::FromString(ShapeName), FText::FromString(FString::FromInt(LODIndex)))),
-																																																																																											FFbxErrors::SkeletalMesh_LOD_MissingMorphTarget);
+					Make sure the name for morphtarget matches with LOD 0"), FText::FromString(ShapeName), FText::FromString(FString::FromInt(LODIndex)))),
+					FFbxErrors::SkeletalMesh_LOD_MissingMorphTarget);
 			}
 		}
 
@@ -2811,7 +2830,7 @@ void UnFbx::FFbxImporter::ImportMorphTargetsInternal( TArray<FbxNode*>& SkelMesh
 			// now we get a shape for whole mesh, import to unreal as a morph target
 			// @todo AssetImportData do we need import data for this temp mesh?
 			UFbxSkeletalMeshImportData* TmpMeshImportData = NULL;
-			USkeletalMesh* TmpSkeletalMesh = (USkeletalMesh*)ImportSkeletalMesh(GetTransientPackage(), SkelMeshNodeArray, NAME_None, (EObjectFlags)0, TmpMeshImportData, &ShapeArray, &ImportData, false);
+			USkeletalMesh* TmpSkeletalMesh = (USkeletalMesh*)ImportSkeletalMesh(GetTransientPackage(), SkelMeshNodeArray, NAME_None, (EObjectFlags)0, TmpMeshImportData, nullptr, &ShapeArray, &ImportData, false);
 			TempMeshes.Add(TmpSkeletalMesh);
 			MorphTargets.Add(Result);
 
diff --git a/Engine/Source/Editor/UnrealEd/Private/FbxImporter.h b/Engine/Source/Editor/UnrealEd/Private/FbxImporter.h
--- a/Engine/Source/Editor/UnrealEd/Private/FbxImporter.h
+++ b/Engine/Source/Editor/UnrealEd/Private/FbxImporter.h
@@ -392,7 +392,7 @@ public:
 	 *
 	 * @return The USkeletalMesh object created
 	 */
-	USkeletalMesh* ImportSkeletalMesh(UObject* InParent, TArray<FbxNode*>& NodeArray, const FName& Name, EObjectFlags Flags, UFbxSkeletalMeshImportData* TemplateImportData, TArray<FbxShape*> *FbxShapeArray=NULL, FSkeletalMeshImportData* OutData=NULL, bool bCreateRenderData = true );
+	USkeletalMesh* ImportSkeletalMesh(UObject* InParent, TArray<FbxNode*>& NodeArray, const FName& Name, EObjectFlags Flags, UFbxSkeletalMeshImportData* TemplateImportData, bool* bCancelOperation = nullptr, TArray<FbxShape*> *FbxShapeArray = nullptr, FSkeletalMeshImportData* OutData = nullptr, bool bCreateRenderData = true );
 
 	/**
 	 * Add to the animation set, the animations contained within the FBX scene, for the given skeletal mesh
@@ -680,6 +680,10 @@ protected:
 	// Initialized as true when start to import a FBX scene
 	bool bFirstMesh;
 	
+	// Set when importing skeletal meshes if the merge bones step fails. Used to track
+	// YesToAll and NoToAll for an entire scene
+	EAppReturnType::Type LastMergeBonesChoice;
+
 	/**
 	 * Collision model list. The key is fbx node name
 	 * If there is an collision model with old name format, the key is empty string("").
diff --git a/Engine/Source/Runtime/Engine/Private/Animation/Skeleton.cpp b/Engine/Source/Runtime/Engine/Private/Animation/Skeleton.cpp
--- a/Engine/Source/Runtime/Engine/Private/Animation/Skeleton.cpp
+++ b/Engine/Source/Runtime/Engine/Private/Animation/Skeleton.cpp
@@ -358,6 +358,12 @@ int32 USkeleton::BuildLinkup(const USkeletalMesh* InSkelMesh)
 	NewMeshLinkup.MeshToSkeletonTable.Empty(NumMeshBones);
 	NewMeshLinkup.MeshToSkeletonTable.AddUninitialized(NumMeshBones);
 
+#if WITH_EDITOR
+	// The message below can potentially fire many times if the skeleton has been put into a state where it is no longer
+	// fully compatible with the mesh we're trying to merge into the bone tree. We use this flag to only show it once per mesh
+	bool bDismissedMessage = false;
+#endif
+
 	for (int32 MeshBoneIndex = 0; MeshBoneIndex < NumMeshBones; MeshBoneIndex++)
 	{
 		const FName MeshBoneName = MeshRefSkel.GetBoneName(MeshBoneIndex);
@@ -368,7 +374,11 @@ int32 USkeleton::BuildLinkup(const USkeletalMesh* InSkelMesh)
 		// not currently supported in-game.
 		if (SkeletonBoneIndex == INDEX_NONE)
 		{
-			FMessageDialog::Open(EAppMsgType::Ok, FText::Format(LOCTEXT("SkeletonBuildLinkupMissingBones", "The Skeleton {0}, is missing bones that SkeletalMesh {1} needs. They will be added now. Please save the Skeleton!"), FText::FromString(GetNameSafe(this)), FText::FromString(GetNameSafe(InSkelMesh))));
+			if(!bDismissedMessage)
+			{
+				FMessageDialog::Open(EAppMsgType::Ok, FText::Format(LOCTEXT("SkeletonBuildLinkupMissingBones", "The Skeleton {0}, is missing bones that SkeletalMesh {1} needs. They will be added now. Please save the Skeleton!"), FText::FromString(GetNameSafe(this)), FText::FromString(GetNameSafe(InSkelMesh))));
+				bDismissedMessage = true;
+			}
 
 			// Re-add all SkelMesh bones to the Skeleton.
 			MergeAllBonesToBoneTree(InSkelMesh);
