Changed the display name of bAutomaticallyHotReloadNewClasses in EditorPerProjectUserSettings to convey that it implies a compile.

diff --git a/Engine/Source/Editor/UnrealEd/Classes/Editor/EditorPerProjectUserSettings.h b/Engine/Source/Editor/UnrealEd/Classes/Editor/EditorPerProjectUserSettings.h
--- a/Engine/Source/Editor/UnrealEd/Classes/Editor/EditorPerProjectUserSettings.h
+++ b/Engine/Source/Editor/UnrealEd/Classes/Editor/EditorPerProjectUserSettings.h
@@ -47,8 +47,8 @@ class UEditorPerProjectUserSettings : public UObject
 	UPROPERTY(EditAnywhere, config, Category=Performance)
 	uint32 bMonitorEditorPerformance:1;
 
-	/** If enabled, any newly added classes will trigger a hot-reload of the module they were added to */
-	UPROPERTY(EditAnywhere, config, Category=HotReload)
+	/** If enabled, any newly added classes will be automatically compiled and trigger a hot-reload of the module they were added to */
+	UPROPERTY(EditAnywhere, config, Category=HotReload, meta=(DisplayName="Automatically Compile Newly Added C++ Classes"))
 	uint32 bAutomaticallyHotReloadNewClasses:1;
 
 	/** If enabled, export level with attachment hierarchy set */
