Forgot to submit this file for default in phat

diff --git a/Engine/Source/Editor/UnrealEd/Private/PreferenceStubs.cpp b/Engine/Source/Editor/UnrealEd/Private/PreferenceStubs.cpp
--- a/Engine/Source/Editor/UnrealEd/Private/PreferenceStubs.cpp
+++ b/Engine/Source/Editor/UnrealEd/Private/PreferenceStubs.cpp
@@ -21,6 +21,7 @@ UPhATSimOptions::UPhATSimOptions(const FObjectInitializer& ObjectInitializer)
 	: Super(ObjectInitializer)
 {
 	PhysicsBlend = 1.0f;
+	bUpdateJointsFromAnimation = false;
 	MaxFPS = -1;
 	TimeDilation = 1.f;
 }
