Get around some VS2015 ICEs

diff --git a/Engine/Source/Editor/UnrealEd/Private/AutoReimport/AutoReimportTests.cpp b/Engine/Source/Editor/UnrealEd/Private/AutoReimport/AutoReimportTests.cpp
--- a/Engine/Source/Editor/UnrealEd/Private/AutoReimport/AutoReimportTests.cpp
+++ b/Engine/Source/Editor/UnrealEd/Private/AutoReimport/AutoReimportTests.cpp
@@ -159,7 +159,7 @@ bool FAutoReimportSimpleCreateTest::RunTest(const FString& Parameters)
 				UE_LOG(LogAutoReimportTests, Error, TEXT("Incorrect number of changes reported (%d != 1)."), Changes.Num());
 			}
 			// Check that we got a relative path
-			else if (Changes[0].Filename.Get() != Filename)
+			else if (!Changes[0].Filename.Get().Equals(Filename))
 			{
 				UE_LOG(LogAutoReimportTests, Error, TEXT("Path reported incorrectly (%s != %s)."), *Changes[0].Filename.Get(), Filename);
 			}
@@ -228,7 +228,7 @@ bool FAutoReimportSimpleModifyTest::RunTest(const FString& Parameters)
 				UE_LOG(LogAutoReimportTests, Error, TEXT("Incorrect number of changes reported (%d != 1)."), Changes.Num());
 			}
 			// Check that we got a relative path
-			else if (Changes[0].Filename.Get() != DstFilename)
+			else if (!Changes[0].Filename.Get().Equals(DstFilename))
 			{
 				UE_LOG(LogAutoReimportTests, Error, TEXT("Path reported incorrectly (%s != %s)."), *Changes[0].Filename.Get(), DstFilename);
 			}
@@ -296,7 +296,7 @@ bool FAutoReimportSimpleDeleteTest::RunTest(const FString& Parameters)
 				UE_LOG(LogAutoReimportTests, Error, TEXT("Incorrect number of changes reported (%d != 1)."), Changes.Num());
 			}
 			// Check that we got a relative path
-			else if (Changes[0].Filename.Get() != Filename)
+			else if (!Changes[0].Filename.Get().Equals(Filename))
 			{
 				UE_LOG(LogAutoReimportTests, Error, TEXT("Path reported incorrectly (%s != %s)."), *Changes[0].Filename.Get(), Filename);
 			}
@@ -362,11 +362,11 @@ bool FAutoReimportSimpleRenameTest::RunTest(const FString& Parameters)
 			{
 				UE_LOG(LogAutoReimportTests, Error, TEXT("Incorrect action for renamed file %s -> %s."), SrcFilename, DstFilename);
 			}
-			else if (Changes[0].Filename.Get() != DstFilename)
+			else if (!Changes[0].Filename.Get().Equals(DstFilename))
 			{
 				UE_LOG(LogAutoReimportTests, Error, TEXT("Source path reported incorrectly (%s != %s)."), *Changes[0].Filename.Get(), DstFilename);
 			}
-			else if (Changes[0].MovedFromFilename.Get() != SrcFilename)
+			else if (!Changes[0].MovedFromFilename.Get().Equals(SrcFilename))
 			{
 				UE_LOG(LogAutoReimportTests, Error, TEXT("Destination path reported incorrectly (%s != %s)."), *Changes[0].MovedFromFilename.Get(), SrcFilename);
 			}
@@ -432,7 +432,7 @@ bool FAutoReimportSimpleMoveExternallyTest::RunTest(const FString& Parameters)
 			{
 				UE_LOG(LogAutoReimportTests, Error, TEXT("Incorrect action for removed file %s -> %s."), SrcFilename, DstFilename);
 			}
-			else if (Changes[0].Filename.Get() != SrcFilename)
+			else if (!Changes[0].Filename.Get().Equals(SrcFilename))
 			{
 				UE_LOG(LogAutoReimportTests, Error, TEXT("Source path reported incorrectly (%s != %s)."), *Changes[0].Filename.Get(), SrcFilename);
 			}
@@ -779,4 +779,4 @@ bool FAutoReimportWildcardFiltersTest::RunTest(const FString& Parameters)
 	});
 
 	return true;
-}
\ No newline at end of file
+}
