Removed some verbose cooking output.

diff --git a/Engine/Source/Editor/UnrealEd/Private/CookOnTheFlyServer.cpp b/Engine/Source/Editor/UnrealEd/Private/CookOnTheFlyServer.cpp
--- a/Engine/Source/Editor/UnrealEd/Private/CookOnTheFlyServer.cpp
+++ b/Engine/Source/Editor/UnrealEd/Private/CookOnTheFlyServer.cpp
@@ -1358,7 +1358,6 @@ uint32 UCookOnTheFlyServer::TickCookOnTheSide( const float TimeSlice, uint32 &Co
 
 #if OUTPUT_TIMING
 		//FScopeTimer PackageManualTimer( ToBuild.GetFilename().ToString(), false );
-		UE_LOG(LogCook, Display,  TEXT("ProcessingPackage %s"), *ToBuild.GetFilename().ToString() );
 #endif
 
 		for ( const auto &PlatformName : TargetPlatformNames )
