Fixed crash when dragging junk text from an external program into the graph editor.

diff --git a/Engine/Source/Editor/UnrealEd/Private/AssetSelection.cpp b/Engine/Source/Editor/UnrealEd/Private/AssetSelection.cpp
--- a/Engine/Source/Editor/UnrealEd/Private/AssetSelection.cpp
+++ b/Engine/Source/Editor/UnrealEd/Private/AssetSelection.cpp
@@ -511,7 +511,9 @@ namespace AssetUtil
 
 				for (int Index = 0; Index < DroppedAssetStrings.Num(); Index++)
 				{
-					FAssetData AssetData = AssetRegistry.GetAssetByObjectPath( *DroppedAssetStrings[ Index ] );
+					// Truncate each string so that it doesn't exceed the maximum allowed length of characters to be converted to an FName
+					FString TruncatedString = DroppedAssetStrings[ Index ].Left( NAME_SIZE );
+					FAssetData AssetData = AssetRegistry.GetAssetByObjectPath( FName( *TruncatedString ) );
 					if ( AssetData.IsValid() )
 					{
 						DroppedAssetData.Add( AssetData );
