Projects can customize default events for newly created Blueprints via ini settings.

Callbacks can be setup to manually add other default nodes to a Blueprint.

Individual modules can also setup, in code, default events using FKismetEditorUtilities::RegisterAutoGeneratedDefaultEvent

diff --git a/Engine/Config/BaseEditorPerProjectUserSettings.ini b/Engine/Config/BaseEditorPerProjectUserSettings.ini
--- a/Engine/Config/BaseEditorPerProjectUserSettings.ini
+++ b/Engine/Config/BaseEditorPerProjectUserSettings.ini
@@ -673,6 +673,17 @@ DefaultBrushRadius=128
 +Node=(Class=DoOnce Key=O Shift=false Ctrl=false Alt=false)
 +Node=(Class=KismetArrayLibrary:Array_Get Key=A Shift=false Ctrl=false Alt=false)
 
+[DefaultEventNodes]
++Node=(TargetClass=Actor TargetEvent="ReceiveBeginPlay")
++Node=(TargetClass=Actor TargetEvent="ReceiveActorBeginOverlap")
++Node=(TargetClass=Actor TargetEvent="ReceiveTick")
++Node=(TargetClass=ActorComponent TargetEvent="ReceiveBeginPlay")
++Node=(TargetClass=ActorComponent TargetEvent="ReceiveTick")
++Node=(TargetClass=GameplayAbility TargetEvent="K2_ActivateAbility")
++Node=(TargetClass=GameplayAbility TargetEvent="K2_OnEndAbility")
++Node=(TargetClass=UserWidget TargetEvent="Construct")
++Node=(TargetClass=UserWidget TargetEvent="Tick")
++Node=(TargetClass=AnimInstance TargetEvent="BlueprintUpdateAnimation")
 
 [MaterialEditorSpawnNodes]
 +Node=(Class=MaterialExpressionAdd Key=A Shift=false Ctrl=false Alt=false)
diff --git a/Engine/Source/Editor/Kismet/Private/BlueprintEditorModule.cpp b/Engine/Source/Editor/Kismet/Private/BlueprintEditorModule.cpp
--- a/Engine/Source/Editor/Kismet/Private/BlueprintEditorModule.cpp
+++ b/Engine/Source/Editor/Kismet/Private/BlueprintEditorModule.cpp
@@ -111,11 +111,15 @@ void FBlueprintEditorModule::StartupModule()
 
 	TSharedPtr<FBlueprintGraphPanelPinFactory> BlueprintGraphPanelPinFactory = MakeShareable(new FBlueprintGraphPanelPinFactory());
 	FEdGraphUtilities::RegisterVisualPinFactory(BlueprintGraphPanelPinFactory);
-}
 
+	PrepareAutoGeneratedDefaultEvents();
+}
 
 void FBlueprintEditorModule::ShutdownModule()
 {
+	// Cleanup all information for auto generated default event nodes by this module
+	FKismetEditorUtilities::UnregisterAutoBlueprintNodeCreation(this);
+
 	SharedBlueprintEditorCommands.Reset();
 	MenuExtensibilityManager.Reset();
 
@@ -204,4 +208,39 @@ void FBlueprintEditorModule::UnregisterSCSEditorCustomization(const FName& InCom
 	SCSEditorCustomizations.Remove(InComponentName);
 }
 
+void FBlueprintEditorModule::PrepareAutoGeneratedDefaultEvents()
+{
+	// Load up all default events that should be spawned for Blueprints that are children of specific classes
+	const FString ConfigSection = TEXT("DefaultEventNodes");
+	const FString SettingName = TEXT("Node");
+	TArray< FString > NodeSpawns;
+	GConfig->GetArray(*ConfigSection, *SettingName, NodeSpawns, GEditorPerProjectIni);
+
+	for(FString CurrentNodeSpawn : NodeSpawns)
+	{
+		FString TargetClassName;
+		if(!FParse::Value(*CurrentNodeSpawn, TEXT("TargetClass="), TargetClassName))
+		{
+			// Could not find a class name, cannot continue with this line
+			continue;
+		}
+
+		UClass* FoundTargetClass = FindObject<UClass>(ANY_PACKAGE, *TargetClassName, true);
+		if(FoundTargetClass)
+		{
+			FString TargetEventFunction;
+			if(!FParse::Value(*CurrentNodeSpawn, TEXT("TargetEvent="), TargetEventFunction))
+			{
+				// Could not find a class name, cannot continue with this line
+				continue;
+			}
+
+			if (FindObject<UFunction>(FoundTargetClass, *TargetEventFunction))
+			{
+				FKismetEditorUtilities::RegisterAutoGeneratedDefaultEvent(this, FoundTargetClass, FName(*TargetEventFunction));
+			}
+		}
+	}
+}
+
 #undef LOCTEXT_NAMESPACE
diff --git a/Engine/Source/Editor/Kismet/Public/BlueprintEditorModule.h b/Engine/Source/Editor/Kismet/Public/BlueprintEditorModule.h
--- a/Engine/Source/Editor/Kismet/Public/BlueprintEditorModule.h
+++ b/Engine/Source/Editor/Kismet/Public/BlueprintEditorModule.h
@@ -148,6 +148,10 @@ public:
 	 */
 	virtual const TSharedRef<FUICommandList> GetsSharedBlueprintEditorCommands() const { return SharedBlueprintEditorCommands.ToSharedRef(); }
 
+private:
+	/** Loads from ini a list of all events that should be auto created for Blueprints of a specific class */
+	void PrepareAutoGeneratedDefaultEvents();
+
 private:
 	TSharedPtr<FExtensibilityManager> MenuExtensibilityManager;
 
diff --git a/Engine/Source/Editor/Persona/Private/PersonaModule.cpp b/Engine/Source/Editor/Persona/Private/PersonaModule.cpp
--- a/Engine/Source/Editor/Persona/Private/PersonaModule.cpp
+++ b/Engine/Source/Editor/Persona/Private/PersonaModule.cpp
@@ -13,6 +13,7 @@
 #include "AnimNotifyDetails.h"
 #include "AnimGraphNodeDetails.h"
 #include "AnimInstanceDetails.h"
+#include "Editor/UnrealEd/Public/Kismet2/KismetEditorUtilities.h"
 
 IMPLEMENT_MODULE( FPersonaModule, Persona );
 
@@ -55,10 +56,14 @@ void FPersonaModule::StartupModule()
 		PropertyModule.RegisterCustomPropertyTypeLayout( "InputScaleBias", FOnGetPropertyTypeCustomizationInstance::CreateStatic( &FInputScaleBiasCustomization::MakeInstance ) );
 		PropertyModule.RegisterCustomPropertyTypeLayout( "BoneReference", FOnGetPropertyTypeCustomizationInstance::CreateStatic( &FBoneReferenceCustomization::MakeInstance ) );
 	}
+
+	FKismetEditorUtilities::RegisterOnBlueprintCreatedCallback(this, UAnimInstance::StaticClass(), FKismetEditorUtilities::FOnBlueprintCreated::CreateRaw(this, &FPersonaModule::OnNewBlueprintCreated));
 }
 
 void FPersonaModule::ShutdownModule()
 {
+	FKismetEditorUtilities::UnregisterAutoBlueprintNodeCreation(this);
+
 	MenuExtensibilityManager.Reset();
 	ToolBarExtensibilityManager.Reset();
 
@@ -82,3 +87,35 @@ TSharedRef<IBlueprintEditor> FPersonaModule::CreatePersona( const EToolkitMode::
 	NewPersona->InitPersona( Mode, InitToolkitHost, Skeleton, Blueprint, AnimationAsset, Mesh );
 	return NewPersona;
 }
+
+void FPersonaModule::OnNewBlueprintCreated(UBlueprint* InBlueprint)
+{
+	if (ensure(InBlueprint->UbergraphPages.Num() > 0))
+	{
+		UEdGraph* EventGraph = InBlueprint->UbergraphPages[0];
+
+		int32 SafeXPosition = 0;
+		int32 SafeYPosition = 0;
+
+		if (EventGraph->Nodes.Num() != 0)
+		{
+			SafeXPosition = EventGraph->Nodes[0]->NodePosX;
+			SafeYPosition = EventGraph->Nodes[EventGraph->Nodes.Num() - 1]->NodePosY + EventGraph->Nodes[EventGraph->Nodes.Num() - 1]->NodeHeight + 100;
+		}
+
+		// add try get owner node
+		UK2Node_CallFunction* GetOwnerNode = NewObject<UK2Node_CallFunction>(EventGraph);
+		UFunction* MakeNodeFunction = UAnimInstance::StaticClass()->FindFunctionByName(GET_FUNCTION_NAME_CHECKED(UAnimInstance, TryGetPawnOwner));
+		GetOwnerNode->CreateNewGuid();
+		GetOwnerNode->PostPlacedNewNode();
+		GetOwnerNode->SetFromFunction(MakeNodeFunction);
+		GetOwnerNode->SetFlags(RF_Transactional);
+		GetOwnerNode->AllocateDefaultPins();
+		GetOwnerNode->NodePosX = SafeXPosition;
+		GetOwnerNode->NodePosY = SafeYPosition;
+		UEdGraphSchema_K2::SetNodeMetaData(GetOwnerNode, FNodeMetadata::DefaultGraphNode);
+		GetOwnerNode->bIsNodeEnabled = false;
+
+		EventGraph->AddNode(GetOwnerNode);
+	}
+}
diff --git a/Engine/Source/Editor/Persona/Public/PersonaModule.h b/Engine/Source/Editor/Persona/Public/PersonaModule.h
--- a/Engine/Source/Editor/Persona/Public/PersonaModule.h
+++ b/Engine/Source/Editor/Persona/Public/PersonaModule.h
@@ -48,6 +48,11 @@ public:
 	/** Gets the extensibility managers for outside entities to extend persona editor's menus and toolbars */
 	virtual TSharedPtr<FExtensibilityManager> GetMenuExtensibilityManager() {return MenuExtensibilityManager;}
 	virtual TSharedPtr<FExtensibilityManager> GetToolBarExtensibilityManager() {return ToolBarExtensibilityManager;}
+
+private:
+	/** When a new AnimBlueprint is created, this will handle post creation work such as adding non-event default nodes */
+	void OnNewBlueprintCreated(UBlueprint* InBlueprint);
+
 private:
 	TSharedPtr<FExtensibilityManager> MenuExtensibilityManager;
 	TSharedPtr<FExtensibilityManager> ToolBarExtensibilityManager;
diff --git a/Engine/Source/Editor/UnrealEd/Private/Kismet2/Kismet2.cpp b/Engine/Source/Editor/UnrealEd/Private/Kismet2/Kismet2.cpp
--- a/Engine/Source/Editor/UnrealEd/Private/Kismet2/Kismet2.cpp
+++ b/Engine/Source/Editor/UnrealEd/Private/Kismet2/Kismet2.cpp
@@ -354,6 +354,8 @@ void FBlueprintUnloader::ReplaceStaleRefs(UBlueprint* NewBlueprint)
 // Static variable definition
 TArray<FString> FKismetEditorUtilities::TrackedBlueprintParentList;
 FKismetEditorUtilities::FOnBlueprintUnloaded FKismetEditorUtilities::OnBlueprintUnloaded;
+TMultiMap<void*, FKismetEditorUtilities::FDefaultEventNodeData> FKismetEditorUtilities::AutoGeneratedDefaultEventsMap;
+TMultiMap<void*, FKismetEditorUtilities::FOnBlueprintCreatedData> FKismetEditorUtilities::OnBlueprintCreatedCallbacks;
 
 /** Create the correct event graphs for this blueprint */
 void FKismetEditorUtilities::CreateDefaultEventGraphs(UBlueprint* Blueprint)
@@ -474,67 +476,36 @@ UBlueprint* FKismetEditorUtilities::CreateBlueprint(UClass* ParentClass, UObject
 	UBlueprintEditorSettings* Settings = GetMutableDefault<UBlueprintEditorSettings>();
 	if(Settings && Settings->bSpawnDefaultBlueprintNodes)
 	{
-		//@TODO: ANIMREFACTOR 2: This kind of code should be on a per-blueprint basis; not centralized here
-		if(UAnimBlueprint* AnimBP = Cast<UAnimBlueprint>(NewBP))
+		// Only add default events if there is an ubergraph and they are supported
+		if(NewBP->UbergraphPages.Num() && FBlueprintEditorUtils::DoesSupportEventGraphs(NewBP))
 		{
-			// add default nodes to event graph
-			UEdGraph* Graph = NewBP->UbergraphPages[0];
+			// Based on the Blueprint type we are constructing, place some starting events.
+			// Note, this cannot happen in the Factories for constructing these Blueprint types due to the fact that creating child BPs circumvent the factories
+			UClass* WidgetClass = FindObject<UClass>(ANY_PACKAGE, TEXT("UserWidget"));
+			UClass* GameplayAbilityClass = FindObject<UClass>(ANY_PACKAGE, TEXT("GameplayAbility"));
 
-			if(Graph->Nodes.Num() == 0)
-			{
-				// add update event graph
-				int32 NodePositionY = 0;
-				UK2Node_Event* EventNode = FKismetEditorUtilities::AddDefaultEventNode(NewBP, Graph, FName(TEXT("BlueprintUpdateAnimation")), UAnimInstance::StaticClass(), NodePositionY);
-				check(EventNode);
-
-				// add try get owner node
-				UK2Node_CallFunction* GetOwnerNode = NewObject<UK2Node_CallFunction>(Graph);
-				UFunction* MakeNodeFunction = UAnimInstance::StaticClass()->FindFunctionByName(GET_FUNCTION_NAME_CHECKED(UAnimInstance, TryGetPawnOwner));
-				GetOwnerNode->CreateNewGuid();
-				GetOwnerNode->PostPlacedNewNode();
-				GetOwnerNode->SetFromFunction(MakeNodeFunction);
-				GetOwnerNode->SetFlags(RF_Transactional);
-				GetOwnerNode->AllocateDefaultPins();
-				GetOwnerNode->NodePosX = EventNode->NodePosX;
-				GetOwnerNode->NodePosY = EventNode->NodePosY + EventNode->NodeHeight + 100;
-				UEdGraphSchema_K2::SetNodeMetaData(GetOwnerNode, FNodeMetadata::DefaultGraphNode);
-				GetOwnerNode->bIsNodeEnabled = false;
-
-				Graph->AddNode(GetOwnerNode);
-			}
-		}
-		else
-		{
-			// Only add default events if there is an ubergraph and they are supported
-			if(NewBP->UbergraphPages.Num() && FBlueprintEditorUtils::DoesSupportEventGraphs(NewBP))
-			{
-				// Based on the Blueprint type we are constructing, place some starting events.
-				// Note, this cannot happen in the Factories for constructing these Blueprint types due to the fact that creating child BPs circumvent the factories
-				UClass* WidgetClass = FindObject<UClass>(ANY_PACKAGE, TEXT("UserWidget"));
-				UClass* GameplayAbilityClass = FindObject<UClass>(ANY_PACKAGE, TEXT("GameplayAbility"));
+			TArray<FName> AutoSpawnedEventNames;
+			int32 NodePositionY = 0;
 
-				int32 NodePositionY = 0;
-
-				if(NewBP->GeneratedClass->IsChildOf(AActor::StaticClass()))
-				{
-					FKismetEditorUtilities::AddDefaultEventNode(NewBP, NewBP->UbergraphPages[0], FName(TEXT("ReceiveBeginPlay")), AActor::StaticClass(), NodePositionY);
-					FKismetEditorUtilities::AddDefaultEventNode(NewBP, NewBP->UbergraphPages[0], FName(TEXT("ReceiveActorBeginOverlap")), AActor::StaticClass(), NodePositionY);
-					FKismetEditorUtilities::AddDefaultEventNode(NewBP, NewBP->UbergraphPages[0], FName(TEXT("ReceiveTick")), AActor::StaticClass(), NodePositionY);
-				}
-				else if(NewBP->GeneratedClass->IsChildOf(UActorComponent::StaticClass()))
-				{
-					FKismetEditorUtilities::AddDefaultEventNode(NewBP, NewBP->UbergraphPages[0], FName(TEXT("ReceiveTick")), UActorComponent::StaticClass(), NodePositionY);
-					FKismetEditorUtilities::AddDefaultEventNode(NewBP, NewBP->UbergraphPages[0], FName(TEXT("ReceiveBeginPlay")), UActorComponent::StaticClass(), NodePositionY);
-				}
-				else if(NewBP->GeneratedClass->IsChildOf(WidgetClass))
+			// Spawn any defined auto generated default events for the class
+			for (TMultiMap<void*, FDefaultEventNodeData>::TIterator DataIt(AutoGeneratedDefaultEventsMap); DataIt; ++DataIt)
+			{
+				FDefaultEventNodeData Data = DataIt.Value();
+				if (NewBP->GeneratedClass->IsChildOf(Data.TargetClass))
 				{
-					FKismetEditorUtilities::AddDefaultEventNode(NewBP, NewBP->UbergraphPages[0], FName(TEXT("Construct")), WidgetClass, NodePositionY);
-					FKismetEditorUtilities::AddDefaultEventNode(NewBP, NewBP->UbergraphPages[0], FName(TEXT("Tick")), WidgetClass, NodePositionY);
+					FKismetEditorUtilities::AddDefaultEventNode(NewBP, NewBP->UbergraphPages[0], Data.EventName, NewBP->GeneratedClass, NodePositionY);
 				}
-				else if(NewBP->GeneratedClass->IsChildOf(GameplayAbilityClass))
+			}
+
+			// Give anyone who wants to do more advanced BP modification post-creation a chance to do so.
+			// Anim Blueprints, for example, adds a non-event node to the main ubergraph.
+			for (TMultiMap<void*, FKismetEditorUtilities::FOnBlueprintCreatedData>::TIterator DataIt(OnBlueprintCreatedCallbacks); DataIt; ++DataIt)
+			{
+				FOnBlueprintCreatedData Data = DataIt.Value();
+				if (NewBP->GeneratedClass->IsChildOf(Data.TargetClass))
 				{
-					FKismetEditorUtilities::AddDefaultEventNode(NewBP, NewBP->UbergraphPages[0], FName(TEXT("K2_ActivateAbility")), GameplayAbilityClass, NodePositionY);
-					FKismetEditorUtilities::AddDefaultEventNode(NewBP, NewBP->UbergraphPages[0], FName(TEXT("K2_OnEndAbility")), GameplayAbilityClass, NodePositionY);
+					FKismetEditorUtilities::FOnBlueprintCreated BlueprintCreatedDelegate = Data.OnBlueprintCreated;
+					BlueprintCreatedDelegate.Execute(NewBP);
 				}
 			}
 		}
@@ -569,6 +540,28 @@ UBlueprint* FKismetEditorUtilities::CreateBlueprint(UClass* ParentClass, UObject
 	return NewBP;
 }
 
+void FKismetEditorUtilities::RegisterAutoGeneratedDefaultEvent(void* InOwner, UClass* InTargetClass, FName InEventName)
+{
+	FDefaultEventNodeData Data;
+	Data.TargetClass = InTargetClass;
+	Data.EventName = InEventName;
+	AutoGeneratedDefaultEventsMap.Add(InOwner, Data);
+}
+
+void FKismetEditorUtilities::RegisterOnBlueprintCreatedCallback(void* InOwner, UClass* InTargetClass, FOnBlueprintCreated InOnBlueprintCreatedCallback)
+{
+	FOnBlueprintCreatedData Data;
+	Data.TargetClass = InTargetClass;
+	Data.OnBlueprintCreated = InOnBlueprintCreatedCallback;
+	OnBlueprintCreatedCallbacks.Add(InOwner, Data);
+}
+
+void FKismetEditorUtilities::UnregisterAutoBlueprintNodeCreation(void* InOwner)
+{
+	AutoGeneratedDefaultEventsMap.Remove(InOwner);
+	OnBlueprintCreatedCallbacks.Remove(InOwner);
+}
+
 UK2Node_Event* FKismetEditorUtilities::AddDefaultEventNode(UBlueprint* InBlueprint, UEdGraph* InGraph, FName InEventName, UClass* InEventClass, int32& InOutNodePosY)
 {
 	UK2Node_Event* NewEventNode = nullptr;
diff --git a/Engine/Source/Editor/UnrealEd/Public/Kismet2/KismetEditorUtilities.h b/Engine/Source/Editor/UnrealEd/Public/Kismet2/KismetEditorUtilities.h
--- a/Engine/Source/Editor/UnrealEd/Public/Kismet2/KismetEditorUtilities.h
+++ b/Engine/Source/Editor/UnrealEd/Public/Kismet2/KismetEditorUtilities.h
@@ -16,6 +16,33 @@ class UBlueprintGeneratedClass;
 class UNREALED_API FKismetEditorUtilities
 {
 public:
+	/** 
+	 * Event that's broadcast anytime a Blueprint is created
+	 */
+	DECLARE_DELEGATE_OneParam(FOnBlueprintCreated, UBlueprint* /*InBlueprint*/);
+
+	/** Manages the TargetClass and EventName to use for spawning default "ghost" nodes in a new Blueprint */
+	struct FDefaultEventNodeData
+	{
+		/** If the new Blueprint is a child of the TargetClass an event will be attempted to be spawned.
+		 *	Hiding the category and other things can prevent the event from being placed
+		 */
+		UClass* TargetClass;
+
+		/** Event Name to spawn a node for */
+		FName EventName;
+	};
+
+	/** Manages the TargetClass and EventName to use for spawning default "ghost" nodes in a new Blueprint */
+	struct FOnBlueprintCreatedData
+	{
+		/** If the new Blueprint is a child of the TargetClass, the callback will be executed */
+		UClass* TargetClass;
+
+		/** Callback to execute */
+		FOnBlueprintCreated OnBlueprintCreated;
+	};
+
 	/**
 	 * Create a new Blueprint and initialize it to a valid state.
 	 *
@@ -225,6 +252,26 @@ public:
 	 */
 	static class UK2Node_Event* AddDefaultEventNode(UBlueprint* InBlueprint, UEdGraph* InGraph, FName InEventName, UClass* InEventClass, int32& InOutNodePosY);
 
+	/**
+	 * Will add an event to the list of default event nodes to be auto-generated for the class or a child of the class
+	 *
+	 * @param InOwner			Method of look-up so these registrations can later be removed when unregistering the Owner.
+	 * @param InTargetClass		If a new Blueprint is a child of the target class, the event will attempt to be placed in the main event graph
+	 * @param InEventName		Name of event to place
+	 */
+	static void RegisterAutoGeneratedDefaultEvent(void* InOwner, UClass* InTargetClass, FName InEventName);
+
+	/**
+	 * Will add an event to a list of callbacks to occur post Blueprint creation if the Blueprint is a child of the class
+	 *
+	 * @param InOwner							Method of look-up so these registrations can later be removed when unregistering the Owner.
+	 * @param InTargetClass						If a new Blueprint is a child of the target class, the event will attempt to be placed in the main event graph
+	 * @param InOnBlueprintCreatedCallback		Callback to call when the Blueprint is created
+	 */
+	static void RegisterOnBlueprintCreatedCallback(void* InOwner, UClass* InTargetClass, FOnBlueprintCreated InOnBlueprintCreatedCallback);
+
+	/** Unregisters a class from having AutoGeneratedDefaultEvent nodes or callbacks for OnBlueprintCreated */
+	static void UnregisterAutoBlueprintNodeCreation(void* InOwner);
 private:
 	/** Stores whether we are already listening for kismet clicks */
 	static bool bIsListeningForClicksOnKismetLog;
@@ -232,6 +279,11 @@ private:
 	/** List of blueprint parent class names cached by IsTrackedBlueprintParent() */
 	static TArray<FString> TrackedBlueprintParentList;
 
+	/** Mapping of classes to names of Events that should be automatically spawned */
+	static TMultiMap<void*, FDefaultEventNodeData> AutoGeneratedDefaultEventsMap;
+
+	/** Mapping of classes to delegate callbacks when a Blueprint is created, occurs post Event node creation */
+	static TMultiMap<void*, FOnBlueprintCreatedData> OnBlueprintCreatedCallbacks;
 private:
 	/** Get IBlueprintEditor for given object, if it exists */
 	static TSharedPtr<class IBlueprintEditor> GetIBlueprintEditorForObject(const UObject* ObjectToFocusOn, bool bOpenEditor);
