Sequencer: Highlight selected curves in the curve editor. Adjust pre-post infinity extrapolation menu items so that they operate on the selected curves if there are any.

diff --git a/Engine/Source/Editor/Sequencer/Private/SSequencerCurveEditor.cpp b/Engine/Source/Editor/Sequencer/Private/SSequencerCurveEditor.cpp
--- a/Engine/Source/Editor/Sequencer/Private/SSequencerCurveEditor.cpp
+++ b/Engine/Source/Editor/Sequencer/Private/SSequencerCurveEditor.cpp
@@ -143,11 +143,14 @@ void SSequencerCurveEditor::UpdateCurveOwner()
 
 	if (bAllFound)
 	{
+		UpdateCurveViewModelSelection();
 		return;
 	}
 
 	CurveOwner = MakeShareable( NewCurveOwner );
 	SetCurveOwner( CurveOwner.Get() );
+
+	UpdateCurveViewModelSelection();
 }
 
 bool SSequencerCurveEditor::GetCurveSnapEnabled() const
@@ -177,6 +180,17 @@ void SSequencerCurveEditor::NodeTreeSelectionChanged()
 		{
 			UpdateCurveOwner();
 		}
+
+		UpdateCurveViewModelSelection();
+	}
+}
+
+void SSequencerCurveEditor::UpdateCurveViewModelSelection()
+{
+	ClearSelectedCurveViewModels();
+	for (auto SelectedCurve : CurveOwner->GetSelectedCurves())
+	{
+		SetSelectedCurveViewModel(SelectedCurve);
 	}
 }
 
diff --git a/Engine/Source/Editor/Sequencer/Private/SSequencerCurveEditor.h b/Engine/Source/Editor/Sequencer/Private/SSequencerCurveEditor.h
--- a/Engine/Source/Editor/Sequencer/Private/SSequencerCurveEditor.h
+++ b/Engine/Source/Editor/Sequencer/Private/SSequencerCurveEditor.h
@@ -53,6 +53,9 @@ private:
 	/** Run whenever the selected curve visibility changes. */
 	void OnCurveEditorCurveVisibilityChanged();
 
+	/** Update the curve view model selection */
+	void UpdateCurveViewModelSelection();
+
 private:
 
 	/** The sequencer which owns this widget. */
diff --git a/Engine/Source/Editor/Sequencer/Private/SequencerCurveOwner.cpp b/Engine/Source/Editor/Sequencer/Private/SequencerCurveOwner.cpp
--- a/Engine/Source/Editor/Sequencer/Private/SequencerCurveOwner.cpp
+++ b/Engine/Source/Editor/Sequencer/Private/SequencerCurveOwner.cpp
@@ -81,6 +81,29 @@ FSequencerCurveOwner::FSequencerCurveOwner( TSharedPtr<FSequencerNodeTree> InSeq
 	}
 }
 
+TArray<FRichCurve*> FSequencerCurveOwner::GetSelectedCurves() const
+{
+	TArray<FRichCurve*> SelectedCurves;
+
+	TArray<TSharedRef<FSectionKeyAreaNode>> KeyAreaNodes;
+	GetAllKeyAreaNodes( SequencerNodeTree, KeyAreaNodes );
+	for ( TSharedRef<FSectionKeyAreaNode> KeyAreaNode : KeyAreaNodes )
+	{
+		for ( TSharedRef<IKeyArea> KeyArea : KeyAreaNode->GetAllKeyAreas() )
+		{
+			FRichCurve* RichCurve = KeyArea->GetRichCurve();
+			if ( RichCurve != nullptr )
+			{
+				if (KeyAreaNode->GetSequencer().GetSelection().IsSelected(KeyAreaNode))
+				{
+					SelectedCurves.Add(RichCurve);
+				}
+			}
+		}
+	}
+	return SelectedCurves;
+}
+
 TArray<FRichCurveEditInfoConst> FSequencerCurveOwner::GetCurves() const
 {
 	return ConstCurves;
diff --git a/Engine/Source/Editor/Sequencer/Private/SequencerCurveOwner.h b/Engine/Source/Editor/Sequencer/Private/SequencerCurveOwner.h
--- a/Engine/Source/Editor/Sequencer/Private/SequencerCurveOwner.h
+++ b/Engine/Source/Editor/Sequencer/Private/SequencerCurveOwner.h
@@ -10,6 +10,9 @@ class FSequencerCurveOwner : public FCurveOwnerInterface
 public:
 	FSequencerCurveOwner(TSharedPtr<FSequencerNodeTree> InSequencerNodeTree, ECurveEditorCurveVisibility::Type CurveVisibility);
 
+	/** Return the set of selected curves */
+	TArray<FRichCurve*> GetSelectedCurves() const;
+
 	/** FCurveOwnerInterface */
 	virtual TArray<FRichCurveEditInfoConst> GetCurves() const override;
 	virtual TArray<FRichCurveEditInfo> GetCurves() override;
diff --git a/Engine/Source/Editor/UnrealEd/Private/SCurveEditor.cpp b/Engine/Source/Editor/UnrealEd/Private/SCurveEditor.cpp
--- a/Engine/Source/Editor/UnrealEd/Private/SCurveEditor.cpp
+++ b/Engine/Source/Editor/UnrealEd/Private/SCurveEditor.cpp
@@ -721,6 +721,8 @@ int32 SCurveEditor::OnPaint(const FPaintArgs& Args, const FGeometry& AllottedGeo
 	int32 KeyLayerId = CurveLayerId + 1;
 	int32 SelectedKeyLayerId = KeyLayerId + 1;
 
+	bool bAnyCurveViewModelsSelected = AnyCurveViewModelsSelected();
+
 	if( AreCurvesVisible() )
 	{
 		//Paint the curves, unlocked curves will be on top
@@ -728,7 +730,7 @@ int32 SCurveEditor::OnPaint(const FPaintArgs& Args, const FGeometry& AllottedGeo
 		{
 			if (CurveViewModel->bIsVisible)
 			{
-				PaintCurve(CurveViewModel, CurveAreaGeometry, ScaleInfo, OutDrawElements, CurveViewModel->bIsLocked ? LockedCurveLayerID : CurveLayerId, MyClippingRect, DrawEffects, InWidgetStyle);
+				PaintCurve(CurveViewModel, CurveAreaGeometry, ScaleInfo, OutDrawElements, CurveViewModel->bIsLocked ? LockedCurveLayerID : CurveLayerId, MyClippingRect, DrawEffects, InWidgetStyle, bAnyCurveViewModelsSelected);
 			}
 		}
 
@@ -737,7 +739,7 @@ int32 SCurveEditor::OnPaint(const FPaintArgs& Args, const FGeometry& AllottedGeo
 		{
 			if (CurveViewModel->bIsVisible)
 			{
-				PaintKeys(CurveViewModel, ScaleInfo, OutDrawElements, KeyLayerId, SelectedKeyLayerId, CurveAreaGeometry, MyClippingRect, DrawEffects, InWidgetStyle);
+				PaintKeys(CurveViewModel, ScaleInfo, OutDrawElements, KeyLayerId, SelectedKeyLayerId, CurveAreaGeometry, MyClippingRect, DrawEffects, InWidgetStyle, bAnyCurveViewModelsSelected);
 			}
 		}
 	}
@@ -756,7 +758,7 @@ int32 SCurveEditor::OnPaint(const FPaintArgs& Args, const FGeometry& AllottedGeo
 }
 
 void SCurveEditor::PaintCurve(TSharedPtr<FCurveViewModel> CurveViewModel, const FGeometry &AllottedGeometry, FTrackScaleInfo &ScaleInfo, FSlateWindowElementList &OutDrawElements,
-	int32 LayerId, const FSlateRect& MyClippingRect, ESlateDrawEffect::Type DrawEffects, const FWidgetStyle &InWidgetStyle )const
+	int32 LayerId, const FSlateRect& MyClippingRect, ESlateDrawEffect::Type DrawEffects, const FWidgetStyle &InWidgetStyle, bool bAnyCurveViewModelsSelected )const
 {
 	if (CurveViewModel.IsValid())
 	{
@@ -764,6 +766,12 @@ void SCurveEditor::PaintCurve(TSharedPtr<FCurveViewModel> CurveViewModel, const
 		{
 			FLinearColor Color = InWidgetStyle.GetColorAndOpacityTint() * CurveViewModel->Color;
 
+			// Fade out curves that are not selected.
+			if (!CurveViewModel->bIsSelected && bAnyCurveViewModelsSelected)
+			{
+				Color *= FLinearColor(1.0f,1.0f,1.0f,0.5f); 
+			}
+
 			// Fade out curves which are locked.
 			if(CurveViewModel->bIsLocked)
 			{
@@ -870,7 +878,7 @@ void SCurveEditor::CreateLinesForSegment( FRichCurve* Curve, const FRichCurveKey
 	}
 }
 
-void SCurveEditor::PaintKeys(TSharedPtr<FCurveViewModel> CurveViewModel, FTrackScaleInfo &ScaleInfo, FSlateWindowElementList &OutDrawElements, int32 LayerId, int32 SelectedLayerId, const FGeometry &AllottedGeometry, const FSlateRect& MyClippingRect, ESlateDrawEffect::Type DrawEffects, const FWidgetStyle &InWidgetStyle ) const
+void SCurveEditor::PaintKeys(TSharedPtr<FCurveViewModel> CurveViewModel, FTrackScaleInfo &ScaleInfo, FSlateWindowElementList &OutDrawElements, int32 LayerId, int32 SelectedLayerId, const FGeometry &AllottedGeometry, const FSlateRect& MyClippingRect, ESlateDrawEffect::Type DrawEffects, const FWidgetStyle &InWidgetStyle, bool bAnyCurveViewModelsSelected ) const
 {
 	FLinearColor KeyColor = CurveViewModel->bIsLocked ? FLinearColor(0.1f,0.1f,0.1f,1.f) : InWidgetStyle.GetColorAndOpacityTint();
 
@@ -892,6 +900,9 @@ void SCurveEditor::PaintKeys(TSharedPtr<FCurveViewModel> CurveViewModel, FTrackS
 		const FSlateBrush* KeyBrush = IsSelected ? FEditorStyle::GetBrush("CurveEd.CurveKeySelected") : FEditorStyle::GetBrush("CurveEd.CurveKey");
 		int32 LayerToUse = IsSelected ? SelectedLayerId: LayerId;
 
+		// Fade out keys that are not selected and whose curve is not selected as well.
+		FLinearColor SelectionTint = !CurveViewModel->bIsSelected && !IsSelected && bAnyCurveViewModelsSelected ? FLinearColor(1.0f,1.0f,1.0f,0.5f) : FLinearColor(1.0f,1.0f,1.0f,1.0f); 
+
 		FSlateDrawElement::MakeBox(
 			OutDrawElements,
 			LayerToUse,
@@ -899,7 +910,7 @@ void SCurveEditor::PaintKeys(TSharedPtr<FCurveViewModel> CurveViewModel, FTrackS
 			KeyBrush,
 			MyClippingRect,
 			DrawEffects,
-			KeyBrush->GetTint( InWidgetStyle ) * InWidgetStyle.GetColorAndOpacityTint() * KeyColor
+			KeyBrush->GetTint( InWidgetStyle ) * InWidgetStyle.GetColorAndOpacityTint() * KeyColor * SelectionTint
 			);
 
 		//Handle drawing the tangent controls for curve
@@ -908,7 +919,7 @@ void SCurveEditor::PaintKeys(TSharedPtr<FCurveViewModel> CurveViewModel, FTrackS
 		bool bIsLeaveSelected = false;
 		if(IsTangentVisible(Curve, KeyHandle, bIsTangentSelected, bIsArrivalSelected, bIsLeaveSelected) && (Curve->GetKeyInterpMode(KeyHandle) == RCIM_Cubic || LastInterpMode == RCIM_Cubic))
 		{
-			PaintTangent(ScaleInfo, Curve, KeyHandle, KeyLocation, OutDrawElements, LayerId, AllottedGeometry, MyClippingRect, DrawEffects, LayerToUse, InWidgetStyle, bIsTangentSelected, bIsArrivalSelected, bIsLeaveSelected);
+			PaintTangent(CurveViewModel, ScaleInfo, Curve, KeyHandle, KeyLocation, OutDrawElements, LayerId, AllottedGeometry, MyClippingRect, DrawEffects, LayerToUse, InWidgetStyle, bIsTangentSelected, bIsArrivalSelected, bIsLeaveSelected, bAnyCurveViewModelsSelected);
 		}
 
 		LastInterpMode = Curve->GetKeyInterpMode(KeyHandle);
@@ -916,7 +927,7 @@ void SCurveEditor::PaintKeys(TSharedPtr<FCurveViewModel> CurveViewModel, FTrackS
 }
 
 
-void SCurveEditor::PaintTangent( FTrackScaleInfo &ScaleInfo, FRichCurve* Curve, FKeyHandle KeyHandle, FVector2D KeyLocation, FSlateWindowElementList &OutDrawElements, int32 LayerId, const FGeometry &AllottedGeometry, const FSlateRect& MyClippingRect, ESlateDrawEffect::Type DrawEffects, int32 LayerToUse, const FWidgetStyle &InWidgetStyle, bool bTangentSelected, bool bIsArrivalSelected, bool bIsLeaveSelected ) const
+void SCurveEditor::PaintTangent( TSharedPtr<FCurveViewModel> CurveViewModel, FTrackScaleInfo &ScaleInfo, FRichCurve* Curve, FKeyHandle KeyHandle, FVector2D KeyLocation, FSlateWindowElementList &OutDrawElements, int32 LayerId, const FGeometry &AllottedGeometry, const FSlateRect& MyClippingRect, ESlateDrawEffect::Type DrawEffects, int32 LayerToUse, const FWidgetStyle &InWidgetStyle, bool bTangentSelected, bool bIsArrivalSelected, bool bIsLeaveSelected, bool bAnyCurveViewModelsSelected ) const
 {
 	FVector2D ArriveTangentLocation, LeaveTangentLocation;
 	GetTangentPoints(ScaleInfo, FSelectedCurveKey(Curve,KeyHandle), ArriveTangentLocation, LeaveTangentLocation);
@@ -932,6 +943,9 @@ void SCurveEditor::PaintTangent( FTrackScaleInfo &ScaleInfo, FRichCurve* Curve,
 	bool LeaveTangentSelected = bTangentSelected && bIsLeaveSelected;
 	bool ArriveTangentSelected = bTangentSelected && bIsArrivalSelected;
 
+	FLinearColor LeaveSelectionTint = !CurveViewModel->bIsSelected && !LeaveTangentSelected && bAnyCurveViewModelsSelected ? FLinearColor(1.0f,1.0f,1.0f,0.5f) : FLinearColor(1.0f,1.0f,1.0f,1.0f); 
+	FLinearColor ArriveSelectionTint = !CurveViewModel->bIsSelected && !ArriveTangentSelected && bAnyCurveViewModelsSelected ? FLinearColor(1.0f,1.0f,1.0f,0.5f) : FLinearColor(1.0f,1.0f,1.0f,1.0f); 
+
 	//Add lines from tangent control point to 'key'
 	TArray<FVector2D> LinePoints;
 	LinePoints.Add(KeyLocation);
@@ -943,7 +957,7 @@ void SCurveEditor::PaintTangent( FTrackScaleInfo &ScaleInfo, FRichCurve* Curve,
 		LinePoints,
 		MyClippingRect,
 		DrawEffects,
-		ArriveTangentSelected ? TangentColorSelected : TangentColor
+		ArriveTangentSelected ? TangentColorSelected * ArriveSelectionTint : TangentColor * ArriveSelectionTint
 		);
 
 	LinePoints.Empty();
@@ -956,7 +970,7 @@ void SCurveEditor::PaintTangent( FTrackScaleInfo &ScaleInfo, FRichCurve* Curve,
 		LinePoints,
 		MyClippingRect,
 		DrawEffects,
-		LeaveTangentSelected ? TangentColorSelected : TangentColor
+		LeaveTangentSelected ? TangentColorSelected * LeaveSelectionTint : TangentColor * LeaveSelectionTint
 		);
 
 	//Arrive tangent control
@@ -967,7 +981,7 @@ void SCurveEditor::PaintTangent( FTrackScaleInfo &ScaleInfo, FRichCurve* Curve,
 		ArriveTangentSelected ? TangentBrushSelected : TangentBrush,
 		MyClippingRect,
 		DrawEffects,
-		ArriveTangentSelected ? TangentBrushSelected->GetTint( InWidgetStyle ) : TangentBrush->GetTint( InWidgetStyle ) * InWidgetStyle.GetColorAndOpacityTint()
+		ArriveTangentSelected ? TangentBrushSelected->GetTint( InWidgetStyle ) * ArriveSelectionTint : TangentBrush->GetTint( InWidgetStyle ) * InWidgetStyle.GetColorAndOpacityTint() * ArriveSelectionTint
 		);
 	//Leave tangent control
 	FSlateDrawElement::MakeBox(
@@ -977,7 +991,7 @@ void SCurveEditor::PaintTangent( FTrackScaleInfo &ScaleInfo, FRichCurve* Curve,
 		LeaveTangentSelected ? TangentBrushSelected : TangentBrush,
 		MyClippingRect,
 		DrawEffects,
-		LeaveTangentSelected ? TangentBrushSelected->GetTint( InWidgetStyle ) : TangentBrush->GetTint( InWidgetStyle ) * InWidgetStyle.GetColorAndOpacityTint()
+		LeaveTangentSelected ? TangentBrushSelected->GetTint( InWidgetStyle ) * LeaveSelectionTint : TangentBrush->GetTint( InWidgetStyle ) * InWidgetStyle.GetColorAndOpacityTint() * LeaveSelectionTint
 		);
 }
 
@@ -2626,6 +2640,35 @@ void SCurveEditor::SetOutputMinMax(float NewMin, float NewMax)
 	}
 }
 
+void SCurveEditor::ClearSelectedCurveViewModels()
+{
+	for(auto CurveViewModel : CurveViewModels)
+	{
+		CurveViewModel->bIsSelected = false;
+	}
+}
+
+void SCurveEditor::SetSelectedCurveViewModel(FRichCurve* CurveToSelect)
+{
+	TSharedPtr<FCurveViewModel> ViewModel = GetViewModelForCurve(CurveToSelect);
+	if (ViewModel.IsValid())
+	{
+		ViewModel.Get()->bIsSelected = true;
+	}
+}
+
+bool SCurveEditor::AnyCurveViewModelsSelected() const
+{
+	for (auto CurveViewModel : CurveViewModels)
+	{
+		if (CurveViewModel->bIsSelected)
+		{
+			return true;
+		}
+	}
+	return false;
+}
+
 TSharedPtr<FCurveViewModel> SCurveEditor::HitTestCurves(  const FGeometry& InMyGeometry, const FPointerEvent& InMouseEvent )
 {
 	if( AreCurvesVisible() )
@@ -2860,13 +2903,18 @@ void SCurveEditor::OnSelectPreInfinityExtrap(ERichCurveExtrapolation Extrapolati
 	const FScopedTransaction Transaction(LOCTEXT("CurveEditor_SetPreInfinityExtrapolation", "Set Pre-Infinity Extrapolation"));
 	CurveOwner->ModifyOwner();
 
+	bool bAnyCurveViewModelsSelected = AnyCurveViewModelsSelected();
+
 	TArray<FRichCurveEditInfo> ChangedCurveEditInfos;
 	for (auto CurveViewModel : CurveViewModels)
 	{
-		if (CurveViewModel->CurveInfo.CurveToEdit->PreInfinityExtrap != Extrapolation)
+		if (!bAnyCurveViewModelsSelected || CurveViewModel->bIsSelected)
 		{
-			CurveViewModel->CurveInfo.CurveToEdit->PreInfinityExtrap = Extrapolation;
-			ChangedCurveEditInfos.Add(CurveViewModel->CurveInfo);
+			if (CurveViewModel->CurveInfo.CurveToEdit->PreInfinityExtrap != Extrapolation)
+			{
+				CurveViewModel->CurveInfo.CurveToEdit->PreInfinityExtrap = Extrapolation;
+				ChangedCurveEditInfos.Add(CurveViewModel->CurveInfo);
+			}
 		}
 	}
 
@@ -2878,13 +2926,30 @@ void SCurveEditor::OnSelectPreInfinityExtrap(ERichCurveExtrapolation Extrapolati
 
 bool SCurveEditor::IsPreInfinityExtrapSelected(ERichCurveExtrapolation Extrapolation)
 {
+	bool bAnyCurveViewModelsSelected = AnyCurveViewModelsSelected();
+
 	for (auto CurveViewModel : CurveViewModels)
 	{
-		if (CurveViewModel->CurveInfo.CurveToEdit->PreInfinityExtrap != Extrapolation)
+		// If there are any curves selected, the setting must match all of the selected curves
+		if (bAnyCurveViewModelsSelected)
 		{
-			return false;
+			if (CurveViewModel->bIsSelected)
+			{
+				if (CurveViewModel->CurveInfo.CurveToEdit->PreInfinityExtrap != Extrapolation)
+				{
+					return false;
+				}
+			}
+		}
+		else
+		{
+			if (CurveViewModel->CurveInfo.CurveToEdit->PreInfinityExtrap != Extrapolation)
+			{
+				return false;
+			}
 		}
 	}
+
 	return CurveViewModels.Num() > 0;
 }
 
@@ -2893,13 +2958,18 @@ void SCurveEditor::OnSelectPostInfinityExtrap(ERichCurveExtrapolation Extrapolat
 	const FScopedTransaction Transaction(LOCTEXT("CurveEditor_SetPostInfinityExtrapolation", "Set Post-Infinity Extrapolation"));
 	CurveOwner->ModifyOwner();
 
+	bool bAnyCurveViewModelsSelected = AnyCurveViewModelsSelected();
+
 	TArray<FRichCurveEditInfo> ChangedCurveEditInfos;
 	for (auto CurveViewModel : CurveViewModels)
 	{
-		if (CurveViewModel->CurveInfo.CurveToEdit->PostInfinityExtrap != Extrapolation)
+		if (!bAnyCurveViewModelsSelected || CurveViewModel->bIsSelected)
 		{
-			CurveViewModel->CurveInfo.CurveToEdit->PostInfinityExtrap = Extrapolation;
-			ChangedCurveEditInfos.Add(CurveViewModel->CurveInfo);
+			if (CurveViewModel->CurveInfo.CurveToEdit->PostInfinityExtrap != Extrapolation)
+			{
+				CurveViewModel->CurveInfo.CurveToEdit->PostInfinityExtrap = Extrapolation;
+				ChangedCurveEditInfos.Add(CurveViewModel->CurveInfo);
+			}
 		}
 	}
 
@@ -2911,13 +2981,30 @@ void SCurveEditor::OnSelectPostInfinityExtrap(ERichCurveExtrapolation Extrapolat
 
 bool SCurveEditor::IsPostInfinityExtrapSelected(ERichCurveExtrapolation Extrapolation)
 {
+	bool bAnyCurveViewModelsSelected = AnyCurveViewModelsSelected();
+
 	for (auto CurveViewModel : CurveViewModels)
 	{
-		if (CurveViewModel->CurveInfo.CurveToEdit->PostInfinityExtrap != Extrapolation)
+		// If there are any curves selected, the setting must match all of the selected curves
+		if (bAnyCurveViewModelsSelected)
+		{
+			if (CurveViewModel->bIsSelected)
+			{
+				if (CurveViewModel->CurveInfo.CurveToEdit->PostInfinityExtrap != Extrapolation)
+				{
+					return false;
+				}
+			}
+		}
+		else
 		{
-			return false;
+			if (CurveViewModel->CurveInfo.CurveToEdit->PostInfinityExtrap != Extrapolation)
+			{
+				return false;
+			}
 		}
 	}
+
 	return CurveViewModels.Num() > 0;
 }
 
diff --git a/Engine/Source/Editor/UnrealEd/Public/SCurveEditor.h b/Engine/Source/Editor/UnrealEd/Public/SCurveEditor.h
--- a/Engine/Source/Editor/UnrealEd/Public/SCurveEditor.h
+++ b/Engine/Source/Editor/UnrealEd/Public/SCurveEditor.h
@@ -86,6 +86,8 @@ public:
 	bool bIsVisible;
 	/** Whether or not the curve is locked from editing. */
 	bool bIsLocked;
+	/** Whether or not the curve is selected. */
+	bool bIsSelected;
 
 	FCurveViewModel(FRichCurveEditInfo InCurveInfo, FLinearColor InColor, bool bInIsLocked)
 	{
@@ -93,6 +95,7 @@ public:
 		Color = InColor;
 		bIsVisible = true;
 		bIsLocked = bInIsLocked;
+		bIsSelected = true;
 	}
 };
 
@@ -421,14 +424,14 @@ private:
 
 	/** Paint a curve */
 	void PaintCurve(TSharedPtr<FCurveViewModel> CurveViewModel, const FGeometry &AllottedGeometry, FTrackScaleInfo &ScaleInfo, FSlateWindowElementList &OutDrawElements, 
-					int32 LayerId, const FSlateRect& MyClippingRect, ESlateDrawEffect::Type DrawEffects, const FWidgetStyle &InWidgetStyle) const;
+					int32 LayerId, const FSlateRect& MyClippingRect, ESlateDrawEffect::Type DrawEffects, const FWidgetStyle &InWidgetStyle, bool bAnyCurveViewModelsSelected) const;
 
 	/** Paint the keys that make up a curve */
-	void PaintKeys(TSharedPtr<FCurveViewModel> CurveViewModel, FTrackScaleInfo &ScaleInfo, FSlateWindowElementList &OutDrawElements, int32 LayerId, int32 SelectedLayerId, const FGeometry &AllottedGeometry, const FSlateRect& MyClippingRect, ESlateDrawEffect::Type DrawEffects, const FWidgetStyle &InWidgetStyle ) const;
+	void PaintKeys(TSharedPtr<FCurveViewModel> CurveViewModel, FTrackScaleInfo &ScaleInfo, FSlateWindowElementList &OutDrawElements, int32 LayerId, int32 SelectedLayerId, const FGeometry &AllottedGeometry, const FSlateRect& MyClippingRect, ESlateDrawEffect::Type DrawEffects, const FWidgetStyle &InWidgetStyle, bool bAnyCurveViewModelsSelected ) const;
 
 	/** Paint the tangent for a key with cubic curves */
-	void PaintTangent( FTrackScaleInfo &ScaleInfo, FRichCurve* Curve, FKeyHandle KeyHandle, FVector2D KeyLocation, FSlateWindowElementList &OutDrawElements, int32 LayerId, 
-					   const FGeometry &AllottedGeometry, const FSlateRect& MyClippingRect, ESlateDrawEffect::Type DrawEffects, int32 LayerToUse, const FWidgetStyle &InWidgetStyle, bool bTangentSelected, bool bIsArrivalSelected, bool bIsLeaveSelected ) const;
+	void PaintTangent( TSharedPtr<FCurveViewModel> CurveViewModel, FTrackScaleInfo &ScaleInfo, FRichCurve* Curve, FKeyHandle KeyHandle, FVector2D KeyLocation, FSlateWindowElementList &OutDrawElements, int32 LayerId, 
+					   const FGeometry &AllottedGeometry, const FSlateRect& MyClippingRect, ESlateDrawEffect::Type DrawEffects, int32 LayerToUse, const FWidgetStyle &InWidgetStyle, bool bTangentSelected, bool bIsArrivalSelected, bool bIsLeaveSelected, bool bAnyCurveViewModelsSelected ) const;
 
 	/** Paint Grid lines, these make it easier to visualize relative distance */
 	void PaintGridLines( const FGeometry &AllottedGeometry, FTrackScaleInfo &ScaleInfo, FSlateWindowElementList &OutDrawElements, int32 LayerId, const FSlateRect& MyClippingRect, ESlateDrawEffect::Type DrawEffects )const;
@@ -561,6 +564,15 @@ protected:
 	/** Access the user-supplied settings object */
 	UCurveEditorSettings* GetSettings() const { return Settings; }
 
+	/** Clear the selected curve view models */
+	UNREALED_API void ClearSelectedCurveViewModels();
+
+	/** Set the selected curve view model that matches the rich curve */
+	UNREALED_API void SetSelectedCurveViewModel(FRichCurve* Curve);
+
+	/** Return whether any curve view models are selected */
+	UNREALED_API bool AnyCurveViewModelsSelected() const;
+
 private:
 
 	/** User-supplied object for this curve editor */
