Checking in UTF8 name support for bone, animation, curves from community developer Alwei

diff --git a/Engine/Source/Editor/UnrealEd/Private/SkeletalMeshEdit.cpp b/Engine/Source/Editor/UnrealEd/Private/SkeletalMeshEdit.cpp
--- a/Engine/Source/Editor/UnrealEd/Private/SkeletalMeshEdit.cpp
+++ b/Engine/Source/Editor/UnrealEd/Private/SkeletalMeshEdit.cpp
@@ -174,7 +174,7 @@ bool UEditorEngine::ReimportFbxAnimation( USkeleton* Skeleton, UAnimSequence* An
 			//this is to make it easier for people who use content creation programs that only export one animation and/or ones that don't allow naming animations			
 			if (FbxImporter->Scene->GetSrcObjectCount(FbxCriteria::ObjectType(FbxAnimStack::ClassId)) > 1 && !ImportData->SourceAnimationName.IsEmpty())
 			{
-				CurAnimStack = FbxCast<FbxAnimStack>(FbxImporter->Scene->FindSrcObject(FbxCriteria::ObjectType(FbxAnimStack::ClassId), TCHAR_TO_ANSI(*ImportData->SourceAnimationName), 0));
+				CurAnimStack = FbxCast<FbxAnimStack>(FbxImporter->Scene->FindSrcObject(FbxCriteria::ObjectType(FbxAnimStack::ClassId), TCHAR_TO_UTF8(*ImportData->SourceAnimationName), 0));
 			}
 			else
 			{
@@ -286,14 +286,14 @@ bool UnFbx::FFbxImporter::IsValidAnimationData(TArray<FbxNode*>& SortedLinks, TA
 		// debug purpose
 		for (int32 BoneIndex = 0; BoneIndex < SortedLinks.Num(); BoneIndex++)
 		{
-			FString BoneName = MakeName(SortedLinks[BoneIndex]->GetName());
+			FString BoneName = UTF8_TO_TCHAR(MakeName(SortedLinks[BoneIndex]->GetName()));
 			UE_LOG(LogFbx, Log, TEXT("SortedLinks :(%d) %s"), BoneIndex, *BoneName );
 		}
 
 		FbxTimeSpan AnimTimeSpan = GetAnimationTimeSpan(SortedLinks[0], CurAnimStack);
 		if (AnimTimeSpan.GetDuration() <= 0)
 		{
-			AddTokenizedErrorMessage(FTokenizedMessage::Create(EMessageSeverity::Warning, FText::Format(LOCTEXT("FBXImport_ZeroLength", "Animation Stack {0} does not contain any valid key. Try different time options when import."), FText::FromString(CurAnimStack->GetName()))), FFbxErrors::Animation_ZeroLength);
+			AddTokenizedErrorMessage(FTokenizedMessage::Create(EMessageSeverity::Warning, FText::Format(LOCTEXT("FBXImport_ZeroLength", "Animation Stack {0} does not contain any valid key. Try different time options when import."), FText::FromString(UTF8_TO_TCHAR(CurAnimStack->GetName())))), FFbxErrors::Animation_ZeroLength);
 			continue;
 		}
 
@@ -345,7 +345,7 @@ void UnFbx::FFbxImporter::FillAndVerifyBoneNames(USkeleton* Skeleton, TArray<Fbx
 	// copy to the data
 	for (int32 BoneIndex = 0; BoneIndex < TrackNum; BoneIndex++)
 	{
-		OutRawBoneNames[BoneIndex] = FName(*FSkeletalMeshImportData::FixupBoneName( (ANSICHAR*)MakeName(SortedLinks[BoneIndex]->GetName()) ));
+		OutRawBoneNames[BoneIndex] = FName(*FSkeletalMeshImportData::FixupBoneName( UTF8_TO_TCHAR(MakeName(SortedLinks[BoneIndex]->GetName())) ));
 	}
 
 	const FReferenceSkeleton& RefSkeleton = Skeleton->GetReferenceSkeleton();
@@ -514,7 +514,7 @@ UAnimSequence * UnFbx::FFbxImporter::ImportAnimations(USkeleton* Skeleton, UObje
 		if (ValidTakeCount > 1)
 		{
 			SequenceName += "_";
-			SequenceName += ANSI_TO_TCHAR(CurAnimStack->GetName());
+			SequenceName += UTF8_TO_TCHAR(CurAnimStack->GetName());
 		}
 
 		// See if this sequence already exists.
@@ -630,7 +630,7 @@ bool UnFbx::FFbxImporter::ValidateAnimStack(TArray<FbxNode*>& SortedLinks, TArra
 	// set current anim stack
 	Scene->SetCurrentAnimationStack(CurAnimStack);
 
-	UE_LOG(LogFbx, Log, TEXT("Parsing AnimStack %s"),ANSI_TO_TCHAR(CurAnimStack->GetName()));
+	UE_LOG(LogFbx, Log, TEXT("Parsing AnimStack %s"),UTF8_TO_TCHAR(CurAnimStack->GetName()));
 
 	// There are a FBX unroll filter bug, so don't bake animation layer at all
 	MergeAllLayerAnimation(CurAnimStack, ResampleRate);
@@ -1024,7 +1024,7 @@ bool UnFbx::FFbxImporter::ImportAnimation(USkeleton* Skeleton, UAnimSequence * D
 
 					const int32 BlendShapeChannelCount = BlendShape->GetBlendShapeChannelCount();
 
-					FString BlendShapeName = ANSI_TO_TCHAR(MakeName(BlendShape->GetName()));
+					FString BlendShapeName = UTF8_TO_TCHAR(MakeName(BlendShape->GetName()));
 
 					for(int32 ChannelIndex = 0; ChannelIndex<BlendShapeChannelCount; ++ChannelIndex)
 					{
@@ -1032,7 +1032,7 @@ bool UnFbx::FFbxImporter::ImportAnimation(USkeleton* Skeleton, UAnimSequence * D
 
 						if(Channel)
 						{
-							FString ChannelName = ANSI_TO_TCHAR(MakeName(Channel->GetName()));
+							FString ChannelName = UTF8_TO_TCHAR(MakeName(Channel->GetName()));
 
 							// Maya adds the name of the blendshape and an underscore to the front of the channel name, so remove it
 							if(ChannelName.StartsWith(BlendShapeName))
@@ -1076,7 +1076,7 @@ bool UnFbx::FFbxImporter::ImportAnimation(USkeleton* Skeleton, UAnimSequence * D
 				if( CurveNode && Property.GetFlag(FbxPropertyAttr::eUserDefined) && 
 					CurveNode->IsAnimated() && IsSupportedCurveDataType(Property.GetPropertyDataType().GetType()) )
 				{
-					FString CurveName = CurveNode->GetName();
+					FString CurveName = UTF8_TO_TCHAR(CurveNode->GetName());
 					UE_LOG(LogFbx, Log, TEXT("CurveName : %s"), *CurveName );
 
 					int32 TotalCount = CurveNode->GetChannelsCount();
