Added per file cooking stats.
Added ddc exists stats.
Reworked the existing ddc stat code to make it neater.

diff --git a/Engine/Source/Developer/DerivedDataCache/Private/DDCStatsHelper.cpp b/Engine/Source/Developer/DerivedDataCache/Private/DDCStatsHelper.cpp
--- a/Engine/Source/Developer/DerivedDataCache/Private/DDCStatsHelper.cpp
+++ b/Engine/Source/Developer/DerivedDataCache/Private/DDCStatsHelper.cpp
@@ -4,61 +4,81 @@
 #include "CookingStatsModule.h"
 #include "ModuleManager.h"
 
-
-uint32 GetValidTLSSlot()
+namespace DDCStats
 {
-	uint32 TLSSlot = FPlatformTLS::AllocTlsSlot();
-	check(FPlatformTLS::IsValidTlsSlot(TLSSlot));
-	FPlatformTLS::SetTlsValue(TLSSlot, nullptr);
-	return TLSSlot;
-}
 
-ICookingStats* CookingStats = nullptr;
+	ICookingStats* GetCookingStats()
+	{
+		static ICookingStats* CookingStats = nullptr;
+		static bool bInitialized = false;
+		if (bInitialized == false)
+		{
+			FCookingStatsModule* CookingStatsModule = FModuleManager::LoadModulePtr<FCookingStatsModule>(TEXT("CookingStats"));
+			if (CookingStatsModule)
+			{
+				CookingStats = &CookingStatsModule->Get();
+			}
+			bInitialized = true;
+		}
+		return CookingStats;
+	}
+
+	static uint32 GetValidTLSSlot()
+	{
+		uint32 TLSSlot = FPlatformTLS::AllocTlsSlot();
+		check(FPlatformTLS::IsValidTlsSlot(TLSSlot));
+		FPlatformTLS::SetTlsValue(TLSSlot, nullptr);
+		return TLSSlot;
+	}
+	
 
-struct FDDCStatsTLSStore
-{
+	struct FDDCStatsTLSStore
+	{
 
-	FDDCStatsTLSStore() : TransactionGuid(NAME_None), CurrentIndex(0), RootScope(NULL) { }
+		FDDCStatsTLSStore() : TransactionGuid(NAME_None), CurrentIndex(0), RootScope(NULL) { }
 
 
-	FName GenerateNewTransactionName()
-	{
-		++CurrentIndex;
-		if (CurrentIndex <= 1)
+		FName GenerateNewTransactionName()
 		{
-			FString TransactionGuidString = TEXT("DDCTransactionId");
-			TransactionGuidString += FGuid::NewGuid().ToString();
-			TransactionGuid = FName(*TransactionGuidString);
-			CurrentIndex = 1;
+			++CurrentIndex;
+			if (CurrentIndex <= 1)
+			{
+				FString TransactionGuidString = TEXT("DDCTransactionId");
+				TransactionGuidString += FGuid::NewGuid().ToString();
+				TransactionGuid = FName(*TransactionGuidString);
+				CurrentIndex = 1;
+			}
+			check(TransactionGuid != NAME_None);
+
+			return FName(TransactionGuid, CurrentIndex);
 		}
-		check(TransactionGuid != NAME_None);
 
-		return FName(TransactionGuid, CurrentIndex);
-	}
+		FName TransactionGuid;
+		int32 CurrentIndex;
+		FDDCScopeStatHelper *RootScope;
+	};
+
+	uint32 CookStatsFDDCStatsTLSStore = GetValidTLSSlot();
 
-	FName TransactionGuid;
-	int32 CurrentIndex;
-	FDDCScopeStatHelper *RootScope;
-};
+}; // namespace DDCStats
 
-uint32 CookStatsFDDCStatsTLSStore = GetValidTLSSlot();
 
-FDDCScopeStatHelper::FDDCScopeStatHelper(const TCHAR* CacheKey, const FName& FunctionName ) : bHasParent(false)
+FDDCScopeStatHelper::FDDCScopeStatHelper(const TCHAR* CacheKey, const FName& FunctionName) : bHasParent(false)
 {
-	FDDCStatsTLSStore* TLSStore = (FDDCStatsTLSStore*)FPlatformTLS::GetTlsValue(CookStatsFDDCStatsTLSStore);
+
+	if (DDCStats::GetCookingStats() == nullptr)
+		return;
+
+	DDCStats::FDDCStatsTLSStore* TLSStore = (DDCStats::FDDCStatsTLSStore*)FPlatformTLS::GetTlsValue(DDCStats::CookStatsFDDCStatsTLSStore);
 	if (TLSStore == nullptr)
 	{
-		TLSStore = new FDDCStatsTLSStore();
-		FPlatformTLS::SetTlsValue(CookStatsFDDCStatsTLSStore, TLSStore);
+		TLSStore = new DDCStats::FDDCStatsTLSStore();
+		FPlatformTLS::SetTlsValue(DDCStats::CookStatsFDDCStatsTLSStore, TLSStore);
 	}
 
 	TransactionGuid = TLSStore->GenerateNewTransactionName();
 
-	if (CookingStats == nullptr)
-	{
-		FCookingStatsModule& CookingStatsModule = FModuleManager::LoadModuleChecked<FCookingStatsModule>(TEXT("CookingStats"));
-		CookingStats = &CookingStatsModule.Get();
-	}
+
 	FString Temp;
 	AddTag(FunctionName, Temp);
 
@@ -74,7 +94,7 @@ FDDCScopeStatHelper::FDDCScopeStatHelper(const TCHAR* CacheKey, const FName& Fun
 	{
 		static const FName NAME_Parent = FName(TEXT("Parent"));
 		AddTag(NAME_Parent, TLSStore->RootScope->TransactionGuid.ToString());
-		
+
 		bHasParent = true;
 	}
 
@@ -90,13 +110,16 @@ FDDCScopeStatHelper::FDDCScopeStatHelper(const TCHAR* CacheKey, const FName& Fun
 
 FDDCScopeStatHelper::~FDDCScopeStatHelper()
 {
+	if (DDCStats::GetCookingStats() == nullptr)
+		return;
+
 	double EndTime = FPlatformTime::Seconds();
 	float DurationMS = (EndTime - StartTime) * 1000.0f;
 
 	static const FName NAME_Duration = FName(TEXT("Duration"));
 	AddTag(NAME_Duration, FString::Printf(TEXT("%fms"), DurationMS));
 
-	FDDCStatsTLSStore* TLSStore = (FDDCStatsTLSStore*)FPlatformTLS::GetTlsValue(CookStatsFDDCStatsTLSStore);
+	DDCStats::FDDCStatsTLSStore* TLSStore = (DDCStats::FDDCStatsTLSStore*)FPlatformTLS::GetTlsValue(DDCStats::CookStatsFDDCStatsTLSStore);
 	check(TLSStore);
 
 	if (TLSStore->RootScope == this)
@@ -107,5 +130,8 @@ FDDCScopeStatHelper::~FDDCScopeStatHelper()
 
 void FDDCScopeStatHelper::AddTag(const FName& Tag, const FString& Value)
 {
-	CookingStats->AddTagValue(TransactionGuid, Tag, Value);
+	if (DDCStats::GetCookingStats())
+	{
+		DDCStats::GetCookingStats()->AddTagValue(TransactionGuid, Tag, Value);
+	}
 }
\ No newline at end of file
diff --git a/Engine/Source/Developer/DerivedDataCache/Private/FileSystemDerivedDataBackend.cpp b/Engine/Source/Developer/DerivedDataCache/Private/FileSystemDerivedDataBackend.cpp
--- a/Engine/Source/Developer/DerivedDataCache/Private/FileSystemDerivedDataBackend.cpp
+++ b/Engine/Source/Developer/DerivedDataCache/Private/FileSystemDerivedDataBackend.cpp
@@ -135,7 +135,10 @@ public:
 	 */
 	virtual bool CachedDataProbablyExists(const TCHAR* CacheKey) override
 	{
-		// FDDCScopeStatHelper Stat(FString(CacheKey), FString(TEXT("FileSystemDerivedDataBackend::CachedDataProbablyExists")));
+		static FName NAME_CachedDataProbablyExists(TEXT("CachedDataProbablyExists"));
+		FDDCScopeStatHelper Stat(CacheKey, NAME_CachedDataProbablyExists);
+		static FName NAME_FileDDCPath(TEXT("FileDDCPath"));
+		Stat.AddTag(NAME_FileDDCPath, CachePath);
 
 		check(!bFailed);
 		FString Filename = BuildFilename(CacheKey);
diff --git a/Engine/Source/Developer/DerivedDataCache/Private/HierarchicalDerivedDataBackend.h b/Engine/Source/Developer/DerivedDataCache/Private/HierarchicalDerivedDataBackend.h
--- a/Engine/Source/Developer/DerivedDataCache/Private/HierarchicalDerivedDataBackend.h
+++ b/Engine/Source/Developer/DerivedDataCache/Private/HierarchicalDerivedDataBackend.h
@@ -65,6 +65,11 @@ public:
 	 */
 	virtual bool CachedDataProbablyExists(const TCHAR* CacheKey) override
 	{
+		static FName NAME_CachedDataProbablyExists(TEXT("CachedDataProbablyExists"));
+		FDDCScopeStatHelper Stat(CacheKey, NAME_CachedDataProbablyExists);
+		const static FName NAME_HierarchicalDDC = FName(TEXT("HierarchicalDDC"));
+		Stat.AddTag(NAME_HierarchicalDDC, FString());
+
 		for (int32 CacheIndex = 0; CacheIndex < InnerBackends.Num(); CacheIndex++)
 		{
 			if (InnerBackends[CacheIndex]->CachedDataProbablyExists(CacheKey))
diff --git a/Engine/Source/Editor/CookingStats/Private/CookingStats.cpp b/Engine/Source/Editor/CookingStats/Private/CookingStats.cpp
--- a/Engine/Source/Editor/CookingStats/Private/CookingStats.cpp
+++ b/Engine/Source/Editor/CookingStats/Private/CookingStats.cpp
@@ -8,7 +8,7 @@ DEFINE_LOG_CATEGORY(LogCookingStats);
 
 FCookingStats::FCookingStats()
 {
-	RunGuid = FName(*FGuid::NewGuid().ToString());
+	RunGuid = FName(*FString::Printf(TEXT("RunID%s"), *FGuid::NewGuid().ToString()));
 }
 
 FCookingStats::~FCookingStats()
@@ -35,15 +35,31 @@ void FCookingStats::AddTagValue(const FName& Key, const FName& Tag, const FStrin
 		Value = &KeyTags.Add(Key);
 	}
 
-	FTag FullTag;
-	FullTag.Name = Tag;
-	FullTag.Value = TagValue;
-	
-	
-	Value->Add(FullTag);
+	Value->Add(Tag, TagValue);
 
 }
 
+
+bool FCookingStats::GetTagValue(const FName& Key, const FName& TagName, FString& OutValue) const
+{
+	FScopeLock ScopeLock(&SyncObject);
+	auto Tags = KeyTags.Find(Key);
+
+	if (Tags == nullptr)
+	{
+		return false;
+	}
+
+	auto* Value = Tags->Find(TagName);
+
+	if (Value == nullptr)
+	{
+		return false;
+	}
+	OutValue = *Value;
+	return true;
+}
+
 bool FCookingStats::SaveStatsAsCSV(const FString& Filename)
 {
 	FString Output;
@@ -56,7 +72,7 @@ bool FCookingStats::SaveStatsAsCSV(const FString& Filename)
 		for (const auto& Tag : Stat.Value)
 		{
 			Output += TEXT(", ");
-			Output += Tag.Name.ToString();
+			Output += Tag.Key.ToString();
 			if (Tag.Value.IsEmpty() == false)
 			{
 				Output += TEXT("=");
@@ -68,7 +84,7 @@ bool FCookingStats::SaveStatsAsCSV(const FString& Filename)
 		for (const auto& Tag : GlobalTags)
 		{
 			Output += TEXT(", ");
-			Output += Tag.Name.ToString();
+			Output += Tag.Key.ToString();
 			if (Tag.Value.IsEmpty() == false)
 			{
 				Output += TEXT("=");
diff --git a/Engine/Source/Editor/CookingStats/Private/CookingStats.h b/Engine/Source/Editor/CookingStats/Private/CookingStats.h
--- a/Engine/Source/Editor/CookingStats/Private/CookingStats.h
+++ b/Engine/Source/Editor/CookingStats/Private/CookingStats.h
@@ -39,6 +39,17 @@ public:
 	*/
 	virtual void AddTagValue(const FName& Key, const FName& Tag, const FString& Value) override;
 
+	/**
+	* GetTagValue
+	* retrieve a previously added tag value
+	*
+	* @param Key to use to get the tag value
+	* @param Tag to use to get the tag value
+	* @param Value the value which will be returned fomr that tag key pair
+	* @return if the key + tag exists will retrun true false otherwise
+	*/
+	virtual bool GetTagValue(const FName& Key, const FName& Tag, FString& Value) const override;
+
 	/**
 	* SaveStatsAsCSV
 	* Save stats in a CSV file
@@ -51,12 +62,7 @@ public:
 private:
 	FName RunGuid; // unique guid for this run of the editor
 
-	struct FTag
-	{
-		FName Name;
-		FString Value;
-	};
-	TMap< FName, TArray<FTag> > KeyTags;
-	TArray<FTag> GlobalTags; // tags to add to every key in this runthrough
-	FCriticalSection SyncObject;
+	TMap< FName, TMap<FName,FString> > KeyTags;
+	TMap<FName, FString> GlobalTags; // tags to add to every key in this runthrough
+	mutable FCriticalSection SyncObject;
 };
\ No newline at end of file
diff --git a/Engine/Source/Editor/CookingStats/Public/ICookingStats.h b/Engine/Source/Editor/CookingStats/Public/ICookingStats.h
--- a/Engine/Source/Editor/CookingStats/Public/ICookingStats.h
+++ b/Engine/Source/Editor/CookingStats/Public/ICookingStats.h
@@ -30,15 +30,27 @@ public:
 	virtual void AddTag(const FName& Key, const FName& Tag) = 0;
 
 	/**
-	* AddTag
+	* AddTagValue
 	* add a tag to a key, the tags will be saved out with the key
 	* tags can be timing information or anything
 	*
 	* @param Key key to add the tag to
 	* @param Tag tag to add to the key
+	* @param VAlue value of the tag we are adding
 	*/
 	virtual void AddTagValue(const FName& Key, const FName& Tag, const FString& Value) = 0;
 
+	/**
+	 * GetTagValue
+	 * retrieve a previously added tag value
+	 *
+	 * @param Key to use to get the tag value
+	 * @param Tag to use to get the tag value
+	 * @param Value the value which will be returned fomr that tag key pair
+	 * @return if the key + tag exists will retrun true false otherwise
+	 */
+	virtual bool GetTagValue(const FName& Key, const FName& Tag, FString& Value) const = 0;
+
 	/**
 	 * SaveStatsAsCSV
 	 * Save stats in a CSV file
diff --git a/Engine/Source/Editor/UnrealEd/Private/CookOnTheFlyServer.cpp b/Engine/Source/Editor/UnrealEd/Private/CookOnTheFlyServer.cpp
--- a/Engine/Source/Editor/UnrealEd/Private/CookOnTheFlyServer.cpp
+++ b/Engine/Source/Editor/UnrealEd/Private/CookOnTheFlyServer.cpp
@@ -44,6 +44,8 @@
 
 DEFINE_LOG_CATEGORY_STATIC(LogCook, Log, All);
 
+// these are the stats which are output to the Saved\Stats\<date>.csv file
+#define ENABLE_COOKER_STATS 1
 
 #define DEBUG_COOKONTHEFLY 0
 #define OUTPUT_TIMING 0
@@ -409,6 +411,223 @@ void OutputHierarchyTimers()
 #define OUTPUT_TIMERS()
 #define OUTPUT_HIERARCHYTIMERS()
 #endif
+
+
+
+#if ENABLE_COOKER_STATS
+#include "TypeHash.h"
+#include "CookingStatsModule.h"
+
+// gather stats about a cooked package
+namespace CookOnTheFlyStats
+{
+	ICookingStats* GetCookingStats()
+	{
+		static ICookingStats* CookingStats = nullptr;
+		static bool bInitialized = false;
+		if (bInitialized == false)
+		{
+			FCookingStatsModule* CookingStatsModule = FModuleManager::LoadModulePtr<FCookingStatsModule>(TEXT("CookingStats"));
+			if (CookingStatsModule)
+			{
+				CookingStats = &CookingStatsModule->Get();
+			}
+			bInitialized = true;
+		}
+		return CookingStats;
+	}
+
+	class FCookerPackageStats
+	{
+	private:
+		const FName Filename; 
+	public:
+		FCookerPackageStats(const FName& InFilename) : Filename(InFilename)
+		{
+
+		};
+		void AddTag(const FName& TagName, const FString& Value)
+		{
+			ICookingStats* CookingStats = GetCookingStats();
+			if (CookingStats)
+			{
+				CookingStats->AddTagValue(Filename, TagName, Value);
+			}
+		}
+
+		bool GetTag(const FName& TagName, FString& Value) const
+		{
+			ICookingStats* CookingStats = GetCookingStats();
+			if (CookingStats)
+			{
+				return CookingStats->GetTagValue(Filename, TagName, Value);
+			}
+			return false;
+		}
+
+		const FName& GetFilename() const 
+		{
+			return Filename;
+		}
+	};
+
+	struct FPackageStat
+	{
+		FPackageStat(const FName& InPackageName, const FName& InStatName)
+		{
+			PackageName = InPackageName;
+			StatName = InStatName;
+		}
+
+		bool operator==(const FPackageStat& InPackageStat) const
+		{
+			return (InPackageStat.PackageName == PackageName) &&
+				(InPackageStat.StatName == StatName);
+		}
+
+		FName PackageName;
+		FName StatName;
+
+	};
+
+	struct FStatInfo
+	{
+		int HitCount;
+		bool bStarted;
+		double StartTime;
+		double Duration;
+	};
+
+	TMap<FPackageStat, FStatInfo> StatInfoMap;
+
+	struct FNameCache
+	{
+	public:
+		FNameCache(const FName& BaseName) 
+		{
+			Duration = FName(*FString::Printf(TEXT("%s_Duration"), *BaseName.ToString()));
+			HitCount = FName(*FString::Printf(TEXT("%s_HitCount"), *BaseName.ToString()));
+		}
+		FName Duration;
+		FName HitCount;
+	};
+	TMap<FName, FNameCache> NameCache;
+
+	const FNameCache& GetNameCache(const FName& BaseName)
+	{
+		const FNameCache* Result = NameCache.Find(BaseName);
+		if (Result == nullptr)
+		{
+			Result = &NameCache.Add(BaseName, FNameCache(BaseName));
+		}
+		return *Result;
+	}
+
+	FStatInfo* GetStatInfo(const FName& StatName, const FName& PackageName)
+	{
+		FPackageStat PackageStat(PackageName, StatName);
+
+		FStatInfo *StatInfo = StatInfoMap.Find(PackageStat);
+		if (StatInfo == nullptr)
+		{
+			StatInfo = &StatInfoMap.Add(PackageStat);
+		}
+		return StatInfo;
+	}
+
+	void StartPackageStat(const FName& Name, const FCookerPackageStats& CookerPackageStats)
+	{
+		FStatInfo* StatInfo = GetStatInfo(Name, CookerPackageStats.GetFilename());
+
+		if (StatInfo->bStarted == false)
+		{
+			StatInfo->StartTime = FPlatformTime::Seconds();
+			StatInfo->bStarted = true;
+		}
+	}
+
+	void EndPackageStat(const FName& Name, FCookerPackageStats& CookerPackageStats)
+	{
+		FStatInfo* StatInfo = GetStatInfo(Name, CookerPackageStats.GetFilename());
+
+		if (StatInfo->bStarted)
+		{
+			++StatInfo->HitCount;
+			double Duration = FPlatformTime::Seconds() - StatInfo->StartTime;
+			StatInfo->Duration += Duration;
+			const FNameCache& CachedFNames = GetNameCache(Name);
+
+			CookerPackageStats.AddTag(CachedFNames.Duration, FString::Printf(TEXT("%fms"), StatInfo->Duration*1000.0f));
+			CookerPackageStats.AddTag(CachedFNames.HitCount, FString::Printf(TEXT("%d"), StatInfo->HitCount));
+
+			StatInfo->bStarted = false;
+		}
+	}
+
+
+
+	class FCookerPackageScopeStat
+	{
+	private:
+		const FName& StatName;
+		FCookerPackageStats& PackageStats;
+	public:
+		FCookerPackageScopeStat(const FName& InStatName, FCookerPackageStats& InPackageStats) : StatName(InStatName), PackageStats(InPackageStats)
+		{
+			Start();
+		}
+
+		~FCookerPackageScopeStat()
+		{
+			End();
+		}
+
+		void Start()
+		{
+			StartPackageStat(StatName, PackageStats);
+		}
+		void End()
+		{
+			EndPackageStat(StatName, PackageStats);
+		}
+	};
+}
+
+uint32 GetTypeHash(const CookOnTheFlyStats::FPackageStat& PackageStat)
+{
+	return HashCombine(GetTypeHash(PackageStat.PackageName), GetTypeHash(PackageStat.StatName));
+}
+
+#define START_PACKAGE(filename) CookOnTheFlyStats::FCookerPackageStats CookerPackageStats(filename);
+
+#define SCOPE_PACKAGE_STAT(name) \
+	static const FName NAME_##name(#name); \
+	CookOnTheFlyStats::FCookerPackageScopeStat CookerPackageScopeStat##name(NAME_##name, CookerPackageStats);
+
+// static const FName NAME_StartTime(TEXT("StartTime"));
+
+#define START_PACKAGE_STAT(name) \
+	static const FName NAME_##name(TEXT(#name)); \
+	CookOnTheFlyStats::StartPackageStat(NAME_##name, CookerPackageStats);
+
+// CookerPackageStats.AddTag(NAME_##nameStartTime, FString::Printf(TEXT("%f"),FPlatformTime::Seconds()));
+
+#define END_PACKAGE_STAT(name) \
+	static const FName NAME2_##name(TEXT(#name)); \
+	CookOnTheFlyStats::EndPackageStat(NAME2_##name, CookerPackageStats);
+
+
+#else
+
+
+#define START_PACKAGE(filename)
+#define SCOPE_PACKAGE_STAT(name)
+#define START_PACKAGE_STAT(name)
+#define END_PACKAGE_STAT(name)
+
+#endif
+
+
 ////////////////////////////////////////////////////////////////
 /// Cook on the fly server
 ///////////////////////////////////////////////////////////////
@@ -1378,6 +1597,8 @@ uint32 UCookOnTheFlyServer::TickCookOnTheSide( const float TimeSlice, uint32 &Co
 		bool bLastLoadWasMap = false;
 		FString LastLoadedMapName;
 
+		START_PACKAGE(ToBuild.GetFilename());
+
 		const FString BuildFilename = ToBuild.GetFilename().ToString();
 
 		if ( ToBuild.GetFilename() != CurrentReentryData.FileName )
@@ -1408,6 +1629,9 @@ uint32 UCookOnTheFlyServer::TickCookOnTheSide( const float TimeSlice, uint32 &Co
 
 		if ( bShouldCook ) // if we should cook the package then cook it otherwise add it to the list of already cooked packages below
 		{
+			SCOPE_PACKAGE_STAT(AllOfLoadPackage);
+
+
 			SCOPE_TIMER(AllOfLoadPackage);
 			UPackage *Package = NULL;
 			{
@@ -1561,6 +1785,7 @@ uint32 UCookOnTheFlyServer::TickCookOnTheSide( const float TimeSlice, uint32 &Co
 		// this will load more packages :)
 		if ( IsCookByTheBookMode() )
 		{
+			SCOPE_PACKAGE_STAT(ResolveRedirectors);
 			SCOPE_TIMER(ResolveRedirectors);
 			GRedirectCollector.ResolveStringAssetReference();
 		}
@@ -1633,6 +1858,8 @@ uint32 UCookOnTheFlyServer::TickCookOnTheSide( const float TimeSlice, uint32 &Co
 
 		if ( PackagesToSave.Num() )
 		{
+			START_PACKAGE_STAT(CacheForCookedPlatformData);
+
 			SCOPE_TIMER(CallBeginCacheForCookedPlatformData);
 			// cache the resources for this package for each platform
 			
@@ -1657,10 +1884,13 @@ uint32 UCookOnTheFlyServer::TickCookOnTheSide( const float TimeSlice, uint32 &Co
 			}
 		}
 
+		END_PACKAGE_STAT(CacheForCookedPlatformData);
+
 		int32 FirstUnsolicitedPackage = PackagesToSave.Num();
 
 		// generate a list of other packages which were loaded with this one
 		{
+			SCOPE_PACKAGE_STAT(UnsolicitedMarkup);
 			SCOPE_TIMER(UnsolicitedMarkup);
 
 			TArray<UObject *> ObjectsInOuter;
@@ -1726,6 +1956,8 @@ uint32 UCookOnTheFlyServer::TickCookOnTheSide( const float TimeSlice, uint32 &Co
 
 		bool bFinishedSave = true;
 
+		START_PACKAGE_STAT(SavePackageCacheForCookedPlatformData);
+
 		if ( PackagesToSave.Num() )
 		{
 			int32 OriginalPackagesToSaveCount = PackagesToSave.Num();
@@ -1840,6 +2072,7 @@ uint32 UCookOnTheFlyServer::TickCookOnTheSide( const float TimeSlice, uint32 &Co
 					GetObjectsWithOuter(Package, ObjectsInPackage);
 					if (FinishPackageCacheForCookedPlatformData(ObjectsInPackage) == false)
 					{
+						// add to back of queue
 						PackagesToSave.Add(Package);
 					}
 				}
@@ -1853,7 +2086,10 @@ uint32 UCookOnTheFlyServer::TickCookOnTheSide( const float TimeSlice, uint32 &Co
 					bShouldSaveAsync = false;
 				}
 
+				END_PACKAGE_STAT(SavePackageCacheForCookedPlatformData);
 				{
+					SCOPE_PACKAGE_STAT(SaveCookedPackage);
+
 					SCOPE_TIMER(SaveCookedPackage);
 					if (SaveCookedPackage(Package, SAVE_KeepGUID | (bShouldSaveAsync ? SAVE_Async : SAVE_None) | (IsCookFlagSet(ECookInitializationFlags::Unversioned) ? SAVE_Unversioned : 0), bWasUpToDate, AllTargetPlatformNames))
 					{
@@ -1870,6 +2106,8 @@ uint32 UCookOnTheFlyServer::TickCookOnTheSide( const float TimeSlice, uint32 &Co
 					}
 					Timer.SavedPackage();
 				}
+				START_PACKAGE_STAT(SavePackageCacheForCookedPlatformData);
+
 
 				if ( IsCookingInEditor() == false )
 				{
@@ -1912,6 +2150,7 @@ uint32 UCookOnTheFlyServer::TickCookOnTheSide( const float TimeSlice, uint32 &Co
 				}
 			}
 		}
+		END_PACKAGE_STAT(SavePackageCacheForCookedPlatformData);
 
 		// if after all this our requested file didn't get saved then we need to mark it as saved (this can happen if the file loaded didn't match the one we saved)
 		// for example say we request A.uasset which doesn't exist however A.umap exists, our load call will succeed as it will find A.umap, then the save will save A.umap
@@ -1945,6 +2184,7 @@ uint32 UCookOnTheFlyServer::TickCookOnTheSide( const float TimeSlice, uint32 &Co
 			break;
 		}
 	}
+	
 
 	OUTPUT_TIMERS();
 
@@ -2191,24 +2431,6 @@ bool UCookOnTheFlyServer::SaveCookedPackage( UPackage* Package, uint32 SaveFlags
 	{
 		FString Name = Package->GetPathName();
 
-		/*FString PkgFilename;
-		FDateTime DependentTimeStamp = FDateTime::MinValue();
-
-		// We always want to use the dependent time stamp when saving a cooked package...
-		// Iterative or not!
-		FString PkgFile;
-		
-
-		if (IsCookFlagSet(ECookInitializationFlags::Iterative) && FPackageName::DoesPackageExist(Name, NULL, &PkgFile))
-		{
-			PkgFilename = PkgFile;
-
-			if (GetPackageTimestamp(FPaths::GetBaseFilename(PkgFilename, false), DependentTimeStamp) == false)
-			{
-				UE_LOG(LogCook, Display, TEXT("Failed to find depedency timestamp for: %s"), *PkgFilename);
-			}
-		}*/
-
 		// Use SandboxFile to do path conversion to properly handle sandbox paths (outside of standard paths in particular).
 		Filename = ConvertToFullSandboxPath(*Filename, true);
 
diff --git a/Engine/Source/Runtime/Core/Public/Misc/CookingStatsInterface.h b/Engine/Source/Runtime/Core/Public/Misc/CookingStatsInterface.h
--- a/Engine/Source/Runtime/Core/Public/Misc/CookingStatsInterface.h
+++ b/Engine/Source/Runtime/Core/Public/Misc/CookingStatsInterface.h
@@ -12,15 +12,15 @@ class ICookingStatsModule : public IModuleInterface
 {
 public:
 	/**
-	 * Tries to gets a pointer to the active AssetRegistryInterface implementation. 
+	 * Tries to gets a pointer to the active CookingStatsModule implementation. 
 	 */
 	static inline ICookingStatsModule * GetPtr()
 	{
-		static FName CookingStats("CookingStats");
-		auto Module = FModuleManager::GetModulePtr<ICookingStatsModule>(CookingStats);
+		static FName NAME_CookingStats("CookingStats");
+		auto Module = FModuleManager::GetModulePtr<ICookingStatsModule>(NAME_CookingStats);
 		if (Module == nullptr)
 		{
-			Module = &FModuleManager::LoadModuleChecked<ICookingStatsModule>(CookingStats);
+			Module = &FModuleManager::LoadModuleChecked<ICookingStatsModule>(NAME_CookingStats);
 		}
 		return reinterpret_cast<ICookingStatsModule*>(Module);
 	}
diff --git a/Engine/Source/Runtime/CoreUObject/Private/UObject/SavePackage.cpp b/Engine/Source/Runtime/CoreUObject/Private/UObject/SavePackage.cpp
--- a/Engine/Source/Runtime/CoreUObject/Private/UObject/SavePackage.cpp
+++ b/Engine/Source/Runtime/CoreUObject/Private/UObject/SavePackage.cpp
@@ -22,43 +22,69 @@ static const FName WorldClassName = FName("World");
 
 #define UE_OUTPUTSTATSTOLOG 0
 
+namespace SavePackageStats
+{
 static const FName CookingStatsName("CookingStats");
 static const FName GSavePackageTransactionId(*FString::Printf(TEXT("SavePackageTransactionId%s"), *FGuid::NewGuid().ToString()));
 static uint32 GSavePackageTransactionNumber = 0;
+
+	ICookingStats* GetCookingStats()
+	{
+		static ICookingStats* CookingStats = nullptr;
+		static bool bInitialized = false;
+		if (bInitialized == false)
+		{
+			FCookingStatsModule* CookingStatsModule = FModuleManager::LoadModulePtr<FCookingStatsModule>(TEXT("CookingStats"));
+			if (CookingStatsModule)
+			{
+				CookingStats = &CookingStatsModule->Get();
+			}
+			bInitialized = true;
+		}
+		return CookingStats;
+	}
+
+
+	// static ICookingStats* CookingStats = GetCookingStats();
+
 void CookStatsStartStat(const FName& Key, const TCHAR *Filename)
 {
 #if UE_OUTPUTSTATSTOLOG
 	UE_LOG(LogSavePackage, Log, TEXT("Starting save package for %s"), Filename);
 #endif
-
-	
-	static FCookingStatsModule& CookingStatsModule = FModuleManager::LoadModuleChecked<FCookingStatsModule>(CookingStatsName);
-	static ICookingStats& CookingStats = CookingStatsModule.Get();
-
-	CookingStats.AddTagValue(Key, TEXT("Filename"), FString(Filename));
+		ICookingStats* CookingStats = GetCookingStats();
+		if (CookingStats)
+		{
+			CookingStats->AddTagValue(Key, TEXT("Filename"), FString(Filename));
+		}
 
 }
 
-void CookStatsAddStat(const FName& Key, const FName& Tag, const float TimeMilliseconds )
+	void CookStatsAddStat(const FName& Key, const FName& Tag, const float TimeMilliseconds)
 {
 #if UE_OUTPUTSTATSTOLOG
 	UE_LOG(LogSavePackage, Log, TEXT("TIMING: %s took %fms"), Tag, TimeMilliseconds);
 #endif
-	static FCookingStatsModule& CookingStatsModule = FModuleManager::LoadModuleChecked<FCookingStatsModule>(CookingStatsName);
-	static ICookingStats& CookingStats = CookingStatsModule.Get();
+		ICookingStats* CookingStats = GetCookingStats();
+		if (CookingStats)
+		{
+			CookingStats->AddTagValue(Key, Tag, FString::Printf(TEXT("%fms"), TimeMilliseconds));
+		}
+	}
 
-	CookingStats.AddTagValue(Key, Tag, FString::Printf(TEXT("%fms"), TimeMilliseconds));
-}
 
 
-#define UE_START_LOG_COOK_TIME(InFilename) double PreviousTime; double StartTime; PreviousTime = StartTime = FPlatformTime::Seconds(); const FName CookStatsKey = FName(GSavePackageTransactionId,++GSavePackageTransactionNumber); \
-	CookStatsStartStat(CookStatsKey, InFilename);
+}; // namespace SavePackageStats
+
+
+#define UE_START_LOG_COOK_TIME(InFilename) double PreviousTime; double StartTime; PreviousTime = StartTime = FPlatformTime::Seconds(); const FName CookStatsKey = FName(SavePackageStats::GSavePackageTransactionId,++SavePackageStats::GSavePackageTransactionNumber); \
+	SavePackageStats::CookStatsStartStat(CookStatsKey, InFilename);
 
 #define UE_LOG_COOK_TIME(TimeType) \
 	{\
 		double CurrentTime = FPlatformTime::Seconds(); \
-		const FName TagName = FName(TimeType);\
-		CookStatsAddStat(CookStatsKey, TagName, (CurrentTime - PreviousTime) * 1000.0f);\
+	const FName TagName = FName(TimeType); \
+	SavePackageStats::CookStatsAddStat(CookStatsKey, TagName, (CurrentTime - PreviousTime) * 1000.0f); \
 		PreviousTime = CurrentTime; \
 	}
 
