Show an OS dialog when trying to show a modal message box from something other than the game thread.

diff --git a/Engine/Source/Editor/UnrealEd/Private/EditorEngine.cpp b/Engine/Source/Editor/UnrealEd/Private/EditorEngine.cpp
--- a/Engine/Source/Editor/UnrealEd/Private/EditorEngine.cpp
+++ b/Engine/Source/Editor/UnrealEd/Private/EditorEngine.cpp
@@ -4009,7 +4009,7 @@ void UEditorEngine::CloseEntryPopupWindow()
 
 EAppReturnType::Type UEditorEngine::OnModalMessageDialog(EAppMsgType::Type InMessage, const FText& InText, const FText& InTitle)
 {
-	if( FSlateApplication::IsInitialized() && FSlateApplication::Get().CanAddModalWindow() )
+	if( IsInGameThread() && FSlateApplication::IsInitialized() && FSlateApplication::Get().CanAddModalWindow() )
 	{
 		return OpenMsgDlgInt(InMessage, InText, InTitle);
 	}
