Fixed lightmass exporter being deleted before the swarm connection is closed, causing occasional crashes when canceling a lighting build

diff --git a/Engine/Source/Editor/UnrealEd/Private/Lightmass/Lightmass.cpp b/Engine/Source/Editor/UnrealEd/Private/Lightmass/Lightmass.cpp
--- a/Engine/Source/Editor/UnrealEd/Private/Lightmass/Lightmass.cpp
+++ b/Engine/Source/Editor/UnrealEd/Private/Lightmass/Lightmass.cpp
@@ -2321,11 +2321,12 @@ FLightmassProcessor::FLightmassProcessor(const FStaticLightingSystem& InSystem,
 
 FLightmassProcessor::~FLightmassProcessor()
 {
+	// Note: the connection must be closed before deleting anything that SwarmCallback accesses
+	Swarm.CloseConnection();
+
 	delete Exporter;
 	delete Importer;
 
-	Swarm.CloseConnection();
-
 	for ( TMap<FGuid, FMappingImportHelper*>::TIterator It(ImportedMappings); It; ++It )
 	{
 		FMappingImportHelper* ImportData = It.Value();
