Sequencer: Pad framing on keyframes by adding next and previous unselected key.

diff --git a/Engine/Source/Editor/UnrealEd/Private/SCurveEditor.cpp b/Engine/Source/Editor/UnrealEd/Private/SCurveEditor.cpp
--- a/Engine/Source/Editor/UnrealEd/Private/SCurveEditor.cpp
+++ b/Engine/Source/Editor/UnrealEd/Private/SCurveEditor.cpp
@@ -769,7 +769,7 @@ void SCurveEditor::PaintCurve(TSharedPtr<FCurveViewModel> CurveViewModel, const
 			// Fade out curves that are not selected.
 			if (!CurveViewModel->bIsSelected && bAnyCurveViewModelsSelected)
 			{
-				Color *= FLinearColor(1.0f,1.0f,1.0f,0.5f); 
+				Color *= FLinearColor(1.0f,1.0f,1.0f,0.2f); 
 			}
 
 			// Fade out curves which are locked.
@@ -901,7 +901,7 @@ void SCurveEditor::PaintKeys(TSharedPtr<FCurveViewModel> CurveViewModel, FTrackS
 		int32 LayerToUse = IsSelected ? SelectedLayerId: LayerId;
 
 		// Fade out keys that are not selected and whose curve is not selected as well.
-		FLinearColor SelectionTint = !CurveViewModel->bIsSelected && !IsSelected && bAnyCurveViewModelsSelected ? FLinearColor(1.0f,1.0f,1.0f,0.5f) : FLinearColor(1.0f,1.0f,1.0f,1.0f); 
+		FLinearColor SelectionTint = !CurveViewModel->bIsSelected && !IsSelected && bAnyCurveViewModelsSelected ? FLinearColor(1.0f,1.0f,1.0f,0.2f) : FLinearColor(1.0f,1.0f,1.0f,1.0f); 
 
 		FSlateDrawElement::MakeBox(
 			OutDrawElements,
@@ -943,8 +943,8 @@ void SCurveEditor::PaintTangent( TSharedPtr<FCurveViewModel> CurveViewModel, FTr
 	bool LeaveTangentSelected = bTangentSelected && bIsLeaveSelected;
 	bool ArriveTangentSelected = bTangentSelected && bIsArrivalSelected;
 
-	FLinearColor LeaveSelectionTint = !CurveViewModel->bIsSelected && !LeaveTangentSelected && bAnyCurveViewModelsSelected ? FLinearColor(1.0f,1.0f,1.0f,0.5f) : FLinearColor(1.0f,1.0f,1.0f,1.0f); 
-	FLinearColor ArriveSelectionTint = !CurveViewModel->bIsSelected && !ArriveTangentSelected && bAnyCurveViewModelsSelected ? FLinearColor(1.0f,1.0f,1.0f,0.5f) : FLinearColor(1.0f,1.0f,1.0f,1.0f); 
+	FLinearColor LeaveSelectionTint = !CurveViewModel->bIsSelected && !LeaveTangentSelected && bAnyCurveViewModelsSelected ? FLinearColor(1.0f,1.0f,1.0f,0.2f) : FLinearColor(1.0f,1.0f,1.0f,1.0f); 
+	FLinearColor ArriveSelectionTint = !CurveViewModel->bIsSelected && !ArriveTangentSelected && bAnyCurveViewModelsSelected ? FLinearColor(1.0f,1.0f,1.0f,0.2f) : FLinearColor(1.0f,1.0f,1.0f,1.0f); 
 
 	//Add lines from tangent control point to 'key'
 	TArray<FVector2D> LinePoints;
@@ -2314,6 +2314,22 @@ void SCurveEditor::ZoomToFitHorizontal()
 				float KeyTime = SelectedKey.Curve->GetKeyTime(SelectedKey.KeyHandle);
 				InMin = FMath::Min(KeyTime, InMin);
 				InMax = FMath::Max(KeyTime, InMax);
+
+				FKeyHandle NextKeyHandle = SelectedKey.Curve->GetNextKey(SelectedKey.KeyHandle);
+				if (SelectedKey.Curve->IsKeyHandleValid(NextKeyHandle))
+				{
+					float NextKeyTime = SelectedKey.Curve->GetKeyTime(NextKeyHandle);
+					InMin = FMath::Min(NextKeyTime, InMin);
+					InMax = FMath::Max(NextKeyTime, InMax);
+				}
+
+				FKeyHandle PreviousKeyHandle = SelectedKey.Curve->GetPreviousKey(SelectedKey.KeyHandle);
+				if (SelectedKey.Curve->IsKeyHandleValid(PreviousKeyHandle))
+				{
+					float PreviousKeyTime = SelectedKey.Curve->GetKeyTime(PreviousKeyHandle);
+					InMin = FMath::Min(PreviousKeyTime, InMin);
+					InMax = FMath::Max(PreviousKeyTime, InMax);
+				}
 			}
 		}
 		else
@@ -2386,6 +2402,22 @@ void SCurveEditor::ZoomToFitVertical()
 				float KeyValue = SelectedKey.Curve->GetKeyValue(SelectedKey.KeyHandle);
 				InMin = FMath::Min(KeyValue, InMin);
 				InMax = FMath::Max(KeyValue, InMax);
+
+				FKeyHandle NextKeyHandle = SelectedKey.Curve->GetNextKey(SelectedKey.KeyHandle);
+				if (SelectedKey.Curve->IsKeyHandleValid(NextKeyHandle))
+				{
+					float NextKeyValue = SelectedKey.Curve->GetKeyValue(NextKeyHandle);
+					InMin = FMath::Min(NextKeyValue, InMin);
+					InMax = FMath::Max(NextKeyValue, InMax);
+				}
+
+				FKeyHandle PreviousKeyHandle = SelectedKey.Curve->GetPreviousKey(SelectedKey.KeyHandle);
+				if (SelectedKey.Curve->IsKeyHandleValid(PreviousKeyHandle))
+				{
+					float PreviousKeyValue = SelectedKey.Curve->GetKeyValue(PreviousKeyHandle);
+					InMin = FMath::Min(PreviousKeyValue, InMin);
+					InMax = FMath::Max(PreviousKeyValue, InMax);
+				}
 			}
 		}
 		else
diff --git a/Engine/Source/Runtime/Engine/Classes/Curves/CurveBase.h b/Engine/Source/Runtime/Engine/Classes/Curves/CurveBase.h
--- a/Engine/Source/Runtime/Engine/Classes/Curves/CurveBase.h
+++ b/Engine/Source/Runtime/Engine/Classes/Curves/CurveBase.h
@@ -307,6 +307,10 @@ public:
 	FRichCurveKey GetFirstKey() const;
 	FRichCurveKey GetLastKey() const;
 
+	/** Get the next or previous key given the key handle */
+	FKeyHandle GetNextKey(FKeyHandle KeyHandle) const;
+	FKeyHandle GetPreviousKey(FKeyHandle KeyHandle) const;
+
 	/** Get number of key in curve. */
 	virtual int32 GetNumKeys() const override;
 	
diff --git a/Engine/Source/Runtime/Engine/Private/CurveBase.cpp b/Engine/Source/Runtime/Engine/Private/CurveBase.cpp
--- a/Engine/Source/Runtime/Engine/Private/CurveBase.cpp
+++ b/Engine/Source/Runtime/Engine/Private/CurveBase.cpp
@@ -292,6 +292,34 @@ FRichCurveKey FRichCurve::GetLastKey() const
 	return Keys[Keys.Num()-1];
 }
 
+FKeyHandle FRichCurve::GetNextKey(FKeyHandle KeyHandle) const
+{
+	int32 KeyIndex = GetIndex(KeyHandle);
+	KeyIndex++;
+	if (Keys.IsValidIndex(KeyIndex))
+	{
+		return GetKeyHandle(KeyIndex);
+	}
+	else
+	{
+		return FKeyHandle();
+	}
+}
+
+FKeyHandle FRichCurve::GetPreviousKey(FKeyHandle KeyHandle) const
+{
+	int32 KeyIndex = GetIndex(KeyHandle);
+	KeyIndex--;
+	if (Keys.IsValidIndex(KeyIndex))
+	{
+		return GetKeyHandle(KeyIndex);
+	}
+	else
+	{
+		return FKeyHandle();
+	}
+}
+
 int32 FRichCurve::GetNumKeys() const
 {
 	return Keys.Num();
