Cooker compilation fix.

diff --git a/Engine/Source/Editor/UnrealEd/Private/CookOnTheFlyServer.cpp b/Engine/Source/Editor/UnrealEd/Private/CookOnTheFlyServer.cpp
--- a/Engine/Source/Editor/UnrealEd/Private/CookOnTheFlyServer.cpp
+++ b/Engine/Source/Editor/UnrealEd/Private/CookOnTheFlyServer.cpp
@@ -591,13 +591,14 @@ namespace CookOnTheFlyStats
 			EndPackageStat(StatName, PackageStats);
 		}
 	};
-}
 
-uint32 GetTypeHash(const CookOnTheFlyStats::FPackageStat& PackageStat)
-{
-	return HashCombine(GetTypeHash(PackageStat.PackageName), GetTypeHash(PackageStat.StatName));
+	uint32 GetTypeHash(const FPackageStat& PackageStat)
+	{
+		return HashCombine(GetTypeHash(PackageStat.PackageName), GetTypeHash(PackageStat.StatName));
+	}
 }
 
+
 #define START_PACKAGE(filename) CookOnTheFlyStats::FCookerPackageStats CookerPackageStats(filename);
 
 #define SCOPE_PACKAGE_STAT(name) \
@@ -4750,4 +4751,4 @@ bool UCookOnTheFlyServer::GetAllPackagesFromAssetRegistry( const FString& AssetR
 	return false;
 }
 
-#undef LOCTEXT_NAMESPACE
\ No newline at end of file
+#undef LOCTEXT_NAMESPACE
