Workaround for VS2015 compiler bug when using operator= on FStrings within lambdas.

diff --git a/Engine/Source/Editor/UnrealEd/Private/AutoReimport/AutoReimportTests.cpp b/Engine/Source/Editor/UnrealEd/Private/AutoReimport/AutoReimportTests.cpp
--- a/Engine/Source/Editor/UnrealEd/Private/AutoReimport/AutoReimportTests.cpp
+++ b/Engine/Source/Editor/UnrealEd/Private/AutoReimport/AutoReimportTests.cpp
@@ -650,7 +650,7 @@ bool FAutoReimportChangeExtensionsTest::RunTest(const FString& Parameters)
 				UE_LOG(LogAutoReimportTests, Error, TEXT("Incorrect action for added file %s."), SrcFilename3);
 				return;
 			}
-			else if (Changes[0].Filename.Get() != SrcFilename3)
+			else if (!Changes[0].Filename.Get().Equals(SrcFilename3))
 			{
 				UE_LOG(LogAutoReimportTests, Error, TEXT("Added file path reported incorrectly (%s != %s)."), *Changes[0].Filename.Get(), SrcFilename3);
 				return;
@@ -767,7 +767,7 @@ bool FAutoReimportWildcardFiltersTest::RunTest(const FString& Parameters)
 				{
 					UE_LOG(LogAutoReimportTests, Error, TEXT("Incorrect number of changes reported (%d != 1)."), Changes.Num());
 				}
-				else if (Changes[0].Filename.Get() != DstFilename3)
+				else if (!Changes[0].Filename.Get().Equals(DstFilename3))
 				{
 					UE_LOG(LogAutoReimportTests, Error, TEXT("Modified file path reported incorrectly (%s != %s)."), *Changes[0].Filename.Get(), DstFilename3);
 				}
