Fixing CIS non-unity errors: including and forward declaring types for UnrealEdMisc.h

diff --git a/Engine/Source/Editor/UnrealEd/Public/UnrealEdMisc.h b/Engine/Source/Editor/UnrealEd/Public/UnrealEdMisc.h
--- a/Engine/Source/Editor/UnrealEd/Public/UnrealEdMisc.h
+++ b/Engine/Source/Editor/UnrealEd/Public/UnrealEdMisc.h
@@ -2,7 +2,10 @@
 
 #pragma once
 
+#include "Editor.h" // for FEditorModeID
+
 class FPerformanceAnalyticsStats;
+class FTickableEditorObject;
 
 namespace EMapChangeType
 {
