Fix OBJ files no longer importing correctly.  Was trying to use the skeletal mesh reimport factory.

https://jira.ol.epicgames.net/browse/UE-18734

diff --git a/Engine/Source/Editor/UnrealEd/Private/Fbx/FbxFactory.cpp b/Engine/Source/Editor/UnrealEd/Private/Fbx/FbxFactory.cpp
--- a/Engine/Source/Editor/UnrealEd/Private/Fbx/FbxFactory.cpp
+++ b/Engine/Source/Editor/UnrealEd/Private/Fbx/FbxFactory.cpp
@@ -592,7 +592,9 @@ void UFbxFactory::CleanUp()
 
 bool UFbxFactory::FactoryCanImport(const FString& Filename)
 {
-	if(FPaths::GetExtension(Filename) == TEXT("fbx"))
+	const FString Extension = FPaths::GetExtension(Filename);
+
+	if( Extension == TEXT("fbx") || Extension == TEXT("obj") )
 	{
 		return true;
 	}
