fixed UE-18757 Texture import causes the texture to be compressed processed twice
was only for DDS_cubemaps, DDS_2d, HDR and IES

diff --git a/Engine/Source/Editor/UnrealEd/Private/Factories/EditorFactories.cpp b/Engine/Source/Editor/UnrealEd/Private/Factories/EditorFactories.cpp
--- a/Engine/Source/Editor/UnrealEd/Private/Factories/EditorFactories.cpp
+++ b/Engine/Source/Editor/UnrealEd/Private/Factories/EditorFactories.cpp
@@ -4046,9 +4046,6 @@ UTexture* UTextureFactory::ImportTexture(UClass* Class, UObject* InParent, FName
 
 			// for now we don't support mip map generation on cubemaps
 			TextureCube->MipGenSettings = TMGS_LeaveExistingMips;
-
-			// generates the mips from the source art
-			TextureCube->PostEditChange();
 		}
 
 		return TextureCube;
@@ -4089,11 +4086,9 @@ UTexture* UTextureFactory::ImportTexture(UClass* Class, UObject* InParent, FName
 
 				if(Texture->HasHDRSource())
 				{
+					// the loader can suggest a compression setting
 					Texture->CompressionSettings = TC_HDR;
 				}
-
-				// generates the mips from the source art
-				Texture->PostEditChange();
 			}
 
 			return Texture;
@@ -4121,8 +4116,8 @@ UTexture* UTextureFactory::ImportTexture(UClass* Class, UObject* InParent, FName
 				TSF_BGRE8,
 				HDRDDSLoadHelper.GetDDSDataPointer()
 				);
+			// the loader can suggest a compression setting
 			TextureCube->CompressionSettings = TC_HDR;
-			TextureCube->PostEditChange();
 		}
 
 		return TextureCube;
@@ -4158,7 +4153,6 @@ UTexture* UTextureFactory::ImportTexture(UClass* Class, UObject* InParent, FName
 				MipGenSettings = TMGS_NoMipmaps;
 				Texture->Brightness = IESLoadHelper.GetBrightness();
 				Texture->TextureMultiplier = Multiplier;
-				Texture->PostEditChange();
 			}
 
 			return Texture;
@@ -4364,7 +4358,13 @@ UObject* UTextureFactory::FactoryCreateBinary
 	Texture->CompressionNoAlpha		= NoAlpha;
 	Texture->DeferCompression		= bDeferCompression;
 	Texture->bDitherMipMapAlpha		= bDitherMipMapAlpha;
-	Texture->MipGenSettings			= MipGenSettings;
+	
+	if(Texture->MipGenSettings == TMGS_FromTextureGroup)
+	{
+		// unless the loader suggest a different setting
+		Texture->MipGenSettings = MipGenSettings;
+	}
+	
 	Texture->bPreserveBorder		= bPreserveBorder;
 
 	Texture->AssetImportData->Update(CurrentFilename);
