Added experimental settings for num worker cookers when using File->Package

diff --git a/Engine/Source/Editor/MainFrame/Private/Frame/MainFrameActions.cpp b/Engine/Source/Editor/MainFrame/Private/Frame/MainFrameActions.cpp
--- a/Engine/Source/Editor/MainFrame/Private/Frame/MainFrameActions.cpp
+++ b/Engine/Source/Editor/MainFrame/Private/Frame/MainFrameActions.cpp
@@ -592,6 +592,12 @@ void FMainFrameActionCallbacks::PackageProject( const FName InPlatformInfoName )
 	}
 #endif
 
+	int32 NumCookers = GetDefault<UEditorExperimentalSettings>()->MultiProcessCooking;
+	if (NumCookers > 0 )
+	{
+		OptionalParams += FString::Printf(TEXT(" -NumCookersToSpawn=%d"), NumCookers);
+	}
+
 	const bool bRunningDebug = FParse::Param(FCommandLine::Get(), TEXT("debug"));
 
 	if (bRunningDebug)
diff --git a/Engine/Source/Editor/UnrealEd/Classes/Settings/EditorExperimentalSettings.h b/Engine/Source/Editor/UnrealEd/Classes/Settings/EditorExperimentalSettings.h
--- a/Engine/Source/Editor/UnrealEd/Classes/Settings/EditorExperimentalSettings.h
+++ b/Engine/Source/Editor/UnrealEd/Classes/Settings/EditorExperimentalSettings.h
@@ -85,6 +85,9 @@ public:
 	UPROPERTY(EditAnywhere, config, Category = Cooking, meta = (DisplayName = "Iterative cooking for builds launched form the editor (launch on)"))
 	bool bIterativeCookingForLaunchOn;
 
+	UPROPERTY(EditAnywhere, config, Category = Cooking, meta = (DisplayName = "Use multiple processes when cooking (only affects File -> Package)"))
+	int32 MultiProcessCooking;
+
 	/** Enables Environment Queries editor */
 	UPROPERTY(EditAnywhere, config, Category = AI, meta = (DisplayName = "Environment Querying System"))
 	bool bEQSEditor;
