Fix inconsistency where source code projects launch on device with the same configuration the editor is running, but blueprint projects always launch on Development.
Added LaunchConfiguration setting to ULevelEditorPlaySettings to specify the configuration you want to launch on device with, defaulting to the same configuration as the editor.

diff --git a/Engine/Source/Editor/UnrealEd/Classes/Settings/LevelEditorPlaySettings.h b/Engine/Source/Editor/UnrealEd/Classes/Settings/LevelEditorPlaySettings.h
--- a/Engine/Source/Editor/UnrealEd/Classes/Settings/LevelEditorPlaySettings.h
+++ b/Engine/Source/Editor/UnrealEd/Classes/Settings/LevelEditorPlaySettings.h
@@ -91,6 +91,17 @@ enum EPlayOnBuildMode
 	PlayOnBuild_Default UMETA(DisplayName="Only Build Code Projects"),
 };
 
+/* Configuration to use when launching on device. */
+UENUM()
+enum EPlayOnLaunchConfiguration
+{
+	/** Launch on device with the same build configuration as the editor. */
+	LaunchConfig_Default UMETA(DisplayName = "Same as Editor"),
+	/** Launch on device with a Debug build configuration. */
+	LaunchConfig_Debug UMETA(DisplayName = "Debug"),
+	/** Launch on device with a Development build configuration. */
+	LaunchConfig_Development UMETA(DisplayName = "Development"),
+};
 
 /**
  * Holds information about a screen resolution to be used for playing.
@@ -119,7 +130,6 @@ public:
 	FString AspectRatio;
 };
 
-
 /**
  * Implements the Editor's play settings.
  */
@@ -174,7 +184,7 @@ public:
 
 public:
 
-	/** The width of the new view port window in pixels (0 = use the desktop's screen resolution). */
+	/** Whether to always have the PIE window on top of the parent windows. */
 	UPROPERTY(config, EditAnywhere, Category = PlayInNewWindow, meta = (ToolTip="Always have the PIE window on top of the parent windows."))
 	bool PIEAlwaysOnTop;
 
@@ -222,10 +232,14 @@ public:
 
 public:
 
-	/** The width of the new view port window in pixels (0 = use the desktop's screen resolution). */
+	/** Whether to build the game before launching on device. */
 	UPROPERTY(config, EditAnywhere, Category = PlayOnDevice)
 	TEnumAsByte<EPlayOnBuildMode> BuildGameBeforeLaunch;
 
+	/* Which build configuration to use when launching on device. */
+	UPROPERTY(config, EditAnywhere, Category = PlayOnDevice)
+	TEnumAsByte<EPlayOnLaunchConfiguration> LaunchConfiguration;
+
 	/** Whether to automatically recompile dirty Blueprints before launching */
 	UPROPERTY(config, EditAnywhere, Category=PlayOnDevice)
 	bool bAutoCompileBlueprintsOnLaunch;
diff --git a/Engine/Source/Editor/UnrealEd/Private/PlayLevel.cpp b/Engine/Source/Editor/UnrealEd/Private/PlayLevel.cpp
--- a/Engine/Source/Editor/UnrealEd/Private/PlayLevel.cpp
+++ b/Engine/Source/Editor/UnrealEd/Private/PlayLevel.cpp
@@ -1630,8 +1630,19 @@ void UEditorEngine::PlayUsingLauncher()
 		if ((bBuildType == EPlayOnBuildMode::PlayOnBuild_Always) || (bBuildType == PlayOnBuild_Default && (bPlayUsingLauncherHasCode) && bPlayUsingLauncherHasCompiler))
 		{
 			LauncherProfile->SetBuildGame(true);
+		}
 
-			// set the build configuration to be the same as the running editor
+		// set the build/launch configuration 
+		switch (PlayInSettings->LaunchConfiguration)
+		{
+		case LaunchConfig_Debug:
+			LauncherProfile->SetBuildConfiguration(EBuildConfigurations::Debug);
+			break;
+		case LaunchConfig_Development:
+			LauncherProfile->SetBuildConfiguration(EBuildConfigurations::Development);
+			break;
+		default:
+			// same as the running editor
 			FString ExeName = FUnrealEdMisc::Get().GetExecutableForCommandlets();
 			if (ExeName.Contains(TEXT("Debug")))
 			{
@@ -1641,8 +1652,9 @@ void UEditorEngine::PlayUsingLauncher()
 			{
 				LauncherProfile->SetBuildConfiguration(EBuildConfigurations::Development);
 			}
+			break;
 		}
-
+		
 		// select the quickest cook mode based on which in editor cook mode is enabled
 		bool bIncrimentalCooking = true;
 		LauncherProfile->AddCookedPlatform(PlayUsingLauncherDeviceId.Left(PlayUsingLauncherDeviceId.Find(TEXT("@"))));
diff --git a/Engine/Source/Editor/UnrealEd/Private/Settings/SettingsClasses.cpp b/Engine/Source/Editor/UnrealEd/Private/Settings/SettingsClasses.cpp
--- a/Engine/Source/Editor/UnrealEd/Private/Settings/SettingsClasses.cpp
+++ b/Engine/Source/Editor/UnrealEd/Private/Settings/SettingsClasses.cpp
@@ -351,6 +351,7 @@ ULevelEditorPlaySettings::ULevelEditorPlaySettings( const FObjectInitializer& Ob
 	RouteGamepadToSecondWindow = false;
 	AutoConnectToServer = true;
 	BuildGameBeforeLaunch = EPlayOnBuildMode::PlayOnBuild_Default;
+	LaunchConfiguration = EPlayOnLaunchConfiguration::LaunchConfig_Default;
 	bAutoCompileBlueprintsOnLaunch = true;
 	CenterNewWindow = true;
 	CenterStandaloneWindow = true;
