Fix for occasional crash when undoing/redoing brush transactions.

diff --git a/Engine/Source/Editor/UnrealEd/Private/EditorTransaction.cpp b/Engine/Source/Editor/UnrealEd/Private/EditorTransaction.cpp
--- a/Engine/Source/Editor/UnrealEd/Private/EditorTransaction.cpp
+++ b/Engine/Source/Editor/UnrealEd/Private/EditorTransaction.cpp
@@ -312,6 +312,7 @@ void FTransaction::Apply()
 		return (BAsComponent ? (BAsComponent->GetOwner() != &A) : true);
 	});
 
+	TArray<ULevel*> LevelsToCommitModelSurface;
 	NumModelsModified = 0;		// Count the number of UModels that were changed.
 	for (auto ChangedObjectIt : ChangedObjects)
 	{
@@ -322,6 +323,14 @@ void FTransaction::Apply()
 			FBSPOps::bspBuildBounds(Model);
 			++NumModelsModified;
 		}
+		
+		if (UModelComponent* ModelComponent = Cast<UModelComponent>(ChangedObject))
+		{
+			ULevel* Level = ModelComponent->GetTypedOuter<ULevel>();
+			check(Level);
+			LevelsToCommitModelSurface.AddUnique(Level);
+		}
+
 		TSharedPtr<ITransactionObjectAnnotation> ChangedObjectTransactionAnnotation = ChangedObjectIt.Value;
 		if (ChangedObjectTransactionAnnotation.IsValid())
 		{
@@ -333,6 +342,12 @@ void FTransaction::Apply()
 		}
 	}
 
+	// Commit model surfaces for unique levels within the transaction
+	for (ULevel* Level : LevelsToCommitModelSurface)
+	{
+		Level->CommitModelSurfaces();
+	}
+
 	// Flip it.
 	if (bFlip)
 	{
diff --git a/Engine/Source/Runtime/Engine/Classes/Engine/Level.h b/Engine/Source/Runtime/Engine/Classes/Engine/Level.h
--- a/Engine/Source/Runtime/Engine/Classes/Engine/Level.h
+++ b/Engine/Source/Runtime/Engine/Classes/Engine/Level.h
@@ -596,7 +596,7 @@ public:
 	/**
 	 * Commits changes made to the UModel's surfaces.
 	 */
-	void CommitModelSurfaces();
+	ENGINE_API void CommitModelSurfaces();
 
 	/**
 	 * Discards the cached data used to render the level's UModel.  Assumes that the
diff --git a/Engine/Source/Runtime/Engine/Private/Components/ModelComponent.cpp b/Engine/Source/Runtime/Engine/Private/Components/ModelComponent.cpp
--- a/Engine/Source/Runtime/Engine/Private/Components/ModelComponent.cpp
+++ b/Engine/Source/Runtime/Engine/Private/Components/ModelComponent.cpp
@@ -253,7 +253,6 @@ void UModelComponent::PostEditUndo()
 	if(ensure(Level))
 	{
 		Level->InvalidateModelSurface();
-//		Level->CommitModelSurfaces();
 	}
 	Super::PostEditUndo();
 }
