Fixed crashes when PIE objects end the transaction buffer.

diff --git a/Engine/Source/Editor/UnrealEd/Classes/Editor/TransBuffer.h b/Engine/Source/Editor/UnrealEd/Classes/Editor/TransBuffer.h
--- a/Engine/Source/Editor/UnrealEd/Classes/Editor/TransBuffer.h
+++ b/Engine/Source/Editor/UnrealEd/Classes/Editor/TransBuffer.h
@@ -110,6 +110,7 @@ public:
 	virtual void SetPrimaryUndoObject( UObject* Object ) override;
 	virtual ITransaction* CreateInternalTransaction() override;
 	virtual bool IsObjectInTransationBuffer( const UObject* Object ) const override;
+	virtual bool ContainsPieObject() const override;
 	virtual bool IsActive() override
 	{
 		return ActiveCount > 0;
diff --git a/Engine/Source/Editor/UnrealEd/Classes/Editor/Transactor.h b/Engine/Source/Editor/UnrealEd/Classes/Editor/Transactor.h
--- a/Engine/Source/Editor/UnrealEd/Classes/Editor/Transactor.h
+++ b/Engine/Source/Editor/UnrealEd/Classes/Editor/Transactor.h
@@ -257,6 +257,9 @@ protected:
 		/** Used by GC to collect referenced objects. */
 		void AddReferencedObjects( FReferenceCollector& Collector );
 
+		/** @return True if this record contains a reference to a pie object */
+		bool ContainsPieObject() const;
+
 		/** Transfers data from an array. */
 		class FReader : public FArchiveUObject
 		{
@@ -490,6 +493,10 @@ public:
 	int32 GetRecordCount() const;
 
 	const UObject* GetPrimaryObject() const { return PrimaryObject; }
+
+	/** @return True if this record contains a reference to a pie object */
+	bool ContainsPieObject() const;
+
 	/**
 	 * Outputs the contents of the ObjectMap to the specified output device.
 	 */
@@ -652,8 +659,11 @@ class UTransactor : public UObject
 	 */
 	virtual void SetPrimaryUndoObject( UObject* Object ) PURE_VIRTUAL(UTransactor::MakePrimaryUndoObject,);
 
+	/** Checks if a specific object is referenced by the transaction buffer */
 	virtual bool IsObjectInTransationBuffer( const UObject* Object ) const { return false; }
 
+	/** @return True if this record contains a reference to a pie object */
+	virtual bool ContainsPieObject() const { return false; }
 
 	// @todo document
 	virtual ITransaction* CreateInternalTransaction() PURE_VIRTUAL(UTransactor::CreateInternalTransaction,return NULL;);
diff --git a/Engine/Source/Editor/UnrealEd/Private/EditorTransaction.cpp b/Engine/Source/Editor/UnrealEd/Private/EditorTransaction.cpp
--- a/Engine/Source/Editor/UnrealEd/Private/EditorTransaction.cpp
+++ b/Engine/Source/Editor/UnrealEd/Private/EditorTransaction.cpp
@@ -124,6 +124,20 @@ int32 FTransaction::GetRecordCount() const
 	return Records.Num();
 }
 
+bool FTransaction::ContainsPieObject() const
+{
+	for( const FObjectRecord& Record : Records )
+	{
+		if( Record.ContainsPieObject() )
+		{
+			return true;
+		}
+	}
+
+	return false;
+}
+
+
 void FTransaction::RemoveRecords( int32 Count /* = 1  */ )
 {
 	if ( Count > 0 && Records.Num() >= Count )
@@ -215,6 +229,29 @@ void FTransaction::FObjectRecord::AddReferencedObjects( FReferenceCollector& Col
 	}
 }
 
+bool FTransaction::FObjectRecord::ContainsPieObject() const
+{
+	{
+		UObject* Obj = Object.Get();
+
+		if(Obj && (Obj->GetOutermost()->PackageFlags & PKG_PlayInEditor) != 0)
+		{
+			return true;
+		}
+	}
+
+	for( const FReferencedObject& ReferencedObject : ReferencedObjects )
+	{
+		const UObject* Obj = ReferencedObject.GetObject();
+		if( Obj && (Obj->GetOutermost()->PackageFlags & PKG_PlayInEditor) != 0 )
+		{
+			return true;
+		}
+	}
+
+	return false;
+}
+
 void FTransaction::AddReferencedObjects( FReferenceCollector& Collector )
 {
 	for( FObjectRecord& ObjectRecord : Records )
@@ -770,5 +807,18 @@ bool UTransBuffer::IsObjectInTransationBuffer( const UObject* Object ) const
 		TransactionObjects.Reset();
 	}
 
+	return false;
+}
+
+bool UTransBuffer::ContainsPieObject() const
+{
+	for( const FTransaction& Transaction : UndoBuffer )
+	{
+		if( Transaction.ContainsPieObject() )
+		{
+			return true;
+		}
+	}
+
 	return false;
 }
\ No newline at end of file
diff --git a/Engine/Source/Editor/UnrealEd/Private/PlayLevel.cpp b/Engine/Source/Editor/UnrealEd/Private/PlayLevel.cpp
--- a/Engine/Source/Editor/UnrealEd/Private/PlayLevel.cpp
+++ b/Engine/Source/Editor/UnrealEd/Private/PlayLevel.cpp
@@ -250,8 +250,18 @@ void UEditorEngine::EndPlayMap()
 			}
 		}
 	}
-	// Garbage Collect
-	CollectGarbage( GARBAGE_COLLECTION_KEEPFLAGS );
+
+	// Clean up any PIE world objects
+	{
+		// The trans buffer should never have a PIE object in it.  If it does though, as a s
+		if( GEditor->Trans->ContainsPieObject() )
+		{
+			GEditor->ResetTransaction( NSLOCTEXT("UnrealEd", "TransactionContainedPIEObject", "A PIE object was in the transaction buffer and had to be destroyed") );
+		}
+
+		// Garbage Collect
+		CollectGarbage(GARBAGE_COLLECTION_KEEPFLAGS);
+	}
 
 	// Make sure that all objects in the temp levels were entirely garbage collected.
 	for( FObjectIterator ObjectIt; ObjectIt; ++ObjectIt )
