Changes to spawning/duplicating actors in SIE. Now actors are always spawned non-permanently into the SIE world. Fixed issues with duplicating actors with inner objects in SIE.

diff --git a/Engine/Source/Editor/LevelEditor/Private/LevelEditorActions.cpp b/Engine/Source/Editor/LevelEditor/Private/LevelEditorActions.cpp
--- a/Engine/Source/Editor/LevelEditor/Private/LevelEditorActions.cpp
+++ b/Engine/Source/Editor/LevelEditor/Private/LevelEditorActions.cpp
@@ -1551,7 +1551,21 @@ bool FLevelEditorActionCallbacks::PasteHere_CanExecute()
 
 void FLevelEditorActionCallbacks::ExecuteExecCommand( FString Command )
 {
-	GUnrealEd->Exec( GetWorld(), *Command );
+	UWorld* OldWorld = nullptr;
+
+	// The play world needs to be selected if it exists
+	if (GIsEditor && GEditor->PlayWorld && !GIsPlayInEditorWorld)
+	{
+		OldWorld = SetPlayInEditorWorld(GEditor->PlayWorld);
+	}
+
+	GUnrealEd->Exec(GetWorld(), *Command);
+
+	// Restore the old world if there was one
+	if (OldWorld)
+	{
+		RestoreEditorWorld(OldWorld);
+	}
 }
 
 void FLevelEditorActionCallbacks::OnSelectAllActorsOfClass( bool bArchetype )
diff --git a/Engine/Source/Editor/UnrealEd/Private/AssetSelection.cpp b/Engine/Source/Editor/UnrealEd/Private/AssetSelection.cpp
--- a/Engine/Source/Editor/UnrealEd/Private/AssetSelection.cpp
+++ b/Engine/Source/Editor/UnrealEd/Private/AssetSelection.cpp
@@ -419,6 +419,14 @@ static AActor* PrivateAddActor( UObject* Asset, UActorFactory* Factory, bool Sel
 		return nullptr;
 	}
 
+	UWorld* OldWorld = nullptr;
+
+	// The play world needs to be selected if it exists
+	if (GIsEditor && GEditor->PlayWorld && !GIsPlayInEditorWorld)
+	{
+		OldWorld = SetPlayInEditorWorld(GEditor->PlayWorld);
+	}
+
 	// For Brushes/Volumes, use the default brush as the template rather than the factory default actor
 	if (NewActorTemplate->IsA(ABrush::StaticClass()) && GWorld->GetDefaultBrush() != nullptr)
 	{
@@ -439,20 +447,15 @@ static AActor* PrivateAddActor( UObject* Asset, UActorFactory* Factory, bool Sel
 	ULevel* DesiredLevel = GWorld->GetCurrentLevel();
 
 	// Don't spawn the actor if the current level is locked.
-	if ( FLevelUtils::IsLevelLocked(DesiredLevel) )
+	if (FLevelUtils::IsLevelLocked(DesiredLevel))
 	{
-		FNotificationInfo Info( NSLOCTEXT("UnrealEd", "Error_OperationDisallowedOnLockedLevel", "The requested operation could not be completed because the level is locked.") );
+		FNotificationInfo Info(NSLOCTEXT("UnrealEd", "Error_OperationDisallowedOnLockedLevel", "The requested operation could not be completed because the level is locked."));
 		Info.ExpireDuration = 3.0f;
 		FSlateNotificationManager::Get().AddNotification(Info);
-		return nullptr;
 	}
-
+	else
 	{
-		FScopedTransaction Transaction( NSLOCTEXT("UnrealEd", "CreateActor", "Create Actor") );
-		if ( !(ObjectFlags & RF_Transactional) )
-		{
-			Transaction.Cancel();
-		}
+		FScopedTransaction Transaction( NSLOCTEXT("UnrealEd", "CreateActor", "Create Actor"), (ObjectFlags & RF_Transactional) != 0 );
 
 		// Create the actor.
 		Actor = Factory->CreateActor( Asset, DesiredLevel, ActorTransform, ObjectFlags, Name );
@@ -467,9 +470,9 @@ static AActor* PrivateAddActor( UObject* Asset, UActorFactory* Factory, bool Sel
 			Actor->InvalidateLightingCache();
 			Actor->PostEditChange();
 		}
-	}
 
-	GEditor->RedrawLevelEditingViewports();
+		GEditor->RedrawLevelEditingViewports();
+	}
 
 	if ( Actor )
 	{
@@ -477,6 +480,12 @@ static AActor* PrivateAddActor( UObject* Asset, UActorFactory* Factory, bool Sel
 		ULevel::LevelDirtiedEvent.Broadcast();
 	}
 
+	// Restore the old world if there was one
+	if (OldWorld)
+	{
+		RestoreEditorWorld(OldWorld);
+	}
+
 	return Actor;
 }
 
diff --git a/Engine/Source/Editor/UnrealEd/Private/EditorEngine.cpp b/Engine/Source/Editor/UnrealEd/Private/EditorEngine.cpp
--- a/Engine/Source/Editor/UnrealEd/Private/EditorEngine.cpp
+++ b/Engine/Source/Editor/UnrealEd/Private/EditorEngine.cpp
@@ -4121,6 +4121,14 @@ AActor* UEditorEngine::UseActorFactory( UActorFactory* Factory, const FAssetData
 	//Load Asset
 	UObject* Asset = AssetData.GetAsset();
 
+	UWorld* OldWorld = nullptr;
+
+	// The play world needs to be selected if it exists
+	if (GIsEditor && GEditor->PlayWorld && !GIsPlayInEditorWorld)
+	{
+		OldWorld = SetPlayInEditorWorld(GEditor->PlayWorld);
+	}
+
 	AActor* Actor = NULL;
 	if ( bIsAllowedToCreateActor )
 	{
@@ -4171,20 +4179,21 @@ AActor* UEditorEngine::UseActorFactory( UActorFactory* Factory, const FAssetData
 					ULevel::LevelDirtiedEvent.Broadcast();
 				}
 			}
-			else
-			{
-				return NULL;
-			}
 		}
 		else
 		{
 			FNotificationInfo Info( NSLOCTEXT("UnrealEd", "Error_OperationDisallowedOnLockedLevel", "The requested operation could not be completed because the level is locked.") );
 			Info.ExpireDuration = 3.0f;
 			FSlateNotificationManager::Get().AddNotification(Info);
-			return NULL;
 		}
 	}
 
+	// Restore the old world if there was one
+	if (OldWorld)
+	{
+		RestoreEditorWorld(OldWorld);
+	}
+
 	return Actor;
 }
 
diff --git a/Engine/Source/Editor/UnrealEd/Private/Factories/EditorFactories.cpp b/Engine/Source/Editor/UnrealEd/Private/Factories/EditorFactories.cpp
--- a/Engine/Source/Editor/UnrealEd/Private/Factories/EditorFactories.cpp
+++ b/Engine/Source/Editor/UnrealEd/Private/Factories/EditorFactories.cpp
@@ -552,6 +552,11 @@ UObject* ULevelFactory::FactoryCreateText
 						SpawnInfo.Name = ActorUniqueName;
 						SpawnInfo.Template = Archetype;
 						SpawnInfo.SpawnCollisionHandlingOverride = ESpawnActorCollisionHandlingMethod::AlwaysSpawn;
+						if (GEditor->bIsSimulatingInEditor)
+						{
+							// During SIE, we don't want to run construction scripts on a BP until it is completely constructed
+							SpawnInfo.bDeferConstruction = true;
+						}
 						AActor* NewActor = World->SpawnActor( TempClass, nullptr, nullptr, SpawnInfo );
 						
 						if( NewActor )
diff --git a/Engine/Source/Runtime/Engine/Private/ActorEditor.cpp b/Engine/Source/Runtime/Engine/Private/ActorEditor.cpp
--- a/Engine/Source/Runtime/Engine/Private/ActorEditor.cpp
+++ b/Engine/Source/Runtime/Engine/Private/ActorEditor.cpp
@@ -25,7 +25,8 @@ void AActor::PreEditChange(UProperty* PropertyThatWillChange)
 		BPGC->UnbindDynamicDelegatesForProperty(this, ObjProp);
 	}
 
-	if ( ReregisterComponentsWhenModified() )
+	// During SIE, allow components to be unregistered here, and then reregistered and reconstructed in PostEditChangeProperty.
+	if (GEditor->bIsSimulatingInEditor || ReregisterComponentsWhenModified())
 	{
 		UnregisterAllComponents();
 	}
@@ -42,7 +43,9 @@ void AActor::PostEditChangeProperty(FPropertyChangedEvent& PropertyChangedEvent)
 	
 	const bool bTransformationChanged = (PropertyName == Name_RelativeLocation || PropertyName == Name_RelativeRotation || PropertyName == Name_RelativeScale3D);
 
-	if ( ReregisterComponentsWhenModified() )
+	// During SIE, allow components to reregistered and reconstructed in PostEditChangeProperty.
+	// This is essential as construction is deferred during spawning / duplication when in SIE.
+	if (GEditor->bIsSimulatingInEditor || ReregisterComponentsWhenModified())
 	{
 		// In the Undo case we have an annotation storing information about constructed components and we do not want
 		// to improperly apply out of date changes so we need to skip registration of all blueprint created components
