Pkginfo commandlet can now open packages from outside of the current project by specifying the absolute package path

diff --git a/Engine/Source/Editor/UnrealEd/Private/Commandlets/PackageUtilities.cpp b/Engine/Source/Editor/UnrealEd/Private/Commandlets/PackageUtilities.cpp
--- a/Engine/Source/Editor/UnrealEd/Private/Commandlets/PackageUtilities.cpp
+++ b/Engine/Source/Editor/UnrealEd/Private/Commandlets/PackageUtilities.cpp
@@ -820,6 +820,20 @@ namespace
 	{ EXPORTSORT_ExportIndex, EXPORTSORT_ExportSize, EXPORTSORT_OuterPathname, EXPORTSORT_ObjectPathname };
 }
 
+/** Given a package filename, creates a linker and a temporary package. The filename does not need to point to a package under the current project content folder */
+FLinkerLoad* CreateLinkerForFilename(const FString& InFilename)
+{
+	FString TempPackageName;
+	TempPackageName = FPaths::Combine(TEXT("/Temp"), *FPaths::GetPath(InFilename.Mid(InFilename.Find(TEXT(":"), ESearchCase::CaseSensitive) + 1)), *FPaths::GetBaseFilename(InFilename));
+	UPackage* Package = FindObjectFast<UPackage>(nullptr, *TempPackageName);
+	if (!Package)
+	{
+		Package = CreatePackage(nullptr, *TempPackageName);
+	}
+	FLinkerLoad* Linker = FLinkerLoad::CreateLinker(Package, *InFilename, LOAD_NoVerify);
+	return Linker;
+}
+
 /**
  * Writes information about the linker to the log.
  *
@@ -1413,6 +1427,7 @@ int32 UPkgInfoCommandlet::Main( const FString& Params )
 				for ( int32 FileIndex = 0; FileIndex < PerTokenFilesInPath.Num(); FileIndex++ )
 				{
 					PerTokenFilesInPath[FileIndex] = FPaths::GetPath(WildcardPath) / PerTokenFilesInPath[FileIndex];
+					FPaths::NormalizeFilename(PerTokenFilesInPath[FileIndex]);
 				}
 			}
 
@@ -1428,7 +1443,7 @@ int32 UPkgInfoCommandlet::Main( const FString& Params )
 
 	for( int32 FileIndex = 0; FileIndex < FilesInPath.Num(); FileIndex++ )
 	{
-		const FString &Filename = FilesInPath[FileIndex];
+		FString Filename = FPaths::ConvertRelativePathToFull(FilesInPath[FileIndex]);
 
 		{
 			// reset the loaders for the packages we want to load so that we don't find the wrong version of the file
@@ -1442,7 +1457,7 @@ int32 UPkgInfoCommandlet::Main( const FString& Params )
 		}
 
 		BeginLoad();
-		auto Linker = GetPackageLinker( NULL, *Filename, LOAD_NoVerify, NULL, NULL );
+		FLinkerLoad* Linker = CreateLinkerForFilename(Filename);
 		EndLoad();
 
 		if( Linker )
