[NO NON-PII COMMIT MESSAGE]

diff --git a/Engine/Source/Editor/UnrealEd/Private/Kismet2/BlueprintEditorUtils.cpp b/Engine/Source/Editor/UnrealEd/Private/Kismet2/BlueprintEditorUtils.cpp
--- a/Engine/Source/Editor/UnrealEd/Private/Kismet2/BlueprintEditorUtils.cpp
+++ b/Engine/Source/Editor/UnrealEd/Private/Kismet2/BlueprintEditorUtils.cpp
@@ -7167,9 +7167,10 @@ void FBlueprintEditorUtils::PostEditChangeBlueprintActors(UBlueprint* Blueprint,
 		// Get the selected Actor set in the level editor context
 		bool bEditorSelectionChanged = false;
 		const USelection* CurrentEditorActorSelection = GEditor ? GEditor->GetSelectedActors() : nullptr;
+		const bool bIncludeDerivedClasses = false;
 
 		TArray<UObject*> MatchingBlueprintObjects;
-		GetObjectsOfClass(Blueprint->GeneratedClass, MatchingBlueprintObjects, true, (RF_ClassDefaultObject | RF_PendingKill));
+		GetObjectsOfClass(Blueprint->GeneratedClass, MatchingBlueprintObjects, bIncludeDerivedClasses, (RF_ClassDefaultObject | RF_PendingKill));
 
 		for (auto ObjIt : MatchingBlueprintObjects)
 		{
