Added additional output to cooker to help track down bug.

diff --git a/Engine/Source/Editor/UnrealEd/Private/CookOnTheFlyServer.cpp b/Engine/Source/Editor/UnrealEd/Private/CookOnTheFlyServer.cpp
--- a/Engine/Source/Editor/UnrealEd/Private/CookOnTheFlyServer.cpp
+++ b/Engine/Source/Editor/UnrealEd/Private/CookOnTheFlyServer.cpp
@@ -2656,6 +2656,8 @@ void UCookOnTheFlyServer::Initialize( ECookMode::Type DesiredCookMode, ECookInit
 	GConfig->GetInt( TEXT("CookSettings"), TEXT("MaxMemoryAllowance"), MaxMemoryAllowanceInMB, GEditorIni );
 	MaxMemoryAllowance = MaxMemoryAllowanceInMB * 1024LL * 1024LL;
 	
+	UE_LOG(LogCook, Display, TEXT("Max memory allowance for cook %dmb, MaxMemoryAllowanceInMB"));
+
 	if ( IsCookByTheBookMode() )
 	{
 		CookByTheBookOptions = new FCookByTheBookOptions();
