Editor: Changing EMapChangeType to enum class, so that it can be forward-declared

diff --git a/Engine/Source/Editor/LevelEditor/Private/LevelEditor.cpp b/Engine/Source/Editor/LevelEditor/Private/LevelEditor.cpp
--- a/Engine/Source/Editor/LevelEditor/Private/LevelEditor.cpp
+++ b/Engine/Source/Editor/LevelEditor/Private/LevelEditor.cpp
@@ -384,7 +384,7 @@ void FLevelEditorModule::BroadcastTakeHighResScreenShots( )
 	TakeHighResScreenShotsEvent.Broadcast();
 }
 
-void FLevelEditorModule::BroadcastMapChanged( UWorld* World, EMapChangeType::Type MapChangeType )
+void FLevelEditorModule::BroadcastMapChanged( UWorld* World, EMapChangeType MapChangeType )
 {
 	MapChangedEvent.Broadcast( World, MapChangeType );
 }
diff --git a/Engine/Source/Editor/LevelEditor/Private/SLevelViewport.cpp b/Engine/Source/Editor/LevelEditor/Private/SLevelViewport.cpp
--- a/Engine/Source/Editor/LevelEditor/Private/SLevelViewport.cpp
+++ b/Engine/Source/Editor/LevelEditor/Private/SLevelViewport.cpp
@@ -945,7 +945,7 @@ TSharedRef< SWidget > SLevelViewport::BuildViewportDragDropContextMenu()
 	return ViewportContextMenuBuilder.MakeWidget();
 }
 
-void SLevelViewport::OnMapChanged( UWorld* World, EMapChangeType::Type MapChangeType )
+void SLevelViewport::OnMapChanged( UWorld* World, EMapChangeType MapChangeType )
 {
 	if( World && ( ( World == GetWorld() ) || ( World->EditorViews[LevelViewportClient->ViewportType].CamUpdated ) ) )
 	{
diff --git a/Engine/Source/Editor/LevelEditor/Public/LevelEditor.h b/Engine/Source/Editor/LevelEditor/Public/LevelEditor.h
--- a/Engine/Source/Editor/LevelEditor/Public/LevelEditor.h
+++ b/Engine/Source/Editor/LevelEditor/Public/LevelEditor.h
@@ -171,7 +171,7 @@ public:
 	 * 
 	 * @param World	The new world
 	 */
-	virtual void BroadcastMapChanged( UWorld* World, EMapChangeType::Type MapChangeType );
+	virtual void BroadcastMapChanged( UWorld* World, EMapChangeType MapChangeType );
 
 	/** Called when an edit command is executed on one or more components in the world */
 	virtual void BroadcastComponentsEdited();
@@ -185,7 +185,7 @@ public:
 	virtual FRedrawLevelEditingViewportsEvent& OnRedrawLevelEditingViewports() { return RedrawLevelEditingViewportsEvent; }
 
 	/** Called when a new map is loaded */
-	DECLARE_EVENT_TwoParams( FLevelEditorModule, FMapChangedEvent, UWorld*, EMapChangeType::Type );
+	DECLARE_EVENT_TwoParams( FLevelEditorModule, FMapChangedEvent, UWorld*, EMapChangeType );
 	virtual FMapChangedEvent& OnMapChanged() { return MapChangedEvent; }
 
 	/** Called when an edit command is executed on components in the world */
diff --git a/Engine/Source/Editor/LevelEditor/Public/SLevelViewport.h b/Engine/Source/Editor/LevelEditor/Public/SLevelViewport.h
--- a/Engine/Source/Editor/LevelEditor/Public/SLevelViewport.h
+++ b/Engine/Source/Editor/LevelEditor/Public/SLevelViewport.h
@@ -567,7 +567,7 @@ private:
 	/**
 	 * Called when a map is changed (loaded,saved,new map, etc)
 	 */
-	void OnMapChanged( UWorld* World, EMapChangeType::Type MapChangeType );
+	void OnMapChanged( UWorld* World, EMapChangeType MapChangeType );
 
 	/** Called in response to an actor being deleted in the level */
 	void OnLevelActorsRemoved(AActor* InActor);
diff --git a/Engine/Source/Editor/Sequencer/Private/Sequencer.cpp b/Engine/Source/Editor/Sequencer/Private/Sequencer.cpp
--- a/Engine/Source/Editor/Sequencer/Private/Sequencer.cpp
+++ b/Engine/Source/Editor/Sequencer/Private/Sequencer.cpp
@@ -458,7 +458,7 @@ void FSequencer::SpawnOrDestroyPuppetObjects( TSharedRef<FMovieSceneInstance> Mo
 	}
 }
 
-void FSequencer::NotifyMapChanged( class UWorld* NewWorld, EMapChangeType::Type MapChangeType )
+void FSequencer::NotifyMapChanged( class UWorld* NewWorld, EMapChangeType MapChangeType )
 {
 	// @todo Sequencer Sub-MovieScenes Needs more investigation of what to spawn
 	// Destroy our puppets because the world is going away.  We probably don't have to do this (the actors will
diff --git a/Engine/Source/Editor/Sequencer/Private/Sequencer.h b/Engine/Source/Editor/Sequencer/Private/Sequencer.h
--- a/Engine/Source/Editor/Sequencer/Private/Sequencer.h
+++ b/Engine/Source/Editor/Sequencer/Private/Sequencer.h
@@ -111,7 +111,7 @@ public:
 	virtual void KeyProperty(FKeyPropertyParams KeyPropertyParams) override;
 	virtual TSharedRef<ISequencerObjectBindingManager> GetObjectBindingManager() const override;
 	virtual FSequencerSelection& GetSelection() override;
-	virtual void NotifyMapChanged(class UWorld* NewWorld, EMapChangeType::Type MapChangeType) override;
+	virtual void NotifyMapChanged(class UWorld* NewWorld, EMapChangeType MapChangeType) override;
 
 	/** @return The current view range */
 	virtual FAnimatedRange GetViewRange() const override;
diff --git a/Engine/Source/Editor/Sequencer/Private/SequencerAssetEditor.cpp b/Engine/Source/Editor/Sequencer/Private/SequencerAssetEditor.cpp
--- a/Engine/Source/Editor/Sequencer/Private/SequencerAssetEditor.cpp
+++ b/Engine/Source/Editor/Sequencer/Private/SequencerAssetEditor.cpp
@@ -96,7 +96,7 @@ void FSequencerAssetEditor::InitSequencerAssetEditor( const EToolkitMode::Type M
 			
 		// We need to find out when the user loads a new map, because we might need to re-create puppet actors
 		// when previewing a MovieScene
-		LevelEditorModule.OnMapChanged().AddSP(Sequencer.ToSharedRef(), &FSequencerAssetEditor::HandleMapChanged);
+		LevelEditorModule.OnMapChanged().AddRaw(this, &FSequencerAssetEditor::HandleMapChanged);
 	}
 }
 
@@ -214,7 +214,7 @@ void FSequencerAssetEditor::AddComponentTrack( UActorComponent* Component )
 }
 
 
-void FSequencerAssetEditor::HandleMapChanged(class UWorld* NewWorld, EMapChangeType::Type MapChangeType)
+void FSequencerAssetEditor::HandleMapChanged(class UWorld* NewWorld, EMapChangeType MapChangeType)
 {
 	Sequencer->NotifyMapChanged(NewWorld, MapChangeType);
 }
diff --git a/Engine/Source/Editor/Sequencer/Private/SequencerAssetEditor.h b/Engine/Source/Editor/Sequencer/Private/SequencerAssetEditor.h
--- a/Engine/Source/Editor/Sequencer/Private/SequencerAssetEditor.h
+++ b/Engine/Source/Editor/Sequencer/Private/SequencerAssetEditor.h
@@ -46,7 +46,7 @@ private:
 	TSharedRef<FExtender> GetContextSensitiveSequencerExtender( const TSharedRef<FUICommandList> CommandList, const TArray<UObject*> ContextSensitiveObjects );
 	void ExtendSequencerAddTrackMenu( FMenuBuilder& AddTrackMenuBuilder, TArray<UObject*> ContextObjects );
 	void AddComponentTrack( UActorComponent* Component );
-	void HandleMapChanged(class UWorld* NewWorld, EMapChangeType::Type MapChangeType);
+	void HandleMapChanged(class UWorld* NewWorld, EMapChangeType MapChangeType);
 
 private:
 
diff --git a/Engine/Source/Editor/Sequencer/Public/ISequencer.h b/Engine/Source/Editor/Sequencer/Public/ISequencer.h
--- a/Engine/Source/Editor/Sequencer/Public/ISequencer.h
+++ b/Engine/Source/Editor/Sequencer/Public/ISequencer.h
@@ -159,5 +159,5 @@ public:
 
 	virtual FSequencerSelection& GetSelection() = 0;
 
-	virtual void NotifyMapChanged(class UWorld* NewWorld, EMapChangeType::Type MapChangeType) = 0;
+	virtual void NotifyMapChanged(class UWorld* NewWorld, EMapChangeType MapChangeType) = 0;
 };
diff --git a/Engine/Source/Editor/UnrealEd/Private/UnrealEdMisc.cpp b/Engine/Source/Editor/UnrealEd/Private/UnrealEdMisc.cpp
--- a/Engine/Source/Editor/UnrealEd/Private/UnrealEdMisc.cpp
+++ b/Engine/Source/Editor/UnrealEd/Private/UnrealEdMisc.cpp
@@ -1563,7 +1563,7 @@ void FUnrealEdMisc::CancelPerformanceSurvey()
 	LevelEditor.OnMapChanged().Remove( OnMapChangedDelegateHandle );
 }
 
-void FUnrealEdMisc::OnMapChanged( UWorld* World, EMapChangeType::Type MapChangeType )
+void FUnrealEdMisc::OnMapChanged( UWorld* World, EMapChangeType MapChangeType )
 {
 	if (bIsSurveyingPerformance)
 	{
diff --git a/Engine/Source/Editor/UnrealEd/Public/UnrealEdMisc.h b/Engine/Source/Editor/UnrealEd/Public/UnrealEdMisc.h
--- a/Engine/Source/Editor/UnrealEd/Public/UnrealEdMisc.h
+++ b/Engine/Source/Editor/UnrealEd/Public/UnrealEdMisc.h
@@ -7,20 +7,22 @@
 class FPerformanceAnalyticsStats;
 class FTickableEditorObject;
 
-namespace EMapChangeType
+
+enum class EMapChangeType
 {
-	enum Type
-	{
-		/** Map has just been loaded*/
-		LoadMap,
-		/** Map is about to be saved*/
-		SaveMap,
-		/** A new map is loaded*/
-		NewMap,
-		/** The world is about to be torn down */
-		TearDownWorld,
-	};
-}
+	/** Map has just been loaded*/
+	LoadMap,
+
+	/** Map is about to be saved*/
+	SaveMap,
+
+	/** A new map is loaded*/
+	NewMap,
+
+	/** The world is about to be torn down */
+	TearDownWorld,
+};
+
 
 /** The public interface for the unreal editor misc singleton. */
 class UNREALED_API FUnrealEdMisc
@@ -247,7 +249,7 @@ private:
 	/**
 	 * Called when a map is changed (loaded,saved,new map, etc)
 	 */
-	void OnMapChanged( UWorld* World, EMapChangeType::Type MapChangeType );
+	void OnMapChanged( UWorld* World, EMapChangeType MapChangeType );
 
 	/** Called when the input manager records a user-defined chord */
 	void OnUserDefinedChordChanged(const FUICommandInfo& CommandInfo);
