Editor: Added missing forward declarations, headers, documentation

diff --git a/Engine/Source/Developer/AssetTools/Public/AssetTypeActions_Base.h b/Engine/Source/Developer/AssetTools/Public/AssetTypeActions_Base.h
--- a/Engine/Source/Developer/AssetTools/Public/AssetTypeActions_Base.h
+++ b/Engine/Source/Developer/AssetTools/Public/AssetTypeActions_Base.h
@@ -5,12 +5,16 @@
 #include "AssetToolsModule.h"
 #include "Toolkits/SimpleAssetEditor.h"
 #include "EditorReimportHandler.h"
+#include "Editor/UnrealEd/Classes/Settings/EditorLoadingSavingSettings.h"
+
 
 /** A base class for all AssetTypeActions. Provides helper functions useful for many types. Deriving from this class is optional. */
 class FAssetTypeActions_Base : public IAssetTypeActions
 {
 public:
-	// Begin IAssetTypeActions implementation
+
+	// IAssetTypeActions interface
+
 	virtual bool HasActions( const TArray<UObject*>& InObjects ) const override
 	{
 		return false;
@@ -68,8 +72,8 @@ public:
 
 	virtual void PerformAssetDiff(UObject* OldAsset, UObject* NewAsset, const struct FRevisionInfo& OldRevision, const struct FRevisionInfo& NewRevision) const override
 	{
-		check(OldAsset != NULL);
-		check(NewAsset != NULL);
+		check(OldAsset != nullptr);
+		check(NewAsset != nullptr);
 
 		// Dump assets to temp text files
 		FString OldTextFilename = DumpAssetToTempFile(OldAsset);
@@ -99,8 +103,6 @@ public:
 	{
 	}
 
-	// End IAssetTypeActions implementation
-
 protected:
 
 	// Here are some convenience functions for common asset type actions logic
diff --git a/Engine/Source/Editor/PropertyEditor/Public/IPropertyTypeCustomization.h b/Engine/Source/Editor/PropertyEditor/Public/IPropertyTypeCustomization.h
--- a/Engine/Source/Editor/PropertyEditor/Public/IPropertyTypeCustomization.h
+++ b/Engine/Source/Editor/PropertyEditor/Public/IPropertyTypeCustomization.h
@@ -2,17 +2,24 @@
 
 #pragma once
 
+#include "EditorStyleSet.h"
+
+
 class IPropertyHandle;
 class FDetailWidgetRow;
+class IPropertyHandle;
 class IPropertyTypeCustomizationUtils;
+class IPropertyUtilities;
 class IDetailChildrenBuilder;
 class FAssetThumbnailPool;
 class IPropertyUtilities;
 
+
 /**
  * Base class for property type customizations
  */
-class IPropertyTypeCustomization : public TSharedFromThis<IPropertyTypeCustomization>
+class IPropertyTypeCustomization
+	: public TSharedFromThis<IPropertyTypeCustomization>
 {
 public:
 	virtual ~IPropertyTypeCustomization() {}
@@ -64,7 +71,7 @@ public:
 	/**
 	*	@return the utilities various widgets need access to certain features of PropertyDetails
 	*/
-	virtual TSharedPtr<IPropertyUtilities> GetPropertyUtilities() const { return NULL; }
+	virtual TSharedPtr<IPropertyUtilities> GetPropertyUtilities() const { return nullptr; }
 };
 
 
diff --git a/Engine/Source/Editor/UnrealEd/Public/Toolkits/AssetEditorManager.h b/Engine/Source/Editor/UnrealEd/Public/Toolkits/AssetEditorManager.h
--- a/Engine/Source/Editor/UnrealEd/Public/Toolkits/AssetEditorManager.h
+++ b/Engine/Source/Editor/UnrealEd/Public/Toolkits/AssetEditorManager.h
@@ -8,14 +8,25 @@
 #include "UnrealEdMessages.h"
 
 
-// This class keeps track of a currently open asset editor; allowing it to be
-// brought into focus, closed, etc..., without concern for how the editor was
-// implemented.
+class FMessageEndpoint;
+class FReferenceCollector;
+class IMessageContext;
+class IToolkitHost;
+class SNotificationItem;
+class UObject;
+
+
+/**
+ * This class keeps track of a currently open asset editor; allowing it to be
+ * brought into focus, closed, etc..., without concern for how the editor was
+ * implemented.
+ */
 class UNREALED_API IAssetEditorInstance
 {
 public:
+
 	virtual FName GetEditorName() const = 0;
-	virtual void FocusWindow(UObject* ObjectToFocusOn = NULL) = 0;
+	virtual void FocusWindow(UObject* ObjectToFocusOn = nullptr) = 0;
 	virtual bool CloseWindow() = 0;
 	virtual bool IsPrimaryEditor() const = 0;
 	virtual void InvokeTab(const struct FTabId& TabId) = 0;
@@ -29,7 +40,8 @@ public:
  *
  * @todo toolkit: Merge this functionality into FToolkitManager
  */
-class UNREALED_API FAssetEditorManager : public FGCObject
+class UNREALED_API FAssetEditorManager
+	: public FGCObject
 {
 public:
 
@@ -40,24 +52,24 @@ public:
 	void OnExit();
 
 	/** 
-	* Tries to open an editor for the specified asset.  Returns true if the asset is open in an editor.
-	* If the file is already open in an editor, it will not create another editor window but instead bring it to front
-	*/
-	bool OpenEditorForAsset(UObject* Asset, const EToolkitMode::Type ToolkitMode = EToolkitMode::Standalone, TSharedPtr< class IToolkitHost > OpenedFromLevelEditor = TSharedPtr< IToolkitHost >());
+	 * Tries to open an editor for the specified asset.  Returns true if the asset is open in an editor.
+	 * If the file is already open in an editor, it will not create another editor window but instead bring it to front
+	 */
+	bool OpenEditorForAsset(UObject* Asset, const EToolkitMode::Type ToolkitMode = EToolkitMode::Standalone, TSharedPtr<IToolkitHost> OpenedFromLevelEditor = TSharedPtr<IToolkitHost>());
 
 	/** 
-	* Tries to open an editor for all of the specified assets. 
-	* If any of the assets are already open, it will not create a new editor for them.
-	* If all assets are of the same type, the supporting AssetTypeAction (if it exists) is responsible for the details of how to handle opening multiple assets at once.
-	*/
-	bool OpenEditorForAssets(const TArray< UObject* >& Assets, const EToolkitMode::Type ToolkitMode = EToolkitMode::Standalone, TSharedPtr< class IToolkitHost > OpenedFromLevelEditor = TSharedPtr< IToolkitHost >());
+	 * Tries to open an editor for all of the specified assets. 
+	 * If any of the assets are already open, it will not create a new editor for them.
+	 * If all assets are of the same type, the supporting AssetTypeAction (if it exists) is responsible for the details of how to handle opening multiple assets at once.
+	 */
+	bool OpenEditorForAssets(const TArray<UObject*>& Assets, const EToolkitMode::Type ToolkitMode = EToolkitMode::Standalone, TSharedPtr<IToolkitHost> OpenedFromLevelEditor = TSharedPtr<IToolkitHost>());
 
 	/** Opens editors for the supplied assets (via OpenEditorForAsset) */
 	void OpenEditorsForAssets(const TArray<FString>& AssetsToOpen);
 
 	/** Returns the primary editor if one is already open for the specified asset.
-	* If there is one open and bFocusIfOpen is true, that editor will be brought to the foreground and focused if possible.
-	*/
+	 * If there is one open and bFocusIfOpen is true, that editor will be brought to the foreground and focused if possible.
+	 */
 	IAssetEditorInstance* FindEditorForAsset(UObject* Asset, bool bFocusIfOpen);
 
 	/** Returns all editors currently opened for the specified asset */
@@ -108,7 +120,7 @@ private:
 private:
 
 	/** Handles FAssetEditorRequestOpenAsset messages. */
-	void HandleRequestOpenAssetMessage(const FAssetEditorRequestOpenAsset& Message, const IMessageContextRef& Context);
+	void HandleRequestOpenAssetMessage(const FAssetEditorRequestOpenAsset& Message, const TSharedRef<IMessageContext, ESPMode::ThreadSafe>& Context);
 
 	/** Opens an asset by path */
 	void OpenEditorForAsset(const FString& AssetPathName);
@@ -153,24 +165,16 @@ private:
 		}
 	};
 
-	/**
-	 * Holds the opened assets.
-	 */
+	/** Holds the opened assets. */
 	TMultiMap<UObject*, IAssetEditorInstance*> OpenedAssets;
 
-	/**
-	 * Holds the opened editors.
-	 */
+	/** Holds the opened editors. */
 	TMultiMap<IAssetEditorInstance*, UObject*> OpenedEditors;
 
-	/**
-	 * Holds the times that editors were opened.
-	 */
+	/** Holds the times that editors were opened. */
 	TMap<IAssetEditorInstance*, FOpenedEditorTime> OpenedEditorTimes;
 
-	/**
-	 * Holds the cumulative time editors have been open by type.
-	 */
+	/** Holds the cumulative time editors have been open by type. */
 	TMap<FName, FAssetEditorAnalyticInfo> EditorUsageAnalytics;
 
 private:
@@ -179,7 +183,7 @@ private:
 	static FAssetEditorManager* Instance;
 
 	/** Holds the messaging endpoint. */
-	FMessageEndpointPtr MessageEndpoint;
+	TSharedPtr<FMessageEndpoint, ESPMode::ThreadSafe> MessageEndpoint;
 
 	/** Holds a delegate to be invoked when the widget ticks. */
 	FTickerDelegate TickDelegate;
diff --git a/Engine/Source/Editor/UnrealEd/Public/Toolkits/AssetEditorToolkit.h b/Engine/Source/Editor/UnrealEd/Public/Toolkits/AssetEditorToolkit.h
--- a/Engine/Source/Editor/UnrealEd/Public/Toolkits/AssetEditorToolkit.h
+++ b/Engine/Source/Editor/UnrealEd/Public/Toolkits/AssetEditorToolkit.h
@@ -3,12 +3,25 @@
 #pragma once
 
 #include "BaseToolkit.h"
+#include "TabManager.h"
 #include "Toolkits/AssetEditorManager.h"		// For IAssetEditorInstance derive
 
 
+class FExtender;
+class FUICommandList;
+class IToolkit;
+class IToolkitHost;
+class SBorder;
+class SStandaloneAssetEditorToolkitHost;
+class SWidget;
+class UObject;
+struct FSlateBrush;
+struct FTabId;
+
 
 DECLARE_DELEGATE_RetVal( bool, FRequestAssetEditorClose );
 
+
 /**
  * The location of the asset editor toolkit tab
  * Note: These values are serialized into an ini file as an int32
@@ -21,10 +34,14 @@ enum class EAssetEditorToolkitTabLocation : int32
 	Standalone,
 };
 
+
 /**
  * Base class for toolkits that are used for asset editing (abstract)
  */
-class UNREALED_API FAssetEditorToolkit : public IAssetEditorInstance, public FBaseToolkit, public TSharedFromThis< FAssetEditorToolkit >
+class UNREALED_API FAssetEditorToolkit
+	: public IAssetEditorInstance
+	, public FBaseToolkit
+	, public TSharedFromThis<FAssetEditorToolkit>
 {
 
 public:
@@ -45,15 +62,15 @@ public:
 	 * @param	ObjectToEdit			The object to edit
 	 * @param	bInIsToolbarFocusable	Whether the buttons on the default toolbar can receive keyboard focus
 	 */
-	virtual void InitAssetEditor(const EToolkitMode::Type Mode, const TSharedPtr< class IToolkitHost >& InitToolkitHost, const FName AppIdentifier, const TSharedRef<FTabManager::FLayout>& StandaloneDefaultLayout, const bool bCreateDefaultStandaloneMenu, const bool bCreateDefaultToolbar, const TArray<UObject*>& ObjectsToEdit, const bool bInIsToolbarFocusable = false);
-	virtual void InitAssetEditor(const EToolkitMode::Type Mode, const TSharedPtr< class IToolkitHost >& InitToolkitHost, const FName AppIdentifier, const TSharedRef<FTabManager::FLayout>& StandaloneDefaultLayout, const bool bCreateDefaultStandaloneMenu, const bool bCreateDefaultToolbar, UObject* ObjectToEdit, const bool bInIsToolbarFocusable = false);
+	virtual void InitAssetEditor(const EToolkitMode::Type Mode, const TSharedPtr<IToolkitHost>& InitToolkitHost, const FName AppIdentifier, const TSharedRef<FTabManager::FLayout>& StandaloneDefaultLayout, const bool bCreateDefaultStandaloneMenu, const bool bCreateDefaultToolbar, const TArray<UObject*>& ObjectsToEdit, const bool bInIsToolbarFocusable = false);
+	virtual void InitAssetEditor(const EToolkitMode::Type Mode, const TSharedPtr<IToolkitHost>& InitToolkitHost, const FName AppIdentifier, const TSharedRef<FTabManager::FLayout>& StandaloneDefaultLayout, const bool bCreateDefaultStandaloneMenu, const bool bCreateDefaultToolbar, UObject* ObjectToEdit, const bool bInIsToolbarFocusable = false);
 
 	/** Virtual destructor, so that we can clean up our app when destroyed */
 	virtual ~FAssetEditorToolkit();
 
 	/** IToolkit interface */
-	virtual void RegisterTabSpawners(const TSharedRef<class FTabManager>& TabManager) override;
-	virtual void UnregisterTabSpawners(const TSharedRef<class FTabManager>& TabManager) override;
+	virtual void RegisterTabSpawners(const TSharedRef<FTabManager>& TabManager) override;
+	virtual void UnregisterTabSpawners(const TSharedRef<FTabManager>& TabManager) override;
 	virtual bool IsAssetEditor() const override;
 	virtual const TArray< UObject* >* GetObjectsCurrentlyBeingEdited() const override;
 	virtual FName GetToolkitFName() const override = 0;				// Must implement in derived class!
@@ -65,11 +82,11 @@ public:
 
 	/** IAssetEditorInstance interface */
 	virtual FName GetEditorName() const override;
-	virtual void FocusWindow(UObject* ObjectToFocusOn = NULL) override;
+	virtual void FocusWindow(UObject* ObjectToFocusOn = nullptr) override;
 	virtual bool CloseWindow() override;
 	virtual bool IsPrimaryEditor() const override { return true; };
-	virtual void InvokeTab(const struct FTabId& TabId) override;
-	virtual TSharedPtr<class FTabManager> GetAssociatedTabManager() override;
+	virtual void InvokeTab(const FTabId& TabId) override;
+	virtual TSharedPtr<FTabManager> GetAssociatedTabManager() override;
 	virtual double GetLastActivationTime() override;
 
 	/**
@@ -106,7 +123,7 @@ public:
 	/** Applies the passed in layout (or the saved user-modified version if available).  Must be called after InitAssetEditor. */
 	void RestoreFromLayout( const TSharedRef<FTabManager::FLayout>& NewLayout );
 
-	/** @return Returns this asset editor's tab manager object.  May be NULL for non-standalone toolkits */
+	/** @return Returns this asset editor's tab manager object.  May be nullptr for non-standalone toolkits */
 	TSharedPtr<FTabManager> GetTabManager()
 	{
 		return TabManager;
@@ -122,10 +139,10 @@ public:
 	virtual void PostRegenerateMenusAndToolbars() { }
 	
 	// Called when another toolkit (such as a ed mode toolkit) is being hosted in this asset editor toolkit
-	virtual void OnToolkitHostingStarted(const TSharedRef< class IToolkit >& Toolkit) {}
+	virtual void OnToolkitHostingStarted(const TSharedRef<IToolkit>& Toolkit) {}
 
 	// Called when another toolkit (such as a ed mode toolkit) is no longer being hosted in this asset editor toolkit
-	virtual void OnToolkitHostingFinished(const TSharedRef< class IToolkit >& Toolkit) {}
+	virtual void OnToolkitHostingFinished(const TSharedRef<IToolkit>& Toolkit) {}
 
 	/** Adds or removes extenders to the default menu or the toolbar menu this asset editor */
 	void AddMenuExtender(TSharedPtr<FExtender> Extender);
@@ -260,8 +277,8 @@ protected:
 	bool bCheckDirtyOnAssetSave;
 
 private:
-	/** The toolkit standalone host; may be NULL */
-	TWeakPtr< class SStandaloneAssetEditorToolkitHost > StandaloneHost;
+	/** The toolkit standalone host; may be nullptr */
+	TWeakPtr<SStandaloneAssetEditorToolkitHost> StandaloneHost;
 
 	/** Static: World centric toolkit host to use for the next created asset editing toolkit */
 	static TWeakPtr< IToolkitHost > PreviousWorldCentricToolkitHostForNewAssetEditor;
@@ -271,16 +288,16 @@ private:
 	TArray<UObject*> EditingObjects;
 	
 	/** Asset Editor Default Toolbar */
-	TSharedPtr<class SWidget> Toolbar;
+	TSharedPtr<SWidget> Toolbar;
 
 	/** The widget that will house the default Toolbar widget */
 	TSharedPtr<SBorder> ToolbarWidgetContent;
 
 	/** The menu extenders to populate the main toolbar with */
-	TArray< TSharedPtr<FExtender> > ToolbarExtenders;
+	TArray<TSharedPtr<FExtender>> ToolbarExtenders;
 
 	/** Additional widgets to be added to the toolbar */
-	TArray< TSharedRef<SWidget> > ToolbarWidgets;
+	TArray<TSharedRef<SWidget>> ToolbarWidgets;
 
 	/** Whether the buttons on the default toolbar can receive keyboard focus */
 	bool bIsToolbarFocusable;
@@ -325,7 +342,6 @@ private:
 };
 
 
-
 /** Indicates that a class has a default menu that is extensible */
 class IHasMenuExtensibility
 {
@@ -333,6 +349,7 @@ public:
 	virtual TSharedPtr<FExtensibilityManager> GetMenuExtensibilityManager() = 0;
 };
 
+
 /** Indicates that a class has a default toolbar that is extensible */
 class IHasToolBarExtensibility
 {
diff --git a/Engine/Source/Editor/UnrealEd/Public/Toolkits/BaseToolkit.h b/Engine/Source/Editor/UnrealEd/Public/Toolkits/BaseToolkit.h
--- a/Engine/Source/Editor/UnrealEd/Public/Toolkits/BaseToolkit.h
+++ b/Engine/Source/Editor/UnrealEd/Public/Toolkits/BaseToolkit.h
@@ -5,12 +5,21 @@
 #include "IToolkit.h"
 
 
+class FTabManager;
+class FUICommandList;
+class FWorkspaceItem;
+class IToolkitHost;
+class SDockableTab;
+class SWidget;
+class UObject;
+
+
 /**
- * Base class for all toolkits (abstract)
+ * Base class for all toolkits (abstract).
  */
-class UNREALED_API FBaseToolkit : public IToolkit
+class UNREALED_API FBaseToolkit
+	: public IToolkit
 {
-
 public:
 
 	/** FBaseToolkit constructor */
@@ -19,22 +28,27 @@ public:
 	/** Virtual destructor */
 	virtual ~FBaseToolkit();
 
-	/** IToolkit interface */
+public:
+
+	// IToolkit interface
+
 	virtual FName GetToolkitContextFName() const override;
-	virtual bool ProcessCommandBindings( const FKeyEvent& InKeyEvent ) const override;
+	virtual bool ProcessCommandBindings(const FKeyEvent& InKeyEvent) const override;
 	virtual bool IsHosted() const override;
-	virtual const TSharedRef< class IToolkitHost > GetToolkitHost() const override;
-	virtual const TMap< EToolkitTabSpot::Type, TArray< TWeakPtr< SDockableTab > > >& GetToolkitTabsInSpots() const override;
+	virtual const TSharedRef<IToolkitHost> GetToolkitHost() const override;
+	virtual const TMap<EToolkitTabSpot::Type, TArray<TWeakPtr<SDockableTab>>>& GetToolkitTabsInSpots() const override;
 	virtual void BringToolkitToFront() override;
-	virtual TSharedPtr<class SWidget> GetInlineContent() const override;
+	virtual TSharedPtr<SWidget> GetInlineContent() const override;
 	virtual bool IsBlueprintEditor() const override;
 	virtual TSharedRef<FWorkspaceItem> GetWorkspaceMenuCategory() const override { return WorkspaceMenuCategory.ToSharedRef(); }
 
+public:
+
 	/** @return	Returns true if this is a world-centric asset editor.  That is, the user is editing the asset inline in a Level Editor app. */
 	bool IsWorldCentricAssetEditor() const;	
 
 	/** @returns Returns our toolkit command list */
-	const TSharedRef< FUICommandList > GetToolkitCommands() const
+	const TSharedRef<FUICommandList> GetToolkitCommands() const
 	{
 		return ToolkitCommands;
 	}
@@ -45,7 +59,7 @@ public:
 	 * @param	TabToAdd		The dockable tab object to add
 	 * @param	TabSpot			Which spot to spawn this tab into
 	 */
-	void AddToolkitTab( const TSharedRef< SDockableTab >& TabToAdd, const EToolkitTabSpot::Type TabSpot );
+	void AddToolkitTab(const TSharedRef<SDockableTab>& TabToAdd, const EToolkitTabSpot::Type TabSpot);
 
 protected:
 
@@ -57,47 +71,52 @@ protected:
 	    make them more easy to distinguish compared to other tabs. */
 	FLinearColor GetTabColorScale() const;
 
-
 protected:
 
 	/** Asset editing mode, set at creation-time and never changes */
 	EToolkitMode::Type ToolkitMode;
 
 	/** List of UI commands for this toolkit.  This should be filled in by the derived class! */
-	TSharedRef< FUICommandList > ToolkitCommands;
+	TSharedRef<FUICommandList> ToolkitCommands;
 
 	/** The host application for this editor.  If editing in world-centric mode, this is the level level editor that we're editing the asset within 
 	    Use GetToolkitHost() method to access this member. */
-	TWeakPtr< class IToolkitHost > ToolkitHost;
+	TWeakPtr<IToolkitHost> ToolkitHost;
 
 	/** Map of toolkit tab spots to known tabs (these are weak pointers and may be invalid after tabs are closed.) */
-	TMap< EToolkitTabSpot::Type, TArray< TWeakPtr< SDockableTab > > > ToolkitTabsInSpots;
+	TMap<EToolkitTabSpot::Type, TArray<TWeakPtr<SDockableTab>>> ToolkitTabsInSpots;
 
 	/** The workspace menu category of this toolkit */
 	TSharedPtr<FWorkspaceItem> WorkspaceMenuCategory;
 };
 
 
-
 /**
  * Base class for all editor mode toolkits.
  */
-class UNREALED_API FModeToolkit : public FBaseToolkit, public TSharedFromThis<FModeToolkit>
+class UNREALED_API FModeToolkit
+	: public FBaseToolkit
+	, public TSharedFromThis<FModeToolkit>
 {
 public:
-	virtual void RegisterTabSpawners(const TSharedRef<class FTabManager>& TabManager) override {}
-	virtual void UnregisterTabSpawners(const TSharedRef<class FTabManager>& TabManager) override {}
 
 	/** Initializes the mode toolkit */
-	virtual void Init(const TSharedPtr< class IToolkitHost >& InitToolkitHost);
+	virtual void Init(const TSharedPtr<IToolkitHost>& InitToolkitHost);
+
+public:
+
+	// FBaseToolkit overrides
+
+	virtual void RegisterTabSpawners(const TSharedRef<FTabManager>& TabManager) override { }
+	virtual void UnregisterTabSpawners(const TSharedRef<FTabManager>& TabManager) override { }
 
-	/** IToolkit interface */
+public:
+
+	// IToolkit interface
 	virtual FText GetToolkitName() const override { return GetBaseToolkitName(); }
 	virtual FText GetToolkitToolTipText() const override { return GetBaseToolkitName(); }
 	virtual FString GetWorldCentricTabPrefix() const override;
 	virtual bool IsAssetEditor() const override;
-	virtual const TArray< UObject* >* GetObjectsCurrentlyBeingEdited() const override;
+	virtual const TArray<UObject*>* GetObjectsCurrentlyBeingEdited() const override;
 	virtual FLinearColor GetWorldCentricTabColorScale() const override;
 };
-
-
diff --git a/Engine/Source/Editor/UnrealEd/Public/Toolkits/IToolkit.h b/Engine/Source/Editor/UnrealEd/Public/Toolkits/IToolkit.h
--- a/Engine/Source/Editor/UnrealEd/Public/Toolkits/IToolkit.h
+++ b/Engine/Source/Editor/UnrealEd/Public/Toolkits/IToolkit.h
@@ -3,6 +3,15 @@
 #pragma once
 
 
+class FEdMode;
+class FTabManager;
+class FWorkspaceItem;
+class IToolkitHost;
+class SDockableTab;
+class SWidget;
+struct FKeyEvent;
+
+
 /** Some toolkits can be spawned as either standalone tools or within an existing level editing UI */
 namespace EToolkitMode
 {
@@ -54,9 +63,10 @@ class IToolkit
 public:
 
 	/** Register tabs that this toolkit can spawn with the TabManager */
-	virtual void RegisterTabSpawners(const TSharedRef<class FTabManager>& TabManager) = 0;
+	virtual void RegisterTabSpawners(const TSharedRef<FTabManager>& TabManager) = 0;
+
 	/** Unregister tabs that this toolkit can spawn */	
-	virtual void UnregisterTabSpawners(const TSharedRef<class FTabManager>& TabManager) = 0;
+	virtual void UnregisterTabSpawners(const TSharedRef<FTabManager>& TabManager) = 0;
 
 public:
 
@@ -94,10 +104,10 @@ public:
 	virtual bool IsHosted() const = 0;
 
 	/** @return Returns the toolkit host for this toolkit */
-	virtual const TSharedRef< class IToolkitHost > GetToolkitHost() const = 0;
+	virtual const TSharedRef<IToolkitHost> GetToolkitHost() const = 0;
 
 	/** @return Returns a map of weak pointers to all of this toolkit's spawned tabs that may currently exist, indexed by the tab spots they are suitable for  */
-	virtual const TMap< EToolkitTabSpot::Type, TArray< TWeakPtr< class SDockableTab > > >& GetToolkitTabsInSpots() const = 0;
+	virtual const TMap<EToolkitTabSpot::Type, TArray<TWeakPtr<SDockableTab>>>& GetToolkitTabsInSpots() const = 0;
 
 	/**
 	 * Processes any UI commands which are activated by the specified event
@@ -106,16 +116,16 @@ public:
 	 *
 	 * @return true if an action was processed
 	 */
-	virtual bool ProcessCommandBindings( const struct FKeyEvent& InKeyEvent ) const = 0;
+	virtual bool ProcessCommandBindings(const FKeyEvent& InKeyEvent) const = 0;
 
 	/** Call this function to bring all of this toolkit's tabs to the foreground in their respective stacks.  Also causes the toolkit's host window to be foregrounded, too! */
 	virtual void BringToolkitToFront() = 0;
 
 	/** @returns the editor mode this toolkit is used for, or null if not relevant. */
-	virtual class FEdMode* GetEditorMode() const = 0;
+	virtual FEdMode* GetEditorMode() const = 0;
 
 	/** @returns the inline content that this toolkit returns if it is an editor mode */
-	virtual TSharedPtr<class SWidget> GetInlineContent() const = 0;
+	virtual TSharedPtr<SWidget> GetInlineContent() const = 0;
 
 	/** Returns if this is a IBlueprintEditor derivation */
 	virtual bool IsBlueprintEditor() const = 0;
diff --git a/Engine/Source/Runtime/Slate/Public/Framework/Docking/TabManager.h b/Engine/Source/Runtime/Slate/Public/Framework/Docking/TabManager.h
--- a/Engine/Source/Runtime/Slate/Public/Framework/Docking/TabManager.h
+++ b/Engine/Source/Runtime/Slate/Public/Framework/Docking/TabManager.h
@@ -2,11 +2,17 @@
 
 #pragma once
 
+#include "SlateIcon.h"
+#include "WorkspaceItem.h"
+
+
+class FTabManager;
 class SDockingArea;
 class SDockTab;
 class SDockingTabStack;
+class SWindow;
 struct FTabMatcher;
-class FTabManager;
+
 
 DECLARE_MULTICAST_DELEGATE_TwoParams(
 	FOnActiveTabChanged,
