[UE-17466] Fixing up references to the generated class and clearing old blueprint node spawners for deleted blueprints

diff --git a/Engine/Source/Editor/UnrealEd/Private/ObjectTools.cpp b/Engine/Source/Editor/UnrealEd/Private/ObjectTools.cpp
--- a/Engine/Source/Editor/UnrealEd/Private/ObjectTools.cpp
+++ b/Engine/Source/Editor/UnrealEd/Private/ObjectTools.cpp
@@ -812,7 +812,7 @@ namespace ObjectTools
 						// then repair the references on the object being consolidated so those objects can be properly disposed of upon deletion.
 						UClass* OldClass = BlueprintToConsolidate->GeneratedClass;
 						UClass* OldSkeletonClass = BlueprintToConsolidate->SkeletonGeneratedClass;
-						FBlueprintCompileReinstancer::ReplaceInstancesOfClass(OldClass, BlueprintToConsolidateTo->GeneratedClass);
+						FBlueprintCompileReinstancer::ReplaceInstancesOfClass(OldClass, BlueprintToConsolidateTo->GeneratedClass, nullptr, nullptr, true );
 						BlueprintToConsolidate->GeneratedClass = OldClass;
 						BlueprintToConsolidate->SkeletonGeneratedClass = OldSkeletonClass;
 					}
@@ -822,6 +822,8 @@ namespace ObjectTools
 				CollectGarbage( GARBAGE_COLLECTION_KEEPFLAGS );
 			}
 
+			FEditorDelegates::OnAssetsPreDelete.Broadcast(ReplaceInfo.ReplaceableObjects);
+
 			// With all references to the objects to consolidate to eliminated from objects that are currently loaded, it should now be safe to delete
 			// the objects to be consolidated themselves, leaving behind a redirector in their place to fix up objects that were not currently loaded at the time
 			// of this operation.
