UE-19006: Fix for HLOD system crash

diff --git a/Engine/Source/Editor/UnrealEd/Private/HierarchicalLOD.cpp b/Engine/Source/Editor/UnrealEd/Private/HierarchicalLOD.cpp
--- a/Engine/Source/Editor/UnrealEd/Private/HierarchicalLOD.cpp
+++ b/Engine/Source/Editor/UnrealEd/Private/HierarchicalLOD.cpp
@@ -285,17 +285,14 @@ void FHierarchicalLODBuilder::FindMST()
 
 void FHierarchicalLODBuilder::HandleHLODVolumes(ULevel* InLevel)
 {	
+	HLODVolumeClusters.Empty();
 	for (int32 ActorId = 0; ActorId < InLevel->Actors.Num(); ++ActorId)
 	{
 		AActor* Actor = InLevel->Actors[ActorId];
 		if (Actor && Actor->IsA( AHierarchicalLODVolume::StaticClass()))
 		{
 			// Came across a HLOD volume			
-			FLODCluster& NewCluster = HLODVolumeClusters.Add(CastChecked<AHierarchicalLODVolume>(Actor));
-
-			//FVector Origin, Extent;
-			//Actor->GetActorBounds(false, Origin, Extent);			
-			//NewCluster.Bound = FSphere(Origin * CM_TO_METER, Extent.GetAbsMax() * CM_TO_METER);
+			FLODCluster& NewCluster = HLODVolumeClusters.Add(CastChecked<AHierarchicalLODVolume>(Actor));			
 		}
 	}
 }
