Fix crash when importing a static mesh over a skeletal mesh with the same name

https://jira.ol.epicgames.net/browse/UE-18557

diff --git a/Engine/Source/Editor/UnrealEd/Private/Fbx/FbxStaticMeshImport.cpp b/Engine/Source/Editor/UnrealEd/Private/Fbx/FbxStaticMeshImport.cpp
--- a/Engine/Source/Editor/UnrealEd/Private/Fbx/FbxStaticMeshImport.cpp
+++ b/Engine/Source/Editor/UnrealEd/Private/Fbx/FbxStaticMeshImport.cpp
@@ -922,6 +922,9 @@ UStaticMesh* UnFbx::FFbxImporter::ImportStaticMeshAsSingle(UObject* InParent, TA
 
 			// Create a package for each mesh
 			Package = CreatePackage(NULL, *NewPackageName);
+
+			// Require the parent because it will have been invalidated from the garbage collection
+			Parent = Package;
 		}
 		else
 		{
