Prevented matinee from being opened during PIE.  This is not safe.

diff --git a/Engine/Source/Editor/UnrealEd/Private/EditorEngine.cpp b/Engine/Source/Editor/UnrealEd/Private/EditorEngine.cpp
--- a/Engine/Source/Editor/UnrealEd/Private/EditorEngine.cpp
+++ b/Engine/Source/Editor/UnrealEd/Private/EditorEngine.cpp
@@ -3343,23 +3343,29 @@ void UEditorEngine::ConvertActorsFromClass( UClass* FromClass, UClass* ToClass )
 
 bool UEditorEngine::ShouldOpenMatinee(AMatineeActor* MatineeActor) const
 {
+	if( PlayWorld )
+	{
+		FMessageDialog::Open( EAppMsgType::Ok, NSLOCTEXT("UnrealEd", "Error_MatineeCantOpenDuringPIE", "Matinee cannot be opened during Play in Editor.") );
+		return false;
+	}
+
 	if ( MatineeActor && !MatineeActor->MatineeData )
 	{
-		FMessageDialog::Open( EAppMsgType::Ok, NSLOCTEXT("UnrealEd", "Error_MatineeActionMustHaveData", "UnrealMatinee must have valid InterpData assigned before being edited.") );
+		FMessageDialog::Open( EAppMsgType::Ok, NSLOCTEXT("UnrealEd", "Error_MatineeActionMustHaveData", "Matinee must have valid InterpData assigned before being edited.") );
 		return false;
 	}
 
 	// Make sure we can't open the same action twice in Matinee.
 	if( GLevelEditorModeTools().IsModeActive(FBuiltinEditorModes::EM_InterpEdit) )
 	{
-		FMessageDialog::Open( EAppMsgType::Ok, NSLOCTEXT("UnrealEd", "MatineeActionAlreadyOpen", "An UnrealMatinee sequence is currently open in an editor.  Please close it before proceeding.") );
+		FMessageDialog::Open( EAppMsgType::Ok, NSLOCTEXT("UnrealEd", "MatineeActionAlreadyOpen", "An Matinee sequence is currently open in an editor.  Please close it before proceeding.") );
 		return false;
 	}
 
 	// Don't let you open Matinee if a transaction is currently active.
 	if( GEditor->IsTransactionActive() )
 	{
-		FMessageDialog::Open( EAppMsgType::Ok, NSLOCTEXT("UnrealEd", "TransactionIsActive", "Undo Transaction Is Active - Cannot Open UnrealMatinee.") );
+		FMessageDialog::Open( EAppMsgType::Ok, NSLOCTEXT("UnrealEd", "TransactionIsActive", "Undo Transaction Is Active - Cannot Open Matinee.") );
 		return false;
 	}
 
