Check for maps specified on command-line before adding -MAPINISECTION=AllMaps. Fixes case where specifying -Map=<MapName> on command line was cooking AllMaps.

diff --git a/Engine/Source/Editor/UnrealEd/Private/Commandlets/CookCommandlet.cpp b/Engine/Source/Editor/UnrealEd/Private/Commandlets/CookCommandlet.cpp
--- a/Engine/Source/Editor/UnrealEd/Private/Commandlets/CookCommandlet.cpp
+++ b/Engine/Source/Editor/UnrealEd/Private/Commandlets/CookCommandlet.cpp
@@ -1065,40 +1065,6 @@ bool UCookCommandlet::NewCook( const TArray<ITargetPlatform*>& Platforms, TArray
 	FString CreateReleaseVersion;
 	FParse::Value( *Params, TEXT("CreateReleaseVersion="), CreateReleaseVersion);
 
-	// Add any map sections specified on command line
-	TArray<FString> AlwaysCookMapList;
-
-	// Add the default map section
-	GEditor->LoadMapListFromIni(TEXT("AlwaysCookMaps"), AlwaysCookMapList);
-
-	TArray<FString> MapList;
-	// Add any map sections specified on command line
-	GEditor->ParseMapSectionIni(*Params, MapList);
-
-	if (MapList.Num() == 0)
-	{
-		// if we didn't find any maps look in the project settings for maps
-
-		UProjectPackagingSettings* PackagingSettings = Cast<UProjectPackagingSettings>(UProjectPackagingSettings::StaticClass()->GetDefaultObject());
-
-		for (const auto& MapToCook : PackagingSettings->MapsToCook)
-		{
-			MapList.Add(MapToCook.FilePath);
-		}
-	}
-
-	// if we still don't have any mapsList check if the allmaps ini section is filled out
-	// this is for backwards compatibility
-	if (MapList.Num() == 0)
-	{
-		GEditor->ParseMapSectionIni(TEXT("-MAPINISECTION=AllMaps"), MapList);
-	}
-
-
-	// put the always cook map list at the front of the map list
-	AlwaysCookMapList.Append(MapList);
-	Swap(MapList, AlwaysCookMapList);
-
 	TArray<FString> CmdLineMapEntries;
 	TArray<FString> CmdLineDirEntries;
 	TArray<FString> CmdLineCultEntries;
@@ -1151,6 +1117,45 @@ bool UCookCommandlet::NewCook( const TArray<ITargetPlatform*>& Platforms, TArray
 		}
 	}
 
+	// Add any map sections specified on command line
+	TArray<FString> AlwaysCookMapList;
+
+	// Add the default map section
+	GEditor->LoadMapListFromIni(TEXT("AlwaysCookMaps"), AlwaysCookMapList);
+
+	TArray<FString> MapList;
+	// Add any map sections specified on command line
+	GEditor->ParseMapSectionIni(*Params, MapList);
+
+	if (MapList.Num() == 0)
+	{
+		// If we didn't find any maps look in the project settings for maps
+
+		UProjectPackagingSettings* PackagingSettings = Cast<UProjectPackagingSettings>(UProjectPackagingSettings::StaticClass()->GetDefaultObject());
+
+		for (const auto& MapToCook : PackagingSettings->MapsToCook)
+		{
+			MapList.Add(MapToCook.FilePath);
+		}
+	}
+
+	// Add any map specified on the command line.
+	for (const auto& MapName : CmdLineMapEntries)
+	{
+		MapList.Add(MapName);
+	}
+
+	// If we still don't have any mapsList check if the allmaps ini section is filled out
+	// this is for backwards compatibility
+	if (MapList.Num() == 0)
+	{
+		GEditor->ParseMapSectionIni(TEXT("-MAPINISECTION=AllMaps"), MapList);
+	}
+
+	// Put the always cook map list at the front of the map list
+	AlwaysCookMapList.Append(MapList);
+	Swap(MapList, AlwaysCookMapList);
+
 	//////////////////////////////////////////////////////////////////////////
 	// start cook by the book 
 
@@ -1163,11 +1168,6 @@ bool UCookCommandlet::NewCook( const TArray<ITargetPlatform*>& Platforms, TArray
 	CookOptions |= Switches.Contains(TEXT("MAPSONLY")) ? ECookByTheBookOptions::MapsOnly : ECookByTheBookOptions::None;
 	CookOptions |= Switches.Contains(TEXT("NODEV")) ? ECookByTheBookOptions::NoDevContent : ECookByTheBookOptions::None;
 
-	for ( const auto& MapName : CmdLineMapEntries )
-	{
-		MapList.Add( MapName );
-	}
-
 	UCookOnTheFlyServer::FCookByTheBookStartupOptions StartupOptions;
 
 	StartupOptions.TargetPlatforms = Platforms;
