Checks for duplicate categories in the deprecated LibraryCategories tag when building the list of Material Functions for the Material Editor right-click context menu.

diff --git a/Engine/Source/Editor/UnrealEd/Private/MaterialGraphSchema.cpp b/Engine/Source/Editor/UnrealEd/Private/MaterialGraphSchema.cpp
--- a/Engine/Source/Editor/UnrealEd/Private/MaterialGraphSchema.cpp
+++ b/Engine/Source/Editor/UnrealEd/Private/MaterialGraphSchema.cpp
@@ -737,9 +737,17 @@ void UMaterialGraphSchema::GetMaterialFunctionActions(FGraphActionMenuBuilder& A
 						LibraryCategoriesProperty->ImportText(*LibraryCategoriesString, DestAddr, PPF_None, NULL, GWarn);
 					}
 
-					for (FString& Category : LibraryCategories)
+					for (const FString& Category : LibraryCategories)
 					{
-						LibraryCategoriesText.Add(FText::FromString(Category));
+						if (!LibraryCategoriesText.ContainsByPredicate(
+							[&Category](const FText& Text)
+							{
+								return Text.ToString() == Category;
+							}
+							))
+						{
+							LibraryCategoriesText.Add(FText::FromString(Category));
+						}
 					}
 
 					if ( LibraryCategoriesText.Num() == 0 )
