CIS Fix

diff --git a/Engine/Source/Editor/UnrealEd/Private/Kismet2/KismetReinstanceUtilities.cpp b/Engine/Source/Editor/UnrealEd/Private/Kismet2/KismetReinstanceUtilities.cpp
--- a/Engine/Source/Editor/UnrealEd/Private/Kismet2/KismetReinstanceUtilities.cpp
+++ b/Engine/Source/Editor/UnrealEd/Private/Kismet2/KismetReinstanceUtilities.cpp
@@ -1326,9 +1326,6 @@ FBlueprintCompileReinstancer::FCDODuplicatesProvider& FBlueprintCompileReinstanc
 	return Delegate;
 }
 
-TArray<UObject*> Objects;
-UClass* Class;
-
 FRecreateUberGraphFrameScope::FRecreateUberGraphFrameScope(UClass* InClass, bool bRecreate)
 	: Class(InClass)
 {
