In FRecreateUberGraphFrameScope member variable "Class" is renamed into "RecompiledClass"

diff --git a/Engine/Source/Editor/UnrealEd/Private/Kismet2/KismetReinstanceUtilities.cpp b/Engine/Source/Editor/UnrealEd/Private/Kismet2/KismetReinstanceUtilities.cpp
--- a/Engine/Source/Editor/UnrealEd/Private/Kismet2/KismetReinstanceUtilities.cpp
+++ b/Engine/Source/Editor/UnrealEd/Private/Kismet2/KismetReinstanceUtilities.cpp
@@ -1327,18 +1327,18 @@ FBlueprintCompileReinstancer::FCDODuplicatesProvider& FBlueprintCompileReinstanc
 }
 
 FRecreateUberGraphFrameScope::FRecreateUberGraphFrameScope(UClass* InClass, bool bRecreate)
-	: Class(InClass)
+	: RecompiledClass(InClass)
 {
-	if (bRecreate && ensure(Class))
+	if (bRecreate && ensure(RecompiledClass))
 	{
 		BP_SCOPED_COMPILER_EVENT_STAT(EKismetCompilerStats_RecreateUberGraphPersistentFrame);
 
 		const bool bIncludeDerivedClasses = true;
-		GetObjectsOfClass(Class, Objects, bIncludeDerivedClasses);
+		GetObjectsOfClass(RecompiledClass, Objects, bIncludeDerivedClasses);
 
 		for (auto Obj : Objects)
 		{
-			Class->DestroyPersistentUberGraphFrame(Obj, true);
+			RecompiledClass->DestroyPersistentUberGraphFrame(Obj, true);
 		}
 	}
 }
@@ -1350,7 +1350,7 @@ FRecreateUberGraphFrameScope::~FRecreateUberGraphFrameScope()
 	{
 		if (IsValid(Obj))
 		{
-			Class->CreatePersistentUberGraphFrame(Obj, false, true);
+			RecompiledClass->CreatePersistentUberGraphFrame(Obj, false, true);
 		}
 	}
 }
\ No newline at end of file
diff --git a/Engine/Source/Editor/UnrealEd/Public/Kismet2/KismetReinstanceUtilities.h b/Engine/Source/Editor/UnrealEd/Public/Kismet2/KismetReinstanceUtilities.h
--- a/Engine/Source/Editor/UnrealEd/Public/Kismet2/KismetReinstanceUtilities.h
+++ b/Engine/Source/Editor/UnrealEd/Public/Kismet2/KismetReinstanceUtilities.h
@@ -16,7 +16,7 @@ struct UNREALED_API FRecreateUberGraphFrameScope
 {
 private:
 	TArray<UObject*> Objects;
-	UClass* Class;
+	UClass* RecompiledClass;
 public:
 	FRecreateUberGraphFrameScope(UClass* InClass, bool bRecreate);
 	~FRecreateUberGraphFrameScope();
