FStringAssetReference fixes.

All changes regarding e-mail discussion titled "Asset registry shenanigans".

Changes:
1. Engine now actively is getting rid of short asset paths in FStringAssetReference during saving (for all packages) and loading (only for older packages).
2. Deprecated direct access to FStringAssetReference.AssetLongPathname and exposed it via ToString and SetPath -- SetPath is making sure the path is in correct format. If the path can't be found on disk, the path is cleared.
3. Also access to FStringAssetReference.AssetLongPathnam via Blueprints is guarded using custom Make node function that uses SetPath instead of simple assign.
4. StringAssetReferenceMap will now contain only paths to packages (not objects) and ini references
5. GetDependencies now has additional parameter that lets you chose if it should resolve ini references or not. It was left not as default to keep old behaviour.

diff --git a/Engine/Source/Developer/AssetTools/Private/AssetRenameManager.cpp b/Engine/Source/Developer/AssetTools/Private/AssetRenameManager.cpp
--- a/Engine/Source/Developer/AssetTools/Private/AssetRenameManager.cpp
+++ b/Engine/Source/Developer/AssetTools/Private/AssetRenameManager.cpp
@@ -541,15 +541,15 @@ void FAssetRenameManager::RenameReferencingStringAssetReferences(const TArray<UP
 
 		FArchive& operator<<(FStringAssetReference& Reference) override
 		{
-			if (Reference.AssetLongPathname == OldAssetPath)
+			if (Reference.ToString() == OldAssetPath)
 			{
-				Reference.AssetLongPathname = NewAssetPath;
+				Reference.SetPath(NewAssetPath);
 			}
 
 			// Generated class path support.
-			if (Reference.AssetLongPathname == OldAssetPath + "_C")
+			if (Reference.ToString() == OldAssetPath + "_C")
 			{
-				Reference.AssetLongPathname = NewAssetPath + "_C";
+				Reference.SetPath(NewAssetPath + "_C");
 			}
 
 			return *this;
diff --git a/Engine/Source/Developer/HotReload/Private/HotReloadClassReinstancer.cpp b/Engine/Source/Developer/HotReload/Private/HotReloadClassReinstancer.cpp
--- a/Engine/Source/Developer/HotReload/Private/HotReloadClassReinstancer.cpp
+++ b/Engine/Source/Developer/HotReload/Private/HotReloadClassReinstancer.cpp
@@ -166,7 +166,16 @@ void FHotReloadClassReinstancer::SerializeCDOProperties(UObject* InObject, FHotR
 		virtual FArchive& operator<<(FStringAssetReference& Value) override
 		{
 			FArchive& Ar = *this;
-			Ar << Value.AssetLongPathname;
+
+			FString Path = Value.ToString();
+
+			Ar << Path;
+
+			if (IsLoading())
+			{
+				Value.SetPath(MoveTemp(Path));
+			}
+
 			return Ar;
 		}
 		/** Archive name, for debugging */
@@ -386,7 +395,16 @@ void FHotReloadClassReinstancer::UpdateDefaultProperties()
 		virtual FArchive& operator<<(FStringAssetReference& Value) override
 		{
 			FArchive& Ar = *this;
-			Ar << Value.AssetLongPathname;
+
+			FString Path = Value.ToString();
+
+			Ar << Path;
+
+			if (IsLoading())
+			{
+				Value.SetPath(MoveTemp(Path));
+			}
+
 			return Ar;
 		}
 	};
diff --git a/Engine/Source/Editor/GameplayAbilitiesEditor/Private/SGameplayCueEditor.cpp b/Engine/Source/Editor/GameplayAbilitiesEditor/Private/SGameplayCueEditor.cpp
--- a/Engine/Source/Editor/GameplayAbilitiesEditor/Private/SGameplayCueEditor.cpp
+++ b/Engine/Source/Editor/GameplayAbilitiesEditor/Private/SGameplayCueEditor.cpp
@@ -647,9 +647,9 @@ private:
 				}
 				else if (ColumnName == CueHandlerColumnName)
 				{
-					if (CueHandlerItem->GameplayCueNotifyObj.AssetLongPathname.IsEmpty() == false)
+					if (CueHandlerItem->GameplayCueNotifyObj.ToString().IsEmpty() == false)
 					{
-						FString ObjName = CueHandlerItem->GameplayCueNotifyObj.AssetLongPathname;
+						FString ObjName = CueHandlerItem->GameplayCueNotifyObj.ToString();
 
 						int32 idx;
 						if (ObjName.FindLastChar(TEXT('.'), idx))
diff --git a/Engine/Source/Editor/IntroTutorials/Private/IntroTutorials.cpp b/Engine/Source/Editor/IntroTutorials/Private/IntroTutorials.cpp
--- a/Engine/Source/Editor/IntroTutorials/Private/IntroTutorials.cpp
+++ b/Engine/Source/Editor/IntroTutorials/Private/IntroTutorials.cpp
@@ -315,7 +315,7 @@ bool FIntroTutorials::MaybeOpenWelcomeTutorial()
 	if(FParse::Param(FCommandLine::Get(), TEXT("TestTutorialAlerts")) || !FEngineBuildSettings::IsInternalBuild())
 	{
 		// try editor startup tutorial
-		TSubclassOf<UEditorTutorial> EditorStartupTutorialClass = LoadClass<UEditorTutorial>(NULL, *GetDefault<UEditorTutorialSettings>()->StartupTutorial.AssetLongPathname, NULL, LOAD_None, NULL);
+		TSubclassOf<UEditorTutorial> EditorStartupTutorialClass = LoadClass<UEditorTutorial>(NULL, *GetDefault<UEditorTutorialSettings>()->StartupTutorial.ToString(), NULL, LOAD_None, NULL);
 		if(EditorStartupTutorialClass != nullptr)
 		{
 			UEditorTutorial* Tutorial = EditorStartupTutorialClass->GetDefaultObject<UEditorTutorial>();
@@ -327,7 +327,7 @@ bool FIntroTutorials::MaybeOpenWelcomeTutorial()
 		}
 
 		// Try project startup tutorial
-		TSubclassOf<UEditorTutorial> ProjectStartupTutorialClass = LoadClass<UEditorTutorial>(NULL, *GetDefault<UTutorialSettings>()->StartupTutorial.AssetLongPathname, NULL, LOAD_None, NULL);
+		TSubclassOf<UEditorTutorial> ProjectStartupTutorialClass = LoadClass<UEditorTutorial>(NULL, *GetDefault<UTutorialSettings>()->StartupTutorial.ToString(), NULL, LOAD_None, NULL);
 		if(ProjectStartupTutorialClass != nullptr)
 		{
 			UEditorTutorial* Tutorial = ProjectStartupTutorialClass->GetDefaultObject<UEditorTutorial>();
diff --git a/Engine/Source/Editor/IntroTutorials/Private/STutorialRoot.cpp b/Engine/Source/Editor/IntroTutorials/Private/STutorialRoot.cpp
--- a/Engine/Source/Editor/IntroTutorials/Private/STutorialRoot.cpp
+++ b/Engine/Source/Editor/IntroTutorials/Private/STutorialRoot.cpp
@@ -130,7 +130,7 @@ void STutorialRoot::LaunchTutorial(UEditorTutorial* InTutorial, IIntroTutorials:
 		if(InTutorial->AssetToUse.IsValid())
 		{
 			TArray<FString> AssetPaths;
-			AssetPaths.Add(InTutorial->AssetToUse.AssetLongPathname);
+			AssetPaths.Add(InTutorial->AssetToUse.ToString());
 			FAssetEditorManager::Get().OpenEditorsForAssets(AssetPaths);
 
 			UObject* Asset = InTutorial->AssetToUse.ResolveObject();
@@ -313,16 +313,16 @@ void STutorialRoot::GoToPreviousStage()
 		{
 			// We went out of bounds for this tutorial, so see if we want to go to another one.
 			CurrentTutorialStage = 0;
-			if (FName(*CurrentTutorial->PreviousTutorial.AssetLongPathname) != NAME_None)
+			if (FName(*CurrentTutorial->PreviousTutorial.ToString()) != NAME_None)
 			{
-				TSubclassOf<UEditorTutorial> PreviousTutorialClass = LoadClass<UEditorTutorial>(NULL, *CurrentTutorial->PreviousTutorial.AssetLongPathname, NULL, LOAD_None, NULL);
+				TSubclassOf<UEditorTutorial> PreviousTutorialClass = LoadClass<UEditorTutorial>(NULL, *CurrentTutorial->PreviousTutorial.ToString(), NULL, LOAD_None, NULL);
 				if (PreviousTutorialClass != nullptr)
 				{
 					LaunchTutorial(PreviousTutorialClass->GetDefaultObject<UEditorTutorial>(), IIntroTutorials::ETutorialStartType::TST_LASTSTAGE, nullptr, FSimpleDelegate(), FSimpleDelegate());
 				}
 				else
 				{
-					FSlateNotificationManager::Get().AddNotification(FNotificationInfo(FText::Format(LOCTEXT("TutorialNotFound", "Could not start previous tutorial {0}"), FText::FromString(CurrentTutorial->PreviousTutorial.AssetLongPathname))));
+					FSlateNotificationManager::Get().AddNotification(FNotificationInfo(FText::Format(LOCTEXT("TutorialNotFound", "Could not start previous tutorial {0}"), FText::FromString(CurrentTutorial->PreviousTutorial.ToString()))));
 				}
 			}
 		}
@@ -379,16 +379,16 @@ void STutorialRoot::GoToNextStage(TWeakPtr<SWindow> InNavigationWindow)
 		{
 			// We went out of bounds for this tutorial, so see if we want to go to another one.
 			CurrentTutorialStage = CurrentTutorial->Stages.Num() - 1;
-			if (FName(*CurrentTutorial->PreviousTutorial.AssetLongPathname) != NAME_None)
+			if (FName(*CurrentTutorial->PreviousTutorial.ToString()) != NAME_None)
 			{
-				TSubclassOf<UEditorTutorial> NextTutorialClass = LoadClass<UEditorTutorial>(NULL, *CurrentTutorial->NextTutorial.AssetLongPathname, NULL, LOAD_None, NULL);
+				TSubclassOf<UEditorTutorial> NextTutorialClass = LoadClass<UEditorTutorial>(NULL, *CurrentTutorial->NextTutorial.ToString(), NULL, LOAD_None, NULL);
 				if (NextTutorialClass != nullptr)
 				{
 					LaunchTutorial(NextTutorialClass->GetDefaultObject<UEditorTutorial>(), IIntroTutorials::ETutorialStartType::TST_RESTART, InNavigationWindow, FSimpleDelegate(), FSimpleDelegate());
 				}
 				else
 				{
-					FSlateNotificationManager::Get().AddNotification(FNotificationInfo(FText::Format(LOCTEXT("TutorialNotFound", "Could not start next tutorial {0}"), FText::FromString(CurrentTutorial->NextTutorial.AssetLongPathname))));
+					FSlateNotificationManager::Get().AddNotification(FNotificationInfo(FText::Format(LOCTEXT("TutorialNotFound", "Could not start next tutorial {0}"), FText::FromString(CurrentTutorial->NextTutorial.ToString()))));
 				}
 			}
 		}
diff --git a/Engine/Source/Editor/IntroTutorials/Private/STutorialsBrowser.cpp b/Engine/Source/Editor/IntroTutorials/Private/STutorialsBrowser.cpp
--- a/Engine/Source/Editor/IntroTutorials/Private/STutorialsBrowser.cpp
+++ b/Engine/Source/Editor/IntroTutorials/Private/STutorialsBrowser.cpp
@@ -55,7 +55,7 @@ public:
 
 		if(Category.Texture.IsValid())
 		{
-			UTexture2D* Texture = LoadObject<UTexture2D>(nullptr, *Category.Texture.AssetLongPathname);
+			UTexture2D* Texture = LoadObject<UTexture2D>(nullptr, *Category.Texture.ToString());
 			if(Texture != nullptr)
 			{
 				FIntPoint TextureSize = Texture->GetImportedSize();
diff --git a/Engine/Source/Editor/IntroTutorials/Private/TutorialStateSettings.cpp b/Engine/Source/Editor/IntroTutorials/Private/TutorialStateSettings.cpp
--- a/Engine/Source/Editor/IntroTutorials/Private/TutorialStateSettings.cpp
+++ b/Engine/Source/Editor/IntroTutorials/Private/TutorialStateSettings.cpp
@@ -15,7 +15,7 @@ void UTutorialStateSettings::PostInitProperties()
 
 	for(const auto& Progress : TutorialsProgress)
 	{
-		TSubclassOf<UEditorTutorial> TutorialClass = LoadClass<UEditorTutorial>(NULL, *Progress.Tutorial.AssetLongPathname, NULL, LOAD_None, NULL);
+		TSubclassOf<UEditorTutorial> TutorialClass = LoadClass<UEditorTutorial>(NULL, *Progress.Tutorial.ToString(), NULL, LOAD_None, NULL);
 		if(TutorialClass != nullptr)
 		{
 			UEditorTutorial* Tutorial = TutorialClass->GetDefaultObject<UEditorTutorial>();
diff --git a/Engine/Source/Editor/KismetCompiler/Private/KismetCompilerVMBackend.cpp b/Engine/Source/Editor/KismetCompiler/Private/KismetCompilerVMBackend.cpp
--- a/Engine/Source/Editor/KismetCompiler/Private/KismetCompilerVMBackend.cpp
+++ b/Engine/Source/Editor/KismetCompiler/Private/KismetCompilerVMBackend.cpp
@@ -577,7 +577,7 @@ public:
 			{
 				Writer << EX_AssetConst;
 				FAssetPtr AssetPtr(Term->ObjectLiteral);
-				EmitStringLiteral(AssetPtr.GetUniqueID().AssetLongPathname);
+				EmitStringLiteral(AssetPtr.GetUniqueID().ToString());
 			}
 			else if (CoerceProperty->IsA(UObjectPropertyBase::StaticClass()))
 			{
diff --git a/Engine/Source/Editor/UnrealEd/Private/Kismet2/Kismet2.cpp b/Engine/Source/Editor/UnrealEd/Private/Kismet2/Kismet2.cpp
--- a/Engine/Source/Editor/UnrealEd/Private/Kismet2/Kismet2.cpp
+++ b/Engine/Source/Editor/UnrealEd/Private/Kismet2/Kismet2.cpp
@@ -630,7 +630,7 @@ UBlueprint* FKismetEditorUtilities::ReloadBlueprint(UBlueprint* StaleBlueprint)
 	FBlueprintUnloader Unloader(StaleBlueprint);
 	Unloader.UnloadBlueprint(/*bResetPackage =*/true);
 
-	UBlueprint* ReloadedBlueprint = Cast<UBlueprint>(StaticLoadObject(UBlueprint::StaticClass(), /*Outer =*/nullptr, *BlueprintAssetRef.AssetLongPathname));
+	UBlueprint* ReloadedBlueprint = Cast<UBlueprint>(StaticLoadObject(UBlueprint::StaticClass(), /*Outer =*/nullptr, *BlueprintAssetRef.ToString()));
 
 	Unloader.ReplaceStaleRefs(ReloadedBlueprint);
 	return ReloadedBlueprint;
diff --git a/Engine/Source/Editor/UnrealEd/Private/Kismet2/KismetDebugUtilities.cpp b/Engine/Source/Editor/UnrealEd/Private/Kismet2/KismetDebugUtilities.cpp
--- a/Engine/Source/Editor/UnrealEd/Private/Kismet2/KismetDebugUtilities.cpp
+++ b/Engine/Source/Editor/UnrealEd/Private/Kismet2/KismetDebugUtilities.cpp
@@ -268,84 +268,101 @@ void FKismetDebugUtilities::OnScriptException(const UObject* ActiveObject, const
 			BlueprintObj->SetObjectBeingDebugged(SavedObjectBeingDebugged);
 		}
 
-		// Extra cleanup after potential interactive handling
-		switch (Info.GetType())
+		const auto DisplayErrorLambda = [&](const FText ErrorTypeName, const TCHAR* Description)
 		{
-		case EBlueprintExceptionType::FatalError:
-		case EBlueprintExceptionType::InfiniteLoop:
+			if (GUnrealEd->PlayWorld != NULL)
 			{
-				if (GUnrealEd->PlayWorld != NULL)
+				GEditor->RequestEndPlayMap();
+				FSlateApplication::Get().LeaveDebuggingMode();
+			}
+
+			// Launch a message box notifying the user why they have been booted
+			{
+				// Callback to display a pop-up showing the Callstack, the user can highlight and copy this if needed
+				auto DisplayCallStackLambda = [](const FText CallStack)
 				{
-					GEditor->RequestEndPlayMap();
-					FSlateApplication::Get().LeaveDebuggingMode();
-					// Launch a message box notifying the user why they have been booted
-					{
-						// Callback to display a pop-up showing the Callstack, the user can highlight and copy this if needed
-						auto DisplayCallStackLambda = [](const FText CallStack)
-						{
-							TSharedPtr<SMultiLineEditableText> TextBlock;
-							TSharedRef<SWidget> DisplayWidget = 
-								SNew(SBox)
-									.MaxDesiredHeight(512)
-									.MaxDesiredWidth(512)
-									.Content()
-									[
-										SNew(SBorder)
-											.BorderImage(FEditorStyle::GetBrush("ToolPanel.GroupBorder"))
-											[
-												SNew(SScrollBox)
-												+SScrollBox::Slot()
-												[
-													SAssignNew(TextBlock, SMultiLineEditableText)
-														.AutoWrapText(true)
-														.IsReadOnly(true)
-														.Text(CallStack)
-												]
-											]
-									];
-
-							FSlateApplication::Get().PushMenu(
-								FSlateApplication::Get().GetActiveTopLevelWindow().ToSharedRef(),
-								FWidgetPath(),
-								DisplayWidget,
-								FSlateApplication::Get().GetCursorPos(),
-								FPopupTransitionEffect(FPopupTransitionEffect::TypeInPopup)
-								);	
-
-							FSlateApplication::Get().SetKeyboardFocus(TextBlock);
-						};
-
-						TSharedRef<FTokenizedMessage> Message = FTokenizedMessage::Create(EMessageSeverity::Error);
-
-						// Display a UObject link to the Blueprint that is the source of the failure
-						Message->AddToken(FTextToken::Create(LOCTEXT( "InfiniteLoopWarning_Blueprint", "Infinite Loop detected in ")));
-						FString BlueprintName;
-						BlueprintObj->GetName(BlueprintName);
-						Message->AddToken(FUObjectToken::Create(BlueprintObj, FText::FromString(BlueprintName)));
-
-						// Display a UObject link to the UFunction that is crashing. Will open the Blueprint if able and focus on the function's graph
-						Message->AddToken(FTextToken::Create(LOCTEXT( "InfiniteLoopWarning_Function", ", asserted during ")));
-						const int32 BreakpointOpCodeOffset = StackFrame.Code - StackFrame.Node->Script.GetData() - 1; //@TODO: Might want to make this a parameter of Info
-						UEdGraphNode* SourceNode = FindSourceNodeForCodeLocation(ActiveObject, StackFrame.Node, BreakpointOpCodeOffset, /*bAllowImpreciseHit=*/ true);
-
-						// If a source node is found, that's the token we want to link, otherwise settle with the UFunction
-						if(SourceNode)
-						{
-							Message->AddToken(FUObjectToken::Create(SourceNode, SourceNode->GetNodeTitle(ENodeTitleType::ListView)));
-						}
-						else
-						{
-							Message->AddToken(FUObjectToken::Create(StackFrame.Node, StackFrame.Node->GetDisplayNameText()));
-						}
+					TSharedPtr<SMultiLineEditableText> TextBlock;
+					TSharedRef<SWidget> DisplayWidget =
+						SNew(SBox)
+						.MaxDesiredHeight(512)
+						.MaxDesiredWidth(512)
+						.Content()
+						[
+							SNew(SBorder)
+							.BorderImage(FEditorStyle::GetBrush("ToolPanel.GroupBorder"))
+							[
+								SNew(SScrollBox)
+								+ SScrollBox::Slot()
+								[
+									SAssignNew(TextBlock, SMultiLineEditableText)
+									.AutoWrapText(true)
+									.IsReadOnly(true)
+									.Text(CallStack)
+								]
+							]
+						];
+
+					FSlateApplication::Get().PushMenu(
+						FSlateApplication::Get().GetActiveTopLevelWindow().ToSharedRef(),
+						FWidgetPath(),
+						DisplayWidget,
+						FSlateApplication::Get().GetCursorPos(),
+						FPopupTransitionEffect(FPopupTransitionEffect::TypeInPopup)
+						);
+
+					FSlateApplication::Get().SetKeyboardFocus(TextBlock);
+				};
+
+				TSharedRef<FTokenizedMessage> Message = FTokenizedMessage::Create(EMessageSeverity::Error);
+
+				// Display a UObject link to the Blueprint that is the source of the failure
+				Message->AddToken(FTextToken::Create(ErrorTypeName));
+				Message->AddToken(FTextToken::Create(LOCTEXT("DisplayErrorLambda_Blueprint", "detected in ")));
+				FString BlueprintName;
+				BlueprintObj->GetName(BlueprintName);
+				Message->AddToken(FUObjectToken::Create(BlueprintObj, FText::FromString(BlueprintName)));
+
+				// Display a UObject link to the UFunction that is crashing. Will open the Blueprint if able and focus on the function's graph
+				Message->AddToken(FTextToken::Create(LOCTEXT("DisplayErrorLambda_Function", ", asserted during ")));
+				const int32 BreakpointOpCodeOffset = StackFrame.Code - StackFrame.Node->Script.GetData() - 1; //@TODO: Might want to make this a parameter of Info
+				UEdGraphNode* SourceNode = FindSourceNodeForCodeLocation(ActiveObject, StackFrame.Node, BreakpointOpCodeOffset, /*bAllowImpreciseHit=*/ true);
+
+				// If a source node is found, that's the token we want to link, otherwise settle with the UFunction
+				if (SourceNode)
+				{
+					Message->AddToken(FUObjectToken::Create(SourceNode, SourceNode->GetNodeTitle(ENodeTitleType::ListView)));
+				}
+				else
+				{
+					Message->AddToken(FUObjectToken::Create(StackFrame.Node, StackFrame.Node->GetDisplayNameText()));
+				}
 
-						// Add an action token to display a pop-up that will display the complete script callstack
-						Message->AddToken(FTextToken::Create(LOCTEXT( "InfiniteLoopWarning_CallStack", " with the following ")));
-						const FText CallStackAsText = FText::FromString(StackFrame.GetStackTrace());
-						Message->AddToken(FActionToken::Create(LOCTEXT( "InfiniteLoopWarning_CallStackLink", "Call Stack" ), LOCTEXT( "InfiniteLoopWarning_CallStackDesc", "Displays the underlying callstack, tracing what function calls led to the assert occuring." ), FOnActionTokenExecuted::CreateStatic(DisplayCallStackLambda, CallStackAsText)));
-						FMessageLog("PIE").AddMessage(Message);
-					}
+				if (!Description)
+				{
+					Message->AddToken(FTextToken::Create(LOCTEXT("DisplayErrorLambda_CallStackNoDescription", " with the following ")));
 				}
+				else
+				{
+					Message->AddToken(FTextToken::Create(LOCTEXT("DisplayErrorLambda_CallStackWithDescription", " with the following message")));
+					Message->AddToken(FTextToken::Create(FText::FromString(FString::Printf(TEXT("\"%s\""), Description))));
+					Message->AddToken(FTextToken::Create(LOCTEXT("DisplayErrorLambda_CallStackWithDescriptionAnd", "and ")));
+				}
+
+				// Add an action token to display a pop-up that will display the complete script callstack
+				const FText CallStackAsText = FText::FromString(StackFrame.GetStackTrace());
+				Message->AddToken(FActionToken::Create(LOCTEXT("DisplayErrorLambda_CallStackLink", "Call Stack"), LOCTEXT("DisplayErrorLambda_CallStackDesc", "Displays the underlying callstack, tracing what function calls led to the assert occuring."), FOnActionTokenExecuted::CreateStatic(DisplayCallStackLambda, CallStackAsText)));
+				FMessageLog("PIE").AddMessage(Message);
 			}
+		};
+
+		// Extra cleanup after potential interactive handling
+		switch (Info.GetType())
+		{
+		case EBlueprintExceptionType::FatalError:
+			DisplayErrorLambda(LOCTEXT("FatalErrorType", "Fatal Error"), *Info.GetDescription());
+			break;
+		case EBlueprintExceptionType::InfiniteLoop:
+			DisplayErrorLambda(LOCTEXT("InfiniteLoopErrorType", "Infinite Loop"), nullptr);
 			break;
 		default:
 			// Left empty intentionally
diff --git a/Engine/Source/Runtime/AssetRegistry/Private/AssetDataGatherer.cpp b/Engine/Source/Runtime/AssetRegistry/Private/AssetDataGatherer.cpp
--- a/Engine/Source/Runtime/AssetRegistry/Private/AssetDataGatherer.cpp
+++ b/Engine/Source/Runtime/AssetRegistry/Private/AssetDataGatherer.cpp
@@ -3,7 +3,7 @@
 #include "AssetRegistryPCH.h"
 
 #define MAX_FILES_TO_PROCESS_BEFORE_FLUSH 250
-#define CACHE_SERIALIZATION_VERSION 4
+#define CACHE_SERIALIZATION_VERSION 5
 
 FAssetDataGatherer::FAssetDataGatherer(const TArray<FString>& InPaths, bool bInIsSynchronous, bool bInLoadAndSaveCache)
 	: StopTaskCounter( 0 )
diff --git a/Engine/Source/Runtime/AssetRegistry/Private/AssetRegistry.cpp b/Engine/Source/Runtime/AssetRegistry/Private/AssetRegistry.cpp
--- a/Engine/Source/Runtime/AssetRegistry/Private/AssetRegistry.cpp
+++ b/Engine/Source/Runtime/AssetRegistry/Private/AssetRegistry.cpp
@@ -732,7 +732,7 @@ bool FAssetRegistry::GetAllAssets(TArray<FAssetData>& OutAssetData) const
 	return true;
 }
 
-bool FAssetRegistry::GetDependencies(FName PackageName, TArray<FName>& OutDependencies, EAssetRegistryDependencyType::Type InDependencyType) const
+bool FAssetRegistry::GetDependencies(FName PackageName, TArray<FName>& OutDependencies, EAssetRegistryDependencyType::Type InDependencyType, bool bResolveIniStringReferences) const
 {
 	const FDependsNode*const* NodePtr = CachedDependsNodes.Find(PackageName);
 	const FDependsNode* Node = NULL;
@@ -743,7 +743,27 @@ bool FAssetRegistry::GetDependencies(FName PackageName, TArray<FName>& OutDepend
 
 	if (Node != NULL)
 	{
-		Node->GetDependencies(OutDependencies, InDependencyType);
+		if (bResolveIniStringReferences)
+		{
+			TArray<FName> DependencyNodes;
+			Node->GetDependencies(DependencyNodes, InDependencyType);
+
+			for (auto NodeIt = DependencyNodes.CreateConstIterator(); NodeIt; ++NodeIt)
+			{
+				FString PackagePath = NodeIt->ToString();
+				const FString* IniFilename = GetIniFilenameFromObjectsReference(PackagePath);
+
+				OutDependencies.Add(
+					(IniFilename != nullptr)
+					? FName(*ResolveIniObjectsReference(PackagePath, IniFilename))
+					: *NodeIt);
+			}
+		}
+		else
+		{
+			Node->GetDependencies(OutDependencies, InDependencyType);
+		}
+
 		return true;
 	}
 	else
@@ -1789,13 +1809,12 @@ void FAssetRegistry::DependencyDataGathered(const double TickStartTime, TArray<F
 
 		for (int32 StringAssetRefIdx = 0; StringAssetRefIdx < Result.StringAssetReferencesMap.Num(); ++StringAssetRefIdx)
 		{
-			FString PackageName, ObjName;
-			Result.StringAssetReferencesMap[StringAssetRefIdx].Split(".", &PackageName, &ObjName, ESearchCase::CaseSensitive);
-
-			if (!PackageDependencies.Contains(*PackageName))
+			if (FPackageName::IsShortPackageName(Result.StringAssetReferencesMap[StringAssetRefIdx]))
 			{
-				PackageDependencies.Add(*PackageName, EAssetRegistryDependencyType::Soft);
+				UE_LOG(LogAssetRegistry, Warning, TEXT("Package with string asset reference with short asset path: %s. This is unsupported, can couse errors and be slow on loading. Please resave the package to fix this."), *Result.PackageName.ToString());
 			}
+
+			PackageDependencies.Add(*Result.StringAssetReferencesMap[StringAssetRefIdx]);
 		}
 
 
diff --git a/Engine/Source/Runtime/AssetRegistry/Private/AssetRegistry.h b/Engine/Source/Runtime/AssetRegistry/Private/AssetRegistry.h
--- a/Engine/Source/Runtime/AssetRegistry/Private/AssetRegistry.h
+++ b/Engine/Source/Runtime/AssetRegistry/Private/AssetRegistry.h
@@ -21,7 +21,7 @@ public:
 	virtual bool GetAssets(const FARFilter& Filter, TArray<FAssetData>& OutAssetData) const override;
 	virtual FAssetData GetAssetByObjectPath( const FName ObjectPath ) const override;
 	virtual bool GetAllAssets(TArray<FAssetData>& OutAssetData) const override;
-	virtual bool GetDependencies(FName PackageName, TArray<FName>& OutDependencies, EAssetRegistryDependencyType::Type InDependencyType = EAssetRegistryDependencyType::All) const override;
+	virtual bool GetDependencies(FName PackageName, TArray<FName>& OutDependencies, EAssetRegistryDependencyType::Type InDependencyType = EAssetRegistryDependencyType::All, bool bResolveIniStringReferences = false) const override;
 	virtual bool GetReferencers(FName PackageName, TArray<FName>& OutReferencers) const override;
 	virtual bool GetAncestorClassNames(FName ClassName, TArray<FName>& OutAncestorClassNames) const override;
 	virtual void GetDerivedClassNames(const TArray<FName>& ClassNames, const TSet<FName>& ExcludedClassNames, TSet<FName>& OutDerivedClassNames) const override;
diff --git a/Engine/Source/Runtime/AssetRegistry/Private/PackageReader.cpp b/Engine/Source/Runtime/AssetRegistry/Private/PackageReader.cpp
--- a/Engine/Source/Runtime/AssetRegistry/Private/PackageReader.cpp
+++ b/Engine/Source/Runtime/AssetRegistry/Private/PackageReader.cpp
@@ -302,11 +302,39 @@ void FPackageReader::SerializeStringAssetReferencesMap(TArray<FString>& OutStrin
 	{
 		Seek(PackageFileSummary.StringAssetReferencesOffset);
 
-		for (int32 ReferenceIdx = 0; ReferenceIdx < PackageFileSummary.StringAssetReferencesCount; ++ReferenceIdx)
+		if (UE4Ver() < VER_UE4_KEEP_ONLY_PACKAGE_NAMES_IN_STRING_ASSET_REFERENCES_MAP)
 		{
-			FString Buf;
-			*this << Buf;
-			OutStringAssetReferencesMap.Add(MoveTemp(Buf));
+			for (int32 ReferenceIdx = 0; ReferenceIdx < PackageFileSummary.StringAssetReferencesCount; ++ReferenceIdx)
+			{
+				FString Buf;
+				*this << Buf;
+
+				if (GetIniFilenameFromObjectsReference(Buf) != nullptr)
+				{
+					OutStringAssetReferencesMap.AddUnique(MoveTemp(Buf));
+				}
+				else
+				{
+					FString NormalizedPath = FPackageName::GetNormalizedObjectPath(MoveTemp(Buf));
+					if (!NormalizedPath.IsEmpty())
+					{
+						OutStringAssetReferencesMap.AddUnique(
+							FPackageName::ObjectPathToPackageName(
+								NormalizedPath
+							)
+						);
+					}
+				}
+			}
+		}
+		else
+		{
+			for (int32 ReferenceIdx = 0; ReferenceIdx < PackageFileSummary.StringAssetReferencesCount; ++ReferenceIdx)
+			{
+				FString Buf;
+				*this << Buf;
+				OutStringAssetReferencesMap.Add(MoveTemp(Buf));
+			}
 		}
 	}
 }
diff --git a/Engine/Source/Runtime/AssetRegistry/Public/IAssetRegistry.h b/Engine/Source/Runtime/AssetRegistry/Public/IAssetRegistry.h
--- a/Engine/Source/Runtime/AssetRegistry/Public/IAssetRegistry.h
+++ b/Engine/Source/Runtime/AssetRegistry/Public/IAssetRegistry.h
@@ -102,8 +102,9 @@ public:
 	 * @param PackageName		the name of the package for which to gather dependencies
 	 * @param OutDependencies	a list of paths to objects that are referenced by the package whose path is PackageName
 	 * @param InDependencyType	which kinds of dependency to include in the output list
+	 * @param bResolveIniStringReferences Tells if the method should also resolve INI references.
 	 */
-	virtual bool GetDependencies(FName PackageName, TArray<FName>& OutDependencies, EAssetRegistryDependencyType::Type InDependencyType = EAssetRegistryDependencyType::All) const = 0;
+	virtual bool GetDependencies(FName PackageName, TArray<FName>& OutDependencies, EAssetRegistryDependencyType::Type InDependencyType = EAssetRegistryDependencyType::All, bool bResolveIniStringReferences = false) const = 0;
 
 	/**
 	 * Gets a list of paths to objects that reference the supplied package. (On disk references ONLY)
diff --git a/Engine/Source/Runtime/Core/Public/UObject/ObjectVersion.h b/Engine/Source/Runtime/Core/Public/UObject/ObjectVersion.h
--- a/Engine/Source/Runtime/Core/Public/UObject/ObjectVersion.h
+++ b/Engine/Source/Runtime/Core/Public/UObject/ObjectVersion.h
@@ -562,6 +562,8 @@ enum EUnrealEngineObjectUE4Version
 	VER_UE4_APEX_CLOTH_TESSELLATION,
 	// Added DecalSize member, solved backward compatibility
 	VER_UE4_DECAL_SIZE,
+	// Keep only package names in StringAssetReferencesMap
+	VER_UE4_KEEP_ONLY_PACKAGE_NAMES_IN_STRING_ASSET_REFERENCES_MAP,
 	// -----<new versions can be added before this line>-------------------------------------------------
 	// - this needs to be the last line (see note below)
 	VER_UE4_AUTOMATIC_VERSION_PLUS_ONE,
diff --git a/Engine/Source/Runtime/CoreUObject/Classes/Object.h b/Engine/Source/Runtime/CoreUObject/Classes/Object.h
--- a/Engine/Source/Runtime/CoreUObject/Classes/Object.h
+++ b/Engine/Source/Runtime/CoreUObject/Classes/Object.h
@@ -896,7 +896,7 @@ struct FTimespan
 
 // A string asset reference
 
-USTRUCT(noexport)
+USTRUCT(noexport, meta=(HasNativeMake="Engine.BlueprintFunctionLibrary.MakeStringAssetReference"))
 struct FStringAssetReference
 {
 	UPROPERTY(EditAnywhere, BlueprintReadWrite, Category=StringAssetReference)
diff --git a/Engine/Source/Runtime/CoreUObject/Private/Misc/PackageName.cpp b/Engine/Source/Runtime/CoreUObject/Private/Misc/PackageName.cpp
--- a/Engine/Source/Runtime/CoreUObject/Private/Misc/PackageName.cpp
+++ b/Engine/Source/Runtime/CoreUObject/Private/Misc/PackageName.cpp
@@ -761,7 +761,54 @@ bool FPackageName::SearchForPackageOnDisk(const FString& PackageName, FString* O
 	return bResult;
 }
 
-FString FPackageName::PackageFromPath( const TCHAR* InPathName )
+bool FPackageName::TryConvertShortPackagePathToLongInObjectPath(const FString& ObjectPath, FString& ConvertedObjectPath)
+{
+	FString PackagePath;
+	FString ObjectName;
+
+	int32 DotPosition = ObjectPath.Find(TEXT("."), ESearchCase::CaseSensitive);
+	if (DotPosition != INDEX_NONE)
+	{
+		PackagePath = ObjectPath.Mid(0, DotPosition);
+		ObjectName = ObjectPath.Mid(DotPosition + 1);
+	}
+	else
+	{
+		PackagePath = ObjectPath;
+	}
+
+	FString LongPackagePath;
+	if (!SearchForPackageOnDisk(PackagePath, &LongPackagePath))
+	{
+		return false;
+	}
+
+	ConvertedObjectPath = FString::Printf(TEXT("%s.%s"), *LongPackagePath, *ObjectName);
+	return true;
+}
+
+FString FPackageName::GetNormalizedObjectPath(const FString& ObjectPath)
+{
+	if (!ObjectPath.IsEmpty() && FPackageName::IsShortPackageName(ObjectPath))
+	{
+		FString LongPath;
+
+		UE_LOG(LogPackageName, Warning, TEXT("String asset reference \"%s\" is in short form, which is unsupported and -- even if valid -- resolving it will be really slow."), *ObjectPath);
+
+		if (!FPackageName::TryConvertShortPackagePathToLongInObjectPath(ObjectPath, LongPath))
+		{
+			UE_LOG(LogPackageName, Warning, TEXT("String asset reference \"%s\" could not be resolved."), *ObjectPath);
+		}
+
+		return LongPath;
+	}
+	else
+	{
+		return ObjectPath;
+	}
+}
+
+FString FPackageName::PackageFromPath(const TCHAR* InPathName)
 {
 	FString PackageName;
 	if (FPackageName::TryConvertFilenameToLongPackageName(InPathName, PackageName))
diff --git a/Engine/Source/Runtime/CoreUObject/Private/Misc/StringAssetReference.cpp b/Engine/Source/Runtime/CoreUObject/Private/Misc/StringAssetReference.cpp
--- a/Engine/Source/Runtime/CoreUObject/Private/Misc/StringAssetReference.cpp
+++ b/Engine/Source/Runtime/CoreUObject/Private/Misc/StringAssetReference.cpp
@@ -9,23 +9,47 @@ FStringAssetReference::FStringAssetReference(const UObject* InObject)
 {
 	if (InObject)
 	{
-		AssetLongPathname = InObject->GetPathName();
+		SetPath(InObject->GetPathName());
 	}
 }
 
+PRAGMA_DISABLE_DEPRECATION_WARNINGS
+FStringAssetReference::~FStringAssetReference()
+{
+
+}
+
+const FString& FStringAssetReference::ToString() const
+{
+	return AssetLongPathname;
+}
+
+void FStringAssetReference::SetPath(FString Path)
+{
+	if (!FPackageName::IsShortPackageName(Path))
+	{
+		AssetLongPathname = MoveTemp(Path);
+	}
+	else
+	{
+		AssetLongPathname = FPackageName::GetNormalizedObjectPath(Path);
+	}
+}
+PRAGMA_POP
+
 bool FStringAssetReference::Serialize(FArchive& Ar)
 {
 #if WITH_EDITOR
 	if (Ar.IsSaving() && Ar.IsPersistent() && FCoreUObjectDelegates::StringAssetReferenceSaving.IsBound())
 	{
-		AssetLongPathname = FCoreUObjectDelegates::StringAssetReferenceSaving.Execute(AssetLongPathname);
+		SetPath(FCoreUObjectDelegates::StringAssetReferenceSaving.Execute(ToString()));
 	}
 #endif // WITH_EDITOR
 	Ar << *this;
 #if WITH_EDITOR
 	if (Ar.IsLoading() && Ar.IsPersistent() && FCoreUObjectDelegates::StringAssetReferenceLoaded.IsBound())
 	{
-		FCoreUObjectDelegates::StringAssetReferenceLoaded.Execute(AssetLongPathname);
+		FCoreUObjectDelegates::StringAssetReferenceLoaded.Execute(ToString());
 	}
 #endif // WITH_EDITOR
 
@@ -39,17 +63,17 @@ bool FStringAssetReference::Serialize(FArchive& Ar)
 }
 bool FStringAssetReference::operator==(FStringAssetReference const& Other) const
 {
-	return AssetLongPathname == Other.AssetLongPathname;
+	return ToString() == Other.ToString();
 }
 void FStringAssetReference::operator=(FStringAssetReference const& Other)
 {
-	AssetLongPathname = Other.AssetLongPathname;
+	SetPath(Other.ToString());
 }
 bool FStringAssetReference::ExportTextItem(FString& ValueStr, FStringAssetReference const& DefaultValue, UObject* Parent, int32 PortFlags, UObject* ExportRootScope) const
 {
-	if (!AssetLongPathname.IsEmpty())
+	if (IsValid())
 	{
-		ValueStr += AssetLongPathname;
+		ValueStr += ToString();
 	}
 	else
 	{
@@ -57,41 +81,43 @@ bool FStringAssetReference::ExportTextItem(FString& ValueStr, FStringAssetRefere
 	}
 	return true;
 }
-bool FStringAssetReference::ImportTextItem( const TCHAR*& Buffer, int32 PortFlags, UObject* Parent, FOutputDevice* ErrorText )
+bool FStringAssetReference::ImportTextItem(const TCHAR*& Buffer, int32 PortFlags, UObject* Parent, FOutputDevice* ErrorText)
 {
-	AssetLongPathname = TEXT("");
-	const TCHAR* NewBuffer = UPropertyHelpers::ReadToken( Buffer, AssetLongPathname, 1 );
-	if( !NewBuffer )
+	FString ImportedPath = TEXT("");
+	const TCHAR* NewBuffer = UPropertyHelpers::ReadToken(Buffer, ImportedPath, 1);
+	if (!NewBuffer)
 	{
 		return false;
 	}
 	Buffer = NewBuffer;
-	if( AssetLongPathname==TEXT("None") )
+	if (ImportedPath == TEXT("None"))
 	{
-		AssetLongPathname = TEXT("");
+		ImportedPath = TEXT("");
 	}
 	else
 	{
-		if( *Buffer == TCHAR('\'') )
+		if (*Buffer == TCHAR('\''))
 		{
-			NewBuffer = UPropertyHelpers::ReadToken( Buffer, AssetLongPathname, 1 );
-			if( !NewBuffer )
+			NewBuffer = UPropertyHelpers::ReadToken(Buffer, ImportedPath, 1);
+			if (!NewBuffer)
 			{
 				return false;
 			}
 			Buffer = NewBuffer;
-			if( *Buffer++ != TCHAR('\'') )
+			if (*Buffer++ != TCHAR('\''))
 			{
 				return false;
 			}
 		}
 	}
 
+	SetPath(MoveTemp(ImportedPath));
+
 #if WITH_EDITOR
 	// Consider this a load, so Config string asset references get cooked
 	if (FCoreUObjectDelegates::StringAssetReferenceLoaded.IsBound())
 	{
-		FCoreUObjectDelegates::StringAssetReferenceLoaded.Execute(AssetLongPathname);
+		FCoreUObjectDelegates::StringAssetReferenceLoaded.Execute(ToString());
 	}
 #endif // WITH_EDITOR
 
@@ -108,7 +134,16 @@ bool FStringAssetReference::SerializeFromMismatchedTag(struct FPropertyTag const
 		static const FName FORCEINLINE GetTypeName() { return NAME_ObjectProperty; }
 	};
 
-	return SerializeFromMismatchedTagTemplate<UObjectTypePolicy>(AssetLongPathname, Tag, Ar);
+	FString Path = ToString();
+
+	bool bReturn = SerializeFromMismatchedTagTemplate<UObjectTypePolicy>(Path, Tag, Ar);
+
+	if (Ar.IsLoading())
+	{
+		SetPath(MoveTemp(Path));
+	}
+
+	return bReturn;
 }
 
 UObject* FStringAssetReference::TryLoad() const
@@ -173,7 +208,7 @@ void FStringAssetReference::FixupForPIE()
 		{
 			check(GPlayInEditorID != -1);
 
-			const FString PIEPath = FString::Printf(TEXT("%s/%s_%d_%s"), *FPackageName::GetLongPackagePath(Path), PLAYWORLD_PACKAGE_PREFIX, GPlayInEditorID, *ShortPackageOuterAndName);
+			FString PIEPath = FString::Printf(TEXT("%s/%s_%d_%s"), *FPackageName::GetLongPackagePath(Path), PLAYWORLD_PACKAGE_PREFIX, GPlayInEditorID, *ShortPackageOuterAndName);
 
 			// Determine if this refers to a package that is being duplicated for PIE
 			for (auto PackageNameIt = PackageNamesBeingDuplicatedForPIE.CreateConstIterator(); PackageNameIt; ++PackageNameIt)
@@ -182,7 +217,7 @@ void FStringAssetReference::FixupForPIE()
 				if (PIEPath.StartsWith(PathPrefix))
 				{
 					// Need to prepend PIE prefix, as we're in PIE and this refers to an object in a PIE package
-					AssetLongPathname = PIEPath;
+					SetPath(MoveTemp(PIEPath));
 
 					break;
 				}
diff --git a/Engine/Source/Runtime/CoreUObject/Private/Misc/StringClassReference.cpp b/Engine/Source/Runtime/CoreUObject/Private/Misc/StringClassReference.cpp
--- a/Engine/Source/Runtime/CoreUObject/Private/Misc/StringClassReference.cpp
+++ b/Engine/Source/Runtime/CoreUObject/Private/Misc/StringClassReference.cpp
@@ -15,7 +15,16 @@ bool FStringClassReference::SerializeFromMismatchedTag(struct FPropertyTag const
 		static const FName FORCEINLINE GetTypeName() { return NAME_ClassProperty; }
 	};
 
-	return SerializeFromMismatchedTagTemplate<UClassTypePolicy>(AssetLongPathname, Tag, Ar);
+	FString Path = ToString();
+
+	bool bReturn = SerializeFromMismatchedTagTemplate<UClassTypePolicy>(Path, Tag, Ar);
+
+	if (Ar.IsLoading())
+	{
+		SetPath(MoveTemp(Path));
+	}
+
+	return bReturn;
 }
 
 UClass* FStringClassReference::ResolveClass() const
diff --git a/Engine/Source/Runtime/CoreUObject/Private/Serialization/ArchiveUObject.cpp b/Engine/Source/Runtime/CoreUObject/Private/Serialization/ArchiveUObject.cpp
--- a/Engine/Source/Runtime/CoreUObject/Private/Serialization/ArchiveUObject.cpp
+++ b/Engine/Source/Runtime/CoreUObject/Private/Serialization/ArchiveUObject.cpp
@@ -67,7 +67,26 @@ FArchive& FArchiveUObject::operator<<( class FAssetPtr& AssetPtr )
 
 FArchive& FArchiveUObject::operator<<(struct FStringAssetReference& Value)
 {
-	*this << Value.AssetLongPathname;
+	FString Path = Value.ToString();
+
+	*this << Path;
+
+	if (IsLoading())
+	{
+		if (UE4Ver() < VER_UE4_KEEP_ONLY_PACKAGE_NAMES_IN_STRING_ASSET_REFERENCES_MAP)
+		{
+			FString NormalizedPath = FPackageName::GetNormalizedObjectPath(Path);
+			if (Value.ToString() != NormalizedPath)
+			{
+				Value.SetPath(NormalizedPath);
+			}
+		}
+		else
+		{
+			Value.SetPath(MoveTemp(Path));
+		}
+	}
+
 	return *this;
 }
 
diff --git a/Engine/Source/Runtime/CoreUObject/Private/Serialization/ObjectReader.cpp b/Engine/Source/Runtime/CoreUObject/Private/Serialization/ObjectReader.cpp
--- a/Engine/Source/Runtime/CoreUObject/Private/Serialization/ObjectReader.cpp
+++ b/Engine/Source/Runtime/CoreUObject/Private/Serialization/ObjectReader.cpp
@@ -46,7 +46,16 @@ FArchive& FObjectReader::operator<<( class FAssetPtr& AssetPtr )
 
 FArchive& FObjectReader::operator<<(FStringAssetReference& Value)
 {
-	return *this << Value.AssetLongPathname;
+	FString Path = Value.ToString();
+
+	*this << Path;
+
+	if (IsLoading())
+	{
+		Value.SetPath(MoveTemp(Path));
+	}
+
+	return *this;
 }
 
 FString FObjectReader::GetArchiveName() const
diff --git a/Engine/Source/Runtime/CoreUObject/Private/Serialization/ObjectWriter.cpp b/Engine/Source/Runtime/CoreUObject/Private/Serialization/ObjectWriter.cpp
--- a/Engine/Source/Runtime/CoreUObject/Private/Serialization/ObjectWriter.cpp
+++ b/Engine/Source/Runtime/CoreUObject/Private/Serialization/ObjectWriter.cpp
@@ -37,7 +37,16 @@ FArchive& FObjectWriter::operator<<( class FAssetPtr& AssetPtr )
 
 FArchive& FObjectWriter::operator<<(FStringAssetReference& Value)
 {
-	return *this << Value.AssetLongPathname;
+	FString Path = Value.ToString();
+
+	*this << Path;
+
+	if (IsLoading())
+	{
+		Value.SetPath(MoveTemp(Path));
+	}
+
+	return *this;
 }
 
 FString FObjectWriter::GetArchiveName() const
diff --git a/Engine/Source/Runtime/CoreUObject/Private/UObject/CoreNet.cpp b/Engine/Source/Runtime/CoreUObject/Private/UObject/CoreNet.cpp
--- a/Engine/Source/Runtime/CoreUObject/Private/UObject/CoreNet.cpp
+++ b/Engine/Source/Runtime/CoreUObject/Private/UObject/CoreNet.cpp
@@ -351,7 +351,16 @@ FArchive& FNetBitWriter::operator<<( UObject*& Object )
 
 FArchive& FNetBitWriter::operator<<(FStringAssetReference& Value)
 {
-	return *this << Value.AssetLongPathname;
+	FString Path = Value.ToString();
+
+	*this << Path;
+
+	if (IsLoading())
+	{
+		Value.SetPath(MoveTemp(Path));
+	}
+
+	return *this;
 }
 
 
@@ -378,7 +387,16 @@ FArchive& FNetBitReader::operator<<( class FName& N )
 
 FArchive& FNetBitReader::operator<<(FStringAssetReference& Value)
 {
-	return *this << Value.AssetLongPathname;
+	FString Path = Value.ToString();
+
+	*this << Path;
+
+	if (IsLoading())
+	{
+		Value.SetPath(MoveTemp(Path));
+	}
+
+	return *this;
 }
 
 static const TCHAR* GLastRPCFailedReason = NULL;
diff --git a/Engine/Source/Runtime/CoreUObject/Private/UObject/SavePackage.cpp b/Engine/Source/Runtime/CoreUObject/Private/UObject/SavePackage.cpp
--- a/Engine/Source/Runtime/CoreUObject/Private/UObject/SavePackage.cpp
+++ b/Engine/Source/Runtime/CoreUObject/Private/UObject/SavePackage.cpp
@@ -627,9 +627,22 @@ public:
 	FArchive& operator<<(FAssetPtr& AssetPtr) override;
 	FArchive& operator<<(FStringAssetReference& Value) override
 	{
-		if ( Value.IsValid() )
+		if (Value.IsValid())
 		{
-			StringAssetReferencesMap.Add(Value.ToString());
+			const FString& Path = Value.ToString();
+			if (GetIniFilenameFromObjectsReference(Path) != nullptr)
+			{
+				StringAssetReferencesMap.AddUnique(Path);
+			}
+			else
+			{
+				FString NormalizedPath = FPackageName::GetNormalizedObjectPath(Path);
+				if (!NormalizedPath.IsEmpty())
+				{
+					StringAssetReferencesMap.AddUnique(FPackageName::ObjectPathToPackageName(NormalizedPath));
+					Value.SetPath(MoveTemp(NormalizedPath));
+				}
+			}
 		}
 		return *this;
 	}
diff --git a/Engine/Source/Runtime/CoreUObject/Private/UObject/UObjectGlobals.cpp b/Engine/Source/Runtime/CoreUObject/Private/UObject/UObjectGlobals.cpp
--- a/Engine/Source/Runtime/CoreUObject/Private/UObject/UObjectGlobals.cpp
+++ b/Engine/Source/Runtime/CoreUObject/Private/UObject/UObjectGlobals.cpp
@@ -569,55 +569,65 @@ UPackage* CreatePackage( UObject* InOuter, const TCHAR* PackageName )
 	return Result;
 }
 
-//
-// Resolve a package and name.
-//
-bool ResolveName( UObject*& InPackage, FString& InOutName, bool Create, bool Throw )
+FString ResolveIniObjectsReference(const FString& ObjectReference, const FString* IniFilename, bool bThrow)
 {
-	FString* IniFilename = NULL;
+	// Get .ini key and section.
+	FString Section = ObjectReference.Mid(1 + ObjectReference.Find(TEXT(":"), ESearchCase::CaseSensitive));
+	int32 i = Section.Find(TEXT("."), ESearchCase::CaseSensitive, ESearchDir::FromEnd);
+	FString Key;
+	if (i != -1)
+	{
+		Key = Section.Mid(i + 1);
+		Section = Section.Left(i);
+	}
+
+	FString Output;
+
+	// Look up name.
+	if (!GConfig->GetString(*Section, *Key, Output, *IniFilename))
+	{
+		if (bThrow == true)
+		{
+			UE_LOG(LogUObjectGlobals, Error, TEXT(" %s %s "), *FString::Printf(TEXT("Can't find '%s' in configuration file section=%s key=%s"), *ObjectReference, *Section, *Key), **IniFilename);
+		}
+	}
 
+	return Output;
+}
+
+const FString* GetIniFilenameFromObjectsReference(const FString& Name)
+{
 	// See if the name is specified in the .ini file.
-	if( FCString::Strnicmp( *InOutName, TEXT("engine-ini:"), FCString::Strlen(TEXT("engine-ini:")) )==0 )
+	if (FCString::Strnicmp(*Name, TEXT("engine-ini:"), FCString::Strlen(TEXT("engine-ini:"))) == 0)
 	{
-		IniFilename = &GEngineIni;
+		return &GEngineIni;
 	}
-	else if( FCString::Strnicmp( *InOutName, TEXT("game-ini:"), FCString::Strlen(TEXT("game-ini:")) )==0 )
+	else if (FCString::Strnicmp(*Name, TEXT("game-ini:"), FCString::Strlen(TEXT("game-ini:"))) == 0)
 	{
-		IniFilename = &GGameIni;
+		return &GGameIni;
 	}
-	else if( FCString::Strnicmp( *InOutName, TEXT("input-ini:"), FCString::Strlen(TEXT("input-ini:")) )==0 )
+	else if (FCString::Strnicmp(*Name, TEXT("input-ini:"), FCString::Strlen(TEXT("input-ini:"))) == 0)
 	{
-		IniFilename = &GInputIni;
+		return &GInputIni;
 	}
-	else if( FCString::Strnicmp( *InOutName, TEXT("editor-ini:"), FCString::Strlen(TEXT("editor-ini:")) )==0 )
+	else if (FCString::Strnicmp(*Name, TEXT("editor-ini:"), FCString::Strlen(TEXT("editor-ini:"))) == 0)
 	{
-		IniFilename = &GEditorIni;
+		return &GEditorIni;
 	}
 
+	return nullptr;
+}
 
-	if( IniFilename && InOutName.Contains(TEXT("."), ESearchCase::CaseSensitive) )
-	{
-		// Get .ini key and section.
-		FString Section = InOutName.Mid(1+InOutName.Find(TEXT(":"), ESearchCase::CaseSensitive));
-		int32 i = Section.Find(TEXT("."), ESearchCase::CaseSensitive, ESearchDir::FromEnd);
-		FString Key;
-		if( i != -1)
-		{
-			Key = Section.Mid(i+1);
-			Section = Section.Left(i);
-		}
+//
+// Resolve a package and name.
+//
+bool ResolveName( UObject*& InPackage, FString& InOutName, bool Create, bool Throw )
+{
+	const FString* IniFilename = GetIniFilenameFromObjectsReference(InOutName);
 
-		// Look up name.
-		FString Result;
-		if( !GConfig->GetString( *Section, *Key, Result, *IniFilename ) )
-		{
-			if( Throw == true )
-			{
-				UE_LOG(LogUObjectGlobals, Error, TEXT( " %s %s " ), *FString::Printf( TEXT("Can't find '%s' in configuration file section=%s key=%s"), *InOutName, *Section, *Key), **IniFilename );
-			}
-			return false;
-		}
-		InOutName = Result;
+	if (IniFilename && InOutName.Contains(TEXT("."), ESearchCase::CaseSensitive))
+	{
+		InOutName = ResolveIniObjectsReference(InOutName, IniFilename, Throw);
 	}
 
 	// Strip off the object class.
diff --git a/Engine/Source/Runtime/CoreUObject/Public/Misc/PackageName.h b/Engine/Source/Runtime/CoreUObject/Public/Misc/PackageName.h
--- a/Engine/Source/Runtime/CoreUObject/Public/Misc/PackageName.h
+++ b/Engine/Source/Runtime/CoreUObject/Public/Misc/PackageName.h
@@ -196,6 +196,25 @@ public:
 	 **/
 	static bool SearchForPackageOnDisk(const FString& PackageName, FString* OutLongPackageName = NULL, FString* OutFilename = NULL, bool bUseLocalizedNames = false);
 
+	/**
+	 * Tries to convert object path with short package name to object path with long package name found on disk (very slow)
+	 *
+	 * @param ObjectPath Path to the object.
+	 * @param OutLongPackageName Converted object path.
+	 *
+	 * @returns True if succeeded. False otherwise.
+	 */
+	static bool TryConvertShortPackagePathToLongInObjectPath(const FString& ObjectPath, FString& ConvertedObjectPath);
+
+	/**
+	 * Gets normlized object path i.e. with long package format.
+	 *
+	 * @param ObjectPath Path to the object.
+	 *
+	 * @returns Normalized path (or empty path, if short object path was given and it wasn't found on the disk).
+	 */
+	static FString GetNormalizedObjectPath(const FString& ObjectPath);
+
 	/**
 	 * Strips all path and extension information from a relative or fully qualified file name.
 	 *
diff --git a/Engine/Source/Runtime/CoreUObject/Public/Misc/StringAssetReference.h b/Engine/Source/Runtime/CoreUObject/Public/Misc/StringAssetReference.h
--- a/Engine/Source/Runtime/CoreUObject/Public/Misc/StringAssetReference.h
+++ b/Engine/Source/Runtime/CoreUObject/Public/Misc/StringAssetReference.h
@@ -11,6 +11,7 @@ struct FPropertyTag;
 struct COREUOBJECT_API FStringAssetReference
 {
 	/** Asset path */
+	DEPRECATED(4.9, "Please don't use AssetLongPathname directly. Instead use SetPath and ToString methods.")
 	FString AssetLongPathname;
 	
 	FStringAssetReference()
@@ -18,20 +19,16 @@ struct COREUOBJECT_API FStringAssetReference
 	}
 
 	FStringAssetReference(const FStringAssetReference& Other)
-		: AssetLongPathname(Other.AssetLongPathname)
 	{
+		SetPath(Other.ToString());
 	}
 
 	/**
 	 * Construct from a path string
 	 */
-	FStringAssetReference(const FString& PathString)
-		: AssetLongPathname(PathString)
+	FStringAssetReference(FString PathString)
 	{
-		if (AssetLongPathname == TEXT("None"))
-		{
-			AssetLongPathname = TEXT("");
-		}
+		SetPath(MoveTemp(PathString));
 	}
 
 	/**
@@ -39,21 +36,27 @@ struct COREUOBJECT_API FStringAssetReference
 	 */
 	FStringAssetReference(const UObject* InObject);
 
+	~FStringAssetReference();
+
 	/**
 	 * Converts in to a string
 	 */
-	const FString& ToString() const 
-	{
-		return AssetLongPathname;
-	}
+	const FString& ToString() const;
 
 	const FString GetLongPackageName() const
 	{
 		FString PackageName;
-		AssetLongPathname.Split(TEXT("."), &PackageName, nullptr, ESearchCase::CaseSensitive, ESearchDir::FromEnd);
+		ToString().Split(TEXT("."), &PackageName, nullptr, ESearchCase::CaseSensitive, ESearchDir::FromEnd);
 		return PackageName;
 	}
 
+	/**
+	 * Sets asset path of this reference.
+	 *
+	 * @param Path The path to the asset.
+	 */
+	void SetPath(FString Path);
+
 	/**
 	 * Attempts to load the asset.
 	 * @return Loaded UObject, or null if the asset fails to load, or if the reference is not valid.
@@ -71,7 +74,7 @@ struct COREUOBJECT_API FStringAssetReference
 	 */
 	void Reset()
 	{		
-		AssetLongPathname = TEXT("");
+		SetPath(TEXT(""));
 	}
 	
 	/**
@@ -79,7 +82,7 @@ struct COREUOBJECT_API FStringAssetReference
 	 */
 	bool IsValid() const
 	{
-		return AssetLongPathname.Len() > 0;
+		return ToString().Len() > 0;
 	}
 
 	bool Serialize(FArchive& Ar);
@@ -95,7 +98,7 @@ struct COREUOBJECT_API FStringAssetReference
 
 	FORCEINLINE friend uint32 GetTypeHash(FStringAssetReference const& This)
 	{
-		return GetTypeHash(This.AssetLongPathname);
+		return GetTypeHash(This.ToString());
 	}
 
 	/** Code needed by AssetPtr to track rather object references should be rechecked */
diff --git a/Engine/Source/Runtime/CoreUObject/Public/UObject/UObjectGlobals.h b/Engine/Source/Runtime/CoreUObject/Public/UObject/UObjectGlobals.h
--- a/Engine/Source/Runtime/CoreUObject/Public/UObject/UObjectGlobals.h
+++ b/Engine/Source/Runtime/CoreUObject/Public/UObject/UObjectGlobals.h
@@ -121,7 +121,26 @@ COREUOBJECT_API TArray<const TCHAR*> ParsePropertyFlags(uint64 Flags);
 
 COREUOBJECT_API UPackage* GetTransientPackage();
 
-COREUOBJECT_API bool ResolveName( UObject*& Outer, FString& Name, bool Create, bool Throw );
+/**
+ * Gets INI file name from object's reference if it contains one.
+ *
+ * @returns If object reference doesn't contain any INI reference the function
+ *		returns nullptr. Otherwise a ptr to INI's file name.
+ */
+COREUOBJECT_API const FString* GetIniFilenameFromObjectsReference(const FString& ObjectsReferenceString);
+
+/**
+ * Resolves ini object path to string object path.
+ *
+ * @param ObjectReference Ini reference.
+ * @param IniFilename Ini filename.
+ * @param bThrow Can this function throw?
+ *
+ * @returns Resolved object path.
+ */
+COREUOBJECT_API FString ResolveIniObjectsReference(const FString& ObjectReference, const FString* IniFilename, bool bThrow = false);
+
+COREUOBJECT_API bool ResolveName( UObject*& Outer, FString& ObjectsReferenceString, bool Create, bool Throw );
 COREUOBJECT_API void SafeLoadError( UObject* Outer, uint32 LoadFlags, const TCHAR* ErrorMessage);
 
 /**
diff --git a/Engine/Source/Runtime/Engine/Classes/Kismet/BlueprintFunctionLibrary.h b/Engine/Source/Runtime/Engine/Classes/Kismet/BlueprintFunctionLibrary.h
--- a/Engine/Source/Runtime/Engine/Classes/Kismet/BlueprintFunctionLibrary.h
+++ b/Engine/Source/Runtime/Engine/Classes/Kismet/BlueprintFunctionLibrary.h
@@ -14,6 +14,22 @@ class UBlueprintFunctionLibrary : public UObject
 	ENGINE_API virtual int32 GetFunctionCallspace(UFunction* Function, void* Parms, FFrame* Stack) override;
 	// End of UObject interface
 
+	UFUNCTION(BlueprintPure, Category="StringAssetReference", CustomThunk, meta = (Keywords = "construct build", NativeMakeFunc))
+	static FStringAssetReference MakeStringAssetReference(const FString& AssetLongPathname);
+
+	/**
+	 * Custom MakeStringAssetReference blueprint node.
+	 */
+	static FStringAssetReference Generic_MakeStringAssetReference(FFrame& StackFrame, const FString& AssetLongPathname);
+
+	DECLARE_FUNCTION(execMakeStringAssetReference)
+	{
+		P_GET_PROPERTY(UStrProperty, Z_Param_AssetLongPathname);
+		P_FINISH;
+
+		*(FStringAssetReference*)Z_Param__Result = UBlueprintFunctionLibrary::Generic_MakeStringAssetReference(Stack, Z_Param_AssetLongPathname);
+	}
+
 protected:
 
 	//bool	CheckAuthority(UObject* 
diff --git a/Engine/Source/Runtime/Engine/Private/KismetFunctionLibrary.cpp b/Engine/Source/Runtime/Engine/Private/KismetFunctionLibrary.cpp
--- a/Engine/Source/Runtime/Engine/Private/KismetFunctionLibrary.cpp
+++ b/Engine/Source/Runtime/Engine/Private/KismetFunctionLibrary.cpp
@@ -25,3 +25,16 @@ int32 UBlueprintFunctionLibrary::GetFunctionCallspace(UFunction* Function, void*
 						(Function->HasAllFunctionFlags(FUNC_BlueprintCosmetic) && GEngine->ShouldAbsorbCosmeticOnlyEvent());
 	return Absorb ? FunctionCallspace::Absorbed : FunctionCallspace::Local;
 }
+
+FStringAssetReference UBlueprintFunctionLibrary::Generic_MakeStringAssetReference(FFrame& Stack, const FString& AssetLongPathname)
+{
+	FStringAssetReference Ref(AssetLongPathname);
+	if (!AssetLongPathname.IsEmpty() && !Ref.IsValid())
+	{
+		FBlueprintExceptionInfo Info(EBlueprintExceptionType::FatalError, TEXT("Asset path not valid. Only long path name is allowed."));
+		FBlueprintCoreDelegates::ThrowScriptException(Stack.Object, Stack, Info);
+		return FStringAssetReference();
+	}
+
+	return Ref;
+}
\ No newline at end of file
diff --git a/Engine/Source/Runtime/Engine/Private/UserInterfaceSettings.cpp b/Engine/Source/Runtime/Engine/Private/UserInterfaceSettings.cpp
--- a/Engine/Source/Runtime/Engine/Private/UserInterfaceSettings.cpp
+++ b/Engine/Source/Runtime/Engine/Private/UserInterfaceSettings.cpp
@@ -41,7 +41,7 @@ float UUserInterfaceSettings::GetDPIScaleBasedOnSize(FIntPoint Size) const
 			if ( CustomScalingRuleClassInstance == nullptr )
 			{
 				FMessageLog("MapCheck").Error()
-					->AddToken(FTextToken::Create(FText::Format(LOCTEXT("CustomScalingRule_NotFound", "Project Settings - User Interface Custom Scaling Rule '{0}' could not be found."), FText::FromString(CustomScalingRuleClass.AssetLongPathname))));
+					->AddToken(FTextToken::Create(FText::Format(LOCTEXT("CustomScalingRule_NotFound", "Project Settings - User Interface Custom Scaling Rule '{0}' could not be found."), FText::FromString(CustomScalingRuleClass.ToString()))));
 				return 1;
 			}
 		}
diff --git a/Engine/Source/Runtime/GameplayAbilities/Private/GameplayCueManager.cpp b/Engine/Source/Runtime/GameplayAbilities/Private/GameplayCueManager.cpp
--- a/Engine/Source/Runtime/GameplayAbilities/Private/GameplayCueManager.cpp
+++ b/Engine/Source/Runtime/GameplayAbilities/Private/GameplayCueManager.cpp
@@ -298,7 +298,7 @@ void UGameplayCueManager::BuildCuesToAddToGlobalSet(const TArray<FAssetData>& As
 			{
 				// Add a new NotifyData entry to our flat list for this one
 				FStringAssetReference StringRef;
-				StringRef.AssetLongPathname = FPackageName::ExportTextPathToObjectPath(*GeneratedClassTag);
+				StringRef.SetPath(FPackageName::ExportTextPathToObjectPath(*GeneratedClassTag));
 
 				OutCuesToAdd.Add(FGameplayCueReferencePair(GameplayCueTag, StringRef));
 
@@ -354,7 +354,7 @@ void UGameplayCueManager::HandleAssetAdded(UObject *Object)
 			if (IsAssetInLoadedPaths(Object))
 			{
 				FStringAssetReference StringRef;
-				StringRef.AssetLongPathname = Blueprint->GeneratedClass->GetPathName();
+				StringRef.SetPath(Blueprint->GeneratedClass->GetPathName());
 
 				TArray<FGameplayCueReferencePair> CuesToAdd;
 				if (StaticCDO)
@@ -387,7 +387,7 @@ void UGameplayCueManager::HandleAssetDeleted(UObject *Object)
 		
 		if (StaticCDO || ActorCDO)
 		{
-			StringRefToRemove.AssetLongPathname = Blueprint->GeneratedClass->GetPathName();
+			StringRefToRemove.SetPath(Blueprint->GeneratedClass->GetPathName());
 		}
 	}
 
