[UE-18384] Fix for actors in hidden levels losing their transform data and becoming visible

diff --git a/Engine/Source/Editor/UnrealEd/Private/Kismet2/KismetReinstanceUtilities.cpp b/Engine/Source/Editor/UnrealEd/Private/Kismet2/KismetReinstanceUtilities.cpp
--- a/Engine/Source/Editor/UnrealEd/Private/Kismet2/KismetReinstanceUtilities.cpp
+++ b/Engine/Source/Editor/UnrealEd/Private/Kismet2/KismetReinstanceUtilities.cpp
@@ -761,6 +761,19 @@ void FActorReplacementHelper::Finalize(const TMap<UObject*, UObject*>& OldToNewI
 		NewActor->ExecuteConstruction(TargetWorldTransform, &DummyComponentData);
 	}	
 
+	// make sure that the actor is properly hidden if it's in a hidden sublevel:
+	bool bIsInHiddenLevel = false;
+	if (ULevel* Level = NewActor->GetLevel())
+	{
+		bIsInHiddenLevel = !Level->bIsVisible;
+	}
+
+	if (bIsInHiddenLevel)
+	{
+		NewActor->bHiddenEdLevel = true;
+		NewActor->MarkComponentsRenderStateDirty();
+	}
+
 	if (TargetAttachParent)
 	{
 		UObject* const* NewTargetAttachParent = OldToNewInstanceMap.Find(TargetAttachParent);
@@ -981,6 +994,16 @@ void FBlueprintCompileReinstancer::ReplaceInstancesOfClass(UClass* OldClass, UCl
 				FRotator Rotation = FRotator::ZeroRotator;
 				if (USceneComponent* OldRootComponent = OldActor->GetRootComponent())
 				{
+					// We need to make sure that the ComponentToWorld transform is up to date, but we don't want to run any initialization logic
+					// so we silence the update, cache it off, revert the change (so no events are raised), and then directly update the transform
+					// with the value calculated in ConditionalUpdateComponentToWorld:
+					FScopedMovementUpdate SilenceMovement(OldRootComponent);
+					
+					OldRootComponent->ConditionalUpdateComponentToWorld();
+					FTransform OldComponentToWorld = OldRootComponent->ComponentToWorld;
+					SilenceMovement.RevertMove();
+					
+					OldRootComponent->ComponentToWorld = OldComponentToWorld;
 					Location = OldActor->GetActorLocation();
 					Rotation = OldActor->GetActorRotation();
 				}
