UE-19126 Fix crash using auto-convex when no sections have collision enabled
UE-19046, UE-19047 Fix non-escaped character and simplify status message

diff --git a/Engine/Source/Editor/StaticMeshEditor/Private/StaticMeshEditor.cpp b/Engine/Source/Editor/StaticMeshEditor/Private/StaticMeshEditor.cpp
--- a/Engine/Source/Editor/StaticMeshEditor/Private/StaticMeshEditor.cpp
+++ b/Engine/Source/Editor/StaticMeshEditor/Private/StaticMeshEditor.cpp
@@ -1699,8 +1699,11 @@ void FStaticMeshEditor::DoDecomp(float InAccuracy, int32 InMaxHullVerts)
 			bs = StaticMesh->BodySetup;
 		}
 
-		// Run actual util to do the work
-		DecomposeMeshToHulls(bs, Verts, CollidingIndices, InAccuracy, InMaxHullVerts);		
+		// Run actual util to do the work (if we have some valid input)
+		if(Verts.Num() >= 3 && CollidingIndices.Num() >= 3)
+		{
+			DecomposeMeshToHulls(bs, Verts, CollidingIndices, InAccuracy, InMaxHullVerts);		
+		}
 
 		// Enable collision, if not already
 		if( !Viewport->GetViewportClient().IsSetShowWireframeCollisionChecked() )
diff --git a/Engine/Source/Editor/UnrealEd/Private/ConvexDecompTool.cpp b/Engine/Source/Editor/UnrealEd/Private/ConvexDecompTool.cpp
--- a/Engine/Source/Editor/UnrealEd/Private/ConvexDecompTool.cpp
+++ b/Engine/Source/Editor/UnrealEd/Private/ConvexDecompTool.cpp
@@ -53,9 +53,10 @@ public:
 
 	void Update(const double overallProgress, const double stageProgress, const double operationProgress, const char * const stage,	const char * const    operation)
 	{
-		FString StatusString = FString::Printf(TEXT("Processing...  [%s %d%]  [%s %d%]"), ANSI_TO_TCHAR(stage), (int)(stageProgress + 0.5), ANSI_TO_TCHAR(operation), (int)(operationProgress + 0.5));
+		FString StatusString = FString::Printf(TEXT("Processing [%s]..."), ANSI_TO_TCHAR(stage));
 
-		GWarn->StatusUpdate(overallProgress*10.f, 1000, FText::FromString(StatusString));
+		GWarn->StatusUpdate(stageProgress*10.f, 1000, FText::FromString(StatusString));
+		//GWarn->StatusUpdate(overallProgress*10.f, 1000, FText::FromString(StatusString));
 	};
 };
 #else
