Fixed a crash caused by nullptr dereference during asset auto reimport

This addresses UE-19121 - Crashed auto-reimporting FBX

diff --git a/Engine/Source/Editor/UnrealEd/Private/AutoReimport/ReimportFeedbackContext.cpp b/Engine/Source/Editor/UnrealEd/Private/AutoReimport/ReimportFeedbackContext.cpp
--- a/Engine/Source/Editor/UnrealEd/Private/AutoReimport/ReimportFeedbackContext.cpp
+++ b/Engine/Source/Editor/UnrealEd/Private/AutoReimport/ReimportFeedbackContext.cpp
@@ -557,8 +557,11 @@ void FReimportFeedbackContext::ProgressReported(const float TotalProgressInterp,
 
 void FReimportFeedbackContext::FinalizeSlowTask()
 {
-	SlowTaskText->SetVisibility(EVisibility::Collapsed);
-	SlowTaskText = nullptr;
+	if (SlowTaskText.IsValid())
+	{
+		SlowTaskText->SetVisibility(EVisibility::Collapsed);
+		SlowTaskText = nullptr;
+	}
 }
 
 #undef LOCTEXT_NAMESPACE
\ No newline at end of file
