[UE-17675] Making sure to duplicate instanced sub-objects (that are not default) upon reinstancing.

diff --git a/Engine/Source/Editor/UnrealEd/Private/Kismet2/KismetReinstanceUtilities.cpp b/Engine/Source/Editor/UnrealEd/Private/Kismet2/KismetReinstanceUtilities.cpp
--- a/Engine/Source/Editor/UnrealEd/Private/Kismet2/KismetReinstanceUtilities.cpp
+++ b/Engine/Source/Editor/UnrealEd/Private/Kismet2/KismetReinstanceUtilities.cpp
@@ -919,6 +919,146 @@ void FActorReplacementHelper::AttachChildActors(USceneComponent* RootComponent,
 	}
 }
 
+// 
+namespace InstancedPropertyUtils
+{
+	typedef TMap<FName, UObject*> FInstancedPropertyMap;
+
+	/** 
+	 * Aids in finding instanced property values that will not be duplicated nor
+	 * copied in CopyPropertiesForUnRelatedObjects().
+	 */
+	class FArchiveInstancedSubObjCollector : public FArchiveUObject
+	{
+	public:
+		//----------------------------------------------------------------------
+		FArchiveInstancedSubObjCollector(UObject* TargetObj, FInstancedPropertyMap& PropertyMapOut, bool bAutoSerialize = true)
+			: Target(TargetObj)
+			, InstancedPropertyMap(PropertyMapOut)
+		{
+			ArIsObjectReferenceCollector = true;
+			ArIsPersistent = false;
+			ArIgnoreArchetypeRef = false;
+
+			if (bAutoSerialize)
+			{
+				RunSerialization();
+			}
+		}
+
+		//----------------------------------------------------------------------
+		FArchive& operator<<(UObject*& Obj)
+		{
+			if (Obj != nullptr)
+			{
+				UProperty* SerializingProperty = GetSerializedProperty();
+				const bool bHasInstancedValue = SerializingProperty->HasAnyPropertyFlags(CPF_PersistentInstance);
+
+				// default sub-objects are handled by CopyPropertiesForUnrelatedObjects()
+				if (bHasInstancedValue && !Obj->IsDefaultSubobject())
+				{
+					
+					UObject* ObjOuter = Obj->GetOuter();
+					bool bIsSubObject = (ObjOuter == Target);
+					// @TODO: handle nested sub-objects when we're more clear on 
+					//        how this'll affect the makeup of the reinstanced object
+// 					while (!bIsSubObject && (ObjOuter != nullptr))
+// 					{
+// 						ObjOuter = ObjOuter->GetOuter();
+// 						bIsSubObject |= (ObjOuter == Target);
+// 					}
+
+					if (bIsSubObject)
+					{
+						InstancedPropertyMap.Add(SerializingProperty->GetFName(), Obj);
+					}
+				}
+			}
+			return *this;
+		}
+
+		//----------------------------------------------------------------------
+		void RunSerialization()
+		{
+			InstancedPropertyMap.Empty();
+			if (Target != nullptr)
+			{
+				Target->Serialize(*this);
+			}
+		}
+
+	private:
+		UObject* Target;
+		FInstancedPropertyMap& InstancedPropertyMap;
+	};
+
+	/** 
+	 * Duplicates and assigns instanced property values that may have been 
+	 * missed by CopyPropertiesForUnRelatedObjects().
+	 */
+	class FArchiveInsertInstancedSubObjects : public FArchiveUObject
+	{
+	public:
+		//----------------------------------------------------------------------
+		FArchiveInsertInstancedSubObjects(UObject* TargetObj, const FInstancedPropertyMap& OldInstancedSubObjs, bool bAutoSerialize = true)
+			: TargetCDO(TargetObj->GetClass()->GetDefaultObject())
+			, Target(TargetObj)
+			, OldInstancedSubObjects(OldInstancedSubObjs)
+		{
+			ArIsObjectReferenceCollector = true;
+			ArIsModifyingWeakAndStrongReferences = true;
+
+			if (bAutoSerialize)
+			{
+				RunSerialization();
+			}
+		}
+
+		//----------------------------------------------------------------------
+		FArchive& operator<<(UObject*& Obj)
+		{
+			if (Obj == nullptr)
+			{
+				UProperty* SerializingProperty = GetSerializedProperty();
+				if (UObject* const* OldInstancedObjPtr = OldInstancedSubObjects.Find(SerializingProperty->GetFName()))
+				{
+					const UObject* OldInstancedObj = *OldInstancedObjPtr;
+					check(SerializingProperty->HasAnyPropertyFlags(CPF_PersistentInstance));
+
+					UClass* TargetClass = TargetCDO->GetClass();
+					// @TODO: Handle nested instances when we have more time to flush this all out  
+					if ( TargetClass->IsChildOf(SerializingProperty->GetOwnerClass()) )
+					{
+						UObjectPropertyBase* SerializingObjProperty = CastChecked<UObjectPropertyBase>(SerializingProperty);
+						// being extra careful, not to create our own instanced version when we expect one from the CDO
+						if (SerializingObjProperty->GetObjectPropertyValue_InContainer(TargetCDO) == nullptr)
+						{
+							// @TODO: What if the instanced object is of the same type 
+							//        that we're currently reinstancing
+							Obj = StaticDuplicateObject(OldInstancedObj, Target, nullptr);// NewObject<UObject>(Target, OldInstancedObj->GetClass()->GetAuthoritativeClass(), OldInstancedObj->GetFName());
+						}
+					}
+				}
+			}
+			return *this;
+		}
+
+		//----------------------------------------------------------------------
+		void RunSerialization()
+		{
+			if ((Target != nullptr) && (OldInstancedSubObjects.Num() != 0))
+			{
+				Target->Serialize(*this);
+			}
+		}
+
+	private:
+		UObject* TargetCDO;
+		UObject* Target;
+		const FInstancedPropertyMap& OldInstancedSubObjects;
+	};
+}
+
 void FBlueprintCompileReinstancer::ReplaceInstancesOfClass(UClass* OldClass, UClass* NewClass, UObject*	OriginalCDO, TSet<UObject*>* ObjectsThatShouldUseOldStuff, bool bClassObjectReplaced, bool bPreserveRootComponent)
 {
 	USelection* SelectedActors;
@@ -1093,7 +1233,10 @@ void FBlueprintCompileReinstancer::ReplaceInstancesOfClass(UClass* OldClass, UCl
 				auto FlagMask = RF_Public | RF_ArchetypeObject | RF_Transactional | RF_Transient | RF_TextExportTransient | RF_InheritableComponentTemplate; //TODO: what about RF_RootSet and RF_Standalone ?
 				NewUObject->SetFlags(OldFlags & FlagMask);
 
+				InstancedPropertyUtils::FInstancedPropertyMap InstancedPropertyMap;
+				InstancedPropertyUtils::FArchiveInstancedSubObjCollector  InstancedSubObjCollector(OldObject, InstancedPropertyMap);
 				UEditorEngine::CopyPropertiesForUnrelatedObjects(OldObject, NewUObject);
+				InstancedPropertyUtils::FArchiveInsertInstancedSubObjects InstancedSubObjSpawner(NewUObject, InstancedPropertyMap);
 
 				if (UAnimInstance* AnimTree = Cast<UAnimInstance>(NewUObject))
 				{
