Changes to remove the UCurve derived types from the Create Asset menu in the Content Browser, while still presenting it as an option in details panels.

diff --git a/Engine/Source/Editor/PropertyEditor/Private/PropertyCustomizationHelpers.cpp b/Engine/Source/Editor/PropertyEditor/Private/PropertyCustomizationHelpers.cpp
--- a/Engine/Source/Editor/PropertyEditor/Private/PropertyCustomizationHelpers.cpp
+++ b/Engine/Source/Editor/PropertyEditor/Private/PropertyCustomizationHelpers.cpp
@@ -298,7 +298,7 @@ namespace PropertyCustomizationHelpers
 			if (Class->IsChildOf(UFactory::StaticClass()) && !Class->HasAnyClassFlags(CLASS_Abstract))
 			{
 				UFactory* Factory = Class->GetDefaultObject<UFactory>();
-				if (Factory->ShouldShowInNewMenu() && ensure(!Factory->GetDisplayName().IsEmpty()))
+				if (Factory->CanCreateNew() && ensure(!Factory->GetDisplayName().IsEmpty()))
 				{
 					UClass* SupportedClass = Factory->GetSupportedClass();
 					if (SupportedClass != nullptr && Classes.ContainsByPredicate([=](const UClass* InClass) { return SupportedClass->IsChildOf(InClass); }))
diff --git a/Engine/Source/Editor/UnrealEd/Classes/Factories/CurveFactory.h b/Engine/Source/Editor/UnrealEd/Classes/Factories/CurveFactory.h
--- a/Engine/Source/Editor/UnrealEd/Classes/Factories/CurveFactory.h
+++ b/Engine/Source/Editor/UnrealEd/Classes/Factories/CurveFactory.h
@@ -34,5 +34,36 @@ class UCurveFloatFactory : public UCurveFactory
 
 	// UFactory interface
 	virtual bool ConfigureProperties() override;
+	virtual bool ShouldShowInNewMenu() const override { return false; }
+	// End of UFactory interface
+};
+
+/**
+ * Factory that creates linear color curve assets
+ */
+
+UCLASS()
+class UCurveLinearColorFactory : public UCurveFactory
+{
+	GENERATED_UCLASS_BODY()
+
+	// UFactory interface
+	virtual bool ConfigureProperties() override;
+	virtual bool ShouldShowInNewMenu() const override { return false; }
+	// End of UFactory interface
+};
+
+/**
+ * Factory that creates vector curve assets
+ */
+
+UCLASS()
+class UCurveVectorFactory : public UCurveFactory
+{
+	GENERATED_UCLASS_BODY()
+
+	// UFactory interface
+	virtual bool ConfigureProperties() override;
+	virtual bool ShouldShowInNewMenu() const override { return false; }
 	// End of UFactory interface
 };
diff --git a/Engine/Source/Editor/UnrealEd/Private/Factories/EditorFactories.cpp b/Engine/Source/Editor/UnrealEd/Private/Factories/EditorFactories.cpp
--- a/Engine/Source/Editor/UnrealEd/Private/Factories/EditorFactories.cpp
+++ b/Engine/Source/Editor/UnrealEd/Private/Factories/EditorFactories.cpp
@@ -6375,12 +6375,11 @@ FString UBlueprintInterfaceFactory::GetDefaultNewAssetName() const
 UCurveFactory::UCurveFactory(const FObjectInitializer& ObjectInitializer)
 	: Super(ObjectInitializer)
 {
-
 	bCreateNew = true;
 	bEditAfterNew = true;
 	SupportedClass = UCurveBase::StaticClass();
 
-	CurveClass = UCurveFloat::StaticClass();
+	CurveClass = nullptr;
 }
 
 class FCurveDataAssetParentFilter : public IClassViewerFilter
@@ -6453,8 +6452,6 @@ UObject* UCurveFactory::FactoryCreateNew(UClass* Class, UObject* InParent, FName
 UCurveFloatFactory::UCurveFloatFactory(const FObjectInitializer& ObjectInitializer)
 	: Super(ObjectInitializer)
 {
-	bCreateNew = true;
-	bEditAfterNew = true;
 	SupportedClass = UCurveFloat::StaticClass();
 	CurveClass = UCurveFloat::StaticClass();
 }
@@ -6464,6 +6461,38 @@ bool UCurveFloatFactory::ConfigureProperties()
 	return true;
 }
 
+/*------------------------------------------------------------------------------
+	UCurveLinearColorFactory implementation.
+------------------------------------------------------------------------------*/
+
+UCurveLinearColorFactory::UCurveLinearColorFactory(const FObjectInitializer& ObjectInitializer)
+	: Super(ObjectInitializer)
+{
+	SupportedClass = UCurveLinearColor::StaticClass();
+	CurveClass = UCurveLinearColor::StaticClass();
+}
+
+bool UCurveLinearColorFactory::ConfigureProperties()
+{
+	return true;
+}
+
+/*------------------------------------------------------------------------------
+	UCurveVectorFactory implementation.
+------------------------------------------------------------------------------*/
+
+UCurveVectorFactory::UCurveVectorFactory(const FObjectInitializer& ObjectInitializer)
+	: Super(ObjectInitializer)
+{
+	SupportedClass = UCurveVector::StaticClass();
+	CurveClass = UCurveVector::StaticClass();
+}
+
+bool UCurveVectorFactory::ConfigureProperties()
+{
+	return true;
+}
+
 /*------------------------------------------------------------------------------
 	UCurveImportFactory implementation.
 ------------------------------------------------------------------------------*/
