Fix map dependencies json generation during cooking  broken from a recent refactor.

diff --git a/Engine/Source/Editor/UnrealEd/Private/CookOnTheFlyServer.cpp b/Engine/Source/Editor/UnrealEd/Private/CookOnTheFlyServer.cpp
--- a/Engine/Source/Editor/UnrealEd/Private/CookOnTheFlyServer.cpp
+++ b/Engine/Source/Editor/UnrealEd/Private/CookOnTheFlyServer.cpp
@@ -3795,38 +3795,23 @@ void UCookOnTheFlyServer::BuildMapDependencyGraph(const FName& PlatformName)
 	FAssetRegistryModule& AssetRegistryModule = FModuleManager::LoadModuleChecked<FAssetRegistryModule>(AssetRegistryName);
 	IAssetRegistry& AssetRegistry = AssetRegistryModule.Get();
 
-
 	// assign chunks for all the map packages
 	for (const auto& CookedPackage : PlatformCookedPackages)
 	{
-		
 		TArray<FAssetData> PackageAssets;
-		ensure(AssetRegistry.GetAssetsByPackageName(CookedPackage, PackageAssets));
-
-		bool bIsMap = false;
-		for (const auto& Asset : PackageAssets)
-		{
-			if (Asset.GetClass()->IsChildOf(UWorld::StaticClass()) ||
-				Asset.GetClass()->IsChildOf(ULevel::StaticClass()))
-			{
-				bIsMap = true;
-			}
-		}
+		FName Name = FName(*FPackageName::FilenameToLongPackageName(CookedPackage.ToString()));
 
-
-		// if we are not a map we don't care about this pacakge
-		if (!bIsMap)
-		{
+		if (!ContainsMap(Name))
 			continue;
-		}
 
 		TSet<FName> DependentPackages;
-		for (const auto& Asset : PackageAssets)
-		{
-			DependentPackages.Add(Asset.PackageName);
-		}
+		TSet<FName> Roots; 
+
+		Roots.Add(Name);
+
+		GetDependentPackages(Roots, DependentPackages);
 
-		MapDependencyGraph.Add(CookedPackage, DependentPackages);
+		MapDependencyGraph.Add(Name, DependentPackages);
 	}
 }
 
