Changed hierarchy of some of the command lists to better reflect when a keypress has the potential to be handled by a parent if necessary.

diff --git a/Engine/Source/Editor/LevelEditor/Public/LevelEditorActions.h b/Engine/Source/Editor/LevelEditor/Public/LevelEditorActions.h
--- a/Engine/Source/Editor/LevelEditor/Public/LevelEditorActions.h
+++ b/Engine/Source/Editor/LevelEditor/Public/LevelEditorActions.h
@@ -19,7 +19,7 @@ public:
 	(
 		"LevelEditor", // Context name for fast lookup
 		NSLOCTEXT("Contexts", "LevelEditor", "Level Editor"), // Localized context name for displaying
-		"MainFrame", // Parent
+		"LevelViewport", // Parent
 		FEditorStyle::GetStyleSetName() // Icon Style Set
 	)
 	{
diff --git a/Engine/Source/Editor/UnrealEd/Private/Toolkits/AssetEditorCommonCommands.h b/Engine/Source/Editor/UnrealEd/Private/Toolkits/AssetEditorCommonCommands.h
--- a/Engine/Source/Editor/UnrealEd/Private/Toolkits/AssetEditorCommonCommands.h
+++ b/Engine/Source/Editor/UnrealEd/Private/Toolkits/AssetEditorCommonCommands.h
@@ -10,7 +10,7 @@ class FAssetEditorCommonCommands : public TCommands< FAssetEditorCommonCommands
 public:
 
 	FAssetEditorCommonCommands()
-		: TCommands< FAssetEditorCommonCommands >( TEXT("AssetEditor"), NSLOCTEXT("Contexts", "AssetEditor", "Asset Editor"), NAME_None, FEditorStyle::GetStyleSetName() )
+		: TCommands< FAssetEditorCommonCommands >( TEXT("AssetEditor"), NSLOCTEXT("Contexts", "AssetEditor", "Asset Editor"), TEXT("EditorViewport"), FEditorStyle::GetStyleSetName() )
 	{
 	}	
 
diff --git a/Engine/Source/Editor/UnrealEd/Public/EditorViewportCommands.h b/Engine/Source/Editor/UnrealEd/Public/EditorViewportCommands.h
--- a/Engine/Source/Editor/UnrealEd/Public/EditorViewportCommands.h
+++ b/Engine/Source/Editor/UnrealEd/Public/EditorViewportCommands.h
@@ -14,7 +14,7 @@ public:
 		(
 			TEXT("EditorViewport"), // Context name for fast lookup
 			NSLOCTEXT("Contexts", "EditorViewportCommands", "Common Viewport Commands"), // Localized context name for displaying
-			TEXT("LevelEditor"),
+			TEXT("MainFrame"),
 			FEditorStyle::GetStyleSetName() // Icon Style Set
 		)
 	{
