Fixed asset placement code to conform to recent SpawnActor changes.

diff --git a/Engine/Source/Editor/UnrealEd/Private/AssetSelection.cpp b/Engine/Source/Editor/UnrealEd/Private/AssetSelection.cpp
--- a/Engine/Source/Editor/UnrealEd/Private/AssetSelection.cpp
+++ b/Engine/Source/Editor/UnrealEd/Private/AssetSelection.cpp
@@ -435,7 +435,6 @@ static AActor* PrivateAddActor( UObject* Asset, UActorFactory* Factory, bool Sel
 
 	const FSnappedPositioningData PositioningData = FSnappedPositioningData(GCurrentLevelEditingViewportClient, GEditor->ClickLocation, GEditor->ClickPlane)
 		.UseFactory(Factory)
-		.UseStartTransform(NewActorTemplate->GetTransform())
 		.UsePlacementExtent(NewActorTemplate->GetPlacementExtent());
 
 	const FTransform ActorTransform = FActorPositioning::GetSnappedSurfaceAlignedTransform(PositioningData);
