Fixes to pivot offset feature.  Changed operation of "Center on Selection" to calculate the center of a Brush according to its vertices' positions.

diff --git a/Engine/Source/Editor/DetailCustomizations/Private/ComponentTransformDetails.cpp b/Engine/Source/Editor/DetailCustomizations/Private/ComponentTransformDetails.cpp
--- a/Engine/Source/Editor/DetailCustomizations/Private/ComponentTransformDetails.cpp
+++ b/Engine/Source/Editor/DetailCustomizations/Private/ComponentTransformDetails.cpp
@@ -1099,6 +1099,8 @@ void FComponentTransformDetails::OnSetLocation( float NewValue, ETextCommit::Typ
 
 	CacheTransform();
 
+	GUnrealEd->UpdatePivotLocationForSelection();
+	GUnrealEd->SetPivotMovedIndependently(false);
 	GUnrealEd->RedrawLevelEditingViewports();
 }
 
@@ -1230,6 +1232,8 @@ void FComponentTransformDetails::OnSetRotation( float NewValue, bool bCommitted,
 			GEditor->EndTransaction();
 		}
 
+		GUnrealEd->UpdatePivotLocationForSelection();
+		GUnrealEd->SetPivotMovedIndependently(false);
 		// Redraw
 		GUnrealEd->RedrawLevelEditingViewports();
 	}
diff --git a/Engine/Source/Editor/LevelEditor/Private/LevelEditorActions.cpp b/Engine/Source/Editor/LevelEditor/Private/LevelEditorActions.cpp
--- a/Engine/Source/Editor/LevelEditor/Private/LevelEditorActions.cpp
+++ b/Engine/Source/Editor/LevelEditor/Private/LevelEditorActions.cpp
@@ -2873,10 +2873,10 @@ void FLevelEditorCommands::RegisterCommands()
 
 	//@ todo Slate better tooltips for pivot options
 	UI_COMMAND( SavePivotToPrePivot, "Set as Pivot Offset", "Sets the current pivot location as the pivot offset for this actor", EUserInterfaceActionType::Button, FInputChord() );
-	UI_COMMAND( ResetPivot, "Reset Pivot Offset", "Resets the pivot offset for this actor", EUserInterfaceActionType::Button, FInputChord() );
-	UI_COMMAND( ResetPrePivot, "Reset Pre-Pivot", "Resets the pre-pivot", EUserInterfaceActionType::Button, FInputChord() );
-	UI_COMMAND( MovePivotHere, "Move Here", "Moves the pivot to the clicked location", EUserInterfaceActionType::Button, FInputChord() );
-	UI_COMMAND( MovePivotHereSnapped, "Move Here (Snapped)", "Moves the pivot to the nearest grid point to the clicked location", EUserInterfaceActionType::Button, FInputChord() );
+	UI_COMMAND( ResetPrePivot, "Reset Pivot Offset", "Resets the pivot offset for this actor", EUserInterfaceActionType::Button, FInputChord() );
+	UI_COMMAND( ResetPivot, "Reset Pivot", "Resets the pivot", EUserInterfaceActionType::Button, FInputChord() );
+	UI_COMMAND( MovePivotHere, "Set Pivot Offset Here", "Sets the pivot offset to the clicked location", EUserInterfaceActionType::Button, FInputChord() );
+	UI_COMMAND( MovePivotHereSnapped, "Set Pivot Offset Here (Snapped)", "Sets the pivot offset to the nearest grid point to the clicked location", EUserInterfaceActionType::Button, FInputChord() );
 	UI_COMMAND( MovePivotToCenter, "Center on Selection", "Centers the pivot to the middle of the selection", EUserInterfaceActionType::Button, FInputChord() );
 
 	UI_COMMAND( ConvertToAdditive, "Additive", "Converts the selected brushes to additive brushes", EUserInterfaceActionType::Button, FInputChord() );
diff --git a/Engine/Source/Editor/LevelEditor/Private/LevelEditorContextMenu.cpp b/Engine/Source/Editor/LevelEditor/Private/LevelEditorContextMenu.cpp
--- a/Engine/Source/Editor/LevelEditor/Private/LevelEditorContextMenu.cpp
+++ b/Engine/Source/Editor/LevelEditor/Private/LevelEditorContextMenu.cpp
@@ -916,14 +916,14 @@ void FLevelEditorContextMenuImpl::FillPivotMenu( FMenuBuilder& MenuBuilder )
 	MenuBuilder.BeginSection("SaveResetPivot");
 	{
 		MenuBuilder.AddMenuEntry( FLevelEditorCommands::Get().SavePivotToPrePivot );
-		MenuBuilder.AddMenuEntry( FLevelEditorCommands::Get().ResetPivot );
+		MenuBuilder.AddMenuEntry( FLevelEditorCommands::Get().ResetPrePivot );
+		MenuBuilder.AddMenuEntry(FLevelEditorCommands::Get().MovePivotHere);
+		MenuBuilder.AddMenuEntry(FLevelEditorCommands::Get().MovePivotHereSnapped);
 	}
 	MenuBuilder.EndSection();
 
 	MenuBuilder.BeginSection("MovePivot");
 	{
-		MenuBuilder.AddMenuEntry( FLevelEditorCommands::Get().MovePivotHere );
-		MenuBuilder.AddMenuEntry( FLevelEditorCommands::Get().MovePivotHereSnapped );
 		MenuBuilder.AddMenuEntry( FLevelEditorCommands::Get().MovePivotToCenter );
 	}
 	MenuBuilder.EndSection();
diff --git a/Engine/Source/Editor/UnrealEd/Classes/Editor/UnrealEdEngine.h b/Engine/Source/Editor/UnrealEd/Classes/Editor/UnrealEdEngine.h
--- a/Engine/Source/Editor/UnrealEd/Classes/Editor/UnrealEdEngine.h
+++ b/Engine/Source/Editor/UnrealEd/Classes/Editor/UnrealEdEngine.h
@@ -266,6 +266,16 @@ public:
 	 */
 	virtual void SetActorSelectionFlags (AActor* InActor);
 
+	/**
+	 * Set whether the pivot has been moved independently or not
+	 */
+	void SetPivotMovedIndependently( bool bMovedIndependently );
+
+	/**
+	 * Return whether the pivot has been moved independently or not
+	 */
+	bool IsPivotMovedIndependently() const;
+
 	/**
 	 * Called to reset the editor's pivot (widget) location using the currently selected objects.  Usually
 	 * called when the selection changes.
@@ -757,4 +767,7 @@ protected:
 
 	/** Handle to the registered OnMatineeEditorClosed delegate. */
 	FDelegateHandle OnMatineeEditorClosedDelegateHandle;
+
+	/** Whether the pivot has been moved independently */
+	bool bPivotMovedIndependently;
 };
diff --git a/Engine/Source/Editor/UnrealEd/Private/EditorSelectUtils.cpp b/Engine/Source/Editor/UnrealEd/Private/EditorSelectUtils.cpp
--- a/Engine/Source/Editor/UnrealEd/Private/EditorSelectUtils.cpp
+++ b/Engine/Source/Editor/UnrealEd/Private/EditorSelectUtils.cpp
@@ -156,10 +156,7 @@ void UUnrealEdEngine::SetPivot( FVector NewPivot, bool bSnapPivotToGrid, bool bI
 	}
 
 	// Set the pivot.
-	//EditorModeTools.CachedLocation	= NewPivot;	// Don't set the cached location, this is our pre-move point
-	EditorModeTools.PivotLocation		= NewPivot;
-	EditorModeTools.SnappedLocation		= NewPivot;
-	EditorModeTools.GridBase			= FVector::ZeroVector;
+	EditorModeTools.SetPivotLocation(NewPivot, false);
 
 	if( bSnapPivotToGrid )
 	{
@@ -274,6 +271,18 @@ void UUnrealEdEngine::SetActorSelectionFlags (AActor* InActor)
 }
 
 
+void UUnrealEdEngine::SetPivotMovedIndependently(bool bMovedIndependently)
+{
+	bPivotMovedIndependently = bMovedIndependently;
+}
+
+
+bool UUnrealEdEngine::IsPivotMovedIndependently() const
+{
+	return bPivotMovedIndependently;
+}
+
+
 void UUnrealEdEngine::UpdatePivotLocationForSelection( bool bOnChange )
 {
 	// Pick a new common pivot, or not.
@@ -355,6 +364,8 @@ void UUnrealEdEngine::UpdatePivotLocationForSelection( bool bOnChange )
 	{
 		ResetPivot();
 	}
+
+	SetPivotMovedIndependently(false);
 }
 
 
diff --git a/Engine/Source/Editor/UnrealEd/Private/LevelEditorViewport.cpp b/Engine/Source/Editor/UnrealEd/Private/LevelEditorViewport.cpp
--- a/Engine/Source/Editor/UnrealEd/Private/LevelEditorViewport.cpp
+++ b/Engine/Source/Editor/UnrealEd/Private/LevelEditorViewport.cpp
@@ -1184,12 +1184,12 @@ FDropQuery FLevelEditorViewportClient::CanDropObjectsAtCoordinates(int32 MouseX,
 		if ( AssetObj->IsA( AActor::StaticClass() ) || bHasActorFactory )
 		{
 			Result.bCanDrop = true;
-			bPivotMovedIndependently = false;
+			GUnrealEd->SetPivotMovedIndependently(false);
 		}
 		else if( AssetObj->IsA( UBrushBuilder::StaticClass()) )
 		{
 			Result.bCanDrop = true;
-			bPivotMovedIndependently = false;
+			GUnrealEd->SetPivotMovedIndependently(false);
 		}
 		else
 		{
@@ -1200,7 +1200,7 @@ FDropQuery FLevelEditorViewportClient::CanDropObjectsAtCoordinates(int32 MouseX,
 				{
 					// If our asset is a material and the target is a valid recipient
 					Result.bCanDrop = true;
-					bPivotMovedIndependently = false;
+					GUnrealEd->SetPivotMovedIndependently(false);
 
 					//if ( HitProxy->IsA(HActor::StaticGetType()) )
 					//{
@@ -1937,7 +1937,7 @@ void FLevelEditorViewportClient::ProcessClick(FSceneView& View, HHitProxy* HitPr
 			}
 
 			// We clicked an actor, allow the pivot to reposition itself.
-			bPivotMovedIndependently = false;
+//			GUnrealEd->SetPivotMovedIndependently(false);
 		}
 		else if (HitProxy->IsA(HInstancedStaticMeshInstance::StaticGetType()))
 		{
@@ -2025,9 +2025,9 @@ void FLevelEditorViewportClient::Tick(float DeltaTime)
 {
 	FEditorViewportClient::Tick(DeltaTime);
 
-	if( !bPivotMovedIndependently && GCurrentLevelEditingViewportClient == this &&
+	if (!GUnrealEd->IsPivotMovedIndependently() && GCurrentLevelEditingViewportClient == this &&
 		bIsRealtime &&
-		( Widget == NULL || !Widget->IsDragging() ) )
+		(Widget == NULL || !Widget->IsDragging()))
 	{
 		// @todo SIE: May be very expensive for lots of actors
 		GUnrealEd->UpdatePivotLocationForSelection();
@@ -2337,7 +2337,7 @@ bool FLevelEditorViewportClient::InputWidgetDelta(FViewport* Viewport, EAxisList
 				else
 				{
 					FSnappingUtils::SnapDragLocationToNearestVertex( ModeTools->PivotLocation, Drag, this );
-					bPivotMovedIndependently = true;
+					GUnrealEd->SetPivotMovedIndependently(true);
 					bOnlyMovedPivot = true;
 				}
 
@@ -2726,7 +2726,7 @@ void FLevelEditorViewportClient::TrackingStopped()
 			GEditor->BroadcastEndObjectMovement(*Actor);
 		}
 
-		if (!bPivotMovedIndependently)
+		if (!GUnrealEd->IsPivotMovedIndependently())
 		{
 			GUnrealEd->UpdatePivotLocationForSelection();
 		}
diff --git a/Engine/Source/Editor/UnrealEd/Private/UnrealEdEngine.cpp b/Engine/Source/Editor/UnrealEd/Private/UnrealEdEngine.cpp
--- a/Engine/Source/Editor/UnrealEd/Private/UnrealEdEngine.cpp
+++ b/Engine/Source/Editor/UnrealEd/Private/UnrealEdEngine.cpp
@@ -147,6 +147,8 @@ void UUnrealEdEngine::Init(IEngineLoop* InEngineLoop)
 			CookServer->Initialize(ECookMode::CookByTheBookFromTheEditor, BaseCookingFlags);
 		}
 	}
+
+	bPivotMovedIndependently = false;
 }
 
 bool CanCookForPlatformInThisProcess( const FString& PlatformName )
diff --git a/Engine/Source/Editor/UnrealEd/Private/UnrealEdSrv.cpp b/Engine/Source/Editor/UnrealEd/Private/UnrealEdSrv.cpp
--- a/Engine/Source/Editor/UnrealEd/Private/UnrealEdSrv.cpp
+++ b/Engine/Source/Editor/UnrealEd/Private/UnrealEdSrv.cpp
@@ -1693,6 +1693,7 @@ bool UUnrealEdEngine::Exec_Pivot( const TCHAR* Str, FOutputDevice& Ar )
 		NoteActorMovement();
 		SetPivot( ClickLocation, false, false );
 		FinishAllSnaps();
+		SetPivotMovedIndependently(true);
 		RedrawLevelEditingViewports();
 	}
 	else if( FParse::Command(&Str,TEXT("SNAPPED")) )
@@ -1700,6 +1701,7 @@ bool UUnrealEdEngine::Exec_Pivot( const TCHAR* Str, FOutputDevice& Ar )
 		NoteActorMovement();
 		SetPivot( ClickLocation, true, false );
 		FinishAllSnaps();
+		SetPivotMovedIndependently(true);
 		RedrawLevelEditingViewports();
 	}
 	else if( FParse::Command(&Str,TEXT("CENTERSELECTION")) )
@@ -1713,10 +1715,43 @@ bool UUnrealEdEngine::Exec_Pivot( const TCHAR* Str, FOutputDevice& Ar )
 
 		for ( FSelectionIterator It( GetSelectedActorIterator() ) ; It ; ++It )
 		{
-			AActor* Actor = static_cast<AActor*>( *It );
-			checkSlow( Actor->IsA(AActor::StaticClass()) );
+			AActor* Actor = CastChecked<AActor>(*It);
+
+			if (ABrush* Brush = Cast<ABrush>(Actor))
+			{
+				// Treat brushes as a special case; calculate an effective position from the center point of the vertices.
+				// This way, "Center on Selection" has a special meaning for brushes.
+				TSet<FVector> UniqueVertices;
+				FVector VertexCenter = FVector::ZeroVector;
+
+				if (Brush->Brush && Brush->Brush->Polys)
+				{
+					for (const auto& Element : Brush->Brush->Polys->Element)
+					{
+						for (const auto& Vertex : Element.Vertices)
+						{
+							UniqueVertices.Add(Vertex);
+						}
+					}
+
+					for (const auto& Vertex : UniqueVertices)
+					{
+						VertexCenter += Vertex;
+					}
+
+					if (UniqueVertices.Num() > 0)
+					{
+						VertexCenter /= UniqueVertices.Num();
+					}
+				}
+
+				Center += Brush->GetTransform().TransformPosition(VertexCenter);
+			}
+			else
+			{
+				Center += Actor->GetActorLocation();
+			}
 
-			Center += Actor->GetActorLocation();
 			Count++;
 		}
 
@@ -1729,6 +1764,7 @@ bool UUnrealEdEngine::Exec_Pivot( const TCHAR* Str, FOutputDevice& Ar )
 
 			SetPivot( ClickLocation, false, false );
 			FinishAllSnaps();
+			SetPivotMovedIndependently(true);
 		}
 
 		RedrawLevelEditingViewports();
@@ -2603,7 +2639,8 @@ bool UUnrealEdEngine::Exec_Actor( UWorld* InWorld, const TCHAR* Str, FOutputDevi
 			FVector Delta( EditorModeTools.PivotLocation - Actor->GetActorLocation() );
 
 			Actor->Modify();
-			Actor->SetPivotOffset(Actor->GetPivotOffset() + Actor->GetTransform().InverseTransformVector(Delta));
+			Actor->SetPivotOffset(Actor->GetTransform().InverseTransformVector(Delta));
+			SetPivotMovedIndependently(false);
 			Actor->PostEditMove(true);
 		}
 
@@ -2625,6 +2662,7 @@ bool UUnrealEdEngine::Exec_Actor( UWorld* InWorld, const TCHAR* Str, FOutputDevi
 
 			Actor->Modify();
 			Actor->SetPivotOffset(FVector::ZeroVector);
+			SetPivotMovedIndependently(false);
 			Actor->PostEditMove(true);
 		}
 
diff --git a/Engine/Source/Editor/UnrealEd/Public/LevelEditorViewport.h b/Engine/Source/Editor/UnrealEd/Public/LevelEditorViewport.h
--- a/Engine/Source/Editor/UnrealEd/Public/LevelEditorViewport.h
+++ b/Engine/Source/Editor/UnrealEd/Public/LevelEditorViewport.h
@@ -644,9 +644,6 @@ public:
 
 	bool					bEnableColorScaling;
 
-	/** If true, the pivot has been moved independently of the actor and position updates should not occur when the actor is moved. */
-	bool					bPivotMovedIndependently;
-
 	/** If true, we switched between two different cameras. Set by matinee, used by the motion blur to invalidate this frames motion vectors */
 	bool					bEditorCameraCut;
 
