Fixed another codepath in asset placement code to conform to recent SpawnActor changes.

diff --git a/Engine/Source/Editor/UnrealEd/Private/Editor/ActorPositioning.cpp b/Engine/Source/Editor/UnrealEd/Private/Editor/ActorPositioning.cpp
--- a/Engine/Source/Editor/UnrealEd/Private/Editor/ActorPositioning.cpp
+++ b/Engine/Source/Editor/UnrealEd/Private/Editor/ActorPositioning.cpp
@@ -150,7 +150,6 @@ FTransform FActorPositioning::GetCurrentViewportPlacementTransform(const AActor&
 		const FSnappedPositioningData PositioningData = FSnappedPositioningData(GCurrentLevelEditingViewportClient, GEditor->ClickLocation, GEditor->ClickPlane)
 			.DrawSnapHelpers(true)
 			.UseFactory(Factory)
-			.UseStartTransform(Actor.GetTransform())
 			.UsePlacementExtent(Actor.GetPlacementExtent());
 
 		ActorTransform = bSnap ? GetSnappedSurfaceAlignedTransform(PositioningData) : GetSurfaceAlignedTransform(PositioningData);
