BP Reinstancer calls "BP->BroadcastCompiled" once all BP are compiled.

diff --git a/Engine/Source/Editor/UnrealEd/Private/Kismet2/Kismet2.cpp b/Engine/Source/Editor/UnrealEd/Private/Kismet2/Kismet2.cpp
--- a/Engine/Source/Editor/UnrealEd/Private/Kismet2/Kismet2.cpp
+++ b/Engine/Source/Editor/UnrealEd/Private/Kismet2/Kismet2.cpp
@@ -675,7 +675,7 @@ bool FKismetEditorUtilities::IsReferencedByUndoBuffer(UBlueprint* Blueprint)
 	return (TotalReferenceCount > NonUndoReferenceCount);
 }
 
-void FKismetEditorUtilities::CompileBlueprint(UBlueprint* BlueprintObj, bool bIsRegeneratingOnLoad, bool bSkipGarbageCollection, bool bSaveIntermediateProducts, FCompilerResultsLog* pResults, bool bSkeletonUpToDate )
+void FKismetEditorUtilities::CompileBlueprint(UBlueprint* BlueprintObj, bool bIsRegeneratingOnLoad, bool bSkipGarbageCollection, bool bSaveIntermediateProducts, FCompilerResultsLog* pResults, bool bSkeletonUpToDate, bool bBatchCompile)
 {
 	FSecondsCounterScope Timer(BlueprintCompileAndLoadTimerData); 
 	BP_SCOPED_COMPILER_EVENT_STAT(EKismetCompilerStats_CompileBlueprint);
@@ -792,7 +792,8 @@ void FKismetEditorUtilities::CompileBlueprint(UBlueprint* BlueprintObj, bool bIs
 		// ReinstanceHelper.VerifyReplacement();
 	}
 
-	{ 
+	if (!bBatchCompile)
+	{
 		BP_SCOPED_COMPILER_EVENT_STAT(EKismetCompilerStats_NotifyBlueprintChanged);
 
 		BlueprintObj->BroadcastCompiled();
@@ -900,7 +901,7 @@ bool FKismetEditorUtilities::GenerateBlueprintSkeleton(UBlueprint* BlueprintObj,
 }
 
 /** Recompiles the bytecode of a blueprint only.  Should only be run for recompiling dependencies during compile on load */
-void FKismetEditorUtilities::RecompileBlueprintBytecode(UBlueprint* BlueprintObj, TArray<UObject*>* ObjLoaded)
+void FKismetEditorUtilities::RecompileBlueprintBytecode(UBlueprint* BlueprintObj, TArray<UObject*>* ObjLoaded, bool bBatchCompile)
 {
 	FSecondsCounterScope Timer(BlueprintCompileAndLoadTimerData); 
 	
@@ -931,7 +932,7 @@ void FKismetEditorUtilities::RecompileBlueprintBytecode(UBlueprint* BlueprintObj
 		BlueprintPackage->SetDirtyFlag(bStartedWithUnsavedChanges);
 	}
 
-	if (!BlueprintObj->bIsRegeneratingOnLoad)
+	if (!BlueprintObj->bIsRegeneratingOnLoad && !bBatchCompile)
 	{
 		BP_SCOPED_COMPILER_EVENT_STAT(EKismetCompilerStats_NotifyBlueprintChanged);
 
diff --git a/Engine/Source/Editor/UnrealEd/Private/Kismet2/KismetReinstanceUtilities.cpp b/Engine/Source/Editor/UnrealEd/Private/Kismet2/KismetReinstanceUtilities.cpp
--- a/Engine/Source/Editor/UnrealEd/Private/Kismet2/KismetReinstanceUtilities.cpp
+++ b/Engine/Source/Editor/UnrealEd/Private/Kismet2/KismetReinstanceUtilities.cpp
@@ -534,6 +534,8 @@ void FBlueprintCompileReinstancer::ReinstanceObjects(bool bForceAlwaysReinstance
 
 		if (QueueToReinstance.Num() && (QueueToReinstance[0] == SharedThis))
 		{
+			TSet<TWeakObjectPtr<UBlueprint>> CompiledBlueprints;
+
 			while (DependentBlueprintsToRecompile.Num())
 			{
 				auto Iter = DependentBlueprintsToRecompile.CreateIterator();
@@ -541,7 +543,8 @@ void FBlueprintCompileReinstancer::ReinstanceObjects(bool bForceAlwaysReinstance
 				Iter.RemoveCurrent();
 				if (auto BP = BPPtr.Get())
 				{
-					FKismetEditorUtilities::CompileBlueprint(BP, false, bSkipGarbageCollection);
+					FKismetEditorUtilities::CompileBlueprint(BP, false, bSkipGarbageCollection, false, nullptr, false, true);
+					CompiledBlueprints.Add(BP);
 				}
 			}
 
@@ -552,7 +555,8 @@ void FBlueprintCompileReinstancer::ReinstanceObjects(bool bForceAlwaysReinstance
 				Iter.RemoveCurrent();
 				if (auto BP = BPPtr.Get())
 				{
-					FKismetEditorUtilities::RecompileBlueprintBytecode(BP);
+					FKismetEditorUtilities::RecompileBlueprintBytecode(BP, nullptr, true);
+					CompiledBlueprints.Add(BP);
 				}
 			}
 
@@ -582,6 +586,11 @@ void FBlueprintCompileReinstancer::ReinstanceObjects(bool bForceAlwaysReinstance
 				}
 			}
 
+			for (auto CompiledBP : CompiledBlueprints)
+			{
+				CompiledBP->BroadcastCompiled();
+			}
+
 			{
 				BP_SCOPED_COMPILER_EVENT_STAT(EKismetCompilerStats_RefreshDependentBlueprintsInReinstancer);
 				for (auto BPPtr : DependentBlueprintsToRefresh)
@@ -593,6 +602,11 @@ void FBlueprintCompileReinstancer::ReinstanceObjects(bool bForceAlwaysReinstance
 				}
 				DependentBlueprintsToRefresh.Empty();
 			}
+
+			if (GEditor)
+			{
+				GEditor->BroadcastBlueprintCompiled();
+			}
 		}
 	}
 }
diff --git a/Engine/Source/Editor/UnrealEd/Public/Kismet2/KismetEditorUtilities.h b/Engine/Source/Editor/UnrealEd/Public/Kismet2/KismetEditorUtilities.h
--- a/Engine/Source/Editor/UnrealEd/Public/Kismet2/KismetEditorUtilities.h
+++ b/Engine/Source/Editor/UnrealEd/Public/Kismet2/KismetEditorUtilities.h
@@ -97,13 +97,13 @@ public:
 	static void CreateDefaultEventGraphs(UBlueprint* Blueprint);
 
 	/** Tries to compile a blueprint, updating any actors in the editor who are using the old class, etc... */
-	static void CompileBlueprint(UBlueprint* BlueprintObj, bool bIsRegeneratingOnLoad=false, bool bSkipGarbageCollection=false, bool bSaveIntermediateProducts = false, class FCompilerResultsLog* pResults = NULL, bool bSkeletonUpToDate = false);
+	static void CompileBlueprint(UBlueprint* BlueprintObj, bool bIsRegeneratingOnLoad = false, bool bSkipGarbageCollection = false, bool bSaveIntermediateProducts = false, class FCompilerResultsLog* pResults = nullptr, bool bSkeletonUpToDate = false, bool bBatchCompile = false);
 
 	/** Generates a blueprint skeleton only.  Minimal compile, no notifications will be sent, no GC, etc.  Only successful if there isn't already a skeleton generated */
 	static bool GenerateBlueprintSkeleton(UBlueprint* BlueprintObj, bool bForceRegeneration = false);
 
 	/** Recompiles the bytecode of a blueprint only.  Should only be run for recompiling dependencies during compile on load */
-	static void RecompileBlueprintBytecode(UBlueprint* BlueprintObj, TArray<UObject*>* ObjLoaded = NULL);
+	static void RecompileBlueprintBytecode(UBlueprint* BlueprintObj, TArray<UObject*>* ObjLoaded = nullptr, bool bBatchCompile = false);
 
 	static void GenerateCppCode(UObject* Obj, TSharedPtr<FString> OutHeaderSource, TSharedPtr<FString> OutCppSource);
 
