Added missing pointer validity check in internationalization export logic.

diff --git a/Engine/Source/Editor/UnrealEd/Private/Commandlets/InternationalizationExportCommandlet.cpp b/Engine/Source/Editor/UnrealEd/Private/Commandlets/InternationalizationExportCommandlet.cpp
--- a/Engine/Source/Editor/UnrealEd/Private/Commandlets/InternationalizationExportCommandlet.cpp
+++ b/Engine/Source/Editor/UnrealEd/Private/Commandlets/InternationalizationExportCommandlet.cpp
@@ -126,8 +126,8 @@ FString ConvertSrcLocationToPORef(const FString& InSrcLocation)
 }
 
 /**
- *	UInternationalizationExportCommandlet
- */
+*	UInternationalizationExportCommandlet
+*/
 UInternationalizationExportCommandlet::UInternationalizationExportCommandlet(const FObjectInitializer& ObjectInitializer)
 	: Super(ObjectInitializer)
 {
@@ -248,33 +248,37 @@ bool UInternationalizationExportCommandlet::DoExport( const FString& SourcePath,
 						const FString& Namespace = ManifestEntry->Namespace;
 						const FLocItem& Source = ManifestEntry->Source;
 
+						// For each context, we may need to create a different or even multiple PO entries.
 						for( auto ContextIter = ManifestEntry->Contexts.CreateConstIterator(); ContextIter; ++ContextIter )
 						{
+							// Create the typical PO entry from the archive entry which matches the exact same namespace, source, and key metadata, if it exists.
 							{
 								const TSharedPtr<FArchiveEntry> ArchiveEntry = InternationalizationArchive->FindEntryBySource( Namespace, Source, ContextIter->KeyMetadataObj );
-							if( ArchiveEntry.IsValid() )
-							{
-								const FString ConditionedArchiveSource = ConditionArchiveStrForPo(ArchiveEntry->Source.Text);
-								const FString ConditionedArchiveTranslation = ConditionArchiveStrForPo(ArchiveEntry->Translation.Text);
-
-								TSharedRef<FPortableObjectEntry> PoEntry = MakeShareable( new FPortableObjectEntry );
-								//@TODO: We support additional metadata entries that can be translated.  How do those fit in the PO file format?  Ex: isMature
-								PoEntry->MsgId = ConditionedArchiveSource;
-								//@TODO: Take into account optional entries and entries that differ by keymetadata.  Ex. Each optional entry needs a unique msgCtxt
-								PoEntry->MsgCtxt = Namespace;
-								PoEntry->MsgStr.Add( ConditionedArchiveTranslation );
-
-								FString PORefString = ConvertSrcLocationToPORef( ContextIter->SourceLocation );
-								PoEntry->AddReference( PORefString ); // Source location.
-								PoEntry->AddExtractedComment( ContextIter->Key ); // "Notes from Programmer" in the form of the Key.
-								PoEntry->AddExtractedComment( PORefString ); // "Notes from Programmer" in the form of the Source Location, since this comes in handy too and OneSky doesn't properly show references, only comments.
-								PortableObj.AddEntry( PoEntry );
+								if( ArchiveEntry.IsValid() )
+								{
+									const FString ConditionedArchiveSource = ConditionArchiveStrForPo(ArchiveEntry->Source.Text);
+									const FString ConditionedArchiveTranslation = ConditionArchiveStrForPo(ArchiveEntry->Translation.Text);
+
+									TSharedRef<FPortableObjectEntry> PoEntry = MakeShareable( new FPortableObjectEntry );
+									//@TODO: We support additional metadata entries that can be translated.  How do those fit in the PO file format?  Ex: isMature
+									PoEntry->MsgId = ConditionedArchiveSource;
+									//@TODO: Take into account optional entries and entries that differ by keymetadata.  Ex. Each optional entry needs a unique msgCtxt
+									PoEntry->MsgCtxt = Namespace;
+									PoEntry->MsgStr.Add( ConditionedArchiveTranslation );
+
+									FString PORefString = ConvertSrcLocationToPORef( ContextIter->SourceLocation );
+									PoEntry->AddReference( PORefString ); // Source location.
+									PoEntry->AddExtractedComment( ContextIter->Key ); // "Notes from Programmer" in the form of the Key.
+									PoEntry->AddExtractedComment( PORefString ); // "Notes from Programmer" in the form of the Source Location, since this comes in handy too and OneSky doesn't properly show references, only comments.
+									PortableObj.AddEntry( PoEntry );
+								}
 							}
-						}
 
+							// If we're exporting for something other than the native culture, we'll need to create PO entries for archive entries based on the native archive's translation.
 							if (CultureName != NativeCultureName)
 							{
 								TSharedPtr<FArchiveEntry> NativeArchiveEntry;
+								// Find the native archive entry which matches the exact same namespace, source, and key metadata, if it exists.
 								for (const auto& NativeArchive : NativeArchives)
 								{
 									const TSharedPtr<FArchiveEntry> PotentialNativeArchiveEntry = NativeArchive->FindEntryBySource( Namespace, Source, ContextIter->KeyMetadataObj );
@@ -287,29 +291,32 @@ bool UInternationalizationExportCommandlet::DoExport( const FString& SourcePath,
 
 								if (NativeArchiveEntry.IsValid())
 								{
+									// Only need to create this PO entry if the native archive entry's translation differs from its source, in which case we need to find the our translation of the native translation.
 									if (!NativeArchiveEntry->Source.IsExactMatch(NativeArchiveEntry->Translation))
 									{
 										const TSharedPtr<FArchiveEntry> ArchiveEntry = InternationalizationArchive->FindEntryBySource( Namespace, NativeArchiveEntry->Translation, NativeArchiveEntry->KeyMetadataObj );
-
-										const FString ConditionedArchiveSource = ConditionArchiveStrForPo(ArchiveEntry->Source.Text);
-										const FString ConditionedArchiveTranslation = ConditionArchiveStrForPo(ArchiveEntry->Translation.Text);
-
-										TSharedRef<FPortableObjectEntry> PoEntry = MakeShareable( new FPortableObjectEntry );
-										//@TODO: We support additional metadata entries that can be translated.  How do those fit in the PO file format?  Ex: isMature
-										PoEntry->MsgId = ConditionedArchiveSource;
-										//@TODO: Take into account optional entries and entries that differ by keymetadata.  Ex. Each optional entry needs a unique msgCtxt
-										PoEntry->MsgCtxt = Namespace;
-										PoEntry->MsgStr.Add( ConditionedArchiveTranslation );
-
-										FString PORefString = ConvertSrcLocationToPORef( ContextIter->SourceLocation );
-										PoEntry->AddReference( PORefString ); // Source location.
-										PoEntry->AddExtractedComment( ContextIter->Key ); // "Notes from Programmer" in the form of the Key.
-										PoEntry->AddExtractedComment( PORefString ); // "Notes from Programmer" in the form of the Source Location, since this comes in handy too and OneSky doesn't properly show references, only comments.
-										PortableObj.AddEntry( PoEntry );
+										if (ArchiveEntry.IsValid())
+										{
+											const FString ConditionedArchiveSource = ConditionArchiveStrForPo(ArchiveEntry->Source.Text);
+											const FString ConditionedArchiveTranslation = ConditionArchiveStrForPo(ArchiveEntry->Translation.Text);
+
+											TSharedRef<FPortableObjectEntry> PoEntry = MakeShareable( new FPortableObjectEntry );
+											//@TODO: We support additional metadata entries that can be translated.  How do those fit in the PO file format?  Ex: isMature
+											PoEntry->MsgId = ConditionedArchiveSource;
+											//@TODO: Take into account optional entries and entries that differ by keymetadata.  Ex. Each optional entry needs a unique msgCtxt
+											PoEntry->MsgCtxt = Namespace;
+											PoEntry->MsgStr.Add( ConditionedArchiveTranslation );
+
+											FString PORefString = ConvertSrcLocationToPORef( ContextIter->SourceLocation );
+											PoEntry->AddReference( PORefString ); // Source location.
+											PoEntry->AddExtractedComment( ContextIter->Key ); // "Notes from Programmer" in the form of the Key.
+											PoEntry->AddExtractedComment( PORefString ); // "Notes from Programmer" in the form of the Source Location, since this comes in handy too and OneSky doesn't properly show references, only comments.
+											PortableObj.AddEntry( PoEntry );
+										}
 									}
 								}
-					}
-				}
+							}
+						}
 					}
 				}
 
@@ -418,7 +425,7 @@ bool UInternationalizationExportCommandlet::DoImport(const FString& SourcePath,
 
 		const FString DestinationCulturePath = DestinationPath / CultureName;
 		FString ArchiveFileName = DestinationCulturePath / ArchiveName;
-		
+
 		if( !FPaths::FileExists(ArchiveFileName) )
 		{
 			UE_LOG( LogInternationalizationExportCommandlet, Error, TEXT("Failed to find destination archive %s."), *ArchiveFileName);
@@ -448,7 +455,7 @@ bool UInternationalizationExportCommandlet::DoImport(const FString& SourcePath,
 				{
 					UE_LOG( LogInternationalizationExportCommandlet, Error, TEXT("Portable Object entry has plural form we did not process.  File: %s  MsgCtxt: %s  MsgId: %s"), *POFilePath, *POEntry->MsgCtxt, *POEntry->MsgId );
 				}
-				
+
 				const FString& Namespace = POEntry->MsgCtxt;
 				const FString& SourceText = ConditionPoStringForArchive(POEntry->MsgId);
 				const FString& Translation = ConditionPoStringForArchive(POEntry->MsgStr[0]);
@@ -460,7 +467,7 @@ bool UInternationalizationExportCommandlet::DoImport(const FString& SourcePath,
 					UE_LOG(LogInternationalizationExportCommandlet, Warning, TEXT("Could not find corresponding archive entry for PO entry.  File: %s  MsgCtxt: %s  MsgId: %s"), *POFilePath, *POEntry->MsgCtxt, *POEntry->MsgId );
 					continue;
 				}
-				
+
 				if( FoundEntry->Translation != Translation )
 				{
 					FoundEntry->Translation = Translation;
@@ -493,7 +500,7 @@ int32 UInternationalizationExportCommandlet::Main( const FString& Params )
 
 
 	UCommandlet::ParseCommandLine(*Params, Tokens, Switches, ParamVals);
-	
+
 	const FString* ParamVal = ParamVals.Find(FString(TEXT("Config")));
 
 	if ( ParamVal )
@@ -508,7 +515,7 @@ int32 UInternationalizationExportCommandlet::Main( const FString& Params )
 
 	//Set config section
 	ParamVal = ParamVals.Find(FString(TEXT("Section")));
-	
+
 
 	if ( ParamVal )
 	{
@@ -555,7 +562,7 @@ int32 UInternationalizationExportCommandlet::Main( const FString& Params )
 
 	GetBoolFromConfig( *SectionName, TEXT("bImportLoc"), bDoImport, ConfigPath );
 	GetBoolFromConfig( *SectionName, TEXT("bExportLoc"), bDoExport, ConfigPath );
-	
+
 	if( !bDoImport && !bDoExport )
 	{
 		UE_LOG(LogInternationalizationExportCommandlet, Error, TEXT("Import/Export operation not detected.  Use bExportLoc or bImportLoc in config section."));
