Asset and Asset class are properly saved as UDStruct default data.

diff --git a/Engine/Source/Editor/UnrealEd/Private/Kismet2/StructureEditorUtils.cpp b/Engine/Source/Editor/UnrealEd/Private/Kismet2/StructureEditorUtils.cpp
--- a/Engine/Source/Editor/UnrealEd/Private/Kismet2/StructureEditorUtils.cpp
+++ b/Engine/Source/Editor/UnrealEd/Private/Kismet2/StructureEditorUtils.cpp
@@ -377,8 +377,13 @@ bool FStructureEditorUtils::ChangeVariableDefaultValue(UUserDefinedStruct* Struc
 		{
 			bResult = true;
 		}
-		else if ((PinType.PinCategory == K2Schema->PC_Object) || (PinType.PinCategory == K2Schema->PC_Interface) || (PinType.PinCategory == K2Schema->PC_Class))
+		else if ((PinType.PinCategory == K2Schema->PC_Object) 
+			|| (PinType.PinCategory == K2Schema->PC_Interface) 
+			|| (PinType.PinCategory == K2Schema->PC_Class)
+			|| (PinType.PinCategory == K2Schema->PC_AssetClass)
+			|| (PinType.PinCategory == K2Schema->PC_Asset))
 		{
+			// K2Schema->DefaultValueSimpleValidation finds an object, passed by path, invalid
 			bResult = true;
 		}
 		else
