Dispose of GPU memory after temp holding texture is finished with.

diff --git a/Engine/Source/Editor/UnrealEd/Private/CubemapUnwrapUtils.cpp b/Engine/Source/Editor/UnrealEd/Private/CubemapUnwrapUtils.cpp
--- a/Engine/Source/Editor/UnrealEd/Private/CubemapUnwrapUtils.cpp
+++ b/Engine/Source/Editor/UnrealEd/Private/CubemapUnwrapUtils.cpp
@@ -84,6 +84,7 @@ namespace CubemapHelpers
 			break;
 		}
 		// Clean up.
+		RenderTargetLongLat->ReleaseResource();
 		RenderTargetLongLat->RemoveFromRoot();
 		RenderTargetLongLat = NULL;
 		delete Canvas;
