Changed FUnrealEdMisc::TickAssetAnalytics to use more appropriate containers

It was using two TArrays with AddUnique, which was *very* slow for large projects, and was causing a hitch after the asset registry had finished populating.

These have been changed to be a TSet and TMap, which can handle unique entries for large data sets (1,000,000 in this particular case) much more efficiently.

diff --git a/Engine/Source/Editor/UnrealEd/Private/UnrealEdMisc.cpp b/Engine/Source/Editor/UnrealEd/Private/UnrealEdMisc.cpp
--- a/Engine/Source/Editor/UnrealEd/Private/UnrealEdMisc.cpp
+++ b/Engine/Source/Editor/UnrealEd/Private/UnrealEdMisc.cpp
@@ -583,13 +583,12 @@ void FUnrealEdMisc::TickAssetAnalytics()
 
 			TArray< FAnalyticsEventAttribute > AssetAttributes;
 			int32 NumMapFiles = 0;
-			TArray< FName > PackageNames;
-			TArray< FName > ClassInstances;
-			TArray< int32 > ClassInstanceCount;
+			TSet< FName > PackageNames;
+			TMap< FName, int32 > ClassInstanceCounts;
 
 			for( auto AssetIter = AssetData.CreateConstIterator(); AssetIter; ++AssetIter )
 			{
-				PackageNames.AddUnique( AssetIter->PackageName );
+				PackageNames.Add( AssetIter->PackageName );
 				if( AssetIter->AssetClass == UWorld::StaticClass()->GetFName()  )
 				{
 					NumMapFiles++;
@@ -597,14 +596,14 @@ void FUnrealEdMisc::TickAssetAnalytics()
 
 				if( AssetIter->AssetClass != NAME_None )
 				{
-					int32 ClassIndex = ClassInstances.AddUnique( AssetIter->AssetClass );
-					if( ClassInstanceCount.Num() < ClassInstances.Num() )
+					int32* ExistingClassCount = ClassInstanceCounts.Find( AssetIter->AssetClass );
+					if( ExistingClassCount )
 					{
-						ClassInstanceCount.Add( 1 );
+						++(*ExistingClassCount);
 					}
 					else
 					{
-						ClassInstanceCount[ ClassIndex ] += 1;
+						ClassInstanceCounts.Add( AssetIter->AssetClass, 1 );
 					}
 				}
 			}
@@ -617,12 +616,9 @@ void FUnrealEdMisc::TickAssetAnalytics()
 
 			TArray< FAnalyticsEventAttribute > AssetInstances;
 			AssetInstances.Add( FAnalyticsEventAttribute( FString( "ProjectId" ), *ProjectSettings.ProjectID.ToString() ));
-			for( auto ClassIter = ClassInstances.CreateIterator(); ClassIter; ++ClassIter )
+			for( auto ClassIter = ClassInstanceCounts.CreateIterator(); ClassIter; ++ClassIter )
 			{
-				if( ClassInstanceCount[ ClassIter.GetIndex() ] > 0 )
-				{
-					AssetInstances.Add( FAnalyticsEventAttribute( ClassIter->ToString(), ClassInstanceCount[ ClassIter.GetIndex() ] ));
-				}
+				AssetInstances.Add( FAnalyticsEventAttribute( ClassIter.Key().ToString(), ClassIter.Value() ) );
 			}
 			// Send class instance analytics
 			FEngineAnalytics::GetProvider().RecordEvent( FString( "Editor.Usage.AssetClasses" ), AssetInstances );
