Clarified the use of Server Game Options and Command Line Options in the tooltips.

diff --git a/Engine/Source/Editor/UnrealEd/Classes/Settings/LevelEditorPlaySettings.h b/Engine/Source/Editor/UnrealEd/Classes/Settings/LevelEditorPlaySettings.h
--- a/Engine/Source/Editor/UnrealEd/Classes/Settings/LevelEditorPlaySettings.h
+++ b/Engine/Source/Editor/UnrealEd/Classes/Settings/LevelEditorPlaySettings.h
@@ -298,11 +298,11 @@ private:
 	UPROPERTY(config, EditAnywhere, Category=MultiplayerOptions)
 	int32 ClientWindowHeight;
 
-	/** Additional options that will be passed to the server as URL parameters. */
+	/** Additional options that will be passed to the server as URL parameters, in the format ?bIsLanMatch=1?listen - any additional command line switches should be passed in the Command Line Arguments field below. */
 	UPROPERTY(config, EditAnywhere, Category=MultiplayerOptions)
 	FString AdditionalServerGameOptions;
 
-	/** Additional command line options that will be passed to standalone game instances. */
+	/** Additional command line options that will be passed to standalone game instances, for example -debug */
 	UPROPERTY(config, EditAnywhere, Category=MultiplayerOptions)
 	FString AdditionalLaunchOptions;
 
