Fixed actor placement in the editor where surface snapping is enabled

The issue here is that SpawnActor has fundamentally changed such that the user-supplied transform is no longer an absolute transform, but multiplied with the class default's current transform. This affected the actor positioning code by further manipulating the desired placement transform with the default translation/rotation of the component type.

This addresses UE-19433.

diff --git a/Engine/Source/Editor/UnrealEd/Private/AssetSelection.cpp b/Engine/Source/Editor/UnrealEd/Private/AssetSelection.cpp
--- a/Engine/Source/Editor/UnrealEd/Private/AssetSelection.cpp
+++ b/Engine/Source/Editor/UnrealEd/Private/AssetSelection.cpp
@@ -437,7 +437,18 @@ static AActor* PrivateAddActor( UObject* Asset, UActorFactory* Factory, bool Sel
 		.UseFactory(Factory)
 		.UsePlacementExtent(NewActorTemplate->GetPlacementExtent());
 
-	const FTransform ActorTransform = FActorPositioning::GetSnappedSurfaceAlignedTransform(PositioningData);
+	FTransform ActorTransform = FActorPositioning::GetSnappedSurfaceAlignedTransform(PositioningData);
+
+	if (GetDefault<ULevelEditorViewportSettings>()->SnapToSurface.bEnabled)
+	{
+		// HACK: If we are aligning rotation to surfaces, we have to factor in the inverse of the actor transform so that the resulting transform after SpawnActor is correct.
+
+		if (auto* RootComponent = NewActorTemplate->GetRootComponent())
+		{
+			RootComponent->UpdateComponentToWorld();
+		}
+		ActorTransform = NewActorTemplate->GetTransform().Inverse() * ActorTransform;
+	}
 
 	// Do not fade snapping indicators over time if the viewport is not realtime
 	bool bClearImmediately = !GCurrentLevelEditingViewportClient || !GCurrentLevelEditingViewportClient->IsRealtime();
diff --git a/Engine/Source/Editor/UnrealEd/Private/Editor/ActorPositioning.cpp b/Engine/Source/Editor/UnrealEd/Private/Editor/ActorPositioning.cpp
--- a/Engine/Source/Editor/UnrealEd/Private/Editor/ActorPositioning.cpp
+++ b/Engine/Source/Editor/UnrealEd/Private/Editor/ActorPositioning.cpp
@@ -153,6 +153,17 @@ FTransform FActorPositioning::GetCurrentViewportPlacementTransform(const AActor&
 			.UsePlacementExtent(Actor.GetPlacementExtent());
 
 		ActorTransform = bSnap ? GetSnappedSurfaceAlignedTransform(PositioningData) : GetSurfaceAlignedTransform(PositioningData);
+
+		if (GetDefault<ULevelEditorViewportSettings>()->SnapToSurface.bEnabled)
+		{
+			// HACK: If we are aligning rotation to surfaces, we have to factor in the inverse of the actor transform so that the resulting transform after SpawnActor is correct.
+
+			if (auto* RootComponent = Actor.GetRootComponent())
+			{
+				RootComponent->UpdateComponentToWorld();
+			}
+			ActorTransform = Actor.GetTransform().Inverse() * ActorTransform;
+		}
 	}
 
 	return ActorTransform;
