EditorViewportFlagsTest.cpp
--Updated the tests pretty name with "System." so that it will be categorized correctly.

diff --git a/Engine/Source/Editor/UnrealEd/Private/Tests/EditorViewportFlagsTest.cpp b/Engine/Source/Editor/UnrealEd/Private/Tests/EditorViewportFlagsTest.cpp
--- a/Engine/Source/Editor/UnrealEd/Private/Tests/EditorViewportFlagsTest.cpp
+++ b/Engine/Source/Editor/UnrealEd/Private/Tests/EditorViewportFlagsTest.cpp
@@ -99,7 +99,7 @@ namespace FViewportTestHelper
 /**
 * Unit test for enabling the view ports show flags
 */
-IMPLEMENT_COMPLEX_AUTOMATION_TEST( FViewportShowflagsToggleOnTest, "Editor.Viewport.Showflags.Toggle On", EAutomationTestFlags::ATF_Editor );
+IMPLEMENT_COMPLEX_AUTOMATION_TEST( FViewportShowflagsToggleOnTest, "System.Editor.Viewport.Showflags.Toggle On", EAutomationTestFlags::ATF_Editor );
 
 void FViewportShowflagsToggleOnTest::GetTests(TArray<FString>& OutBeautifiedNames, TArray<FString>& OutTestCommands) const
 {
@@ -134,7 +134,7 @@ bool FViewportShowflagsToggleOnTest::RunTest( const FString& Parameters )
 /**
 * Unit test for disabling the view ports show flags
 */
-IMPLEMENT_COMPLEX_AUTOMATION_TEST(FViewportShowflagsToggleOffTest, "Editor.Viewport.Showflags.Toggle Off", EAutomationTestFlags::ATF_Editor);
+IMPLEMENT_COMPLEX_AUTOMATION_TEST(FViewportShowflagsToggleOffTest, "System.Editor.Viewport.Showflags.Toggle Off", EAutomationTestFlags::ATF_Editor);
 
 void FViewportShowflagsToggleOffTest::GetTests(TArray<FString>& OutBeautifiedNames, TArray<FString>& OutTestCommands) const
 {
