Added file size to CookedAssetRegistry json file.

diff --git a/Engine/Source/Editor/UnrealEd/Private/Commandlets/ChunkManifestGenerator.cpp b/Engine/Source/Editor/UnrealEd/Private/Commandlets/ChunkManifestGenerator.cpp
--- a/Engine/Source/Editor/UnrealEd/Private/Commandlets/ChunkManifestGenerator.cpp
+++ b/Engine/Source/Editor/UnrealEd/Private/Commandlets/ChunkManifestGenerator.cpp
@@ -747,11 +747,14 @@ bool FChunkManifestGenerator::SaveCookedPackageAssetRegistry( const FString& San
 			FString PlatformSandboxPath = SandboxPath.Replace(TEXT("[Platform]"), *Platform->PlatformName());
 			
 			FDateTime TimeStamp = IFileManager::Get().GetTimeStamp( *PlatformSandboxPath );
+			int64 FileSize = IFileManager::Get().FileSize(*PlatformSandboxPath);
 
-			Json->WriteValue( "SourcePackageName", PackageName.ToString() );
-			Json->WriteValue( "CookedPackageName", PlatformSandboxPath );
-			Json->WriteValue( "CookedPackageTimeStamp", TimeStamp.ToString() );
+			Json->WriteValue( TEXT("SourcePackageName"), PackageName.ToString() );
+			Json->WriteValue( TEXT("CookedPackageName"), PlatformSandboxPath );
+			Json->WriteValue( TEXT("CookedPackageTimeStamp"), TimeStamp.ToString() );
+			Json->WriteValue( TEXT("FileSize"), FString::Printf(TEXT("%ll"), FileSize) );
 			
+
 			Json->WriteArrayStart("AssetData");
 			for (const auto& AssetData : AssetRegistryData)
 			{	// Add only assets that have actually been cooked and belong to any chunk
