Replaced extracted comment sorting in PO files with reference comment sorting.

diff --git a/Engine/Source/Editor/UnrealEd/Private/Localization/PortableObjectFormatDOM.cpp b/Engine/Source/Editor/UnrealEd/Private/Localization/PortableObjectFormatDOM.cpp
--- a/Engine/Source/Editor/UnrealEd/Private/Localization/PortableObjectFormatDOM.cpp
+++ b/Engine/Source/Editor/UnrealEd/Private/Localization/PortableObjectFormatDOM.cpp
@@ -812,7 +812,7 @@ void FPortableObjectFormatDOM::SortEntries()
 	// Sort keys.
 	for (const TSharedPtr<FPortableObjectEntry>& Entry : Entries)
 	{
-		Entry->ExtractedComments.Sort();
+		Entry->ReferenceComments.Sort();
 	}
 
 	// Sort by namespace, then keys, then source text.
