Changed dialogue wave namespace, removed optional contexts, added new object version for these changes, and updated asset localization gather logic to just-in-time gather for dialogue wave packages saved before this version.

diff --git a/Engine/Source/Editor/UnrealEd/Private/Commandlets/GatherTextFromAssetsCommandlet.cpp b/Engine/Source/Editor/UnrealEd/Private/Commandlets/GatherTextFromAssetsCommandlet.cpp
--- a/Engine/Source/Editor/UnrealEd/Private/Commandlets/GatherTextFromAssetsCommandlet.cpp
+++ b/Engine/Source/Editor/UnrealEd/Private/Commandlets/GatherTextFromAssetsCommandlet.cpp
@@ -361,8 +361,43 @@ int32 UGatherTextFromAssetsCommandlet::Main(const FString& Params)
 			FPackageFileSummary PackageFileSummary;
 			(*FileReader) << PackageFileSummary;
 
-			// Package has gatherable text data in its header, process immediately.
-			if (PackageFileSummary.GetFileVersionUE4() >= VER_UE4_SERIALIZE_TEXT_IN_PACKAGES)
+			bool MustLoadForGather = false;
+
+			// Packages not resaved since localization gathering flagging was added to packages must be loaded.
+			if (PackageFileSummary.GetFileVersionUE4() < VER_UE4_PACKAGE_REQUIRES_LOCALIZATION_GATHER_FLAGGING)
+			{
+				MustLoadForGather = true;
+			}
+			// Package not resaved since gatherable text data was added to package headers must be loaded, since their package header won't contain pregathered text data.
+			else if (PackageFileSummary.GetFileVersionUE4() < VER_UE4_SERIALIZE_TEXT_IN_PACKAGES)
+			{
+				// Fallback on the old package flag check.
+				if (PackageFileSummary.PackageFlags & PKG_RequiresLocalizationGather)
+				{
+					MustLoadForGather = true;
+				}
+			}
+			else if (PackageFileSummary.GetFileVersionUE4() < VER_UE4_DIALOGUE_WAVE_NAMESPACE_AND_CONTEXT_CHANGES)
+			{
+				IAssetRegistry& AssetRegistry = AssetRegistryModule.Get();
+				TArray<FAssetData> AssetDataInPackage;
+				AssetRegistry.GetAssetsByPackageName(*FPackageName::FilenameToLongPackageName(PackageFile), AssetDataInPackage);
+				for (const FAssetData& AssetData : AssetDataInPackage)
+				{
+					if (AssetData.AssetClass == UDialogueWave::StaticClass()->GetFName())
+					{
+						MustLoadForGather = true;
+					}
+				}
+			}
+				 
+			// Add package to list of packages to load fully and process.
+			if (MustLoadForGather)
+			{
+				PackageFileNamesToLoad.Add(PackageFile);
+			}
+			// Process immediately packages that don't require loading to process.
+			else
 			{
 				TArray<FGatherableTextData> GatherableTextDataArray;
 
@@ -379,11 +414,6 @@ int32 UGatherTextFromAssetsCommandlet::Main(const FString& Params)
 					ProcessGatherableTextDataArray(PackageFile, GatherableTextDataArray);
 				}
 			}
-			// Package not resaved since gatherable text data was added to package headers, defer processing until after loading.
-			else if (PackageFileSummary.PackageFlags & PKG_RequiresLocalizationGather || PackageFileSummary.GetFileVersionUE4() < VER_UE4_PACKAGE_REQUIRES_LOCALIZATION_GATHER_FLAGGING)
-			{
-				PackageFileNamesToLoad.Add(PackageFile);
-			}
 		}
 	}
 
diff --git a/Engine/Source/Runtime/Core/Public/UObject/ObjectVersion.h b/Engine/Source/Runtime/Core/Public/UObject/ObjectVersion.h
--- a/Engine/Source/Runtime/Core/Public/UObject/ObjectVersion.h
+++ b/Engine/Source/Runtime/Core/Public/UObject/ObjectVersion.h
@@ -566,6 +566,8 @@ enum EUnrealEngineObjectUE4Version
 	VER_UE4_KEEP_ONLY_PACKAGE_NAMES_IN_STRING_ASSET_REFERENCES_MAP,
 	// Support sound cue not saving out editor only data
 	VER_UE4_COOKED_ASSETS_IN_EDITOR_SUPPORT,
+	// Updated dialogue wave localization gathering logic.
+	VER_UE4_DIALOGUE_WAVE_NAMESPACE_AND_CONTEXT_CHANGES,
 	// -----<new versions can be added before this line>-------------------------------------------------
 	// - this needs to be the last line (see note below)
 	VER_UE4_AUTOMATIC_VERSION_PLUS_ONE,
diff --git a/Engine/Source/Runtime/Engine/Private/DialogueWave.cpp b/Engine/Source/Runtime/Engine/Private/DialogueWave.cpp
--- a/Engine/Source/Runtime/Engine/Private/DialogueWave.cpp
+++ b/Engine/Source/Runtime/Engine/Private/DialogueWave.cpp
@@ -1,4 +1,4 @@
-// Copyright 1998-2015 Epic Games, Inc. All Rights Reserved.
+﻿// Copyright 1998-2015 Epic Games, Inc. All Rights Reserved.
 
 #include "EnginePrivate.h"
 #include "ActiveSound.h"
@@ -57,7 +57,7 @@ namespace
 		TArray< FTextSourceSiteContext > ContextSpecificVariations;
 	};
 
-	const FString FDialogueHelper::DialogueNamespace						= TEXT("Dialogue");
+	const FString FDialogueHelper::DialogueNamespace						= TEXT("");
 	const FString FDialogueHelper::PropertyName_VoiceActorDirection		= TEXT("Voice Actor Direction");
 	const FString FDialogueHelper::PropertyName_Speaker					= TEXT("Speaker");
 	const FString FDialogueHelper::PropertyName_Speakers					= TEXT("Speakers");
@@ -194,7 +194,7 @@ namespace
 				{
 					Context.KeyName = DialogueWave->GetContextLocalizationKey( DialogueContext );
 					Context.SiteDescription = SourceLocation;
-					Context.IsOptional = true;
+					Context.IsOptional = false;
 					Context.KeyMetaData = KeyMetaDataObject->Values.Num() > 0 ? *KeyMetaDataObject : FLocMetadataObject();
 					Context.InfoMetaData = InfoMetaDataObject->Values.Num() > 0 ? *InfoMetaDataObject : FLocMetadataObject();
 				}
@@ -218,34 +218,6 @@ namespace
 			}
 		}
 
-		// Create the base, non-optional entry
-		{
-			TSharedPtr< FLocMetadataObject > InfoMetaDataObject = MakeShareable( new FLocMetadataObject() );
-
-			// Setup a loc metadata object with all the context specific info.  This usually includes human readable descriptions of the dialogue
-			{
-				if( VariationsDisplayInfoList.Num() > 0 )
-				{
-					InfoMetaDataObject->SetArrayField( PropertyName_DialogueVariations, VariationsDisplayInfoList );
-				}
-
-				if( !VoiceActorDirection.IsEmpty() )
-				{
-					InfoMetaDataObject->SetStringField( PropertyName_VoiceActorDirection, VoiceActorDirection );
-				}
-			}
-
-			
-			// Setup the context
-			{
-				Base.KeyName = DialogueKey;
-				Base.SiteDescription = SourceLocation;
-				Base.IsOptional = false;
-				Base.InfoMetaData = InfoMetaDataObject->Values.Num() > 0 ? *InfoMetaDataObject : FLocMetadataObject();
-			}
-
-		}
-
 		return true;
 	}
 
