Updated PO export/import logic to use generate msgctxt based on namespace and, conditionally, key (only if key metadata is present). Also added info metadata as extracted comments.

diff --git a/Engine/Source/Editor/UnrealEd/Private/Commandlets/InternationalizationExportCommandlet.cpp b/Engine/Source/Editor/UnrealEd/Private/Commandlets/InternationalizationExportCommandlet.cpp
--- a/Engine/Source/Editor/UnrealEd/Private/Commandlets/InternationalizationExportCommandlet.cpp
+++ b/Engine/Source/Editor/UnrealEd/Private/Commandlets/InternationalizationExportCommandlet.cpp
@@ -14,8 +14,114 @@ static const TCHAR* NewLineDelimiter = TEXT("\n");
 /**
 *	Helper Functions
 */
-FString ConditionArchiveStrForPo(const FString& InStr)
+namespace
 {
+	FString ConditionIdentityForPOMsgCtxt(const FString& Namespace, const FString& Key, const TSharedPtr<FLocMetadataObject>& KeyMetaData)
+	{
+		const auto& EscapeMsgCtxtParticle = [](const FString& InStr) -> FString
+		{
+			FString Result;
+			for (const TCHAR C : InStr)
+			{
+				switch (C)
+				{
+				case TEXT(','):		Result += TEXT("\\,");	break;
+				default:			Result += C;			break;
+				}
+			}
+			return Result;
+		};
+
+		FString EscapedNamespace = EscapeMsgCtxtParticle(Namespace);
+		FString EscapedKey = EscapeMsgCtxtParticle(Key);
+
+		return KeyMetaData.IsValid() ? FString::Printf(TEXT("%s,%s"), *EscapedNamespace, *EscapedKey) : EscapedNamespace;
+	}
+
+	void ParsePOMsgCtxtForIdentity(const FString& MsgCtxt, FString& OutNamespace, FString& OutKey)
+	{
+		static const int32 OutputBufferCount = 2;
+		FString* OutputBuffers[OutputBufferCount] = { &OutNamespace, &OutKey };
+		int32 OutputBufferIndex = 0;
+
+		FString EscapeSequenceBuffer;
+
+		auto HandleEscapeSequenceBuffer = [&]()
+		{
+			// Insert unescaped sequence if needed.
+			if (!EscapeSequenceBuffer.IsEmpty())
+			{
+				bool EscapeSequenceIdentified = true;
+
+				// Identify escape sequence
+				TCHAR UnescapedCharacter = 0;
+				if (EscapeSequenceBuffer == TEXT("\\,"))
+				{
+					UnescapedCharacter = ',';
+				}
+				else
+				{
+					EscapeSequenceIdentified = false;
+				}
+
+				// If identified, append the processed sequence as the unescaped character.
+				if (EscapeSequenceIdentified)
+				{
+					*OutputBuffers[OutputBufferIndex] += UnescapedCharacter;
+				}
+				// If it was not identified, preserve the escape sequence and append it.
+				else
+				{
+					*OutputBuffers[OutputBufferIndex] += EscapeSequenceBuffer;
+				}
+				// Either way, we've appended something based on the buffer and it should be reset.
+				EscapeSequenceBuffer.Empty();
+			}
+		};
+
+		for (const TCHAR C : MsgCtxt)
+		{
+			// If we're out of buffers, break out. The particle list is longer than expected.
+			if (OutputBufferIndex >= OutputBufferCount)
+			{
+				UE_LOG( LogInternationalizationExportCommandlet, Warning, TEXT("msgctxt found in PO has too many parts: %s"), *MsgCtxt );
+				break;
+			}
+
+			// Not in an escape sequence.
+			if (EscapeSequenceBuffer.IsEmpty())
+			{
+				// Comma marks the delimiter between namespace and key, if present.
+				if(C == TEXT(','))
+				{
+					++OutputBufferIndex;
+				}
+				// Regular character, just copy over.
+				else if (C != TEXT('\\'))
+				{
+					*OutputBuffers[OutputBufferIndex] += C;
+				}
+				// Start of an escape sequence, put in escape sequence buffer.
+				else
+				{
+					EscapeSequenceBuffer += C;
+				}
+			}
+			// If already in an escape sequence.
+			else
+			{
+				// Append to escape sequence buffer.
+				EscapeSequenceBuffer += C;
+
+				HandleEscapeSequenceBuffer();
+			}
+		}
+		// Catch any trailing backslashes.
+		HandleEscapeSequenceBuffer();
+	}
+
+	FString ConditionArchiveStrForPo(const FString& InStr)
+	{
 	FString Result;
 	for (const TCHAR C : InStr)
 	{
@@ -30,10 +136,10 @@ FString ConditionArchiveStrForPo(const FString& InStr)
 		}
 	}
 	return Result;
-}
+	}
 
-FString ConditionPoStringForArchive(const FString& InStr)
-{
+	FString ConditionPoStringForArchive(const FString& InStr)
+	{
 	FString Result;
 	FString EscapeSequenceBuffer;
 
@@ -114,15 +220,15 @@ FString ConditionPoStringForArchive(const FString& InStr)
 	// Catch any trailing backslashes.
 	HandleEscapeSequenceBuffer();
 	return Result;
-}
-
+	}
 
-FString ConvertSrcLocationToPORef(const FString& InSrcLocation)
-{
+	FString ConvertSrcLocationToPORef(const FString& InSrcLocation)
+	{
 	// Source location format: /Path1/Path2/file.cpp - line 123
 	// PO Reference format: /Path1/Path2/file.cpp:123
 	// @TODO: Note, we assume the source location format here but it could be arbitrary.
 	return InSrcLocation.Replace(TEXT(" - line "), TEXT(":"));
+	}
 }
 
 /**
@@ -251,9 +357,11 @@ bool UInternationalizationExportCommandlet::DoExport( const FString& SourcePath,
 						// For each context, we may need to create a different or even multiple PO entries.
 						for( auto ContextIter = ManifestEntry->Contexts.CreateConstIterator(); ContextIter; ++ContextIter )
 						{
+							const FContext& Context = *ContextIter;
+
 							// Create the typical PO entry from the archive entry which matches the exact same namespace, source, and key metadata, if it exists.
 							{
-								const TSharedPtr<FArchiveEntry> ArchiveEntry = InternationalizationArchive->FindEntryBySource( Namespace, Source, ContextIter->KeyMetadataObj );
+								const TSharedPtr<FArchiveEntry> ArchiveEntry = InternationalizationArchive->FindEntryBySource( Namespace, Source, Context.KeyMetadataObj );
 								if( ArchiveEntry.IsValid() )
 								{
 									const FString ConditionedArchiveSource = ConditionArchiveStrForPo(ArchiveEntry->Source.Text);
@@ -262,14 +370,29 @@ bool UInternationalizationExportCommandlet::DoExport( const FString& SourcePath,
 									TSharedRef<FPortableObjectEntry> PoEntry = MakeShareable( new FPortableObjectEntry );
 									//@TODO: We support additional metadata entries that can be translated.  How do those fit in the PO file format?  Ex: isMature
 									PoEntry->MsgId = ConditionedArchiveSource;
-									//@TODO: Take into account optional entries and entries that differ by keymetadata.  Ex. Each optional entry needs a unique msgCtxt
-									PoEntry->MsgCtxt = Namespace;
+									PoEntry->MsgCtxt = ConditionIdentityForPOMsgCtxt(Namespace, Context.Key, Context.KeyMetadataObj);
 									PoEntry->MsgStr.Add( ConditionedArchiveTranslation );
 
-									FString PORefString = ConvertSrcLocationToPORef( ContextIter->SourceLocation );
+									const FString PORefString = ConvertSrcLocationToPORef( Context.SourceLocation );
 									PoEntry->AddReference( PORefString ); // Source location.
-									PoEntry->AddExtractedComment( ContextIter->Key ); // "Notes from Programmer" in the form of the Key.
-									PoEntry->AddExtractedComment( PORefString ); // "Notes from Programmer" in the form of the Source Location, since this comes in handy too and OneSky doesn't properly show references, only comments.
+
+									PoEntry->AddExtractedComment( FString::Printf(TEXT("Key:\t%s"), *Context.Key) ); // "Notes from Programmer" in the form of the Key.
+									PoEntry->AddExtractedComment( FString::Printf(TEXT("SourceLocation:\t%s"), *PORefString) ); // "Notes from Programmer" in the form of the Source Location, since this comes in handy too and OneSky doesn't properly show references, only comments.
+									TArray<FString> InfoMetaDataStrings;
+									if (Context.InfoMetadataObj.IsValid())
+									{
+										for (auto InfoMetaDataPair : Context.InfoMetadataObj->Values)
+										{
+											const FString KeyName = InfoMetaDataPair.Key;
+											const TSharedPtr<FLocMetadataValue> Value = InfoMetaDataPair.Value;
+											InfoMetaDataStrings.Add(FString::Printf(TEXT("InfoMetaData:\t\"%s\" : \"%s\""), *KeyName, *Value->AsString()));
+										}
+									}
+									if (InfoMetaDataStrings.Num())
+									{
+										PoEntry->AddExtractedComments(InfoMetaDataStrings);
+									}
+
 									PortableObj.AddEntry( PoEntry );
 								}
 							}
@@ -281,7 +404,7 @@ bool UInternationalizationExportCommandlet::DoExport( const FString& SourcePath,
 								// Find the native archive entry which matches the exact same namespace, source, and key metadata, if it exists.
 								for (const auto& NativeArchive : NativeArchives)
 								{
-									const TSharedPtr<FArchiveEntry> PotentialNativeArchiveEntry = NativeArchive->FindEntryBySource( Namespace, Source, ContextIter->KeyMetadataObj );
+									const TSharedPtr<FArchiveEntry> PotentialNativeArchiveEntry = NativeArchive->FindEntryBySource( Namespace, Source, Context.KeyMetadataObj );
 									if (PotentialNativeArchiveEntry.IsValid())
 									{
 										NativeArchiveEntry = PotentialNativeArchiveEntry;
@@ -303,14 +426,29 @@ bool UInternationalizationExportCommandlet::DoExport( const FString& SourcePath,
 											TSharedRef<FPortableObjectEntry> PoEntry = MakeShareable( new FPortableObjectEntry );
 											//@TODO: We support additional metadata entries that can be translated.  How do those fit in the PO file format?  Ex: isMature
 											PoEntry->MsgId = ConditionedArchiveSource;
-											//@TODO: Take into account optional entries and entries that differ by keymetadata.  Ex. Each optional entry needs a unique msgCtxt
-											PoEntry->MsgCtxt = Namespace;
+											PoEntry->MsgCtxt = ConditionIdentityForPOMsgCtxt(Namespace, Context.Key, Context.KeyMetadataObj);
 											PoEntry->MsgStr.Add( ConditionedArchiveTranslation );
 
-											FString PORefString = ConvertSrcLocationToPORef( ContextIter->SourceLocation );
+											const FString PORefString = ConvertSrcLocationToPORef( Context.SourceLocation );
 											PoEntry->AddReference( PORefString ); // Source location.
-											PoEntry->AddExtractedComment( ContextIter->Key ); // "Notes from Programmer" in the form of the Key.
-											PoEntry->AddExtractedComment( PORefString ); // "Notes from Programmer" in the form of the Source Location, since this comes in handy too and OneSky doesn't properly show references, only comments.
+
+											PoEntry->AddExtractedComment( FString::Printf(TEXT("Key:\t%s"), *Context.Key) ); // "Notes from Programmer" in the form of the Key.
+											PoEntry->AddExtractedComment( FString::Printf(TEXT("SourceLocation:\t%s"), *PORefString) ); // "Notes from Programmer" in the form of the Source Location, since this comes in handy too and OneSky doesn't properly show references, only comments.
+											TArray<FString> InfoMetaDataStrings;
+											if (Context.InfoMetadataObj.IsValid())
+											{
+												for (auto InfoMetaDataPair : Context.InfoMetadataObj->Values)
+												{
+													const FString KeyName = InfoMetaDataPair.Key;
+													const TSharedPtr<FLocMetadataValue> Value = InfoMetaDataPair.Value;
+													InfoMetaDataStrings.Add(FString::Printf(TEXT("InfoMetaData:\t\"%s\" : \"%s\""), *KeyName, *Value->AsString()));
+												}
+											}
+											if (InfoMetaDataStrings.Num())
+											{
+												PoEntry->AddExtractedComments(InfoMetaDataStrings);
+											}
+
 											PortableObj.AddEntry( PoEntry );
 										}
 									}
@@ -423,8 +561,23 @@ bool UInternationalizationExportCommandlet::DoImport(const FString& SourcePath,
 		}
 
 
+		const FString ManifestFileName = DestinationPath / ManifestName;
+
+		TSharedPtr< FJsonObject > ManifestJsonObject = NULL;
+		ManifestJsonObject = ReadJSONTextFile( ManifestFileName );
+
+		FJsonInternationalizationManifestSerializer ManifestSerializer;
+		TSharedRef< FInternationalizationManifest > InternationalizationManifest = MakeShareable( new FInternationalizationManifest );
+		ManifestSerializer.DeserializeManifest( ManifestJsonObject.ToSharedRef(), InternationalizationManifest );
+
+		if( !FPaths::FileExists(ManifestFileName) )
+		{
+			UE_LOG( LogInternationalizationExportCommandlet, Error, TEXT("Failed to find manifest %s."), *ManifestFileName);
+			continue;
+		}
+
 		const FString DestinationCulturePath = DestinationPath / CultureName;
-		FString ArchiveFileName = DestinationCulturePath / ArchiveName;
+		const FString ArchiveFileName = DestinationCulturePath / ArchiveName;
 
 		if( !FPaths::FileExists(ArchiveFileName) )
 		{
@@ -456,12 +609,37 @@ bool UInternationalizationExportCommandlet::DoImport(const FString& SourcePath,
 					UE_LOG( LogInternationalizationExportCommandlet, Error, TEXT("Portable Object entry has plural form we did not process.  File: %s  MsgCtxt: %s  MsgId: %s"), *POFilePath, *POEntry->MsgCtxt, *POEntry->MsgId );
 				}
 
-				const FString& Namespace = POEntry->MsgCtxt;
+				FString Key;
+				FString Namespace;
+				ParsePOMsgCtxtForIdentity(POEntry->MsgCtxt, Namespace, Key);
 				const FString& SourceText = ConditionPoStringForArchive(POEntry->MsgId);
 				const FString& Translation = ConditionPoStringForArchive(POEntry->MsgStr[0]);
 
+				TSharedPtr<FLocMetadataObject> KeyMetaDataObject;
+				// Get key metadata from the manifest, using the namespace and key.
+				if (!Key.IsEmpty())
+				{
+					// Find manifest entry by namespace
+					for (auto ManifestEntryIterator = InternationalizationManifest->GetEntriesByContextIdIterator(); ManifestEntryIterator; ++ManifestEntryIterator)
+					{
+						const FString& ManifestEntryNamespace = ManifestEntryIterator->Key;
+						const TSharedRef<FManifestEntry>& ManifestEntry = ManifestEntryIterator->Value;
+						if (ManifestEntry->Namespace == Namespace)
+						{
+							FContext* const MatchingContext = ManifestEntry->Contexts.FindByPredicate([&](FContext& Context) -> bool
+								{
+									return Context.Key == Key;
+								});
+							if (MatchingContext)
+							{
+								KeyMetaDataObject = MatchingContext->KeyMetadataObj;
+							}
+						}
+					}
+				}
+
 				//@TODO: Take into account optional entries and entries that differ by keymetadata.  Ex. Each optional entry needs a unique msgCtxt
-				TSharedPtr< FArchiveEntry > FoundEntry = InternationalizationArchive->FindEntryBySource( Namespace, SourceText, NULL );
+				const TSharedPtr< FArchiveEntry > FoundEntry = InternationalizationArchive->FindEntryBySource( Namespace, SourceText, KeyMetaDataObject );
 				if( !FoundEntry.IsValid() )
 				{
 					UE_LOG(LogInternationalizationExportCommandlet, Warning, TEXT("Could not find corresponding archive entry for PO entry.  File: %s  MsgCtxt: %s  MsgId: %s"), *POFilePath, *POEntry->MsgCtxt, *POEntry->MsgId );
