Fixed compilation warning-as-error regarding shadowed variable in GatherTextFromAssetsCommandlet.cpp

diff --git a/Engine/Source/Editor/UnrealEd/Private/Commandlets/GatherTextFromAssetsCommandlet.cpp b/Engine/Source/Editor/UnrealEd/Private/Commandlets/GatherTextFromAssetsCommandlet.cpp
--- a/Engine/Source/Editor/UnrealEd/Private/Commandlets/GatherTextFromAssetsCommandlet.cpp
+++ b/Engine/Source/Editor/UnrealEd/Private/Commandlets/GatherTextFromAssetsCommandlet.cpp
@@ -241,14 +241,14 @@ int32 UGatherTextFromAssetsCommandlet::Main(const FString& Params)
 		}
 	}
 
-	TArray<FAssetData> AssetData;
-	AssetRegistryModule.Get().GetAssets(Filter, AssetData);
+	TArray<FAssetData> AssetDataArray;
+	AssetRegistryModule.Get().GetAssets(Filter, AssetDataArray);
 
 	FString UAssetPackageExtension = FPackageName::GetAssetPackageExtension();
 	TSet< FString > LongPackageNamesToExclude;
-	for (int Index = 0; Index < AssetData.Num(); Index++)
+	for (int Index = 0; Index < AssetDataArray.Num(); Index++)
 	{
-		LongPackageNamesToExclude.Add( FPackageName::LongPackageNameToFilename( AssetData[Index].PackageName.ToString(), UAssetPackageExtension ) );
+		LongPackageNamesToExclude.Add( FPackageName::LongPackageNameToFilename( AssetDataArray[Index].PackageName.ToString(), UAssetPackageExtension ) );
 	}
 
 	//Get whether we should fix broken properties that we find.
