Fixed orthographic box selection using the an incorrect world-space box for intersection

This addresses UE-19518

diff --git a/Engine/Source/Editor/UnrealEd/Private/DragTool_BoxSelect.cpp b/Engine/Source/Editor/UnrealEd/Private/DragTool_BoxSelect.cpp
--- a/Engine/Source/Editor/UnrealEd/Private/DragTool_BoxSelect.cpp
+++ b/Engine/Source/Editor/UnrealEd/Private/DragTool_BoxSelect.cpp
@@ -241,11 +241,9 @@ void FDragTool_ActorBoxSelect::CalculateBox( FBox& OutBox )
 	FVector4 StartScreenPos = View->PixelToScreen(Start.X, Start.Y, 0);
 	FVector4 EndScreenPos = View->PixelToScreen(End.X, End.Y, 0);
 
-	FMatrix InvViewMatrix = View->ViewMatrices.GetInvViewMatrix();
-	FMatrix InvProjMatrix = View->ViewMatrices.GetInvProjMatrix();
+	FVector TransformedStart = View->ScreenToWorld(View->PixelToScreen(Start.X, Start.Y, 0.5f));
+	FVector TransformedEnd = View->ScreenToWorld(View->PixelToScreen(End.X, End.Y, 0.5f));
 
-	FVector TransformedStart = InvViewMatrix.TransformFVector4(InvProjMatrix.TransformFVector4(FVector4(StartScreenPos.X, StartScreenPos.Y, 0.5f, 1.0f)));
-	FVector TransformedEnd = InvViewMatrix.TransformFVector4(InvProjMatrix.TransformFVector4(FVector4(EndScreenPos.X, EndScreenPos.Y, 0.5f, 1.0f)));
 	// Create a bounding box based on the start/end points (normalizes the points).
 	OutBox.Init();
 	OutBox += TransformedStart;
