Fix vehicle crash caused when trying to create a convex hull around no vertices. The real issue here was inside PhAT which will now do the right thing and give a warning.

diff --git a/Engine/Source/Editor/PhAT/Private/PhAT.cpp b/Engine/Source/Editor/PhAT/Private/PhAT.cpp
--- a/Engine/Source/Editor/PhAT/Private/PhAT.cpp
+++ b/Engine/Source/Editor/PhAT/Private/PhAT.cpp
@@ -2186,8 +2186,13 @@ void FPhAT::OnResetBoneCollision()
 			int32 BoneIndex = SharedData->EditorSkelMesh->RefSkeleton.FindBoneIndex(BodySetup->BoneName);
 			check(BoneIndex != INDEX_NONE);
 
-			FPhysicsAssetUtils::CreateCollisionFromBone(BodySetup, SharedData->EditorSkelMesh, BoneIndex, SharedData->NewBodyData, (SharedData->NewBodyData.VertWeight == EVW_DominantWeight)? SharedData->DominantWeightBoneInfos: SharedData->AnyWeightBoneInfos);
-			BodyIndices.AddUnique(SharedData->SelectedBodies[i].Index);
+			if(FPhysicsAssetUtils::CreateCollisionFromBone(BodySetup, SharedData->EditorSkelMesh, BoneIndex, SharedData->NewBodyData, (SharedData->NewBodyData.VertWeight == EVW_DominantWeight)? SharedData->DominantWeightBoneInfos: SharedData->AnyWeightBoneInfos))
+			{
+				BodyIndices.AddUnique(SharedData->SelectedBodies[i].Index);
+			}else
+			{
+				FPhysicsAssetUtils::DestroyBody(SharedData->PhysicsAsset, SharedData->SelectedBodies[i].Index);
+			}
 		}
 
 		//deselect first
diff --git a/Engine/Source/Editor/UnrealEd/Private/PhysicsAssetUtils.cpp b/Engine/Source/Editor/UnrealEd/Private/PhysicsAssetUtils.cpp
--- a/Engine/Source/Editor/UnrealEd/Private/PhysicsAssetUtils.cpp
+++ b/Engine/Source/Editor/UnrealEd/Private/PhysicsAssetUtils.cpp
@@ -442,16 +442,24 @@ bool CreateCollisionFromBone( UBodySetup* bs, USkeletalMesh* skelMesh, int32 Bon
 	// Deal with creating a single convex hull
 	else if (Params.GeomType == EFG_SingleConvexHull)
 	{
-		FKConvexElem ConvexElem;
+		if (BoneInfo->Positions.Num())
+		{
+			FKConvexElem ConvexElem;
 
-		// Add all of the vertices for this bone to the convex element
-		for( int32 index=0; index<BoneInfo->Positions.Num(); index++ )
+			// Add all of the vertices for this bone to the convex element
+			for( int32 index=0; index<BoneInfo->Positions.Num(); index++ )
+			{
+				ConvexElem.VertexData.Add(BoneInfo->Positions[index]);
+			}
+			ConvexElem.UpdateElemBox();
+			bs->AggGeom.ConvexElems.Add(ConvexElem);
+		}else
 		{
-			ConvexElem.VertexData.Add(BoneInfo->Positions[index]);
+			FMessageLog EditorErrors("EditorErrors");
+			EditorErrors.Warning(NSLOCTEXT("PhysicsAssetUtils", "ConvexNoPositions", "Unable to create a convex hull for the given bone as there are no vertices associated with the bone."));
+			EditorErrors.Open();
+			return false;
 		}
-		ConvexElem.UpdateElemBox();
-
-		bs->AggGeom.ConvexElems.Add(ConvexElem);
 	}
 	else if (Params.GeomType == EFG_MultiConvexHull)
 	{
