HLOD: Print user friendly message if you try build and meshes don't have have LODs #jira UE-14356

diff --git a/Engine/Source/Editor/UnrealEd/Private/HierarchicalLOD.cpp b/Engine/Source/Editor/UnrealEd/Private/HierarchicalLOD.cpp
--- a/Engine/Source/Editor/UnrealEd/Private/HierarchicalLOD.cpp
+++ b/Engine/Source/Editor/UnrealEd/Private/HierarchicalLOD.cpp
@@ -97,7 +97,7 @@ void FHierarchicalLODBuilder::BuildClusters(ULevel* InLevel, const bool bCreateM
 	CollectGarbage( GARBAGE_COLLECTION_KEEPFLAGS, true );
 
 	// only build if it's enabled
-	if(InLevel->GetWorld()->GetWorldSettings()->bEnableHierarchicalLODSystem)
+	if(InLevel->GetWorld()->GetWorldSettings()->bEnableHierarchicalLODSystem && InLevel->GetWorld()->GetWorldSettings()->HierarchicalLODSetup.Num() != 0)
 	{
 		// Handle HierachicalLOD volumes first
 		HandleHLODVolumes(InLevel);
@@ -151,7 +151,7 @@ void FHierarchicalLODBuilder::BuildClusters(ULevel* InLevel, const bool bCreateM
 		FMessageLog MapCheck("MapCheck");
 		MapCheck.Warning()
 			->AddToken(FUObjectToken::Create(InLevel->GetWorld()->GetWorldSettings()))
-			->AddToken(FTextToken::Create(LOCTEXT("MapCheck_Message_HLODSystemNotEnabled", "Hierarchical LOD System is disabled, unable to build LOD actors.")))
+			->AddToken(FTextToken::Create(LOCTEXT("MapCheck_Message_HLODSystemNotEnabled", "Hierarchical LOD System is disabled or no HLOD level settings available, unable to build LOD actors.")))
 			->AddToken(FMapErrorToken::Create(FMapErrors::HLODSystemNotEnabled));
 	}
 
