Refactored the Open Assets test to be a complex test.
--Deleted a couple of helper functions that were used for this.
--Removed the old individual tests, replaced with one test.
--New tests will show in the UFE automation tab with the same naming scheme.

diff --git a/Engine/Source/Editor/UnrealEd/Private/Tests/AutomationEditorCommon.cpp b/Engine/Source/Editor/UnrealEd/Private/Tests/AutomationEditorCommon.cpp
--- a/Engine/Source/Editor/UnrealEd/Private/Tests/AutomationEditorCommon.cpp
+++ b/Engine/Source/Editor/UnrealEd/Private/Tests/AutomationEditorCommon.cpp
@@ -900,114 +900,11 @@ void FEditorAutomationTestUtilities::CollectGameContentTestsByClass(UClass * Cla
 	}
 }
 
-/**
-* Generates a list of misc. assets from the GAME.
-* This is to be used by the GetTest() function.
-*/
-void FEditorAutomationTestUtilities::CollectMiscGameContentTestsByClass(TArray<FString>& OutBeautifiedNames, TArray <FString>& OutTestCommands)
-{
-	//Setting the Asset Registry
-	FAssetRegistryModule& AssetRegistryModule = FModuleManager::LoadModuleChecked<FAssetRegistryModule>(TEXT("AssetRegistry"));
-	TArray<FAssetData> ObjectList;
-	FARFilter AssetFilter;
-
-	//This is a list of classes that we don't want to be in the misc category.
-	TArray<FName> ExcludeClassesList;
-	ExcludeClassesList.Add(UAimOffsetBlendSpace::StaticClass()->GetFName());
-	ExcludeClassesList.Add(UAimOffsetBlendSpace1D::StaticClass()->GetFName());
-	ExcludeClassesList.Add(UAnimBlueprint::StaticClass()->GetFName());
-	ExcludeClassesList.Add(UAnimMontage::StaticClass()->GetFName());
-	ExcludeClassesList.Add(UAnimSequence::StaticClass()->GetFName());
-	ExcludeClassesList.Add(UBehaviorTree::StaticClass()->GetFName());
-	ExcludeClassesList.Add(UBlendSpace::StaticClass()->GetFName());
-	ExcludeClassesList.Add(UBlendSpace1D::StaticClass()->GetFName());
-	ExcludeClassesList.Add(UBlueprint::StaticClass()->GetFName());
-	ExcludeClassesList.Add(UDestructibleMesh::StaticClass()->GetFName());
-	ExcludeClassesList.Add(UDialogueVoice::StaticClass()->GetFName());
-	ExcludeClassesList.Add(UDialogueWave::StaticClass()->GetFName());
-	ExcludeClassesList.Add(UFont::StaticClass()->GetFName());
-	ExcludeClassesList.Add(UMaterial::StaticClass()->GetFName());
-	ExcludeClassesList.Add(UMaterialFunction::StaticClass()->GetFName());
-	ExcludeClassesList.Add(UMaterialInstanceConstant::StaticClass()->GetFName());
-	ExcludeClassesList.Add(UMaterialParameterCollection::StaticClass()->GetFName());
-	ExcludeClassesList.Add(UParticleSystem::StaticClass()->GetFName());
-	ExcludeClassesList.Add(UPhysicalMaterial::StaticClass()->GetFName());
-	ExcludeClassesList.Add(UPhysicsAsset::StaticClass()->GetFName());
-	ExcludeClassesList.Add(UReverbEffect::StaticClass()->GetFName());
-	ExcludeClassesList.Add(USkeletalMesh::StaticClass()->GetFName());
-	ExcludeClassesList.Add(USkeleton::StaticClass()->GetFName());
-	ExcludeClassesList.Add(USlateBrushAsset::StaticClass()->GetFName());
-	ExcludeClassesList.Add(USlateWidgetStyleAsset::StaticClass()->GetFName());
-	ExcludeClassesList.Add(USoundAttenuation::StaticClass()->GetFName());
-	ExcludeClassesList.Add(USoundClass::StaticClass()->GetFName());
-	ExcludeClassesList.Add(USoundCue::StaticClass()->GetFName());
-	ExcludeClassesList.Add(USoundMix::StaticClass()->GetFName());
-	ExcludeClassesList.Add(USoundWave::StaticClass()->GetFName());
-	ExcludeClassesList.Add(UStaticMesh::StaticClass()->GetFName());
-	ExcludeClassesList.Add(USubsurfaceProfile::StaticClass()->GetFName());
-	ExcludeClassesList.Add(UTexture::StaticClass()->GetFName());
-	ExcludeClassesList.Add(UTexture2D::StaticClass()->GetFName());
-	ExcludeClassesList.Add(UTextureCube::StaticClass()->GetFName());
-	ExcludeClassesList.Add(UTextureRenderTarget::StaticClass()->GetFName());
-	ExcludeClassesList.Add(UTextureRenderTarget2D::StaticClass()->GetFName());
-	ExcludeClassesList.Add(UUserDefinedEnum::StaticClass()->GetFName());
-	ExcludeClassesList.Add(UWorld::StaticClass()->GetFName());
-
-	//Generating the list of assets.
-	//This list is being filtered by the game folder and class type.  The results are placed into the ObjectList variable.
-	AssetFilter.PackagePaths.Add("/Game");
-	AssetFilter.bRecursivePaths = true;
-	AssetRegistryModule.Get().GetAssets(AssetFilter, ObjectList);
-
-	//Loop through the list of assets, make their path full and a string, then add them to the test.
-	for (auto ObjIter = ObjectList.CreateConstIterator(); ObjIter; ++ObjIter)
-	{
-		const FAssetData& Asset = *ObjIter;
-		//First we check if the class is valid.  If not then we move onto the next object.
-		if (Asset.GetClass() != NULL)
-		{
-			//Variable that holds the class FName for the current asset iteration.
-			FName AssetClassFName = Asset.GetClass()->GetFName();
-
-			//Counter used to keep track for the following for loop.
-			float ExcludedClassesCounter = 1;
-
-			for (auto ExcludeIter = ExcludeClassesList.CreateConstIterator(); ExcludeIter; ++ExcludeIter)
-			{
-				FName ExludedName = *ExcludeIter;
-
-				//If the classes are the same then we don't want this asset. So we move onto the next one instead.
-				if (AssetClassFName == ExludedName)
-				{
-					break;
-				}
-
-				//We run out of class names in our Excluded list then we want the current ObjectList asset.
-				if ((ExcludedClassesCounter + 1) > ExcludeClassesList.Num())
-				{
-					FString Filename = Asset.ObjectPath.ToString();
-					//convert to full paths
-					Filename = FPackageName::LongPackageNameToFilename(Filename);
-
-					if (FAutomationTestFramework::GetInstance().ShouldTestContent(Filename))
-					{
-						FString BeautifiedFilename = Asset.AssetName.ToString();
-						OutBeautifiedNames.Add(BeautifiedFilename);
-						OutTestCommands.Add(Asset.ObjectPath.ToString());
-					}
-
-					break;
-				}
-				ExcludedClassesCounter++;
-			}
-		}
-	}
-}
-
 /**
 * Generates a list of assets from the GAME by a specific type.
+* This is to be used by the GetTest() function.
 */
-void FEditorAutomationTestUtilities::CollectGameContentByClass(const UClass * Class, bool bRecursiveClass, TArray<FString>& OutAssetList)
+void FEditorAutomationTestUtilities::CollectGameContentTests(TArray<FString>& OutBeautifiedNames, TArray <FString>& OutTestCommands)
 {
 	//Setting the Asset Registry
 	FAssetRegistryModule& AssetRegistryModule = FModuleManager::LoadModuleChecked<FAssetRegistryModule>(TEXT("AssetRegistry"));
@@ -1015,12 +912,10 @@ void FEditorAutomationTestUtilities::CollectGameContentByClass(const UClass * Cl
 	//Variable setups
 	TArray<FAssetData> ObjectList;
 	FARFilter AssetFilter;
-
-	//Generating the list of assets.
-	//This list is being filtered by the game folder and class type.  The results are placed into the ObjectList variable.
-	AssetFilter.ClassNames.Add(Class->GetFName());
+	
+	//removed path as a filter as it causes two large lists to be sorted.  Filtering on "game" directory on iteration
 	AssetFilter.PackagePaths.Add("/Game");
-	AssetFilter.bRecursiveClasses = bRecursiveClass;
+	AssetFilter.bRecursiveClasses = true;
 	AssetFilter.bRecursivePaths = true;
 	AssetRegistryModule.Get().GetAssets(AssetFilter, ObjectList);
 
@@ -1029,11 +924,17 @@ void FEditorAutomationTestUtilities::CollectGameContentByClass(const UClass * Cl
 	{
 		const FAssetData& Asset = *ObjIter;
 		FString Filename = Asset.ObjectPath.ToString();
-		//convert to full paths
-		Filename = FPackageName::LongPackageNameToFilename(Filename);
-		if (FAutomationTestFramework::GetInstance().ShouldTestContent(Filename))
+
+		if (Filename.StartsWith("/Game"))
 		{
-			OutAssetList.Add(Asset.ObjectPath.ToString());
+			//convert to full paths
+			Filename = FPackageName::LongPackageNameToFilename(Filename);
+			if (FAutomationTestFramework::GetInstance().ShouldTestContent(Filename))
+			{
+				FString BeautifiedFilename = FString::Printf(TEXT("%s.%s"), *Asset.GetClass()->GetFName().ToString(), *Asset.AssetName.ToString());
+				OutBeautifiedNames.Add(BeautifiedFilename);
+				OutTestCommands.Add(Asset.ObjectPath.ToString());
+			}
 		}
 	}
-}
+}
\ No newline at end of file
diff --git a/Engine/Source/Editor/UnrealEd/Private/Tests/EditorOpenAssetAutomationTests.cpp b/Engine/Source/Editor/UnrealEd/Private/Tests/EditorOpenAssetAutomationTests.cpp
--- a/Engine/Source/Editor/UnrealEd/Private/Tests/EditorOpenAssetAutomationTests.cpp
+++ b/Engine/Source/Editor/UnrealEd/Private/Tests/EditorOpenAssetAutomationTests.cpp
@@ -3,37 +3,9 @@
 #include "UnrealEd.h"
 #include "Tests/AutomationTestSettings.h"
 #include "AutomationEditorCommon.h"
-#include "AutomationTest.h"
-#include "AssertionMacros.h"
 #include "AutomationCommon.h"
-#include "AssetEditorManager.h"
 #include "ShaderCompiler.h"
 
-//Includes needed for opening certain assets
-#include "Materials/MaterialFunction.h"
-#include "Slate/SlateBrushAsset.h"
-#include "BehaviorTree/BehaviorTree.h"
-#include "Engine/DestructibleMesh.h"
-#include "Animation/AimOffsetBlendSpace.h"
-#include "Animation/AimOffsetBlendSpace1D.h"
-#include "Animation/BlendSpace1D.h"
-#include "Engine/UserDefinedEnum.h"
-#include "Sound/DialogueVoice.h"
-#include "Sound/DialogueWave.h"
-#include "Materials/MaterialInstanceConstant.h"
-#include "Materials/MaterialParameterCollection.h"
-#include "PhysicalMaterials/PhysicalMaterial.h"
-#include "PhysicsEngine/PhysicsAsset.h"
-#include "Sound/ReverbEffect.h"
-#include "Sound/SoundAttenuation.h"
-#include "Engine/StaticMesh.h"
-#include "Sound/SoundClass.h"
-#include "Sound/SoundCue.h"
-#include "Sound/SoundMix.h"
-#include "Sound/SoundWave.h"
-#include "Engine/SubsurfaceProfile.h"
-#include "Engine/Font.h"
-
 #define ASSET_OPEN_INTERVAL 1.5f
 
 bool FOpenActualAssetEditors(const FString& Parameters)
@@ -238,654 +210,17 @@ void FOpenAllAssetTests(OpenAssetParameters& AssetList)
 	ADD_LATENT_AUTOMATION_COMMAND(FCloseAllAssetEditorsCommand());
 }
 
-//////////////////////////////////////////////////////////////////////////
-/**
-* This test opens each AimOffsetBlendSpace into its sub-editor.
-*/
-IMPLEMENT_COMPLEX_AUTOMATION_TEST(FOpenAimOffsetBlendSpaceAssetEditors, "Project.Tools.Skeletal Mesh.Animation.BlendSpaces.Open Aim Offset BlendSpace Assets", (EAutomationTestFlags::ATF_Editor | EAutomationTestFlags::ATF_RequiresUser));
-
-void FOpenAimOffsetBlendSpaceAssetEditors::GetTests(TArray<FString>& OutBeautifiedNames, TArray<FString>& OutTestCommands) const
-{
-	//This grabs each AimOffsetBlendSpace asset in the Game/Content directory
-	FEditorAutomationTestUtilities::CollectGameContentTestsByClass(UAimOffsetBlendSpace::StaticClass(), true, OutBeautifiedNames, OutTestCommands);
-}
-
-bool FOpenAimOffsetBlendSpaceAssetEditors::RunTest(const FString& Parameters)
-{
-	bool bDidTheTestPass = FOpenActualAssetEditors(Parameters);
-	CollectGarbage(GARBAGE_COLLECTION_KEEPFLAGS);
-	return bDidTheTestPass;
-}
-
-//////////////////////////////////////////////////////////////////////////
-/**
-* This test opens each AimOffsetBlendSpace1D into its sub-editor.
-*/
-IMPLEMENT_COMPLEX_AUTOMATION_TEST(FOpenAimOffsetBlendSpace1DAssetEditors, "Project.Tools.Skeletal Mesh.Animation.BlendSpaces.Open Aim Offset BlendSpace1D Assets", (EAutomationTestFlags::ATF_Editor | EAutomationTestFlags::ATF_RequiresUser));
-
-void FOpenAimOffsetBlendSpace1DAssetEditors::GetTests(TArray<FString>& OutBeautifiedNames, TArray<FString>& OutTestCommands) const
-{
-	//This grabs each AimOffsetBlendSpace1D asset in the Game/Content directory
-	FEditorAutomationTestUtilities::CollectGameContentTestsByClass(UAimOffsetBlendSpace1D::StaticClass(), true, OutBeautifiedNames, OutTestCommands);
-}
-
-bool FOpenAimOffsetBlendSpace1DAssetEditors::RunTest(const FString& Parameters)
-{
-	bool bDidTheTestPass = FOpenActualAssetEditors(Parameters);
-	CollectGarbage(GARBAGE_COLLECTION_KEEPFLAGS);
-	return bDidTheTestPass;
-}
-
-//////////////////////////////////////////////////////////////////////////
-/**
-* This test opens each AnimMontage into its sub-editor.
-*/
-IMPLEMENT_COMPLEX_AUTOMATION_TEST(FOpenAnimMontageAssetEditors, "Project.Tools.Skeletal Mesh.Animation.Open AnimMontage Assets", (EAutomationTestFlags::ATF_Editor | EAutomationTestFlags::ATF_RequiresUser));
-
-void FOpenAnimMontageAssetEditors::GetTests(TArray<FString>& OutBeautifiedNames, TArray<FString>& OutTestCommands) const
-{
-	//This grabs each AnimMontage asset in the Game/Content directory
-	FEditorAutomationTestUtilities::CollectGameContentTestsByClass(UAnimMontage::StaticClass(), true, OutBeautifiedNames, OutTestCommands);
-}
-
-bool FOpenAnimMontageAssetEditors::RunTest(const FString& Parameters)
-{
-	bool bDidTheTestPass = FOpenActualAssetEditors(Parameters);
-	CollectGarbage(GARBAGE_COLLECTION_KEEPFLAGS);
-	return bDidTheTestPass;
-}
-
-//////////////////////////////////////////////////////////////////////////
-/**
-* This test opens each AnimSequence into its sub-editor.
-*/
-IMPLEMENT_COMPLEX_AUTOMATION_TEST(FOpenAnimSequenceAssetEditors, "Project.Tools.Skeletal Mesh.Animation.Open AnimSequence Assets", (EAutomationTestFlags::ATF_Editor | EAutomationTestFlags::ATF_RequiresUser));
-
-void FOpenAnimSequenceAssetEditors::GetTests(TArray<FString>& OutBeautifiedNames, TArray<FString>& OutTestCommands) const
-{
-	//This grabs each AnimSequence asset in the Game/Content directory
-	FEditorAutomationTestUtilities::CollectGameContentTestsByClass(UAnimSequence::StaticClass(), true, OutBeautifiedNames, OutTestCommands);
-}
-
-bool FOpenAnimSequenceAssetEditors::RunTest(const FString& Parameters)
-{
-	bool bDidTheTestPass = FOpenActualAssetEditors(Parameters);
-	CollectGarbage(GARBAGE_COLLECTION_KEEPFLAGS);
-	return bDidTheTestPass;
-}
-
-//////////////////////////////////////////////////////////////////////////
-/**
-* This test opens each BEHAVIOR TREE into its sub-editor.
-*/
-IMPLEMENT_COMPLEX_AUTOMATION_TEST(FOpenBehaviorTreeAssetEditors, "Project.Tools.AI.Open Behavior Tree Assets", (EAutomationTestFlags::ATF_Editor | EAutomationTestFlags::ATF_RequiresUser));
-
-void FOpenBehaviorTreeAssetEditors::GetTests(TArray<FString>& OutBeautifiedNames, TArray<FString>& OutTestCommands) const
-{
-	//This grabs each Behavior Tree asset in the Game/Content directory
-	FEditorAutomationTestUtilities::CollectGameContentTestsByClass(UBehaviorTree::StaticClass(), true, OutBeautifiedNames, OutTestCommands);
-}
-
-bool FOpenBehaviorTreeAssetEditors::RunTest(const FString& Parameters)
-{
-	bool bDidTheTestPass = FOpenActualAssetEditors(Parameters);
-	CollectGarbage(GARBAGE_COLLECTION_KEEPFLAGS);
-	return bDidTheTestPass;
-}
-
-//////////////////////////////////////////////////////////////////////////
-/**
-* This test opens each BlendSpace into its sub-editor.
-*/
-IMPLEMENT_COMPLEX_AUTOMATION_TEST(FOpenBlendSpaceAssetEditors, "Project.Tools.Skeletal Mesh.Animation.BlendSpaces.Open BlendSpace Assets", (EAutomationTestFlags::ATF_Editor | EAutomationTestFlags::ATF_RequiresUser));
-
-void FOpenBlendSpaceAssetEditors::GetTests(TArray<FString>& OutBeautifiedNames, TArray<FString>& OutTestCommands) const
-{
-	//This grabs each BlendSpace asset in the Game/Content directory
-	FEditorAutomationTestUtilities::CollectGameContentTestsByClass(UBlendSpace::StaticClass(), false, OutBeautifiedNames, OutTestCommands);
-}
-
-bool FOpenBlendSpaceAssetEditors::RunTest(const FString& Parameters)
-{
-	bool bDidTheTestPass = FOpenActualAssetEditors(Parameters);
-	CollectGarbage(GARBAGE_COLLECTION_KEEPFLAGS);
-	return bDidTheTestPass;
-}
-
-//////////////////////////////////////////////////////////////////////////
-/**
-* This test opens each BlendSpace1D into its sub-editor.
-*/
-IMPLEMENT_COMPLEX_AUTOMATION_TEST(FOpenBlendSpace1DAssetEditors, "Project.Tools.Skeletal Mesh.Animation.BlendSpaces.Open BlendSpace1D Assets", (EAutomationTestFlags::ATF_Editor | EAutomationTestFlags::ATF_RequiresUser));
-
-void FOpenBlendSpace1DAssetEditors::GetTests(TArray<FString>& OutBeautifiedNames, TArray<FString>& OutTestCommands) const
-{
-	//This grabs each BlendSpace1D asset in the Game/Content directory
-	FEditorAutomationTestUtilities::CollectGameContentTestsByClass(UBlendSpace1D::StaticClass(), false, OutBeautifiedNames, OutTestCommands);
-}
-
-bool FOpenBlendSpace1DAssetEditors::RunTest(const FString& Parameters)
-{
-	bool bDidTheTestPass = FOpenActualAssetEditors(Parameters);
-	CollectGarbage(GARBAGE_COLLECTION_KEEPFLAGS);
-	return bDidTheTestPass;
-}
-
-//////////////////////////////////////////////////////////////////////////
-/**
-* This test opens each BLUEPRINT into its sub-editor.
-*/
-IMPLEMENT_COMPLEX_AUTOMATION_TEST(FOpenBlueprintAssetEditors, "Project.Tools.Blueprint.Open Blueprint Assets", (EAutomationTestFlags::ATF_Editor | EAutomationTestFlags::ATF_RequiresUser));
-
-void FOpenBlueprintAssetEditors::GetTests(TArray<FString>& OutBeautifiedNames, TArray<FString>& OutTestCommands) const
-{
-	//This grabs each Blueprints asset in the Game/Content directory
-	FEditorAutomationTestUtilities::CollectGameContentTestsByClass(UBlueprint::StaticClass(), true, OutBeautifiedNames, OutTestCommands);
-}
-
-bool FOpenBlueprintAssetEditors::RunTest(const FString& Parameters)
-{
-	bool bDidTheTestPass = FOpenActualAssetEditors(Parameters);
-	CollectGarbage(GARBAGE_COLLECTION_KEEPFLAGS);
-	return bDidTheTestPass;
-}
-
-//////////////////////////////////////////////////////////////////////////
-/**
-* This test opens each UserDefinedEnum into its sub-editor.
-*/
-IMPLEMENT_COMPLEX_AUTOMATION_TEST(FOpenUserDefinedEnumAssetEditors, "Project.Tools.Blueprint.Open User Defined Enum Assets", (EAutomationTestFlags::ATF_Editor | EAutomationTestFlags::ATF_RequiresUser));
-
-void FOpenUserDefinedEnumAssetEditors::GetTests(TArray<FString>& OutBeautifiedNames, TArray<FString>& OutTestCommands) const
-{
-	//This grabs each UserDefinedEnum asset in the Game/Content directory
-	FEditorAutomationTestUtilities::CollectGameContentTestsByClass(UUserDefinedEnum::StaticClass(), true, OutBeautifiedNames, OutTestCommands);
-}
-
-bool FOpenUserDefinedEnumAssetEditors::RunTest(const FString& Parameters)
-{
-	bool bDidTheTestPass = FOpenActualAssetEditors(Parameters);
-	CollectGarbage(GARBAGE_COLLECTION_KEEPFLAGS);
-	return bDidTheTestPass;
-}
-
-//////////////////////////////////////////////////////////////////////////
-/**
-* This test opens each DESTRUCTIBLE MESH into its sub-editor.
-*/
-IMPLEMENT_COMPLEX_AUTOMATION_TEST(FOpenDestructibleMeshAssetEditors, "Project.Tools.Skeletal Mesh.Open Destructible Mesh Assets", (EAutomationTestFlags::ATF_Editor | EAutomationTestFlags::ATF_RequiresUser));
-
-void FOpenDestructibleMeshAssetEditors::GetTests(TArray<FString>& OutBeautifiedNames, TArray<FString>& OutTestCommands) const
-{
-	//This grabs each DestructibleMesh asset in the Game/Content directory
-	FEditorAutomationTestUtilities::CollectGameContentTestsByClass(UDestructibleMesh::StaticClass(), false, OutBeautifiedNames, OutTestCommands);
-}
-
-bool FOpenDestructibleMeshAssetEditors::RunTest(const FString& Parameters)
-{
-	bool bDidTheTestPass = FOpenActualAssetEditors(Parameters);
-	CollectGarbage(GARBAGE_COLLECTION_KEEPFLAGS);
-	return bDidTheTestPass;
-}
-
-//////////////////////////////////////////////////////////////////////////
-/**
-* This test opens each DialogueVoice into its sub-editor.
-*/
-IMPLEMENT_COMPLEX_AUTOMATION_TEST(FOpenDialogueVoiceAssetEditors, "Project.Tools.Sounds.Open Dialogue Voice Assets", (EAutomationTestFlags::ATF_Editor | EAutomationTestFlags::ATF_RequiresUser));
-
-void FOpenDialogueVoiceAssetEditors::GetTests(TArray<FString>& OutBeautifiedNames, TArray<FString>& OutTestCommands) const
-{
-	//This grabs each DialogueVoice asset in the Game/Content directory
-	FEditorAutomationTestUtilities::CollectGameContentTestsByClass(UDialogueVoice::StaticClass(), true, OutBeautifiedNames, OutTestCommands);
-}
-
-bool FOpenDialogueVoiceAssetEditors::RunTest(const FString& Parameters)
-{
-	bool bDidTheTestPass = FOpenActualAssetEditors(Parameters);
-	CollectGarbage(GARBAGE_COLLECTION_KEEPFLAGS);
-	return bDidTheTestPass;
-}
-
-//////////////////////////////////////////////////////////////////////////
-/**
-* This test opens each DialogueWave into its sub-editor.
-*/
-IMPLEMENT_COMPLEX_AUTOMATION_TEST(FOpenDialogueWaveAssetEditors, "Project.Tools.Sounds.Open Dialogue Wave Assets", (EAutomationTestFlags::ATF_Editor | EAutomationTestFlags::ATF_RequiresUser));
-
-void FOpenDialogueWaveAssetEditors::GetTests(TArray<FString>& OutBeautifiedNames, TArray<FString>& OutTestCommands) const
-{
-	//This grabs each DialogueWave asset in the Game/Content directory
-	FEditorAutomationTestUtilities::CollectGameContentTestsByClass(UDialogueWave::StaticClass(), true, OutBeautifiedNames, OutTestCommands);
-}
-
-bool FOpenDialogueWaveAssetEditors::RunTest(const FString& Parameters)
-{
-	bool bDidTheTestPass = FOpenActualAssetEditors(Parameters);
-	CollectGarbage(GARBAGE_COLLECTION_KEEPFLAGS);
-	return bDidTheTestPass;
-}
-
-//////////////////////////////////////////////////////////////////////////
-/**
-* This test opens each Material into its sub-editor.
-*/
-IMPLEMENT_COMPLEX_AUTOMATION_TEST(FOpenMaterialAssetEditors, "Project.Tools.Materials.Open Material Assets", (EAutomationTestFlags::ATF_Editor | EAutomationTestFlags::ATF_RequiresUser));
-
-void FOpenMaterialAssetEditors::GetTests(TArray<FString>& OutBeautifiedNames, TArray<FString>& OutTestCommands) const
-{
-	//This grabs each Material asset in the Game/Content directory
-	FEditorAutomationTestUtilities::CollectGameContentTestsByClass(UMaterial::StaticClass(), true, OutBeautifiedNames, OutTestCommands);
-}
-
-bool FOpenMaterialAssetEditors::RunTest(const FString& Parameters)
-{
-	bool bDidTheTestPass = FOpenActualAssetEditors(Parameters);
-	CollectGarbage(GARBAGE_COLLECTION_KEEPFLAGS);
-	return bDidTheTestPass;
-}
-
-////////////////////////////////////////////////////////////////////////////
-///**
-//* This test opens each MaterialFunction into its sub-editor.
-//*/
-IMPLEMENT_COMPLEX_AUTOMATION_TEST(FOpenMaterialFunctionAssetEditors, "Project.Tools.Materials.Open Material Function Assets", (EAutomationTestFlags::ATF_Editor | EAutomationTestFlags::ATF_RequiresUser));
-
-void FOpenMaterialFunctionAssetEditors::GetTests(TArray<FString>& OutBeautifiedNames, TArray<FString>& OutTestCommands) const
-{
-	//This grabs each MaterialFunction asset in the Game/Content directory
-	FEditorAutomationTestUtilities::CollectGameContentTestsByClass(UMaterialFunction::StaticClass(), true, OutBeautifiedNames, OutTestCommands);
-}
-
-bool FOpenMaterialFunctionAssetEditors::RunTest(const FString& Parameters)
-{
-	bool bDidTheTestPass = FOpenActualAssetEditors(Parameters);
-	CollectGarbage(GARBAGE_COLLECTION_KEEPFLAGS);
-	return bDidTheTestPass;
-}
-
-//////////////////////////////////////////////////////////////////////////
-/**
-* This test opens each MaterialInstanceConstant into its sub-editor.
-*/
-IMPLEMENT_COMPLEX_AUTOMATION_TEST(FOpenMaterialInstanceConstantAssetEditors, "Project.Tools.Materials.Open Material Instance Constant Assets", (EAutomationTestFlags::ATF_Editor | EAutomationTestFlags::ATF_RequiresUser));
-
-void FOpenMaterialInstanceConstantAssetEditors::GetTests(TArray<FString>& OutBeautifiedNames, TArray<FString>& OutTestCommands) const
-{
-	//This grabs each MaterialInstanceConstant asset in the Game/Content directory
-	FEditorAutomationTestUtilities::CollectGameContentTestsByClass(UMaterialInstanceConstant::StaticClass(), true, OutBeautifiedNames, OutTestCommands);
-}
-
-bool FOpenMaterialInstanceConstantAssetEditors::RunTest(const FString& Parameters)
-{
-	bool bDidTheTestPass = FOpenActualAssetEditors(Parameters);
-	CollectGarbage(GARBAGE_COLLECTION_KEEPFLAGS);
-	return bDidTheTestPass;
-}
-
-//////////////////////////////////////////////////////////////////////////
-/**
-* This test opens each MaterialParameterCollection into its sub-editor.
-*/
-IMPLEMENT_COMPLEX_AUTOMATION_TEST(FOpenMaterialParameterCollectionAssetEditors, "Project.Tools.Materials.Open Material Parameter Collection Assets", (EAutomationTestFlags::ATF_Editor | EAutomationTestFlags::ATF_RequiresUser));
-
-void FOpenMaterialParameterCollectionAssetEditors::GetTests(TArray<FString>& OutBeautifiedNames, TArray<FString>& OutTestCommands) const
-{
-	//This grabs each MaterialParameterCollection asset in the Game/Content directory
-	FEditorAutomationTestUtilities::CollectGameContentTestsByClass(UMaterialParameterCollection::StaticClass(), true, OutBeautifiedNames, OutTestCommands);
-}
-
-bool FOpenMaterialParameterCollectionAssetEditors::RunTest(const FString& Parameters)
-{
-	bool bDidTheTestPass = FOpenActualAssetEditors(Parameters);
-	CollectGarbage(GARBAGE_COLLECTION_KEEPFLAGS);
-	return bDidTheTestPass;
-}
-
-//////////////////////////////////////////////////////////////////////////
-/**
-* This test opens each PARTICLE SYSTEM into its sub-editor.
-*/
-IMPLEMENT_COMPLEX_AUTOMATION_TEST(FOpenParticleSystemAssetEditors, "Project.Tools.Particle System.Open Particle System Assets", (EAutomationTestFlags::ATF_Editor | EAutomationTestFlags::ATF_RequiresUser));
-
-void FOpenParticleSystemAssetEditors::GetTests(TArray<FString>& OutBeautifiedNames, TArray<FString>& OutTestCommands) const
-{
-	//This grabs each ParticleSystem asset in the Game/Content directory
-	FEditorAutomationTestUtilities::CollectGameContentTestsByClass(UParticleSystem::StaticClass(), true, OutBeautifiedNames, OutTestCommands);
-}
-
-bool FOpenParticleSystemAssetEditors::RunTest(const FString& Parameters)
-{
-	bool bDidTheTestPass = FOpenActualAssetEditors(Parameters);
-	CollectGarbage(GARBAGE_COLLECTION_KEEPFLAGS);
-	return bDidTheTestPass;
-}
-
-//////////////////////////////////////////////////////////////////////////
-/**
-* This test opens each PhysicalMaterial into its sub-editor.
-*/
-IMPLEMENT_COMPLEX_AUTOMATION_TEST(FOpenPhysicalMaterialAssetEditors, "Project.Tools.Physics.Open Physical Material Assets", (EAutomationTestFlags::ATF_Editor | EAutomationTestFlags::ATF_RequiresUser));
-
-void FOpenPhysicalMaterialAssetEditors::GetTests(TArray<FString>& OutBeautifiedNames, TArray<FString>& OutTestCommands) const
-{
-	//This grabs each PhysicalMaterial asset in the Game/Content directory
-	FEditorAutomationTestUtilities::CollectGameContentTestsByClass(UPhysicalMaterial::StaticClass(), true, OutBeautifiedNames, OutTestCommands);
-}
-
-bool FOpenPhysicalMaterialAssetEditors::RunTest(const FString& Parameters)
-{
-	bool bDidTheTestPass = FOpenActualAssetEditors(Parameters);
-	CollectGarbage(GARBAGE_COLLECTION_KEEPFLAGS);
-	return bDidTheTestPass;
-}
-
-
-//////////////////////////////////////////////////////////////////////////
-/**
-* This test opens each PhysicsAsset into its sub-editor.
-*/
-IMPLEMENT_COMPLEX_AUTOMATION_TEST(FOpenPhysicsAssetAssetEditors, "Project.Tools.Physics.Open Physics Assets", (EAutomationTestFlags::ATF_Editor | EAutomationTestFlags::ATF_RequiresUser));
-
-void FOpenPhysicsAssetAssetEditors::GetTests(TArray<FString>& OutBeautifiedNames, TArray<FString>& OutTestCommands) const
-{
-	//This grabs each PhysicsAsset asset in the Game/Content directory
-	FEditorAutomationTestUtilities::CollectGameContentTestsByClass(UPhysicsAsset::StaticClass(), true, OutBeautifiedNames, OutTestCommands);
-}
-
-bool FOpenPhysicsAssetAssetEditors::RunTest(const FString& Parameters)
-{
-	bool bDidTheTestPass = FOpenActualAssetEditors(Parameters);
-	CollectGarbage(GARBAGE_COLLECTION_KEEPFLAGS);
-	return bDidTheTestPass;
-}
-
-
-//////////////////////////////////////////////////////////////////////////
-/**
-* This test opens each ReverbEffect into its sub-editor.
-*/
-IMPLEMENT_COMPLEX_AUTOMATION_TEST(FOpenReverbEffectAssetEditors, "Project.Tools.Sounds.Open Reverb Effect Assets", (EAutomationTestFlags::ATF_Editor | EAutomationTestFlags::ATF_RequiresUser));
-
-void FOpenReverbEffectAssetEditors::GetTests(TArray<FString>& OutBeautifiedNames, TArray<FString>& OutTestCommands) const
-{
-	//This grabs each ReverbEffect asset in the Game/Content directory
-	FEditorAutomationTestUtilities::CollectGameContentTestsByClass(UReverbEffect::StaticClass(), true, OutBeautifiedNames, OutTestCommands);
-}
-
-bool FOpenReverbEffectAssetEditors::RunTest(const FString& Parameters)
-{
-	bool bDidTheTestPass = FOpenActualAssetEditors(Parameters);
-	CollectGarbage(GARBAGE_COLLECTION_KEEPFLAGS);
-	return bDidTheTestPass;
-}
-
-
-
-//////////////////////////////////////////////////////////////////////////
-/**
-* This test opens each SlateWidgetStyleAsset into its sub-editor.
-*/
-IMPLEMENT_COMPLEX_AUTOMATION_TEST(FOpenSlateWidgetStyleAssetAssetEditors, "Project.Tools.Slate.Open Slate Widget Style Assets", (EAutomationTestFlags::ATF_Editor | EAutomationTestFlags::ATF_RequiresUser));
-
-void FOpenSlateWidgetStyleAssetAssetEditors::GetTests(TArray<FString>& OutBeautifiedNames, TArray<FString>& OutTestCommands) const
-{
-	//This grabs each SlateWidgetStyleAsset asset in the Game/Content directory
-	FEditorAutomationTestUtilities::CollectGameContentTestsByClass(USlateWidgetStyleAsset::StaticClass(), true, OutBeautifiedNames, OutTestCommands);
-}
-
-bool FOpenSlateWidgetStyleAssetAssetEditors::RunTest(const FString& Parameters)
-{
-	bool bDidTheTestPass = FOpenActualAssetEditors(Parameters);
-	CollectGarbage(GARBAGE_COLLECTION_KEEPFLAGS);
-	return bDidTheTestPass;
-}
 
 //////////////////////////////////////////////////////////////////////////
 /**
-* This test opens each Slate Brush Asset into its sub-editor.
+* This is a complex test that opens every asset into it's editor.
 */
-IMPLEMENT_COMPLEX_AUTOMATION_TEST(FOpenSlateBrushAssetAssetEditors, "Project.Tools.Slate.Open Slate Brush Assets", (EAutomationTestFlags::ATF_Editor | EAutomationTestFlags::ATF_RequiresUser));
-
-void FOpenSlateBrushAssetAssetEditors::GetTests(TArray<FString>& OutBeautifiedNames, TArray<FString>& OutTestCommands) const
-{
-	//This grabs each SlateBrushAsset asset in the Game/Content directory
-	FEditorAutomationTestUtilities::CollectGameContentTestsByClass(USlateBrushAsset::StaticClass(), true, OutBeautifiedNames, OutTestCommands);
-}
-
-bool FOpenSlateBrushAssetAssetEditors::RunTest(const FString& Parameters)
-{
-	bool bDidTheTestPass = FOpenActualAssetEditors(Parameters);
-	CollectGarbage(GARBAGE_COLLECTION_KEEPFLAGS);
-	return bDidTheTestPass;
-}
-
-//////////////////////////////////////////////////////////////////////////
-/**
-* This test opens each SoundAttenuation into its sub-editor.
-*/
-IMPLEMENT_COMPLEX_AUTOMATION_TEST(FOpenSoundAttenuationAssetEditors, "Project.Tools.Sounds.Open Sound Attenuation Assets", (EAutomationTestFlags::ATF_Editor | EAutomationTestFlags::ATF_RequiresUser));
-
-void FOpenSoundAttenuationAssetEditors::GetTests(TArray<FString>& OutBeautifiedNames, TArray<FString>& OutTestCommands) const
-{
-	//This grabs each SoundAttenuation asset in the Game/Content directory
-	FEditorAutomationTestUtilities::CollectGameContentTestsByClass(USoundAttenuation::StaticClass(), true, OutBeautifiedNames, OutTestCommands);
-}
-
-bool FOpenSoundAttenuationAssetEditors::RunTest(const FString& Parameters)
-{
-	bool bDidTheTestPass = FOpenActualAssetEditors(Parameters);
-	CollectGarbage(GARBAGE_COLLECTION_KEEPFLAGS);
-	return bDidTheTestPass;
-}
-
-//////////////////////////////////////////////////////////////////////////
-/**
-* This test opens each Static Mesh into their sub-editor.
-*/
-IMPLEMENT_COMPLEX_AUTOMATION_TEST(FOpenStaticMeshAssetEditors, "Project.Tools.Static Mesh.Open Static Mesh Assets", (EAutomationTestFlags::ATF_Editor | EAutomationTestFlags::ATF_RequiresUser));
-
-void FOpenStaticMeshAssetEditors::GetTests(TArray<FString>& OutBeautifiedNames, TArray<FString>& OutTestCommands) const
-{
-	//This grabs each Static Mesh in the Game/Content
-	FEditorAutomationTestUtilities::CollectGameContentTestsByClass(UStaticMesh::StaticClass(), true, OutBeautifiedNames, OutTestCommands);
-}
-
-bool FOpenStaticMeshAssetEditors::RunTest(const FString& Parameters)
-{
-	bool bDidTheTestPass = FOpenActualAssetEditors(Parameters);
-	CollectGarbage(GARBAGE_COLLECTION_KEEPFLAGS);
-	return bDidTheTestPass;
-}
-
-//////////////////////////////////////////////////////////////////////////
-/**
-* This test opens each Skeletal Mesh into their sub-editor.
-*/
-IMPLEMENT_COMPLEX_AUTOMATION_TEST(FOpenSkeletalMeshAssetEditors, "Project.Tools.Skeletal Mesh.Open Skeletal Mesh Assets", (EAutomationTestFlags::ATF_Editor | EAutomationTestFlags::ATF_RequiresUser));
-
-void FOpenSkeletalMeshAssetEditors::GetTests(TArray<FString>& OutBeautifiedNames, TArray<FString>& OutTestCommands) const
-{
-	//This grabs each SKeletal Mesh in the Game/Content
-	FEditorAutomationTestUtilities::CollectGameContentTestsByClass(USkeletalMesh::StaticClass(), false, OutBeautifiedNames, OutTestCommands);
-}
-
-bool FOpenSkeletalMeshAssetEditors::RunTest(const FString& Parameters)
-{
-	bool bDidTheTestPass = FOpenActualAssetEditors(Parameters);
-	CollectGarbage(GARBAGE_COLLECTION_KEEPFLAGS);
-	return bDidTheTestPass;
-}
-
-//////////////////////////////////////////////////////////////////////////
-/**
-* This test opens each Skeleton into its sub-editor.
-*/
-IMPLEMENT_COMPLEX_AUTOMATION_TEST(FOpenSkeletonAssetEditors, "Project.Tools.Skeletal Mesh.Open Skeleton Assets", (EAutomationTestFlags::ATF_Editor | EAutomationTestFlags::ATF_RequiresUser));
-
-void FOpenSkeletonAssetEditors::GetTests(TArray<FString>& OutBeautifiedNames, TArray<FString>& OutTestCommands) const
-{
-	//This grabs each Skeleton asset in the Game/Content directory
-	FEditorAutomationTestUtilities::CollectGameContentTestsByClass(USkeleton::StaticClass(), true, OutBeautifiedNames, OutTestCommands);
-}
-
-bool FOpenSkeletonAssetEditors::RunTest(const FString& Parameters)
-{
-	bool bDidTheTestPass = FOpenActualAssetEditors(Parameters);
-	CollectGarbage(GARBAGE_COLLECTION_KEEPFLAGS);
-	return bDidTheTestPass;
-}
-
-//////////////////////////////////////////////////////////////////////////
-/**
-* This test opens each SoundClass into its sub-editor.
-*/
-IMPLEMENT_COMPLEX_AUTOMATION_TEST(FOpenSoundClassAssetEditors, "Project.Tools.Sounds.Open SoundClass Assets", (EAutomationTestFlags::ATF_Editor | EAutomationTestFlags::ATF_RequiresUser));
-
-void FOpenSoundClassAssetEditors::GetTests(TArray<FString>& OutBeautifiedNames, TArray<FString>& OutTestCommands) const
-{
-	//This grabs each SoundClass asset in the Game/Content directory
-	FEditorAutomationTestUtilities::CollectGameContentTestsByClass(USoundClass::StaticClass(), true, OutBeautifiedNames, OutTestCommands);
-}
-
-bool FOpenSoundClassAssetEditors::RunTest(const FString& Parameters)
-{
-	bool bDidTheTestPass = FOpenActualAssetEditors(Parameters);
-	CollectGarbage(GARBAGE_COLLECTION_KEEPFLAGS);
-	return bDidTheTestPass;
-}
-
-
-//////////////////////////////////////////////////////////////////////////
-/**
-* This test opens each SoundCue into its sub-editor.
-*/
-IMPLEMENT_COMPLEX_AUTOMATION_TEST(FOpenSoundCueAssetEditors, "Project.Tools.Sounds.Open Sound Cue Assets", (EAutomationTestFlags::ATF_Editor | EAutomationTestFlags::ATF_RequiresUser));
-
-void FOpenSoundCueAssetEditors::GetTests(TArray<FString>& OutBeautifiedNames, TArray<FString>& OutTestCommands) const
-{
-	//This grabs each SoundCue asset in the Game/Content directory
-	FEditorAutomationTestUtilities::CollectGameContentTestsByClass(USoundCue::StaticClass(), true, OutBeautifiedNames, OutTestCommands);
-}
-
-bool FOpenSoundCueAssetEditors::RunTest(const FString& Parameters)
-{
-	bool bDidTheTestPass = FOpenActualAssetEditors(Parameters);
-	CollectGarbage(GARBAGE_COLLECTION_KEEPFLAGS);
-	return bDidTheTestPass;
-}
-
-
-//////////////////////////////////////////////////////////////////////////
-/**
-* This test opens each SoundMix into its sub-editor.
-*/
-IMPLEMENT_COMPLEX_AUTOMATION_TEST(FOpenSoundMixAssetEditors, "Project.Tools.Sounds.Open Sound Mix Assets", (EAutomationTestFlags::ATF_Editor | EAutomationTestFlags::ATF_RequiresUser));
-
-void FOpenSoundMixAssetEditors::GetTests(TArray<FString>& OutBeautifiedNames, TArray<FString>& OutTestCommands) const
-{
-	//This grabs each SoundMix asset in the Game/Content directory
-	FEditorAutomationTestUtilities::CollectGameContentTestsByClass(USoundMix::StaticClass(), true, OutBeautifiedNames, OutTestCommands);
-}
-
-bool FOpenSoundMixAssetEditors::RunTest(const FString& Parameters)
-{
-	bool bDidTheTestPass = FOpenActualAssetEditors(Parameters);
-	CollectGarbage(GARBAGE_COLLECTION_KEEPFLAGS);
-	return bDidTheTestPass;
-}
-
-
-//////////////////////////////////////////////////////////////////////////
-/**
-* This test opens each SoundWave into its sub-editor.
-*/
-IMPLEMENT_COMPLEX_AUTOMATION_TEST(FOpenSoundWaveAssetEditors, "Project.Tools.Sounds.Open Sound Wave Assets", (EAutomationTestFlags::ATF_Editor | EAutomationTestFlags::ATF_RequiresUser));
-
-void FOpenSoundWaveAssetEditors::GetTests(TArray<FString>& OutBeautifiedNames, TArray<FString>& OutTestCommands) const
-{
-	//This grabs each SoundWave asset in the Game/Content directory
-	FEditorAutomationTestUtilities::CollectGameContentTestsByClass(USoundWave::StaticClass(), true, OutBeautifiedNames, OutTestCommands);
-}
-
-bool FOpenSoundWaveAssetEditors::RunTest(const FString& Parameters)
-{
-	bool bDidTheTestPass = FOpenActualAssetEditors(Parameters);
-	CollectGarbage(GARBAGE_COLLECTION_KEEPFLAGS);
-	return bDidTheTestPass;
-}
-
-//////////////////////////////////////////////////////////////////////////
-/**
-* This test opens each SubsurfaceProfile into its sub-editor.
-*/
-IMPLEMENT_COMPLEX_AUTOMATION_TEST(FOpenSubsurfaceProfileAssetEditors, "Project.Tools.Materials.Open Subsurface Profile Assets", (EAutomationTestFlags::ATF_Editor | EAutomationTestFlags::ATF_RequiresUser));
-
-void FOpenSubsurfaceProfileAssetEditors::GetTests(TArray<FString>& OutBeautifiedNames, TArray<FString>& OutTestCommands) const
-{
-	//This grabs each SubsurfaceProfile asset in the Game/Content directory
-	FEditorAutomationTestUtilities::CollectGameContentTestsByClass(USubsurfaceProfile::StaticClass(), true, OutBeautifiedNames, OutTestCommands);
-}
-
-bool FOpenSubsurfaceProfileAssetEditors::RunTest(const FString& Parameters)
-{
-	bool bDidTheTestPass = FOpenActualAssetEditors(Parameters);
-	CollectGarbage(GARBAGE_COLLECTION_KEEPFLAGS);
-	return bDidTheTestPass;
-}
-
-
-
-//////////////////////////////////////////////////////////////////////////
-/**
-* This test opens each TEXTURE asset into its sub-editor.
-*/
-IMPLEMENT_COMPLEX_AUTOMATION_TEST(FOpenTextureAssetEditors, "Project.Tools.Textures.Open Texture Assets", (EAutomationTestFlags::ATF_Editor | EAutomationTestFlags::ATF_RequiresUser));
-
-void FOpenTextureAssetEditors::GetTests(TArray<FString>& OutBeautifiedNames, TArray<FString>& OutTestCommands) const
-{
-	//This grabs each texture file in the Game/Content
-	FEditorAutomationTestUtilities::CollectGameContentTestsByClass(UTexture::StaticClass(), true, OutBeautifiedNames, OutTestCommands);
-}
-
-bool FOpenTextureAssetEditors::RunTest(const FString& Parameters)
-{
-	bool bDidTheTestPass = FOpenActualAssetEditors(Parameters);
-	CollectGarbage(GARBAGE_COLLECTION_KEEPFLAGS);
-	return bDidTheTestPass;
-}
-
-
-//////////////////////////////////////////////////////////////////////////
-/**
-* This test opens each Font into its sub-editor.
-*/
-IMPLEMENT_COMPLEX_AUTOMATION_TEST(FOpenFontAssetEditors, "Project.Tools.Textures.Open Font Assets", (EAutomationTestFlags::ATF_Editor | EAutomationTestFlags::ATF_RequiresUser));
-
-void FOpenFontAssetEditors::GetTests(TArray<FString>& OutBeautifiedNames, TArray<FString>& OutTestCommands) const
-{
-	//This grabs each Font asset in the Game/Content directory
-	FEditorAutomationTestUtilities::CollectGameContentTestsByClass(UFont::StaticClass(), true, OutBeautifiedNames, OutTestCommands);
-}
-
-bool FOpenFontAssetEditors::RunTest(const FString& Parameters)
-{
-	bool bDidTheTestPass = FOpenActualAssetEditors(Parameters);
-	CollectGarbage(GARBAGE_COLLECTION_KEEPFLAGS);
-	return bDidTheTestPass;
-}
-
-
-////////////////////////////////////////////////////////////////////////////
-///**
-//* This test opens each misc asset into its sub-editor.
-//*/
-IMPLEMENT_COMPLEX_AUTOMATION_TEST(FOpenMiscAssetEditors, "Project.Tools.Misc.Open Misc Assets", (EAutomationTestFlags::ATF_Editor | EAutomationTestFlags::ATF_RequiresUser));
+IMPLEMENT_COMPLEX_AUTOMATION_TEST(FOpenMiscAssetEditors, "Project.Tools.Open Assets", (EAutomationTestFlags::ATF_Editor | EAutomationTestFlags::ATF_RequiresUser));
 
 void FOpenMiscAssetEditors::GetTests(TArray<FString>& OutBeautifiedNames, TArray<FString>& OutTestCommands) const
 {
 	//This grabs each TextureRenderTarget2D asset in the Game/Content directory
-	FEditorAutomationTestUtilities::CollectMiscGameContentTestsByClass(OutBeautifiedNames, OutTestCommands);
+	FEditorAutomationTestUtilities::CollectGameContentTests(OutBeautifiedNames, OutTestCommands);
 }
 
 bool FOpenMiscAssetEditors::RunTest(const FString& Parameters)
@@ -893,7 +228,4 @@ bool FOpenMiscAssetEditors::RunTest(const FString& Parameters)
 	bool bDidTheTestPass = FOpenActualAssetEditors(Parameters);
 	CollectGarbage(GARBAGE_COLLECTION_KEEPFLAGS);
 	return bDidTheTestPass;
-}
-
-
-
+}
\ No newline at end of file
diff --git a/Engine/Source/Editor/UnrealEd/Public/Tests/AutomationEditorCommon.h b/Engine/Source/Editor/UnrealEd/Public/Tests/AutomationEditorCommon.h
--- a/Engine/Source/Editor/UnrealEd/Public/Tests/AutomationEditorCommon.h
+++ b/Engine/Source/Editor/UnrealEd/Public/Tests/AutomationEditorCommon.h
@@ -1,15 +1,12 @@
 // Copyright 1998-2015 Epic Games, Inc. All Rights Reserved.
 
 #pragma once
+
 #include "AssetRegistryModule.h"
 #include "ModuleManager.h"
 #include "AutomationCommon.h"
 #include "Tests/AutomationTestSettings.h"
 
-//Includes needed for opening certain assets
-#include "Materials/MaterialFunction.h"
-#include "Slate/SlateBrushAsset.h"
-
 namespace AutomationEditorCommonUtils
 {
 	/**
@@ -302,16 +299,11 @@ public:
 	*/
 	static void CollectGameContentTestsByClass(UClass * Class, bool bRecursiveClass, TArray<FString>& OutBeautifiedNames, TArray <FString>& OutTestCommands);
 		
-	/**
-	* Generates a list of misc. assets from the GAME.
-	* This is to be used by the GetTest() function.
-	*/
-	static void CollectMiscGameContentTestsByClass(TArray<FString>& OutBeautifiedNames, TArray <FString>& OutTestCommands);
-	
 	/**
 	* Generates a list of assets from the GAME by a specific type.
+	* This is to be used by the GetTest() function.
 	*/
-	static void CollectGameContentByClass(const UClass * Class, bool bRecursiveClass, TArray<FString>& OutAssetList);
+	static void CollectGameContentTests(TArray<FString>& OutBeautifiedNames, TArray <FString>& OutTestCommands);
 };
 
 
