Fix non-unity compile

diff --git a/Engine/Source/Editor/UnrealEd/Private/Tests/EditorOpenAssetAutomationTests.cpp b/Engine/Source/Editor/UnrealEd/Private/Tests/EditorOpenAssetAutomationTests.cpp
--- a/Engine/Source/Editor/UnrealEd/Private/Tests/EditorOpenAssetAutomationTests.cpp
+++ b/Engine/Source/Editor/UnrealEd/Private/Tests/EditorOpenAssetAutomationTests.cpp
@@ -5,6 +5,7 @@
 #include "AutomationEditorCommon.h"
 #include "AutomationCommon.h"
 #include "ShaderCompiler.h"
+#include "AssetEditorManager.h"
 
 #define ASSET_OPEN_INTERVAL 1.5f
 
