Fixed issue where widget deltas were not being correctly calculated when the widget was positioned behind the camera.

diff --git a/Engine/Source/Editor/UnrealEd/Private/UnrealWidget.cpp b/Engine/Source/Editor/UnrealEd/Private/UnrealWidget.cpp
--- a/Engine/Source/Editor/UnrealEd/Private/UnrealWidget.cpp
+++ b/Engine/Source/Editor/UnrealEd/Private/UnrealWidget.cpp
@@ -67,6 +67,12 @@ FWidget::FWidget()
 	bSnapEnabled = false;
 	bDefaultVisibility = true;
 	bIsOrthoDrawingFullRing = false;
+
+	Origin = FVector2D::ZeroVector;
+	XAxisDir = FVector2D::ZeroVector;
+	YAxisDir = FVector2D::ZeroVector;
+	ZAxisDir = FVector2D::ZeroVector;
+	DragStartPos = FVector2D::ZeroVector;
 }
 
 extern ENGINE_API void StringSize(UFont* Font,int32& XL,int32& YL,const TCHAR* Text, FCanvas* Canvas);
@@ -163,32 +169,34 @@ void FWidget::Render( const FSceneView* View,FPrimitiveDrawInterface* PDI, FEdit
 		return;
 	}
 
-	FVector Loc = ViewportClient->GetWidgetLocation();
-	if(!View->ScreenToPixel(View->WorldToScreen(Loc),Origin))
+	FVector WidgetLocation = ViewportClient->GetWidgetLocation();
+	FVector2D NewOrigin;
+	if (View->ScreenToPixel(View->WorldToScreen(WidgetLocation), NewOrigin))
 	{
-		Origin.X = Origin.Y = 0;
+		// Only update the viewport-space origin if the position was in front of the camera
+		Origin = NewOrigin;
 	}
 
 	switch( ViewportClient->GetWidgetMode() )
 	{
 		case WM_Translate:
-			Render_Translate( View, PDI, ViewportClient, Loc, bDrawWidget );
+			Render_Translate(View, PDI, ViewportClient, WidgetLocation, bDrawWidget);
 			break;
 
 		case WM_Rotate:
-			Render_Rotate( View, PDI, ViewportClient, Loc, bDrawWidget );
+			Render_Rotate(View, PDI, ViewportClient, WidgetLocation, bDrawWidget);
 			break;
 
 		case WM_Scale:
-			Render_Scale( View, PDI, ViewportClient, Loc, bDrawWidget );
+			Render_Scale(View, PDI, ViewportClient, WidgetLocation, bDrawWidget);
 			break;
 
 		case WM_TranslateRotateZ:
-			Render_TranslateRotateZ( View, PDI, ViewportClient, Loc, bDrawWidget );
+			Render_TranslateRotateZ(View, PDI, ViewportClient, WidgetLocation, bDrawWidget);
 			break;
 
 		case WM_2D:
-			Render_2D( View, PDI, ViewportClient, Loc, bDrawWidget );
+			Render_2D(View, PDI, ViewportClient, WidgetLocation, bDrawWidget);
 			break;
 
 		default:
@@ -272,13 +280,34 @@ void FWidget::Render_Axis( const FSceneView* View, FPrimitiveDrawInterface* PDI,
 		PDI->SetHitProxy( NULL );
 	}
 
+	FVector2D NewOrigin;
 	FVector2D AxisEnd;
-	if(!View->ScreenToPixel(View->WorldToScreen(ArrowToWorld.TransformPosition(FVector(64,0,0))),AxisEnd))
+	const FVector AxisEndWorld = ArrowToWorld.TransformPosition(FVector(64, 0, 0));
+	const FVector WidgetOrigin = InMatrix.GetOrigin();
+
+	if (View->ScreenToPixel(View->WorldToScreen(WidgetOrigin), NewOrigin) &&
+		View->ScreenToPixel(View->WorldToScreen(AxisEndWorld), AxisEnd))
 	{
-		AxisEnd.X = AxisEnd.Y = 0;
+		// If both the origin and the axis endpoint are in front of the camera, trivially calculate the viewport space axis direction
+		OutAxisDir = (AxisEnd - NewOrigin).GetSafeNormal();
+	}
+	else
+	{
+		// If either the origin or axis endpoint are behind the camera, translate the entire widget in front of the camera in the view direction before performing the
+		// viewport space calculation
+		const FMatrix InvViewMatrix = View->ViewMatrices.GetInvViewMatrix();
+		const FVector ViewLocation = InvViewMatrix.GetOrigin();
+		const FVector ViewDirection = InvViewMatrix.GetUnitAxis(EAxis::Z);
+		const FVector Offset = ViewDirection * (FVector::DotProduct(ViewLocation - WidgetOrigin, ViewDirection) + 100.0f);
+		const FVector AdjustedWidgetOrigin = WidgetOrigin + Offset;
+		const FVector AdjustedWidgetAxisEnd = AxisEndWorld + Offset;
+
+		if (View->ScreenToPixel(View->WorldToScreen(AdjustedWidgetOrigin), NewOrigin) &&
+			View->ScreenToPixel(View->WorldToScreen(AdjustedWidgetAxisEnd), AxisEnd))
+		{
+			OutAxisDir = -(AxisEnd - NewOrigin).GetSafeNormal();
+		}
 	}
-
-	OutAxisDir = (AxisEnd - Origin).GetSafeNormal();
 }
 
 void FWidget::Render_Cube( FPrimitiveDrawInterface* PDI, const FMatrix& InMatrix, const UMaterialInterface* InMaterial, const FVector& InScale )
