Percentage based scaling no longer enabled when multiple actors/components are selected.

diff --git a/Engine/Source/Editor/UnrealEd/Private/EditorServer.cpp b/Engine/Source/Editor/UnrealEd/Private/EditorServer.cpp
--- a/Engine/Source/Editor/UnrealEd/Private/EditorServer.cpp
+++ b/Engine/Source/Editor/UnrealEd/Private/EditorServer.cpp
@@ -545,8 +545,7 @@ void UEditorEngine::LoadAndSelectAssets( TArray<FAssetData>& Assets, UClass* Typ
 
 bool UEditorEngine::UsePercentageBasedScaling() const
 {
-	// Use percentage based scaling if the user setting is enabled or more than component or actor is selected.  Multiplicative scaling doesn't work when more than one object is selected
-	return GetDefault<ULevelEditorViewportSettings>()->UsePercentageBasedScaling() || GetSelectedActorCount() > 1 || GetSelectedComponentCount() > 1;
+	return GetDefault<ULevelEditorViewportSettings>()->UsePercentageBasedScaling();
 }
 
 bool UEditorEngine::Exec_Brush( UWorld* InWorld, const TCHAR* Str, FOutputDevice& Ar )
