Fix import option detection

https://jira.ol.epicgames.net/browse/UE-19358

diff --git a/Engine/Source/Editor/UnrealEd/Private/Fbx/FbxMainImport.cpp b/Engine/Source/Editor/UnrealEd/Private/Fbx/FbxMainImport.cpp
--- a/Engine/Source/Editor/UnrealEd/Private/Fbx/FbxMainImport.cpp
+++ b/Engine/Source/Editor/UnrealEd/Private/Fbx/FbxMainImport.cpp
@@ -585,7 +585,11 @@ bool FFbxImporter::GetSceneInfo(FString Filename, FbxSceneInfo& SceneInfo)
 		}
 		
 		SceneInfo.bHasAnimation = false;
-		for (int32 AnimCurveNodeIndex = 0; AnimCurveNodeIndex < Scene->GetSrcObjectCount<FbxAnimCurveNode>(); AnimCurveNodeIndex++)
+		int32 AnimCurveNodeCount = Scene->GetSrcObjectCount<FbxAnimCurveNode>();
+		// sadly Max export with animation curve node by default without any change, so 
+		// we'll have to skip the first two curves, which is translation/rotation
+		// if there is a valid animation, we'd expect there are more curve nodes than 2. 
+		for (int32 AnimCurveNodeIndex = 2; AnimCurveNodeIndex < AnimCurveNodeCount; AnimCurveNodeIndex++)
 		{
 			FbxAnimCurveNode* CurAnimCruveNode = Scene->GetSrcObject<FbxAnimCurveNode>(AnimCurveNodeIndex);
 			if (CurAnimCruveNode->IsAnimated(true))
