UE-19728 Crash when setting local variable.

diff --git a/Engine/Source/Editor/UnrealEd/Private/Kismet2/BlueprintEditorUtils.cpp b/Engine/Source/Editor/UnrealEd/Private/Kismet2/BlueprintEditorUtils.cpp
--- a/Engine/Source/Editor/UnrealEd/Private/Kismet2/BlueprintEditorUtils.cpp
+++ b/Engine/Source/Editor/UnrealEd/Private/Kismet2/BlueprintEditorUtils.cpp
@@ -4616,25 +4616,28 @@ FBPVariableDescription FBlueprintEditorUtils::DuplicateVariableDescription(UBlue
 
 void FBlueprintEditorUtils::SetDefaultValueOnUserDefinedStructProperty(UBlueprint* InBlueprint, FString InScopeName, FBPVariableDescription* InOutVariableDescription)
 {
-	if (InOutVariableDescription->VarType.PinSubCategoryObject.IsValid() && InOutVariableDescription->VarType.PinSubCategoryObject->IsA(UUserDefinedStruct::StaticClass()))
+	auto UserDefinedStruct = InOutVariableDescription ? Cast<UUserDefinedStruct>(InOutVariableDescription->VarType.PinSubCategoryObject.Get()) : nullptr;
+	if (UserDefinedStruct)
 	{
 		// Create a variable reference for the variable, so we can resolve it and use it to generate the default value
 		FMemberReference VariableReference;
 		VariableReference.SetLocalMember(InOutVariableDescription->VarName, InScopeName, InOutVariableDescription->VarGuid);
-		UStruct* FunctionScope = VariableReference.GetMemberScope(InBlueprint->SkeletonGeneratedClass);
-		UProperty* VariableProperty = VariableReference.ResolveMember<UProperty>(InBlueprint->SkeletonGeneratedClass);
+		auto VariableProperty = VariableReference.ResolveMember<UStructProperty>(InBlueprint->SkeletonGeneratedClass);
 
-		TSharedPtr<FStructOnScope> StructData = MakeShareable(new FStructOnScope(FunctionScope));
+		if (VariableProperty && ensure(VariableProperty->Struct == UserDefinedStruct))
+		{
+			FStructOnScope StructData(UserDefinedStruct);
 
-		// Create the default value for the property
-		FStructureEditorUtils::Fill_MakeStructureDefaultValue(VariableProperty, StructData->GetStructMemory());
+			// Create the default value for the property
+			FStructureEditorUtils::Fill_MakeStructureDefaultValue(UserDefinedStruct, StructData.GetStructMemory());
 
-		// Export the default value as a string so we can store it with the variable description
-		FString DefaultValueString;
-		FBlueprintEditorUtils::PropertyValueToString(VariableProperty, StructData->GetStructMemory(), DefaultValueString);
+			// Export the default value as a string so we can store it with the variable description
+			FString DefaultValueString;
+			FBlueprintEditorUtils::PropertyValueToString_Direct(VariableProperty, StructData.GetStructMemory(), DefaultValueString);
 
-		// Set the default value
-		InOutVariableDescription->DefaultValue = DefaultValueString;
+			// Set the default value
+			InOutVariableDescription->DefaultValue = DefaultValueString;
+		}
 	}
 }
 
@@ -7488,7 +7491,12 @@ bool FBlueprintEditorUtils::PropertyValueFromString(const UProperty* Property, c
 
 bool FBlueprintEditorUtils::PropertyValueToString(const UProperty* Property, const uint8* Container, FString& OutForm)
 {
-	check(Property && Container);
+	return PropertyValueToString_Direct(Property, Property->ContainerPtrToValuePtr<const uint8>(Container), OutForm);
+}
+
+bool FBlueprintEditorUtils::PropertyValueToString_Direct(const UProperty* Property, const uint8* DirectValue, FString& OutForm)
+{
+	check(Property && DirectValue);
 	OutForm.Reset();
 	if (Property->IsA(UStructProperty::StaticClass()))
 	{
@@ -7503,25 +7511,25 @@ bool FBlueprintEditorUtils::PropertyValueToString(const UProperty* Property, con
 		if (StructProperty->Struct == VectorStruct)
 		{
 			FVector Vector;
-			Property->CopyCompleteValue(&Vector, Property->ContainerPtrToValuePtr<uint8>(Container));
+			Property->CopyCompleteValue(&Vector, DirectValue);
 			OutForm = FString::Printf(TEXT("%f,%f,%f"), Vector.X, Vector.Y, Vector.Z);
 		}
 		else if (StructProperty->Struct == RotatorStruct)
 		{
 			FRotator Rotator;
-			Property->CopyCompleteValue(&Rotator, Property->ContainerPtrToValuePtr<uint8>(Container));
+			Property->CopyCompleteValue(&Rotator, DirectValue);
 			OutForm = FString::Printf(TEXT("%f,%f,%f"), Rotator.Pitch, Rotator.Yaw, Rotator.Roll);
 		}
 		else if (StructProperty->Struct == TransformStruct)
 		{
 			FTransform Transform;
-			Property->CopyCompleteValue(&Transform, Property->ContainerPtrToValuePtr<uint8>(Container));
+			Property->CopyCompleteValue(&Transform, DirectValue);
 			OutForm = Transform.ToString();
 		}
 		else if (StructProperty->Struct == LinearColorStruct)
 		{
 			FLinearColor Color;
-			Property->CopyCompleteValue(&Color, Property->ContainerPtrToValuePtr<uint8>(Container));
+			Property->CopyCompleteValue(&Color, DirectValue);
 			OutForm = Color.ToString();
 		}
 	}
@@ -7529,7 +7537,7 @@ bool FBlueprintEditorUtils::PropertyValueToString(const UProperty* Property, con
 	bool bSuccedded = true;
 	if (OutForm.IsEmpty())
 	{
-		bSuccedded = Property->ExportText_InContainer(0, OutForm, Container, Container, nullptr, 0);
+		bSuccedded = Property->ExportText_Direct(OutForm, DirectValue, DirectValue, nullptr, 0);
 	}
 	return bSuccedded;
 }
diff --git a/Engine/Source/Editor/UnrealEd/Public/Kismet2/BlueprintEditorUtils.h b/Engine/Source/Editor/UnrealEd/Public/Kismet2/BlueprintEditorUtils.h
--- a/Engine/Source/Editor/UnrealEd/Public/Kismet2/BlueprintEditorUtils.h
+++ b/Engine/Source/Editor/UnrealEd/Public/Kismet2/BlueprintEditorUtils.h
@@ -997,6 +997,8 @@ public:
 
 	static bool PropertyValueToString(const UProperty* Property, const uint8* Container, FString& OutForm);
 
+	static bool PropertyValueToString_Direct(const UProperty* Property, const uint8* DirectValue, FString& OutForm);
+
 	/** Call PostEditChange() on all Actors based on the given Blueprint */
 	static void PostEditChangeBlueprintActors(UBlueprint* Blueprint, bool bComponentEditChange = false);
 
