Fixed typo.

diff --git a/Engine/Source/Editor/UnrealEd/Classes/Factories/FbxImportUI.h b/Engine/Source/Editor/UnrealEd/Classes/Factories/FbxImportUI.h
--- a/Engine/Source/Editor/UnrealEd/Classes/Factories/FbxImportUI.h
+++ b/Engine/Source/Editor/UnrealEd/Classes/Factories/FbxImportUI.h
@@ -59,7 +59,7 @@ class UFbxImportUI : public UObject
 	UPROPERTY(EditAnywhere, AdvancedDisplay, config, Category=Mesh, meta=(ToolTip="If enabled, combines all meshes into a single mesh", ImportType="StaticMesh"))
 	uint32 bCombineMeshes:1;
 
-	/** Skeleton to use for imported asset. When importing a mesh, leaving this as "None" will create a new skeleton. When importing and animation this MUST be specified to import the asset. */
+	/** Skeleton to use for imported asset. When importing a mesh, leaving this as "None" will create a new skeleton. When importing an animation this MUST be specified to import the asset. */
 	UPROPERTY(EditAnywhere, Category=Mesh, meta=(ImportType="SkeletalMesh|Animation"))
 	class USkeleton* Skeleton;
 
