A possible fix for cooking crash (better printf format).

diff --git a/Engine/Source/Editor/UnrealEd/Private/Commandlets/ChunkManifestGenerator.cpp b/Engine/Source/Editor/UnrealEd/Private/Commandlets/ChunkManifestGenerator.cpp
--- a/Engine/Source/Editor/UnrealEd/Private/Commandlets/ChunkManifestGenerator.cpp
+++ b/Engine/Source/Editor/UnrealEd/Private/Commandlets/ChunkManifestGenerator.cpp
@@ -752,7 +752,7 @@ bool FChunkManifestGenerator::SaveCookedPackageAssetRegistry( const FString& San
 			Json->WriteValue( TEXT("SourcePackageName"), PackageName.ToString() );
 			Json->WriteValue( TEXT("CookedPackageName"), PlatformSandboxPath );
 			Json->WriteValue( TEXT("CookedPackageTimeStamp"), TimeStamp.ToString() );
-			Json->WriteValue( TEXT("FileSize"), FString::Printf(TEXT("%ll"), FileSize) );
+			Json->WriteValue( TEXT("FileSize"), FString::Printf(TEXT("%lld"), FileSize) );
 			
 
 			Json->WriteArrayStart("AssetData");
