Fix slowness when deleting volumes.  It was forcing bsp to rebuild and cause all sorts of hitches due to reregistering actors.

diff --git a/Engine/Source/Editor/UnrealEd/Private/EditorActor.cpp b/Engine/Source/Editor/UnrealEd/Private/EditorActor.cpp
--- a/Engine/Source/Editor/UnrealEd/Private/EditorActor.cpp
+++ b/Engine/Source/Editor/UnrealEd/Private/EditorActor.cpp
@@ -698,9 +698,10 @@ bool UUnrealEdEngine::edactDeleteSelected( UWorld* InWorld, bool bVerifyDeletion
 	// is modify'd only once.
 	TArray<ULevel*> LevelsAlreadyModified;
 	// A list of levels that will need their Bsp updated after the deletion is complete
-	TArray<ULevel*> LevelsToRebuild;
+	TSet<ULevel*> LevelsToRebuildBSP;
+	TSet<ULevel*> LevelsToRebuildNavigation;
+
 
-	bool	bBrushWasDeleted = false;
 	bool	bRequestedDeleteAllByLevel = false;
 	bool	bRequestedDeleteAllByActor = false;
 	EAppMsgType::Type MessageType = ActorsToDelete.Num() > 1 ? EAppMsgType::YesNoYesAllNoAll : EAppMsgType::YesNo;
@@ -784,16 +785,24 @@ bool UUnrealEdEngine::edactDeleteSelected( UWorld* InWorld, bool bVerifyDeletion
 			}
 		}
 	
+		bool bRebuildNavigation = false;
+
 		ABrush* Brush = Cast< ABrush >( Actor );
 		if (Brush && !FActorEditorUtils::IsABuilderBrush(Brush)) // Track whether or not a brush actor was deleted.
 		{
-			bBrushWasDeleted = true;
 			ULevel* BrushLevel = Actor->GetLevel();
-			if (BrushLevel)
+			if (BrushLevel && !Brush->IsVolumeBrush() )
+			{
+				LevelsToRebuildBSP.Add( BrushLevel );
+				// Rebuilding bsp will also take care of navigation
+				LevelsToRebuildNavigation.Remove( BrushLevel );
+			}
+			else if( BrushLevel && !LevelsToRebuildBSP.Contains( BrushLevel ) )
 			{
-				LevelsToRebuild.AddUnique(BrushLevel);
+				LevelsToRebuildNavigation.Add( BrushLevel );
 			}
 		}
+
 		// If the actor about to be deleted is in a group, be sure to remove it from the group
 		AGroupActor* ActorParentGroup = AGroupActor::GetParentForActor(Actor);
 		if(ActorParentGroup)
@@ -831,16 +840,32 @@ bool UUnrealEdEngine::edactDeleteSelected( UWorld* InWorld, bool bVerifyDeletion
 	CollectGarbage( GARBAGE_COLLECTION_KEEPFLAGS );
 
 	NoteSelectionChange();
-	// If any brush actors were deleted, update the Bsp in the appropriate levels
-	if (bBrushWasDeleted)
+	// If any brush actors were modified, update the Bsp in the appropriate levels
+	if (LevelsToRebuildBSP.Num())
 	{
 		FlushRenderingCommands();
 
-		for (int32 LevelIdx = 0; LevelIdx < LevelsToRebuild.Num(); ++LevelIdx)
+		for ( ULevel* Level : LevelsToRebuildBSP )
+		{
+			GEditor->RebuildLevel(*Level);
+		}
+	}
+
+	if( LevelsToRebuildNavigation.Num() )
+	{
+		UWorld* World = GetEditorWorldContext().World();
+		UNavigationSystem* NavSys = UNavigationSystem::GetCurrent(World);
+		if (NavSys)
 		{
-			GEditor->RebuildLevel(*(LevelsToRebuild[LevelIdx]));
+			for ( ULevel* Level : LevelsToRebuildNavigation )
+			{
+				NavSys->UpdateLevelCollision(Level);
+			}
 		}
+	}
 
+	if( LevelsToRebuildBSP.Num() || LevelsToRebuildNavigation.Num() )
+	{
 		RedrawLevelEditingViewports();
 		ULevel::LevelDirtiedEvent.Broadcast();
 	}
