UMG - Adding a way to Slate Application to explictly flush the rendering state.  This should be used whenever PIE sessions end to help flush any UMG / Slate / Slate Renderer caching that may have occured the previous frame that needs to be flushed to allow UTextures or UMaterialInterfaces referenced to be cleaned up properly.  Additionally, Tooltips generated by a callback binding, should correctly be released after the cursor moves off the widget.  Previously they were cached forever because we don't create a new window now for game tooltips.

diff --git a/Engine/Source/Editor/UnrealEd/Private/PlayLevel.cpp b/Engine/Source/Editor/UnrealEd/Private/PlayLevel.cpp
--- a/Engine/Source/Editor/UnrealEd/Private/PlayLevel.cpp
+++ b/Engine/Source/Editor/UnrealEd/Private/PlayLevel.cpp
@@ -251,6 +251,12 @@ void UEditorEngine::EndPlayMap()
 		}
 	}
 
+	// Flush any render commands and released accessed UTextures and materials to give them a chance to be collected.
+	if ( FSlateApplication::IsInitialized() )
+	{
+		FSlateApplication::Get().FlushRenderState();
+	}
+
 	// Clean up any PIE world objects
 	{
 		// The trans buffer should never have a PIE object in it.  If it does though, as a s
@@ -2890,7 +2896,9 @@ UGameInstance* UEditorEngine::CreatePIEGameInstance(int32 PIEInstance, bool bInS
 					FSlateApplication::Get().AddWindowAsNativeChild(PieWindow, MainWindow, true);
 				}
 				else
+				{
 					FSlateApplication::Get().AddWindow(PieWindow);
+				}
 #endif
 
 				TSharedRef<SOverlay> ViewportOverlayWidgetRef = SNew(SOverlay);
@@ -3080,7 +3088,7 @@ FViewport* UEditorEngine::GetActiveViewport()
 		return ActiveLevelViewport->GetActiveViewport();
 	}
 	
-	return NULL;
+	return nullptr;
 }
 
 FViewport* UEditorEngine::GetPIEViewport()
@@ -3100,7 +3108,7 @@ FViewport* UEditorEngine::GetPIEViewport()
 				// We can't use FindChecked here because when using the dedicated server option we don't initialize this map 
 				//	(we don't use a viewport for the PIE context in this case)
 				FSlatePlayInEditorInfo * SlatePlayInEditorSessionPtr = SlatePlayInEditorMap.Find(WorldContext.ContextHandle);
-				if (SlatePlayInEditorSessionPtr != NULL && SlatePlayInEditorSessionPtr->SlatePlayInEditorWindowViewport.IsValid() )
+				if ( SlatePlayInEditorSessionPtr != nullptr && SlatePlayInEditorSessionPtr->SlatePlayInEditorWindowViewport.IsValid() )
 				{
 					return SlatePlayInEditorSessionPtr->SlatePlayInEditorWindowViewport.Get();
 				}
@@ -3108,7 +3116,7 @@ FViewport* UEditorEngine::GetPIEViewport()
 		}
 	}
 
-	return NULL;
+	return nullptr;
 }
 
 void UEditorEngine::ToggleBetweenPIEandSIE( bool bNewSession )
@@ -3117,7 +3125,7 @@ void UEditorEngine::ToggleBetweenPIEandSIE( bool bNewSession )
 
 	// The first PIE world context is the one that can toggle between PIE and SIE
 	// Network PIE/SIE toggling is not really meant to be supported.
-	FSlatePlayInEditorInfo * SlateInfoPtr = NULL;
+	FSlatePlayInEditorInfo * SlateInfoPtr = nullptr;
 	for (auto It = WorldList.CreateIterator(); It && !SlateInfoPtr; ++It)
 	{
 		FWorldContext &WorldContext = *It;
diff --git a/Engine/Source/Runtime/Slate/Private/Framework/Application/SlateApplication.cpp b/Engine/Source/Runtime/Slate/Private/Framework/Application/SlateApplication.cpp
--- a/Engine/Source/Runtime/Slate/Private/Framework/Application/SlateApplication.cpp
+++ b/Engine/Source/Runtime/Slate/Private/Framework/Application/SlateApplication.cpp
@@ -2039,6 +2039,12 @@ void FSlateApplication::UnregisterGameViewport()
 		GameViewportWidget.Pin()->SetActive(false);
 	}
 	GameViewportWidget.Reset();
+}
+
+void FSlateApplication::FlushRenderState()
+{
+	// Flush any render commands because we're about to release shader accesses and not keep them alive.
+	Renderer->FlushCommands();
 
 	// Release any temporary material or texture resources we may have cached and are reporting to prevent
 	// GC on those resources.  If the game viewport is being unregistered, we need to flush these resources
@@ -2873,27 +2879,30 @@ void FSlateApplication::SpawnToolTip( const TSharedRef<IToolTip>& InToolTip, con
 
 void FSlateApplication::CloseToolTip()
 {
-	TSharedPtr< SWindow > PinnedToolTipWindow( ToolTipWindow.Pin() );
-	if( PinnedToolTipWindow.IsValid() && PinnedToolTipWindow->IsVisible() )
+	// Notify the source widget that its tooltip is closing
+	if ( SWidget* SourceWidget = ActiveToolTipWidgetSource.Pin().Get() )
 	{
-		// Notify the source widget that its tooltip is closing
-		if (SWidget* SourceWidget = ActiveToolTipWidgetSource.Pin().Get())
-		{
-			SourceWidget->OnToolTipClosing();
-		}
+		SourceWidget->OnToolTipClosing();
+	}
 
-		TSharedPtr<IToolTip> StableActiveToolTip = ActiveToolTip.Pin();
-		if ( StableActiveToolTip.IsValid() )
-		{
-			StableActiveToolTip->OnClosed();
-		}
+	// Notify the active tooltip that it's being closed.
+	TSharedPtr<IToolTip> StableActiveToolTip = ActiveToolTip.Pin();
+	if ( StableActiveToolTip.IsValid() )
+	{
+		StableActiveToolTip->OnClosed();
+	}
 
+	// If the tooltip had a new window holding it, hide the window.
+	TSharedPtr< SWindow > PinnedToolTipWindow( ToolTipWindow.Pin() );
+	if( PinnedToolTipWindow.IsValid() && PinnedToolTipWindow->IsVisible() )
+	{
 		// Hide the tool tip window.  We don't destroy the window, because we want to reuse it for future tool tips.
 		PinnedToolTipWindow->HideWindow();
-
-		ActiveToolTip.Reset();
-		ActiveToolTipWidgetSource.Reset();
 	}
+
+	ActiveToolTip.Reset();
+	ActiveToolTipWidgetSource.Reset();
+
 	ToolTipOffsetDirection = EToolTipOffsetDirection::Undetermined;
 }
 
diff --git a/Engine/Source/Runtime/Slate/Public/Framework/Application/SlateApplication.h b/Engine/Source/Runtime/Slate/Public/Framework/Application/SlateApplication.h
--- a/Engine/Source/Runtime/Slate/Public/Framework/Application/SlateApplication.h
+++ b/Engine/Source/Runtime/Slate/Public/Framework/Application/SlateApplication.h
@@ -416,6 +416,11 @@ public:
 	 */
 	void UnregisterGameViewport();
 
+	/**
+	 * Flushes the render state of slate, releasing accesses and flushing all render commands.
+	 */
+	void FlushRenderState();
+
 	/**
 	 * Sets specified user focus to the SWidget representing the currently active game viewport
 	 */
diff --git a/Engine/Source/Runtime/UMG/Private/Components/Widget.cpp b/Engine/Source/Runtime/UMG/Private/Components/Widget.cpp
--- a/Engine/Source/Runtime/UMG/Private/Components/Widget.cpp
+++ b/Engine/Source/Runtime/UMG/Private/Components/Widget.cpp
@@ -75,11 +75,14 @@ public:
 
 	virtual void OnClosed() override
 	{
+		//TODO Notify interface implementing widget of closure
+
 		CachedToolTip.Reset();
 	}
 
 	virtual void OnOpening() override
 	{
+		//TODO Notify interface implementing widget of opening
 	}
 
 public:
@@ -792,7 +795,7 @@ void UWidget::SynchronizeProperties()
 		SafeWidget->SetToolTipText(GAME_SAFE_OPTIONAL_BINDING(FText, ToolTipText));
 	}
 
-#if WITH_EDITOR
+#if WITH_EDITORONLY_DATA
 	// In editor builds we add metadata to the widget so that once hit with the widget reflector it can report
 	// where it comes from, what blueprint, what the name of the widget was...etc.
 	SafeWidget->AddMetadata<FReflectionMetaData>(MakeShareable(new FReflectionMetaData(GetFName(), GetClass(), WidgetGeneratedBy)));
