#ENGINE - Change viewport update order to allow matinee viewports to happen first...

--------
Integrated using branch UE4-To-UE4-Orion (reversed) of change#2614987 by [REDACTED_PERSON] on 2015/07/09 09:34:25.

diff --git a/Engine/Source/Editor/UnrealEd/Private/EditorEngine.cpp b/Engine/Source/Editor/UnrealEd/Private/EditorEngine.cpp
--- a/Engine/Source/Editor/UnrealEd/Private/EditorEngine.cpp
+++ b/Engine/Source/Editor/UnrealEd/Private/EditorEngine.cpp
@@ -1394,7 +1394,7 @@ void UEditorEngine::Tick( float DeltaSeconds, bool bIdleMode )
 
 	// Update viewports.
 
-	for(int32 ViewportIndex = 0;ViewportIndex < AllViewportClients.Num();ViewportIndex++)
+	for (int32 ViewportIndex = AllViewportClients.Num()-1; ViewportIndex >= 0; ViewportIndex--)
 	{
 		FEditorViewportClient* ViewportClient = AllViewportClients[ ViewportIndex ];
 
