Fixed issue with material not updating correctly when using iterative cook on the side.

diff --git a/Engine/Source/Editor/UnrealEd/Private/CookOnTheFlyServer.cpp b/Engine/Source/Editor/UnrealEd/Private/CookOnTheFlyServer.cpp
--- a/Engine/Source/Editor/UnrealEd/Private/CookOnTheFlyServer.cpp
+++ b/Engine/Source/Editor/UnrealEd/Private/CookOnTheFlyServer.cpp
@@ -2266,7 +2266,11 @@ void UCookOnTheFlyServer::MarkPackageDirtyForCooker( UPackage *Package )
 				{
 					// if this package was previously cooked and we are doing a cook by the book 
 					// we need to recook this package before finishing cook by the book
-					CookRequests.EnqueueUnique( FFilePlatformRequest(PackageFFileName, CookedPlatforms) );
+					CookRequests.EnqueueUnique(FFilePlatformRequest(PackageFFileName, CookedPlatforms));
+				}
+				else
+				{
+					CookByTheBookOptions->PreviousCookRequests.Add(FFilePlatformRequest(PackageFFileName, CookedPlatforms));
 				}
 			}
 		}
@@ -4543,8 +4547,8 @@ void UCookOnTheFlyServer::MaybeMarkPackageAsAlreadyLoaded(UPackage *Package)
 
 	if (bShouldMarkAsAlreadyProcessed)
 	{
-		Package->PackageFlags |= PKG_ReloadingForCooker;
-	}
+			Package->PackageFlags |= PKG_ReloadingForCooker;
+		}
 
 
 	/*FString Name = Package->GetName();
