Added additional checks to cooker to help track down issue with Packages being saved after being marked as reloading for cooker.

diff --git a/Engine/Source/Editor/UnrealEd/Private/CookOnTheFlyServer.cpp b/Engine/Source/Editor/UnrealEd/Private/CookOnTheFlyServer.cpp
--- a/Engine/Source/Editor/UnrealEd/Private/CookOnTheFlyServer.cpp
+++ b/Engine/Source/Editor/UnrealEd/Private/CookOnTheFlyServer.cpp
@@ -2262,7 +2262,7 @@ void UCookOnTheFlyServer::MarkPackageDirtyForCooker( UPackage *Package )
 			// if we have already cooked this package and we have made changes then recook ;)
 			if ( CookedPackages.GetCookedPlatforms(PackageFFileName, CookedPlatforms) )
 			{
-				if ( IsCookByTheBookRunning() )
+				if (IsCookByTheBookRunning())
 				{
 					// if this package was previously cooked and we are doing a cook by the book 
 					// we need to recook this package before finishing cook by the book
@@ -2438,6 +2438,8 @@ bool UCookOnTheFlyServer::SaveCookedPackage( UPackage* Package, uint32 SaveFlags
 
 	if (Filename.Len())
 	{
+		check(!(Package->PackageFlags & PKG_ReloadingForCooker));
+
 		FString Name = Package->GetPathName();
 
 		// Use SandboxFile to do path conversion to properly handle sandbox paths (outside of standard paths in particular).
@@ -3928,7 +3930,9 @@ void UCookOnTheFlyServer::CreateSandboxFile()
 	// Use SandboxFile to do path conversion to properly handle sandbox paths (outside of standard paths in particular).
 	SandboxFile->Initialize(&FPlatformFileManager::Get().GetPlatformFile(), *FString::Printf(TEXT("-sandbox=\"%s\""), *OutputDirectory));
 }
-
+#if DO_CHECK
+static bool bIsInitializingSandbox = false;
+#endif
 void UCookOnTheFlyServer::InitializeSandbox()
 {
 	if ( SandboxFile == NULL )
@@ -3940,12 +3944,19 @@ void UCookOnTheFlyServer::InitializeSandbox()
 
 		if (!IsChildCooker())
 		{
+#if DO_CHECK
+			bIsInitializingSandbox = true;
+#endif
 			if (IsCookFlagSet(ECookInitializationFlags::Iterative))
 			{
 				PopulateCookedPackagesFromDisk(TargetPlatforms);
 			}
 
 			CleanSandbox(IsCookFlagSet(ECookInitializationFlags::Iterative));
+#if DO_CHECK
+			bIsInitializingSandbox = false;
+#endif
+
 		}
 	}
 }
@@ -4514,7 +4525,9 @@ void UCookOnTheFlyServer::MaybeMarkPackageAsAlreadyLoaded(UPackage *Package)
 	check(IsCookingInEditor()==false);
 	check(IsCookByTheBookMode());
 
-
+#if DO_CHECK
+	check(bIsInitializingSandbox == false);
+#endif
 	FName StandardName = GetCachedStandardPackageFileFName(Package);
 
 	bool bShouldMarkAsAlreadyProcessed = false;
@@ -4547,6 +4560,8 @@ void UCookOnTheFlyServer::MaybeMarkPackageAsAlreadyLoaded(UPackage *Package)
 
 	if (bShouldMarkAsAlreadyProcessed)
 	{
+		if (Package->IsFullyLoaded() == false)
+		{
 			Package->PackageFlags |= PKG_ReloadingForCooker;
 		}
 
