Fixed date time stamp and filesize for cooked asset registry json file.

diff --git a/Engine/Source/Editor/UnrealEd/Private/Commandlets/ChunkManifestGenerator.cpp b/Engine/Source/Editor/UnrealEd/Private/Commandlets/ChunkManifestGenerator.cpp
--- a/Engine/Source/Editor/UnrealEd/Private/Commandlets/ChunkManifestGenerator.cpp
+++ b/Engine/Source/Editor/UnrealEd/Private/Commandlets/ChunkManifestGenerator.cpp
@@ -746,6 +746,8 @@ bool FChunkManifestGenerator::SaveCookedPackageAssetRegistry( const FString& San
 
 			FString PlatformSandboxPath = SandboxPath.Replace(TEXT("[Platform]"), *Platform->PlatformName());
 			
+			FPackageName::FindPackageFileWithoutExtension(PlatformSandboxPath, PlatformSandboxPath);
+
 			FDateTime TimeStamp = IFileManager::Get().GetTimeStamp( *PlatformSandboxPath );
 			int64 FileSize = IFileManager::Get().FileSize(*PlatformSandboxPath);
 
