Remove unused font editor assets

diff --git a/Engine/Source/Editor/FontEditor/Private/SFontEditorViewport.cpp b/Engine/Source/Editor/FontEditor/Private/SFontEditorViewport.cpp
--- a/Engine/Source/Editor/FontEditor/Private/SFontEditorViewport.cpp
+++ b/Engine/Source/Editor/FontEditor/Private/SFontEditorViewport.cpp
@@ -82,7 +82,7 @@ FFontEditorViewportClient::FFontEditorViewportClient(TWeakPtr<SFontEditorViewpor
 void FFontEditorViewportClient::Draw(FViewport* Viewport, FCanvas* Canvas)
 {
 	UFont* Font = FontEditorViewportPtr.Pin()->GetFontEditor().Pin()->GetFont();
-	
+
 	if (!FontEditorViewportPtr.Pin()->IsPreviewViewport())
 	{
 		FVector2D Ratio = FVector2D(GetViewportHorizontalScrollBarRatio(), GetViewportVerticalScrollBarRatio());
@@ -141,7 +141,7 @@ void FFontEditorViewportClient::Draw(FViewport* Viewport, FCanvas* Canvas)
 							Canvas->SetHitProxy(new HObject(Texture));
 						
 							// Draw a simple tile
-							Canvas->DrawTile( XPos, YPos, Width, Height, 0.0f, 0.0f, 1.0f, 1.0f, FLinearColor::White, GEditor->BkgndHi->Resource);
+							Canvas->DrawTile( XPos, YPos, Width, Height, 0.0f, 0.0f, 1.0f, 1.0f, FLinearColor::White, nullptr );
 						
 							Canvas->SetHitProxy(NULL);
 						}
@@ -151,7 +151,7 @@ void FFontEditorViewportClient::Draw(FViewport* Viewport, FCanvas* Canvas)
 							// Draw a selected background
 							if (Texture->IsSelected())
 							{
-								Canvas->DrawTile( XPos, YPos, Width, Height, 0.0f, 0.0f, 1.0f, 1.0f, FLinearColor::White, GEditor->BkgndHi->Resource);
+								Canvas->DrawTile( XPos, YPos, Width, Height, 0.0f, 0.0f, 1.0f, 1.0f, FLinearColor(0.084f, .127f, 0.098f), nullptr );
 							}
 							else
 							{
diff --git a/Engine/Source/Editor/UnrealEd/Classes/Editor/EditorEngine.h b/Engine/Source/Editor/UnrealEd/Classes/Editor/EditorEngine.h
--- a/Engine/Source/Editor/UnrealEd/Classes/Editor/EditorEngine.h
+++ b/Engine/Source/Editor/UnrealEd/Classes/Editor/EditorEngine.h
@@ -304,12 +304,6 @@ public:
 	UPROPERTY()
 	class UTexture2D* Bad;
 
-	UPROPERTY()
-	class UTexture2D* Bkgnd;
-
-	UPROPERTY()
-	class UTexture2D* BkgndHi;
-
 	// Font used by Canvas-based editors
 	UPROPERTY()
 	class UFont* EditorFont;
diff --git a/Engine/Source/Editor/UnrealEd/Private/EditorEngine.cpp b/Engine/Source/Editor/UnrealEd/Private/EditorEngine.cpp
--- a/Engine/Source/Editor/UnrealEd/Private/EditorEngine.cpp
+++ b/Engine/Source/Editor/UnrealEd/Private/EditorEngine.cpp
@@ -213,8 +213,6 @@ UEditorEngine::UEditorEngine(const FObjectInitializer& ObjectInitializer)
 		struct FConstructorStatics
 		{
 			ConstructorHelpers::FObjectFinder<UTexture2D> BadTexture;
-			ConstructorHelpers::FObjectFinder<UTexture2D> BackgroundTexture;
-			ConstructorHelpers::FObjectFinder<UTexture2D> BackgroundHiTexture;
 			ConstructorHelpers::FObjectFinder<UStaticMesh> EditorCubeMesh;
 			ConstructorHelpers::FObjectFinder<UStaticMesh> EditorSphereMesh;
 			ConstructorHelpers::FObjectFinder<UStaticMesh> EditorPlaneMesh;
@@ -222,8 +220,6 @@ UEditorEngine::UEditorEngine(const FObjectInitializer& ObjectInitializer)
 			ConstructorHelpers::FObjectFinder<UFont> SmallFont;
 			FConstructorStatics()
 				: BadTexture(TEXT("/Engine/EditorResources/Bad"))
-				, BackgroundTexture(TEXT("/Engine/EditorResources/Bkgnd"))
-				, BackgroundHiTexture(TEXT("/Engine/EditorResources/BkgndHi"))
 				, EditorCubeMesh(TEXT("/Engine/EditorMeshes/EditorCube"))
 				, EditorSphereMesh(TEXT("/Engine/EditorMeshes/EditorSphere"))
 				, EditorPlaneMesh(TEXT("/Engine/EditorMeshes/EditorPlane"))
@@ -235,8 +231,6 @@ UEditorEngine::UEditorEngine(const FObjectInitializer& ObjectInitializer)
 		static FConstructorStatics ConstructorStatics;
 
 		Bad = ConstructorStatics.BadTexture.Object;
-		Bkgnd = ConstructorStatics.BackgroundTexture.Object;
-		BkgndHi = ConstructorStatics.BackgroundHiTexture.Object;
 		EditorCube = ConstructorStatics.EditorCubeMesh.Object;
 		EditorSphere = ConstructorStatics.EditorSphereMesh.Object;
 		EditorPlane = ConstructorStatics.EditorPlaneMesh.Object;
