All AssetImportData properties are now instanced

This ensures 2 things:
  - That the property is always visible on an relevant UI
  - That we will not suffer any data loss if any new derived UAssetImportData types are added in future

This addresses UE-19895 and UE-19893

diff --git a/Engine/Source/Editor/DetailCustomizations/Private/AssetImportDataCustomization.cpp b/Engine/Source/Editor/DetailCustomizations/Private/AssetImportDataCustomization.cpp
--- a/Engine/Source/Editor/DetailCustomizations/Private/AssetImportDataCustomization.cpp
+++ b/Engine/Source/Editor/DetailCustomizations/Private/AssetImportDataCustomization.cpp
@@ -151,6 +151,7 @@ FReply FAssetImportDataCustomization::OnChangePathClicked(int32 Index) const
 	if (OpenFilenames.Num() == 1)
 	{
 		ImportData->UpdateFilenameOnly(FPaths::ConvertRelativePathToFull(OpenFilenames[0]), Index);
+		ImportData->MarkPackageDirty();
 	}
 	return FReply::Handled();
 }
diff --git a/Engine/Source/Editor/UnrealEd/Private/Editor.cpp b/Engine/Source/Editor/UnrealEd/Private/Editor.cpp
--- a/Engine/Source/Editor/UnrealEd/Private/Editor.cpp
+++ b/Engine/Source/Editor/UnrealEd/Private/Editor.cpp
@@ -111,6 +111,8 @@
 #include "Engine/SimpleConstructionScript.h"
 #include "K2Node_AddComponent.h"
 
+#include "AutoReimport/AutoReimportUtilities.h"
+
 #include "Settings/EditorSettings.h"
 
 #define LOCTEXT_NAMESPACE "UnrealEd.Editor"
@@ -407,6 +409,16 @@ void FReimportManager::GetNewReimportPath(UObject* Obj, TArray<FString>& InOutFi
 
 	FileTypes = FString::Printf(TEXT("All Files (%s)|%s|%s"),*AllExtensions,*AllExtensions,*FileTypes);
 
+	FString DefaultFolder;
+	FString DefaultFile;
+	
+	TArray<FString> ExistingPaths = Utils::ExtractSourceFilePaths(Obj);
+	if (ExistingPaths.Num() > 0)
+	{
+		DefaultFolder = FPaths::GetPath(ExistingPaths[0]);
+		DefaultFile = FPaths::GetCleanFilename(ExistingPaths[0]);
+	}
+
 	// Prompt the user for the filenames
 	TArray<FString> OpenFilenames;
 	IDesktopPlatform* DesktopPlatform = FDesktopPlatformModule::Get();
@@ -426,8 +438,8 @@ void FReimportManager::GetNewReimportPath(UObject* Obj, TArray<FString>& InOutFi
 		bOpened = DesktopPlatform->OpenFileDialog(
 			ParentWindowWindowHandle,
 			Title,
-			TEXT(""),
-			TEXT(""),
+			*DefaultFolder,
+			*DefaultFile,
 			FileTypes,
 			bAllowMultiSelect ? EFileDialogFlags::Multiple : EFileDialogFlags::None,
 			OpenFilenames
diff --git a/Engine/Source/Runtime/Engine/Classes/Curves/CurveBase.h b/Engine/Source/Runtime/Engine/Classes/Curves/CurveBase.h
--- a/Engine/Source/Runtime/Engine/Classes/Curves/CurveBase.h
+++ b/Engine/Source/Runtime/Engine/Classes/Curves/CurveBase.h
@@ -545,9 +545,10 @@ public:
 #if WITH_EDITORONLY_DATA
 	/** Override to ensure we write out the asset import data */
 	virtual void GetAssetRegistryTags(TArray<FAssetRegistryTag>& OutTags) const override;
+	virtual void PostInitProperties() override;
 	virtual void PostLoad() override;
 
-	UPROPERTY()
+	UPROPERTY(VisibleAnywhere, Instanced, Category=ImportSettings)
 	class UAssetImportData* AssetImportData;
 
 	/** The filename imported to create this object. Relative to this object's package, BaseDir() or absolute */
diff --git a/Engine/Source/Runtime/Engine/Classes/Engine/CurveTable.h b/Engine/Source/Runtime/Engine/Classes/Engine/CurveTable.h
--- a/Engine/Source/Runtime/Engine/Classes/Engine/CurveTable.h
+++ b/Engine/Source/Runtime/Engine/Classes/Engine/CurveTable.h
@@ -33,9 +33,10 @@ class UCurveTable
 
 #if WITH_EDITORONLY_DATA
 	virtual void GetAssetRegistryTags(TArray<FAssetRegistryTag>& OutTags) const override;
+	virtual void PostInitProperties() override;
 	virtual void PostLoad() override;
 
-	UPROPERTY()
+	UPROPERTY(VisibleAnywhere, Instanced, Category=ImportSettings)
 	class UAssetImportData* AssetImportData;
 
 	/** The filename imported to create this object. Relative to this object's package, BaseDir() or absolute */
diff --git a/Engine/Source/Runtime/Engine/Classes/Engine/DataTable.h b/Engine/Source/Runtime/Engine/Classes/Engine/DataTable.h
--- a/Engine/Source/Runtime/Engine/Classes/Engine/DataTable.h
+++ b/Engine/Source/Runtime/Engine/Classes/Engine/DataTable.h
@@ -47,9 +47,10 @@ class UDataTable
 	ENGINE_API static void AddReferencedObjects(UObject* InThis, FReferenceCollector& Collector);
 #if WITH_EDITORONLY_DATA
 	ENGINE_API virtual void GetAssetRegistryTags(TArray<FAssetRegistryTag>& OutTags) const override;
+	ENGINE_API virtual void PostInitProperties() override;
 	ENGINE_API virtual void PostLoad() override;
 
-	UPROPERTY()
+	UPROPERTY(VisibleAnywhere, Instanced, Category=ImportSettings)
 	class UAssetImportData* AssetImportData;
 
 	/** The filename imported to create this object. Relative to this object's package, BaseDir() or absolute */
diff --git a/Engine/Source/Runtime/Engine/Classes/Engine/Texture.h b/Engine/Source/Runtime/Engine/Classes/Engine/Texture.h
--- a/Engine/Source/Runtime/Engine/Classes/Engine/Texture.h
+++ b/Engine/Source/Runtime/Engine/Classes/Engine/Texture.h
@@ -361,7 +361,7 @@ public:
 	UPROPERTY()
 	FString SourceFilePath_DEPRECATED;
 
-	UPROPERTY()
+	UPROPERTY(VisibleAnywhere, Instanced, Category=ImportSettings)
 	class UAssetImportData* AssetImportData;
 
 public:
@@ -687,6 +687,7 @@ public:
 	ENGINE_API virtual void PostEditChangeProperty(FPropertyChangedEvent& PropertyChangedEvent) override;
 #endif // WITH_EDITOR
 	ENGINE_API virtual void Serialize(FArchive& Ar) override;
+	ENGINE_API virtual void PostInitProperties() override;
 	ENGINE_API virtual void PostLoad() override;
 	ENGINE_API virtual void PreSave() override;
 	ENGINE_API virtual void BeginDestroy() override;
diff --git a/Engine/Source/Runtime/Engine/Classes/Sound/SoundWave.h b/Engine/Source/Runtime/Engine/Classes/Sound/SoundWave.h
--- a/Engine/Source/Runtime/Engine/Classes/Sound/SoundWave.h
+++ b/Engine/Source/Runtime/Engine/Classes/Sound/SoundWave.h
@@ -219,7 +219,7 @@ class ENGINE_API USoundWave : public USoundBase
 	UPROPERTY()
 	FString SourceFileTimestamp_DEPRECATED;
 
-	UPROPERTY()
+	UPROPERTY(VisibleAnywhere, Instanced, Category=ImportSettings)
 	class UAssetImportData* AssetImportData;
 	
 #endif // WITH_EDITORONLY_DATA
diff --git a/Engine/Source/Runtime/Engine/Classes/VectorField/VectorFieldStatic.h b/Engine/Source/Runtime/Engine/Classes/VectorField/VectorFieldStatic.h
--- a/Engine/Source/Runtime/Engine/Classes/VectorField/VectorFieldStatic.h
+++ b/Engine/Source/Runtime/Engine/Classes/VectorField/VectorFieldStatic.h
@@ -37,7 +37,7 @@ public:
 	UPROPERTY()
 	FString SourceFilePath_DEPRECATED;
 
-	UPROPERTY()
+	UPROPERTY(VisibleAnywhere, Instanced, Category=ImportSettings)
 	class UAssetImportData* AssetImportData;
 #endif // WITH_EDITORONLY_DATA
 
@@ -50,6 +50,7 @@ public:
 	virtual void Serialize(FArchive& Ar) override;
 #if WITH_EDITORONLY_DATA
 	virtual void GetAssetRegistryTags(TArray<FAssetRegistryTag>& OutTags) const override;
+	virtual void PostInitProperties() override;
 #endif
 	// End UObject interface.
 
diff --git a/Engine/Source/Runtime/Engine/Private/CurveBase.cpp b/Engine/Source/Runtime/Engine/Private/CurveBase.cpp
--- a/Engine/Source/Runtime/Engine/Private/CurveBase.cpp
+++ b/Engine/Source/Runtime/Engine/Private/CurveBase.cpp
@@ -1087,9 +1087,6 @@ bool FRichCurve::operator==(const FRichCurve& Curve) const
 UCurveBase::UCurveBase(const FObjectInitializer& ObjectInitializer)
 	: Super(ObjectInitializer)
 {
-#if WITH_EDITORONLY_DATA
-	AssetImportData = CreateEditorOnlyDefaultSubobject<UAssetImportData>(TEXT("AssetImportData"));
-#endif
 }
 
 #if WITH_EDITORONLY_DATA
@@ -1103,6 +1100,16 @@ void UCurveBase::GetAssetRegistryTags(TArray<FAssetRegistryTag>& OutTags) const
 	Super::GetAssetRegistryTags(OutTags);
 }
 
+void UCurveBase::PostInitProperties()
+{
+	if (!HasAnyFlags(RF_ClassDefaultObject))
+	{
+		AssetImportData = NewObject<UAssetImportData>(this, TEXT("AssetImportData"));
+	}
+
+	Super::PostInitProperties();
+}
+
 void UCurveBase::PostLoad()
 {
 	Super::PostLoad();
diff --git a/Engine/Source/Runtime/Engine/Private/CurveTable.cpp b/Engine/Source/Runtime/Engine/Private/CurveTable.cpp
--- a/Engine/Source/Runtime/Engine/Private/CurveTable.cpp
+++ b/Engine/Source/Runtime/Engine/Private/CurveTable.cpp
@@ -18,9 +18,6 @@ DECLARE_CYCLE_STAT(TEXT("CurveTableRowHandle Eval"),STAT_CurveTableRowHandleEval
 UCurveTable::UCurveTable(const FObjectInitializer& ObjectInitializer)
 	: Super(ObjectInitializer)
 {
-#if WITH_EDITORONLY_DATA
-	AssetImportData = CreateEditorOnlyDefaultSubobject<UAssetImportData>(TEXT("AssetImportData"));
-#endif
 }
 
 /** Util that removes invalid chars and then make an FName */
@@ -107,6 +104,16 @@ void UCurveTable::GetAssetRegistryTags(TArray<FAssetRegistryTag>& OutTags) const
 	Super::GetAssetRegistryTags(OutTags);
 }
 
+void UCurveTable::PostInitProperties()
+{
+	if (!HasAnyFlags(RF_ClassDefaultObject))
+	{
+		AssetImportData = NewObject<UAssetImportData>(this, TEXT("AssetImportData"));
+	}
+
+	Super::PostInitProperties();
+}
+
 void UCurveTable::PostLoad()
 {
 	Super::PostLoad();
diff --git a/Engine/Source/Runtime/Engine/Private/DataTable.cpp b/Engine/Source/Runtime/Engine/Private/DataTable.cpp
--- a/Engine/Source/Runtime/Engine/Private/DataTable.cpp
+++ b/Engine/Source/Runtime/Engine/Private/DataTable.cpp
@@ -41,8 +41,6 @@ UDataTable::UDataTable(const FObjectInitializer& ObjectInitializer)
 			UPackage::GetTypeSpecificLocalizationDataGatheringCallbacks().Add(UDataTable::StaticClass(), &GatherDataTableForLocalization);
 		}
 	} AutomaticRegistrationOfLocalizationGatherer;
-
-	AssetImportData = CreateEditorOnlyDefaultSubobject<UAssetImportData>(TEXT("AssetImportData"));
 #endif
 }
 
@@ -169,6 +167,16 @@ void UDataTable::GetAssetRegistryTags(TArray<FAssetRegistryTag>& OutTags) const
 	Super::GetAssetRegistryTags(OutTags);
 }
 
+void UDataTable::PostInitProperties()
+{
+	if (!HasAnyFlags(RF_ClassDefaultObject))
+	{
+		AssetImportData = NewObject<UAssetImportData>(this, TEXT("AssetImportData"));
+	}
+
+	Super::PostInitProperties();
+}
+
 void UDataTable::PostLoad()
 {
 	Super::PostLoad();
diff --git a/Engine/Source/Runtime/Engine/Private/SoundWave.cpp b/Engine/Source/Runtime/Engine/Private/SoundWave.cpp
--- a/Engine/Source/Runtime/Engine/Private/SoundWave.cpp
+++ b/Engine/Source/Runtime/Engine/Private/SoundWave.cpp
@@ -61,10 +61,6 @@ USoundWave::USoundWave(const FObjectInitializer& ObjectInitializer)
 	Volume = 1.0;
 	Pitch = 1.0;
 	CompressionQuality = 40;
-
-#if WITH_EDITORONLY_DATA
-	AssetImportData = CreateEditorOnlyDefaultSubobject<UAssetImportData>(TEXT("AssetImportData"));
-#endif
 }
 
 SIZE_T USoundWave::GetResourceSize(EResourceSizeMode::Type Mode)
@@ -278,6 +274,13 @@ void USoundWave::PostInitProperties()
 	{
 		InvalidateCompressedData();
 	}
+
+#if WITH_EDITORONLY_DATA
+	if (!HasAnyFlags(RF_ClassDefaultObject))
+	{
+		AssetImportData = NewObject<UAssetImportData>(this, TEXT("AssetImportData"));
+	}
+#endif
 }
 
 FByteBulkData* USoundWave::GetCompressedData(FName Format)
diff --git a/Engine/Source/Runtime/Engine/Private/Texture.cpp b/Engine/Source/Runtime/Engine/Private/Texture.cpp
--- a/Engine/Source/Runtime/Engine/Private/Texture.cpp
+++ b/Engine/Source/Runtime/Engine/Private/Texture.cpp
@@ -60,8 +60,6 @@ UTexture::UTexture(const FObjectInitializer& ObjectInitializer)
 	PaddingColor = FColor::Black;
 	ChromaKeyColor = FColorList::Magenta;
 	ChromaKeyThreshold = 1.0f / 255.0f;
-
-	AssetImportData = CreateEditorOnlyDefaultSubobject<UAssetImportData>(TEXT("AssetImportData"));
 	
 #endif // #if WITH_EDITORONLY_DATA
 
@@ -255,6 +253,17 @@ void UTexture::Serialize(FArchive& Ar)
 #endif // #if WITH_EDITORONLY_DATA
 }
 
+void UTexture::PostInitProperties()
+{
+#if WITH_EDITORONLY_DATA
+	if (!HasAnyFlags(RF_ClassDefaultObject))
+	{
+		AssetImportData = NewObject<UAssetImportData>(this, TEXT("AssetImportData"));
+	}
+#endif
+	Super::PostInitProperties();
+}
+
 void UTexture::PostLoad()
 {
 	Super::PostLoad();
diff --git a/Engine/Source/Runtime/Engine/Private/VectorField.cpp b/Engine/Source/Runtime/Engine/Private/VectorField.cpp
--- a/Engine/Source/Runtime/Engine/Private/VectorField.cpp
+++ b/Engine/Source/Runtime/Engine/Private/VectorField.cpp
@@ -261,9 +261,6 @@ private:
 UVectorFieldStatic::UVectorFieldStatic(const FObjectInitializer& ObjectInitializer)
 	: Super(ObjectInitializer)
 {
-#if WITH_EDITORONLY_DATA
-	AssetImportData = CreateEditorOnlyDefaultSubobject<UAssetImportData>(TEXT("AssetImportData"));
-#endif
 }
 
 void UVectorFieldStatic::InitInstance(FVectorFieldInstance* Instance, bool bPreviewInstance)
@@ -359,6 +356,16 @@ void UVectorFieldStatic::GetAssetRegistryTags(TArray<FAssetRegistryTag>& OutTags
 
 	Super::GetAssetRegistryTags(OutTags);
 }
+
+void UVectorFieldStatic::PostInitProperties()
+{
+	if (!HasAnyFlags(RF_ClassDefaultObject))
+	{
+		AssetImportData = NewObject<UAssetImportData>(this, TEXT("AssetImportData"));
+	}
+
+	Super::PostInitProperties();
+}
 #endif
 
 class FVectorFieldCollectorResources : public FOneFrameResource
