Reinstated control to move the pivot by Middle-dragging the centre point in perspective mode.  This allows the vertex snapping mode (holding V+MMB drag) to work as before.

diff --git a/Engine/Source/Editor/UnrealEd/Private/MouseDeltaTracker.cpp b/Engine/Source/Editor/UnrealEd/Private/MouseDeltaTracker.cpp
--- a/Engine/Source/Editor/UnrealEd/Private/MouseDeltaTracker.cpp
+++ b/Engine/Source/Editor/UnrealEd/Private/MouseDeltaTracker.cpp
@@ -416,11 +416,11 @@ void FMouseDeltaTracker::AddDelta(FEditorViewportClient* InViewportClient, FKey
 	// We allow an exception for dragging with the left and/or right mouse button while holding control
 	// as that simulates moving objects with the gizmo
 	//
-	// We also allow the exception of the middle mouse button when Alt is pressed as it 
-	// allows movement of only the pivot.
+	// We also allow the exception of the middle mouse button when Alt is pressed, or when the current axis is the pivot centre, as it 
+	// allows movement of only the pivot
 	const bool bIsOrthoObjectRotation = bControlDown && InViewportClient->IsOrtho();
 	const bool bUsingDragTool = UsingDragTool();
-	const bool bUsingAxis = !bUsingDragTool && (LeftMouseButtonDown || (bAltDown && MiddleMouseButtonDown) || ((bIsOrthoObjectRotation || bControlDown) && RightMouseButtonDown));
+	const bool bUsingAxis = !bUsingDragTool && (LeftMouseButtonDown || (bAltDown && MiddleMouseButtonDown) || (SaveAxis == EAxisList::Screen && MiddleMouseButtonDown) || ((bIsOrthoObjectRotation || bControlDown) && RightMouseButtonDown));
 
 	ConditionalBeginUsingDragTool( InViewportClient );
 
