ClearSourceFiles exec command now operates on all asset types

Added 'Path' parameter to specify a folder path to operate within (eg CleanSourceFiles Path=/Game), in addition to the 'Find' parameter which removes source files only if they contain any of the specified comma-separated values.

diff --git a/Engine/Source/Editor/UnrealEd/Private/UnrealEdSrv.cpp b/Engine/Source/Editor/UnrealEd/Private/UnrealEdSrv.cpp
--- a/Engine/Source/Editor/UnrealEd/Private/UnrealEdSrv.cpp
+++ b/Engine/Source/Editor/UnrealEd/Private/UnrealEdSrv.cpp
@@ -50,6 +50,7 @@
 #include "Engine/StaticMeshActor.h"
 #include "EngineUtils.h"
 #include "LevelEditor.h"
+#include "AutoReimport/AssetSourceFilenameCache.h"
 #include "ComponentEditorUtils.h"
 
 
@@ -1101,102 +1102,101 @@ bool UUnrealEdEngine::Exec( UWorld* InWorld, const TCHAR* Stream, FOutputDevice&
 	{
 		struct Local
 		{
-			static bool RemoveSourcePath( UAssetImportData* Data, const TArray<FString>& SearchTerms )
+			static bool RemoveSourcePath( const FAssetImportInfo& ImportInfo, const FAssetData& AssetData, const TArray<FString>* SearchTerms )
 			{
 				FAssetImportInfo AssetImportInfo;
 
 				bool bModified = false;
-				for (const auto& File : Data->SourceData.SourceFiles)
+				for (const auto& File : ImportInfo.SourceFiles)
 				{
-					if( !File.RelativeFilename.IsEmpty() && !SearchTerms.ContainsByPredicate([&](const FString& SearchTerm){ return File.RelativeFilename.Contains(SearchTerm); }) )
+					const bool bRemoveFile = File.RelativeFilename.IsEmpty() || !SearchTerms ||
+						SearchTerms->ContainsByPredicate([&](const FString& SearchTerm){ return File.RelativeFilename.Contains(SearchTerm); });
+
+					if( bRemoveFile )
 					{
-						AssetImportInfo.Insert(File);
+						UE_LOG(LogUnrealEdSrv, Log, TEXT("Removing Path: %s"), *File.RelativeFilename);
+						bModified = true;
 					}
 					else
 					{
-						UE_LOG(LogUnrealEdSrv, Log, TEXT("Removing Path: %s"), *File.RelativeFilename);
-						bModified = true;
+						AssetImportInfo.Insert(File);
 					}
 				}
 
 				if (bModified)
 				{
-					Data->Modify();
-					Data->SourceData = AssetImportInfo;
-					return true;
+					if (UObject* Asset = AssetData.GetAsset())
+					{
+						UAssetImportData* ImportData = nullptr;
+
+						// Root out the asset import data property
+						for (UObjectProperty* Property : TFieldRange<UObjectProperty>(Asset->GetClass()))
+						{
+							ImportData = Cast<UAssetImportData>(Property->GetObjectPropertyValue(Property->ContainerPtrToValuePtr<UObject*>(Asset)));
+							if (ImportData)
+							{
+								Asset->Modify();
+								ImportData->SourceData = AssetImportInfo;
+								return true;
+							}
+						}
+					}
 				}
 
 				return false;
 			}
-		};
 
-		FString SearchTermStr;
-		if (FParse::Value(Str, TEXT("Find="), SearchTermStr, false))
-		{
-			TArray<FString> SearchTerms;
-			SearchTermStr.ParseIntoArray( SearchTerms, TEXT(","), true );
-
-			TArray<UObject*> ModifiedObjects;
-			if( SearchTerms.Num() )
+			static void RemoveSourcePaths( const TArray<FAssetData>& AllAssets, const TArray<FString>* SearchTerms )
 			{
-				FAssetRegistryModule& AssetRegistryModule = FModuleManager::LoadModuleChecked<FAssetRegistryModule>(TEXT("AssetRegistry"));
-
-				TArray<FAssetData> StaticMeshes;
-				TArray<FAssetData> SkeletalMeshes;
-				TArray<FAssetData> AnimSequences;
-				TArray<FAssetData> DestructibleMeshes;
+				FScopedSlowTask SlowTask(AllAssets.Num(), NSLOCTEXT("UnrealEd", "ClearingSourceFiles", "Clearing Source Files"));
+				SlowTask.MakeDialog(true);
 
-				GWarn->BeginSlowTask(NSLOCTEXT("UnrealEd", "ClearingSourceFiles", "Clearing Source Files"), true, true);
-				AssetRegistryModule.Get().GetAssetsByClass(UStaticMesh::StaticClass()->GetFName(), StaticMeshes);
-
-				AssetRegistryModule.Get().GetAssetsByClass(USkeletalMesh::StaticClass()->GetFName(), SkeletalMeshes);
-
-				AssetRegistryModule.Get().GetAssetsByClass(UAnimSequence::StaticClass()->GetFName(), AnimSequences);
-
-				AssetRegistryModule.Get().GetAssetsByClass(UDestructibleMesh::StaticClass()->GetFName(), DestructibleMeshes);
-
-				for (const FAssetData& StaticMesh : StaticMeshes)
+				for (const FAssetData& Asset : AllAssets)
 				{
-					UStaticMesh* Mesh = Cast<UStaticMesh>(StaticMesh.GetAsset());
+					SlowTask.EnterProgressFrame();
 
-					if (Mesh && Mesh->AssetImportData && Local::RemoveSourcePath( Mesh->AssetImportData, SearchTerms ) )
+					// Optimization - check the asset has import information before loading it
+					TOptional<FAssetImportInfo> ImportInfo = FAssetSourceFilenameCache::ExtractAssetImportInfo(Asset.TagsAndValues);
+					if (ImportInfo.IsSet() && ImportInfo->SourceFiles.Num())
 					{
-						ModifiedObjects.Add( Mesh );
+						RemoveSourcePath(ImportInfo.GetValue(), Asset, SearchTerms);
 					}
 				}
+			}
+		};
 
-				for (const FAssetData& SkelMesh : SkeletalMeshes)
-				{
-					USkeletalMesh* Mesh = Cast<USkeletalMesh>(SkelMesh.GetAsset());
-
-					if (Mesh && Mesh->AssetImportData && Local::RemoveSourcePath(Mesh->AssetImportData, SearchTerms))
-					{
-						ModifiedObjects.Add(Mesh);
-					}
-				}
+		FAssetRegistryModule& AssetRegistryModule = FModuleManager::LoadModuleChecked<FAssetRegistryModule>(TEXT("AssetRegistry"));
 
-				for (const FAssetData& AnimSequence : AnimSequences)
-				{
-					UAnimSequence* Sequence = Cast<UAnimSequence>(AnimSequence.GetAsset());
+		FString Path;
+		FParse::Value(Str, TEXT("Path="), Path, false);
 
-					if (Sequence && Sequence->AssetImportData && Local::RemoveSourcePath( Sequence->AssetImportData, SearchTerms ) )
-					{
-						ModifiedObjects.Add(Sequence);
-					}
-				}
+		TArray<FAssetData> AllAssets;
+		if (!Path.IsEmpty())
+		{
+			AssetRegistryModule.Get().GetAssetsByPath(*Path, AllAssets, true);
+		}
+		else
+		{
+			AssetRegistryModule.Get().GetAllAssets(AllAssets);
+		}
 
-				for (const FAssetData& DestMesh : DestructibleMeshes)
-				{
-					UDestructibleMesh* Mesh = Cast<UDestructibleMesh>(DestMesh.GetAsset());
+		FString SearchTermStr;
+		if (FParse::Value(Str, TEXT("Find="), SearchTermStr, false))
+		{
+			// Searching for particular paths to remove
+			TArray<FString> SearchTerms;
+			SearchTermStr.ParseIntoArray( SearchTerms, TEXT(","), true );
 
-					if (Mesh && Mesh->AssetImportData && Local::RemoveSourcePath(Mesh->AssetImportData, SearchTerms))
-					{
-						ModifiedObjects.Add(Mesh);
-					}
-				}
+			TArray<UObject*> ModifiedObjects;
+			if( SearchTerms.Num() )
+			{
+				Local::RemoveSourcePaths(AllAssets, &SearchTerms);
 			}
-
-			GWarn->EndSlowTask();
+		}
+		else
+		{
+			// Remove every source path on any asset
+			Local::RemoveSourcePaths(AllAssets, nullptr);
 		}
 	}
 	else if (FParse::Command(&Str, TEXT("RenameAssets")))
