Stopped markup of packages for reload in sandbox initialization.

diff --git a/Engine/Source/Editor/UnrealEd/Private/CookOnTheFlyServer.cpp b/Engine/Source/Editor/UnrealEd/Private/CookOnTheFlyServer.cpp
--- a/Engine/Source/Editor/UnrealEd/Private/CookOnTheFlyServer.cpp
+++ b/Engine/Source/Editor/UnrealEd/Private/CookOnTheFlyServer.cpp
@@ -824,7 +824,8 @@ UCookOnTheFlyServer::UCookOnTheFlyServer(const FObjectInitializer& ObjectInitial
 	CurrentCookMode(ECookMode::CookOnTheFly),
 	CookByTheBookOptions(NULL),
 	CookFlags(ECookInitializationFlags::None),
-	bIsSavingPackage( false )
+	bIsSavingPackage( false ),
+	bIsInitializingSandbox(false)
 {
 }
 
@@ -2819,7 +2820,7 @@ bool UCookOnTheFlyServer::GetCurrentIniVersionStrings( const ITargetPlatform* Ta
 
 	for (const auto& RValue : VersionedRValues)
 	{
-		static const auto CVar = IConsoleManager::Get().FindTConsoleVariableDataInt(TEXT("r.CompileShadersForDevelopment"));
+		static const auto CVar = IConsoleManager::Get().FindTConsoleVariableDataInt(*RValue);
 		if (CVar)
 		{
 			FString VersionedRValueString = FString::Printf(TEXT("%s:%d"), *RValue, CVar->GetValueOnGameThread());
@@ -3945,9 +3946,7 @@ void UCookOnTheFlyServer::CreateSandboxFile()
 	// Use SandboxFile to do path conversion to properly handle sandbox paths (outside of standard paths in particular).
 	SandboxFile->Initialize(&FPlatformFileManager::Get().GetPlatformFile(), *FString::Printf(TEXT("-sandbox=\"%s\""), *OutputDirectory));
 }
-#if DO_CHECK
-static bool bIsInitializingSandbox = false;
-#endif
+
 void UCookOnTheFlyServer::InitializeSandbox()
 {
 	if ( SandboxFile == NULL )
@@ -3959,18 +3958,14 @@ void UCookOnTheFlyServer::InitializeSandbox()
 
 		if (!IsChildCooker())
 		{
-#if DO_CHECK
 			bIsInitializingSandbox = true;
-#endif
 			if (IsCookFlagSet(ECookInitializationFlags::Iterative))
 			{
 				PopulateCookedPackagesFromDisk(TargetPlatforms);
 			}
 
 			CleanSandbox(IsCookFlagSet(ECookInitializationFlags::Iterative));
-#if DO_CHECK
 			bIsInitializingSandbox = false;
-#endif
 
 		}
 	}
@@ -4540,9 +4535,11 @@ void UCookOnTheFlyServer::MaybeMarkPackageAsAlreadyLoaded(UPackage *Package)
 	check(IsCookingInEditor()==false);
 	check(IsCookByTheBookMode());
 
-#if DO_CHECK
-	check(bIsInitializingSandbox == false);
-#endif
+	if (bIsInitializingSandbox)
+	{
+		return;
+	}
+
 	FName StandardName = GetCachedStandardPackageFileFName(Package);
 
 	bool bShouldMarkAsAlreadyProcessed = false;
