Missed file in last checkin.

diff --git a/Engine/Source/Editor/UnrealEd/Classes/CookOnTheSide/CookOnTheFlyServer.h b/Engine/Source/Editor/UnrealEd/Classes/CookOnTheSide/CookOnTheFlyServer.h
--- a/Engine/Source/Editor/UnrealEd/Classes/CookOnTheSide/CookOnTheFlyServer.h
+++ b/Engine/Source/Editor/UnrealEd/Classes/CookOnTheSide/CookOnTheFlyServer.h
@@ -721,8 +721,10 @@ private:
 	/** Max memory the cooker should use before forcing a gc */
 	uint64 MaxMemoryAllowance;
 
+	
 	ECookInitializationFlags CookFlags;
 	TAutoPtr<class FSandboxPlatformFile> SandboxFile;
+	bool bIsInitializingSandbox; // stop recursion into callbacks when we are initializing sandbox
 	bool bIsSavingPackage; // used to stop recursive mark package dirty functions
 	TSet<FName> PackagesKeptFromPreviousCook; // used for iterative cooking this is a list of the packages which were kept from the previous cook
 
