Curve Editor: Fix for adding keys. A key is created at the curve evaluation point if the mouse is hovered on the curve. Otherwise, the key is created at the mouse point.

diff --git a/Engine/Source/Editor/UnrealEd/Private/SCurveEditor.cpp b/Engine/Source/Editor/UnrealEd/Private/SCurveEditor.cpp
--- a/Engine/Source/Editor/UnrealEd/Private/SCurveEditor.cpp
+++ b/Engine/Source/Editor/UnrealEd/Private/SCurveEditor.cpp
@@ -1847,19 +1847,19 @@ void SCurveEditor::ProcessClick(const FGeometry& InMyGeometry, const FPointerEve
 				if (HoveredCurve.IsValid())
 				{
 					CurvesToAddKeysTo->Add(HoveredCurve);
-					bAddKeysInline = false;
+					bAddKeysInline = true;
 				}
 				else
 				{
 					if (CurveViewModels.Num() == 1)
 					{
 						CurvesToAddKeysTo->Add(CurveViewModels[0]);
-						bAddKeysInline = true;
+						bAddKeysInline = false;
 					}
 					else
 					{
 						CurvesToAddKeysTo->Append(CurveViewModels);
-						bAddKeysInline = false;
+						bAddKeysInline = true;
 					}
 				}
 				AddNewKey(InMyGeometry, InMouseEvent.GetScreenSpacePosition(), CurvesToAddKeysTo, bAddKeysInline);
@@ -2500,7 +2500,7 @@ void SCurveEditor::CreateContextMenu(const FGeometry& InMyGeometry, const FPoint
 			MenuItemLabel = FText::Format(AddKeyToCurveLabelFormat, FText::FromName(HoveredCurve->CurveInfo.CurveName));
 			MenuItemToolTip = FText::Format(AddKeyToCurveToolTipFormat, FText::FromName(HoveredCurve->CurveInfo.CurveName));
 			CurvesToAddKeysTo->Add(HoveredCurve);
-			bAddKeysInline = false;
+			bAddKeysInline = true;
 		}
 		else
 		{
@@ -2509,14 +2509,14 @@ void SCurveEditor::CreateContextMenu(const FGeometry& InMyGeometry, const FPoint
 				MenuItemLabel = FText::Format(AddKeyToCurveLabelFormat, FText::FromName(CurveViewModels[0]->CurveInfo.CurveName));
 				MenuItemToolTip = FText::Format(AddKeyToCurveToolTipFormat, FText::FromName(CurveViewModels[0]->CurveInfo.CurveName));
 				CurvesToAddKeysTo->Add(CurveViewModels[0]);
-				bAddKeysInline = true;
+				bAddKeysInline = false;
 			}
 			else
 			{
 				MenuItemLabel = LOCTEXT("AddKeyToAllCurves", "Add key to all curves");
 				MenuItemToolTip = LOCTEXT("AddKeyToAllCurveToolTip", "Adds a key at the hovered time to all curves.  Keys can also be added with Shift + Click.");
 				CurvesToAddKeysTo->Append(CurveViewModels);
-				bAddKeysInline = false;
+				bAddKeysInline = true;
 			}
 		}
 
