Fixed import factories that support multiple classes not working properly with auto-reimporting

This addresses UE-19912

diff --git a/Engine/Source/Editor/UnrealEd/Private/AutoReimport/ContentDirectoryMonitor.cpp b/Engine/Source/Editor/UnrealEd/Private/AutoReimport/ContentDirectoryMonitor.cpp
--- a/Engine/Source/Editor/UnrealEd/Private/AutoReimport/ContentDirectoryMonitor.cpp
+++ b/Engine/Source/Editor/UnrealEd/Private/AutoReimport/ContentDirectoryMonitor.cpp
@@ -149,9 +149,12 @@ UObject* AttemptImport(UClass* InFactoryType, UPackage* Package, FName InName, b
 	if (UFactory* Factory = NewObject<UFactory>(GetTransientPackage(), InFactoryType))
 	{
 		Factory->AddToRoot();
-		if (Factory->ConfigureProperties() && Factory->SupportedClass)
+		if (Factory->ConfigureProperties())
 		{
-			Asset = UFactory::StaticImportObject(Factory->SupportedClass, Package, InName, RF_Public | RF_Standalone, bCancelled, *FullFilename, nullptr, Factory);
+			if (auto* SupportedClass = Factory->ResolveSupportedClass())
+			{
+				Asset = UFactory::StaticImportObject(SupportedClass, Package, InName, RF_Public | RF_Standalone, bCancelled, *FullFilename, nullptr, Factory);
+			}
 		}
 		Factory->RemoveFromRoot();
 	}
