C++ generation without BP duplication

The FKismetEditorUtilities::GenerateCppCode should create a duplicate of BP before the BP is compiles with EKismetCompileType::Cpp. The BP cannot be duplicated into the transient package, because of FKismetEditorUtilities::StripExternalComponents functions (that removes instanced objects fro, CDO) called in FKismetEditorUtilities::CompileBlueprint.

diff --git a/Engine/Source/Editor/UnrealEd/Private/Kismet2/Kismet2.cpp b/Engine/Source/Editor/UnrealEd/Private/Kismet2/Kismet2.cpp
--- a/Engine/Source/Editor/UnrealEd/Private/Kismet2/Kismet2.cpp
+++ b/Engine/Source/Editor/UnrealEd/Private/Kismet2/Kismet2.cpp
@@ -960,54 +960,27 @@ void FKismetEditorUtilities::GenerateCppCode(UObject* Obj, TSharedPtr<FString> O
 		check(OutHeaderSource.IsValid());
 		check(OutCppSource.IsValid());
 
-		//TGuardValue<bool> DuplicatingReadOnly(InBlueprintObj->bDuplicatingReadOnly, true);
+		auto BlueprintObj = InBlueprintObj;
 		{
-			auto BlueprintObj = DuplicateObject<UBlueprint>(InBlueprintObj, GetTransientPackage(), *InBlueprintObj->GetName());
-			{
-				// FIX NOT-DUPLICATED CURVES
-
-				auto NewTimelineOwnerClass = CastChecked<UBlueprintGeneratedClass>(BlueprintObj->GeneratedClass);
-				TSet<UCurveBase*> AllCurves;
-				for (auto Timeline : NewTimelineOwnerClass->Timelines)
-				{
-					Timeline->GetAllCurves(AllCurves);
-				}
-				TMap<UCurveBase*, UCurveBase*> CurveMap;
-				ensure(InBlueprintObj->GetOutermost() != NewTimelineOwnerClass->GetOutermost());
-				for (auto OldCurve : AllCurves)
-				{
-					if (OldCurve && OldCurve->GetOutermost() == InBlueprintObj->GetOutermost())
-					{
-						UCurveBase* NewCurve = DuplicateObject<UCurveBase>(OldCurve, NewTimelineOwnerClass);
-						CurveMap.Add(OldCurve, NewCurve);
-					}
-				}
-				FArchiveReplaceObjectRef<UCurveBase> ReplaceCurvesAr(NewTimelineOwnerClass, CurveMap, /*bNullPrivateRefs=*/ false, /*bIgnoreOuterRef=*/ true, /*bIgnoreArchetypeRef=*/ true);
-			}
-
-			{
-				auto Reinstancer = FBlueprintCompileReinstancer::Create(BlueprintObj->GeneratedClass);
+			auto Reinstancer = FBlueprintCompileReinstancer::Create(BlueprintObj->GeneratedClass);
 
-				IKismetCompilerInterface& Compiler = FModuleManager::LoadModuleChecked<IKismetCompilerInterface>(KISMET_COMPILER_MODULENAME);
+			IKismetCompilerInterface& Compiler = FModuleManager::LoadModuleChecked<IKismetCompilerInterface>(KISMET_COMPILER_MODULENAME);
 
-				TGuardValue<bool> GuardTemplateNameFlag(GCompilingBlueprint, true);
-				FCompilerResultsLog Results;
+			TGuardValue<bool> GuardTemplateNameFlag(GCompilingBlueprint, true);
+			FCompilerResultsLog Results;
 
-				FKismetCompilerOptions CompileOptions;
-				CompileOptions.CompileType = EKismetCompileType::Cpp;
-				CompileOptions.OutCppSourceCode = OutCppSource;
-				CompileOptions.OutHeaderSourceCode = OutHeaderSource;
-				Compiler.CompileBlueprint(BlueprintObj, CompileOptions, Results);
+			FKismetCompilerOptions CompileOptions;
+			CompileOptions.CompileType = EKismetCompileType::Cpp;
+			CompileOptions.OutCppSourceCode = OutCppSource;
+			CompileOptions.OutHeaderSourceCode = OutHeaderSource;
+			Compiler.CompileBlueprint(BlueprintObj, CompileOptions, Results);
 
-				if (EBlueprintType::BPTYPE_Interface == BlueprintObj->BlueprintType && OutCppSource.IsValid())
-				{
-					OutCppSource->Empty();
-				}
+			if (EBlueprintType::BPTYPE_Interface == BlueprintObj->BlueprintType && OutCppSource.IsValid())
+			{
+				OutCppSource->Empty(); // ugly temp hack
 			}
-			BlueprintObj->RemoveGeneratedClasses();
-			BlueprintObj->ClearFlags(RF_Standalone);
-			BlueprintObj->MarkPendingKill();
 		}
+		CompileBlueprint(BlueprintObj);
 	}
 	else if ((UDEnum || UDStruct) && OutHeaderSource.IsValid())
 	{
