[BuildWatcher] Fix for RootEditor_OnMac

Initialization order issue with UCookOnTheFlyServer properties

diff --git a/Engine/Source/Editor/UnrealEd/Private/CookOnTheFlyServer.cpp b/Engine/Source/Editor/UnrealEd/Private/CookOnTheFlyServer.cpp
--- a/Engine/Source/Editor/UnrealEd/Private/CookOnTheFlyServer.cpp
+++ b/Engine/Source/Editor/UnrealEd/Private/CookOnTheFlyServer.cpp
@@ -823,8 +823,8 @@ UCookOnTheFlyServer::UCookOnTheFlyServer(const FObjectInitializer& ObjectInitial
 	CurrentCookMode(ECookMode::CookOnTheFly),
 	CookByTheBookOptions(NULL),
 	CookFlags(ECookInitializationFlags::None),
-	bIsSavingPackage( false ),
-	bIsInitializingSandbox(false)
+	bIsInitializingSandbox(false),
+	bIsSavingPackage(false)
 {
 }
 
