Added units to auto-reimport threshold property

diff --git a/Engine/Source/Editor/UnrealEd/Classes/Settings/EditorLoadingSavingSettings.h b/Engine/Source/Editor/UnrealEd/Classes/Settings/EditorLoadingSavingSettings.h
--- a/Engine/Source/Editor/UnrealEd/Classes/Settings/EditorLoadingSavingSettings.h
+++ b/Engine/Source/Editor/UnrealEd/Classes/Settings/EditorLoadingSavingSettings.h
@@ -110,7 +110,7 @@ public:
 	UPROPERTY(EditAnywhere, config, AdvancedDisplay,Category=AutoReimport, meta=(DisplayName="Directories to Monitor", ToolTip="Lists every directory to monitor for content changes. Can be virtual package paths (eg /Game/ or /MyPlugin/), or absolute paths on disk."))
 	TArray<FAutoReimportDirectoryConfig> AutoReimportDirectorySettings;
 
-	UPROPERTY(EditAnywhere, config, AdvancedDisplay,Category=AutoReimport, meta=(ClampMin=0, ClampMax=60, DisplayName="Import Threshold Time", ToolTip="Specifies an amount of time to wait before a specific file change is considered for auto reimport"))
+	UPROPERTY(EditAnywhere, config, AdvancedDisplay, Category=AutoReimport, meta=(ClampMin=0, ClampMax=60, Units=Seconds, DisplayName="Import Threshold Time", ToolTip="Specifies an amount of time to wait before a specific file change is considered for auto reimport"))
 	float AutoReimportThreshold;
 	UPROPERTY(EditAnywhere, config, AdvancedDisplay, Category=AutoReimport, meta=(DisplayName="Auto Create Assets", ToolTip="When enabled, newly added source content files will be automatically imported into new assets."))
 	bool bAutoCreateAssets;
