#FBX : Save Animation time change when import

https://jira.ol.epicgames.net/browse/UE-7544

Merging

//depot/UE4-Orion/Engine/Source/Editor/UnrealEd/Classes/Factories/FbxAnimSequenceImportData.h

to //depot/Ue/Engine/Source/Editor/UnrealEd/Classes/Factories/FbxAnimSequenceImportData.h

diff --git a/Engine/Source/Editor/UnrealEd/Classes/Factories/FbxAnimSequenceImportData.h b/Engine/Source/Editor/UnrealEd/Classes/Factories/FbxAnimSequenceImportData.h
--- a/Engine/Source/Editor/UnrealEd/Classes/Factories/FbxAnimSequenceImportData.h
+++ b/Engine/Source/Editor/UnrealEd/Classes/Factories/FbxAnimSequenceImportData.h
@@ -29,7 +29,7 @@ class UNREALED_API UFbxAnimSequenceImportData : public UFbxAssetImportData
 	GENERATED_UCLASS_BODY()
 
 	/** Type of asset to import from the FBX file */
-	UPROPERTY(EditAnywhere, Category=ImportSettings, meta=(DisplayName = "Animation Length"))
+	UPROPERTY(EditAnywhere, Category=ImportSettings, config, meta=(DisplayName = "Animation Length"))
 	TEnumAsByte<enum EFBXAnimationLengthImportType> AnimationLength;
 
 	/** Type of asset to import from the FBX file */
