Cis non-unity fix

diff --git a/Engine/Source/Editor/UnrealEd/Private/AnimationEditorUtils.cpp b/Engine/Source/Editor/UnrealEd/Private/AnimationEditorUtils.cpp
--- a/Engine/Source/Editor/UnrealEd/Private/AnimationEditorUtils.cpp
+++ b/Engine/Source/Editor/UnrealEd/Private/AnimationEditorUtils.cpp
@@ -4,6 +4,7 @@
 #include "AnimationEditorUtils.h"
 #include "AssetToolsModule.h"
 #include "Animation/AnimComposite.h"
+#include "Animation/AnimCompress.h"
 #include "Animation/BlendSpace.h"
 #include "Animation/BlendSpace1D.h"
 #include "Animation/AimOffsetBlendSpace.h"
