Added more warnings and error cases to cooker to help track down issue with reloading packages.

diff --git a/Engine/Source/Editor/UnrealEd/Private/CookOnTheFlyServer.cpp b/Engine/Source/Editor/UnrealEd/Private/CookOnTheFlyServer.cpp
--- a/Engine/Source/Editor/UnrealEd/Private/CookOnTheFlyServer.cpp
+++ b/Engine/Source/Editor/UnrealEd/Private/CookOnTheFlyServer.cpp
@@ -2453,7 +2453,11 @@ ESavePackageResult UCookOnTheFlyServer::SaveCookedPackage(UPackage* Package, uin
 
 	if (Filename.Len())
 	{
-		check(!(Package->PackageFlags & PKG_ReloadingForCooker));
+		if (Package->PackageFlags & PKG_ReloadingForCooker)
+		{
+			UE_LOG(LogCook, Warning, TEXT("Package %s marked as reloading for cook by was requested to save"), *Package->GetPathName());
+			UE_LOG(LogCook, Fatal, TEXT("Package %s marked as reloading for cook by was requested to save"), *Package->GetPathName());
+		}
 
 		FString Name = Package->GetPathName();
 
@@ -2622,6 +2626,13 @@ ESavePackageResult UCookOnTheFlyServer::SaveCookedPackage(UPackage* Package, uin
 	return Result;
 }
 
+// this is a duplciate of the 
+static bool IsMobileHDR()
+{
+	static auto* MobileHDRCvar = IConsoleManager::Get().FindTConsoleVariableDataInt(TEXT("r.MobileHDR"));
+	return MobileHDRCvar->GetValueOnAnyThread() == 1;
+}
+
 void UCookOnTheFlyServer::Initialize( ECookMode::Type DesiredCookMode, ECookInitializationFlags InCookFlags, const FString &InOutputDirectoryOverride )
 {
 	OutputDirectoryOverride = InOutputDirectoryOverride;
@@ -2690,12 +2701,7 @@ void UCookOnTheFlyServer::Initialize( ECookMode::Type DesiredCookMode, ECookInit
 
 	ITargetPlatformManagerModule& TPM = GetTargetPlatformManagerRef();
 	const TArray<ITargetPlatform*>& Platforms = TPM.GetCookingTargetPlatforms();
-
-	if ( ((CookFlags & ECookInitializationFlags::Iterative) != ECookInitializationFlags::None) && 
-		((CookFlags & ECookInitializationFlags::Unversioned) != ECookInitializationFlags::None) )
-	{
-		UE_LOG(LogCook, Warning, TEXT("Iterative cooking is unstable when using unversioned cooked content, please disable one of these flags"));
-	}
+	UE_LOG(LogCook, Warning, TEXT("Mobile HDR setting %d"), IsMobileHDR());
 
 	// always generate the asset registry before starting to cook, for either method
 	GenerateAssetRegistry(Platforms);
@@ -3045,6 +3051,7 @@ void UCookOnTheFlyServer::PopulateCookedPackagesFromDisk( const TArray<ITargetPl
 				if (Diff >= 0.0)
 				{
 					CookedPackages.Add(FFilePlatformRequest( FName(*StandardCookedFilename), PlatformFName) );
+					UE_LOG(LogCook, Display, TEXT("Package %s for %s up to date, avoiding cook"), *StandardCookedFilename, *PlatformFName.ToString());
 				}
 			}
 		}
