CIS fix for non-unity

diff --git a/Engine/Source/Editor/UnrealEd/Private/Commandlets/ParticleSystemAuditCommandlet.cpp b/Engine/Source/Editor/UnrealEd/Private/Commandlets/ParticleSystemAuditCommandlet.cpp
--- a/Engine/Source/Editor/UnrealEd/Private/Commandlets/ParticleSystemAuditCommandlet.cpp
+++ b/Engine/Source/Editor/UnrealEd/Private/Commandlets/ParticleSystemAuditCommandlet.cpp
@@ -1,6 +1,7 @@
 // Copyright 1998-2015 Epic Games, Inc. All Rights Reserved.
 
 #include "UnrealEd.h"
+#include "CollectionManagerModule.h"
 #include "Runtime/Engine/Public/ObjectEditorUtils.h"
 #include "Runtime/Engine/Classes/Particles/ParticleEmitter.h"
 #include "Runtime/Engine/Classes/Particles/ParticleModuleRequired.h"
