Added skip editor content button to file -> package settings.
Added NeverCookDirectories to cook and also to project settings.

diff --git a/Engine/Source/Editor/MainFrame/Private/Frame/MainFrameActions.cpp b/Engine/Source/Editor/MainFrame/Private/Frame/MainFrameActions.cpp
--- a/Engine/Source/Editor/MainFrame/Private/Frame/MainFrameActions.cpp
+++ b/Engine/Source/Editor/MainFrame/Private/Frame/MainFrameActions.cpp
@@ -514,6 +514,11 @@ void FMainFrameActionCallbacks::PackageProject( const FName InPlatformInfoName )
 		}
 	}
 
+	if (PackagingSettings->bSkipEditorContent)
+	{
+		OptionalParams += TEXT(" -SKIPEDITORCONTENT");
+	}
+
 	if (PackagingSettings->UsePakFile)
 	{
 	  if (PlatformInfo->TargetPlatformName != FName("HTML5")) 
diff --git a/Engine/Source/Editor/UnrealEd/Classes/Settings/ProjectPackagingSettings.h b/Engine/Source/Editor/UnrealEd/Classes/Settings/ProjectPackagingSettings.h
--- a/Engine/Source/Editor/UnrealEd/Classes/Settings/ProjectPackagingSettings.h
+++ b/Engine/Source/Editor/UnrealEd/Classes/Settings/ProjectPackagingSettings.h
@@ -153,6 +153,12 @@ public:
 	UPROPERTY(config, EditAnywhere, Category=Packaging, AdvancedDisplay, meta=(DisplayName="Create compressed cooked packages"))
 	bool bCompressed;
 	
+	/**
+	* Skip editor content
+	*/
+	UPROPERTY(config, EditAnywhere, Category = Packaging, AdvancedDisplay, meta = (DisplayName = "Do not include editor content in this package may cause game to crash / error if you are using this content."))
+	bool bSkipEditorContent;
+
 	/**
 	 * List of maps to include when no other map list is specified on commandline
 	 */
@@ -166,6 +172,15 @@ public:
 	UPROPERTY(config, EditAnywhere, Category=Packaging, AdvancedDisplay, meta=(DisplayName="Additional Asset Directories to Cook", RelativeToGameContentDir))
 	TArray<FDirectoryPath> DirectoriesToAlwaysCook;
 	
+
+	/**
+	* Directories containing .uasset files that should always be cooked regardless of whether they're referenced by anything in your project
+	* Note: These paths are relative to your project Content directory
+	*/
+	UPROPERTY(config, EditAnywhere, Category = Packaging, AdvancedDisplay, meta = (DisplayName = "Directories to never cook", RelativeToGameContentDir))
+	TArray<FDirectoryPath> DirectoriesToNeverCook;
+
+
 	/**
 	 * Directories containing files that should always be added to the .pak file (if using a .pak file; otherwise they're copied as individual files)
 	 * This is used to stage additional files that you manually load via the UFS (Unreal File System) file IO API
diff --git a/Engine/Source/Editor/UnrealEd/Private/Commandlets/CookCommandlet.cpp b/Engine/Source/Editor/UnrealEd/Private/Commandlets/CookCommandlet.cpp
--- a/Engine/Source/Editor/UnrealEd/Private/Commandlets/CookCommandlet.cpp
+++ b/Engine/Source/Editor/UnrealEd/Private/Commandlets/CookCommandlet.cpp
@@ -1068,6 +1068,7 @@ bool UCookCommandlet::NewCook( const TArray<ITargetPlatform*>& Platforms, TArray
 	TArray<FString> CmdLineMapEntries;
 	TArray<FString> CmdLineDirEntries;
 	TArray<FString> CmdLineCultEntries;
+	TArray<FString> CmdLineNeverCookDirEntries;
 	for (int32 SwitchIdx = 0; SwitchIdx < Switches.Num(); SwitchIdx++)
 	{
 		const FString& Switch = Switches[SwitchIdx];
@@ -1111,12 +1112,14 @@ bool UCookCommandlet::NewCook( const TArray<ITargetPlatform*>& Platforms, TArray
 	{
 		const FString AbsoluteGameContentDir = FPaths::ConvertRelativePathToFull(FPaths::GameContentDir());
 		const UProjectPackagingSettings* const PackagingSettings = GetDefault<UProjectPackagingSettings>();
-		for(const auto& DirToCook : PackagingSettings->DirectoriesToAlwaysCook)
+		for (const auto& DirToCook : PackagingSettings->DirectoriesToAlwaysCook)
 		{
 			CmdLineDirEntries.Add(AbsoluteGameContentDir / DirToCook.Path);
 		}
+
 	}
 
+
 	// Add any map sections specified on command line
 	TArray<FString> AlwaysCookMapList;
 
@@ -1173,6 +1176,7 @@ bool UCookCommandlet::NewCook( const TArray<ITargetPlatform*>& Platforms, TArray
 	StartupOptions.TargetPlatforms = Platforms;
 	Swap( StartupOptions.CookMaps, MapList );
 	Swap( StartupOptions.CookDirectories, CmdLineDirEntries );
+	Swap( StartupOptions.NeverCookDirectories, CmdLineNeverCookDirEntries);
 	Swap( StartupOptions.CookCultures, CmdLineCultEntries );
 	Swap( StartupOptions.DLCName, DLCName );
 	Swap( StartupOptions.BasedOnReleaseVersion, BasedOnReleaseVersion );
diff --git a/Engine/Source/Editor/UnrealEd/Private/CookOnTheFlyServer.cpp b/Engine/Source/Editor/UnrealEd/Private/CookOnTheFlyServer.cpp
--- a/Engine/Source/Editor/UnrealEd/Private/CookOnTheFlyServer.cpp
+++ b/Engine/Source/Editor/UnrealEd/Private/CookOnTheFlyServer.cpp
@@ -1570,6 +1570,7 @@ uint32 UCookOnTheFlyServer::TickCookOnTheSide( const float TimeSlice, uint32 &Co
 #endif
 			continue;
 		}
+
 #if DEBUG_COOKONTHEFLY
 		UE_LOG(LogCook, Display, TEXT("Processing package %s"), *ToBuild.GetFilename().ToString());
 #endif
@@ -1632,6 +1633,16 @@ uint32 UCookOnTheFlyServer::TickCookOnTheSide( const float TimeSlice, uint32 &Co
 			}
 		}
 
+
+		if (NeverCookPackageList.Contains(ToBuild.GetFilename()))
+		{
+#if DEBUG_COOKONTHEFLY
+			UE_LOG(LogCook, Display, TEXT("Package %s requested but is in the never cook package list, discarding request"), *ToBuild.GetFilename().ToString());
+#endif
+			bShouldCook = false;
+		}
+
+
 		if ( bShouldCook ) // if we should cook the package then cook it otherwise add it to the list of already cooked packages below
 		{
 			SCOPE_PACKAGE_STAT(AllOfLoadPackage);
@@ -1976,6 +1987,13 @@ uint32 UCookOnTheFlyServer::TickCookOnTheSide( const float TimeSlice, uint32 &Co
 					continue;
 				}
 
+				const FName StandardPackageFilename = GetCachedStandardPackageFileFName(Package);
+				if (NeverCookPackageList.Contains(StandardPackageFilename))
+				{
+					// refuse to save this package, it's clearly one of the undesirables
+					continue;
+				}
+
 				// if we are processing unsolicited packages we can optionally not save these right now
 				// the unsolicited packages which we missed now will be picked up on next run
 				// we want to do this in cook on the fly also, if there is a new network package request instead of saving unsolicited packages we can process the requested package
@@ -3351,7 +3369,7 @@ void UCookOnTheFlyServer::AddFileToCook( TArray<FName>& InOutFilesToCook, const
 	}
 }
 
-void UCookOnTheFlyServer::CollectFilesToCook(TArray<FName>& FilesInPath, const TArray<FString> &CookMaps, const TArray<FString> &InCookDirectories, const TArray<FString> &CookCultures, const TArray<FString> &IniMapSections, bool bCookAll, bool bMapsOnly, bool bNoDev )
+void UCookOnTheFlyServer::CollectFilesToCook(TArray<FName>& FilesInPath, const TArray<FString>& CookMaps, const TArray<FString>& InCookDirectories, const TArray<FString> &CookCultures, const TArray<FString> &IniMapSections, bool bCookAll, bool bMapsOnly, bool bNoDev )
 {
 	if (CookByTheBookOptions->bIsChildCooker)
 	{
@@ -4022,7 +4040,52 @@ void UCookOnTheFlyServer::StartCookByTheBook( const FCookByTheBookStartupOptions
 	CookByTheBookOptions->bForceDisableCompressedPackages = !!(CookOptions & ECookByTheBookOptions::ForceDisableCompressed);
 	CookByTheBookOptions->bIsChildCooker = CookByTheBookStartupOptions.ChildCookFileName.Len() > 0 ? true : false;
 	CookByTheBookOptions->ChildCookFilename = CookByTheBookStartupOptions.ChildCookFileName;
+	
+	NeverCookPackageList.Empty();
+	{
+		const FString AbsoluteGameContentDir = FPaths::ConvertRelativePathToFull(FPaths::GameContentDir());
+		const UProjectPackagingSettings* const PackagingSettings = GetDefault<UProjectPackagingSettings>();
+
+		TArray<FString> NeverCookDirectories = CookByTheBookStartupOptions.NeverCookDirectories;
 
+		for (const auto& DirToNotCook : PackagingSettings->DirectoriesToNeverCook)
+		{
+			NeverCookDirectories.Add(AbsoluteGameContentDir / DirToNotCook.Path);
+		}
+
+		for (const auto& NeverCookDirectory : NeverCookDirectories)
+		{
+			// add the packages to the never cook package list
+			struct FNeverCookDirectoryWalker : public IPlatformFile::FDirectoryVisitor
+			{
+			private:
+				FThreadSafeNameSet &NeverCookPackageList;
+			public:
+				FNeverCookDirectoryWalker(FThreadSafeNameSet &InNeverCookPackageList) : NeverCookPackageList(InNeverCookPackageList) { }
+
+				// IPlatformFile::FDirectoryVisitor interface
+				virtual bool Visit(const TCHAR* FilenameOrDirectory, bool bIsDirectory)
+				{
+					if (bIsDirectory)
+					{
+						return true;
+					}
+					FString StandardFilename = FString(FilenameOrDirectory);
+					FPaths::MakeStandardFilename(StandardFilename);
+
+					NeverCookPackageList.Add(FName(*StandardFilename));
+					return true;
+				}
+
+			} NeverCookDirectoryWalker(NeverCookPackageList);
+
+			IPlatformFile& PlatformFile = FPlatformFileManager::Get().GetPlatformFile();
+
+			PlatformFile.IterateDirectoryRecursively(*NeverCookDirectory, NeverCookDirectoryWalker);
+		}
+
+	}
+	
 
 	CookByTheBookOptions->TargetPlatformNames.Empty();
 	for (const auto &Platform : TargetPlatforms)
@@ -4563,6 +4626,13 @@ void UCookOnTheFlyServer::MaybeMarkPackageAsAlreadyLoaded(UPackage *Package)
 		UE_LOG(LogCook, Display, TEXT("Marking %s as reloading for cooker because it's been cooked for platforms%s."), *StandardName.ToString(), *Platforms);
 	}
 
+
+	if (NeverCookPackageList.Contains(StandardName))
+	{
+		bShouldMarkAsAlreadyProcessed = true;
+		UE_LOG(LogCook, Display, TEXT("Marking %s as reloading for cooker because it was requested as never cook package."), *StandardName.ToString());
+	}
+
 	/*
 	// if we are currently loading a package then the request would already be dequed so this check doesn't work
 	// need to figure out a better way to do this 
