Missed file in last checkin

diff --git a/Engine/Source/Editor/UnrealEd/Classes/CookOnTheSide/CookOnTheFlyServer.h b/Engine/Source/Editor/UnrealEd/Classes/CookOnTheSide/CookOnTheFlyServer.h
--- a/Engine/Source/Editor/UnrealEd/Classes/CookOnTheSide/CookOnTheFlyServer.h
+++ b/Engine/Source/Editor/UnrealEd/Classes/CookOnTheSide/CookOnTheFlyServer.h
@@ -627,6 +627,11 @@ private:
 			FScopeLock NamesLock(&NamesCritical);
 			return Names.Contains(InName);
 		}
+		void Empty()
+		{
+			FScopeLock NamesLock(&NamesCritical);
+			Names.Empty();
+		}
 	};
 private:
 	/** Current cook mode the cook on the fly server is running in */
@@ -754,6 +759,7 @@ private:
 	FFilenameQueue CookRequests; // list of requested files
 	FThreadSafeUnsolicitedPackagesList UnsolicitedCookedPackages;
 	FThreadSafeFilenameSet CookedPackages; // set of files which have been cooked when needing to recook a file the entry will need to be removed from here
+	FThreadSafeNameSet NeverCookPackageList;
 	FThreadSafeNameSet UncookedEditorOnlyPackages; // set of packages that have been rejected due to being referenced by editor-only properties
 
 	FString GetCachedPackageFilename( const FName& PackageName ) const;
@@ -832,6 +838,7 @@ public:
 		TArray<ITargetPlatform*> TargetPlatforms;
 		TArray<FString> CookMaps;
 		TArray<FString> CookDirectories;
+		TArray<FString> NeverCookDirectories;
 		TArray<FString> CookCultures; 
 		TArray<FString> IniMapSections;
 		ECookByTheBookOptions CookOptions;
