- Utilized animation recorder to enable recording a skeletal mesh's animation from live gameplay into a baked anim that can be played back/exported/etc.
- moved animation recorder from Persona to UnrealEd module

(merged from branch)

diff --git a/Engine/Source/Editor/LevelEditor/Private/LevelEditorActions.cpp b/Engine/Source/Editor/LevelEditor/Private/LevelEditorActions.cpp
--- a/Engine/Source/Editor/LevelEditor/Private/LevelEditorActions.cpp
+++ b/Engine/Source/Editor/LevelEditor/Private/LevelEditorActions.cpp
@@ -61,7 +61,7 @@
 #include "Engine/StaticMesh.h"
 #include "Engine/Light.h"
 #include "Animation/SkeletalMeshActor.h"
-#include "Editor/Persona/Public/AnimationRecorder.h"
+#include "Editor/UnrealEd/Public/Animation/AnimationRecorder.h"
 #include "Editor/KismetWidgets/Public/CreateBlueprintFromActorDialog.h"
 #include "EditorProjectSettings.h"
 #include "HierarchicalLODUtils.h"
diff --git a/Engine/Source/Editor/Persona/Private/SAnimationDlgs.cpp b/Engine/Source/Editor/Persona/Private/SAnimationDlgs.cpp
--- a/Engine/Source/Editor/Persona/Private/SAnimationDlgs.cpp
+++ b/Engine/Source/Editor/Persona/Private/SAnimationDlgs.cpp
@@ -10,217 +10,11 @@
 #define LOCTEXT_NAMESPACE "SAnimationDlgs"
 
 /////////////////////////////////////////////////////
-/////////////////////////////////////////////////////
-/////////////////////////////////////////////////////
-/////////////////////////////////////////////////////
-
-/////////////////////////////////////////////////////
-// Create Animation dialog for recording animation
+// create path picker for importing assets
 /////////////////////////////////////////////////////
 
-FText SCreateAnimationDlg::LastUsedAssetPath;
 FText SImportPathDialog::LastUsedAssetPath;
 
-void SCreateAnimationDlg::Construct(const FArguments& InArgs)
-{
-	AssetPath = FText::FromString(FPackageName::GetLongPackagePath(InArgs._DefaultAssetPath.ToString()));
-	AssetName = FText::FromString(FPackageName::GetLongPackageAssetName(InArgs._DefaultAssetPath.ToString()));
-
-	if (AssetPath.IsEmpty())
-	{
-		AssetPath = LastUsedAssetPath;
-		// still empty?
-		if (AssetPath.IsEmpty())
-		{
-			AssetPath = FText::FromString(TEXT("/Game"));
-		}
-	}
-	else
-	{
-		LastUsedAssetPath = AssetPath;
-	}
-
-	if (AssetName.IsEmpty())
-	{
-		// find default name for them
-		FAssetToolsModule& AssetToolsModule = FModuleManager::Get().LoadModuleChecked<FAssetToolsModule>("AssetTools");
-		FString OutPackageName, OutAssetName;
-		FString PackageName = AssetPath.ToString() + TEXT("/NewAnimation");
-
-		AssetToolsModule.Get().CreateUniqueAssetName(PackageName,  TEXT(""), OutPackageName, OutAssetName);
-		AssetName = FText::FromString(OutAssetName);
-	}
-
-	FPathPickerConfig PathPickerConfig;
-	PathPickerConfig.DefaultPath = AssetPath.ToString();
-	PathPickerConfig.OnPathSelected = FOnPathSelected::CreateSP(this, &SCreateAnimationDlg::OnPathChange);
-	PathPickerConfig.bAddDefaultPath = true;
-
-	FContentBrowserModule& ContentBrowserModule = FModuleManager::LoadModuleChecked<FContentBrowserModule>("ContentBrowser");
-
-	SWindow::Construct(SWindow::FArguments()
-		.Title(LOCTEXT("SCreateAnimationDlg_Title", "Create New Animation Object"))
-		.SupportsMinimize(false)
-		.SupportsMaximize(false)
-		//.SizingRule( ESizingRule::Autosized )
-		.ClientSize(FVector2D(450,450))
-		[
-			SNew(SVerticalBox)
-
-			+ SVerticalBox::Slot() // Add user input block
-			.Padding(2)
-			[
-				SNew(SBorder)
-				.BorderImage(FEditorStyle::GetBrush("ToolPanel.GroupBorder"))
-				[
-					SNew(SVerticalBox)
-
-					+SVerticalBox::Slot()
-					.AutoHeight()
-					[
-						SNew(STextBlock)
-						.Text(LOCTEXT("SelectPath", "Select Path to create animation"))
-						.Font( FSlateFontInfo( FPaths::EngineContentDir() / TEXT("Slate/Fonts/Roboto-Regular.ttf"), 14 ) )
-					]
-
-					+SVerticalBox::Slot()
-					.FillHeight(1)
-					.Padding(3)
-					[
-						ContentBrowserModule.Get().CreatePathPicker(PathPickerConfig)
-					]
-
-					+SVerticalBox::Slot()
-					.AutoHeight()
-					[
-						SNew(SSeparator)
-					]
-
-					+SVerticalBox::Slot()
-					.AutoHeight()
-					.Padding(3)
-					[
-						SNew(SHorizontalBox)
-
-						+SHorizontalBox::Slot()
-						.AutoWidth()
-						.Padding(0, 0, 10, 0)
-						.VAlign(VAlign_Center)
-						[
-							SNew(STextBlock)
-							.Text(LOCTEXT("AnimationName", "Animation Name"))
-						]
-
-						+SHorizontalBox::Slot()
-						[
-							SNew(SEditableTextBox)
-							.Text(AssetName)
-							.OnTextCommitted(this, &SCreateAnimationDlg::OnNameChange)
-							.MinDesiredWidth(250)
-						]
-					]
-				]
-			]
-
-			+SVerticalBox::Slot()
-			.AutoHeight()
-			.HAlign(HAlign_Right)
-			.Padding(5)
-			[
-				SNew(SUniformGridPanel)
-				.SlotPadding(FEditorStyle::GetMargin("StandardDialog.SlotPadding"))
-				.MinDesiredSlotWidth(FEditorStyle::GetFloat("StandardDialog.MinDesiredSlotWidth"))
-				.MinDesiredSlotHeight(FEditorStyle::GetFloat("StandardDialog.MinDesiredSlotHeight"))
-				+SUniformGridPanel::Slot(0,0)
-				[
-					SNew(SButton) 
-					.HAlign(HAlign_Center)
-					.ContentPadding(FEditorStyle::GetMargin("StandardDialog.ContentPadding"))
-					.Text(LOCTEXT("OK", "OK"))
-					.OnClicked(this, &SCreateAnimationDlg::OnButtonClick, EAppReturnType::Ok)
-				]
-				+SUniformGridPanel::Slot(1,0)
-					[
-						SNew(SButton) 
-						.HAlign(HAlign_Center)
-						.ContentPadding(FEditorStyle::GetMargin("StandardDialog.ContentPadding"))
-						.Text(LOCTEXT("Cancel", "Cancel"))
-						.OnClicked(this, &SCreateAnimationDlg::OnButtonClick, EAppReturnType::Cancel)
-					]
-			]
-		]);
-}
-
-void SCreateAnimationDlg::OnNameChange(const FText& NewName, ETextCommit::Type CommitInfo)
-{
-	AssetName = NewName;
-}
-
-void SCreateAnimationDlg::OnPathChange(const FString& NewPath)
-{
-	AssetPath = FText::FromString(NewPath);
-	LastUsedAssetPath = AssetPath;
-}
-
-FReply SCreateAnimationDlg::OnButtonClick(EAppReturnType::Type ButtonID)
-{
-	UserResponse = ButtonID;
-
-	if (ButtonID != EAppReturnType::Cancel)
-	{
-		if (!ValidatePackage())
-		{
-			// reject the request
-			return FReply::Handled();
-		}
-	}
-
-	RequestDestroyWindow();
-
-	return FReply::Handled();
-}
-
-/** Ensures supplied package name information is valid */
-bool SCreateAnimationDlg::ValidatePackage()
-{
-	FText Reason;
-	FString FullPath = GetFullAssetPath();
-
-	if(!FPackageName::IsValidLongPackageName(FullPath, false, &Reason)
-		|| !FName(*AssetName.ToString()).IsValidObjectName(Reason))
-	{
-		FMessageDialog::Open(EAppMsgType::Ok, Reason);
-		return false;
-	}
-
-	return true;
-}
-
-EAppReturnType::Type SCreateAnimationDlg::ShowModal()
-{
-	GEditor->EditorAddModalWindow(SharedThis(this));
-	return UserResponse;
-}
-
-FString SCreateAnimationDlg::GetAssetPath()
-{
-	return AssetPath.ToString();
-}
-
-FString SCreateAnimationDlg::GetAssetName()
-{
-	return AssetName.ToString();
-}
-
-FString SCreateAnimationDlg::GetFullAssetPath()
-{
-	return AssetPath.ToString() + "/" + AssetName.ToString();
-}
-
-/////////////////////////////////////////////////////
-// create path picker for importing assets
-/////////////////////////////////////////////////////
-
 void SImportPathDialog::Construct(const FArguments& InArgs)
 {
 	AssetPath = FText::FromString(FPackageName::GetLongPackagePath(InArgs._DefaultAssetPath.ToString()));
diff --git a/Engine/Source/Editor/Persona/Private/SAnimationDlgs.h b/Engine/Source/Editor/Persona/Private/SAnimationDlgs.h
--- a/Engine/Source/Editor/Persona/Private/SAnimationDlgs.h
+++ b/Engine/Source/Editor/Persona/Private/SAnimationDlgs.h
@@ -2,50 +2,6 @@
 
 #pragma once
 
-class SCreateAnimationDlg: public SWindow
-{
-public:
-	SLATE_BEGIN_ARGS(SCreateAnimationDlg)
-	{
-	}
-	
-	SLATE_ARGUMENT(FText, DefaultAssetPath)
-	SLATE_END_ARGS()
-
-		SCreateAnimationDlg()
-		: UserResponse(EAppReturnType::Cancel)
-	{
-		}
-
-	void Construct(const FArguments& InArgs);
-
-public:
-	/** Displays the dialog in a blocking fashion */
-	EAppReturnType::Type ShowModal();
-
-	/** Gets the resulting asset path */
-	FString GetAssetPath();
-
-	/** Gets the resulting asset name */
-	FString GetAssetName();
-
-	/** Gets the resulting full asset path (path+'/'+name) */
-	FString GetFullAssetPath();
-
-protected:
-	void OnPathChange(const FString& NewPath);
-	void OnNameChange(const FText& NewName, ETextCommit::Type CommitInfo);
-	FReply OnButtonClick(EAppReturnType::Type ButtonID);
-
-	bool ValidatePackage();
-
-	EAppReturnType::Type UserResponse;
-	FText AssetPath;
-	FText AssetName;
-
-	static FText LastUsedAssetPath;
-};
-
 class SImportPathDialog: public SWindow
 {
 public:
diff --git a/Engine/Source/Editor/Persona/Private/Shared/AnimationRecorder.cpp b/Engine/Source/Editor/Persona/Private/Shared/AnimationRecorder.cpp
deleted file mode 100644
diff --git a/Engine/Source/Editor/Persona/Public/AnimationRecorder.h b/Engine/Source/Editor/Persona/Public/AnimationRecorder.h
deleted file mode 100644
diff --git a/Engine/Source/Editor/UnrealEd/Classes/Editor/EditorEngine.h b/Engine/Source/Editor/UnrealEd/Classes/Editor/EditorEngine.h
--- a/Engine/Source/Editor/UnrealEd/Classes/Editor/EditorEngine.h
+++ b/Engine/Source/Editor/UnrealEd/Classes/Editor/EditorEngine.h
@@ -766,6 +766,8 @@ public:
 	bool	HandleListMapPackageDependenciesCommand( const TCHAR* Str, FOutputDevice& Ar );
 	bool	HandleRebuildVolumesCommand( const TCHAR* Str, FOutputDevice& Ar, UWorld* InWorld );
 	bool	HandleRemoveArchtypeFlagCommand( const TCHAR* Str, FOutputDevice& Ar );
+	bool	HandleRecordAnimationCommand(UWorld* InWorld, const TCHAR* InStr, FOutputDevice& Ar);
+	bool	HandleStopRecordAnimationCommand(UWorld* InWorld, const TCHAR* InStr, FOutputDevice& Ar);
 
 	/**
 	 * Initializes the Editor.
diff --git a/Engine/Source/Editor/UnrealEd/Private/Animation/AnimationRecorder.cpp b/Engine/Source/Editor/UnrealEd/Private/Animation/AnimationRecorder.cpp
new file mode 100644
--- /dev/null
+++ b/Engine/Source/Editor/UnrealEd/Private/Animation/AnimationRecorder.cpp
@@ -0,0 +1,440 @@
+// Copyright 1998-2015 Epic Games, Inc. All Rights Reserved.
+
+
+#include "UnrealEd.h"
+#include "AnimationRecorder.h"
+#include "SCreateAnimationDlg.h"
+#include "Runtime/AssetRegistry/Public/AssetRegistryModule.h"
+#include "SNotificationList.h"
+#include "NotificationManager.h"
+#include "EngineLogs.h"
+
+#define LOCTEXT_NAMESPACE "FAnimationRecorder"
+
+/////////////////////////////////////////////////////
+
+FAnimationRecorder::FAnimationRecorder()
+	: IntervalTime(1.0f/DEFAULT_SAMPLERATE)
+	, MaxFrame(DEFAULT_SAMPLERATE * 60 * 10) // for now, set the max limit to 10 mins. 
+	, AnimationObject(NULL)
+	, bRecordLocalToWorld(false)
+{
+}
+
+FAnimationRecorder::~FAnimationRecorder()
+{
+	StopRecord(false);
+}
+
+// get record configuration
+bool GetRecordConfig(FString& AssetPath, FString& AssetName)
+{
+	TSharedRef<SCreateAnimationDlg> NewAnimDlg =
+		SNew(SCreateAnimationDlg);
+
+	if(NewAnimDlg->ShowModal() != EAppReturnType::Cancel)
+	{
+		AssetPath = NewAnimDlg->GetFullAssetPath();
+		AssetName = NewAnimDlg->GetAssetName();
+		return true;
+	}
+
+	return false;
+}
+
+bool FAnimationRecorder::TriggerRecordAnimation(USkeletalMeshComponent* Component)
+{
+	// ask for path
+	FString AssetPath;
+	FString AssetName;
+
+	if (!Component || !Component->SkeletalMesh || !Component->SkeletalMesh->Skeleton)
+	{
+		return false;
+	}
+
+	if (GetRecordConfig(AssetPath, AssetName))
+	{
+		return TriggerRecordAnimation(Component, AssetPath, AssetName);
+	}
+
+	return false;
+}
+
+bool FAnimationRecorder::TriggerRecordAnimation(USkeletalMeshComponent* Component, FString AssetPath, FString AssetName)
+{
+	if (!Component || !Component->SkeletalMesh || !Component->SkeletalMesh->Skeleton)
+	{
+		return false;
+	}
+
+	// create the asset
+	UObject* Parent = AssetPath.IsEmpty() ? nullptr : CreatePackage(nullptr, *AssetPath);
+	if (Parent == nullptr)
+	{
+		// bad or no path passed in, do the popup
+		if (GetRecordConfig(AssetPath, AssetName) == false)
+		{
+			return false;
+		}
+		
+		Parent = CreatePackage(nullptr, *AssetPath);
+	}
+
+	UObject* const Object = LoadObject<UObject>(Parent, *AssetName, nullptr, LOAD_None, nullptr);
+	// if object with same name exists, warn user
+	if (Object)
+	{
+		FMessageDialog::Open(EAppMsgType::Ok, NSLOCTEXT("UnrealEd", "Error_AssetExist", "Asset with same name exists. Can't overwrite another asset"));
+		return false;		// failed
+	}
+
+	// If not, create new one now.
+	UAnimSequence* const NewSeq = NewObject<UAnimSequence>(Parent, *AssetName, RF_Public | RF_Standalone);
+	if (NewSeq)
+	{
+		// set skeleton
+		NewSeq->SetSkeleton(Component->SkeletalMesh->Skeleton);
+		// Notify the asset registry
+		FAssetRegistryModule::AssetCreated(NewSeq);
+		StartRecord(Component, NewSeq);
+
+		return true;
+	}
+
+	return false;
+}
+
+
+void FAnimationRecorder::StartRecord(USkeletalMeshComponent* Component, UAnimSequence* InAnimationObject)
+{
+	TimePassed = 0.f;
+	AnimationObject = InAnimationObject;
+
+	AnimationObject->RawAnimationData.Empty();
+	AnimationObject->TrackToSkeletonMapTable.Empty();
+	AnimationObject->AnimationTrackNames.Empty();
+
+	PreviousSpacesBases = Component->GetSpaceBases();
+
+	LastFrame = 0;
+	AnimationObject->SequenceLength = 0.f;
+	AnimationObject->NumFrames = 0;
+
+	USkeleton* AnimSkeleton = AnimationObject->GetSkeleton();
+	// add all frames
+	for (int32 BoneIndex=0; BoneIndex <PreviousSpacesBases.Num(); ++BoneIndex)
+	{
+		// verify if this bone exists in skeleton
+		int32 BoneTreeIndex = AnimSkeleton->GetSkeletonBoneIndexFromMeshBoneIndex(Component->SkeletalMesh, BoneIndex);
+		if (BoneTreeIndex != INDEX_NONE)
+		{
+			// add tracks for the bone existing
+			FName BoneTreeName = AnimSkeleton->GetReferenceSkeleton().GetBoneName(BoneTreeIndex);
+			AnimationObject->AnimationTrackNames.Add(BoneTreeName);
+			// add mapping to skeleton bone track
+			AnimationObject->TrackToSkeletonMapTable.Add(FTrackToSkeletonMap(BoneTreeIndex));
+		}
+	}
+
+	int32 NumTracks = AnimationObject->AnimationTrackNames.Num();
+	// add tracks
+	AnimationObject->RawAnimationData.AddZeroed(NumTracks);
+
+	// record the first frame;
+	Record(Component, PreviousSpacesBases, 0);
+}
+
+UAnimSequence* FAnimationRecorder::StopRecord(bool bShowMessage)
+{
+	if (AnimationObject)
+	{
+		int32 NumFrames = LastFrame  + 1;
+		AnimationObject->NumFrames = NumFrames;
+
+		// can't use TimePassed. That is just total time that has been passed, not necessarily match with frame count
+		AnimationObject->SequenceLength = (NumFrames>1)? (NumFrames-1) * IntervalTime : MINIMUM_ANIMATION_LENGTH;
+		AnimationObject->PostProcessSequence();
+		AnimationObject->MarkPackageDirty();
+
+		UAnimSequence* ReturnObject = AnimationObject;
+
+		// notify to user
+		if (bShowMessage)
+		{
+			const FText NotificationText = FText::Format( LOCTEXT("RecordAnimation", "'{0}' has been successfully recorded [{1} frames : {2} sec(s)]" ), 
+				FText::FromString( AnimationObject->GetName() ),
+				FText::AsNumber( AnimationObject->NumFrames ),
+				FText::AsNumber( AnimationObject->SequenceLength ));
+
+			//This is not showing well in the Persona, so opening dialog first. 
+			//right now it will crash if you don't wait until end of the record, so it is important for users to know
+			//this is done
+					
+			FNotificationInfo Info(NotificationText);
+			Info.ExpireDuration = 5.0f;
+			Info.bUseLargeFont = false;
+			TSharedPtr<SNotificationItem> Notification = FSlateNotificationManager::Get().AddNotification(Info);
+			if ( Notification.IsValid() )
+			{
+				Notification->SetCompletionState( SNotificationItem::CS_Success );
+			}
+
+			FAssetRegistryModule::AssetCreated(AnimationObject);
+			
+			FMessageDialog::Open(EAppMsgType::Ok, NotificationText);
+		}
+
+		AnimationObject = NULL;
+		PreviousSpacesBases.Empty();
+
+		return ReturnObject;
+	}
+
+	return NULL;
+}
+
+// return false if nothing to update
+// return true if it has properly updated
+void FAnimationRecorder::UpdateRecord(USkeletalMeshComponent* Component, float DeltaTime)
+{
+	// if no animation object, return
+	if (!AnimationObject)
+	{
+		return;
+	}
+
+	const float MaxDelta = 1.f/10.f;
+	float UseDelta = FMath::Min<float>(DeltaTime, MaxDelta);
+
+	// if longer than 1 second, you'll see the halt on the animation
+	if (UseDelta <= 0.f)
+	{
+		return;
+	}
+
+	float PreviousTimePassed = TimePassed;
+	TimePassed += UseDelta;
+
+	// time passed has been updated
+	// now find what frames we need to update
+	int32 FramesRecorded = LastFrame;
+	int32 FramesToRecord = FPlatformMath::TruncToInt(TimePassed / IntervalTime);
+
+	if (FramesRecorded >= FramesToRecord)
+	{
+		return;
+	}
+
+	const TArray<FTransform> & SpaceBases = Component->GetSpaceBases();
+
+	check (SpaceBases.Num() == PreviousSpacesBases.Num());
+
+	TArray<FTransform> BlendedSpaceBases;
+	BlendedSpaceBases.AddZeroed(SpaceBases.Num());
+
+	UE_LOG(LogAnimation, Warning, TEXT("DeltaTime : %0.2f, Current Frame Count : %d, Frames To Record : %d, TimePassed : %0.2f"), DeltaTime
+		, FramesRecorded, FramesToRecord, TimePassed);
+
+	// if we need to record frame
+	while (FramesToRecord > FramesRecorded)
+	{
+		// find what frames we need to record
+		// convert to time
+		float CurrentTime = (FramesRecorded + 1) * IntervalTime;
+		float BlendAlpha = (CurrentTime - PreviousTimePassed)/UseDelta;
+
+		UE_LOG(LogAnimation, Warning, TEXT("Current Frame Count : %d, BlendAlpha : %0.2f"), FramesRecorded + 1, BlendAlpha);
+
+		// for now we just concern component space, not skeleton space
+		for (int32 BoneIndex=0; BoneIndex<SpaceBases.Num(); ++BoneIndex)
+		{
+			BlendedSpaceBases[BoneIndex].Blend(PreviousSpacesBases[BoneIndex], SpaceBases[BoneIndex], BlendAlpha);
+		}
+
+		Record(Component, BlendedSpaceBases, FramesRecorded + 1);
+		++FramesRecorded;
+	}
+
+	//save to current transform
+	PreviousSpacesBases = Component->GetSpaceBases();
+
+	// if we passed MaxFrame, just stop it
+	if (FramesRecorded > MaxFrame)
+	{
+		UE_LOG(LogAnimation, Warning, TEXT("Animation Recording exceeds the time limited (10 mins). Stopping recording animation... "))
+		StopRecord(true);
+	}
+}
+
+void FAnimationRecorder::Record( USkeletalMeshComponent* Component, TArray<FTransform> SpacesBases, int32 FrameToAdd )
+{
+	if (ensure (AnimationObject))
+	{
+		USkeleton* AnimSkeleton = AnimationObject->GetSkeleton();
+		for (int32 TrackIndex=0; TrackIndex <AnimationObject->RawAnimationData.Num(); ++TrackIndex)
+		{
+			// verify if this bone exists in skeleton
+			int32 BoneTreeIndex = AnimationObject->GetSkeletonIndexFromTrackIndex(TrackIndex);
+			if (BoneTreeIndex != INDEX_NONE)
+			{
+				int32 BoneIndex = AnimSkeleton->GetMeshBoneIndexFromSkeletonBoneIndex(Component->SkeletalMesh, BoneTreeIndex);
+				int32 ParentIndex = Component->SkeletalMesh->RefSkeleton.GetParentIndex(BoneIndex);
+				FTransform LocalTransform = SpacesBases[BoneIndex];
+				if ( ParentIndex != INDEX_NONE )
+				{
+					LocalTransform.SetToRelativeTransform(SpacesBases[ParentIndex]);
+				}
+				// if record local to world, we'd like to consider component to world to be in root
+				else if (bRecordLocalToWorld)
+				{
+					LocalTransform *= Component->ComponentToWorld;
+				}
+
+				FRawAnimSequenceTrack& RawTrack = AnimationObject->RawAnimationData[TrackIndex];
+				RawTrack.PosKeys.Add(LocalTransform.GetTranslation());
+				RawTrack.RotKeys.Add(LocalTransform.GetRotation());
+				RawTrack.ScaleKeys.Add(LocalTransform.GetScale3D());
+
+				// verification
+				check (FrameToAdd == RawTrack.PosKeys.Num()-1);
+			}
+		}
+
+		LastFrame = FrameToAdd;
+	}
+}
+
+
+
+
+void FAnimationRecorderManager::Tick(float DeltaTime)
+{
+	for (auto& Inst : RecorderInstances)
+	{
+		Inst.Update(DeltaTime);
+	}
+}
+
+FAnimationRecorderManager::FAnimationRecorderManager()
+{
+}
+
+FAnimationRecorderManager::~FAnimationRecorderManager()
+{
+}
+
+FAnimationRecorderManager& FAnimationRecorderManager::Get()
+{
+	static FAnimationRecorderManager AnimRecorderManager;
+	return AnimRecorderManager;
+}
+
+
+FAnimRecorderInstance::FAnimRecorderInstance()
+	: Actor(nullptr),
+	SkelComp(nullptr),
+	Recorder(nullptr)
+{
+}
+
+void FAnimRecorderInstance::Init(AActor* InActor, USkeletalMeshComponent* InComponent, FString InAssetPath, FString InAssetName)
+{
+	Actor = InActor;
+	SkelComp = InComponent;
+	AssetPath = InAssetPath;
+	AssetName = InAssetName;
+	Recorder = new FAnimationRecorder();
+	Recorder->bRecordLocalToWorld = true;
+}
+
+FAnimRecorderInstance::~FAnimRecorderInstance()
+{
+	if (Recorder)
+	{
+		delete Recorder;
+		Recorder = nullptr;
+	}
+}
+
+bool FAnimRecorderInstance::BeginRecording()
+{
+	if (Recorder)
+	{
+		return Recorder->TriggerRecordAnimation(SkelComp.Get(), AssetPath, AssetName);
+	}
+
+	return false;
+}
+
+void FAnimRecorderInstance::Update(float DeltaTime)
+{
+	if (Recorder)
+	{
+		Recorder->UpdateRecord(SkelComp.Get(), DeltaTime);
+	}
+}
+
+void FAnimRecorderInstance::FinishRecording()
+{
+	const FText FinishRecordingAnimationSlowTask = LOCTEXT("FinishRecordingAnimationSlowTask", "Finalizing recorded animation");
+	GWarn->BeginSlowTask(FinishRecordingAnimationSlowTask, false);
+
+	if (Recorder)
+	{
+		Recorder->StopRecord(false);
+	}
+
+	GWarn->EndSlowTask();
+}
+
+bool FAnimationRecorderManager::RecordAnimation(AActor* Actor, USkeletalMeshComponent* Component, FString AssetPath, FString AssetName)
+{
+	int32 const NewInstIdx = RecorderInstances.AddDefaulted();
+	
+	FAnimRecorderInstance& NewInst = RecorderInstances[NewInstIdx];
+	NewInst.Init(Actor, Component, AssetPath, AssetName);
+
+	bool const bSuccess = NewInst.BeginRecording();
+	if (bSuccess == false)
+	{
+		// failed, remove it
+		RecorderInstances.RemoveAtSwap(NewInstIdx);
+	}
+
+	return bSuccess;
+}
+
+void FAnimationRecorderManager::StopRecordingAnimation(AActor* Actor, USkeletalMeshComponent* Component)
+{
+	FAnimRecorderInstance* InstToStop = nullptr;
+	for (int32 Idx = 0; Idx < RecorderInstances.Num(); ++Idx)
+	{
+		FAnimRecorderInstance& Inst = RecorderInstances[Idx];
+		if (Inst.Actor == Actor && Inst.SkelComp == Component)
+		{
+			// stop and finalize recoded data
+			Inst.FinishRecording();
+
+			// remove instance, which will clean itself up
+			RecorderInstances.RemoveAtSwap(Idx);
+
+			// all done
+			break;
+		}
+	}
+}
+
+void FAnimationRecorderManager::StopRecordingAllAnimations()
+{
+	for (auto& Inst : RecorderInstances)
+	{
+		Inst.FinishRecording();
+	}
+
+	RecorderInstances.Empty();
+}
+
+
+#undef LOCTEXT_NAMESPACE
+
diff --git a/Engine/Source/Editor/UnrealEd/Private/Animation/SCreateAnimationDlg.cpp b/Engine/Source/Editor/UnrealEd/Private/Animation/SCreateAnimationDlg.cpp
new file mode 100644
--- /dev/null
+++ b/Engine/Source/Editor/UnrealEd/Private/Animation/SCreateAnimationDlg.cpp
@@ -0,0 +1,217 @@
+// Copyright 1998-2015 Epic Games, Inc. All Rights Reserved.
+
+
+#include "UnrealEd.h"
+//#include "Persona.h"
+#include "SCreateAnimationDlg.h"
+#include "Editor/ContentBrowser/Public/ContentBrowserModule.h"
+#include "Developer/AssetTools/Public/AssetToolsModule.h"
+
+#define LOCTEXT_NAMESPACE "SCreateAnimationDlg"
+
+/////////////////////////////////////////////////////
+/////////////////////////////////////////////////////
+/////////////////////////////////////////////////////
+/////////////////////////////////////////////////////
+
+/////////////////////////////////////////////////////
+// Create Animation dialog for recording animation
+/////////////////////////////////////////////////////
+
+FText SCreateAnimationDlg::LastUsedAssetPath;
+
+void SCreateAnimationDlg::Construct(const FArguments& InArgs)
+{
+	AssetPath = FText::FromString(FPackageName::GetLongPackagePath(InArgs._DefaultAssetPath.ToString()));
+	AssetName = FText::FromString(FPackageName::GetLongPackageAssetName(InArgs._DefaultAssetPath.ToString()));
+
+	if (AssetPath.IsEmpty())
+	{
+		AssetPath = LastUsedAssetPath;
+		// still empty?
+		if (AssetPath.IsEmpty())
+		{
+			AssetPath = FText::FromString(TEXT("/Game"));
+		}
+	}
+	else
+	{
+		LastUsedAssetPath = AssetPath;
+	}
+
+	if (AssetName.IsEmpty())
+	{
+		// find default name for them
+		FAssetToolsModule& AssetToolsModule = FModuleManager::Get().LoadModuleChecked<FAssetToolsModule>("AssetTools");
+		FString OutPackageName, OutAssetName;
+		FString PackageName = AssetPath.ToString() + TEXT("/NewAnimation");
+
+		AssetToolsModule.Get().CreateUniqueAssetName(PackageName,  TEXT(""), OutPackageName, OutAssetName);
+		AssetName = FText::FromString(OutAssetName);
+	}
+
+	FPathPickerConfig PathPickerConfig;
+	PathPickerConfig.DefaultPath = AssetPath.ToString();
+	PathPickerConfig.OnPathSelected = FOnPathSelected::CreateSP(this, &SCreateAnimationDlg::OnPathChange);
+	PathPickerConfig.bAddDefaultPath = true;
+
+	FContentBrowserModule& ContentBrowserModule = FModuleManager::LoadModuleChecked<FContentBrowserModule>("ContentBrowser");
+
+	SWindow::Construct(SWindow::FArguments()
+		.Title(LOCTEXT("SCreateAnimationDlg_Title", "Create New Animation Object"))
+		.SupportsMinimize(false)
+		.SupportsMaximize(false)
+		//.SizingRule( ESizingRule::Autosized )
+		.ClientSize(FVector2D(450,450))
+		[
+			SNew(SVerticalBox)
+
+			+ SVerticalBox::Slot() // Add user input block
+			.Padding(2)
+			[
+				SNew(SBorder)
+				.BorderImage(FEditorStyle::GetBrush("ToolPanel.GroupBorder"))
+				[
+					SNew(SVerticalBox)
+
+					+SVerticalBox::Slot()
+					.AutoHeight()
+					[
+						SNew(STextBlock)
+						.Text(LOCTEXT("SelectPath", "Select Path to create animation"))
+						.Font( FSlateFontInfo( FPaths::EngineContentDir() / TEXT("Slate/Fonts/Roboto-Regular.ttf"), 14 ) )
+					]
+
+					+SVerticalBox::Slot()
+					.FillHeight(1)
+					.Padding(3)
+					[
+						ContentBrowserModule.Get().CreatePathPicker(PathPickerConfig)
+					]
+
+					+SVerticalBox::Slot()
+					.AutoHeight()
+					[
+						SNew(SSeparator)
+					]
+
+					+SVerticalBox::Slot()
+					.AutoHeight()
+					.Padding(3)
+					[
+						SNew(SHorizontalBox)
+
+						+SHorizontalBox::Slot()
+						.AutoWidth()
+						.Padding(0, 0, 10, 0)
+						.VAlign(VAlign_Center)
+						[
+							SNew(STextBlock)
+							.Text(LOCTEXT("AnimationName", "Animation Name"))
+						]
+
+						+SHorizontalBox::Slot()
+						[
+							SNew(SEditableTextBox)
+							.Text(AssetName)
+							.OnTextCommitted(this, &SCreateAnimationDlg::OnNameChange)
+							.MinDesiredWidth(250)
+						]
+					]
+				]
+			]
+
+			+SVerticalBox::Slot()
+			.AutoHeight()
+			.HAlign(HAlign_Right)
+			.Padding(5)
+			[
+				SNew(SUniformGridPanel)
+				.SlotPadding(FEditorStyle::GetMargin("StandardDialog.SlotPadding"))
+				.MinDesiredSlotWidth(FEditorStyle::GetFloat("StandardDialog.MinDesiredSlotWidth"))
+				.MinDesiredSlotHeight(FEditorStyle::GetFloat("StandardDialog.MinDesiredSlotHeight"))
+				+SUniformGridPanel::Slot(0,0)
+				[
+					SNew(SButton) 
+					.HAlign(HAlign_Center)
+					.ContentPadding(FEditorStyle::GetMargin("StandardDialog.ContentPadding"))
+					.Text(LOCTEXT("OK", "OK"))
+					.OnClicked(this, &SCreateAnimationDlg::OnButtonClick, EAppReturnType::Ok)
+				]
+				+SUniformGridPanel::Slot(1,0)
+					[
+						SNew(SButton) 
+						.HAlign(HAlign_Center)
+						.ContentPadding(FEditorStyle::GetMargin("StandardDialog.ContentPadding"))
+						.Text(LOCTEXT("Cancel", "Cancel"))
+						.OnClicked(this, &SCreateAnimationDlg::OnButtonClick, EAppReturnType::Cancel)
+					]
+			]
+		]);
+}
+
+void SCreateAnimationDlg::OnNameChange(const FText& NewName, ETextCommit::Type CommitInfo)
+{
+	AssetName = NewName;
+}
+
+void SCreateAnimationDlg::OnPathChange(const FString& NewPath)
+{
+	AssetPath = FText::FromString(NewPath);
+	LastUsedAssetPath = AssetPath;
+}
+
+FReply SCreateAnimationDlg::OnButtonClick(EAppReturnType::Type ButtonID)
+{
+	UserResponse = ButtonID;
+
+	if (ButtonID != EAppReturnType::Cancel)
+	{
+		if (!ValidatePackage())
+		{
+			// reject the request
+			return FReply::Handled();
+		}
+	}
+
+	RequestDestroyWindow();
+
+	return FReply::Handled();
+}
+
+/** Ensures supplied package name information is valid */
+bool SCreateAnimationDlg::ValidatePackage()
+{
+	FText Reason;
+	FString FullPath = GetFullAssetPath();
+
+	if(!FPackageName::IsValidLongPackageName(FullPath, false, &Reason)
+		|| !FName(*AssetName.ToString()).IsValidObjectName(Reason))
+	{
+		FMessageDialog::Open(EAppMsgType::Ok, Reason);
+		return false;
+	}
+
+	return true;
+}
+
+EAppReturnType::Type SCreateAnimationDlg::ShowModal()
+{
+	GEditor->EditorAddModalWindow(SharedThis(this));
+	return UserResponse;
+}
+
+FString SCreateAnimationDlg::GetAssetPath()
+{
+	return AssetPath.ToString();
+}
+
+FString SCreateAnimationDlg::GetAssetName()
+{
+	return AssetName.ToString();
+}
+
+FString SCreateAnimationDlg::GetFullAssetPath()
+{
+	return AssetPath.ToString() + "/" + AssetName.ToString();
+}
diff --git a/Engine/Source/Editor/UnrealEd/Private/Animation/SCreateAnimationDlg.h b/Engine/Source/Editor/UnrealEd/Private/Animation/SCreateAnimationDlg.h
new file mode 100644
--- /dev/null
+++ b/Engine/Source/Editor/UnrealEd/Private/Animation/SCreateAnimationDlg.h
@@ -0,0 +1,47 @@
+// Copyright 1998-2015 Epic Games, Inc. All Rights Reserved.
+
+#pragma once
+
+class SCreateAnimationDlg: public SWindow
+{
+public:
+	SLATE_BEGIN_ARGS(SCreateAnimationDlg)
+	{
+	}
+	
+	SLATE_ARGUMENT(FText, DefaultAssetPath)
+	SLATE_END_ARGS()
+
+		SCreateAnimationDlg()
+		: UserResponse(EAppReturnType::Cancel)
+	{
+		}
+
+	void Construct(const FArguments& InArgs);
+
+public:
+	/** Displays the dialog in a blocking fashion */
+	EAppReturnType::Type ShowModal();
+
+	/** Gets the resulting asset path */
+	FString GetAssetPath();
+
+	/** Gets the resulting asset name */
+	FString GetAssetName();
+
+	/** Gets the resulting full asset path (path+'/'+name) */
+	FString GetFullAssetPath();
+
+protected:
+	void OnPathChange(const FString& NewPath);
+	void OnNameChange(const FText& NewName, ETextCommit::Type CommitInfo);
+	FReply OnButtonClick(EAppReturnType::Type ButtonID);
+
+	bool ValidatePackage();
+
+	EAppReturnType::Type UserResponse;
+	FText AssetPath;
+	FText AssetName;
+
+	static FText LastUsedAssetPath;
+};
diff --git a/Engine/Source/Editor/UnrealEd/Private/EditorEngine.cpp b/Engine/Source/Editor/UnrealEd/Private/EditorEngine.cpp
--- a/Engine/Source/Editor/UnrealEd/Private/EditorEngine.cpp
+++ b/Engine/Source/Editor/UnrealEd/Private/EditorEngine.cpp
@@ -116,6 +116,8 @@
 #include "PhysicsPublic.h"
 #include "Engine/CoreSettings.h"
 
+#include "AnimationRecorder.h"
+
 DEFINE_LOG_CATEGORY_STATIC(LogEditor, Log, All);
 
 #define LOCTEXT_NAMESPACE "UnrealEd.Editor"
@@ -1574,6 +1576,8 @@ void UEditorEngine::Tick( float DeltaSeconds, bool bIdleMode )
 
 	FUnrealEdMisc::Get().TickPerformanceAnalytics();
 
+	FAnimationRecorderManager::Get().Tick(DeltaSeconds);
+
 	// If the fadeout animation has completed for the undo/redo notification item, allow it to be deleted
 	if(UndoRedoNotificationItem.IsValid() && UndoRedoNotificationItem->GetCompletionState() == SNotificationItem::CS_None)
 	{
diff --git a/Engine/Source/Editor/UnrealEd/Private/EditorServer.cpp b/Engine/Source/Editor/UnrealEd/Private/EditorServer.cpp
--- a/Engine/Source/Editor/UnrealEd/Private/EditorServer.cpp
+++ b/Engine/Source/Editor/UnrealEd/Private/EditorServer.cpp
@@ -69,6 +69,10 @@
 #include "UnrealEngine.h"
 #include "AI/Navigation/NavLinkRenderingComponent.h"
 
+#include "PhysicsPublic.h"
+
+#include "AnimationRecorder.h"
+
 #include "Settings/EditorSettings.h"
 
 DEFINE_LOG_CATEGORY_STATIC(LogEditorServer, Log, All);
@@ -5608,6 +5612,14 @@ bool UEditorEngine::Exec( UWorld* InWorld, const TCHAR* Stream, FOutputDevice& A
 	{
 		HandleRemoveArchtypeFlagCommand( Str, Ar );
 	}
+	else if (FParse::Command(&Str, TEXT("RecordAnimation")))
+	{
+		bProcessed = HandleRecordAnimationCommand(InWorld, Str, Ar);
+	}
+	else if (FParse::Command(&Str, TEXT("StopRecordingAnimation")))
+	{
+		bProcessed = HandleStopRecordAnimationCommand(InWorld, Str, Ar);
+	}
 	else
 	{
 		bProcessed = FBlueprintEditorUtils::KismetDiagnosticExec(Stream, Ar);
@@ -6377,6 +6389,107 @@ bool UEditorEngine::HandleRemoveArchtypeFlagCommand( const TCHAR* Str, FOutputDe
 	return true;
 }
 
+
+bool UEditorEngine::HandleRecordAnimationCommand(UWorld* InWorld, const TCHAR* InStr, FOutputDevice& Ar)
+{
+	const TCHAR* Str = InStr;
+	// parse actor name
+	TCHAR ActorName[128];
+	AActor* FoundActor = nullptr;
+	if (FParse::Token(Str, ActorName, ARRAY_COUNT(ActorName), 0))
+	{
+		FString const ActorNameStr = FString(ActorName);
+		for (ULevel const* Level : InWorld->GetLevels())
+		{
+			if (Level)
+			{
+				for (AActor* Actor : Level->Actors)
+				{
+					if (Actor)
+					{
+						if (Actor->GetName() == ActorNameStr)
+						{
+							FoundActor = Actor;
+							break;
+						}
+					}
+				}
+			}
+		}
+	}
+
+	if (FoundActor)
+	{
+		USkeletalMeshComponent* const SkelComp = FoundActor->FindComponentByClass<USkeletalMeshComponent>();
+		if (SkelComp)
+		{
+			TCHAR AssetPath[256];
+			FParse::Token(Str, AssetPath, ARRAY_COUNT(AssetPath), 0);
+			FString const AssetName = FPackageName::GetLongPackageAssetName(AssetPath);
+			return FAnimationRecorderManager::Get().RecordAnimation(FoundActor, SkelComp, AssetPath, AssetName);
+		}
+	}
+
+	return false;
+}
+
+bool UEditorEngine::HandleStopRecordAnimationCommand(UWorld* InWorld, const TCHAR* InStr, FOutputDevice& Ar)
+{
+	const TCHAR* Str = InStr;
+
+	// parse actor name
+	TCHAR ActorName[128];
+	AActor* FoundActor = nullptr;
+	bool bStopAll = false;
+	if (FParse::Token(Str, ActorName, ARRAY_COUNT(ActorName), 0))
+	{
+		FString const ActorNameStr = FString(ActorName);
+
+		if (ActorNameStr.ToLower() == TEXT("all"))
+		{
+			bStopAll = true;
+		}
+		else if (InWorld)
+		{
+			// search for the actor by name
+			for (ULevel* Level : InWorld->GetLevels())
+			{
+				if (Level)
+				{
+					for (AActor* Actor : Level->Actors)
+					{
+						if (Actor)
+						{
+							if (Actor->GetName() == ActorNameStr)
+							{
+								FoundActor = Actor;
+								break;
+							}
+						}
+					}
+				}
+			}
+		}
+	}
+
+	if (bStopAll)
+	{
+		FAnimationRecorderManager::Get().StopRecordingAllAnimations();
+		return true;
+	}
+	else if (FoundActor)
+	{
+		USkeletalMeshComponent* const SkelComp = FoundActor->FindComponentByClass<USkeletalMeshComponent>();
+		if (SkelComp)
+		{
+			FAnimationRecorderManager::Get().StopRecordingAnimation(FoundActor, SkelComp);
+			return true;
+		}
+	}
+
+	return false;
+}
+
 /**
  * @return true if the given component's StaticMesh can be merged with other StaticMeshes
  */
diff --git a/Engine/Source/Editor/UnrealEd/Public/Animation/AnimationRecorder.h b/Engine/Source/Editor/UnrealEd/Public/Animation/AnimationRecorder.h
new file mode 100644
--- /dev/null
+++ b/Engine/Source/Editor/UnrealEd/Public/Animation/AnimationRecorder.h
@@ -0,0 +1,91 @@
+// Copyright 1998-2015 Epic Games, Inc. All Rights Reserved.
+
+#pragma once
+
+//////////////////////////////////////////////////////////////////////////
+// FAnimationRecorder
+
+// records the mesh pose to animation input
+struct UNREALED_API FAnimationRecorder
+{
+private:
+	float IntervalTime;
+	int32 MaxFrame;
+	int32 LastFrame;
+	float TimePassed;
+	UAnimSequence* AnimationObject;
+	TArray<FTransform> PreviousSpacesBases;
+
+public:
+	FAnimationRecorder();
+	~FAnimationRecorder();
+
+	/** Starts recording an animation. Prompts for asset path and name via dialog. */
+	bool TriggerRecordAnimation(USkeletalMeshComponent* Component);
+	/** Starts recording an animation. */
+	bool TriggerRecordAnimation(USkeletalMeshComponent* Component, FString AssetPath, FString AssetName);
+
+	void StartRecord(USkeletalMeshComponent* Component, UAnimSequence* InAnimationObject);
+	UAnimSequence* StopRecord(bool bShowMessage);
+	void UpdateRecord(USkeletalMeshComponent* Component, float DeltaTime);
+	UAnimSequence* GetAnimationObject() const { return AnimationObject; }
+	bool InRecording() const { return AnimationObject != nullptr; }
+	float GetTimeRecorded() const { return TimePassed; }
+
+	// If true, it will record root to include LocalToWorld
+	bool bRecordLocalToWorld;
+
+private:
+	void Record( USkeletalMeshComponent* Component, TArray<FTransform> SpacesBases, int32 FrameToAdd );
+};
+
+//////////////////////////////////////////////////////////////////////////
+// FAnimRecorderInstance
+
+struct FAnimRecorderInstance
+{
+public:
+	FAnimRecorderInstance();
+	~FAnimRecorderInstance();
+
+	void Init(AActor* InActor, USkeletalMeshComponent* InComponent, FString InAssetPath, FString InAssetName);
+
+	bool BeginRecording();
+	void Update(float DeltaTime);
+	void FinishRecording();
+
+	TWeakObjectPtr<AActor> Actor;
+	TWeakObjectPtr<USkeletalMeshComponent> SkelComp;
+	FString AssetPath;
+	FString AssetName;
+
+	FAnimationRecorder* Recorder;
+};
+
+
+//////////////////////////////////////////////////////////////////////////
+// FAnimationRecorderManager
+
+struct FAnimationRecorderManager
+{
+public:
+	/** Singleton accessor */
+	static FAnimationRecorderManager& Get();
+
+	/** Destructor */
+	virtual ~FAnimationRecorderManager();
+
+	/** Starts recording an animation. */
+	bool RecordAnimation(AActor* Actor, USkeletalMeshComponent* Component, FString AssetPath, FString AssetName);
+	void StopRecordingAnimation(AActor* Actor, USkeletalMeshComponent* Component);
+	void StopRecordingAllAnimations();
+
+	void Tick(float DeltaTime);
+
+private:
+	/** Constructor, private - use Get() function */
+	FAnimationRecorderManager();
+
+	TArray<FAnimRecorderInstance> RecorderInstances;
+};
+
