Fix CIS (mismatched LOCTEXT_NAMESPACE).

diff --git a/Engine/Source/Editor/UnrealEd/Private/Animation/SCreateAnimationDlg.cpp b/Engine/Source/Editor/UnrealEd/Private/Animation/SCreateAnimationDlg.cpp
--- a/Engine/Source/Editor/UnrealEd/Private/Animation/SCreateAnimationDlg.cpp
+++ b/Engine/Source/Editor/UnrealEd/Private/Animation/SCreateAnimationDlg.cpp
@@ -215,3 +215,5 @@ FString SCreateAnimationDlg::GetFullAssetPath()
 {
 	return AssetPath.ToString() + "/" + AssetName.ToString();
 }
+
+#undef LOCTEXT_NAMESPACE
