Prevent prompt to checkout package notifications from appearing in unattended builds.  Licensee suggestion.

diff --git a/Engine/Source/Editor/UnrealEd/Private/UnrealEdSrv.cpp b/Engine/Source/Editor/UnrealEd/Private/UnrealEdSrv.cpp
--- a/Engine/Source/Editor/UnrealEd/Private/UnrealEdSrv.cpp
+++ b/Engine/Source/Editor/UnrealEd/Private/UnrealEdSrv.cpp
@@ -1428,7 +1428,7 @@ bool UUnrealEdEngine::IsUserInteracting()
 
 void UUnrealEdEngine::AttemptModifiedPackageNotification()
 {
-	if( bNeedToPromptForCheckout )
+	if( bNeedToPromptForCheckout && !FApp::IsUnattended() )
 	{
 		// Defer prompting for checkout if we cant prompt because of the following:
 		// The user is interacting with something,
