#ANIM: fix retarget option to support naming option

Merging

//depot/UE4-Orion/Engine/Source/...

to //depot/UE4/Engine/Source/...

diff --git a/Engine/Source/Developer/AssetTools/Private/AssetTypeActions/AssetTypeActions_AnimBlueprint.cpp b/Engine/Source/Developer/AssetTools/Private/AssetTypeActions/AssetTypeActions_AnimBlueprint.cpp
--- a/Engine/Source/Developer/AssetTools/Private/AssetTypeActions/AssetTypeActions_AnimBlueprint.cpp
+++ b/Engine/Source/Developer/AssetTools/Private/AssetTypeActions/AssetTypeActions_AnimBlueprint.cpp
@@ -157,11 +157,11 @@ void FAssetTypeActions_AnimBlueprint::ExecuteFindSkeleton(TArray<TWeakObjectPtr<
 	}
 }
 
-void FAssetTypeActions_AnimBlueprint::RetargetAnimationHandler(USkeleton* OldSkeleton, USkeleton* NewSkeleton, bool bRemapReferencedAssets, bool bConvertSpaces, bool bDuplicateAssets, TArray<TWeakObjectPtr<UObject>> AnimBlueprints)
+void FAssetTypeActions_AnimBlueprint::RetargetAnimationHandler(USkeleton* OldSkeleton, USkeleton* NewSkeleton, bool bRemapReferencedAssets, bool bConvertSpaces, const EditorAnimUtils::FNameDuplicationRule* NameRule, TArray<TWeakObjectPtr<UObject>> AnimBlueprints)
 {
 	if(!OldSkeleton || OldSkeleton->GetPreviewMesh(true))
 	{
-		EditorAnimUtils::RetargetAnimations(OldSkeleton, NewSkeleton, AnimBlueprints, bRemapReferencedAssets, bDuplicateAssets, bConvertSpaces);
+		EditorAnimUtils::RetargetAnimations(OldSkeleton, NewSkeleton, AnimBlueprints, bRemapReferencedAssets, NameRule, bConvertSpaces);
 	}
 	else
 	{
@@ -193,7 +193,7 @@ void FAssetTypeActions_AnimBlueprint::RetargetAssets(TArray<UObject*> InAnimBlue
 	const FText Message = LOCTEXT("RemapSkeleton_Warning", "Select the skeleton to remap this asset to.");
 	auto AnimBlueprints = GetTypedWeakObjectPtrs<UObject>(InAnimBlueprints);
 
-	SAnimationRemapSkeleton::ShowWindow(OldSkeleton, Message, FOnRetargetAnimation::CreateSP(this, &FAssetTypeActions_AnimBlueprint::RetargetAnimationHandler, bDuplicateAssets, AnimBlueprints));
+	SAnimationRemapSkeleton::ShowWindow(OldSkeleton, Message, bDuplicateAssets, FOnRetargetAnimation::CreateSP(this, &FAssetTypeActions_AnimBlueprint::RetargetAnimationHandler, AnimBlueprints));
 }
 
 TSharedPtr<SWidget> FAssetTypeActions_AnimBlueprint::GetThumbnailOverlay(const FAssetData& AssetData) const
diff --git a/Engine/Source/Developer/AssetTools/Private/AssetTypeActions/AssetTypeActions_AnimBlueprint.h b/Engine/Source/Developer/AssetTools/Private/AssetTypeActions/AssetTypeActions_AnimBlueprint.h
--- a/Engine/Source/Developer/AssetTools/Private/AssetTypeActions/AssetTypeActions_AnimBlueprint.h
+++ b/Engine/Source/Developer/AssetTools/Private/AssetTypeActions/AssetTypeActions_AnimBlueprint.h
@@ -30,5 +30,5 @@ private:
 	void RetargetAssets(TArray<UObject*> InAnimBlueprints, bool bDuplicateAssets);
 
 	/** Handler for retargeting */
-	void RetargetAnimationHandler(USkeleton* OldSkeleton, USkeleton* NewSkeleton, bool bRemapReferencedAssets, bool bConvertSpaces, bool bDuplicateAssets, TArray<TWeakObjectPtr<UObject>> AnimBlueprints);
+	void RetargetAnimationHandler(USkeleton* OldSkeleton, USkeleton* NewSkeleton, bool bRemapReferencedAssets, bool bConvertSpaces, const EditorAnimUtils::FNameDuplicationRule* NameRule, TArray<TWeakObjectPtr<UObject>> AnimBlueprints);
 };
\ No newline at end of file
diff --git a/Engine/Source/Developer/AssetTools/Private/AssetTypeActions/AssetTypeActions_AnimationAsset.cpp b/Engine/Source/Developer/AssetTools/Private/AssetTypeActions/AssetTypeActions_AnimationAsset.cpp
--- a/Engine/Source/Developer/AssetTools/Private/AssetTypeActions/AssetTypeActions_AnimationAsset.cpp
+++ b/Engine/Source/Developer/AssetTools/Private/AssetTypeActions/AssetTypeActions_AnimationAsset.cpp
@@ -152,11 +152,11 @@ void FAssetTypeActions_AnimationAsset::ExecuteFindSkeleton(TArray<TWeakObjectPtr
 	}
 }
 
-void FAssetTypeActions_AnimationAsset::RetargetAnimationHandler(USkeleton* OldSkeleton, USkeleton* NewSkeleton, bool bRemapReferencedAssets, bool bConvertSpaces, bool bDuplicateAssets,  TArray<TWeakObjectPtr<UObject>> InAnimAssets)
+void FAssetTypeActions_AnimationAsset::RetargetAnimationHandler(USkeleton* OldSkeleton, USkeleton* NewSkeleton, bool bRemapReferencedAssets, bool bConvertSpaces, const EditorAnimUtils::FNameDuplicationRule* NameRule, TArray<TWeakObjectPtr<UObject>> InAnimAssets)
 {
 	if((!OldSkeleton || OldSkeleton->GetPreviewMesh(false)))
 	{
-		EditorAnimUtils::RetargetAnimations(OldSkeleton, NewSkeleton, InAnimAssets, bRemapReferencedAssets, bDuplicateAssets, bConvertSpaces);
+		EditorAnimUtils::RetargetAnimations(OldSkeleton, NewSkeleton, InAnimAssets, bRemapReferencedAssets, NameRule, bConvertSpaces);
 	}
 	else
 	{
@@ -174,9 +174,9 @@ void FAssetTypeActions_AnimationAsset::RetargetAnimationHandler(USkeleton* OldSk
 	}
 }
 
-void FAssetTypeActions_AnimationAsset::RetargetNonSkeletonAnimationHandler(USkeleton* OldSkeleton, USkeleton* NewSkeleton, bool bRemapReferencedAssets, bool bConvertSpaces, bool bDuplicateAssets, TArray<TWeakObjectPtr<UObject>> InAnimAssets, TWeakPtr<IToolkitHost> EditWithinLevelEditor)
+void FAssetTypeActions_AnimationAsset::RetargetNonSkeletonAnimationHandler(USkeleton* OldSkeleton, USkeleton* NewSkeleton, bool bRemapReferencedAssets, bool bConvertSpaces, const EditorAnimUtils::FNameDuplicationRule* NameRule, TArray<TWeakObjectPtr<UObject>> InAnimAssets, TWeakPtr<IToolkitHost> EditWithinLevelEditor)
 {
-	RetargetAnimationHandler(OldSkeleton, NewSkeleton, bRemapReferencedAssets, bConvertSpaces, bDuplicateAssets, InAnimAssets);
+	RetargetAnimationHandler(OldSkeleton, NewSkeleton, bRemapReferencedAssets, bConvertSpaces, NameRule, InAnimAssets);
 
 	if(NewSkeleton)
 	{
@@ -225,11 +225,11 @@ void FAssetTypeActions_AnimationAsset::RetargetAssets(TArray<UObject*> InAnimAss
 
 	if (bOpenEditor)
 	{
-		SAnimationRemapSkeleton::ShowWindow(OldSkeleton, Message, FOnRetargetAnimation::CreateSP(this, &FAssetTypeActions_AnimationAsset::RetargetNonSkeletonAnimationHandler, bDuplicateAssets, AnimAssets, TWeakPtr<class IToolkitHost> (EditWithinLevelEditor)) );
+		SAnimationRemapSkeleton::ShowWindow(OldSkeleton, Message, bDuplicateAssets, FOnRetargetAnimation::CreateSP(this, &FAssetTypeActions_AnimationAsset::RetargetNonSkeletonAnimationHandler, AnimAssets, TWeakPtr<class IToolkitHost> (EditWithinLevelEditor)) );
 	}
 	else
 	{
-		SAnimationRemapSkeleton::ShowWindow(OldSkeleton, Message, FOnRetargetAnimation::CreateSP(this, &FAssetTypeActions_AnimationAsset::RetargetAnimationHandler, bDuplicateAssets, AnimAssets) );
+		SAnimationRemapSkeleton::ShowWindow(OldSkeleton, Message, bDuplicateAssets, FOnRetargetAnimation::CreateSP(this, &FAssetTypeActions_AnimationAsset::RetargetAnimationHandler, AnimAssets) );
 	}
 }
 
diff --git a/Engine/Source/Developer/AssetTools/Private/AssetTypeActions/AssetTypeActions_AnimationAsset.h b/Engine/Source/Developer/AssetTools/Private/AssetTypeActions/AssetTypeActions_AnimationAsset.h
--- a/Engine/Source/Developer/AssetTools/Private/AssetTypeActions/AssetTypeActions_AnimationAsset.h
+++ b/Engine/Source/Developer/AssetTools/Private/AssetTypeActions/AssetTypeActions_AnimationAsset.h
@@ -1,6 +1,7 @@
 // Copyright 1998-2015 Epic Games, Inc. All Rights Reserved.
 
 #pragma once
+#include "EditorAnimUtils.h"
 
 class FAssetTypeActions_AnimationAsset : public FAssetTypeActions_Base
 {
@@ -27,6 +28,6 @@ private:
 	void RetargetAssets(TArray<UObject*> InAnimBlueprints, bool bDuplicateAssets, bool bOpenEditor, TSharedPtr<class IToolkitHost> EditWithinLevelEditor);
 
 	/** Handler for retargeting */
-	void RetargetAnimationHandler(USkeleton* OldSkeleton, USkeleton* NewSkeleton, bool bRemapReferencedAssets, bool bConvertSpaces, bool bDuplicateAssets, TArray<TWeakObjectPtr<UObject>> InAnimAssets);
-	void RetargetNonSkeletonAnimationHandler(USkeleton* OldSkeleton, USkeleton* NewSkeleton, bool bRemapReferencedAssets, bool bConvertSpaces, bool bDuplicateAssets, TArray<TWeakObjectPtr<UObject>> InAnimAssets, TWeakPtr<class IToolkitHost> EditWithinLevelEditor);
+	void RetargetAnimationHandler(USkeleton* OldSkeleton, USkeleton* NewSkeleton, bool bRemapReferencedAssets, bool bConvertSpaces, const EditorAnimUtils::FNameDuplicationRule* NameRule, TArray<TWeakObjectPtr<UObject>> InAnimAssets);
+	void RetargetNonSkeletonAnimationHandler(USkeleton* OldSkeleton, USkeleton* NewSkeleton, bool bRemapReferencedAssets, bool bConvertSpaces, const EditorAnimUtils::FNameDuplicationRule* NameRule, TArray<TWeakObjectPtr<UObject>> InAnimAssets, TWeakPtr<class IToolkitHost> EditWithinLevelEditor);
 };
\ No newline at end of file
diff --git a/Engine/Source/Developer/AssetTools/Private/AssetTypeActions/AssetTypeActions_Skeleton.cpp b/Engine/Source/Developer/AssetTools/Private/AssetTypeActions/AssetTypeActions_Skeleton.cpp
--- a/Engine/Source/Developer/AssetTools/Private/AssetTypeActions/AssetTypeActions_Skeleton.cpp
+++ b/Engine/Source/Developer/AssetTools/Private/AssetTypeActions/AssetTypeActions_Skeleton.cpp
@@ -519,7 +519,7 @@ void FAssetTypeActions_Skeleton::CreateRig(const TWeakObjectPtr<USkeleton> Skele
 	}
 }
 
-void FAssetTypeActions_Skeleton::RetargetAnimationHandler(USkeleton* OldSkeleton, USkeleton* NewSkeleton, bool bRemapReferencedAssets, bool bConvertSpaces)
+void FAssetTypeActions_Skeleton::RetargetAnimationHandler(USkeleton* OldSkeleton, USkeleton* NewSkeleton, bool bRemapReferencedAssets, bool bConvertSpaces, const EditorAnimUtils::FNameDuplicationRule* NameRule)
 {
 	if((OldSkeleton && OldSkeleton->GetPreviewMesh(true) == NULL))
 	{
@@ -538,6 +538,7 @@ void FAssetTypeActions_Skeleton::RetargetAnimationHandler(USkeleton* OldSkeleton
 		return;
 	}
 
+	// namerule should be null
 	// find all assets who references old skeleton
 	TArray<FName> Packages;
 
@@ -568,9 +569,8 @@ void FAssetTypeActions_Skeleton::ExecuteRetargetSkeleton(TArray<TWeakObjectPtr<U
 			USkeleton* OldSkeleton = (*SkelIt).Get();
 
 			const FText Message = LOCTEXT("RetargetSkeleton_Warning", "This only converts animation data -i.e. animation assets and Anim Blueprints. \nIf you'd like to convert SkeletalMesh, use the context menu (Assign Skeleton) for each mesh. \n\nIf you'd like to convert mesh as well, please do so before converting animation data. \nOtherwise you will lose any extra track that is in the new mesh.");
-
 			// ask user what they'd like to change to 
-			SAnimationRemapSkeleton::ShowWindow(OldSkeleton, Message, FOnRetargetAnimation::CreateSP(this, &FAssetTypeActions_Skeleton::RetargetAnimationHandler));
+			SAnimationRemapSkeleton::ShowWindow(OldSkeleton, Message, false, FOnRetargetAnimation::CreateSP(this, &FAssetTypeActions_Skeleton::RetargetAnimationHandler));
 		}
 	}
 }
diff --git a/Engine/Source/Developer/AssetTools/Private/AssetTypeActions/AssetTypeActions_Skeleton.h b/Engine/Source/Developer/AssetTools/Private/AssetTypeActions/AssetTypeActions_Skeleton.h
--- a/Engine/Source/Developer/AssetTools/Private/AssetTypeActions/AssetTypeActions_Skeleton.h
+++ b/Engine/Source/Developer/AssetTools/Private/AssetTypeActions/AssetTypeActions_Skeleton.h
@@ -1,6 +1,7 @@
 // Copyright 1998-2015 Epic Games, Inc. All Rights Reserved.
 
 #pragma once
+#include "EditorAnimUtils.h"
 
 /**
  * Remap Skeleton Asset Data
@@ -70,5 +71,5 @@ private: // Helper functions
 	void OnAssetCreated(TArray<UObject*> NewAssets) const;
 
 	/** Handler for retargeting */
-	void RetargetAnimationHandler(USkeleton* OldSkeleton, USkeleton* NewSkeleton, bool bRemapReferencedAssets, bool bConvertSpaces);
+	void RetargetAnimationHandler(USkeleton* OldSkeleton, USkeleton* NewSkeleton, bool bRemapReferencedAssets, bool bConvertSpaces, const EditorAnimUtils::FNameDuplicationRule* NameRule);
 };
\ No newline at end of file
diff --git a/Engine/Source/Editor/EditorStyle/Private/SlateEditorStyle.cpp b/Engine/Source/Editor/EditorStyle/Private/SlateEditorStyle.cpp
--- a/Engine/Source/Editor/EditorStyle/Private/SlateEditorStyle.cpp
+++ b/Engine/Source/Editor/EditorStyle/Private/SlateEditorStyle.cpp
@@ -3594,8 +3594,10 @@ void FSlateEditorStyle::FStyle::SetupGraphEditorStyles()
 				);
 			}
 
-			Set( "Persona.RetargetManager.BoldFont",							TTF_CORE_FONT( "Fonts/Roboto-Bold", 10 ) );
+			Set( "Persona.RetargetManager.BoldFont",							TTF_CORE_FONT( "Fonts/Roboto-Bold", 12 ) );
+			Set( "Persona.RetargetManager.SmallBoldFont",						TTF_CORE_FONT( "Fonts/Roboto-Bold", 10 ) );
 			Set( "Persona.RetargetManager.FilterFont",							TTF_CORE_FONT( "Fonts/Roboto-Regular", 10 ) );
+			Set( "Persona.RetargetManager.ItalicFont",							TTF_FONT( "Fonts/Roboto-Italic", 9 ) );
 
 			Set("Persona.RetargetManager.ImportantText", FTextBlockStyle(NormalText)
 				.SetFont(TTF_CORE_FONT("Fonts/Roboto-Bold", 11))
diff --git a/Engine/Source/Editor/Persona/Private/SAnimationSequenceBrowser.cpp b/Engine/Source/Editor/Persona/Private/SAnimationSequenceBrowser.cpp
--- a/Engine/Source/Editor/Persona/Private/SAnimationSequenceBrowser.cpp
+++ b/Engine/Source/Editor/Persona/Private/SAnimationSequenceBrowser.cpp
@@ -340,9 +340,9 @@ void SAnimationSequenceBrowser::OnReimportAnimation(TArray<FAssetData> SelectedA
 	}
 }
 
-void SAnimationSequenceBrowser::RetargetAnimationHandler(USkeleton* OldSkeleton, USkeleton* NewSkeleton, bool bRemapReferencedAssets, bool bConvertSpaces, TArray<TWeakObjectPtr<UObject>> InAnimAssets)
+void SAnimationSequenceBrowser::RetargetAnimationHandler(USkeleton* OldSkeleton, USkeleton* NewSkeleton, bool bRemapReferencedAssets, bool bConvertSpaces, const EditorAnimUtils::FNameDuplicationRule* NameRule, TArray<TWeakObjectPtr<UObject>> InAnimAssets)
 {
-	UObject* AssetToOpen = EditorAnimUtils::RetargetAnimations(OldSkeleton, NewSkeleton, InAnimAssets, bRemapReferencedAssets, true, bConvertSpaces);
+	UObject* AssetToOpen = EditorAnimUtils::RetargetAnimations(OldSkeleton, NewSkeleton, InAnimAssets, bRemapReferencedAssets, NameRule, bConvertSpaces);
 
 	if(UAnimationAsset* AnimAsset = Cast<UAnimationAsset>(AssetToOpen))
 	{
@@ -379,7 +379,7 @@ void SAnimationSequenceBrowser::OnCreateCopy(TArray<FAssetData> Selected)
 		{
 			auto AnimAssetsToConvert = FObjectEditorUtils::GetTypedWeakObjectPtrs<UObject>(AnimAssets);
 			// ask user what they'd like to change to 
-			SAnimationRemapSkeleton::ShowWindow(OldSkeleton, Message, FOnRetargetAnimation::CreateSP(this, &SAnimationSequenceBrowser::RetargetAnimationHandler, AnimAssetsToConvert));
+			SAnimationRemapSkeleton::ShowWindow(OldSkeleton, Message, true, FOnRetargetAnimation::CreateSP(this, &SAnimationSequenceBrowser::RetargetAnimationHandler, AnimAssetsToConvert));
 		}
 	}
 }
diff --git a/Engine/Source/Editor/Persona/Private/SAnimationSequenceBrowser.h b/Engine/Source/Editor/Persona/Private/SAnimationSequenceBrowser.h
--- a/Engine/Source/Editor/Persona/Private/SAnimationSequenceBrowser.h
+++ b/Engine/Source/Editor/Persona/Private/SAnimationSequenceBrowser.h
@@ -6,6 +6,7 @@
 #include "AssetData.h"
 #include "Editor/ContentBrowser/Public/ContentBrowserModule.h"
 #include "PreviewScene.h"
+#include "EditorAnimUtils.h"
 
 //////////////////////////////////////////////////////////////////////////
 // FAnimationAssetViewportClient
@@ -207,7 +208,7 @@ protected:
 	FSyncToAssetsDelegate SyncToAssetsDelegate;
 	FGetCurrentSelectionDelegate GetCurrentSelectionDelegate;
 
-	void RetargetAnimationHandler(USkeleton* OldSkeleton, USkeleton* NewSkeleton, bool bRemapReferencedAssets, bool bConvertSpaces, TArray<TWeakObjectPtr<UObject>> InAnimAssets);
+	void RetargetAnimationHandler(USkeleton* OldSkeleton, USkeleton* NewSkeleton, bool bRemapReferencedAssets, bool bConvertSpaces, const EditorAnimUtils::FNameDuplicationRule* NameRule, TArray<TWeakObjectPtr<UObject>> InAnimAssets);
 
 private:
 	/** Updates the animation preview in the tooltip */
diff --git a/Engine/Source/Editor/UnrealEd/Private/EditorAnimUtils.cpp b/Engine/Source/Editor/UnrealEd/Private/EditorAnimUtils.cpp
--- a/Engine/Source/Editor/UnrealEd/Private/EditorAnimUtils.cpp
+++ b/Engine/Source/Editor/UnrealEd/Private/EditorAnimUtils.cpp
@@ -19,7 +19,7 @@ namespace EditorAnimUtils
 {
 	//////////////////////////////////////////////////////////////////
 	// FAnimationRetargetContext
-	FAnimationRetargetContext::FAnimationRetargetContext(const TArray<FAssetData>& AssetsToRetarget, bool bRetargetReferredAssets, bool bInConvertAnimationDataInComponentSpaces) 
+	FAnimationRetargetContext::FAnimationRetargetContext(const TArray<FAssetData>& AssetsToRetarget, bool bRetargetReferredAssets, bool bInConvertAnimationDataInComponentSpaces, const FNameDuplicationRule& NameRule) 
 		: SingleTargetObject(NULL)
 		, bConvertAnimationDataInComponentSpaces(bInConvertAnimationDataInComponentSpaces)
 	{
@@ -32,7 +32,7 @@ namespace EditorAnimUtils
 		Initialize(WeakObjectList,bRetargetReferredAssets);
 	}
 
-	FAnimationRetargetContext::FAnimationRetargetContext(TArray<TWeakObjectPtr<UObject>> AssetsToRetarget, bool bRetargetReferredAssets, bool bInConvertAnimationDataInComponentSpaces) 
+	FAnimationRetargetContext::FAnimationRetargetContext(TArray<TWeakObjectPtr<UObject>> AssetsToRetarget, bool bRetargetReferredAssets, bool bInConvertAnimationDataInComponentSpaces, const FNameDuplicationRule& NameRule) 
 		: SingleTargetObject(NULL)
 		, bConvertAnimationDataInComponentSpaces(bInConvertAnimationDataInComponentSpaces)
 	{
@@ -99,6 +99,27 @@ namespace EditorAnimUtils
 				DuplicatedBlueprints.Num() > 0;
 	}
 
+	TArray<UObject*> FAnimationRetargetContext::GetAllDuplicates() const
+	{
+		TArray<UObject*> Duplicates;
+
+		if (AnimSequencesToRetarget.Num() > 0)
+		{
+			Duplicates.Append(AnimSequencesToRetarget);
+		}
+
+		if(ComplexAnimsToRetarget.Num() > 0)
+		{
+			Duplicates.Append(ComplexAnimsToRetarget);
+		}
+
+		if(AnimBlueprintsToRetarget.Num() > 0)
+		{
+			Duplicates.Append(AnimBlueprintsToRetarget);
+		}
+		return Duplicates;
+	}
+
 	UObject* FAnimationRetargetContext::GetSingleTargetObject() const
 	{
 		return SingleTargetObject;
@@ -133,13 +154,13 @@ namespace EditorAnimUtils
 		return NULL;
 	}
 
-	void FAnimationRetargetContext::DuplicateAssetsToRetarget(UPackage* DestinationPackage)
+	void FAnimationRetargetContext::DuplicateAssetsToRetarget(UPackage* DestinationPackage, const FNameDuplicationRule* NameRule)
 	{
 		if(!HasDuplicates())
 		{
-			DuplicatedSequences = DuplicateAssets<UAnimSequence>(AnimSequencesToRetarget, DestinationPackage);
-			DuplicatedComplexAssets = DuplicateAssets<UAnimationAsset>(ComplexAnimsToRetarget, DestinationPackage);
-			DuplicatedBlueprints = DuplicateAssets<UAnimBlueprint>(AnimBlueprintsToRetarget, DestinationPackage);
+			DuplicatedSequences = DuplicateAssets<UAnimSequence>(AnimSequencesToRetarget, DestinationPackage, NameRule);
+			DuplicatedComplexAssets = DuplicateAssets<UAnimationAsset>(ComplexAnimsToRetarget, DestinationPackage, NameRule);
+			DuplicatedBlueprints = DuplicateAssets<UAnimBlueprint>(AnimBlueprintsToRetarget, DestinationPackage, NameRule);
 
 			DuplicatedSequences.GenerateValueArray(AnimSequencesToRetarget);
 			DuplicatedComplexAssets.GenerateValueArray(ComplexAnimsToRetarget);
@@ -241,19 +262,19 @@ namespace EditorAnimUtils
 	}
 
 	//////////////////////////////////////////////////////////////////
-	UObject* RetargetAnimations(USkeleton* OldSkeleton, USkeleton* NewSkeleton, TArray<TWeakObjectPtr<UObject>> AssetsToRetarget, bool bRetargetReferredAssets, bool bDuplicateAssetsBeforeRetarget, bool bConvertSpace)
+	UObject* RetargetAnimations(USkeleton* OldSkeleton, USkeleton* NewSkeleton, TArray<TWeakObjectPtr<UObject>> AssetsToRetarget, bool bRetargetReferredAssets, const FNameDuplicationRule* NameRule, bool bConvertSpace)
 	{
 		FAnimationRetargetContext RetargetContext(AssetsToRetarget, bRetargetReferredAssets, bConvertSpace);
-		return RetargetAnimations(OldSkeleton, NewSkeleton, RetargetContext, bRetargetReferredAssets, bDuplicateAssetsBeforeRetarget);
+		return RetargetAnimations(OldSkeleton, NewSkeleton, RetargetContext, bRetargetReferredAssets, NameRule);
 	}
 
-	UObject* RetargetAnimations(USkeleton* OldSkeleton, USkeleton* NewSkeleton, const TArray<FAssetData>& AssetsToRetarget, bool bRetargetReferredAssets, bool bDuplicateAssetsBeforeRetarget, bool bConvertSpace)
+	UObject* RetargetAnimations(USkeleton* OldSkeleton, USkeleton* NewSkeleton, const TArray<FAssetData>& AssetsToRetarget, bool bRetargetReferredAssets, const FNameDuplicationRule* NameRule, bool bConvertSpace)
 	{
 		FAnimationRetargetContext RetargetContext(AssetsToRetarget, bRetargetReferredAssets, bConvertSpace);
-		return RetargetAnimations(OldSkeleton, NewSkeleton, RetargetContext, bRetargetReferredAssets, bDuplicateAssetsBeforeRetarget);
+		return RetargetAnimations(OldSkeleton, NewSkeleton, RetargetContext, bRetargetReferredAssets, NameRule);
 	}
 
-	UObject* RetargetAnimations(USkeleton* OldSkeleton, USkeleton* NewSkeleton, FAnimationRetargetContext& RetargetContext, bool bRetargetReferredAssets, bool bDuplicateAssetsBeforeRetarget)
+	UObject* RetargetAnimations(USkeleton* OldSkeleton, USkeleton* NewSkeleton, FAnimationRetargetContext& RetargetContext, bool bRetargetReferredAssets, const FNameDuplicationRule* NameRule)
 	{
 		check(NewSkeleton);
 		UObject* OriginalObject  = RetargetContext.GetSingleTargetObject();
@@ -261,9 +282,9 @@ namespace EditorAnimUtils
 
 		if(	RetargetContext.HasAssetsToRetarget() )
 		{
-			if(bDuplicateAssetsBeforeRetarget)
+			if(NameRule)
 			{
-				RetargetContext.DuplicateAssetsToRetarget(DuplicationDestPackage);
+				RetargetContext.DuplicateAssetsToRetarget(DuplicationDestPackage, NameRule);
 			}
 			RetargetContext.RetargetAnimations(OldSkeleton, NewSkeleton);
 		}
@@ -272,12 +293,12 @@ namespace EditorAnimUtils
 		Notification.ExpireDuration = 5.f;
 
 		UObject* NotifyLinkObject = OriginalObject;
-		if(OriginalObject && bDuplicateAssetsBeforeRetarget)
+		if(OriginalObject && NameRule)
 		{
 			NotifyLinkObject = RetargetContext.GetDuplicate(OriginalObject);
 		}
 
-		if(!bDuplicateAssetsBeforeRetarget)
+		if(!NameRule)
 		{
 			if(OriginalObject)
 			{
@@ -309,14 +330,41 @@ namespace EditorAnimUtils
 
 		FSlateNotificationManager::Get().AddNotification(Notification);
 
-		if(OriginalObject && bDuplicateAssetsBeforeRetarget)
+		// sync newly created objects on CB
+		if (NotifyLinkObject)
+		{
+			TArray<UObject*> NewObjects = RetargetContext.GetAllDuplicates();
+			TArray<FAssetData> CurrentSelection;
+			for(auto& NewObject : NewObjects)
+			{
+				CurrentSelection.Add(FAssetData(NewObject));
+			}
+
+			FContentBrowserModule& ContentBrowserModule = FModuleManager::Get().LoadModuleChecked<FContentBrowserModule>("ContentBrowser");
+			ContentBrowserModule.Get().SyncBrowserToAssets(CurrentSelection);
+		}
+		if(OriginalObject && NameRule)
 		{
 			return RetargetContext.GetDuplicate(OriginalObject);
 		}
 		return NULL;
 	}
 
-	TMap<UObject*, UObject*> DuplicateAssetsInternal(const TArray<UObject*>& AssetsToDuplicate, UPackage* DestinationPackage)
+	FString CreateDesiredName(UObject* Asset, const FNameDuplicationRule* NameRule)
+	{
+		check(Asset);
+
+		FString NewName = Asset->GetName();
+
+		if(NameRule)
+		{
+			NewName = NameRule->Rename(Asset);
+		}
+
+		return NewName;
+	}
+
+	TMap<UObject*, UObject*> DuplicateAssetsInternal(const TArray<UObject*>& AssetsToDuplicate, UPackage* DestinationPackage, const FNameDuplicationRule* NameRule)
 	{
 		FAssetToolsModule& AssetToolsModule = FModuleManager::LoadModuleChecked<FAssetToolsModule>("AssetTools");
 
@@ -327,11 +375,11 @@ namespace EditorAnimUtils
 			UObject* Asset = (*Iter);
 			if(!DuplicateMap.Contains(Asset))
 			{
-				FString PathName = FPackageName::GetLongPackagePath(DestinationPackage->GetName());
+				FString PathName = (NameRule)? NameRule->FolderPath : FPackageName::GetLongPackagePath(DestinationPackage->GetName());
 
 				FString ObjectName;
 				FString NewPackageName;
-				AssetToolsModule.Get().CreateUniqueAssetName(PathName+"/"+ Asset->GetName(), TEXT("_Copy"), NewPackageName, ObjectName);
+				AssetToolsModule.Get().CreateUniqueAssetName(PathName+"/"+ CreateDesiredName(Asset, NameRule), TEXT(""), NewPackageName, ObjectName);
 
 				// create one on skeleton folder
 				UObject* NewAsset = AssetToolsModule.Get().DuplicateAsset(ObjectName, PathName, Asset);
@@ -414,6 +462,16 @@ namespace EditorAnimUtils
 			}
 		}
 	}
+
+	FString FNameDuplicationRule::Rename(const UObject* Asset) const
+	{
+		check(Asset);
+
+		FString NewName = Asset->GetName();
+
+		NewName = NewName.Replace(*ReplaceFrom, *ReplaceTo);
+		return FString::Printf(TEXT("%s%s%s"), *Prefix, *NewName, *Suffix);
+	}
 }
 
 #undef LOCTEXT_NAMESPACE
diff --git a/Engine/Source/Editor/UnrealEd/Private/SSkeletonWidget.cpp b/Engine/Source/Editor/UnrealEd/Private/SSkeletonWidget.cpp
--- a/Engine/Source/Editor/UnrealEd/Private/SSkeletonWidget.cpp
+++ b/Engine/Source/Editor/UnrealEd/Private/SSkeletonWidget.cpp
@@ -405,10 +405,19 @@ void SAnimationRemapSkeleton::Construct( const FArguments& InArgs )
 	bShowOnlyCompatibleSkeletons = false;
 	OnRetargetAnimationDelegate = InArgs._OnRetargetDelegate;
 
-	TSharedRef<SVerticalBox> Widget = SNew(SVerticalBox);
+	TSharedRef<SVerticalBox> RetargetWidget = SNew(SVerticalBox);
+
+	RetargetWidget->AddSlot()
+	[
+		SNew(STextBlock)
+		.AutoWrapText(true)
+		.Font(FEditorStyle::GetFontStyle("Persona.RetargetManager.SmallBoldFont"))
+		.Text(LOCTEXT("RetargetBasePose_OptionLabel", "Retarget Options"))
+	];
+
 	if(InArgs._ShowRemapOption)
 	{
-		Widget->AddSlot()
+		RetargetWidget->AddSlot()
 		[
 			SNew(SCheckBox)
 			.IsChecked(this, &SAnimationRemapSkeleton::IsRemappingReferencedAssets)
@@ -425,7 +434,7 @@ void SAnimationRemapSkeleton::Construct( const FArguments& InArgs )
 	{
 		TSharedPtr<SToolTip> ConvertSpaceTooltip = IDocumentation::Get()->CreateToolTip(FText::FromString("Check if you'd like to convert animation data to new skeleton space. If this is false, it won't convert any animation data to new space."),
 			NULL, FString("Shared/Editors/Persona"), FString("AnimRemapSkeleton_ConvertSpace"));
-		Widget->AddSlot()
+		RetargetWidget->AddSlot()
 			[
 				SNew(SCheckBox)
 				.IsChecked(this, &SAnimationRemapSkeleton::IsConvertSpacesChecked)
@@ -443,7 +452,7 @@ void SAnimationRemapSkeleton::Construct( const FArguments& InArgs )
 	{
 		TSharedPtr<SToolTip> ConvertSpaceTooltip = IDocumentation::Get()->CreateToolTip(FText::FromString("Check if you'd like to show only the skeleton that uses the same rig."),
 			NULL, FString("Shared/Editors/Persona"), FString("AnimRemapSkeleton_ShowCompatbielSkeletons")); // @todo add tooltip
-		Widget->AddSlot()
+		RetargetWidget->AddSlot()
 			[
 				SNew(SCheckBox)
 				.IsChecked(this, &SAnimationRemapSkeleton::IsShowOnlyCompatibleSkeletonsChecked)
@@ -458,6 +467,162 @@ void SAnimationRemapSkeleton::Construct( const FArguments& InArgs )
 		bShowOnlyCompatibleSkeletons = true;
 	}
 
+	TSharedRef<SHorizontalBox> OptionWidget = SNew(SHorizontalBox);
+	OptionWidget->AddSlot()
+		[
+			RetargetWidget
+		];
+
+	if (InArgs._ShowDuplicateAssetOption)
+	{
+		TSharedRef<SVerticalBox> NameOptionWidget = SNew(SVerticalBox);
+
+		NameOptionWidget->AddSlot()
+		[	
+			SNew(SVerticalBox)
+			+SVerticalBox::Slot()
+			.AutoHeight()
+			.Padding(2, 3)
+			[
+				SNew(STextBlock)
+				.AutoWrapText(true)
+				.Font(FEditorStyle::GetFontStyle("Persona.RetargetManager.SmallBoldFont"))
+				.Text(LOCTEXT("RetargetBasePose_RenameLable", "New Asset Name"))
+			]
+
+			+SVerticalBox::Slot()
+			.AutoHeight()
+			.Padding(2, 1)
+			[
+				SNew(SHorizontalBox)
+				+SHorizontalBox::Slot()
+				[
+					SNew(STextBlock).Text(LOCTEXT("RemapSkeleton_DupeName_Prefix", "Prefix"))
+				]
+
+				+SHorizontalBox::Slot()
+				[
+					SNew(SEditableTextBox)
+						.Text(this, &SAnimationRemapSkeleton::GetPrefixName)
+						.MinDesiredWidth(100)
+						.OnTextChanged(this, &SAnimationRemapSkeleton::SetPrefixName)
+						.IsReadOnly(false)
+						.RevertTextOnEscape(true)
+				]
+
+			]
+
+			+SVerticalBox::Slot()
+			.AutoHeight()
+			.Padding(2, 1)
+			[
+				SNew(SHorizontalBox)
+				+SHorizontalBox::Slot()
+				[
+					SNew(STextBlock).Text(LOCTEXT("RemapSkeleton_DupeName_Suffix", "Suffix"))
+				]
+
+				+SHorizontalBox::Slot()
+				[
+					SNew(SEditableTextBox)
+						.Text(this, &SAnimationRemapSkeleton::GetSuffixName)
+						.MinDesiredWidth(100)
+						.OnTextChanged(this, &SAnimationRemapSkeleton::SetSuffixName)
+						.IsReadOnly(false)
+						.RevertTextOnEscape(true)
+				]
+
+			]
+
+			+SVerticalBox::Slot()
+			.AutoHeight()
+			.Padding(2, 1)
+			[
+				SNew(SHorizontalBox)
+				+SHorizontalBox::Slot()
+				[
+					SNew(STextBlock).Text(LOCTEXT("RemapSkeleton_DupeName_ReplaceFrom", "Replace "))
+				]
+
+				+SHorizontalBox::Slot()
+				[
+					SNew(SEditableTextBox)
+						.Text(this, &SAnimationRemapSkeleton::GetReplaceFrom)
+						.MinDesiredWidth(50)
+						.OnTextChanged(this, &SAnimationRemapSkeleton::SetReplaceFrom)
+						.IsReadOnly(false)
+						.RevertTextOnEscape(true)
+				]
+
+				+SHorizontalBox::Slot()
+				.Padding(5, 0)
+				.HAlign(HAlign_Center)
+				[
+					SNew(STextBlock).Text(LOCTEXT("RemapSkeleton_DupeName_ReplaceTo", "with "))
+				]
+
+				+SHorizontalBox::Slot()
+				[
+					SNew(SEditableTextBox)
+						.Text(this, &SAnimationRemapSkeleton::GetReplaceTo)
+						.MinDesiredWidth(50)
+						.OnTextChanged(this, &SAnimationRemapSkeleton::SetReplaceTo)
+						.IsReadOnly(false)
+						.RevertTextOnEscape(true)
+				]
+			]
+
+			+SVerticalBox::Slot()
+			.AutoHeight()
+			.Padding(2, 3)
+			[
+				SNew(SHorizontalBox)
+				+SHorizontalBox::Slot()
+				.Padding(5, 5)
+				[
+					SNew(STextBlock).Text(this,  &SAnimationRemapSkeleton::GetExampleText)
+						.Font(FEditorStyle::GetFontStyle("Persona.RetargetManager.ItalicFont"))
+				]
+			]
+
+			+SVerticalBox::Slot()
+			.AutoHeight()
+			.Padding(2, 3)
+			[
+				SNew(SHorizontalBox)
+
+				+SHorizontalBox::Slot()
+				.AutoWidth()
+				[
+					SNew(STextBlock)
+					.Text(LOCTEXT("RemapSkeleton_DupeName_Folder", "Folder "))
+						.Font(FEditorStyle::GetFontStyle("Persona.RetargetManager.SmallBoldFont"))
+				]
+
+				+SHorizontalBox::Slot()
+				.FillWidth(1)
+				.HAlign(HAlign_Center)
+				[
+					SNew(STextBlock).Text(this, &SAnimationRemapSkeleton::GetFolderPath)
+				]
+
+				+SHorizontalBox::Slot()
+				.AutoWidth()
+				[
+					SNew(SButton)
+					.HAlign(HAlign_Center)
+					.Text(LOCTEXT("RemapSkeleton_DupeName_ShowFolderOption", "Change..."))
+					.OnClicked(this, &SAnimationRemapSkeleton::ShowFolderOption)
+				]
+			]
+		];
+
+		OptionWidget->AddSlot()
+		[
+			NameOptionWidget
+		];
+	}
+
 	TSharedPtr<SToolTip> SkeletonTooltip = IDocumentation::Get()->CreateToolTip(FText::FromString("Pick a skeleton for this mesh"), NULL, FString("Shared/Editors/Persona"), FString("Skeleton"));
 
 	this->ChildSlot
@@ -494,33 +659,12 @@ void SAnimationRemapSkeleton::Construct( const FArguments& InArgs )
 					]
 				]
 
-				+SVerticalBox::Slot()
-				.AutoHeight()
-				.Padding(2, 10)
-				.HAlign(HAlign_Fill)
-				[
-					SNew(SSeparator)
-					.Orientation(Orient_Horizontal)
-				]
-
-				+SVerticalBox::Slot()
-				.AutoHeight()
-				.HAlign(HAlign_Center)
-				.Padding(5)
-				[
-					SNew(STextBlock)
-					.Font(FSlateFontInfo(FPaths::EngineContentDir() / TEXT("Slate/Fonts/Roboto-Regular.ttf"), 10))
-					.ColorAndOpacity(FLinearColor::Red)
-					.Text(LOCTEXT("RemapSkeleton_Warning_Title", "[WARNING]"))
-				]
-
 				+SVerticalBox::Slot()
 				.AutoHeight()
 				.Padding(5)
 				[
 					SNew(STextBlock)
-					.Font(FSlateFontInfo(FPaths::EngineContentDir() / TEXT("Slate/Fonts/Roboto-Regular.ttf"), 10))
-					.ColorAndOpacity(FLinearColor::Red)
+					.Font(FEditorStyle::GetFontStyle("Persona.RetargetManager.FilterFont"))
 					.Text(InArgs._WarningMessage)
 				]
 
@@ -536,12 +680,7 @@ void SAnimationRemapSkeleton::Construct( const FArguments& InArgs )
 				[
 					SAssignNew(AssetPickerBox, SBox)
 					.WidthOverride(400)
-				]
-
-				+SVerticalBox::Slot()
-				.AutoHeight()
-				[
-					Widget
+					.HeightOverride(300)
 				]
 
 				+SVerticalBox::Slot()
@@ -563,7 +702,7 @@ void SAnimationRemapSkeleton::Construct( const FArguments& InArgs )
 					+SUniformGridPanel::Slot(0, 0)
 					[
 						SNew(SButton).HAlign(HAlign_Center)
-						.Text(LOCTEXT("RemapSkeleton_Apply", "Select"))
+						.Text(LOCTEXT("RemapSkeleton_Apply", "Retarget"))
 						.IsEnabled(this, &SAnimationRemapSkeleton::CanApply)
 						.OnClicked(this, &SAnimationRemapSkeleton::OnApply)
 						.HAlign(HAlign_Center)
@@ -600,8 +739,9 @@ void SAnimationRemapSkeleton::Construct( const FArguments& InArgs )
 					// put nice message here
 					SNew(STextBlock)
 					.AutoWrapText(true)
-					.Font(FEditorStyle::GetFontStyle("Persona.RetargetManager.BoldFont"))
-					.Text(LOCTEXT("RetargetBasePose_WarningMessage", "Make sure you have the similar retarget base pose. If they don't look alike here, you can edit your base pose in the Retarget Manager window to look alike."))
+					.Font(FEditorStyle::GetFontStyle("Persona.RetargetManager.FilterFont"))
+					.ColorAndOpacity(FLinearColor::Red)
+					.Text(LOCTEXT("RetargetBasePose_WarningMessage", "*Make sure you have the similar retarget base pose. \nIf they don't look alike here, you can edit your base pose in the Retarget Manager window."))
 				]
 
 				+SVerticalBox::Slot()
@@ -609,7 +749,6 @@ void SAnimationRemapSkeleton::Construct( const FArguments& InArgs )
 				.Padding(0, 5)
 				[
 					SNew(SHorizontalBox)
-
 					+SHorizontalBox::Slot()
 					[
 						SNew(SVerticalBox)
@@ -618,19 +757,17 @@ void SAnimationRemapSkeleton::Construct( const FArguments& InArgs )
 						[
 							SNew(STextBlock)
 							.Text(LOCTEXT("SourceSkeleteonTitle", "[Source]"))
-							.Font(FEditorStyle::GetFontStyle("Persona.RetargetManager.FilterFont"))
+							.Font(FEditorStyle::GetFontStyle("Persona.RetargetManager.BoldFont"))
 							.AutoWrapText(true)
 						]
 
 						+ SVerticalBox::Slot()
 						.AutoHeight()
+						.Padding(5, 5)
 						[
 							SAssignNew(SourceViewport, SBasePoseViewport)
 							.Skeleton(OldSkeleton)
 						]
-						/*SAssignNew(SourceViewport, SBasePoseViewport)
-						.Title(TEXT("[Source]"))
-						.Skeleton(OldSkeleton)*/
 					]
 
 					+SHorizontalBox::Slot()
@@ -641,25 +778,63 @@ void SAnimationRemapSkeleton::Construct( const FArguments& InArgs )
 						[
 							SNew(STextBlock)
 							.Text(LOCTEXT("TargetSkeleteonTitle", "[Target]"))
-							.Font(FEditorStyle::GetFontStyle("Persona.RetargetManager.FilterFont"))
+							.Font(FEditorStyle::GetFontStyle("Persona.RetargetManager.BoldFont"))
 							.AutoWrapText(true)
 						]
 
-						+ SVerticalBox::Slot()
+						+SVerticalBox::Slot()
 						.AutoHeight()
+						.Padding(5, 5)
 						[
 							SAssignNew(TargetViewport, SBasePoseViewport)
 							.Skeleton(nullptr)
 						]
-						/*SAssignNew(TargetViewport, SBasePoseViewport)
-						.Title(TEXT("[Target]"))
-						.Skeleton(NULL)*/
+					]
+				]
+
+				+SVerticalBox::Slot()
+				.AutoHeight()
+				.Padding(0, 5)
+				[
+					SNew(SVerticalBox)
+					+SVerticalBox::Slot()
+					.AutoHeight()
+					.Padding(5, 5)
+					[
+						OptionWidget
 					]
 				]
 			]
 		];
 
 	UpdateAssetPicker();
+	UpdateExampleText();
+}
+
+FReply SAnimationRemapSkeleton::ShowFolderOption()
+{
+	TSharedRef<SSelectFolderDlg> NewAnimDlg =
+	SNew(SSelectFolderDlg);
+
+	if(NewAnimDlg->ShowModal() != EAppReturnType::Cancel)
+	{
+		NameDuplicateRule.FolderPath = NewAnimDlg->GetAssetPath();
+	}
+
+	if (WidgetWindow.IsValid())
+	{
+		WidgetWindow.Pin()->BringToFront(true);
+	}
+
+	return FReply::Handled();
+}
+
+void SAnimationRemapSkeleton::UpdateExampleText()
+{
+	FString ReplaceFrom = FString::Printf(TEXT("Old Name : ###%s###"), *NameDuplicateRule.ReplaceFrom);
+	FString ReplaceTo = FString::Printf(TEXT("New Name : %s###%s###%s"), *NameDuplicateRule.Prefix, *NameDuplicateRule.ReplaceTo, *NameDuplicateRule.Suffix);
+
+	ExampleText = FText::FromString(FString::Printf(TEXT("%s\n%s"), *ReplaceFrom, *ReplaceTo));
 }
 
 ECheckBoxState SAnimationRemapSkeleton::IsRemappingReferencedAssets() const
@@ -719,7 +894,7 @@ FReply SAnimationRemapSkeleton::OnApply()
 {
 	if (OnRetargetAnimationDelegate.IsBound())
 	{
-		OnRetargetAnimationDelegate.Execute(OldSkeleton, NewSkeleton, bRemapReferencedAssets, bConvertSpaces);
+		OnRetargetAnimationDelegate.Execute(OldSkeleton, NewSkeleton, bRemapReferencedAssets, bConvertSpaces, &NameDuplicateRule);
 	}
 
 	CloseWindow();
@@ -747,7 +922,7 @@ void SAnimationRemapSkeleton::CloseWindow()
 	}
 }
 
-void SAnimationRemapSkeleton::ShowWindow(USkeleton* OldSkeleton, const FText& WarningMessage, FOnRetargetAnimation RetargetDelegate)
+void SAnimationRemapSkeleton::ShowWindow(USkeleton* OldSkeleton, const FText& WarningMessage, bool bDuplicateAssets, FOnRetargetAnimation RetargetDelegate)
 {
 	if(DialogWindow.IsValid())
 	{
@@ -773,6 +948,7 @@ void SAnimationRemapSkeleton::ShowWindow(USkeleton* OldSkeleton, const FText& Wa
 			.ShowRemapOption(true)
 			.ShowConvertSpacesOption(OldSkeleton != NULL)
 			.ShowCompatibleDisplayOption(OldSkeleton != NULL)
+			.ShowDuplicateAssetOption(bDuplicateAssets)
 			.OnRetargetDelegate(RetargetDelegate)
 		];
 
@@ -1221,5 +1397,114 @@ TSharedPtr<SWidget> SBasePoseViewport::MakeViewportToolbar()
 {
 	return nullptr;
 }
+
+
+/////////////////////////////////////////////////
+// select folder dialog
+//////////////////////////////////////////////////
+void SSelectFolderDlg::Construct(const FArguments& InArgs)
+{
+	AssetPath = FText::FromString(FPackageName::GetLongPackagePath(InArgs._DefaultAssetPath.ToString()));
+
+	if(AssetPath.IsEmpty())
+	{
+		AssetPath = FText::FromString(TEXT("/Game"));
+	}
+
+	FPathPickerConfig PathPickerConfig;
+	PathPickerConfig.DefaultPath = AssetPath.ToString();
+	PathPickerConfig.OnPathSelected = FOnPathSelected::CreateSP(this, &SSelectFolderDlg::OnPathChange);
+	PathPickerConfig.bAddDefaultPath = true;
+
+	FContentBrowserModule& ContentBrowserModule = FModuleManager::LoadModuleChecked<FContentBrowserModule>("ContentBrowser");
+
+	SWindow::Construct(SWindow::FArguments()
+		.Title(LOCTEXT("SSelectFolderDlg_Title", "Create New Animation Object"))
+		.SupportsMinimize(false)
+		.SupportsMaximize(false)
+		//.SizingRule( ESizingRule::Autosized )
+		.ClientSize(FVector2D(450, 450))
+		[
+			SNew(SVerticalBox)
+
+			+ SVerticalBox::Slot() // Add user input block
+			.Padding(2)
+			[
+				SNew(SBorder)
+				.BorderImage(FEditorStyle::GetBrush("ToolPanel.GroupBorder"))
+				[
+					SNew(SVerticalBox)
+
+					+SVerticalBox::Slot()
+					.AutoHeight()
+					[
+						SNew(STextBlock)
+						.Text(LOCTEXT("SelectPath", "Select Path"))
+						.Font(FSlateFontInfo(FPaths::EngineContentDir() / TEXT("Slate/Fonts/Roboto-Regular.ttf"), 14))
+					]
+
+					+SVerticalBox::Slot()
+					.FillHeight(1)
+					.Padding(3)
+					[
+						ContentBrowserModule.Get().CreatePathPicker(PathPickerConfig)
+					]
+				]
+			]
+
+			+SVerticalBox::Slot()
+			.AutoHeight()
+			.HAlign(HAlign_Right)
+			.Padding(5)
+			[
+				SNew(SUniformGridPanel)
+				.SlotPadding(FEditorStyle::GetMargin("StandardDialog.SlotPadding"))
+				.MinDesiredSlotWidth(FEditorStyle::GetFloat("StandardDialog.MinDesiredSlotWidth"))
+				.MinDesiredSlotHeight(FEditorStyle::GetFloat("StandardDialog.MinDesiredSlotHeight"))
+				+SUniformGridPanel::Slot(0, 0)
+				[
+					SNew(SButton)
+					.HAlign(HAlign_Center)
+					.ContentPadding(FEditorStyle::GetMargin("StandardDialog.ContentPadding"))
+					.Text(LOCTEXT("OK", "OK"))
+					.OnClicked(this, &SSelectFolderDlg::OnButtonClick, EAppReturnType::Ok)
+				]
+				+SUniformGridPanel::Slot(1, 0)
+				[
+					SNew(SButton)
+					.HAlign(HAlign_Center)
+					.ContentPadding(FEditorStyle::GetMargin("StandardDialog.ContentPadding"))
+					.Text(LOCTEXT("Cancel", "Cancel"))
+					.OnClicked(this, &SSelectFolderDlg::OnButtonClick, EAppReturnType::Cancel)
+				]
+			]
+		]);
+}
+
+void SSelectFolderDlg::OnPathChange(const FString& NewPath)
+{
+	AssetPath = FText::FromString(NewPath);
+}
+
+FReply SSelectFolderDlg::OnButtonClick(EAppReturnType::Type ButtonID)
+{
+	UserResponse = ButtonID;
+
+	RequestDestroyWindow();
+
+	return FReply::Handled();
+}
+
+EAppReturnType::Type SSelectFolderDlg::ShowModal()
+{
+	GEditor->EditorAddModalWindow(SharedThis(this));
+	return UserResponse;
+}
+
+FString SSelectFolderDlg::GetAssetPath()
+{
+	return AssetPath.ToString();
+}
+
 #undef LOCTEXT_NAMESPACE 
 
diff --git a/Engine/Source/Editor/UnrealEd/Public/EditorAnimUtils.h b/Engine/Source/Editor/UnrealEd/Public/EditorAnimUtils.h
--- a/Engine/Source/Editor/UnrealEd/Public/EditorAnimUtils.h
+++ b/Engine/Source/Editor/UnrealEd/Public/EditorAnimUtils.h
@@ -4,11 +4,34 @@
 
 namespace EditorAnimUtils
 {
+	struct FNameDuplicationRule
+	{
+		FString Prefix;
+		FString Suffix;
+		// replace string
+		FString ReplaceFrom;
+		FString ReplaceTo;
+		// folder path
+		FString FolderPath;
+
+		FNameDuplicationRule() 
+			: Prefix(), Suffix(), ReplaceFrom(), ReplaceTo(), FolderPath(TEXT("/Game"))
+		{}
+
+		bool ShouldRename()
+		{
+			// should be able to replace "XXX" to ""
+			return ( !Prefix.IsEmpty() || !Suffix.IsEmpty() || !ReplaceFrom.IsEmpty());
+		}
+
+		FString Rename(const UObject* Asset) const;
+	};
+
 	class FAnimationRetargetContext
 	{
 	public:
-		FAnimationRetargetContext(const TArray<FAssetData>& AssetsToRetarget, bool bRetargetReferredAssets, bool bInConvertAnimationDataInComponentSpaces);
-		FAnimationRetargetContext(TArray<TWeakObjectPtr<UObject>> AssetsToRetarget, bool bRetargetReferredAssets, bool bInConvertAnimationDataInComponentSpaces);
+		FAnimationRetargetContext(const TArray<FAssetData>& AssetsToRetarget, bool bRetargetReferredAssets, bool bInConvertAnimationDataInComponentSpaces, const FNameDuplicationRule& NameRule=FNameDuplicationRule());
+		FAnimationRetargetContext(TArray<TWeakObjectPtr<UObject>> AssetsToRetarget, bool bRetargetReferredAssets, bool bInConvertAnimationDataInComponentSpaces, const FNameDuplicationRule& NameRule=FNameDuplicationRule());
 
 		/** Were we supplied anything that we can retarget */
 		bool HasAssetsToRetarget() const;
@@ -22,8 +45,11 @@ namespace EditorAnimUtils
 		/** Returns the duplicate of the supplied object if there is one, otherwise NULL */
 		UObject* GetDuplicate(const UObject* OriginalObject) const;
 
+		/** return all duplicates */
+		TArray<UObject*> GetAllDuplicates() const;
+
 		/** Duplicates the assets stored for retargetting, populating maps of original assets to new asset */
-		void DuplicateAssetsToRetarget(UPackage* DestinationPackage);
+		void DuplicateAssetsToRetarget(UPackage* DestinationPackage, const FNameDuplicationRule* NameRule);
 
 		/** Retarget the contained assets */
 		void RetargetAnimations(USkeleton* OldSkeleton, USkeleton* NewSkeleton);
@@ -58,7 +84,7 @@ namespace EditorAnimUtils
 	 * @param bDuplicatedAssetsBeforeRetarget	If true the assets are duplicated and then the duplicates are retargetted
 	 * @param bConvertSpace						Do the conversion in component space of the animation to match new target
 	 */
-	UNREALED_API UObject* RetargetAnimations(USkeleton* OldSkeleton, USkeleton* NewSkeleton, TArray<TWeakObjectPtr<UObject>> AssetsToRetarget, bool bRetargetReferredAssets, bool bDuplicateAssetsBeforeRetarget, bool bConvertSpace);
+	UNREALED_API UObject* RetargetAnimations(USkeleton* OldSkeleton, USkeleton* NewSkeleton, TArray<TWeakObjectPtr<UObject>> AssetsToRetarget, bool bRetargetReferredAssets, const FNameDuplicationRule* NameRule, bool bConvertSpace);
 
 	/**
 	 * Retargets the supplied FAssetDatas (as long as they are an animation asset), optionally duplicating them and retargetting their reference assets too
@@ -69,7 +95,7 @@ namespace EditorAnimUtils
 	 * @param bDuplicatedAssetsBeforeRetarget	If true the assets are duplicated and then the duplicates are retargetted
 	 * @param bConvertSpace						Do the conversion in component space of the animation to match new target
 	 */
-	UNREALED_API UObject* RetargetAnimations(USkeleton* OldSkeleton, USkeleton* NewSkeleton, const TArray<FAssetData>& AssetsToRetarget, bool bRetargetReferredAssets, bool bDuplicateAssetsBeforeRetarget, bool bConvertSpace);
+	UNREALED_API UObject* RetargetAnimations(USkeleton* OldSkeleton, USkeleton* NewSkeleton, const TArray<FAssetData>& AssetsToRetarget, bool bRetargetReferredAssets, const FNameDuplicationRule* NameRule, bool bConvertSpace);
 
 	/**
 	 * Retargets the supplied FAnimationRetargetContext, optionally duplicating the assets and retargetting the assets reference assets too. Is called by other overloads of RetargetAnimations
@@ -80,7 +106,7 @@ namespace EditorAnimUtils
 	 * @param bDuplicatedAssetsBeforeRetarget	If true the assets are duplicated and then the duplicates are retargetted
 	 * @param bConvertSpace						Do the conversion in component space of the animation to match new target
 	 */
-	UObject* RetargetAnimations(USkeleton* OldSkeleton, USkeleton* NewSkeleton, FAnimationRetargetContext& RetargetContext, bool bRetargetReferredAssets, bool bDuplicateAssetsBeforeRetarget);
+	UObject* RetargetAnimations(USkeleton* OldSkeleton, USkeleton* NewSkeleton, FAnimationRetargetContext& RetargetContext, bool bRetargetReferredAssets, const FNameDuplicationRule* NameRule);
 
 	// Populates the supplied TArrays with any animation assets that this blueprint refers too
 	void GetAllAnimationSequencesReferredInBlueprint(UAnimBlueprint* AnimBlueprint, TArray<UAnimationAsset*>& ComplexAnims, TArray<UAnimSequence*>& AnimSequences);
@@ -96,7 +122,7 @@ namespace EditorAnimUtils
 	 *
 	 * @return	TMap of original animation to duplicate
 	 */
-	TMap<UObject*, UObject*> DuplicateAssetsInternal(const TArray<UObject*>& AssetsToDuplicate, UPackage* DestinationPackage);
+	TMap<UObject*, UObject*> DuplicateAssetsInternal(const TArray<UObject*>& AssetsToDuplicate, UPackage* DestinationPackage, const FNameDuplicationRule* NameRule);
 
 	/**
 	 * Duplicates the supplied AssetsToDuplicate and returns a map of original asset to duplicate. Templated wrapper that calls DuplicateAssetInternal.
@@ -107,7 +133,7 @@ namespace EditorAnimUtils
 	 * @return	TMap of original animation to duplicate
 	 */
 	template<class AssetType>
-	TMap<AssetType*, AssetType*> DuplicateAssets(const TArray<AssetType*>& AssetsToDuplicate, UPackage* DestinationPackage)
+	TMap<AssetType*, AssetType*> DuplicateAssets(const TArray<AssetType*>& AssetsToDuplicate, UPackage* DestinationPackage, const FNameDuplicationRule* NameRule)
 	{
 		TArray<UObject*> Assets;
 		for(auto Iter = AssetsToDuplicate.CreateConstIterator(); Iter; ++Iter)
@@ -115,7 +141,7 @@ namespace EditorAnimUtils
 			Assets.Add(*Iter);
 		}
 
-		TMap<UObject*, UObject*> AssetMap = DuplicateAssetsInternal(Assets, DestinationPackage);
+		TMap<UObject*, UObject*> AssetMap = DuplicateAssetsInternal(Assets, DestinationPackage, NameRule);
 
 		TMap<AssetType*, AssetType*> ReturnMap;
 		for(auto Iter = AssetMap.CreateIterator(); Iter; ++Iter)
diff --git a/Engine/Source/Editor/UnrealEd/Public/SSkeletonWidget.h b/Engine/Source/Editor/UnrealEd/Public/SSkeletonWidget.h
--- a/Engine/Source/Editor/UnrealEd/Public/SSkeletonWidget.h
+++ b/Engine/Source/Editor/UnrealEd/Public/SSkeletonWidget.h
@@ -7,9 +7,12 @@
 #include "SlateBasics.h"
 #include "PreviewScene.h"
 #include "SEditorViewport.h"
+#include "EditorAnimUtils.h"
 
 class UAnimSet;
 
+using namespace EditorAnimUtils;
+
 /**
  * This below code is to select skeleton from the list 
  */
@@ -288,7 +291,7 @@ private:
  * Slate panel for choose Skeleton for assets to relink
  */
 
-DECLARE_DELEGATE_FourParams(FOnRetargetAnimation, USkeleton* /*OldSkeleton*/, USkeleton* /*NewSkeleton*/, bool /*bRemapReferencedAssets*/, bool /*bConvertSpaces*/)
+DECLARE_DELEGATE_FiveParams(FOnRetargetAnimation, USkeleton* /*OldSkeleton*/, USkeleton* /*NewSkeleton*/, bool /*bRemapReferencedAssets*/, bool /*bConvertSpaces*/, const FNameDuplicationRule* /*NameRule*/)
 
 class SAnimationRemapSkeleton : public SCompoundWidget
 {
@@ -308,6 +311,7 @@ public:
 		SLATE_ARGUMENT( bool, ShowRemapOption )
 		SLATE_ARGUMENT( bool, ShowConvertSpacesOption )
 		SLATE_ARGUMENT( bool, ShowCompatibleDisplayOption )
+		SLATE_ARGUMENT( bool, ShowDuplicateAssetOption )
 		SLATE_EVENT(FOnRetargetAnimation, OnRetargetDelegate)
 	SLATE_END_ARGS()	
 
@@ -391,6 +395,70 @@ private:
 	* Refreshes asset picker - call when asset picker option changes
 	*/
 	void UpdateAssetPicker();
+
+	/*
+	 * Duplicate Name Rule 
+	 */
+	FNameDuplicationRule NameDuplicateRule;
+	FText ExampleText;
+
+	FText GetPrefixName() const
+	{
+		return FText::FromString(NameDuplicateRule.Prefix);
+	}
+
+	void SetPrefixName(const FText &InText)
+	{
+		NameDuplicateRule.Prefix = InText.ToString();
+		UpdateExampleText();
+	}
+
+	FText GetSuffixName() const
+	{
+		return FText::FromString(NameDuplicateRule.Suffix);
+	}
+
+	void SetSuffixName(const FText &InText)
+	{
+		NameDuplicateRule.Suffix = InText.ToString();
+		UpdateExampleText();
+	}
+
+	FText GetReplaceFrom() const
+	{
+		return FText::FromString(NameDuplicateRule.ReplaceFrom);
+	}
+
+	void SetReplaceFrom(const FText &InText)
+	{
+		NameDuplicateRule.ReplaceFrom = InText.ToString();
+		UpdateExampleText();
+	}
+
+	FText GetReplaceTo() const
+	{
+		return FText::FromString(NameDuplicateRule.ReplaceTo);
+	}
+
+	void SetReplaceTo(const FText &InText)
+	{
+		NameDuplicateRule.ReplaceTo = InText.ToString();
+		UpdateExampleText();
+	}
+
+	FText GetExampleText() const
+	{
+		return ExampleText;
+	}
+
+	void UpdateExampleText();
+
+	FReply ShowFolderOption();
+	FText GetFolderPath() const
+	{
+		return FText::FromString(NameDuplicateRule.FolderPath);
+	}
+
 public:
 
 	/**
@@ -400,7 +468,7 @@ public:
 	 *
 	 * @return true if successfully selected new skeleton
 	 */
-	static UNREALED_API void ShowWindow(USkeleton* OldSkeleton, const FText& WarningMessage, FOnRetargetAnimation RetargetDelegate);
+	static UNREALED_API void ShowWindow(USkeleton* OldSkeleton, const FText& WarningMessage, bool bDuplicateAssets, FOnRetargetAnimation RetargetDelegate);
 
 	static TSharedPtr<SWindow> DialogWindow;
 };
@@ -521,6 +589,38 @@ public:
 	TArray< TSharedPtr<FName> > BoneNames;
 };
 
+class SSelectFolderDlg: public SWindow
+{
+public:
+	SLATE_BEGIN_ARGS(SSelectFolderDlg)
+	{
+	}
+
+	SLATE_ARGUMENT(FText, DefaultAssetPath)
+	SLATE_END_ARGS()
+
+		SSelectFolderDlg()
+		: UserResponse(EAppReturnType::Cancel)
+	{
+		}
+
+	void Construct(const FArguments& InArgs);
+
+public:
+	/** Displays the dialog in a blocking fashion */
+	EAppReturnType::Type ShowModal();
+
+	/** Gets the resulting asset path */
+	FString GetAssetPath();
+
+protected:
+	void OnPathChange(const FString& NewPath);
+	FReply OnButtonClick(EAppReturnType::Type ButtonID);
+
+	EAppReturnType::Type UserResponse;
+	FText AssetPath;
+};
+
 #undef LOCTEXT_NAMESPACE
 
 #endif // __SSkeletonWidget_h__
