Merging

//depot/UE4-Orion/Engine/Source/Editor/UnrealEd/Private/EditorAnimUtils.cpp

to //depot/UE4/Engine/Source/Editor/UnrealEd/Private/EditorAnimUtils.cpp

diff --git a/Engine/Source/Editor/UnrealEd/Private/EditorAnimUtils.cpp b/Engine/Source/Editor/UnrealEd/Private/EditorAnimUtils.cpp
--- a/Engine/Source/Editor/UnrealEd/Private/EditorAnimUtils.cpp
+++ b/Engine/Source/Editor/UnrealEd/Private/EditorAnimUtils.cpp
@@ -12,6 +12,7 @@
 #include "Editor/Persona/Public/PersonaModule.h"
 #include "ObjectEditorUtils.h"
 #include "SNotificationList.h"
+#include "Editor/ContentBrowser/Public/ContentBrowserModule.h"
 
 #define LOCTEXT_NAMESPACE "EditorAnimUtils"
 
