Editor: Added an editor setting to control how brushes get mirrored (either using scale or by directly modifying the brush geometry and leaving scale untouched)
[INTEGRATE] Change 2663677 by [REDACTED_PERSON][REDACTED_HANDLE] on 2015/08/20 19:21:15

diff --git a/Engine/Source/Editor/UnrealEd/Classes/Settings/LevelEditorMiscSettings.h b/Engine/Source/Editor/UnrealEd/Classes/Settings/LevelEditorMiscSettings.h
--- a/Engine/Source/Editor/UnrealEd/Classes/Settings/LevelEditorMiscSettings.h
+++ b/Engine/Source/Editor/UnrealEd/Classes/Settings/LevelEditorMiscSettings.h
@@ -38,6 +38,10 @@ public:
 	UPROPERTY(EditAnywhere, config, Category=Editing, meta=( DisplayName = "Preserve Actor Scale on Replace" ))
 	uint32 bReplaceRespectsScale:1;
 
+	/** If enabled, mirroring actions will directly affect the geometry of brushes rather than their scale */
+	UPROPERTY(EditAnywhere, config, Category=Editing, meta=( DisplayName = "Mirror Brushes Directly" ))
+	uint32 bMirrorBrushesWithoutScale:1;
+
 public:
 	/** If checked audio playing in the editor will continue to play even if the editor is in the background */
 	UPROPERTY(EditAnywhere, config, Category=Sound)
diff --git a/Engine/Source/Editor/UnrealEd/Private/UnrealEdSrv.cpp b/Engine/Source/Editor/UnrealEd/Private/UnrealEdSrv.cpp
--- a/Engine/Source/Editor/UnrealEd/Private/UnrealEdSrv.cpp
+++ b/Engine/Source/Editor/UnrealEd/Private/UnrealEdSrv.cpp
@@ -1778,7 +1778,9 @@ static void MirrorActors(const FVector& MirrorScale)
 	const FScopedTransaction Transaction( NSLOCTEXT("UnrealEd", "MirroringActors", "Mirroring Actors") );
 
 	// Fires ULevel::LevelDirtiedEvent when falling out of scope.
-	FScopedLevelDirtied		LevelDirtyCallback;
+	FScopedLevelDirtied LevelDirtyCallback;
+
+	const bool bUseScale = !GetDefault<ULevelEditorMiscSettings>()->bMirrorBrushesWithoutScale;
 
 	for ( FSelectionIterator It( GEditor->GetSelectedActorIterator() ) ; It ; ++It )
 	{
@@ -1786,26 +1788,75 @@ static void MirrorActors(const FVector& MirrorScale)
 		checkSlow( Actor->IsA(AActor::StaticClass()) );
 
 		const FVector PivotLocation = GLevelEditorModeTools().PivotLocation;
-		ABrush* Brush = Cast< ABrush >( Actor );
-		if( Brush && Brush->Brush )
-		{
-			// Brushes have to reverse their poly vertex order and recalculate the normal as negating one of the scale axes
-			// changes the handedness of the local transform.
-			Brush->Modify();
-			Brush->Brush->Modify();
-			Brush->Brush->Polys->Modify();
 
-			for( int32 poly = 0 ; poly < Brush->Brush->Polys->Element.Num() ; poly++ )
+		if (bUseScale)
+		{
+			ABrush* Brush = Cast< ABrush >(Actor);
+			if (Brush && Brush->Brush)
 			{
-				FPoly* Poly = &(Brush->Brush->Polys->Element[poly]);
+				// Brushes have to reverse their poly vertex order and recalculate the normal as negating one of the scale axes
+				// changes the handedness of the local transform.
+				Brush->Modify();
+				Brush->Brush->Modify();
+				Brush->Brush->Polys->Modify();
 
-				Poly->Reverse();
-				Poly->CalcNormal();
+				for (int32 poly = 0; poly < Brush->Brush->Polys->Element.Num(); poly++)
+				{
+					FPoly* Poly = &(Brush->Brush->Polys->Element[poly]);
+
+					Poly->Reverse();
+					Poly->CalcNormal();
+				}
 			}
+
+			Actor->Modify();
+			Actor->EditorApplyMirror(MirrorScale, PivotLocation);
 		}
+		else
+		{
+			ABrush* Brush = Cast< ABrush >( Actor );
+				if( Brush && Brush->Brush )
+			{
+				Brush->Modify();
+				Brush->Brush->Modify();
+				Brush->Brush->Polys->Modify();
 
-		Actor->Modify();
-		Actor->EditorApplyMirror( MirrorScale, PivotLocation );
+				const FVector LocalToWorldOffset = ( Brush->GetActorLocation() - PivotLocation );
+				const FVector LocationOffset = ( LocalToWorldOffset * MirrorScale ) - LocalToWorldOffset;
+
+				Brush->SetActorLocation( Brush->GetActorLocation() + LocationOffset, false );
+				Brush->SetPrePivot( Brush->GetPrePivot() * MirrorScale );
+
+				for( int32 poly = 0 ; poly < Brush->Brush->Polys->Element.Num() ; poly++ )
+				{
+					FPoly* Poly = &(Brush->Brush->Polys->Element[poly]);
+
+					Poly->TextureU *= MirrorScale;
+					Poly->TextureV *= MirrorScale;
+
+					Poly->Base += LocalToWorldOffset;
+					Poly->Base *= MirrorScale;
+					Poly->Base -= LocalToWorldOffset;
+					Poly->Base -= LocationOffset;
+
+					for( int32 vtx = 0 ; vtx < Poly->Vertices.Num(); vtx++ )
+					{
+						Poly->Vertices[vtx] += LocalToWorldOffset;
+						Poly->Vertices[vtx] *= MirrorScale;
+						Poly->Vertices[vtx] -= LocalToWorldOffset;
+						Poly->Vertices[vtx] -= LocationOffset;
+					}
+
+					Poly->Reverse();
+					Poly->CalcNormal();
+				}
+			}
+			else
+			{
+				Actor->Modify();
+				Actor->EditorApplyMirror( MirrorScale, PivotLocation );
+			}
+		}
 
 		Actor->InvalidateLightingCache();
 		Actor->PostEditMove( true );
