Invalidate cooked content on any ini change.  Added more cooking stats.

diff --git a/Engine/Source/Editor/UnrealEd/Private/CookOnTheFlyServer.cpp b/Engine/Source/Editor/UnrealEd/Private/CookOnTheFlyServer.cpp
--- a/Engine/Source/Editor/UnrealEd/Private/CookOnTheFlyServer.cpp
+++ b/Engine/Source/Editor/UnrealEd/Private/CookOnTheFlyServer.cpp
@@ -440,6 +440,39 @@ namespace CookOnTheFlyStats
 		return CookingStats;
 	}
 
+
+
+	class FCookerStat
+	{
+	private:
+		bool bStarted;
+		double StartTime;
+		const FName TagName;
+	public:
+		FCookerStat(const FName& InTagName) : TagName(InTagName)
+		{
+			StartTime = FPlatformTime::Seconds();
+			bStarted = true;
+		}
+
+		~FCookerStat()
+		{
+			Stop();
+		}
+
+		void Stop()
+		{
+			ICookingStats* CookingStats = GetCookingStats();
+			if (CookingStats)
+			{
+				static const FName NAME_CookerStat(TEXT("CookerStat"));
+				double Duration = FPlatformTime::Seconds() - StartTime * 1000.0f;
+				CookingStats->AddTagValue(NAME_CookerStat, TagName, FString::Printf(TEXT("%f"), (float)(Duration)));
+			}
+		}
+
+	};
+
 	class FCookerPackageStats
 	{
 	private:
@@ -601,6 +634,10 @@ namespace CookOnTheFlyStats
 	}
 }
 
+// global cooking stat which isn't associated with a specific package
+#define SCOPE_COOKING_STAT(name) \
+	static const FName NAME_##name(#name); \
+	CookOnTheFlyStats::FCookerStat Stat_##name(NAME_##name);
 
 #define START_PACKAGE(filename) CookOnTheFlyStats::FCookerPackageStats CookerPackageStats(filename);
 
@@ -623,7 +660,7 @@ namespace CookOnTheFlyStats
 
 #else
 
-
+#define SCOPE_COOKING_STAT(name)
 #define START_PACKAGE(filename)
 #define SCOPE_PACKAGE_STAT(name)
 #define START_PACKAGE_STAT(name)
@@ -858,6 +895,8 @@ TStatId UCookOnTheFlyServer::GetStatId() const
 
 bool UCookOnTheFlyServer::StartNetworkFileServer( const bool BindAnyPort )
 {
+	SCOPE_COOKING_STAT(StartNetworkFileServer);
+	
 	check( IsCookOnTheFlyMode() );
 	//GetDerivedDataCacheRef().WaitForQuiescence(false);
 
@@ -2653,6 +2692,8 @@ static bool IsMobileHDR()
 
 void UCookOnTheFlyServer::Initialize( ECookMode::Type DesiredCookMode, ECookInitializationFlags InCookFlags, const FString &InOutputDirectoryOverride )
 {
+	SCOPE_COOKING_STAT(Initialize);
+
 	OutputDirectoryOverride = InOutputDirectoryOverride;
 	CurrentCookMode = DesiredCookMode;
 	CookFlags = InCookFlags;
@@ -2784,24 +2825,59 @@ void GetVersionFormatNumbersForIniVersionStrings( TArray<FString>& IniVersionStr
 	}
 }
 
+#define INVALIDATE_ON_ANY_INI_CHANGE 1
+
 bool UCookOnTheFlyServer::GetCurrentIniVersionStrings( const ITargetPlatform* TargetPlatform, TArray<FString>& IniVersionStrings ) const
 {
+	SCOPE_COOKING_STAT(GetCurrentIniVersionStrings);
 	// there is a list of important ini settings in the Editor config 
+#if INVALIDATE_ON_ANY_INI_CHANGE // use a full list of ini settings to invalidate the cooked content or use the entire engine / editor ini
+	TArray<FString> IniFiles;
+	IniFiles.Add(FString(TEXT("Engine")));
+	IniFiles.Add(FString(TEXT("Editor")));
+	IniFiles.Add(FString(TEXT("EditorKeyBindings")));
+	IniFiles.Add(FString(TEXT("EditorLayout")));
+	IniFiles.Add(FString(TEXT("EditorSettings")));
+	IniFiles.Add(FString(TEXT("EditorPerProject")));
+	IniFiles.Add(FString(TEXT("Compat")));
+	IniFiles.Add(FString(TEXT("Lightmass")));
+	IniFiles.Add(FString(TEXT("Scalability")));
+	IniFiles.Add(FString(TEXT("Input")));
+	IniFiles.Add(FString(TEXT("GameUserSettings")));
+	IniFiles.Add(FString(TEXT("Game")));
+
+
+	for (const auto& IniFilename : IniFiles)
+	{
+		FConfigFile PlatformIniFile;
+		FConfigCacheIni::LoadLocalIniFile(PlatformIniFile, *IniFilename, true, *TargetPlatform->IniPlatformName());
+
+		for ( const auto& IniSection : PlatformIniFile)
+		{
+			const FString SectionName = FString::Printf(TEXT("%s:%s"), *IniFilename, *IniSection.Key);
+
+			for (const auto& IniValue : IniSection.Value)
+			{
+				IniVersionStrings.Add(FString::Printf(TEXT("%s:%s:%s"), *SectionName, *IniValue.Key.ToString(), *IniValue.Value));
+			}
+		}
+	}
+#else // invalidate only on some ini changes which we know cause problems (less full recooks but could be unstable)
 	TArray<FString> IniVersionedParams;
-	GConfig->GetArray( TEXT("CookSettings"), TEXT("VersionedIniParams"), IniVersionedParams, GEditorIni );
+	GConfig->GetArray(TEXT("CookSettings"), TEXT("VersionedIniParams"), IniVersionedParams, GEditorIni);
 
 
 	// used to store temporary platform specific ini files
-	TMap<FString,FConfigFile*> PlatformIniFiles;
+	TMap<FString, FConfigFile*> PlatformIniFiles;
 
 	// if the old one doesn't contain all the settings in the new one then we fail this check
-	for ( const auto& IniVersioned : IniVersionedParams )
+	for (const auto& IniVersioned : IniVersionedParams)
 	{
 
 		TArray<FString> IniVersionedArray;
 		IniVersioned.ParseIntoArray(IniVersionedArray, TEXT(":"), false);
 
-		if ( IniVersionedArray.Num() != 3 )
+		if (IniVersionedArray.Num() != 3)
 		{
 #if DEBUG_COOKONTHEFLY
 			UE_LOG(LogCook, Warning, TEXT("Invalid entry in CookSettings, VersionedIniParams %s"), *IniVersioned);
@@ -2814,17 +2890,17 @@ bool UCookOnTheFlyServer::GetCurrentIniVersionStrings( const ITargetPlatform* Ta
 		const FString& Key = IniVersionedArray[2];
 
 		// const FString& IniFilename = FPaths::GeneratedConfigDir() / TargetPlatform->IniPlatformName() / Filename + TEXT(".ini");
-		FConfigFile *PlatformIniFile = PlatformIniFiles.FindRef( Filename );
-		if ( PlatformIniFile == NULL )
+		FConfigFile *PlatformIniFile = PlatformIniFiles.FindRef(Filename);
+		if (PlatformIniFile == NULL)
 		{
 			PlatformIniFile = new FConfigFile();
-			FConfigCacheIni::LoadLocalIniFile( *PlatformIniFile, *Filename, true, *TargetPlatform->IniPlatformName() );
-			PlatformIniFiles.Add( Filename, PlatformIniFile );
+			FConfigCacheIni::LoadLocalIniFile(*PlatformIniFile, *Filename, true, *TargetPlatform->IniPlatformName());
+			PlatformIniFiles.Add(Filename, PlatformIniFile);
 		}
 
 		// get the value of the entry
 		FString Value;
-		if ( !PlatformIniFile->GetString(*Section, *Key, Value) )
+		if (!PlatformIniFile->GetString(*Section, *Key, Value))
 		{
 #if DEBUG_COOKONTHEFLY
 			UE_LOG(LogCook, Warning, TEXT("Unable to find entry in CookSettings, VersionedIniParams %s, assume default is being used"), *IniVersioned);
@@ -2832,12 +2908,18 @@ bool UCookOnTheFlyServer::GetCurrentIniVersionStrings( const ITargetPlatform* Ta
 			continue;
 		}
 
-		FString CurrentVersionString = FString::Printf(TEXT("%s:%s:%s:%s"), *Filename, *Section, *Key, *Value );
+		FString CurrentVersionString = FString::Printf(TEXT("%s:%s:%s:%s"), *Filename, *Section, *Key, *Value);
 
-		IniVersionStrings.Emplace( MoveTemp(CurrentVersionString) );
+		IniVersionStrings.Emplace(MoveTemp(CurrentVersionString));
 	}
 
-
+	// clean up our temporary platform ini files
+	for (const auto& PlatformIniFile : PlatformIniFiles)
+	{
+		delete PlatformIniFile.Value;
+	}
+	PlatformIniFiles.Empty();
+#endif
 	TArray<FString> VersionedRValues;
 	GConfig->GetArray(TEXT("CookSettings"), TEXT("VersionedIntRValues"), VersionedRValues, GEditorIni);
 
@@ -2852,14 +2934,6 @@ bool UCookOnTheFlyServer::GetCurrentIniVersionStrings( const ITargetPlatform* Ta
 	}
 
 
-	// clean up our temporary platform ini files
-	for ( const auto& PlatformIniFile : PlatformIniFiles )
-	{
-		delete PlatformIniFile.Value;
-	}
-	PlatformIniFiles.Empty();
-
-
 	const UTextureLODSettings& LodSettings = TargetPlatform->GetTextureLODSettings();
 
 	UEnum* TextureGroupEnum = FindObject<UEnum>( NULL, TEXT("Engine.TextureGroup") );
@@ -2911,7 +2985,7 @@ bool UCookOnTheFlyServer::GetCurrentIniVersionStrings( const ITargetPlatform* Ta
 
 bool UCookOnTheFlyServer::GetCookedIniVersionStrings( const ITargetPlatform* TargetPlatform, TArray<FString>& IniVersionStrings ) const
 {
-
+	SCOPE_COOKING_STAT(GetCookedIniVersionStrings);
 	const FString EditorIni = FPaths::GameDir() / TEXT("CookedIniVersion.txt");
 	const FString SandboxEditorIni = ConvertToFullSandboxPath(*EditorIni, true);
 
@@ -2924,6 +2998,8 @@ bool UCookOnTheFlyServer::GetCookedIniVersionStrings( const ITargetPlatform* Tar
 
 bool UCookOnTheFlyServer::CacheIniVersionStringsMap( const ITargetPlatform* TargetPlatform ) const
 {
+	SCOPE_COOKING_STAT(CacheIniVersionStringsMap);
+
 	// check if the cached ones are filled out
 	const FName TargetPlatformName = FName(*TargetPlatform->PlatformName());
 	TArray<FString>* FoundCookedIniVersionStrings = CachedIniVersionStringsMap.Find( TargetPlatformName );
@@ -2941,6 +3017,8 @@ bool UCookOnTheFlyServer::CacheIniVersionStringsMap( const ITargetPlatform* Targ
 
 bool UCookOnTheFlyServer::IniSettingsOutOfDate( const ITargetPlatform* TargetPlatform ) const
 {
+	SCOPE_COOKING_STAT(IniSettingsOutOfDate);
+
 	TArray<FString> CurrentIniVersionStrings;
 	if ( GetCurrentIniVersionStrings(TargetPlatform, CurrentIniVersionStrings) == false )
 	{
@@ -3244,6 +3322,7 @@ void UCookOnTheFlyServer::CleanSandbox( const bool bIterative )
 
 void UCookOnTheFlyServer::GenerateAssetRegistry(const TArray<ITargetPlatform*>& Platforms)
 {
+	SCOPE_COOKING_STAT(GenerateAssetRegistry);
 	if (IsChildCooker())
 	{
 		// don't generate the asset registry
@@ -3973,6 +4052,7 @@ void UCookOnTheFlyServer::CreateSandboxFile()
 
 void UCookOnTheFlyServer::InitializeSandbox()
 {
+	SCOPE_COOKING_STAT(InitializeSandbox);
 	if ( SandboxFile == NULL )
 	{
 		ITargetPlatformManagerModule& TPM = GetTargetPlatformManagerRef();
