Fix deprecation warnings

diff --git a/Engine/Source/Editor/UnrealEd/Private/UnrealEdSrv.cpp b/Engine/Source/Editor/UnrealEd/Private/UnrealEdSrv.cpp
--- a/Engine/Source/Editor/UnrealEd/Private/UnrealEdSrv.cpp
+++ b/Engine/Source/Editor/UnrealEd/Private/UnrealEdSrv.cpp
@@ -1825,7 +1825,7 @@ static void MirrorActors(const FVector& MirrorScale)
 				const FVector LocationOffset = ( LocalToWorldOffset * MirrorScale ) - LocalToWorldOffset;
 
 				Brush->SetActorLocation( Brush->GetActorLocation() + LocationOffset, false );
-				Brush->SetPrePivot( Brush->GetPrePivot() * MirrorScale );
+				Brush->SetPivotOffset( Brush->GetPivotOffset() * MirrorScale );
 
 				for( int32 poly = 0 ; poly < Brush->Brush->Polys->Element.Num() ; poly++ )
 				{
