Moving JSON tests out of UnrealEd and into Json module.

diff --git a/Engine/Source/Editor/UnrealEd/Private/Tests/EditorAutomationTests.cpp b/Engine/Source/Editor/UnrealEd/Private/Tests/EditorAutomationTests.cpp
--- a/Engine/Source/Editor/UnrealEd/Private/Tests/EditorAutomationTests.cpp
+++ b/Engine/Source/Editor/UnrealEd/Private/Tests/EditorAutomationTests.cpp
@@ -4,7 +4,6 @@
 #include "UnrealEd.h"
 #include "Tests/AutomationTestSettings.h"
 #include "ObjectTools.h"
-#include "Json.h"
 #include "AssetRegistryModule.h"
 #include "AssetToolsModule.h"
 #include "AssetSelection.h"
@@ -190,689 +189,6 @@ bool FGenericImportAssetsAutomationTest::RunTest(const FString& Parameters)
 
 //////////////////////////////////////////////////////////////////////////
 
-/**
- * FJsonAutomationTest
- * Simple unit test that runs Json's in-built test cases
- */
-IMPLEMENT_SIMPLE_AUTOMATION_TEST( FJsonAutomationTest, "System.Engine.FileSystem.JSON", EAutomationTestFlags::ATF_SmokeTest)
-
-typedef TJsonWriterFactory< TCHAR, TCondensedJsonPrintPolicy<TCHAR> > FCondensedJsonStringWriterFactory;
-typedef TJsonWriter< TCHAR, TCondensedJsonPrintPolicy<TCHAR> > FCondensedJsonStringWriter;
-
-typedef TJsonWriterFactory< TCHAR, TPrettyJsonPrintPolicy<TCHAR> > FPrettyJsonStringWriterFactory;
-typedef TJsonWriter< TCHAR, TPrettyJsonPrintPolicy<TCHAR> > FPrettyJsonStringWriter;
-
-/** 
- * Execute the Json test cases
- *
- * @return	true if the test was successful, false otherwise
- */
-bool FJsonAutomationTest::RunTest(const FString& Parameters)
-{
-	// Null Case
-	{
-		const FString InputString = TEXT("");
-		TSharedRef< TJsonReader<> > Reader = TJsonReaderFactory<>::Create( InputString );
-
-		TSharedPtr<FJsonObject> Object;
-		check( FJsonSerializer::Deserialize( Reader, Object ) == false );
-		check( !Object.IsValid() );
-	}
-
-	// Empty Object Case
-	{
-		const FString InputString = TEXT("{}");
-		TSharedRef< TJsonReader<> > Reader = TJsonReaderFactory<>::Create( InputString );
-
-		TSharedPtr<FJsonObject> Object;
-		check( FJsonSerializer::Deserialize( Reader, Object ) );
-		check( Object.IsValid() );
-
-		FString OutputString;
-		TSharedRef< FCondensedJsonStringWriter > Writer = FCondensedJsonStringWriterFactory::Create( &OutputString );
-		check( FJsonSerializer::Serialize( Object.ToSharedRef(), Writer ) );
-		check( InputString == OutputString );
-	}
-
-	// Empty Array Case
-	{
-		const FString InputString = TEXT("[]");
-		TSharedRef< TJsonReader<> > Reader = TJsonReaderFactory<>::Create( InputString );
-
-		TArray< TSharedPtr<FJsonValue> > Array;
-		check( FJsonSerializer::Deserialize( Reader, Array ) );
-		check( Array.Num() == 0 );
-
-		FString OutputString;
-		TSharedRef< FCondensedJsonStringWriter > Writer = FCondensedJsonStringWriterFactory::Create( &OutputString );
-		check( FJsonSerializer::Serialize( Array, Writer ) );
-		check( InputString == OutputString );
-	}
-
-	// Simple Array Case
-	{
-		const FString InputString = 
-			TEXT("[")
-			TEXT(	"{")
-			TEXT(		"\"Value\":\"Some String\"")
-			TEXT(	"}")
-			TEXT("]");
-
-		TSharedRef< TJsonReader<> > Reader = TJsonReaderFactory<>::Create( InputString );
-
-		TArray< TSharedPtr<FJsonValue> > Array;
-		bool bSuccessful = FJsonSerializer::Deserialize(Reader, Array);
-		check(bSuccessful);
-		check( Array.Num() == 1 );
-		check( Array[0].IsValid() )
-
-		TSharedPtr< FJsonObject > Object = Array[0]->AsObject();
-		check( Object.IsValid() );
-		check( Object->GetStringField( TEXT("Value") ) == TEXT("Some String") )
-
-		FString OutputString;
-		TSharedRef< FCondensedJsonStringWriter > Writer = FCondensedJsonStringWriterFactory::Create( &OutputString );
-		check( FJsonSerializer::Serialize( Array, Writer ) );
-		check( InputString == OutputString );
-	}
-
-	// Object Array Case
-	{
-		const FString InputString =
-			TEXT("[")
-			TEXT(	"{")
-			TEXT(		"\"Value\":\"Some String1\"")
-			TEXT(	"},")
-			TEXT(	"{")
-			TEXT(		"\"Value\":\"Some String2\"")
-			TEXT(	"},")
-			TEXT(	"{")
-			TEXT(		"\"Value\":\"Some String3\"")
-			TEXT(	"}")
-			TEXT("]");
-
-		TSharedRef< TJsonReader<> > Reader = TJsonReaderFactory<>::Create(InputString);
-
-		TArray< TSharedPtr<FJsonValue> > Array;
-
-		bool bSuccessful = FJsonSerializer::Deserialize(Reader, Array);
-		check(bSuccessful);
-		check(Array.Num() == 3);
-		check(Array[0].IsValid());
-		check(Array[1].IsValid());
-		check(Array[2].IsValid());
-
-		TSharedPtr< FJsonObject > Object = Array[0]->AsObject();
-		check(Object.IsValid());
-		check(Object->GetStringField(TEXT("Value")) == TEXT("Some String1"));
-
-		Object = Array[1]->AsObject();
-		check(Object.IsValid());
-		check(Object->GetStringField(TEXT("Value")) == TEXT("Some String2"));
-
-		Object = Array[2]->AsObject();
-		check(Object.IsValid());
-		check(Object->GetStringField(TEXT("Value")) == TEXT("Some String3"));
-
-		FString OutputString;
-		TSharedRef< FCondensedJsonStringWriter > Writer = FCondensedJsonStringWriterFactory::Create(&OutputString);
-		check(FJsonSerializer::Serialize(Array, Writer));
-		check(InputString == OutputString);
-	}
-
-	// Number Array Case
-	{
-		const FString InputString =
-			TEXT("[")
-			TEXT("10,")
-			TEXT("20,")
-			TEXT("30,")
-			TEXT("40")
-			TEXT("]");
-
-		TSharedRef< TJsonReader<> > Reader = TJsonReaderFactory<>::Create(InputString);
-
-		TArray< TSharedPtr<FJsonValue> > Array;
-		bool bSuccessful = FJsonSerializer::Deserialize(Reader, Array);
-		check(bSuccessful);
-		check(Array.Num() == 4);
-		check(Array[0].IsValid());
-		check(Array[1].IsValid());
-		check(Array[2].IsValid());
-		check(Array[3].IsValid());
-
-		double Number = Array[0]->AsNumber();
-		check(Number == 10);
-
-		Number = Array[1]->AsNumber();
-		check(Number == 20);
-
-		Number = Array[2]->AsNumber();
-		check(Number == 30);
-
-		Number = Array[3]->AsNumber();
-		check(Number == 40);
-
-		FString OutputString;
-		TSharedRef< FCondensedJsonStringWriter > Writer = FCondensedJsonStringWriterFactory::Create(&OutputString);
-		check(FJsonSerializer::Serialize(Array, Writer));
-		check(InputString == OutputString);
-	}
-
-	// String Array Case
-	{
-		const FString InputString =
-			TEXT("[")
-			TEXT("\"Some String1\",")
-			TEXT("\"Some String2\",")
-			TEXT("\"Some String3\",")
-			TEXT("\"Some String4\"")
-			TEXT("]");
-
-		TSharedRef< TJsonReader<> > Reader = TJsonReaderFactory<>::Create(InputString);
-
-		TArray< TSharedPtr<FJsonValue> > Array;
-		bool bSuccessful = FJsonSerializer::Deserialize(Reader, Array);
-		check(bSuccessful);
-		check(Array.Num() == 4);
-		check(Array[0].IsValid());
-		check(Array[1].IsValid());
-		check(Array[2].IsValid());
-		check(Array[3].IsValid());
-
-		FString Text = Array[0]->AsString();
-		check(Text == TEXT("Some String1"));
-
-		Text = Array[1]->AsString();
-		check(Text == TEXT("Some String2"));
-
-		Text = Array[2]->AsString();
-		check(Text == TEXT("Some String3"));
-
-		Text = Array[3]->AsString();
-		check(Text == TEXT("Some String4"));
-
-		FString OutputString;
-		TSharedRef< FCondensedJsonStringWriter > Writer = FCondensedJsonStringWriterFactory::Create(&OutputString);
-		check(FJsonSerializer::Serialize(Array, Writer));
-		check(InputString == OutputString);
-	}
-
-	// Complex Array Case
-	{
-		const FString InputString =
-			TEXT("[")
-			TEXT(	"\"Some String1\",")
-			TEXT(	"10,")
-			TEXT(	"{")
-			TEXT(		"\"Value\":\"Some String3\"")
-			TEXT(	"},")
-			TEXT(	"[")
-			TEXT(		"\"Some String4\",")
-			TEXT(		"\"Some String5\"")
-			TEXT(	"],")
-			TEXT(	"true,")
-			TEXT(	"null")
-			TEXT("]");
-
-		TSharedRef< TJsonReader<> > Reader = TJsonReaderFactory<>::Create(InputString);
-
-		TArray< TSharedPtr<FJsonValue> > Array;
-		bool bSuccessful = FJsonSerializer::Deserialize(Reader, Array);
-		check(bSuccessful);
-		check(Array.Num() == 6);
-		check(Array[0].IsValid());
-		check(Array[1].IsValid());
-		check(Array[2].IsValid());
-		check(Array[3].IsValid());
-		check(Array[4].IsValid());
-		check(Array[5].IsValid());
-
-		FString Text = Array[0]->AsString();
-		check(Text == TEXT("Some String1"));
-
-		double Number = Array[1]->AsNumber();
-		check(Number == 10);
-
-		TSharedPtr< FJsonObject > Object = Array[2]->AsObject();
-		check(Object.IsValid());
-		check(Object->GetStringField(TEXT("Value")) == TEXT("Some String3"));
-
-		const TArray<TSharedPtr< FJsonValue >>& InnerArray = Array[3]->AsArray();
-		check(InnerArray.Num() == 2);
-		check(Array[0].IsValid());
-		check(Array[1].IsValid());
-
-		Text = InnerArray[0]->AsString();
-		check(Text == TEXT("Some String4"));
-
-		Text = InnerArray[1]->AsString();
-		check(Text == TEXT("Some String5"));
-
-		bool Boolean = Array[4]->AsBool();
-		check(Boolean == true);
-
-		check(Array[5]->IsNull() == true);
-
-		FString OutputString;
-		TSharedRef< FCondensedJsonStringWriter > Writer = FCondensedJsonStringWriterFactory::Create(&OutputString);
-		check(FJsonSerializer::Serialize(Array, Writer));
-		check(InputString == OutputString);
-	}
-
-	// String Test
-	{
-		const FString InputString =
-			TEXT("{")
-			TEXT(	"\"Value\":\"Some String, Escape Chars: \\\\, \\\", \\/, \\b, \\f, \\n, \\r, \\t, \\u002B\"")
-			TEXT("}");
-		TSharedRef< TJsonReader<> > Reader = TJsonReaderFactory<>::Create( InputString );
-
-		TSharedPtr<FJsonObject> Object;
-		bool bSuccessful = FJsonSerializer::Deserialize(Reader, Object);
-		check(bSuccessful);
-		check( Object.IsValid() );
-
-		const TSharedPtr<FJsonValue>* Value = Object->Values.Find(TEXT("Value"));
-		check(Value && (*Value)->Type == EJson::String);
-		const FString String = (*Value)->AsString();
-		check(String == TEXT("Some String, Escape Chars: \\, \", /, \b, \f, \n, \r, \t, +"));
-
-		FString OutputString;
-		TSharedRef< FCondensedJsonStringWriter > Writer = FCondensedJsonStringWriterFactory::Create( &OutputString );
-		check( FJsonSerializer::Serialize( Object.ToSharedRef(), Writer ) );
-
-		const FString TestOutput =
-			TEXT("{")
-			TEXT(	"\"Value\":\"Some String, Escape Chars: \\\\, \\\", /, \\b, \\f, \\n, \\r, \\t, +\"")
-			TEXT("}");
-		check(OutputString == TestOutput);
-	}
-
-	// Number Test
-	{
-		const FString InputString =
-			TEXT("{")
-			TEXT(	"\"Value1\":2.544e+15,")
-			TEXT(	"\"Value2\":-0.544E-2,")
-			TEXT(	"\"Value3\":251e3,")
-			TEXT(	"\"Value4\":-0.0,")
-			TEXT(	"\"Value5\":843")
-			TEXT("}");
-		TSharedRef< TJsonReader<> > Reader = TJsonReaderFactory<>::Create( InputString );
-
-		TSharedPtr<FJsonObject> Object;
-		bool bSuccessful = FJsonSerializer::Deserialize(Reader, Object);
-		check(bSuccessful);
-		check( Object.IsValid() );
-
-		double TestValues[] = {2.544e+15, -0.544e-2, 251e3, -0.0, 843};
-		for (int32 i = 0; i < 5; ++i)
-		{
-			const TSharedPtr<FJsonValue>* Value = Object->Values.Find(FString::Printf(TEXT("Value%i"), i + 1));
-			check(Value && (*Value)->Type == EJson::Number);
-			const double Number = (*Value)->AsNumber();
-			check(Number == TestValues[i]);
-		}
-
-		FString OutputString;
-		TSharedRef< FCondensedJsonStringWriter > Writer = FCondensedJsonStringWriterFactory::Create( &OutputString );
-		check( FJsonSerializer::Serialize( Object.ToSharedRef(), Writer ) );
-
-		// %g isn't standardized, so we use the same %g format that is used inside PrintJson instead of hardcoding the values here
-		const FString TestOutput = FString::Printf(
-			TEXT("{")
-			TEXT(	"\"Value1\":%.17g,")
-			TEXT(	"\"Value2\":%.17g,")
-			TEXT(	"\"Value3\":%.17g,")
-			TEXT(	"\"Value4\":%.17g,")
-			TEXT(	"\"Value5\":%.17g")
-			TEXT("}"),
-			TestValues[0], TestValues[1], TestValues[2], TestValues[3], TestValues[4]);
-		check(OutputString == TestOutput);
-	}
-
-	// Boolean/Null Test
-	{
-		const FString InputString =
-			TEXT("{")
-			TEXT(	"\"Value1\":true,")
-			TEXT(	"\"Value2\":true,")
-			TEXT(	"\"Value3\":faLsE,")
-			TEXT(	"\"Value4\":null,")
-			TEXT(	"\"Value5\":NULL")
-			TEXT("}");
-		TSharedRef< TJsonReader<> > Reader = TJsonReaderFactory<>::Create( InputString );
-
-		TSharedPtr<FJsonObject> Object;
-		bool bSuccessful = FJsonSerializer::Deserialize(Reader, Object);
-		check(bSuccessful);
-		check( Object.IsValid() );
-
-		bool TestValues[] = {true, true, false};
-		for (int32 i = 0; i < 5; ++i)
-		{
-			const TSharedPtr<FJsonValue>* Value = Object->Values.Find(FString::Printf(TEXT("Value%i"), i + 1));
-			check(Value);
-			if (i < 3)
-			{
-				check((*Value)->Type == EJson::Boolean);
-				const bool Bool = (*Value)->AsBool();
-				check(Bool == TestValues[i]);
-			}
-			else
-			{
-				check((*Value)->Type == EJson::Null);
-				check((*Value)->IsNull());
-			}
-		}
-
-		FString OutputString;
-		TSharedRef< FCondensedJsonStringWriter > Writer = FCondensedJsonStringWriterFactory::Create( &OutputString );
-		check( FJsonSerializer::Serialize( Object.ToSharedRef(), Writer ) );
-
-		const FString TestOutput =
-			TEXT("{")
-			TEXT(	"\"Value1\":true,")
-			TEXT(	"\"Value2\":true,")
-			TEXT(	"\"Value3\":false,")
-			TEXT(	"\"Value4\":null,")
-			TEXT(	"\"Value5\":null")
-			TEXT("}");
-		check(OutputString == TestOutput);
-	}
-
-	// Object Test && extra whitespace test
-	{
-		const FString InputStringWithExtraWhitespace =
-			TEXT("		\n\r\n	   {")
-			TEXT(	"\"Object\":")
-			TEXT(	"{")
-			TEXT(		"\"NestedValue\":null,")
-			TEXT(		"\"NestedObject\":{}")
-			TEXT(	"},")
-			TEXT(	"\"Value\":true")
-			TEXT("}		\n\r\n	   ");
-
-		const FString InputString =
-			TEXT("{")
-			TEXT(	"\"Object\":")
-			TEXT(	"{")
-			TEXT(		"\"NestedValue\":null,")
-			TEXT(		"\"NestedObject\":{}")
-			TEXT(	"},")
-			TEXT(	"\"Value\":true")
-			TEXT("}");
-
-		TSharedRef< TJsonReader<> > Reader = TJsonReaderFactory<>::Create( InputStringWithExtraWhitespace );
-
-		TSharedPtr<FJsonObject> Object;
-		bool bSuccessful = FJsonSerializer::Deserialize(Reader, Object);
-		check(bSuccessful);
-		check( Object.IsValid() );
-
-		const TSharedPtr<FJsonValue>* InnerValueFail = Object->Values.Find(TEXT("InnerValue"));
-		check(!InnerValueFail);
-
-		const TSharedPtr<FJsonValue>* ObjectValue = Object->Values.Find(TEXT("Object"));
-		check(ObjectValue && (*ObjectValue)->Type == EJson::Object);
-		const TSharedPtr<FJsonObject> InnerObject = (*ObjectValue)->AsObject();
-		check(InnerObject.IsValid());
-
-		{
-			const TSharedPtr<FJsonValue>* NestedValueValue = InnerObject->Values.Find(TEXT("NestedValue"));
-			check(NestedValueValue && (*NestedValueValue)->Type == EJson::Null);
-			check((*NestedValueValue)->IsNull());
-
-			const TSharedPtr<FJsonValue>* NestedObjectValue = InnerObject->Values.Find(TEXT("NestedObject"));
-			check(NestedObjectValue && (*NestedObjectValue)->Type == EJson::Object);
-			const TSharedPtr<FJsonObject> InnerInnerObject = (*NestedObjectValue)->AsObject();
-			check(InnerInnerObject.IsValid());
-
-			{
-				const TSharedPtr<FJsonValue>* NestedValueValueFail = InnerInnerObject->Values.Find(TEXT("NestedValue"));
-				check(!NestedValueValueFail);
-			}
-		}
-
-		const TSharedPtr<FJsonValue>* ValueValue = Object->Values.Find(TEXT("Value"));
-		check(ValueValue && (*ValueValue)->Type == EJson::Boolean);
-		const bool Bool = (*ValueValue)->AsBool();
-		check(Bool);
-
-		FString OutputString;
-		TSharedRef< FCondensedJsonStringWriter > Writer = FCondensedJsonStringWriterFactory::Create( &OutputString );
-		check( FJsonSerializer::Serialize( Object.ToSharedRef(), Writer ) );
-		check(OutputString == InputString);
-	}
-
-	// Array Test
-	{
-		const FString InputString =
-			TEXT("{")
-			TEXT(	"\"Array\":")
-			TEXT(	"[")
-			TEXT(		"[],")
-			TEXT(		"\"Some String\",")
-			TEXT(		"\"Another String\",")
-			TEXT(		"null,")
-			TEXT(		"true,")
-			TEXT(		"false,")
-			TEXT(		"45,")
-			TEXT(		"{}")
-			TEXT(	"]")
-			TEXT("}");
-		TSharedRef< TJsonReader<> > Reader = TJsonReaderFactory<>::Create( InputString );
-
-		TSharedPtr<FJsonObject> Object;
-		bool bSuccessful = FJsonSerializer::Deserialize(Reader, Object);
-		check(bSuccessful);
-		check( Object.IsValid() );
-
-		const TSharedPtr<FJsonValue>* InnerValueFail = Object->Values.Find(TEXT("InnerValue"));
-		check(!InnerValueFail);
-
-		const TSharedPtr<FJsonValue>* ArrayValue = Object->Values.Find(TEXT("Array"));
-		check(ArrayValue && (*ArrayValue)->Type == EJson::Array);
-		const TArray< TSharedPtr<FJsonValue> > Array = (*ArrayValue)->AsArray();
-		check(Array.Num() == 8);
-
-		EJson ValueTypes[] = {EJson::Array, EJson::String, EJson::String, EJson::Null,
-			EJson::Boolean, EJson::Boolean, EJson::Number, EJson::Object};
-		for (int32 i = 0; i < Array.Num(); ++i)
-		{
-			const TSharedPtr<FJsonValue>& Value = Array[i];
-			check(Value.IsValid());
-			check(Value->Type == ValueTypes[i]);
-		}
-
-		const TArray< TSharedPtr<FJsonValue> >& InnerArray = Array[0]->AsArray();
-		check(InnerArray.Num() == 0);
-		check(Array[1]->AsString() == TEXT("Some String"));
-		check(Array[2]->AsString() == TEXT("Another String"));
-		check(Array[3]->IsNull());
-		check(Array[4]->AsBool());
-		check(!Array[5]->AsBool());
-		check(FMath::Abs(Array[6]->AsNumber() - 45.f) < KINDA_SMALL_NUMBER);
-		const TSharedPtr<FJsonObject> InnerObject = Array[7]->AsObject();
-		check(InnerObject.IsValid());
-
-		FString OutputString;
-		TSharedRef< FCondensedJsonStringWriter > Writer = FCondensedJsonStringWriterFactory::Create( &OutputString );
-		check( FJsonSerializer::Serialize( Object.ToSharedRef(), Writer ) );
-		check(OutputString == InputString);
-	}
-
-	// Pretty Print Test
-	{
-		const FString InputString =
-			TEXT("{")									LINE_TERMINATOR
-			TEXT("	\"Data1\": \"value\",")				LINE_TERMINATOR
-			TEXT("	\"Data2\": \"value\",")				LINE_TERMINATOR
-			TEXT("	\"Array\": [")						LINE_TERMINATOR
-			TEXT("		{")								LINE_TERMINATOR
-			TEXT("			\"InnerData1\": \"value\"")	LINE_TERMINATOR
-			TEXT("		},")							LINE_TERMINATOR
-			TEXT("		[],")							LINE_TERMINATOR
-			TEXT("		[ 1, 2, 3, 4 ],")				LINE_TERMINATOR
-			TEXT("		{")								LINE_TERMINATOR
-			TEXT("		},")							LINE_TERMINATOR
-			TEXT("		\"value\",")					LINE_TERMINATOR
-			TEXT("		\"value\"")						LINE_TERMINATOR
-			TEXT("	],")								LINE_TERMINATOR
-			TEXT("	\"Object\":")						LINE_TERMINATOR
-			TEXT("	{")									LINE_TERMINATOR
-			TEXT("	}")									LINE_TERMINATOR
-			TEXT("}");
-		TSharedRef< TJsonReader<> > Reader = TJsonReaderFactory<>::Create( InputString );
-
-		TSharedPtr<FJsonObject> Object;
-		check( FJsonSerializer::Deserialize( Reader, Object ) );
-		check( Object.IsValid() );
-
-		FString OutputString;
-		TSharedRef< FPrettyJsonStringWriter > Writer = FPrettyJsonStringWriterFactory::Create( &OutputString );
-		check( FJsonSerializer::Serialize( Object.ToSharedRef(), Writer ) );
-		check(OutputString == InputString);
-	}
-	  
-	// Line and Character # test
-	{
-		const FString InputString =
-			TEXT("{")									LINE_TERMINATOR
-			TEXT("	\"Data1\": \"value\",")				LINE_TERMINATOR
-			TEXT("	\"Array\":")						LINE_TERMINATOR
-			TEXT("	[")									LINE_TERMINATOR
-			TEXT("		12345,")						LINE_TERMINATOR
-			TEXT("		True")							LINE_TERMINATOR
-			TEXT("	],")								LINE_TERMINATOR
-			TEXT("	\"Object\":")						LINE_TERMINATOR
-			TEXT("	{")									LINE_TERMINATOR
-			TEXT("	}")									LINE_TERMINATOR
-			TEXT("}");
-		TSharedRef< TJsonReader<> > Reader = TJsonReaderFactory<>::Create( InputString );
-
-		EJsonNotation Notation = EJsonNotation::Null;
-		check( Reader->ReadNext( Notation ) && Notation == EJsonNotation::ObjectStart );
-		check( Reader->GetLineNumber() == 1 && Reader->GetCharacterNumber() == 1 );
-
-		check( Reader->ReadNext( Notation ) && Notation == EJsonNotation::String );
-		check( Reader->GetLineNumber() == 2 && Reader->GetCharacterNumber() == 17 );
-
-		check( Reader->ReadNext( Notation ) && Notation == EJsonNotation::ArrayStart );
-		check( Reader->GetLineNumber() == 4 && Reader->GetCharacterNumber() == 2 );
-
-		check( Reader->ReadNext( Notation ) && Notation == EJsonNotation::Number );
-		check( Reader->GetLineNumber() == 5 && Reader->GetCharacterNumber() == 7 );
-
-		check( Reader->ReadNext( Notation ) && Notation == EJsonNotation::Boolean );
-		check( Reader->GetLineNumber() == 6 && Reader->GetCharacterNumber() == 6 );
-	}
-
-	// Failure Cases
-	TArray<FString> FailureInputs;
-
-	// Unclosed Object
-	FailureInputs.Add(
-		TEXT("{"));
-
-	// Values in Object without identifiers
-	FailureInputs.Add(
-		TEXT("{")
-		TEXT(	"\"Value1\",")
-		TEXT(	"\"Value2\",")
-		TEXT(	"43")
-		TEXT("}"));
-
-	// Unexpected End Of Input Found
-	FailureInputs.Add(
-		TEXT("{")
-		TEXT(	"\"Object\":")
-		TEXT(	"{")
-		TEXT(		"\"NestedValue\":null,"));
-
-	// Missing first brace
-	FailureInputs.Add(
-		TEXT(	"\"Object\":")
-		TEXT(		"{")
-		TEXT(		"\"NestedValue\":null,")
-		TEXT(		"\"NestedObject\":{}")
-		TEXT(	"},")
-		TEXT(	"\"Value\":true")
-		TEXT("}"));
-
-	// Missing last character
-	FailureInputs.Add(
-		TEXT("{")
-		TEXT(	"\"Object\":")
-		TEXT(	"{")
-		TEXT(		"\"NestedValue\":null,")
-		TEXT(		"\"NestedObject\":{}")
-		TEXT(	"},")
-		TEXT(	"\"Value\":true"));
-
-	// Extra last character
-	FailureInputs.Add(
-		TEXT("{")
-		TEXT(	"\"Object\":")
-		TEXT(	"{")
-		TEXT(		"\"NestedValue\":null,")
-		TEXT(		"\"NestedObject\":{}")
-		TEXT(	"},")
-		TEXT(	"\"Value\":true")
-		TEXT("}0"));
-
-	// Missing comma
-	FailureInputs.Add(
-		TEXT("{")
-		TEXT(	"\"Value1\":null,")
-		TEXT(	"\"Value2\":\"string\"")
-		TEXT(	"\"Value3\":65.3")
-		TEXT("}"));
-
-	// Extra comma
-	FailureInputs.Add(
-		TEXT("{")
-		TEXT(	"\"Value1\":null,")
-		TEXT(	"\"Value2\":\"string\",")
-		TEXT(	"\"Value3\":65.3,")
-		TEXT("}"));
-
-	// Badly formed true/false/null
-	FailureInputs.Add(TEXT("{\"Value\":tru}"));
-	FailureInputs.Add(TEXT("{\"Value\":full}"));
-	FailureInputs.Add(TEXT("{\"Value\":nulle}"));
-	FailureInputs.Add(TEXT("{\"Value\":n%ll}"));
-
-	// Floating Point Failures
-	FailureInputs.Add(TEXT("{\"Value\":65.3e}"));
-	FailureInputs.Add(TEXT("{\"Value\":65.}"));
-	FailureInputs.Add(TEXT("{\"Value\":.7}"));
-	FailureInputs.Add(TEXT("{\"Value\":+6}"));
-	FailureInputs.Add(TEXT("{\"Value\":01}"));
-	FailureInputs.Add(TEXT("{\"Value\":00.56}"));
-	FailureInputs.Add(TEXT("{\"Value\":-1.e+4}"));
-	FailureInputs.Add(TEXT("{\"Value\":2e+}"));
-
-	// Bad Escape Characters
-	FailureInputs.Add(TEXT("{\"Value\":\"Hello\\xThere\"}"));
-	FailureInputs.Add(TEXT("{\"Value\":\"Hello\\u123There\"}"));
-	FailureInputs.Add(TEXT("{\"Value\":\"Hello\\RThere\"}"));
-
-	for (int32 i = 0; i < FailureInputs.Num(); ++i)
-	{
-		TSharedRef< TJsonReader<> > Reader = TJsonReaderFactory<>::Create( FailureInputs[i] );
-
-		TSharedPtr<FJsonObject> Object;
-		check( FJsonSerializer::Deserialize( Reader, Object ) == false );
-		check( !Object.IsValid() );
-	}
-
-	return true;
-}
-
-//////////////////////////////////////////////////////////////////////////
-
 /**
  * Pie Test
  * Verification PIE works
diff --git a/Engine/Source/Runtime/Json/Private/Tests/JsonTests.cpp b/Engine/Source/Runtime/Json/Private/Tests/JsonTests.cpp
new file mode 100644
--- /dev/null
+++ b/Engine/Source/Runtime/Json/Private/Tests/JsonTests.cpp
@@ -0,0 +1,685 @@
+// Copyright 1998-2015 Epic Games, Inc. All Rights Reserved.
+
+#include "JsonPrivatePCH.h"
+#include "Json.h"
+
+/**
+ * FJsonAutomationTest
+ * Simple unit test that runs Json's in-built test cases
+ */
+IMPLEMENT_SIMPLE_AUTOMATION_TEST( FJsonAutomationTest, "System.Engine.FileSystem.JSON", EAutomationTestFlags::ATF_SmokeTest)
+
+typedef TJsonWriterFactory< TCHAR, TCondensedJsonPrintPolicy<TCHAR> > FCondensedJsonStringWriterFactory;
+typedef TJsonWriter< TCHAR, TCondensedJsonPrintPolicy<TCHAR> > FCondensedJsonStringWriter;
+
+typedef TJsonWriterFactory< TCHAR, TPrettyJsonPrintPolicy<TCHAR> > FPrettyJsonStringWriterFactory;
+typedef TJsonWriter< TCHAR, TPrettyJsonPrintPolicy<TCHAR> > FPrettyJsonStringWriter;
+
+/** 
+ * Execute the Json test cases
+ *
+ * @return	true if the test was successful, false otherwise
+ */
+bool FJsonAutomationTest::RunTest(const FString& Parameters)
+{
+	// Null Case
+	{
+		const FString InputString = TEXT("");
+		TSharedRef< TJsonReader<> > Reader = TJsonReaderFactory<>::Create( InputString );
+
+		TSharedPtr<FJsonObject> Object;
+		check( FJsonSerializer::Deserialize( Reader, Object ) == false );
+		check( !Object.IsValid() );
+	}
+
+	// Empty Object Case
+	{
+		const FString InputString = TEXT("{}");
+		TSharedRef< TJsonReader<> > Reader = TJsonReaderFactory<>::Create( InputString );
+
+		TSharedPtr<FJsonObject> Object;
+		check( FJsonSerializer::Deserialize( Reader, Object ) );
+		check( Object.IsValid() );
+
+		FString OutputString;
+		TSharedRef< FCondensedJsonStringWriter > Writer = FCondensedJsonStringWriterFactory::Create( &OutputString );
+		check( FJsonSerializer::Serialize( Object.ToSharedRef(), Writer ) );
+		check( InputString == OutputString );
+	}
+
+	// Empty Array Case
+	{
+		const FString InputString = TEXT("[]");
+		TSharedRef< TJsonReader<> > Reader = TJsonReaderFactory<>::Create( InputString );
+
+		TArray< TSharedPtr<FJsonValue> > Array;
+		check( FJsonSerializer::Deserialize( Reader, Array ) );
+		check( Array.Num() == 0 );
+
+		FString OutputString;
+		TSharedRef< FCondensedJsonStringWriter > Writer = FCondensedJsonStringWriterFactory::Create( &OutputString );
+		check( FJsonSerializer::Serialize( Array, Writer ) );
+		check( InputString == OutputString );
+	}
+
+	// Simple Array Case
+	{
+		const FString InputString = 
+			TEXT("[")
+			TEXT(	"{")
+			TEXT(		"\"Value\":\"Some String\"")
+			TEXT(	"}")
+			TEXT("]");
+
+		TSharedRef< TJsonReader<> > Reader = TJsonReaderFactory<>::Create( InputString );
+
+		TArray< TSharedPtr<FJsonValue> > Array;
+		bool bSuccessful = FJsonSerializer::Deserialize(Reader, Array);
+		check(bSuccessful);
+		check( Array.Num() == 1 );
+		check( Array[0].IsValid() )
+
+		TSharedPtr< FJsonObject > Object = Array[0]->AsObject();
+		check( Object.IsValid() );
+		check( Object->GetStringField( TEXT("Value") ) == TEXT("Some String") )
+
+		FString OutputString;
+		TSharedRef< FCondensedJsonStringWriter > Writer = FCondensedJsonStringWriterFactory::Create( &OutputString );
+		check( FJsonSerializer::Serialize( Array, Writer ) );
+		check( InputString == OutputString );
+	}
+
+	// Object Array Case
+	{
+		const FString InputString =
+			TEXT("[")
+			TEXT(	"{")
+			TEXT(		"\"Value\":\"Some String1\"")
+			TEXT(	"},")
+			TEXT(	"{")
+			TEXT(		"\"Value\":\"Some String2\"")
+			TEXT(	"},")
+			TEXT(	"{")
+			TEXT(		"\"Value\":\"Some String3\"")
+			TEXT(	"}")
+			TEXT("]");
+
+		TSharedRef< TJsonReader<> > Reader = TJsonReaderFactory<>::Create(InputString);
+
+		TArray< TSharedPtr<FJsonValue> > Array;
+
+		bool bSuccessful = FJsonSerializer::Deserialize(Reader, Array);
+		check(bSuccessful);
+		check(Array.Num() == 3);
+		check(Array[0].IsValid());
+		check(Array[1].IsValid());
+		check(Array[2].IsValid());
+
+		TSharedPtr< FJsonObject > Object = Array[0]->AsObject();
+		check(Object.IsValid());
+		check(Object->GetStringField(TEXT("Value")) == TEXT("Some String1"));
+
+		Object = Array[1]->AsObject();
+		check(Object.IsValid());
+		check(Object->GetStringField(TEXT("Value")) == TEXT("Some String2"));
+
+		Object = Array[2]->AsObject();
+		check(Object.IsValid());
+		check(Object->GetStringField(TEXT("Value")) == TEXT("Some String3"));
+
+		FString OutputString;
+		TSharedRef< FCondensedJsonStringWriter > Writer = FCondensedJsonStringWriterFactory::Create(&OutputString);
+		check(FJsonSerializer::Serialize(Array, Writer));
+		check(InputString == OutputString);
+	}
+
+	// Number Array Case
+	{
+		const FString InputString =
+			TEXT("[")
+			TEXT("10,")
+			TEXT("20,")
+			TEXT("30,")
+			TEXT("40")
+			TEXT("]");
+
+		TSharedRef< TJsonReader<> > Reader = TJsonReaderFactory<>::Create(InputString);
+
+		TArray< TSharedPtr<FJsonValue> > Array;
+		bool bSuccessful = FJsonSerializer::Deserialize(Reader, Array);
+		check(bSuccessful);
+		check(Array.Num() == 4);
+		check(Array[0].IsValid());
+		check(Array[1].IsValid());
+		check(Array[2].IsValid());
+		check(Array[3].IsValid());
+
+		double Number = Array[0]->AsNumber();
+		check(Number == 10);
+
+		Number = Array[1]->AsNumber();
+		check(Number == 20);
+
+		Number = Array[2]->AsNumber();
+		check(Number == 30);
+
+		Number = Array[3]->AsNumber();
+		check(Number == 40);
+
+		FString OutputString;
+		TSharedRef< FCondensedJsonStringWriter > Writer = FCondensedJsonStringWriterFactory::Create(&OutputString);
+		check(FJsonSerializer::Serialize(Array, Writer));
+		check(InputString == OutputString);
+	}
+
+	// String Array Case
+	{
+		const FString InputString =
+			TEXT("[")
+			TEXT("\"Some String1\",")
+			TEXT("\"Some String2\",")
+			TEXT("\"Some String3\",")
+			TEXT("\"Some String4\"")
+			TEXT("]");
+
+		TSharedRef< TJsonReader<> > Reader = TJsonReaderFactory<>::Create(InputString);
+
+		TArray< TSharedPtr<FJsonValue> > Array;
+		bool bSuccessful = FJsonSerializer::Deserialize(Reader, Array);
+		check(bSuccessful);
+		check(Array.Num() == 4);
+		check(Array[0].IsValid());
+		check(Array[1].IsValid());
+		check(Array[2].IsValid());
+		check(Array[3].IsValid());
+
+		FString Text = Array[0]->AsString();
+		check(Text == TEXT("Some String1"));
+
+		Text = Array[1]->AsString();
+		check(Text == TEXT("Some String2"));
+
+		Text = Array[2]->AsString();
+		check(Text == TEXT("Some String3"));
+
+		Text = Array[3]->AsString();
+		check(Text == TEXT("Some String4"));
+
+		FString OutputString;
+		TSharedRef< FCondensedJsonStringWriter > Writer = FCondensedJsonStringWriterFactory::Create(&OutputString);
+		check(FJsonSerializer::Serialize(Array, Writer));
+		check(InputString == OutputString);
+	}
+
+	// Complex Array Case
+	{
+		const FString InputString =
+			TEXT("[")
+			TEXT(	"\"Some String1\",")
+			TEXT(	"10,")
+			TEXT(	"{")
+			TEXT(		"\"Value\":\"Some String3\"")
+			TEXT(	"},")
+			TEXT(	"[")
+			TEXT(		"\"Some String4\",")
+			TEXT(		"\"Some String5\"")
+			TEXT(	"],")
+			TEXT(	"true,")
+			TEXT(	"null")
+			TEXT("]");
+
+		TSharedRef< TJsonReader<> > Reader = TJsonReaderFactory<>::Create(InputString);
+
+		TArray< TSharedPtr<FJsonValue> > Array;
+		bool bSuccessful = FJsonSerializer::Deserialize(Reader, Array);
+		check(bSuccessful);
+		check(Array.Num() == 6);
+		check(Array[0].IsValid());
+		check(Array[1].IsValid());
+		check(Array[2].IsValid());
+		check(Array[3].IsValid());
+		check(Array[4].IsValid());
+		check(Array[5].IsValid());
+
+		FString Text = Array[0]->AsString();
+		check(Text == TEXT("Some String1"));
+
+		double Number = Array[1]->AsNumber();
+		check(Number == 10);
+
+		TSharedPtr< FJsonObject > Object = Array[2]->AsObject();
+		check(Object.IsValid());
+		check(Object->GetStringField(TEXT("Value")) == TEXT("Some String3"));
+
+		const TArray<TSharedPtr< FJsonValue >>& InnerArray = Array[3]->AsArray();
+		check(InnerArray.Num() == 2);
+		check(Array[0].IsValid());
+		check(Array[1].IsValid());
+
+		Text = InnerArray[0]->AsString();
+		check(Text == TEXT("Some String4"));
+
+		Text = InnerArray[1]->AsString();
+		check(Text == TEXT("Some String5"));
+
+		bool Boolean = Array[4]->AsBool();
+		check(Boolean == true);
+
+		check(Array[5]->IsNull() == true);
+
+		FString OutputString;
+		TSharedRef< FCondensedJsonStringWriter > Writer = FCondensedJsonStringWriterFactory::Create(&OutputString);
+		check(FJsonSerializer::Serialize(Array, Writer));
+		check(InputString == OutputString);
+	}
+
+	// String Test
+	{
+		const FString InputString =
+			TEXT("{")
+			TEXT(	"\"Value\":\"Some String, Escape Chars: \\\\, \\\", \\/, \\b, \\f, \\n, \\r, \\t, \\u002B\"")
+			TEXT("}");
+		TSharedRef< TJsonReader<> > Reader = TJsonReaderFactory<>::Create( InputString );
+
+		TSharedPtr<FJsonObject> Object;
+		bool bSuccessful = FJsonSerializer::Deserialize(Reader, Object);
+		check(bSuccessful);
+		check( Object.IsValid() );
+
+		const TSharedPtr<FJsonValue>* Value = Object->Values.Find(TEXT("Value"));
+		check(Value && (*Value)->Type == EJson::String);
+		const FString String = (*Value)->AsString();
+		check(String == TEXT("Some String, Escape Chars: \\, \", /, \b, \f, \n, \r, \t, +"));
+
+		FString OutputString;
+		TSharedRef< FCondensedJsonStringWriter > Writer = FCondensedJsonStringWriterFactory::Create( &OutputString );
+		check( FJsonSerializer::Serialize( Object.ToSharedRef(), Writer ) );
+
+		const FString TestOutput =
+			TEXT("{")
+			TEXT(	"\"Value\":\"Some String, Escape Chars: \\\\, \\\", /, \\b, \\f, \\n, \\r, \\t, +\"")
+			TEXT("}");
+		check(OutputString == TestOutput);
+	}
+
+	// Number Test
+	{
+		const FString InputString =
+			TEXT("{")
+			TEXT(	"\"Value1\":2.544e+15,")
+			TEXT(	"\"Value2\":-0.544E-2,")
+			TEXT(	"\"Value3\":251e3,")
+			TEXT(	"\"Value4\":-0.0,")
+			TEXT(	"\"Value5\":843")
+			TEXT("}");
+		TSharedRef< TJsonReader<> > Reader = TJsonReaderFactory<>::Create( InputString );
+
+		TSharedPtr<FJsonObject> Object;
+		bool bSuccessful = FJsonSerializer::Deserialize(Reader, Object);
+		check(bSuccessful);
+		check( Object.IsValid() );
+
+		double TestValues[] = {2.544e+15, -0.544e-2, 251e3, -0.0, 843};
+		for (int32 i = 0; i < 5; ++i)
+		{
+			const TSharedPtr<FJsonValue>* Value = Object->Values.Find(FString::Printf(TEXT("Value%i"), i + 1));
+			check(Value && (*Value)->Type == EJson::Number);
+			const double Number = (*Value)->AsNumber();
+			check(Number == TestValues[i]);
+		}
+
+		FString OutputString;
+		TSharedRef< FCondensedJsonStringWriter > Writer = FCondensedJsonStringWriterFactory::Create( &OutputString );
+		check( FJsonSerializer::Serialize( Object.ToSharedRef(), Writer ) );
+
+		// %g isn't standardized, so we use the same %g format that is used inside PrintJson instead of hardcoding the values here
+		const FString TestOutput = FString::Printf(
+			TEXT("{")
+			TEXT(	"\"Value1\":%.17g,")
+			TEXT(	"\"Value2\":%.17g,")
+			TEXT(	"\"Value3\":%.17g,")
+			TEXT(	"\"Value4\":%.17g,")
+			TEXT(	"\"Value5\":%.17g")
+			TEXT("}"),
+			TestValues[0], TestValues[1], TestValues[2], TestValues[3], TestValues[4]);
+		check(OutputString == TestOutput);
+	}
+
+	// Boolean/Null Test
+	{
+		const FString InputString =
+			TEXT("{")
+			TEXT(	"\"Value1\":true,")
+			TEXT(	"\"Value2\":true,")
+			TEXT(	"\"Value3\":faLsE,")
+			TEXT(	"\"Value4\":null,")
+			TEXT(	"\"Value5\":NULL")
+			TEXT("}");
+		TSharedRef< TJsonReader<> > Reader = TJsonReaderFactory<>::Create( InputString );
+
+		TSharedPtr<FJsonObject> Object;
+		bool bSuccessful = FJsonSerializer::Deserialize(Reader, Object);
+		check(bSuccessful);
+		check( Object.IsValid() );
+
+		bool TestValues[] = {true, true, false};
+		for (int32 i = 0; i < 5; ++i)
+		{
+			const TSharedPtr<FJsonValue>* Value = Object->Values.Find(FString::Printf(TEXT("Value%i"), i + 1));
+			check(Value);
+			if (i < 3)
+			{
+				check((*Value)->Type == EJson::Boolean);
+				const bool Bool = (*Value)->AsBool();
+				check(Bool == TestValues[i]);
+			}
+			else
+			{
+				check((*Value)->Type == EJson::Null);
+				check((*Value)->IsNull());
+			}
+		}
+
+		FString OutputString;
+		TSharedRef< FCondensedJsonStringWriter > Writer = FCondensedJsonStringWriterFactory::Create( &OutputString );
+		check( FJsonSerializer::Serialize( Object.ToSharedRef(), Writer ) );
+
+		const FString TestOutput =
+			TEXT("{")
+			TEXT(	"\"Value1\":true,")
+			TEXT(	"\"Value2\":true,")
+			TEXT(	"\"Value3\":false,")
+			TEXT(	"\"Value4\":null,")
+			TEXT(	"\"Value5\":null")
+			TEXT("}");
+		check(OutputString == TestOutput);
+	}
+
+	// Object Test && extra whitespace test
+	{
+		const FString InputStringWithExtraWhitespace =
+			TEXT("		\n\r\n	   {")
+			TEXT(	"\"Object\":")
+			TEXT(	"{")
+			TEXT(		"\"NestedValue\":null,")
+			TEXT(		"\"NestedObject\":{}")
+			TEXT(	"},")
+			TEXT(	"\"Value\":true")
+			TEXT("}		\n\r\n	   ");
+
+		const FString InputString =
+			TEXT("{")
+			TEXT(	"\"Object\":")
+			TEXT(	"{")
+			TEXT(		"\"NestedValue\":null,")
+			TEXT(		"\"NestedObject\":{}")
+			TEXT(	"},")
+			TEXT(	"\"Value\":true")
+			TEXT("}");
+
+		TSharedRef< TJsonReader<> > Reader = TJsonReaderFactory<>::Create( InputStringWithExtraWhitespace );
+
+		TSharedPtr<FJsonObject> Object;
+		bool bSuccessful = FJsonSerializer::Deserialize(Reader, Object);
+		check(bSuccessful);
+		check( Object.IsValid() );
+
+		const TSharedPtr<FJsonValue>* InnerValueFail = Object->Values.Find(TEXT("InnerValue"));
+		check(!InnerValueFail);
+
+		const TSharedPtr<FJsonValue>* ObjectValue = Object->Values.Find(TEXT("Object"));
+		check(ObjectValue && (*ObjectValue)->Type == EJson::Object);
+		const TSharedPtr<FJsonObject> InnerObject = (*ObjectValue)->AsObject();
+		check(InnerObject.IsValid());
+
+		{
+			const TSharedPtr<FJsonValue>* NestedValueValue = InnerObject->Values.Find(TEXT("NestedValue"));
+			check(NestedValueValue && (*NestedValueValue)->Type == EJson::Null);
+			check((*NestedValueValue)->IsNull());
+
+			const TSharedPtr<FJsonValue>* NestedObjectValue = InnerObject->Values.Find(TEXT("NestedObject"));
+			check(NestedObjectValue && (*NestedObjectValue)->Type == EJson::Object);
+			const TSharedPtr<FJsonObject> InnerInnerObject = (*NestedObjectValue)->AsObject();
+			check(InnerInnerObject.IsValid());
+
+			{
+				const TSharedPtr<FJsonValue>* NestedValueValueFail = InnerInnerObject->Values.Find(TEXT("NestedValue"));
+				check(!NestedValueValueFail);
+			}
+		}
+
+		const TSharedPtr<FJsonValue>* ValueValue = Object->Values.Find(TEXT("Value"));
+		check(ValueValue && (*ValueValue)->Type == EJson::Boolean);
+		const bool Bool = (*ValueValue)->AsBool();
+		check(Bool);
+
+		FString OutputString;
+		TSharedRef< FCondensedJsonStringWriter > Writer = FCondensedJsonStringWriterFactory::Create( &OutputString );
+		check( FJsonSerializer::Serialize( Object.ToSharedRef(), Writer ) );
+		check(OutputString == InputString);
+	}
+
+	// Array Test
+	{
+		const FString InputString =
+			TEXT("{")
+			TEXT(	"\"Array\":")
+			TEXT(	"[")
+			TEXT(		"[],")
+			TEXT(		"\"Some String\",")
+			TEXT(		"\"Another String\",")
+			TEXT(		"null,")
+			TEXT(		"true,")
+			TEXT(		"false,")
+			TEXT(		"45,")
+			TEXT(		"{}")
+			TEXT(	"]")
+			TEXT("}");
+		TSharedRef< TJsonReader<> > Reader = TJsonReaderFactory<>::Create( InputString );
+
+		TSharedPtr<FJsonObject> Object;
+		bool bSuccessful = FJsonSerializer::Deserialize(Reader, Object);
+		check(bSuccessful);
+		check( Object.IsValid() );
+
+		const TSharedPtr<FJsonValue>* InnerValueFail = Object->Values.Find(TEXT("InnerValue"));
+		check(!InnerValueFail);
+
+		const TSharedPtr<FJsonValue>* ArrayValue = Object->Values.Find(TEXT("Array"));
+		check(ArrayValue && (*ArrayValue)->Type == EJson::Array);
+		const TArray< TSharedPtr<FJsonValue> > Array = (*ArrayValue)->AsArray();
+		check(Array.Num() == 8);
+
+		EJson ValueTypes[] = {EJson::Array, EJson::String, EJson::String, EJson::Null,
+			EJson::Boolean, EJson::Boolean, EJson::Number, EJson::Object};
+		for (int32 i = 0; i < Array.Num(); ++i)
+		{
+			const TSharedPtr<FJsonValue>& Value = Array[i];
+			check(Value.IsValid());
+			check(Value->Type == ValueTypes[i]);
+		}
+
+		const TArray< TSharedPtr<FJsonValue> >& InnerArray = Array[0]->AsArray();
+		check(InnerArray.Num() == 0);
+		check(Array[1]->AsString() == TEXT("Some String"));
+		check(Array[2]->AsString() == TEXT("Another String"));
+		check(Array[3]->IsNull());
+		check(Array[4]->AsBool());
+		check(!Array[5]->AsBool());
+		check(FMath::Abs(Array[6]->AsNumber() - 45.f) < KINDA_SMALL_NUMBER);
+		const TSharedPtr<FJsonObject> InnerObject = Array[7]->AsObject();
+		check(InnerObject.IsValid());
+
+		FString OutputString;
+		TSharedRef< FCondensedJsonStringWriter > Writer = FCondensedJsonStringWriterFactory::Create( &OutputString );
+		check( FJsonSerializer::Serialize( Object.ToSharedRef(), Writer ) );
+		check(OutputString == InputString);
+	}
+
+	// Pretty Print Test
+	{
+		const FString InputString =
+			TEXT("{")									LINE_TERMINATOR
+			TEXT("	\"Data1\": \"value\",")				LINE_TERMINATOR
+			TEXT("	\"Data2\": \"value\",")				LINE_TERMINATOR
+			TEXT("	\"Array\": [")						LINE_TERMINATOR
+			TEXT("		{")								LINE_TERMINATOR
+			TEXT("			\"InnerData1\": \"value\"")	LINE_TERMINATOR
+			TEXT("		},")							LINE_TERMINATOR
+			TEXT("		[],")							LINE_TERMINATOR
+			TEXT("		[ 1, 2, 3, 4 ],")				LINE_TERMINATOR
+			TEXT("		{")								LINE_TERMINATOR
+			TEXT("		},")							LINE_TERMINATOR
+			TEXT("		\"value\",")					LINE_TERMINATOR
+			TEXT("		\"value\"")						LINE_TERMINATOR
+			TEXT("	],")								LINE_TERMINATOR
+			TEXT("	\"Object\":")						LINE_TERMINATOR
+			TEXT("	{")									LINE_TERMINATOR
+			TEXT("	}")									LINE_TERMINATOR
+			TEXT("}");
+		TSharedRef< TJsonReader<> > Reader = TJsonReaderFactory<>::Create( InputString );
+
+		TSharedPtr<FJsonObject> Object;
+		check( FJsonSerializer::Deserialize( Reader, Object ) );
+		check( Object.IsValid() );
+
+		FString OutputString;
+		TSharedRef< FPrettyJsonStringWriter > Writer = FPrettyJsonStringWriterFactory::Create( &OutputString );
+		check( FJsonSerializer::Serialize( Object.ToSharedRef(), Writer ) );
+		check(OutputString == InputString);
+	}
+	  
+	// Line and Character # test
+	{
+		const FString InputString =
+			TEXT("{")									LINE_TERMINATOR
+			TEXT("	\"Data1\": \"value\",")				LINE_TERMINATOR
+			TEXT("	\"Array\":")						LINE_TERMINATOR
+			TEXT("	[")									LINE_TERMINATOR
+			TEXT("		12345,")						LINE_TERMINATOR
+			TEXT("		True")							LINE_TERMINATOR
+			TEXT("	],")								LINE_TERMINATOR
+			TEXT("	\"Object\":")						LINE_TERMINATOR
+			TEXT("	{")									LINE_TERMINATOR
+			TEXT("	}")									LINE_TERMINATOR
+			TEXT("}");
+		TSharedRef< TJsonReader<> > Reader = TJsonReaderFactory<>::Create( InputString );
+
+		EJsonNotation Notation = EJsonNotation::Null;
+		check( Reader->ReadNext( Notation ) && Notation == EJsonNotation::ObjectStart );
+		check( Reader->GetLineNumber() == 1 && Reader->GetCharacterNumber() == 1 );
+
+		check( Reader->ReadNext( Notation ) && Notation == EJsonNotation::String );
+		check( Reader->GetLineNumber() == 2 && Reader->GetCharacterNumber() == 17 );
+
+		check( Reader->ReadNext( Notation ) && Notation == EJsonNotation::ArrayStart );
+		check( Reader->GetLineNumber() == 4 && Reader->GetCharacterNumber() == 2 );
+
+		check( Reader->ReadNext( Notation ) && Notation == EJsonNotation::Number );
+		check( Reader->GetLineNumber() == 5 && Reader->GetCharacterNumber() == 7 );
+
+		check( Reader->ReadNext( Notation ) && Notation == EJsonNotation::Boolean );
+		check( Reader->GetLineNumber() == 6 && Reader->GetCharacterNumber() == 6 );
+	}
+
+	// Failure Cases
+	TArray<FString> FailureInputs;
+
+	// Unclosed Object
+	FailureInputs.Add(
+		TEXT("{"));
+
+	// Values in Object without identifiers
+	FailureInputs.Add(
+		TEXT("{")
+		TEXT(	"\"Value1\",")
+		TEXT(	"\"Value2\",")
+		TEXT(	"43")
+		TEXT("}"));
+
+	// Unexpected End Of Input Found
+	FailureInputs.Add(
+		TEXT("{")
+		TEXT(	"\"Object\":")
+		TEXT(	"{")
+		TEXT(		"\"NestedValue\":null,"));
+
+	// Missing first brace
+	FailureInputs.Add(
+		TEXT(	"\"Object\":")
+		TEXT(		"{")
+		TEXT(		"\"NestedValue\":null,")
+		TEXT(		"\"NestedObject\":{}")
+		TEXT(	"},")
+		TEXT(	"\"Value\":true")
+		TEXT("}"));
+
+	// Missing last character
+	FailureInputs.Add(
+		TEXT("{")
+		TEXT(	"\"Object\":")
+		TEXT(	"{")
+		TEXT(		"\"NestedValue\":null,")
+		TEXT(		"\"NestedObject\":{}")
+		TEXT(	"},")
+		TEXT(	"\"Value\":true"));
+
+	// Extra last character
+	FailureInputs.Add(
+		TEXT("{")
+		TEXT(	"\"Object\":")
+		TEXT(	"{")
+		TEXT(		"\"NestedValue\":null,")
+		TEXT(		"\"NestedObject\":{}")
+		TEXT(	"},")
+		TEXT(	"\"Value\":true")
+		TEXT("}0"));
+
+	// Missing comma
+	FailureInputs.Add(
+		TEXT("{")
+		TEXT(	"\"Value1\":null,")
+		TEXT(	"\"Value2\":\"string\"")
+		TEXT(	"\"Value3\":65.3")
+		TEXT("}"));
+
+	// Extra comma
+	FailureInputs.Add(
+		TEXT("{")
+		TEXT(	"\"Value1\":null,")
+		TEXT(	"\"Value2\":\"string\",")
+		TEXT(	"\"Value3\":65.3,")
+		TEXT("}"));
+
+	// Badly formed true/false/null
+	FailureInputs.Add(TEXT("{\"Value\":tru}"));
+	FailureInputs.Add(TEXT("{\"Value\":full}"));
+	FailureInputs.Add(TEXT("{\"Value\":nulle}"));
+	FailureInputs.Add(TEXT("{\"Value\":n%ll}"));
+
+	// Floating Point Failures
+	FailureInputs.Add(TEXT("{\"Value\":65.3e}"));
+	FailureInputs.Add(TEXT("{\"Value\":65.}"));
+	FailureInputs.Add(TEXT("{\"Value\":.7}"));
+	FailureInputs.Add(TEXT("{\"Value\":+6}"));
+	FailureInputs.Add(TEXT("{\"Value\":01}"));
+	FailureInputs.Add(TEXT("{\"Value\":00.56}"));
+	FailureInputs.Add(TEXT("{\"Value\":-1.e+4}"));
+	FailureInputs.Add(TEXT("{\"Value\":2e+}"));
+
+	// Bad Escape Characters
+	FailureInputs.Add(TEXT("{\"Value\":\"Hello\\xThere\"}"));
+	FailureInputs.Add(TEXT("{\"Value\":\"Hello\\u123There\"}"));
+	FailureInputs.Add(TEXT("{\"Value\":\"Hello\\RThere\"}"));
+
+	for (int32 i = 0; i < FailureInputs.Num(); ++i)
+	{
+		TSharedRef< TJsonReader<> > Reader = TJsonReaderFactory<>::Create( FailureInputs[i] );
+
+		TSharedPtr<FJsonObject> Object;
+		check( FJsonSerializer::Deserialize( Reader, Object ) == false );
+		check( !Object.IsValid() );
+	}
+
+	return true;
+}
\ No newline at end of file
