Fixed editing of brush builder parameters to constrain property values depending on other properties.

diff --git a/Engine/Source/Editor/UnrealEd/Classes/Builders/ConeBuilder.h b/Engine/Source/Editor/UnrealEd/Classes/Builders/ConeBuilder.h
--- a/Engine/Source/Editor/UnrealEd/Classes/Builders/ConeBuilder.h
+++ b/Engine/Source/Editor/UnrealEd/Classes/Builders/ConeBuilder.h
@@ -47,6 +47,9 @@ public:
 	UPROPERTY(EditAnywhere, Category=BrushSettings)
 	uint32 Hollow:1;
 
+	// Begin UObject interface
+	virtual void PostEditChangeProperty(struct FPropertyChangedEvent& PropertyChangedEvent) override;
+	// End UObject interface
 
 	// Begin UBrushBuilder Interface
 	virtual bool Build( UWorld* InWorld, ABrush* InBrush = NULL ) override;
diff --git a/Engine/Source/Editor/UnrealEd/Classes/Builders/CubeBuilder.h b/Engine/Source/Editor/UnrealEd/Classes/Builders/CubeBuilder.h
--- a/Engine/Source/Editor/UnrealEd/Classes/Builders/CubeBuilder.h
+++ b/Engine/Source/Editor/UnrealEd/Classes/Builders/CubeBuilder.h
@@ -44,6 +44,10 @@ public:
 	uint32 Tessellated:1;
 
 
+	// Begin UObject interface
+	virtual void PostEditChangeProperty(struct FPropertyChangedEvent& PropertyChangedEvent) override;
+	// End UObject interface
+
 	// Begin UBrushBuilder Interface
 	virtual bool Build( UWorld* InWorld, ABrush* InBrush = NULL ) override;
 	// End UBrushBuilder Interface
diff --git a/Engine/Source/Editor/UnrealEd/Classes/Builders/CylinderBuilder.h b/Engine/Source/Editor/UnrealEd/Classes/Builders/CylinderBuilder.h
--- a/Engine/Source/Editor/UnrealEd/Classes/Builders/CylinderBuilder.h
+++ b/Engine/Source/Editor/UnrealEd/Classes/Builders/CylinderBuilder.h
@@ -48,6 +48,9 @@ public:
 	UPROPERTY(EditAnywhere, Category=BrushSettings)
 	uint32 Hollow:1;
 
+	// Begin UObject interface
+	virtual void PostEditChangeProperty(struct FPropertyChangedEvent& PropertyChangedEvent) override;
+	// End UObject interface
 
 	// Begin UBrushBuilder Interface
 	virtual bool Build( UWorld* InWorld, ABrush* InBrush = NULL ) override;
diff --git a/Engine/Source/Editor/UnrealEd/Private/EditorBrushBuilder.cpp b/Engine/Source/Editor/UnrealEd/Private/EditorBrushBuilder.cpp
--- a/Engine/Source/Editor/UnrealEd/Private/EditorBrushBuilder.cpp
+++ b/Engine/Source/Editor/UnrealEd/Private/EditorBrushBuilder.cpp
@@ -296,6 +296,44 @@ void UConeBuilder::BuildCone( int32 Direction, bool InAlignToSide, int32 InSides
 		Poly3i( Direction, n+i, n+InSides, n+((i+1)%InSides), Item );
 }
 
+void UConeBuilder::PostEditChangeProperty(struct FPropertyChangedEvent& PropertyChangedEvent)
+{
+	if (PropertyChangedEvent.Property)
+	{
+		static FName Name_Z(GET_MEMBER_NAME_CHECKED(UConeBuilder, Z));
+		static FName Name_CapZ(GET_MEMBER_NAME_CHECKED(UConeBuilder, CapZ));
+
+		const float ZEpsilon = 0.1f;
+
+		if (PropertyChangedEvent.Property->GetFName() == Name_Z && Z <= CapZ)
+		{
+			Z = CapZ + ZEpsilon;
+		}
+
+		if (PropertyChangedEvent.Property->GetFName() == Name_CapZ && CapZ >= Z)
+		{
+			CapZ = FMath::Max(0.0f, Z - ZEpsilon);
+		}
+
+		static FName Name_OuterRadius(GET_MEMBER_NAME_CHECKED(UConeBuilder, OuterRadius));
+		static FName Name_InnerRadius(GET_MEMBER_NAME_CHECKED(UConeBuilder, InnerRadius));
+
+		const float RadiusEpsilon = 0.1f;
+
+		if (PropertyChangedEvent.Property->GetFName() == Name_OuterRadius && OuterRadius <= InnerRadius)
+		{
+			OuterRadius = InnerRadius + RadiusEpsilon;
+		}
+
+		if (PropertyChangedEvent.Property->GetFName() == Name_InnerRadius && InnerRadius >= OuterRadius)
+		{
+			InnerRadius = FMath::Max(0.0f, OuterRadius - RadiusEpsilon);
+		}
+	}
+
+	Super::PostEditChangeProperty(PropertyChangedEvent);
+}
+
 bool UConeBuilder::Build( UWorld* InWorld, ABrush* InBrush )
 {
 	if( Sides<3 )
@@ -389,7 +427,55 @@ void UCubeBuilder::BuildCube( int32 Direction, float dx, float dy, float dz, boo
 	}
 }
 
-bool UCubeBuilder::Build( UWorld* InWorld, ABrush* InBrush )
+void UCubeBuilder::PostEditChangeProperty(struct FPropertyChangedEvent& PropertyChangedEvent)
+{
+	if (PropertyChangedEvent.Property)
+	{
+		static FName Name_X(GET_MEMBER_NAME_CHECKED(UCubeBuilder, X));
+		static FName Name_Y(GET_MEMBER_NAME_CHECKED(UCubeBuilder, Y));
+		static FName Name_Z(GET_MEMBER_NAME_CHECKED(UCubeBuilder, Z));
+		static FName Name_WallThickness(GET_MEMBER_NAME_CHECKED(UCubeBuilder, WallThickness));
+
+		const float ThicknessEpsilon = 0.1f;
+
+		if (PropertyChangedEvent.Property->GetFName() == Name_X && X <= WallThickness)
+		{
+			X = WallThickness + ThicknessEpsilon;
+		}
+
+		if (PropertyChangedEvent.Property->GetFName() == Name_Y && Y <= WallThickness)
+		{
+			Y = WallThickness + ThicknessEpsilon;
+		}
+
+		if (PropertyChangedEvent.Property->GetFName() == Name_Z && Z <= WallThickness)
+		{
+			Z = WallThickness + ThicknessEpsilon;
+		}
+
+		if (PropertyChangedEvent.Property->GetFName() == Name_WallThickness && WallThickness >= FMath::Min3(X, Y, Z))
+		{
+			WallThickness = FMath::Max(0.0f, FMath::Min3(X, Y, Z) - ThicknessEpsilon);
+		}
+
+		static FName Name_Hollow(GET_MEMBER_NAME_CHECKED(UCubeBuilder, Hollow));
+		static FName Name_Tesselated(GET_MEMBER_NAME_CHECKED(UCubeBuilder, Tessellated));
+
+		if (PropertyChangedEvent.Property->GetFName() == Name_Hollow && Hollow && Tessellated)
+		{
+			Hollow = false;
+		}
+
+		if (PropertyChangedEvent.Property->GetFName() == Name_Tesselated && Hollow && Tessellated)
+		{
+			Tessellated = false;
+		}
+	}
+
+	Super::PostEditChangeProperty(PropertyChangedEvent);
+}
+
+bool UCubeBuilder::Build(UWorld* InWorld, ABrush* InBrush)
 {
 	if( Z<=0 || Y<=0 || X<=0 )
 		return BadParameters(LOCTEXT("CubeInvalidDimensions", "Invalid cube dimensions"));
@@ -512,7 +598,7 @@ void UCurvedStairBuilder::BuildCurvedStair( int32 Direction )
 	Poly4i( Direction, BottomInnerStart + NumSteps, InnerStart + (NumSteps * 2), OuterStart + (NumSteps * 2), BottomOuterStart + NumSteps, FName(TEXT("back")) );
 }
 
-bool UCurvedStairBuilder::Build( UWorld* InWorld, ABrush* InBrush )
+bool UCurvedStairBuilder::Build(UWorld* InWorld, ABrush* InBrush)
 {
 	if( AngleOfCurve<1 || AngleOfCurve>360 )
 		return BadParameters(LOCTEXT("StairAngleOutOfRange", "Angle is out of range."));
@@ -569,7 +655,30 @@ void UCylinderBuilder::BuildCylinder( int32 Direction, bool InAlignToSide, int32
 		Poly4i( Direction, n+i*2, n+i*2+1, n+((i*2+3)%(2*InSides)), n+((i*2+2)%(2*InSides)), FName(TEXT("Wall")) );
 }
 
-bool UCylinderBuilder::Build( UWorld* InWorld, ABrush* InBrush )
+void UCylinderBuilder::PostEditChangeProperty(struct FPropertyChangedEvent& PropertyChangedEvent)
+{
+	if (PropertyChangedEvent.Property)
+	{
+		static FName Name_OuterRadius(GET_MEMBER_NAME_CHECKED(UCylinderBuilder, OuterRadius));
+		static FName Name_InnerRadius(GET_MEMBER_NAME_CHECKED(UCylinderBuilder, InnerRadius));
+
+		const float RadiusEpsilon = 0.1f;
+
+		if (PropertyChangedEvent.Property->GetFName() == Name_OuterRadius && OuterRadius <= InnerRadius)
+		{
+			OuterRadius = InnerRadius + RadiusEpsilon;
+		}
+
+		if (PropertyChangedEvent.Property->GetFName() == Name_InnerRadius && InnerRadius >= OuterRadius)
+		{
+			InnerRadius = FMath::Max(0.0f, OuterRadius - RadiusEpsilon);
+		}
+	}
+
+	Super::PostEditChangeProperty(PropertyChangedEvent);
+}
+
+bool UCylinderBuilder::Build(UWorld* InWorld, ABrush* InBrush)
 {
 	if( Sides<3 )
 		return BadParameters(LOCTEXT("CylinderInvalidSides", "Not enough cylinder sides."));
@@ -624,7 +733,7 @@ ULinearStairBuilder::ULinearStairBuilder(const FObjectInitializer& ObjectInitial
 	ToolTip = TEXT("BrushBuilderName_LinearStair");
 }
 
-bool ULinearStairBuilder::Build( UWorld* InWorld, ABrush* InBrush )
+bool ULinearStairBuilder::Build(UWorld* InWorld, ABrush* InBrush)
 {
 	// Check for bad input.
 	if( StepLength<=0 || StepHeight<=0 || StepWidth<=0 )
@@ -751,7 +860,7 @@ USheetBuilder::USheetBuilder(const FObjectInitializer& ObjectInitializer)
 	ToolTip = TEXT("BrushBuilderName_Sheet");
 }
 
-bool USheetBuilder::Build( UWorld* InWorld, ABrush* InBrush )
+bool USheetBuilder::Build(UWorld* InWorld, ABrush* InBrush)
 {
 	if( Y<=0 || X<=0 || XSegments<=0 || YSegments<=0 )
 		return BadParameters(LOCTEXT("SheetInvalidParams", "Invalid sheet parameters."));
@@ -908,7 +1017,7 @@ void USpiralStairBuilder::BuildCurvedStair( int32 Direction )
 	}
 }
 
-bool USpiralStairBuilder::Build( UWorld* InWorld, ABrush* InBrush )
+bool USpiralStairBuilder::Build(UWorld* InWorld, ABrush* InBrush)
 {
 	if( InnerRadius<1 || StepWidth<1 || NumSteps<1 || NumStepsPer360<3 )
 		return BadParameters(LOCTEXT("SpiralStairInvalidStepParams", "Invalid step parameters."));
@@ -974,7 +1083,7 @@ void UTetrahedronBuilder::BuildTetrahedron( float R, int32 InSphereExtrapolation
 	Extrapolate(2,0,5,InSphereExtrapolation,Radius);
 }
 
-bool UTetrahedronBuilder::Build( UWorld* InWorld, ABrush* InBrush )
+bool UTetrahedronBuilder::Build(UWorld* InWorld, ABrush* InBrush)
 {
 	if( Radius<=0 || SphereExtrapolation<=0 )
 		return BadParameters(LOCTEXT("TetrahedronInvalidParams", "Invalid tetrahedron parameters."));
@@ -1036,7 +1145,7 @@ void UVolumetricBuilder::BuildVolumetric( int32 Direction, int32 InNumSheets, fl
 	}
 }
 
-bool UVolumetricBuilder::Build( UWorld* InWorld, ABrush* InBrush )
+bool UVolumetricBuilder::Build(UWorld* InWorld, ABrush* InBrush)
 {
 	if( NumSheets<2 )
 		return BadParameters(LOCTEXT("VolumetricInvalidSheets", "Invalid volumetric sheet count."));
