Fix for mirrored blocking volumes not correctly generating collision geometry.

diff --git a/Engine/Source/Editor/UnrealEd/Private/EditorBsp.cpp b/Engine/Source/Editor/UnrealEd/Private/EditorBsp.cpp
--- a/Engine/Source/Editor/UnrealEd/Private/EditorBsp.cpp
+++ b/Engine/Source/Editor/UnrealEd/Private/EditorBsp.cpp
@@ -1144,6 +1144,8 @@ int UEditorEngine::bspBrushCSG
 	const FRotator Rotation = Actor->GetActorRotation();
 	const FVector Location = Actor->GetActorLocation();
 
+	const bool bIsMirrored = (Scale.X * Scale.Y * Scale.Z < 0.0f);
+
 	// Cache actor transform which is used for the geometry being built
 	Brush->OwnerLocationWhenLastBuilt = Location;
 	Brush->OwnerRotationWhenLastBuilt = Rotation;
@@ -1184,6 +1186,13 @@ int UEditorEngine::bspBrushCSG
 		DestEdPoly.Rotate( Rotation );
 		DestEdPoly.Transform( Location );
 
+		// Reverse winding and normal if the parent brush is mirrored
+		if (bIsMirrored)
+		{
+			DestEdPoly.Reverse();
+			DestEdPoly.CalcNormal();
+		}
+
 		// Add poly to the temp model.
 		new(TempModel->Polys->Element)FPoly( DestEdPoly );
 	}
diff --git a/Engine/Source/Editor/UnrealEd/Private/UnrealEdSrv.cpp b/Engine/Source/Editor/UnrealEd/Private/UnrealEdSrv.cpp
--- a/Engine/Source/Editor/UnrealEd/Private/UnrealEdSrv.cpp
+++ b/Engine/Source/Editor/UnrealEd/Private/UnrealEdSrv.cpp
@@ -1778,9 +1778,7 @@ static void MirrorActors(const FVector& MirrorScale)
 	const FScopedTransaction Transaction( NSLOCTEXT("UnrealEd", "MirroringActors", "Mirroring Actors") );
 
 	// Fires ULevel::LevelDirtiedEvent when falling out of scope.
-	FScopedLevelDirtied LevelDirtyCallback;
-
-	const bool bUseScale = !GetDefault<ULevelEditorMiscSettings>()->bMirrorBrushesWithoutScale;
+	FScopedLevelDirtied		LevelDirtyCallback;
 
 	for ( FSelectionIterator It( GEditor->GetSelectedActorIterator() ) ; It ; ++It )
 	{
@@ -1789,73 +1787,13 @@ static void MirrorActors(const FVector& MirrorScale)
 
 		const FVector PivotLocation = GLevelEditorModeTools().PivotLocation;
 
-		if (bUseScale)
-		{
-			ABrush* Brush = Cast< ABrush >(Actor);
-			if (Brush && Brush->Brush)
-			{
-				// Brushes have to reverse their poly vertex order and recalculate the normal as negating one of the scale axes
-				// changes the handedness of the local transform.
-				Brush->Modify();
-				Brush->Brush->Modify();
-				Brush->Brush->Polys->Modify();
-
-				for (int32 poly = 0; poly < Brush->Brush->Polys->Element.Num(); poly++)
-				{
-					FPoly* Poly = &(Brush->Brush->Polys->Element[poly]);
-
-					Poly->Reverse();
-					Poly->CalcNormal();
-				}
-			}
+		Actor->Modify();
+		Actor->EditorApplyMirror( MirrorScale, PivotLocation );
 
-			Actor->Modify();
-			Actor->EditorApplyMirror(MirrorScale, PivotLocation);
-		}
-		else
+		ABrush* Brush = Cast< ABrush >(Actor);
+		if (Brush && Brush->BrushComponent)
 		{
-			ABrush* Brush = Cast< ABrush >( Actor );
-				if( Brush && Brush->Brush )
-			{
-				Brush->Modify();
-				Brush->Brush->Modify();
-				Brush->Brush->Polys->Modify();
-
-				const FVector LocalToWorldOffset = ( Brush->GetActorLocation() - PivotLocation );
-				const FVector LocationOffset = ( LocalToWorldOffset * MirrorScale ) - LocalToWorldOffset;
-
-				Brush->SetActorLocation( Brush->GetActorLocation() + LocationOffset, false );
-				Brush->SetPivotOffset( Brush->GetPivotOffset() * MirrorScale );
-
-				for( int32 poly = 0 ; poly < Brush->Brush->Polys->Element.Num() ; poly++ )
-				{
-					FPoly* Poly = &(Brush->Brush->Polys->Element[poly]);
-
-					Poly->TextureU *= MirrorScale;
-					Poly->TextureV *= MirrorScale;
-
-					Poly->Base += LocalToWorldOffset;
-					Poly->Base *= MirrorScale;
-					Poly->Base -= LocalToWorldOffset;
-					Poly->Base -= LocationOffset;
-
-					for( int32 vtx = 0 ; vtx < Poly->Vertices.Num(); vtx++ )
-					{
-						Poly->Vertices[vtx] += LocalToWorldOffset;
-						Poly->Vertices[vtx] *= MirrorScale;
-						Poly->Vertices[vtx] -= LocalToWorldOffset;
-						Poly->Vertices[vtx] -= LocationOffset;
-					}
-
-					Poly->Reverse();
-					Poly->CalcNormal();
-				}
-			}
-			else
-			{
-				Actor->Modify();
-				Actor->EditorApplyMirror( MirrorScale, PivotLocation );
-			}
+			Brush->BrushComponent->RequestUpdateBrushCollision();
 		}
 
 		Actor->InvalidateLightingCache();
@@ -1875,6 +1813,7 @@ static void MirrorActors(const FVector& MirrorScale)
 	GEditor->RedrawLevelEditingViewports();
 }
 
+
 /**
 * Gathers up a list of selection FPolys from selected static meshes.
 *
diff --git a/Engine/Source/Runtime/Engine/Classes/Components/BrushComponent.h b/Engine/Source/Runtime/Engine/Classes/Components/BrushComponent.h
--- a/Engine/Source/Runtime/Engine/Classes/Components/BrushComponent.h
+++ b/Engine/Source/Runtime/Engine/Classes/Components/BrushComponent.h
@@ -50,6 +50,9 @@ public:
 #endif
 	// End UPrimitiveComponent interface.
 
+	/** If the transform mirroring no longer reflects the body setup, request its recalculation */
+	ENGINE_API void RequestUpdateBrushCollision();
+
 	/** Create the AggGeom collection-of-convex-primitives from the Brush UModel data. */
 	ENGINE_API void BuildSimpleBrushCollision();
 };
diff --git a/Engine/Source/Runtime/Engine/Classes/Engine/Brush.h b/Engine/Source/Runtime/Engine/Classes/Engine/Brush.h
--- a/Engine/Source/Runtime/Engine/Classes/Engine/Brush.h
+++ b/Engine/Source/Runtime/Engine/Classes/Engine/Brush.h
@@ -143,6 +143,7 @@ public:
 	virtual void PostLoad() override;
 	virtual void PostEditMove(bool bFinished) override;
 	virtual void PostEditChangeProperty(FPropertyChangedEvent& PropertyChangedEvent) override;
+	virtual void PostEditChangeChainProperty(FPropertyChangedChainEvent& PropertyChangedEvent) override;
 #endif // WITH_EDITOR
 
 	virtual bool Modify(bool bAlwaysMarkDirty = false) override;
diff --git a/Engine/Source/Runtime/Engine/Private/Brush.cpp b/Engine/Source/Runtime/Engine/Private/Brush.cpp
--- a/Engine/Source/Runtime/Engine/Private/Brush.cpp
+++ b/Engine/Source/Runtime/Engine/Private/Brush.cpp
@@ -55,6 +55,18 @@ void ABrush::PostEditMove(bool bFinished)
 	Super::PostEditMove(bFinished);
 }
 
+void ABrush::PostEditChangeChainProperty(FPropertyChangedChainEvent& PropertyChangedEvent)
+{
+	// Prior to reregistering the BrushComponent (done in the Super), request an update to the Body Setup to take into account any change
+	// in the mirroring of the Actor. This will actually be updated when the component is reregistered.
+	if (BrushComponent && PropertyChangedEvent.Property && PropertyChangedEvent.Property->GetName() == TEXT("RelativeScale3D"))
+	{
+		BrushComponent->RequestUpdateBrushCollision();
+	}
+
+	Super::PostEditChangeChainProperty(PropertyChangedEvent);
+}
+
 void ABrush::PostEditChangeProperty(FPropertyChangedEvent& PropertyChangedEvent)
 {
 	if(Brush)
diff --git a/Engine/Source/Runtime/Engine/Private/Components/BrushComponent.cpp b/Engine/Source/Runtime/Engine/Private/Components/BrushComponent.cpp
--- a/Engine/Source/Runtime/Engine/Private/Components/BrushComponent.cpp
+++ b/Engine/Source/Runtime/Engine/Private/Components/BrushComponent.cpp
@@ -546,6 +546,10 @@ void UBrushComponent::PostLoad()
 		BrushBodySetup->bGenerateMirroredCollision = false;
 	}
 
+	// If loading a brush with mirroring whose body setup has not been created correctly, request that it be rebuilt now.
+	// The rebuilding will actually happen in the UBodySetup::PostLoad.
+	RequestUpdateBrushCollision();
+
 #if WITH_EDITOR
 	AActor* Owner = GetOwner();
 
@@ -730,6 +734,22 @@ bool UBrushComponent::ComponentIsTouchingSelectionFrustum(const FConvexVolume& I
 #endif
 
 
+void UBrushComponent::RequestUpdateBrushCollision()
+{
+	if (BrushBodySetup)
+	{
+		const bool bIsMirrored = (RelativeScale3D.X * RelativeScale3D.Y * RelativeScale3D.Z) < 0.0f;
+		if ((BrushBodySetup->bGenerateNonMirroredCollision && bIsMirrored) || (BrushBodySetup->bGenerateMirroredCollision && !bIsMirrored))
+		{
+			// Brushes only maintain one convex mesh as they can't be transformed at runtime.
+			// Here we invalidate the body setup, and specify whether we wish to build a non-mirrored or a mirrored mesh.
+			BrushBodySetup->bGenerateNonMirroredCollision = !bIsMirrored;
+			BrushBodySetup->bGenerateMirroredCollision = bIsMirrored;
+			BrushBodySetup->InvalidatePhysicsData();
+		}
+	}
+}
+
 void UBrushComponent::BuildSimpleBrushCollision()
 {
 	AActor* Owner = GetOwner();
@@ -747,12 +767,13 @@ void UBrushComponent::BuildSimpleBrushCollision()
 
 	// No complex collision, so use the simple for that
 	BrushBodySetup->CollisionTraceFlag = CTF_UseSimpleAsComplex;
-	// Don't need mirrored version of collision data
-	BrushBodySetup->bGenerateMirroredCollision = false;
+
+	RequestUpdateBrushCollision();
 
 #if WITH_EDITOR
 	// Convert collision model into convex hulls.
 	BrushBodySetup->CreateFromModel( Brush, true );
+
 	RecreatePhysicsState();
 #endif // WITH_EDITOR
 
diff --git a/Engine/Source/Runtime/Engine/Private/PhysicsEngine/BodySetup.cpp b/Engine/Source/Runtime/Engine/Private/PhysicsEngine/BodySetup.cpp
--- a/Engine/Source/Runtime/Engine/Private/PhysicsEngine/BodySetup.cpp
+++ b/Engine/Source/Runtime/Engine/Private/PhysicsEngine/BodySetup.cpp
@@ -475,7 +475,7 @@ public:
 		{
 			const FKConvexElem& ConvexElem = BodySetup->AggGeom.ConvexElems[i];
 
-			if (ConvexElem.ConvexMesh)
+			if (ConvexElem.ConvexMesh || ConvexElem.ConvexMeshNegX)
 			{
 				PxTransform PLocalPose;
 				bool bUseNegX = CalcMeshNegScaleCompensation(Scale3D, PLocalPose);
