Removed unused editor setting property.

diff --git a/Engine/Source/Editor/UnrealEd/Classes/Settings/LevelEditorMiscSettings.h b/Engine/Source/Editor/UnrealEd/Classes/Settings/LevelEditorMiscSettings.h
--- a/Engine/Source/Editor/UnrealEd/Classes/Settings/LevelEditorMiscSettings.h
+++ b/Engine/Source/Editor/UnrealEd/Classes/Settings/LevelEditorMiscSettings.h
@@ -38,10 +38,6 @@ public:
 	UPROPERTY(EditAnywhere, config, Category=Editing, meta=( DisplayName = "Preserve Actor Scale on Replace" ))
 	uint32 bReplaceRespectsScale:1;
 
-	/** If enabled, mirroring actions will directly affect the geometry of brushes rather than their scale */
-	UPROPERTY(EditAnywhere, config, Category=Editing, meta=( DisplayName = "Mirror Brushes Directly" ))
-	uint32 bMirrorBrushesWithoutScale:1;
-
 public:
 	/** If checked audio playing in the editor will continue to play even if the editor is in the background */
 	UPROPERTY(EditAnywhere, config, Category=Sound)
