Updated animation recording
- fixed some issues causing inaccuracies in the output data
- added some cvars to control sample rate and toggle world space recording
- moved control code from EditorEngine to Engine so it's available in standalone builds
- set compression to bitwise for faster compression after recording

(merge from branch)

diff --git a/Engine/Source/Editor/UnrealEd/Classes/Editor/EditorEngine.h b/Engine/Source/Editor/UnrealEd/Classes/Editor/EditorEngine.h
--- a/Engine/Source/Editor/UnrealEd/Classes/Editor/EditorEngine.h
+++ b/Engine/Source/Editor/UnrealEd/Classes/Editor/EditorEngine.h
@@ -766,8 +766,6 @@ public:
 	bool	HandleListMapPackageDependenciesCommand( const TCHAR* Str, FOutputDevice& Ar );
 	bool	HandleRebuildVolumesCommand( const TCHAR* Str, FOutputDevice& Ar, UWorld* InWorld );
 	bool	HandleRemoveArchtypeFlagCommand( const TCHAR* Str, FOutputDevice& Ar );
-	bool	HandleRecordAnimationCommand(UWorld* InWorld, const TCHAR* InStr, FOutputDevice& Ar);
-	bool	HandleStopRecordAnimationCommand(UWorld* InWorld, const TCHAR* InStr, FOutputDevice& Ar);
 
 	/**
 	 * Initializes the Editor.
diff --git a/Engine/Source/Editor/UnrealEd/Private/Animation/AnimationRecorder.cpp b/Engine/Source/Editor/UnrealEd/Private/Animation/AnimationRecorder.cpp
--- a/Engine/Source/Editor/UnrealEd/Private/Animation/AnimationRecorder.cpp
+++ b/Engine/Source/Editor/UnrealEd/Private/Animation/AnimationRecorder.cpp
@@ -3,6 +3,8 @@
 
 #include "UnrealEd.h"
 #include "AnimationRecorder.h"
+#include "Animation/AnimCompress.h"
+#include "Animation/AnimCompress_BitwiseCompressOnly.h"
 #include "SCreateAnimationDlg.h"
 #include "Runtime/AssetRegistry/Public/AssetRegistryModule.h"
 #include "SNotificationList.h"
@@ -13,12 +15,27 @@
 
 /////////////////////////////////////////////////////
 
+float FAnimationRecorder::DefaultSampleRate = DEFAULT_SAMPLERATE;
+
+static TAutoConsoleVariable<float> CVarAnimRecorderSampleRate(
+	TEXT("AnimRecorder.SampleRate"),
+	DEFAULT_SAMPLERATE,
+	TEXT("Sets the sample rate for the animation recorder system"),
+	ECVF_Default);
+
+static TAutoConsoleVariable<int32> CVarAnimRecorderWorldSpace(
+	TEXT("AnimRecorder.RecordInWorldSpace"),
+	1,
+	TEXT("True to record anim keys in world space, false to record only in local space."),
+	ECVF_Default);
+
 FAnimationRecorder::FAnimationRecorder()
-	: IntervalTime(1.0f/DEFAULT_SAMPLERATE)
-	, MaxFrame(DEFAULT_SAMPLERATE * 60 * 10) // for now, set the max limit to 10 mins. 
-	, AnimationObject(NULL)
+	: AnimationObject(nullptr)
 	, bRecordLocalToWorld(false)
+	, bAutoSaveAsset(false)
+
 {
+	SetSampleRate(DefaultSampleRate);
 }
 
 FAnimationRecorder::~FAnimationRecorder()
@@ -26,13 +43,38 @@ FAnimationRecorder::~FAnimationRecorder()
 	StopRecord(false);
 }
 
-// get record configuration
-bool GetRecordConfig(FString& AssetPath, FString& AssetName)
+void FAnimationRecorder::SetSampleRate(float SampleRateHz)
+{
+	if (SampleRateHz <= 0.f)
+	{
+		// invalid rate passed in, fall back to default
+		SampleRateHz = FAnimationRecorder::DefaultSampleRate;
+	}
+
+	IntervalTime = 1.0f / SampleRateHz;
+	MaxFrame = SampleRateHz * 60 * 10;		 // for now, set the max limit to 10 mins. 
+}
+
+bool FAnimationRecorder::SetAnimCompressionScheme(TSubclassOf<UAnimCompress> SchemeClass)
 {
-	TSharedRef<SCreateAnimationDlg> NewAnimDlg =
-		SNew(SCreateAnimationDlg);
+	if (AnimationObject)
+	{
+		UAnimCompress* const SchemeObject = NewObject<UAnimCompress>(GetTransientPackage(), SchemeClass);
+		if (SchemeObject)
+		{
+			AnimationObject->CompressionScheme = SchemeObject;
+			return true;
+		}
+	}
 
-	if(NewAnimDlg->ShowModal() != EAppReturnType::Cancel)
+	return false;
+}
+
+// Internal. Pops up a dialog to get saved asset path
+static bool PromptUserForAssetPath(FString& AssetPath, FString& AssetName)
+{
+	TSharedRef<SCreateAnimationDlg> NewAnimDlg = SNew(SCreateAnimationDlg);
+	if (NewAnimDlg->ShowModal() != EAppReturnType::Cancel)
 	{
 		AssetPath = NewAnimDlg->GetFullAssetPath();
 		AssetName = NewAnimDlg->GetAssetName();
@@ -44,7 +86,6 @@ bool GetRecordConfig(FString& AssetPath, FString& AssetName)
 
 bool FAnimationRecorder::TriggerRecordAnimation(USkeletalMeshComponent* Component)
 {
-	// ask for path
 	FString AssetPath;
 	FString AssetName;
 
@@ -53,7 +94,8 @@ bool FAnimationRecorder::TriggerRecordAnimation(USkeletalMeshComponent* Componen
 		return false;
 	}
 
-	if (GetRecordConfig(AssetPath, AssetName))
+	// ask for path
+	if (PromptUserForAssetPath(AssetPath, AssetName))
 	{
 		return TriggerRecordAnimation(Component, AssetPath, AssetName);
 	}
@@ -73,7 +115,7 @@ bool FAnimationRecorder::TriggerRecordAnimation(USkeletalMeshComponent* Componen
 	if (Parent == nullptr)
 	{
 		// bad or no path passed in, do the popup
-		if (GetRecordConfig(AssetPath, AssetName) == false)
+		if (PromptUserForAssetPath(AssetPath, AssetName) == false)
 		{
 			return false;
 		}
@@ -116,6 +158,7 @@ void FAnimationRecorder::StartRecord(USkeletalMeshComponent* Component, UAnimSeq
 	AnimationObject->AnimationTrackNames.Empty();
 
 	PreviousSpacesBases = Component->GetSpaceBases();
+	PreviousComponentToWorld = Component->ComponentToWorld;
 
 	LastFrame = 0;
 	AnimationObject->SequenceLength = 0.f;
@@ -142,7 +185,7 @@ void FAnimationRecorder::StartRecord(USkeletalMeshComponent* Component, UAnimSeq
 	AnimationObject->RawAnimationData.AddZeroed(NumTracks);
 
 	// record the first frame;
-	Record(Component, PreviousSpacesBases, 0);
+	Record(Component, PreviousComponentToWorld, PreviousSpacesBases, 0);
 }
 
 UAnimSequence* FAnimationRecorder::StopRecord(bool bShowMessage)
@@ -153,19 +196,31 @@ UAnimSequence* FAnimationRecorder::StopRecord(bool bShowMessage)
 		AnimationObject->NumFrames = NumFrames;
 
 		// can't use TimePassed. That is just total time that has been passed, not necessarily match with frame count
-		AnimationObject->SequenceLength = (NumFrames>1)? (NumFrames-1) * IntervalTime : MINIMUM_ANIMATION_LENGTH;
+		AnimationObject->SequenceLength = (NumFrames>1) ? (NumFrames-1) * IntervalTime : MINIMUM_ANIMATION_LENGTH;
 		AnimationObject->PostProcessSequence();
 		AnimationObject->MarkPackageDirty();
+		
+		// save the package to disk, for convenience and so we can run this in standalone mode
+		if (bAutoSaveAsset)
+		{
+			UPackage* const Package = AnimationObject->GetOutermost();
+			FString const PackageName = Package->GetName();
+			FString const PackageFileName = FPackageName::LongPackageNameToFilename(PackageName, FPackageName::GetAssetPackageExtension());
+
+			UPackage::SavePackage(Package, NULL, RF_Standalone, *PackageFileName, GError, nullptr, false, true, SAVE_NoError);
+		}
 
 		UAnimSequence* ReturnObject = AnimationObject;
 
 		// notify to user
 		if (bShowMessage)
 		{
-			const FText NotificationText = FText::Format( LOCTEXT("RecordAnimation", "'{0}' has been successfully recorded [{1} frames : {2} sec(s)]" ), 
-				FText::FromString( AnimationObject->GetName() ),
-				FText::AsNumber( AnimationObject->NumFrames ),
-				FText::AsNumber( AnimationObject->SequenceLength ));
+			const FText NotificationText = FText::Format(LOCTEXT("RecordAnimation", "'{0}' has been successfully recorded [{1} frames : {2} sec(s) @ {3} Hz]"),
+				FText::FromString(AnimationObject->GetName()),
+				FText::AsNumber(AnimationObject->NumFrames),
+				FText::AsNumber(AnimationObject->SequenceLength),
+				FText::AsNumber(1.f / IntervalTime)
+				);
 
 			//This is not showing well in the Persona, so opening dialog first. 
 			//right now it will crash if you don't wait until end of the record, so it is important for users to know
@@ -194,70 +249,67 @@ UAnimSequence* FAnimationRecorder::StopRecord(bool bShowMessage)
 	return NULL;
 }
 
-// return false if nothing to update
-// return true if it has properly updated
 void FAnimationRecorder::UpdateRecord(USkeletalMeshComponent* Component, float DeltaTime)
 {
 	// if no animation object, return
-	if (!AnimationObject)
+	if (!AnimationObject || !Component)
 	{
 		return;
 	}
 
-	const float MaxDelta = 1.f/10.f;
-	float UseDelta = FMath::Min<float>(DeltaTime, MaxDelta);
-
-	// if longer than 1 second, you'll see the halt on the animation
-	if (UseDelta <= 0.f)
+	// no sime time, no record
+	if (DeltaTime <= 0.f)
 	{
 		return;
 	}
 
-	float PreviousTimePassed = TimePassed;
-	TimePassed += UseDelta;
+	float const PreviousTimePassed = TimePassed;
+	TimePassed += DeltaTime;
 
 	// time passed has been updated
 	// now find what frames we need to update
 	int32 FramesRecorded = LastFrame;
 	int32 FramesToRecord = FPlatformMath::TruncToInt(TimePassed / IntervalTime);
 
-	if (FramesRecorded >= FramesToRecord)
+	if (FramesRecorded < FramesToRecord)
 	{
-		return;
-	}
+		const TArray<FTransform>& SpaceBases = Component->GetSpaceBases();
 
-	const TArray<FTransform> & SpaceBases = Component->GetSpaceBases();
+		check(SpaceBases.Num() == PreviousSpacesBases.Num());
 
-	check (SpaceBases.Num() == PreviousSpacesBases.Num());
+		TArray<FTransform> BlendedSpaceBases;
+		BlendedSpaceBases.AddZeroed(SpaceBases.Num());
 
-	TArray<FTransform> BlendedSpaceBases;
-	BlendedSpaceBases.AddZeroed(SpaceBases.Num());
+		UE_LOG(LogAnimation, Warning, TEXT("DeltaTime : %0.2f, Current Frame Count : %d, Frames To Record : %d, TimePassed : %0.2f"), DeltaTime
+			, FramesRecorded, FramesToRecord, TimePassed);
 
-	UE_LOG(LogAnimation, Warning, TEXT("DeltaTime : %0.2f, Current Frame Count : %d, Frames To Record : %d, TimePassed : %0.2f"), DeltaTime
-		, FramesRecorded, FramesToRecord, TimePassed);
+		// if we need to record frame
+		while (FramesToRecord > FramesRecorded)
+		{
+			// find what frames we need to record
+			// convert to time
+			float CurrentTime = (FramesRecorded + 1) * IntervalTime;
+			float BlendAlpha = (CurrentTime - PreviousTimePassed) / DeltaTime;
 
-	// if we need to record frame
-	while (FramesToRecord > FramesRecorded)
-	{
-		// find what frames we need to record
-		// convert to time
-		float CurrentTime = (FramesRecorded + 1) * IntervalTime;
-		float BlendAlpha = (CurrentTime - PreviousTimePassed)/UseDelta;
+			UE_LOG(LogAnimation, Warning, TEXT("Current Frame Count : %d, BlendAlpha : %0.2f"), FramesRecorded + 1, BlendAlpha);
 
-		UE_LOG(LogAnimation, Warning, TEXT("Current Frame Count : %d, BlendAlpha : %0.2f"), FramesRecorded + 1, BlendAlpha);
+			// for now we just concern component space, not skeleton space
+			for (int32 BoneIndex = 0; BoneIndex<SpaceBases.Num(); ++BoneIndex)
+			{
+				BlendedSpaceBases[BoneIndex].Blend(PreviousSpacesBases[BoneIndex], SpaceBases[BoneIndex], BlendAlpha);
+			}
 
-		// for now we just concern component space, not skeleton space
-		for (int32 BoneIndex=0; BoneIndex<SpaceBases.Num(); ++BoneIndex)
-		{
-			BlendedSpaceBases[BoneIndex].Blend(PreviousSpacesBases[BoneIndex], SpaceBases[BoneIndex], BlendAlpha);
-		}
+			FTransform BlendedComponentToWorld;
+			BlendedComponentToWorld.Blend(PreviousComponentToWorld, Component->ComponentToWorld, BlendAlpha);
 
-		Record(Component, BlendedSpaceBases, FramesRecorded + 1);
-		++FramesRecorded;
+			Record(Component, BlendedComponentToWorld, BlendedSpaceBases, FramesRecorded + 1);
+			++FramesRecorded;
+		}
 	}
 
 	//save to current transform
 	PreviousSpacesBases = Component->GetSpaceBases();
+	PreviousComponentToWorld = Component->ComponentToWorld;
 
 	// if we passed MaxFrame, just stop it
 	if (FramesRecorded > MaxFrame)
@@ -267,9 +319,9 @@ void FAnimationRecorder::UpdateRecord(USkeletalMeshComponent* Component, float D
 	}
 }
 
-void FAnimationRecorder::Record( USkeletalMeshComponent* Component, TArray<FTransform> SpacesBases, int32 FrameToAdd )
+void FAnimationRecorder::Record( USkeletalMeshComponent* Component, FTransform const& ComponentToWorld, TArray<FTransform> SpacesBases, int32 FrameToAdd )
 {
-	if (ensure (AnimationObject))
+	if (ensure(AnimationObject))
 	{
 		USkeleton* AnimSkeleton = AnimationObject->GetSkeleton();
 		for (int32 TrackIndex=0; TrackIndex <AnimationObject->RawAnimationData.Num(); ++TrackIndex)
@@ -288,7 +340,7 @@ void FAnimationRecorder::Record( USkeletalMeshComponent* Component, TArray<FTran
 				// if record local to world, we'd like to consider component to world to be in root
 				else if (bRecordLocalToWorld)
 				{
-					LocalTransform *= Component->ComponentToWorld;
+					LocalTransform *= ComponentToWorld;
 				}
 
 				FRawAnimSequenceTrack& RawTrack = AnimationObject->RawAnimationData[TrackIndex];
@@ -305,9 +357,6 @@ void FAnimationRecorder::Record( USkeletalMeshComponent* Component, TArray<FTran
 	}
 }
 
-
-
-
 void FAnimationRecorderManager::Tick(float DeltaTime)
 {
 	for (auto& Inst : RecorderInstances)
@@ -338,14 +387,19 @@ FAnimRecorderInstance::FAnimRecorderInstance()
 {
 }
 
-void FAnimRecorderInstance::Init(AActor* InActor, USkeletalMeshComponent* InComponent, FString InAssetPath, FString InAssetName)
+void FAnimRecorderInstance::Init(AActor* InActor, USkeletalMeshComponent* InComponent, FString InAssetPath, FString InAssetName, float SampleRateHz, bool bRecordInWorldSpace)
 {
 	Actor = InActor;
 	SkelComp = InComponent;
 	AssetPath = InAssetPath;
 	AssetName = InAssetName;
 	Recorder = new FAnimationRecorder();
-	Recorder->bRecordLocalToWorld = true;
+	Recorder->SetSampleRate(SampleRateHz);
+	Recorder->bRecordLocalToWorld = bRecordInWorldSpace;
+	Recorder->SetAnimCompressionScheme(UAnimCompress_BitwiseCompressOnly::StaticClass());
+
+	// always autosave through this codepath
+	Recorder->bAutoSaveAsset = true;
 }
 
 FAnimRecorderInstance::~FAnimRecorderInstance()
@@ -382,7 +436,7 @@ void FAnimRecorderInstance::FinishRecording()
 
 	if (Recorder)
 	{
-		Recorder->StopRecord(false);
+		Recorder->StopRecord(true);
 	}
 
 	GWarn->EndSlowTask();
@@ -393,7 +447,7 @@ bool FAnimationRecorderManager::RecordAnimation(AActor* Actor, USkeletalMeshComp
 	int32 const NewInstIdx = RecorderInstances.AddDefaulted();
 	
 	FAnimRecorderInstance& NewInst = RecorderInstances[NewInstIdx];
-	NewInst.Init(Actor, Component, AssetPath, AssetName);
+	NewInst.Init(Actor, Component, AssetPath, AssetName, CVarAnimRecorderSampleRate.GetValueOnAnyThread(), (CVarAnimRecorderWorldSpace.GetValueOnAnyThread() != 0));
 
 	bool const bSuccess = NewInst.BeginRecording();
 	if (bSuccess == false)
diff --git a/Engine/Source/Editor/UnrealEd/Private/EditorServer.cpp b/Engine/Source/Editor/UnrealEd/Private/EditorServer.cpp
--- a/Engine/Source/Editor/UnrealEd/Private/EditorServer.cpp
+++ b/Engine/Source/Editor/UnrealEd/Private/EditorServer.cpp
@@ -5646,14 +5646,6 @@ bool UEditorEngine::Exec( UWorld* InWorld, const TCHAR* Stream, FOutputDevice& A
 	{
 		HandleRemoveArchtypeFlagCommand( Str, Ar );
 	}
-	else if (FParse::Command(&Str, TEXT("RecordAnimation")))
-	{
-		bProcessed = HandleRecordAnimationCommand(InWorld, Str, Ar);
-	}
-	else if (FParse::Command(&Str, TEXT("StopRecordingAnimation")))
-	{
-		bProcessed = HandleStopRecordAnimationCommand(InWorld, Str, Ar);
-	}
 	else
 	{
 		bProcessed = FBlueprintEditorUtils::KismetDiagnosticExec(Stream, Ar);
@@ -6423,107 +6415,6 @@ bool UEditorEngine::HandleRemoveArchtypeFlagCommand( const TCHAR* Str, FOutputDe
 	return true;
 }
 
-
-bool UEditorEngine::HandleRecordAnimationCommand(UWorld* InWorld, const TCHAR* InStr, FOutputDevice& Ar)
-{
-	const TCHAR* Str = InStr;
-	// parse actor name
-	TCHAR ActorName[128];
-	AActor* FoundActor = nullptr;
-	if (FParse::Token(Str, ActorName, ARRAY_COUNT(ActorName), 0))
-	{
-		FString const ActorNameStr = FString(ActorName);
-		for (ULevel const* Level : InWorld->GetLevels())
-		{
-			if (Level)
-			{
-				for (AActor* Actor : Level->Actors)
-				{
-					if (Actor)
-					{
-						if (Actor->GetName() == ActorNameStr)
-						{
-							FoundActor = Actor;
-							break;
-						}
-					}
-				}
-			}
-		}
-	}
-
-	if (FoundActor)
-	{
-		USkeletalMeshComponent* const SkelComp = FoundActor->FindComponentByClass<USkeletalMeshComponent>();
-		if (SkelComp)
-		{
-			TCHAR AssetPath[256];
-			FParse::Token(Str, AssetPath, ARRAY_COUNT(AssetPath), 0);
-			FString const AssetName = FPackageName::GetLongPackageAssetName(AssetPath);
-			return FAnimationRecorderManager::Get().RecordAnimation(FoundActor, SkelComp, AssetPath, AssetName);
-		}
-	}
-
-	return false;
-}
-
-bool UEditorEngine::HandleStopRecordAnimationCommand(UWorld* InWorld, const TCHAR* InStr, FOutputDevice& Ar)
-{
-	const TCHAR* Str = InStr;
-
-	// parse actor name
-	TCHAR ActorName[128];
-	AActor* FoundActor = nullptr;
-	bool bStopAll = false;
-	if (FParse::Token(Str, ActorName, ARRAY_COUNT(ActorName), 0))
-	{
-		FString const ActorNameStr = FString(ActorName);
-
-		if (ActorNameStr.ToLower() == TEXT("all"))
-		{
-			bStopAll = true;
-		}
-		else if (InWorld)
-		{
-			// search for the actor by name
-			for (ULevel* Level : InWorld->GetLevels())
-			{
-				if (Level)
-				{
-					for (AActor* Actor : Level->Actors)
-					{
-						if (Actor)
-						{
-							if (Actor->GetName() == ActorNameStr)
-							{
-								FoundActor = Actor;
-								break;
-							}
-						}
-					}
-				}
-			}
-		}
-	}
-
-	if (bStopAll)
-	{
-		FAnimationRecorderManager::Get().StopRecordingAllAnimations();
-		return true;
-	}
-	else if (FoundActor)
-	{
-		USkeletalMeshComponent* const SkelComp = FoundActor->FindComponentByClass<USkeletalMeshComponent>();
-		if (SkelComp)
-		{
-			FAnimationRecorderManager::Get().StopRecordingAnimation(FoundActor, SkelComp);
-			return true;
-		}
-	}
-
-	return false;
-}
-
 /**
  * @return true if the given component's StaticMesh can be merged with other StaticMeshes
  */
diff --git a/Engine/Source/Editor/UnrealEd/Public/Animation/AnimationRecorder.h b/Engine/Source/Editor/UnrealEd/Public/Animation/AnimationRecorder.h
--- a/Engine/Source/Editor/UnrealEd/Public/Animation/AnimationRecorder.h
+++ b/Engine/Source/Editor/UnrealEd/Public/Animation/AnimationRecorder.h
@@ -15,6 +15,9 @@ private:
 	float TimePassed;
 	UAnimSequence* AnimationObject;
 	TArray<FTransform> PreviousSpacesBases;
+	FTransform PreviousComponentToWorld;
+
+	static float DefaultSampleRate;
 
 public:
 	FAnimationRecorder();
@@ -32,11 +35,18 @@ public:
 	bool InRecording() const { return AnimationObject != nullptr; }
 	float GetTimeRecorded() const { return TimePassed; }
 
-	// If true, it will record root to include LocalToWorld
-	bool bRecordLocalToWorld;
+	/** Sets a new sample rate for this recorder. Don't call while recording. */
+	void SetSampleRate(float SampleRateHz);
+
+	bool SetAnimCompressionScheme(TSubclassOf<class UAnimCompress> SchemeClass);
+
+	/** If true, it will record root to include LocalToWorld */
+	uint8 bRecordLocalToWorld :1;
+	/** If true, asset will be saved to disk after recording. If false, asset will remain in mem and can be manually saved. */
+	uint8 bAutoSaveAsset : 1;
 
 private:
-	void Record( USkeletalMeshComponent* Component, TArray<FTransform> SpacesBases, int32 FrameToAdd );
+	void Record(USkeletalMeshComponent* Component, FTransform const& ComponentToWorld, TArray<FTransform> SpacesBases, int32 FrameToAdd);
 };
 
 //////////////////////////////////////////////////////////////////////////
@@ -48,7 +58,7 @@ public:
 	FAnimRecorderInstance();
 	~FAnimRecorderInstance();
 
-	void Init(AActor* InActor, USkeletalMeshComponent* InComponent, FString InAssetPath, FString InAssetName);
+	void Init(AActor* InActor, USkeletalMeshComponent* InComponent, FString InAssetPath, FString InAssetName, float SampleRateHz, bool bRecordInWorldSpace);
 
 	bool BeginRecording();
 	void Update(float DeltaTime);
@@ -66,7 +76,7 @@ public:
 //////////////////////////////////////////////////////////////////////////
 // FAnimationRecorderManager
 
-struct FAnimationRecorderManager
+struct UNREALED_API FAnimationRecorderManager
 {
 public:
 	/** Singleton accessor */
diff --git a/Engine/Source/Runtime/Engine/Classes/Engine/Engine.h b/Engine/Source/Runtime/Engine/Classes/Engine/Engine.h
--- a/Engine/Source/Runtime/Engine/Classes/Engine/Engine.h
+++ b/Engine/Source/Runtime/Engine/Classes/Engine/Engine.h
@@ -1696,6 +1696,9 @@ public:
 	bool HandleDumpTicksCommand( UWorld* InWorld, const TCHAR* Cmd, FOutputDevice& Ar );
 	bool HandleGammaCommand( const TCHAR* Cmd, FOutputDevice& Ar );
 
+	bool HandleRecordAnimationCommand(UWorld* InWorld, const TCHAR* InStr, FOutputDevice& Ar);
+	bool HandleStopRecordAnimationCommand(UWorld* InWorld, const TCHAR* InStr, FOutputDevice& Ar);
+
 	// Only compile in when STATS is set
 #if STATS
 	bool HandleDumpParticleMemCommand( const TCHAR* Cmd, FOutputDevice& Ar );
diff --git a/Engine/Source/Runtime/Engine/Private/GameEngine.cpp b/Engine/Source/Runtime/Engine/Private/GameEngine.cpp
--- a/Engine/Source/Runtime/Engine/Private/GameEngine.cpp
+++ b/Engine/Source/Runtime/Engine/Private/GameEngine.cpp
@@ -40,6 +40,10 @@
 #include "GameDelegates.h"
 #include "Engine/CoreSettings.h"
 
+#if WITH_EDITOR
+#include "Editor/UnrealEd/Public/Animation/AnimationRecorder.h"
+#endif
+
 ENGINE_API bool GDisallowNetworkTravel = false;
 
 /** Benchmark results to the log */
@@ -1195,6 +1199,14 @@ void UGameEngine::Tick( float DeltaSeconds, bool bIdleMode )
 			bCheckForMovieCapture = false;
 		}
 	}
+
+#if WITH_EDITOR
+	// tick animation recorder. available only in editor builds
+	if (!IsRunningDedicatedServer() && !IsRunningCommandlet())
+	{
+		FAnimationRecorderManager::Get().Tick(DeltaSeconds);
+	}
+#endif
 }
 
 
diff --git a/Engine/Source/Runtime/Engine/Private/UnrealEngine.cpp b/Engine/Source/Runtime/Engine/Private/UnrealEngine.cpp
--- a/Engine/Source/Runtime/Engine/Private/UnrealEngine.cpp
+++ b/Engine/Source/Runtime/Engine/Private/UnrealEngine.cpp
@@ -117,6 +117,10 @@
 
 #include "GameFramework/OnlineSession.h"
 
+#if WITH_EDITOR
+#include "Editor/UnrealEd/Public/Animation/AnimationRecorder.h"
+#endif
+
 #include "InstancedReferenceSubobjectHelper.h"
 
 DEFINE_LOG_CATEGORY(LogEngine);
@@ -3033,6 +3037,16 @@ bool UEngine::Exec( UWorld* InWorld, const TCHAR* Cmd, FOutputDevice& Ar )
 	{
 		GNetworkProfiler.Exec( InWorld, Cmd, Ar );
 	}
+#endif
+#if WITH_EDITOR
+	else if (FParse::Command(&Cmd, TEXT("RecordAnimation")))
+	{
+		return HandleRecordAnimationCommand(InWorld, Cmd, Ar);
+	}
+	else if (FParse::Command(&Cmd, TEXT("StopRecordingAnimation")))
+	{
+		return HandleStopRecordAnimationCommand(InWorld, Cmd, Ar);
+	}
 #endif
 	else 
 	{
@@ -6125,6 +6139,111 @@ bool UEngine::HandleGammaCommand( const TCHAR* Cmd, FOutputDevice& Ar )
 	return true;
 }
 
+bool UEngine::HandleRecordAnimationCommand(UWorld* InWorld, const TCHAR* InStr, FOutputDevice& Ar)
+{
+#if WITH_EDITOR
+	const TCHAR* Str = InStr;
+	// parse actor name
+	TCHAR ActorName[128];
+	AActor* FoundActor = nullptr;
+	if (FParse::Token(Str, ActorName, ARRAY_COUNT(ActorName), 0))
+	{
+		FString const ActorNameStr = FString(ActorName);
+		for (ULevel const* Level : InWorld->GetLevels())
+		{
+			if (Level)
+			{
+				for (AActor* Actor : Level->Actors)
+				{
+					if (Actor)
+					{
+						if (Actor->GetName() == ActorNameStr)
+						{
+							FoundActor = Actor;
+							break;
+						}
+					}
+				}
+			}
+		}
+	}
+
+	if (FoundActor)
+	{
+		USkeletalMeshComponent* const SkelComp = FoundActor->FindComponentByClass<USkeletalMeshComponent>();
+		if (SkelComp)
+		{
+			TCHAR AssetPath[256];
+			FParse::Token(Str, AssetPath, ARRAY_COUNT(AssetPath), 0);
+			FString const AssetName = FPackageName::GetLongPackageAssetName(AssetPath);
+			return FAnimationRecorderManager::Get().RecordAnimation(FoundActor, SkelComp, AssetPath, AssetName);
+		}
+	}
+#endif		// WITH_EDITOR
+
+	return false;
+}
+
+bool UEngine::HandleStopRecordAnimationCommand(UWorld* InWorld, const TCHAR* InStr, FOutputDevice& Ar)
+{
+#if WITH_EDITOR
+	const TCHAR* Str = InStr;
+
+	// parse actor name
+	TCHAR ActorName[128];
+	AActor* FoundActor = nullptr;
+	bool bStopAll = false;
+	if (FParse::Token(Str, ActorName, ARRAY_COUNT(ActorName), 0))
+	{
+		FString const ActorNameStr = FString(ActorName);
+
+		if (ActorNameStr.ToLower() == TEXT("all"))
+		{
+			bStopAll = true;
+		}
+		else if (InWorld)
+		{
+			// search for the actor by name
+			for (ULevel* Level : InWorld->GetLevels())
+			{
+				if (Level)
+				{
+					for (AActor* Actor : Level->Actors)
+					{
+						if (Actor)
+						{
+							if (Actor->GetName() == ActorNameStr)
+							{
+								FoundActor = Actor;
+								break;
+							}
+						}
+					}
+				}
+			}
+		}
+	}
+
+	if (bStopAll)
+	{
+		FAnimationRecorderManager::Get().StopRecordingAllAnimations();
+		return true;
+	}
+	else if (FoundActor)
+	{
+		USkeletalMeshComponent* const SkelComp = FoundActor->FindComponentByClass<USkeletalMeshComponent>();
+		if (SkelComp)
+		{
+			FAnimationRecorderManager::Get().StopRecordingAnimation(FoundActor, SkelComp);
+			return true;
+		}
+	}
+#endif		// WITH_EDITOR
+
+	return false;
+}
+
+
 /**
  * Computes a color to use for property coloration for the given object.
  *
