#ANIM: fix issue with importing curve value

Merging

//depot/UE4-Orion/Engine/Source/Editor/UnrealEd/Private/SkeletalMeshEdit.cpp

to //depot/UE4/Engine/Source/Editor/UnrealEd/Private/SkeletalMeshEdit.cpp

diff --git a/Engine/Source/Editor/UnrealEd/Private/SkeletalMeshEdit.cpp b/Engine/Source/Editor/UnrealEd/Private/SkeletalMeshEdit.cpp
--- a/Engine/Source/Editor/UnrealEd/Private/SkeletalMeshEdit.cpp
+++ b/Engine/Source/Editor/UnrealEd/Private/SkeletalMeshEdit.cpp
@@ -699,7 +699,7 @@ bool UnFbx::FFbxImporter::ImportCurve(const FbxAnimCurve* FbxCurve, FFloatCurve
 			FbxAnimCurveKey Key = FbxCurve->KeyGet(KeyIndex);
 			FbxTime KeyTime = Key.GetTime() - AnimTimeSpan.GetStart();
 			float Value = Key.GetValue() * ValueScale;
-			FKeyHandle NewKeyHandle = Curve->FloatCurve.AddKey(KeyTime.GetSecondDouble(), Value, true);
+			FKeyHandle NewKeyHandle = Curve->FloatCurve.AddKey(KeyTime.GetSecondDouble(), Value, false);
 
 			FbxAnimCurveDef::ETangentMode KeyTangentMode = Key.GetTangentMode();
 			FbxAnimCurveDef::EInterpolationType KeyInterpMode = Key.GetInterpolation();
