AnimRecorder: force LOD 0 on the recorded component for duration of the recording

(merge from branch)

diff --git a/Engine/Source/Editor/UnrealEd/Private/Animation/AnimationRecorder.cpp b/Engine/Source/Editor/UnrealEd/Private/Animation/AnimationRecorder.cpp
--- a/Engine/Source/Editor/UnrealEd/Private/Animation/AnimationRecorder.cpp
+++ b/Engine/Source/Editor/UnrealEd/Private/Animation/AnimationRecorder.cpp
@@ -400,6 +400,12 @@ void FAnimRecorderInstance::Init(AActor* InActor, USkeletalMeshComponent* InComp
 
 	// always autosave through this codepath
 	Recorder->bAutoSaveAsset = true;
+
+	if (InComponent)
+	{
+		CachedSkelCompForcedLodModel = InComponent->ForcedLodModel;
+		InComponent->ForcedLodModel = 1;
+	}
 }
 
 FAnimRecorderInstance::~FAnimRecorderInstance()
@@ -432,14 +438,16 @@ void FAnimRecorderInstance::Update(float DeltaTime)
 void FAnimRecorderInstance::FinishRecording()
 {
 	const FText FinishRecordingAnimationSlowTask = LOCTEXT("FinishRecordingAnimationSlowTask", "Finalizing recorded animation");
-	GWarn->BeginSlowTask(FinishRecordingAnimationSlowTask, false);
-
 	if (Recorder)
 	{
 		Recorder->StopRecord(true);
 	}
 
-	GWarn->EndSlowTask();
+	// restore force lod setting
+	if (SkelComp.IsValid())
+	{
+		SkelComp->ForcedLodModel = CachedSkelCompForcedLodModel;
+	}
 }
 
 bool FAnimationRecorderManager::RecordAnimation(AActor* Actor, USkeletalMeshComponent* Component, FString AssetPath, FString AssetName)
diff --git a/Engine/Source/Editor/UnrealEd/Public/Animation/AnimationRecorder.h b/Engine/Source/Editor/UnrealEd/Public/Animation/AnimationRecorder.h
--- a/Engine/Source/Editor/UnrealEd/Public/Animation/AnimationRecorder.h
+++ b/Engine/Source/Editor/UnrealEd/Public/Animation/AnimationRecorder.h
@@ -69,6 +69,9 @@ public:
 	FString AssetPath;
 	FString AssetName;
 
+	/** Original ForcedLodModel setting on the SkelComp, so we can modify it and restore it when we are done. */
+	int CachedSkelCompForcedLodModel;
+
 	FAnimationRecorder* Recorder;
 };
 
