Cooker now saves stats periodically so that you can have stats if the build times out.
Added some more stats to help track down issues.

diff --git a/Engine/Source/Editor/CookingStats/Private/CookingStats.cpp b/Engine/Source/Editor/CookingStats/Private/CookingStats.cpp
--- a/Engine/Source/Editor/CookingStats/Private/CookingStats.cpp
+++ b/Engine/Source/Editor/CookingStats/Private/CookingStats.cpp
@@ -60,7 +60,7 @@ bool FCookingStats::GetTagValue(const FName& Key, const FName& TagName, FString&
 	return true;
 }
 
-bool FCookingStats::SaveStatsAsCSV(const FString& Filename)
+bool FCookingStats::SaveStatsAsCSV(const FString& Filename) const
 {
 	FString Output;
 	FScopeLock ScopeLock(&SyncObject);
diff --git a/Engine/Source/Editor/CookingStats/Private/CookingStats.h b/Engine/Source/Editor/CookingStats/Private/CookingStats.h
--- a/Engine/Source/Editor/CookingStats/Private/CookingStats.h
+++ b/Engine/Source/Editor/CookingStats/Private/CookingStats.h
@@ -57,7 +57,7 @@ public:
 	* @param Filename file name to save comma delimited stats to
 	* @return true if succeeded
 	*/
-	virtual bool SaveStatsAsCSV(const FString& Filename) override;
+	virtual bool SaveStatsAsCSV(const FString& Filename) const override;
 
 private:
 	FName RunGuid; // unique guid for this run of the editor
diff --git a/Engine/Source/Editor/CookingStats/Public/ICookingStats.h b/Engine/Source/Editor/CookingStats/Public/ICookingStats.h
--- a/Engine/Source/Editor/CookingStats/Public/ICookingStats.h
+++ b/Engine/Source/Editor/CookingStats/Public/ICookingStats.h
@@ -58,7 +58,7 @@ public:
 	 * @param Filename file name to save comma delimited stats to 
 	 * @return true if succeeded
 	 */
-	virtual bool SaveStatsAsCSV(const FString& Filename) = 0;
+	virtual bool SaveStatsAsCSV(const FString& Filename) const = 0;
 };
 
 
diff --git a/Engine/Source/Editor/UnrealEd/Classes/CookOnTheSide/CookOnTheFlyServer.h b/Engine/Source/Editor/UnrealEd/Classes/CookOnTheSide/CookOnTheFlyServer.h
--- a/Engine/Source/Editor/UnrealEd/Classes/CookOnTheSide/CookOnTheFlyServer.h
+++ b/Engine/Source/Editor/UnrealEd/Classes/CookOnTheSide/CookOnTheFlyServer.h
@@ -1292,6 +1292,12 @@ private:
 	bool SaveCurrentIniSettings( const ITargetPlatform* TargetPlatform ) const;
 
 
+	/**
+	 * Save the accumulated cooker stats to the Saved\Stats folder
+	 * saves to different file name pending running multiprocess cooker / not
+	 */
+	void SaveCookerStats() const;
+
 	/**
 	 * IsCookFlagSet
 	 * 
diff --git a/Engine/Source/Editor/UnrealEd/Private/CookOnTheFlyServer.cpp b/Engine/Source/Editor/UnrealEd/Private/CookOnTheFlyServer.cpp
--- a/Engine/Source/Editor/UnrealEd/Private/CookOnTheFlyServer.cpp
+++ b/Engine/Source/Editor/UnrealEd/Private/CookOnTheFlyServer.cpp
@@ -466,8 +466,8 @@ namespace CookOnTheFlyStats
 			if (CookingStats)
 			{
 				static const FName NAME_CookerStat(TEXT("CookerStat"));
-				double Duration = FPlatformTime::Seconds() - StartTime * 1000.0f;
-				CookingStats->AddTagValue(NAME_CookerStat, TagName, FString::Printf(TEXT("%f"), (float)(Duration)));
+				double Duration = (FPlatformTime::Seconds() - StartTime) * 1000.0f;
+				CookingStats->AddTagValue(NAME_CookerStat, TagName, FString::Printf(TEXT("%fms"), (float)(Duration)));
 			}
 		}
 
@@ -1564,6 +1564,17 @@ uint32 UCookOnTheFlyServer::TickCookOnTheSide( const float TimeSlice, uint32 &Co
 		}
 	}
 
+	{
+		static double SavedStatsStart = FPlatformTime::Seconds();
+		double Now = FPlatformTime::Seconds();
+		const double Timeout = 60 * 10; // 10 minutes
+		if ((Now - SavedStatsStart) > Timeout)
+		{
+			SaveCookerStats();
+			SavedStatsStart = Now;
+		}
+	}
+
 	while (!GIsRequestingExit || CurrentCookMode == ECookMode::CookByTheBook)
 	{
 		// if we just cooked a map then don't process anything the rest of this tick
@@ -1672,7 +1683,7 @@ uint32 UCookOnTheFlyServer::TickCookOnTheSide( const float TimeSlice, uint32 &Co
 			}
 		}
 
-
+		check(IsInGameThread());
 		if (NeverCookPackageList.Contains(ToBuild.GetFilename()))
 		{
 #if DEBUG_COOKONTHEFLY
@@ -2027,6 +2038,7 @@ uint32 UCookOnTheFlyServer::TickCookOnTheSide( const float TimeSlice, uint32 &Co
 				}
 
 				const FName StandardPackageFilename = GetCachedStandardPackageFileFName(Package);
+				check(IsInGameThread());
 				if (NeverCookPackageList.Contains(StandardPackageFilename))
 				{
 					// refuse to save this package, it's clearly one of the undesirables
@@ -3038,10 +3050,17 @@ bool UCookOnTheFlyServer::IniSettingsOutOfDate( const ITargetPlatform* TargetPla
 	}
 
 	check( FoundCookedIniVersionStrings );
+
+	TSet<FString> CookedIniVersionStrings;
+	for (const auto& CookedString : *FoundCookedIniVersionStrings)
+	{
+		CookedIniVersionStrings.Add(CookedString);
+	}
+
 	bool bCurrentIniSettingsChanged = false;
 	for ( const auto& CurrentVersionString : CurrentIniVersionStrings ) 
 	{
-		if ( FoundCookedIniVersionStrings->Contains(CurrentVersionString) == false )
+		if ( CookedIniVersionStrings.Contains(CurrentVersionString) == false )
 		{
 			bCurrentIniSettingsChanged = true;
 			break;
@@ -3794,7 +3813,15 @@ const FString UCookOnTheFlyServer::GetCookedAssetRegistryFilename(const FString&
 	return CookedAssetRegistryFilename;
 }
 
+void UCookOnTheFlyServer::SaveCookerStats() const
+{
+	static const FName CookingStatsName("CookingStats");
+	FCookingStatsModule& CookingStatsModule = FModuleManager::LoadModuleChecked<FCookingStatsModule>(CookingStatsName);
+	ICookingStats& CookingStats = CookingStatsModule.Get();
 
+	const FString StatsFilename = GetStatsFilename(CookByTheBookOptions->ChildCookFilename);
+	CookingStats.SaveStatsAsCSV(StatsFilename);
+}
 
 void UCookOnTheFlyServer::CookByTheBookFinished()
 {
@@ -3805,13 +3832,9 @@ void UCookOnTheFlyServer::CookByTheBookFinished()
 	UPackage::WaitForAsyncFileWrites();
 
 	GetDerivedDataCacheRef().WaitForQuiescence(true);
-	
-	static const FName CookingStatsName("CookingStats");
-	FCookingStatsModule& CookingStatsModule = FModuleManager::LoadModuleChecked<FCookingStatsModule>(CookingStatsName);
-	ICookingStats& CookingStats = CookingStatsModule.Get();
 
-	const FString StatsFilename = GetStatsFilename(CookByTheBookOptions->ChildCookFilename);
-	CookingStats.SaveStatsAsCSV(StatsFilename);
+
+	SaveCookerStats();
 
 	if (IsChildCooker())
 	{
@@ -4706,7 +4729,7 @@ void UCookOnTheFlyServer::MaybeMarkPackageAsAlreadyLoaded(UPackage *Package)
 		UE_LOG(LogCook, Display, TEXT("Marking %s as reloading for cooker because it's been cooked for platforms%s."), *StandardName.ToString(), *Platforms);
 	}
 
-
+	check(IsInGameThread());
 	if (NeverCookPackageList.Contains(StandardName))
 	{
 		bShouldMarkAsAlreadyProcessed = true;
