Integrating CL 2668961 into main:

Improvements and fixes for dependency preloading
 - Fix check when mutliple load requests reference a shared asset
 - Added a package load flag to disable dependency preloading on that package
 - Fix a crash when loading default materials by disabling preloading
 - Added WIP code to cooker to sort maps in pak file.

Display level load times for persistent levels when using "stat levels"

diff --git a/Engine/Source/Editor/UnrealEd/Private/Commandlets/ChunkManifestGenerator.cpp b/Engine/Source/Editor/UnrealEd/Private/Commandlets/ChunkManifestGenerator.cpp
--- a/Engine/Source/Editor/UnrealEd/Private/Commandlets/ChunkManifestGenerator.cpp
+++ b/Engine/Source/Editor/UnrealEd/Private/Commandlets/ChunkManifestGenerator.cpp
@@ -503,7 +503,7 @@ void FChunkManifestGenerator::AddAssetToFileOrderRecursive(FAssetData* InAsset,
 		OutEncounteredNames.Add(InAsset->PackageName);
 
 		TArray<FName> Dependencies;
-		AssetRegistry.GetDependencies(InAsset->PackageName, Dependencies);
+		AssetRegistry.GetDependencies(InAsset->PackageName, Dependencies, EAssetRegistryDependencyType::Hard);
 
 		for (auto DependencyName : Dependencies)
 		{
@@ -602,10 +602,33 @@ inline void ConvertFilenameToPakFormat(FString& InOutPath)
 	}
 }
 
+void InsertIntoOrder(const TMap<FName, FAssetData*>& InAssets, TMap<FAssetData*, int32>& InOrder, FName InPackageName, int32 InOrderIndex)
+{
+	if (InAssets.Contains(InPackageName))
+	{
+		InOrder.Add(InAssets[InPackageName], InOrderIndex);
+	}
+}
+
 FString FChunkManifestGenerator::CreateCookerFileOrderString(const TMap<FName, FAssetData*>& InAssetData, const TArray<FName>& InMaps)
 {
 	FString FileOrderString;
+	TArray<FAssetData*> TopLevelMapNodes;
 	TArray<FAssetData*> TopLevelNodes;
+	TMap<FAssetData*, int32> PreferredMapOrders;
+
+#if 0
+	// TODO: PreferredMapOrders should be filled out with the maps in the order we wish them to be pak'd.
+	int32 Order = 0;
+	static const TCHAR* MapOrderList[] = 
+	{
+	};
+
+	for (int32 i = 0; i < ARRAY_COUNT(MapOrderList); ++i)
+	{
+		InsertIntoOrder(InAssetData, PreferredMapOrders, FName(MapOrderList[i]), i);
+	}
+#endif
 
 	for (auto Asset : InAssetData)
 	{
@@ -632,17 +655,31 @@ FString FChunkManifestGenerator::CreateCookerFileOrderString(const TMap<FName, F
 		{
 			if (bIsMap)
 			{
-				TopLevelNodes.Insert(Asset.Value, 0);
+				TopLevelMapNodes.Add(Asset.Value);
 			}
 			else
 			{
-				TopLevelNodes.Insert(Asset.Value, TopLevelNodes.Num());
+				TopLevelNodes.Add(Asset.Value);
 			}
 		}
 	}
 
+	TopLevelMapNodes.Sort([&PreferredMapOrders](const FAssetData& A, const FAssetData& B)
+	{
+		auto OrderA = PreferredMapOrders.Find(&A);
+		auto OrderB = PreferredMapOrders.Find(&B);
+		auto IndexA = OrderA ? *OrderA : INT_MAX;
+		auto IndexB = OrderB ? *OrderB : INT_MAX;
+		return IndexA < IndexB;
+	});
+
 	TArray<FName> FileOrder;
 	TArray<FName> EncounteredNames;
+	for (auto Asset : TopLevelMapNodes)
+	{
+		AddAssetToFileOrderRecursive(Asset, FileOrder, EncounteredNames, InAssetData, InMaps);
+	}
+
 	for (auto Asset : TopLevelNodes)
 	{
 		AddAssetToFileOrderRecursive(Asset, FileOrder, EncounteredNames, InAssetData, InMaps);
diff --git a/Engine/Source/Runtime/CoreUObject/Private/Serialization/AsyncLoading.cpp b/Engine/Source/Runtime/CoreUObject/Private/Serialization/AsyncLoading.cpp
--- a/Engine/Source/Runtime/CoreUObject/Private/Serialization/AsyncLoading.cpp
+++ b/Engine/Source/Runtime/CoreUObject/Private/Serialization/AsyncLoading.cpp
@@ -463,12 +463,18 @@ void FAsyncLoadingThread::ProcessAsyncPackageRequest(FAsyncPackageDesc* InReques
 				{
 					if (!InDependencyTracker.Contains(DependencyName))
 					{
-						QueuedPackagesCounter.Increment();
-						const int32 RequestID = GPackageRequestID.Increment();
-						FAsyncLoadingThread::Get().AddPendingRequest(RequestID);
-						FAsyncPackageDesc DependencyPackageRequest(RequestID, DependencyName, NAME_None, FGuid(), FLoadPackageAsyncDelegate(), InRequest->PackageFlags, INDEX_NONE, InRequest->Priority);
-						InDependencyTracker.Add(InRequest->Name);
-						ProcessAsyncPackageRequest(&DependencyPackageRequest, InRootPackage, InDependencyTracker);
+						// Try and find the package in memory first
+						UPackage* ClassPackage = FindObjectFast<UPackage>( NULL, DependencyName, false, false );
+
+						if (ClassPackage == nullptr)
+						{
+							QueuedPackagesCounter.Increment();
+							const int32 RequestID = GPackageRequestID.Increment();
+							FAsyncLoadingThread::Get().AddPendingRequest(RequestID);
+							FAsyncPackageDesc DependencyPackageRequest(RequestID, DependencyName, NAME_None, FGuid(), FLoadPackageAsyncDelegate(), InRequest->PackageFlags, INDEX_NONE, InRequest->Priority);
+							InDependencyTracker.Add(InRequest->Name);
+							ProcessAsyncPackageRequest(&DependencyPackageRequest, InRootPackage, InDependencyTracker);
+						}
 					}
 				}
 			}
@@ -506,9 +512,9 @@ int32 FAsyncLoadingThread::CreateAsyncPackagesFromQueue()
 		double Timer = 0;
 		{
 			SCOPE_SECONDS_COUNTER(Timer);
-			DependencyTracker.Empty();
 			for (auto PackageRequest : QueueCopy)
 			{
+				DependencyTracker.Empty();
 				ProcessAsyncPackageRequest(PackageRequest, nullptr, DependencyTracker);
 			}
 		}
diff --git a/Engine/Source/Runtime/CoreUObject/Private/UObject/UObjectGlobals.cpp b/Engine/Source/Runtime/CoreUObject/Private/UObject/UObjectGlobals.cpp
--- a/Engine/Source/Runtime/CoreUObject/Private/UObject/UObjectGlobals.cpp
+++ b/Engine/Source/Runtime/CoreUObject/Private/UObject/UObjectGlobals.cpp
@@ -988,8 +988,9 @@ UPackage* LoadPackageInternal(UPackage* InOuter, const TCHAR* InLongPackageName,
 		FExclusiveLoadPackageTimeTracker::FScopedPackageTracker Tracker(Result);
 
 #if !WITH_EDITOR
+		bool bAllowDependencyPreloading = ((LoadFlags & LOAD_DisableDependencyPreloading) == 0);
 		static auto CVarPreloadDependencies = IConsoleManager::Get().FindConsoleVariable(TEXT("s.PreloadPackageDependencies"));
-		if (CVarPreloadDependencies && CVarPreloadDependencies->GetInt() != 0)
+		if (bAllowDependencyPreloading && CVarPreloadDependencies && CVarPreloadDependencies->GetInt() != 0)
 		{
 			if (Result->PackageFlags & PKG_ProcessingDependencies)
 			{
diff --git a/Engine/Source/Runtime/CoreUObject/Public/UObject/ObjectBase.h b/Engine/Source/Runtime/CoreUObject/Public/UObject/ObjectBase.h
--- a/Engine/Source/Runtime/CoreUObject/Public/UObject/ObjectBase.h
+++ b/Engine/Source/Runtime/CoreUObject/Public/UObject/ObjectBase.h
@@ -54,7 +54,7 @@ enum ELoadFlags
 //	LOAD_Unused						= 0x00000200,
 //	LOAD_Unused						= 0x00000400,
 //	LOAD_Unused						= 0x00000800,
-//	LOAD_Unused						= 0x00001000,
+	LOAD_DisableDependencyPreloading = 0x00001000,	// Bypass dependency preloading system
 	LOAD_Quiet						= 0x00002000,   // No log warnings.
 	LOAD_FindIfFail					= 0x00004000,	// Tries FindObject if a linker cannot be obtained (e.g. package is currently being compiled)
 	LOAD_MemoryReader				= 0x00008000,	// Loads the file into memory and serializes from there.
diff --git a/Engine/Source/Runtime/Engine/Private/Materials/Material.cpp b/Engine/Source/Runtime/Engine/Private/Materials/Material.cpp
--- a/Engine/Source/Runtime/Engine/Private/Materials/Material.cpp
+++ b/Engine/Source/Runtime/Engine/Private/Materials/Material.cpp
@@ -346,7 +346,7 @@ void UMaterialInterface::InitDefaultMaterials()
 				GDefaultMaterials[Domain] = FindObject<UMaterial>(NULL,GDefaultMaterialNames[Domain]);
 				if (GDefaultMaterials[Domain] == NULL)
 				{
-					GDefaultMaterials[Domain] = LoadObject<UMaterial>(NULL,GDefaultMaterialNames[Domain],NULL,LOAD_None,NULL);
+					GDefaultMaterials[Domain] = LoadObject<UMaterial>(NULL, GDefaultMaterialNames[Domain], NULL, LOAD_DisableDependencyPreloading, NULL);
 					checkf(GDefaultMaterials[Domain] != NULL, TEXT("Cannot load default material '%s'"), GDefaultMaterialNames[Domain]);
 				}
 				GDefaultMaterials[Domain]->AddToRoot();
diff --git a/Engine/Source/Runtime/Engine/Private/UnrealEngine.cpp b/Engine/Source/Runtime/Engine/Private/UnrealEngine.cpp
--- a/Engine/Source/Runtime/Engine/Private/UnrealEngine.cpp
+++ b/Engine/Source/Runtime/Engine/Private/UnrealEngine.cpp
@@ -11734,13 +11734,14 @@ int32 UEngine::RenderStatLevels(UWorld* World, FViewport* Viewport, FCanvas* Can
 	{
 		// First entry - always persistent level
 		FString MapName	= SubLevelsStatusList[0].PackageName.ToString();
+		UPackage* LevelPackage = FindObjectFast<UPackage>(NULL, SubLevelsStatusList[0].PackageName);
 		if (SubLevelsStatusList[0].bPlayerInside)
 		{
-			MapName = *FString::Printf( TEXT("->  %s"), *MapName );
+			MapName = *FString::Printf(TEXT("->  %s - %4.1f sec"), *MapName, LevelPackage->GetLoadTime());
 		}
 		else
 		{
-			MapName = *FString::Printf( TEXT("    %s"), *MapName );
+			MapName = *FString::Printf(TEXT("    %s - %4.1f sec"), *MapName, LevelPackage->GetLoadTime());
 		}
 
 		Canvas->DrawShadowedString(X, Y, *MapName, GetSmallFont(), FColor(127, 127, 127));
