#ANIM: don't support break curve, but just auto

Since weight of tangent isn't supported, break curve doesn't work either. Technically if no weight, it should work, but most of break curve coming in seems to have weight, so I've added ticket UE-20414, and removed break curve support.
For now it will all import as auto.

Merging

//depot/UE4-Orion/Engine/Source/Editor/UnrealEd/Private/SkeletalMeshEdit.cpp

to //depot/UE4/Engine/Source/Editor/UnrealEd/Private/SkeletalMeshEdit.cpp

diff --git a/Engine/Source/Editor/UnrealEd/Private/SkeletalMeshEdit.cpp b/Engine/Source/Editor/UnrealEd/Private/SkeletalMeshEdit.cpp
--- a/Engine/Source/Editor/UnrealEd/Private/SkeletalMeshEdit.cpp
+++ b/Engine/Source/Editor/UnrealEd/Private/SkeletalMeshEdit.cpp
@@ -742,17 +742,21 @@ bool UnFbx::FFbxImporter::ImportCurve(const FbxAnimCurve* FbxCurve, FFloatCurve
 				break;
 			}
 
+			// break or any other tangent mode doesn't work well with DCC
+			// it's because we don't support tangent weights, break with tangent weights won't work
+			// I added new ticket to support this, but meanwhile, we'll have to just import using auto. 
+			// @Todo: fix me: UE-20414
 			// when we import tangent, we only support break or user
 			// since it's modified by DCC and we only assume these two are valid
 			// auto does our own stuff, which doesn't work with what you see in DCC
-			if (KeyTangentMode & FbxAnimCurveDef::eTangentBreak)
-			{
-				NewTangentMode = RCTM_Break;
-			}
-			else
-			{
-				NewTangentMode = RCTM_User;
-			}
+// 			if (KeyTangentMode & FbxAnimCurveDef::eTangentGenericBreak)
+// 			{
+// 				NewTangentMode = RCTM_Break;
+// 			}
+// 			else
+// 			{
+// 				NewTangentMode = RCTM_User;
+// 			}
 
 			// @fix me : weight of tangent is not used, but we'll just save this for future where we might use it. 
 			switch (KeyTangentWeightMode)
