Integrated github pull request #1315 (JIRA UE-17656): Modified DDS cubemap loading to prevent an array overflow.

diff --git a/Engine/Source/Editor/UnrealEd/Private/Factories/EditorFactories.cpp b/Engine/Source/Editor/UnrealEd/Private/Factories/EditorFactories.cpp
--- a/Engine/Source/Editor/UnrealEd/Private/Factories/EditorFactories.cpp
+++ b/Engine/Source/Editor/UnrealEd/Private/Factories/EditorFactories.cpp
@@ -4011,6 +4011,12 @@ UTexture* UTextureFactory::ImportTexture(UClass* Class, UObject* InParent, FName
 			return nullptr;
 		}
 
+		if (NumMips > MAX_TEXTURE_MIP_COUNT)
+		{
+			Warn->Logf(ELogVerbosity::Error, TEXT("DDS file contains an unsupported number of mipmap levels."));
+			return nullptr;
+		}
+
 		// create the cube texture
 		UTextureCube* TextureCube = CreateTextureCube( InParent, Name, Flags );
 
