Fixes for UserDefinedEnums becoming corrupt in Blueprints after modifying the UserDefinedEnum asset.

 FEnumEditorUtils::ResolveEnumerator now iterates over every enum stored in the resolver to find the enum's name and from there it's up-to-date value.

diff --git a/Engine/Source/Editor/UnrealEd/Private/Kismet2/EnumEditorUtils.cpp b/Engine/Source/Editor/UnrealEd/Private/Kismet2/EnumEditorUtils.cpp
--- a/Engine/Source/Editor/UnrealEd/Private/Kismet2/EnumEditorUtils.cpp
+++ b/Engine/Source/Editor/UnrealEd/Private/Kismet2/EnumEditorUtils.cpp
@@ -318,14 +318,16 @@ int32 FEnumEditorUtils::ResolveEnumerator(const UEnum* Enum, FArchive& Ar, int32
 	const FArchiveEnumeratorResolver* EnumeratorResolver = (FArchiveEnumeratorResolver*)(&Ar);
 	if(Enum == EnumeratorResolver->Enum)
 	{
-		const auto& OldNames = EnumeratorResolver->OldNames;
-		if(EnumeratorValue < OldNames.Num())
+		for (TPair<FName, uint8> OldName : EnumeratorResolver->OldNames)
 		{
-			const FName EnumeratorName = OldNames[EnumeratorValue].Key;
-			const int32 NewEnumIndex = Enum->GetValueByName(EnumeratorName);
-			if(INDEX_NONE != NewEnumIndex)
+			if (OldName.Value == EnumeratorValue)
 			{
-				return NewEnumIndex;
+				const FName EnumeratorName = OldName.Key;
+				const int32 NewEnumValue = Enum->GetValueByName(EnumeratorName);
+				if(INDEX_NONE != NewEnumValue)
+				{
+					return NewEnumValue;
+				}
 			}
 		}
 		return Enum->GetMaxEnumValue();
