#ANIM: curve : fix curve editor zoom to fit vertial issue when min/max is same

Merging

//depot/UE4-Orion/Engine/Source/Editor/UnrealEd/Private/SCurveEditor.cpp

to //depot/UE4-Orion/Engine/Source/Editor/UnrealEd/Private/SCurveEditor.cpp

diff --git a/Engine/Source/Editor/UnrealEd/Private/SCurveEditor.cpp b/Engine/Source/Editor/UnrealEd/Private/SCurveEditor.cpp
--- a/Engine/Source/Editor/UnrealEd/Private/SCurveEditor.cpp
+++ b/Engine/Source/Editor/UnrealEd/Private/SCurveEditor.cpp
@@ -2434,6 +2434,13 @@ void SCurveEditor::ZoomToFitVertical()
 
 		const float MinZoomRange = (TotalKeys > 0 ) ? CONST_MinViewRange: CONST_DefaultZoomRange;
 
+		// if in max and in min is same, then include 0.f
+		if (InMax == InMin)
+		{
+			InMax = FMath::Max(InMax, 0.f);
+			InMin = FMath::Min(InMin, 0.f);
+		}
+
 		// Clamp the minimum size
 		float Size = InMax - InMin;
 		if( Size < MinZoomRange )
