Forcing pin tooltip data to load before writing BP API data

diff --git a/Engine/Source/Editor/UnrealEd/Private/Commandlets/GenerateBlueprintAPICommandlet.cpp b/Engine/Source/Editor/UnrealEd/Private/Commandlets/GenerateBlueprintAPICommandlet.cpp
--- a/Engine/Source/Editor/UnrealEd/Private/Commandlets/GenerateBlueprintAPICommandlet.cpp
+++ b/Engine/Source/Editor/UnrealEd/Private/Commandlets/GenerateBlueprintAPICommandlet.cpp
@@ -861,6 +861,10 @@ static void GenerateBlueprintAPIUtils::DumpActionMenuItem(uint32 Indent, FGraphA
 			{
 				if (!Pin->bHidden)
 				{
+					// @hack: Some pin data will not be available until requested for display, specifically tooltip strings for call function nodes:
+					FString Scratch;
+					Node->GetPinHoverText(*Pin, Scratch);
+
 					if (!bFirst)
 					{
 						ActionEntry += ",";
