Fix to infinite loop crash when performing a viewport camera transition while Matinee has a locked actor defined.

diff --git a/Engine/Source/Editor/UnrealEd/Private/EditorViewportClient.cpp b/Engine/Source/Editor/UnrealEd/Private/EditorViewportClient.cpp
--- a/Engine/Source/Editor/UnrealEd/Private/EditorViewportClient.cpp
+++ b/Engine/Source/Editor/UnrealEd/Private/EditorViewportClient.cpp
@@ -930,6 +930,10 @@ void FEditorViewportClient::Tick(float DeltaTime)
 
 	FViewportCameraTransform& ViewTransform = GetViewTransform();
 	const bool bIsAnimating = ViewTransform.UpdateTransition();
+	if (bIsAnimating && GetViewportType() == LVT_Perspective)
+	{
+		PerspectiveCameraMoved();
+	}
 
 	if ( bIsTracking )
 	{
diff --git a/Engine/Source/Editor/UnrealEd/Private/LevelEditorViewport.cpp b/Engine/Source/Editor/UnrealEd/Private/LevelEditorViewport.cpp
--- a/Engine/Source/Editor/UnrealEd/Private/LevelEditorViewport.cpp
+++ b/Engine/Source/Editor/UnrealEd/Private/LevelEditorViewport.cpp
@@ -2068,17 +2068,12 @@ void FLevelEditorViewportClient::UpdateViewForLockedActor()
 	bUseControllingActorViewInfo = false;
 	ControllingActorViewInfo = FMinimalViewInfo();
 
-	const AActor* Actor = ActorLockedByMatinee.IsValid() ? ActorLockedByMatinee.Get() : ActorLockedToCamera.Get();
+	AActor* Actor = ActorLockedByMatinee.IsValid() ? ActorLockedByMatinee.Get() : ActorLockedToCamera.Get();
 	if( Actor != NULL )
 	{
 		// Check if the viewport is transitioning
 		FViewportCameraTransform& ViewTransform = GetViewTransform();
-		if (ViewTransform.IsPlaying())
-		{
-			// Move actor to the transitioned viewport
-			PerspectiveCameraMoved();
-		}
-		else
+		if (!ViewTransform.IsPlaying())
 		{
 			// Update transform
 			if (Actor->GetAttachParentActor() != NULL)
