Added support for OrthoFreelook viewport types when calculating the view rotation matrix for the viewport.

diff --git a/Engine/Source/Editor/UnrealEd/Private/EditorViewportClient.cpp b/Engine/Source/Editor/UnrealEd/Private/EditorViewportClient.cpp
--- a/Engine/Source/Editor/UnrealEd/Private/EditorViewportClient.cpp
+++ b/Engine/Source/Editor/UnrealEd/Private/EditorViewportClient.cpp
@@ -801,6 +801,14 @@ FSceneView* FEditorViewportClient::CalcSceneView(FSceneViewFamily* ViewFamily)
 					FPlane(0, 1, 0, 0),
 					FPlane(0, 0, ViewLocation.X, 1));
 			}
+			else if (EffectiveViewportType == LVT_OrthoFreelook)
+			{
+				ViewInitOptions.ViewRotationMatrix = FMatrix(
+					FPlane(0, 0, 1, 0),
+					FPlane(1, 0, 0, 0),
+					FPlane(0, 1, 0, 0),
+					FPlane(0, 0, ViewLocation.X, 1));
+			}
 			else
 			{
 				// Unknown viewport type
