Changed placement offset for lights so that they're placed closer to the cursor.

diff --git a/Engine/Source/Editor/UnrealEd/Private/Factories/ActorFactory.cpp b/Engine/Source/Editor/UnrealEd/Private/Factories/ActorFactory.cpp
--- a/Engine/Source/Editor/UnrealEd/Private/Factories/ActorFactory.cpp
+++ b/Engine/Source/Editor/UnrealEd/Private/Factories/ActorFactory.cpp
@@ -1551,7 +1551,7 @@ UActorFactoryDirectionalLight::UActorFactoryDirectionalLight(const FObjectInitia
 {
 	DisplayName = LOCTEXT("DirectionalLightDisplayName", "Directional Light");
 	NewActorClass = ADirectionalLight::StaticClass();
-	SpawnPositionOffset = FVector(200, 0, 0);
+	SpawnPositionOffset = FVector(50, 0, 0);
 	bUseSurfaceOrientation = true;
 }
 
@@ -1563,7 +1563,7 @@ UActorFactorySpotLight::UActorFactorySpotLight(const FObjectInitializer& ObjectI
 {
 	DisplayName = LOCTEXT("SpotLightDisplayName", "Spot Light");
 	NewActorClass = ASpotLight::StaticClass();
-	SpawnPositionOffset = FVector(200, 0, 0);
+	SpawnPositionOffset = FVector(50, 0, 0);
 	bUseSurfaceOrientation = true;
 }
 
@@ -1575,7 +1575,7 @@ UActorFactoryPointLight::UActorFactoryPointLight(const FObjectInitializer& Objec
 {
 	DisplayName = LOCTEXT("PointLightDisplayName", "Point Light");
 	NewActorClass = APointLight::StaticClass();
-	SpawnPositionOffset = FVector(200, 0, 0);
+	SpawnPositionOffset = FVector(50, 0, 0);
 	bUseSurfaceOrientation = true;
 }
 
@@ -1597,7 +1597,7 @@ UActorFactorySphereReflectionCapture::UActorFactorySphereReflectionCapture(const
 {
 	DisplayName = LOCTEXT("ReflectionCaptureSphereDisplayName", "Sphere Reflection Capture");
 	NewActorClass = ASphereReflectionCapture::StaticClass();
-	SpawnPositionOffset = FVector(200, 0, 0);
+	SpawnPositionOffset = FVector(50, 0, 0);
 	bUseSurfaceOrientation = true;
 }
 
@@ -1609,7 +1609,7 @@ UActorFactoryBoxReflectionCapture::UActorFactoryBoxReflectionCapture(const FObje
 {
 	DisplayName = LOCTEXT("ReflectionCaptureBoxDisplayName", "Box Reflection Capture");
 	NewActorClass = ABoxReflectionCapture::StaticClass();
-	SpawnPositionOffset = FVector(200, 0, 0);
+	SpawnPositionOffset = FVector(50, 0, 0);
 	bUseSurfaceOrientation = true;
 }
 
@@ -1621,7 +1621,7 @@ UActorFactoryPlaneReflectionCapture::UActorFactoryPlaneReflectionCapture(const F
 {
 	DisplayName = LOCTEXT("ReflectionCapturePlaneDisplayName", "Plane Reflection Capture");
 	NewActorClass = APlaneReflectionCapture::StaticClass();
-	SpawnPositionOffset = FVector(200, 0, 0);
+	SpawnPositionOffset = FVector(50, 0, 0);
 	bUseSurfaceOrientation = true;
 }
 
