Grouped actors pasted with 'Paste Here' now have their group actor and pivot position set correctly.

diff --git a/Engine/Source/Editor/UnrealEd/Private/EditorServer.cpp b/Engine/Source/Editor/UnrealEd/Private/EditorServer.cpp
--- a/Engine/Source/Editor/UnrealEd/Private/EditorServer.cpp
+++ b/Engine/Source/Editor/UnrealEd/Private/EditorServer.cpp
@@ -3378,6 +3378,9 @@ void UEditorEngine::PasteSelectedActorsFromClipboard( UWorld* InWorld, const FTe
 				const FVector Location = bbox.GetCenter();
 				const FVector Adjust = Origin - Location;
 
+				// List of group actors in the selection
+				TArray<AGroupActor*> GroupActors;
+
 				// Move the actors.
 				AActor* SingleActor = NULL;
 				for ( FSelectionIterator It( GEditor->GetSelectedActorIterator() ) ; It ; ++It )
@@ -3385,6 +3388,16 @@ void UEditorEngine::PasteSelectedActorsFromClipboard( UWorld* InWorld, const FTe
 					AActor* Actor = static_cast<AActor*>( *It );
 					checkSlow( Actor->IsA(AActor::StaticClass()) );
 
+					// If this actor is in a group, add it to the list
+					if (GEditor->bGroupingActive)
+					{
+						AGroupActor* ActorGroupRoot = AGroupActor::GetRootForActor(Actor, true, true);
+						if (ActorGroupRoot)
+						{
+							GroupActors.AddUnique(ActorGroupRoot);
+						}
+					}
+
 					SingleActor = Actor;
 					Actor->SetActorLocation(Actor->GetActorLocation() + Adjust, false);
 					Actor->PostEditMove(true);
@@ -3392,7 +3405,16 @@ void UEditorEngine::PasteSelectedActorsFromClipboard( UWorld* InWorld, const FTe
 
 				// Update the pivot location.
 				check(SingleActor);
-				SetPivot( SingleActor->GetActorLocation(), false, true );
+				SetPivot(SingleActor->GetActorLocation(), false, true);
+
+				// If grouping is active, go through the unique group actors and update the group actor location
+				if (GEditor->bGroupingActive)
+				{
+					for (AGroupActor* GroupActor : GroupActors)
+					{
+						GroupActor->CenterGroupLocation();
+					}
+				}
 			}
 		}
 
