Moved editor promotion settings tests into their own file.

diff --git a/Engine/Source/Editor/UnrealEd/Private/Tests/AutomationEditorPromotionCommon.cpp b/Engine/Source/Editor/UnrealEd/Private/Tests/AutomationEditorPromotionCommon.cpp
--- a/Engine/Source/Editor/UnrealEd/Private/Tests/AutomationEditorPromotionCommon.cpp
+++ b/Engine/Source/Editor/UnrealEd/Private/Tests/AutomationEditorPromotionCommon.cpp
@@ -7,6 +7,8 @@
 
 #define LOCTEXT_NAMESPACE "EditorPromotionTestCommon"
 
+DEFINE_LOG_CATEGORY_STATIC(LogEditorPromotionTests, Log, All);
+
 /**
 * Finds a visible widget by type.     SLOW!!!!!
 *
@@ -161,6 +163,26 @@ void FEditorPromotionTestUtilities::SendCommandToCurrentEditor(const FInputChord
 	}
 }
 
+/**
+* Gets an object property value by name
+*
+* @param TargetObject - The object to modify
+* @param InVariableName - The name of the property
+*/
+FString FEditorPromotionTestUtilities::GetPropertyByName(UObject* TargetObject, const FString& InVariableName)
+{
+	UProperty* FoundProperty = FindField<UProperty>(TargetObject->GetClass(), *InVariableName);
+	if (FoundProperty)
+	{
+		FString ValueString;
+		const uint8* PropertyAddr = FoundProperty->ContainerPtrToValuePtr<uint8>(TargetObject);
+		FoundProperty->ExportTextItem(ValueString, PropertyAddr, NULL, NULL, PPF_None);
+		return ValueString;
+	}
+	return TEXT("");
+}
+
+
 /**
 * Sets an object property value by name
 *
@@ -183,4 +205,62 @@ void FEditorPromotionTestUtilities::SetPropertyByName(UObject* TargetObject, con
 	}
 }
 
+/**
+* Ends a PIE session
+*/
+void FEditorPromotionTestUtilities::EndPIE()
+{
+	GUnrealEd->RequestEndPlayMap();
+}
+
+/**
+* Takes an automation screenshot
+*
+* @param ScreenshotName - The sub name to use for the screenshot
+*/
+void FEditorPromotionTestUtilities::TakeScreenshot(const FString& ScreenshotName, bool bUseTopWindow)
+{
+	//Update the screenshot name, then take a screenshot.
+	if (FAutomationTestFramework::GetInstance().IsScreenshotAllowed())
+	{
+		TSharedPtr<SWindow> Window;
+
+		if (bUseTopWindow)
+		{
+			Window = FSlateApplication::Get().GetActiveTopLevelWindow();
+		}
+		else
+		{
+			//Find the main editor window
+			TArray<TSharedRef<SWindow> > AllWindows;
+			FSlateApplication::Get().GetAllVisibleWindowsOrdered(AllWindows);
+			if (AllWindows.Num() == 0)
+			{
+				UE_LOG(LogEditorPromotionTests, Error, TEXT("ERROR: Could not find the main editor window."));
+				return;
+			}
+
+			Window = AllWindows[0];
+		}
+
+		if (Window.IsValid())
+		{
+			FString ScreenshotFileName;
+			const FString TestName = FString::Printf(TEXT("EditorBuildPromotion/%s"), *ScreenshotName);
+			AutomationCommon::GetScreenshotPath(TestName, ScreenshotFileName, false);
+
+			TSharedRef<SWidget> WindowRef = Window.ToSharedRef();
+
+			TArray<FColor> OutImageData;
+			FIntVector OutImageSize;
+			FSlateApplication::Get().TakeScreenshot(WindowRef, OutImageData, OutImageSize);
+			FAutomationTestFramework::GetInstance().OnScreenshotCaptured().ExecuteIfBound(OutImageSize.X, OutImageSize.Y, OutImageData, ScreenshotFileName);
+		}
+		else
+		{
+			UE_LOG(LogEditorPromotionTests, Error, TEXT("Failed to find editor window for screenshot (%s)"), *ScreenshotName);
+		}
+	}
+}
+
 #undef LOCTEXT_NAMESPACE
diff --git a/Engine/Source/Editor/UnrealEd/Private/Tests/EditorBuildPromotionTests.cpp b/Engine/Source/Editor/UnrealEd/Private/Tests/EditorBuildPromotionTests.cpp
--- a/Engine/Source/Editor/UnrealEd/Private/Tests/EditorBuildPromotionTests.cpp
+++ b/Engine/Source/Editor/UnrealEd/Private/Tests/EditorBuildPromotionTests.cpp
@@ -162,53 +162,53 @@ namespace EditorBuildPromotionTestUtils
 		double Duration;
 	};
 
-	/**
-	* Exports the current editor keybindings
-	*
-	* @param TargetFilename - The name of the file to export to
-	*/
-	static void ExportKeybindings(const FString& TargetFilename)
-	{
-		FInputBindingManager::Get().SaveInputBindings();
-		GConfig->Flush(false, GEditorKeyBindingsIni);
-		IFileManager::Get().Copy(*TargetFilename, *GEditorKeyBindingsIni);
-	}
-
-	/**
-	* Imports new editor keybindings  (Unused)
-	*
-	* @param TargetFilename - The name of the file to import from
-	*/
-	static void ImportKeybindings(const FString& TargetFilename)
-	{
-		GConfig->Flush(true, GEditorKeyBindingsIni);
-		IFileManager::Get().Copy(*GEditorKeyBindingsIni, *TargetFilename);
-		GConfig->LoadFile(GEditorKeyBindingsIni);
-	}
-
-	/**
-	* Exports the current editor settings
-	*
-	* @param TargetFilename - The name of the file to export to
-	*/
-	static void ExportEditorSettings(const FString& TargetFilename)
-	{
-		FInputBindingManager::Get().SaveInputBindings();
-		GConfig->Flush(false, GEditorPerProjectIni);
-		IFileManager::Get().Copy(*TargetFilename, *GEditorPerProjectIni);
-	}
-
-	/**
-	* Imports new editor settings (unused)
-	*
-	* @param TargetFilename - The name of the file to export to
-	*/
-	static void ImportEditorSettings(const FString& TargetFilename)
-	{
-		GConfig->Flush(true, GEditorPerProjectIni);
-		IFileManager::Get().Copy(*GEditorPerProjectIni, *TargetFilename);
-		GConfig->LoadFile(GEditorPerProjectIni);
-	}
+	///**
+	//* Exports the current editor keybindings
+	//*
+	//* @param TargetFilename - The name of the file to export to
+	//*/
+	//static void ExportKeybindings(const FString& TargetFilename)
+	//{
+	//	FInputBindingManager::Get().SaveInputBindings();
+	//	GConfig->Flush(false, GEditorKeyBindingsIni);
+	//	IFileManager::Get().Copy(*TargetFilename, *GEditorKeyBindingsIni);
+	//}
+
+	///**
+	//* Imports new editor keybindings  (Unused)
+	//*
+	//* @param TargetFilename - The name of the file to import from
+	//*/
+	//static void ImportKeybindings(const FString& TargetFilename)
+	//{
+	//	GConfig->Flush(true, GEditorKeyBindingsIni);
+	//	IFileManager::Get().Copy(*GEditorKeyBindingsIni, *TargetFilename);
+	//	GConfig->LoadFile(GEditorKeyBindingsIni);
+	//}
+
+	///**
+	//* Exports the current editor settings
+	//*
+	//* @param TargetFilename - The name of the file to export to
+	//*/
+	//static void ExportEditorSettings(const FString& TargetFilename)
+	//{
+	//	FInputBindingManager::Get().SaveInputBindings();
+	//	GConfig->Flush(false, GEditorPerProjectIni);
+	//	IFileManager::Get().Copy(*TargetFilename, *GEditorPerProjectIni);
+	//}
+
+	///**
+	//* Imports new editor settings (unused)
+	//*
+	//* @param TargetFilename - The name of the file to export to
+	//*/
+	//static void ImportEditorSettings(const FString& TargetFilename)
+	//{
+	//	GConfig->Flush(true, GEditorPerProjectIni);
+	//	IFileManager::Get().Copy(*GEditorPerProjectIni, *TargetFilename);
+	//	GConfig->LoadFile(GEditorPerProjectIni);
+	//}
 
 	/**
 	* Sends the MaterialEditor->Apply UI command
@@ -736,53 +736,6 @@ namespace EditorBuildPromotionTestUtils
 		GUnrealEd->BuildLighting(LightingBuildOptions);
 	}
 
-	/**
-	* Creates a new keybinding chord and sets it for the supplied command and context
-	*
-	* @param CommandContext - The context of the command
-	* @param Command - The command name to get
-	* @param Key - The keybinding chord key
-	* @param ModifierKey - The keybinding chord modifier key
-	*/
-	static FInputChord SetKeybinding(const FString& CommandContext, const FString& Command, const FKey Key, const EModifierKey::Type ModifierKey)
-	{
-		FInputChord NewChord(Key, ModifierKey);
-		if (!FEditorPromotionTestUtilities::SetEditorKeybinding(CommandContext, Command, NewChord))
-		{
-			// Trigger a failure when used in an automated test
-			UE_LOG(LogEditorBuildPromotionTests, Error, TEXT("Could not find keybinding for %s using context %s"), *Command, *CommandContext);
-		}
-		return NewChord;
-	}
-
-	/**
-	* Retrieves the current keybinding for a command and compares it against the expected binding.
-	* Triggers an automation test failure if keybind cannot be retrieved or does not match expected binding.
-	*
-	* @param CommandContext - The context of the command
-	* @param Command - The command name to get
-	* @param ExpectedChord - The chord value to compare against
-	*/
-	static void CompareKeybindings(const FString& CommandContext, const FString& Command, FInputChord ExpectedChord)
-	{
-		FInputChord CurrentChord = FEditorPromotionTestUtilities::GetEditorKeybinding(CommandContext, Command);
-		if (!CurrentChord.IsValidChord())
-		{
-			UE_LOG(LogEditorBuildPromotionTests, Error, TEXT("Could not find keybinding for %s using context %s"), *Command, *CommandContext);
-		}
-		else
-		{
-			if (CurrentChord == ExpectedChord)
-			{
-				UE_LOG(LogEditorBuildPromotionTests, Display, TEXT("%s keybinding correct."), *Command);
-			}
-			else
-			{
-				UE_LOG(LogEditorBuildPromotionTests, Error, TEXT("%s keybinding incorrect."), *Command);
-			}
-		}
-	}
-
 	/**
 	* Gets an object property value by name
 	*
@@ -813,64 +766,6 @@ namespace EditorBuildPromotionTestUtils
 		GUnrealEd->RequestPlaySession(false, ActiveLevelViewport, bSimulateInEditor, NULL, NULL, -1, false);
 	}
 
-	/**
-	* Ends a PIE session
-	*/
-	static void EndPIE()
-	{
-		GUnrealEd->RequestEndPlayMap();
-	}
-
-	/**
-	* Takes an automation screenshot
-	*
-	* @param ScreenshotName - The sub name to use for the screenshot
-	*/
-	static void TakeScreenshot(const FString& ScreenshotName, bool bUseTopWindow = false)
-	{
-		//Update the screenshot name, then take a screenshot.
-		if (FAutomationTestFramework::GetInstance().IsScreenshotAllowed())
-		{
-			TSharedPtr<SWindow> Window;
-
-			if (bUseTopWindow)
-			{
-				Window = FSlateApplication::Get().GetActiveTopLevelWindow();
-			}
-			else
-			{
-				//Find the main editor window
-				TArray<TSharedRef<SWindow> > AllWindows;
-				FSlateApplication::Get().GetAllVisibleWindowsOrdered(AllWindows);
-				if (AllWindows.Num() == 0)
-				{
-					UE_LOG(LogEditorAutomationTests, Error, TEXT("ERROR: Could not find the main editor window."));
-					return;
-				}
-
-				Window = AllWindows[0];
-			}
-
-			if (Window.IsValid())
-			{
-				FString ScreenshotFileName;
-				const FString TestName = FString::Printf(TEXT("EditorBuildPromotion/%s"), *ScreenshotName);
-				AutomationCommon::GetScreenshotPath(TestName, ScreenshotFileName, false);
-
-				TSharedRef<SWidget> WindowRef = Window.ToSharedRef();
-
-				TArray<FColor> OutImageData;
-				FIntVector OutImageSize;
-				FSlateApplication::Get().TakeScreenshot(WindowRef, OutImageData, OutImageSize);
-				FAutomationTestFramework::GetInstance().OnScreenshotCaptured().ExecuteIfBound(OutImageSize.X, OutImageSize.Y, OutImageData, ScreenshotFileName);
-			}
-			else
-			{
-				UE_LOG(LogEditorBuildPromotionTests, Error, TEXT("Failed to find editor window for screenshot (%s)"), *ScreenshotName);
-			}
-		}
-	}
-
 	/**
 	* Adds a default mesh to the level
 	*
@@ -1290,15 +1185,15 @@ namespace BuildPromotionTestHelper
 			FString PropertyName = Assets[AssetIndex].PropertyName;
 			FString NewPropertyValue = Assets[AssetIndex].PropertyValue;
 
-			FString OldPropertyValue = EditorBuildPromotionTestUtils::GetPropertyByName(CurrentAsset, PropertyName);
+			FString OldPropertyValue = FEditorPromotionTestUtilities::GetPropertyByName(CurrentAsset, PropertyName);
 			FEditorPromotionTestUtilities::SetPropertyByName(CurrentAsset, PropertyName, NewPropertyValue);
 			UE_LOG(LogEditorBuildPromotionTests, Display, TEXT("Modified asset.  %s = %s"), *PropertyName, *NewPropertyValue);
 
 			//Get the property again and use that to compare the redo action.  Parsing the new value may change the formatting a bit. ie) 100 becomes 100.0000
-			FString ParsedNewValue = EditorBuildPromotionTestUtils::GetPropertyByName(CurrentAsset, PropertyName);
+			FString ParsedNewValue = FEditorPromotionTestUtilities::GetPropertyByName(CurrentAsset, PropertyName);
 
 			GEditor->UndoTransaction();
-			FString CurrentValue = EditorBuildPromotionTestUtils::GetPropertyByName(CurrentAsset, PropertyName);
+			FString CurrentValue = FEditorPromotionTestUtilities::GetPropertyByName(CurrentAsset, PropertyName);
 			if (CurrentValue == OldPropertyValue)
 			{
 				UE_LOG(LogEditorBuildPromotionTests, Display, TEXT("Undo %s change successful"), *PropertyName);
@@ -1309,7 +1204,7 @@ namespace BuildPromotionTestHelper
 			}
 
 			GEditor->RedoTransaction();
-			CurrentValue = EditorBuildPromotionTestUtils::GetPropertyByName(CurrentAsset, PropertyName);
+			CurrentValue = FEditorPromotionTestUtilities::GetPropertyByName(CurrentAsset, PropertyName);
 			if (CurrentValue == ParsedNewValue)
 			{
 				UE_LOG(LogEditorBuildPromotionTests, Display, TEXT("Redo %s change successful"), *PropertyName);
@@ -2664,7 +2559,7 @@ namespace BuildPromotionTestHelper
 				IAssetEditorInstance* AssetEditor = FAssetEditorManager::Get().FindEditorForAsset(BlueprintObject, true);
 				IBlueprintEditor* BlueprintEditor = (IBlueprintEditor*)AssetEditor;
 
-				EditorBuildPromotionTestUtils::TakeScreenshot(TEXT("BlueprintComponentVariables"), true);
+				FEditorPromotionTestUtilities::TakeScreenshot(TEXT("BlueprintComponentVariables"), true);
 
 				UE_LOG(LogEditorBuildPromotionTests, Display, TEXT("Switched to components mode"));
 				BlueprintEditor->SetCurrentMode(FBlueprintEditorApplicationModes::BlueprintComponentsMode);
@@ -2820,7 +2715,7 @@ namespace BuildPromotionTestHelper
 				IAssetEditorInstance* AssetEditor = FAssetEditorManager::Get().FindEditorForAsset(BlueprintObject, true);
 				FBlueprintEditor* BlueprintEditor = (FBlueprintEditor*)AssetEditor;
 
-				EditorBuildPromotionTestUtils::TakeScreenshot(TEXT("BlueprintMeshVariable"), true);
+				FEditorPromotionTestUtilities::TakeScreenshot(TEXT("BlueprintMeshVariable"), true);
 
 				UE_LOG(LogEditorBuildPromotionTests, Display, TEXT("Switched to components mode"));
 				BlueprintEditor->SetCurrentMode(FBlueprintEditorApplicationModes::BlueprintComponentsMode);
@@ -2838,7 +2733,7 @@ namespace BuildPromotionTestHelper
 		{
 			if (BlueprintObject)
 			{
-				EditorBuildPromotionTestUtils::TakeScreenshot(TEXT("BlueprintComponent"), true);
+				FEditorPromotionTestUtilities::TakeScreenshot(TEXT("BlueprintComponent"), true);
 			}
 			return true;
 		}
@@ -3199,7 +3094,7 @@ namespace BuildPromotionTestHelper
 		{
 			if (BlueprintObject)
 			{
-				EditorBuildPromotionTestUtils::TakeScreenshot(TEXT("BlueprintPIE_Start"));
+				FEditorPromotionTestUtilities::TakeScreenshot(TEXT("BlueprintPIE_Start"));
 			}
 			return true;
 		}
@@ -3250,9 +3145,9 @@ namespace BuildPromotionTestHelper
 			{
 				if (GEditor->PlayWorld != NULL)
 				{
-					EditorBuildPromotionTestUtils::TakeScreenshot(TEXT("BlueprintPIE_End"));
+					FEditorPromotionTestUtilities::TakeScreenshot(TEXT("BlueprintPIE_End"));
 				}
-				EditorBuildPromotionTestUtils::EndPIE();
+				FEditorPromotionTestUtilities::EndPIE();
 			}
 			return true;
 		}
@@ -3333,7 +3228,7 @@ namespace BuildPromotionTestHelper
 			{
 				UE_LOG(LogEditorBuildPromotionTests, Display, TEXT("Clearing the breakpoint"));
 				FKismetDebugUtilities::ClearBreakpoints(BlueprintObject);
-				EditorBuildPromotionTestUtils::EndPIE();
+				FEditorPromotionTestUtilities::EndPIE();
 				CurrentWorld = PlacedBlueprint->GetWorld();
 			}
 			return true;
@@ -3518,9 +3413,9 @@ namespace BuildPromotionTestHelper
 				if (GEditor->PlayWorld != NULL)
 				{
 					//Take a screenshot and end the PIE session
-					EditorBuildPromotionTestUtils::TakeScreenshot(TEXT("LevelBlueprint"), false);
+					FEditorPromotionTestUtilities::TakeScreenshot(TEXT("LevelBlueprint"), false);
 				}
-				EditorBuildPromotionTestUtils::EndPIE();
+				FEditorPromotionTestUtilities::EndPIE();
 			}
 			return true;
 		}
@@ -3574,7 +3469,7 @@ namespace BuildPromotionTestHelper
 			ULevel* Level = World->GetCurrentLevel();
 
 			UE_LOG(LogEditorBuildPromotionTests, Display, TEXT("Taking a screenshot"));
-			EditorBuildPromotionTestUtils::TakeScreenshot(TEXT("VisibilityOff"));
+			FEditorPromotionTestUtilities::TakeScreenshot(TEXT("VisibilityOff"));
 
 			UE_LOG(LogEditorBuildPromotionTests, Display, TEXT("Turning level visibility on"));
 			bool bShouldBeVisible = true;
@@ -3593,7 +3488,7 @@ namespace BuildPromotionTestHelper
 			ULevel* Level = World->GetCurrentLevel();
 
 			UE_LOG(LogEditorBuildPromotionTests, Display, TEXT("Taking a screenshot"));
-			EditorBuildPromotionTestUtils::TakeScreenshot(TEXT("VisibilityOn"));
+			FEditorPromotionTestUtilities::TakeScreenshot(TEXT("VisibilityOn"));
 
 			FEditorFileUtils::SaveLevel(Level, TEXT("/Game/Maps/EditorBuildPromotionTest"));
 			GUnrealEd->Exec(World, TEXT("MAP REBUILD ALLVISIBLE"));
@@ -3667,25 +3562,6 @@ bool FBuildPromotionInitialCleanupTest::RunTest(const FString& Parameters)
 	return true;
 }
 
-/**
-* Latent command to check if PIE is running
-*/
-DEFINE_LATENT_AUTOMATION_COMMAND(FSettingsCheckForPIECommand);
-bool FSettingsCheckForPIECommand::Update()
-{
-	if (GEditor->PlayWorld != NULL)
-	{
-		//Success
-		UE_LOG(LogEditorBuildPromotionTests, Display, TEXT("PlayInEditor keyboard shortcut success"));
-		EditorBuildPromotionTestUtils::EndPIE();
-	}
-	else
-	{
-		UE_LOG(LogEditorBuildPromotionTests, Error, TEXT("PlayInEditor keyboard shortcut failed"));
-	}
-	return true;
-}
-
 /**
 * Latent command to run the main build promotion test
 */
@@ -3703,86 +3579,6 @@ bool FLogTestResultCommand::Update()
 	return true;
 }
 
-/**
-* Automation test that handles setting keybindings and editor preferences
-*/
-IMPLEMENT_SIMPLE_AUTOMATION_TEST(FBuildPromotionSettingsTest, "System.Promotion.Editor Promotion Pass.Step 1 Main Editor Test.Settings", EAutomationTestFlags::EditorContext | EAutomationTestFlags::EngineFilter);
-
-bool FBuildPromotionSettingsTest::RunTest(const FString& Parameters)
-{
-
-//Previous Editor Settings
-	UE_LOG(LogEditorBuildPromotionTests, Display, TEXT("Exporting Current keybindings and editor settings"));
-
-	//Export the original keybindings
-	const FString TargetOriginalKeybindFile = FString::Printf(TEXT("%s/BuildPromotion/OriginalKeybindings-%d.ini"), *FPaths::AutomationDir(), FEngineVersion::Current().GetChangelist());
-	EditorBuildPromotionTestUtils::ExportKeybindings(TargetOriginalKeybindFile);
-
-	//Export the original preferences
-	const FString TargetOriginalPreferenceFile = FString::Printf(TEXT("%s/BuildPromotion/OriginalPreferences-%d.ini"), *FPaths::AutomationDir(), FEngineVersion::Current().GetChangelist());
-	EditorBuildPromotionTestUtils::ExportEditorSettings(TargetOriginalPreferenceFile);
-
-//New Editor Settings
-	//Bind H to CreateEmptyLayer keybinding
-	UE_LOG(LogEditorBuildPromotionTests, Display, TEXT("Binding create empty layer shortcut"));
-	FInputChord NewCreateChord = EditorBuildPromotionTestUtils::SetKeybinding(TEXT("LayersView"), TEXT("CreateEmptyLayer"), EKeys::H, EModifierKey::None);
-	
-	//Bind J to RequestRenameLayer
-	UE_LOG(LogEditorBuildPromotionTests, Display, TEXT("Binding request rename layer shortcut"));
-	FInputChord NewRenameChord = EditorBuildPromotionTestUtils::SetKeybinding(TEXT("LayersView"), TEXT("RequestRenameLayer"), EKeys::J, EModifierKey::None);
-	
-	// Bind CTRL+L to PIE
-	UE_LOG(LogEditorBuildPromotionTests, Display, TEXT("Binding play shortcut (PIE)"));
-	FInputChord NewPIEChord = EditorBuildPromotionTestUtils::SetKeybinding(TEXT("PlayWorld"), TEXT("RepeatLastPlay"), EKeys::L, EModifierKey::Control);
-	
-	//Export the keybindings
-	const FString TargetKeybindFile = FString::Printf(TEXT("%s/BuildPromotion/Keybindings-%d.ini"), *FPaths::AutomationDir(), FEngineVersion::Current().GetChangelist());
-	UE_LOG(LogEditorBuildPromotionTests, Display, TEXT("Exporting keybind"));
-	EditorBuildPromotionTestUtils::ExportKeybindings(TargetKeybindFile);
-
-
-	UEditorStyleSettings* EditorStyleSettings = GetMutableDefault<UEditorStyleSettings>();
-	FString OldStyleSetting = EditorBuildPromotionTestUtils::GetPropertyByName(EditorStyleSettings, TEXT("bUseSmallToolBarIcons"));
-	FEditorPromotionTestUtilities::SetPropertyByName(EditorStyleSettings, TEXT("bUseSmallToolBarIcons"), TEXT("true"));
-
-	UE_LOG(LogEditorBuildPromotionTests, Display, TEXT("Set UseSmallToolBarIcons"));
-
-	//Export the preferences
-	const FString TargetPreferenceFile = FString::Printf(TEXT("%s/BuildPromotion/Preferences-%d.ini"), *FPaths::AutomationDir(), FEngineVersion::Current().GetChangelist());
-	EditorBuildPromotionTestUtils::ExportEditorSettings(TargetPreferenceFile);
-
-	//Take a screenshot of the small icons
-	EditorBuildPromotionTestUtils::TakeScreenshot(TEXT("Small Toolbar Icons"));
-
-	UE_LOG(LogEditorBuildPromotionTests, Display, TEXT("Exported keybindings and preferences to %s/BuildPromotion"), *FPaths::AutomationDir());
-
-	//Change the setting back
-	FEditorPromotionTestUtilities::SetPropertyByName(EditorStyleSettings, TEXT("bUseSmallToolBarIcons"), OldStyleSetting);
-
-	// Verify keybindings were assigned correctly
-	EditorBuildPromotionTestUtils::CompareKeybindings(TEXT("LayersView"), TEXT("CreateEmptyLayer"), NewCreateChord);
-	EditorBuildPromotionTestUtils::CompareKeybindings(TEXT("LayersView"), TEXT("RequestRenameLayer"), NewRenameChord);
-	EditorBuildPromotionTestUtils::CompareKeybindings(TEXT("PlayWorld"), TEXT("RepeatLastPlay"), NewPIEChord);
-
-	//Focus the main editor
-	TArray< TSharedRef<SWindow> > AllWindows;
-	FSlateApplication::Get().GetAllVisibleWindowsOrdered(AllWindows);
-	FSlateApplication::Get().ProcessWindowActivatedEvent(FWindowActivateEvent(FWindowActivateEvent::EA_Activate, AllWindows[0]));
-
-	//Send the PIE event
-	FKeyEvent PIEKeyEvent(EKeys::L, FModifierKeysState(false, false, true, false, false, false, false, false, false), false, 0/*UserIndex*/, 0x4C, 0x4C);
-	FSlateApplication::Get().ProcessKeyDownEvent(PIEKeyEvent);
-	FSlateApplication::Get().ProcessKeyUpEvent(PIEKeyEvent);
-
-	UE_LOG(LogEditorBuildPromotionTests, Display, TEXT("Sent PIE keyboard shortcut"));
-
-	ADD_LATENT_AUTOMATION_COMMAND(FWaitLatentCommand(3.f));
-	ADD_LATENT_AUTOMATION_COMMAND(FSettingsCheckForPIECommand());
-	ADD_LATENT_AUTOMATION_COMMAND(FLogTestResultCommand(&ExecutionInfo));
-
-	return true;
-}
-
 /**
 * Latent command to run the main build promotion test
 */
@@ -3811,7 +3607,7 @@ DEFINE_LATENT_AUTOMATION_COMMAND(FEndPIECommand);
 bool FEndPIECommand::Update()
 {
 	UE_LOG(LogEditorBuildPromotionTests, Display, TEXT("Ending PIE"));
-	EditorBuildPromotionTestUtils::EndPIE();
+	FEditorPromotionTestUtilities::EndPIE();
 	return true;
 }
 
diff --git a/Engine/Source/Editor/UnrealEd/Private/Tests/EditorSettingsTests.cpp b/Engine/Source/Editor/UnrealEd/Private/Tests/EditorSettingsTests.cpp
new file mode 100644
--- /dev/null
+++ b/Engine/Source/Editor/UnrealEd/Private/Tests/EditorSettingsTests.cpp
@@ -0,0 +1,290 @@
+// Copyright 1998-2015 Epic Games, Inc. All Rights Reserved.
+
+#include "UnrealEd.h"
+
+// Automation
+#include "AutomationTest.h"
+#include "AutomationEditorCommon.h"
+#include "AutomationEditorPromotionCommon.h"
+
+#define LOCTEXT_NAMESPACE "EditorSettingsPromotionTests"
+
+DEFINE_LOG_CATEGORY_STATIC(LogEditorSettingsTests, All, All);
+
+//Latent commands
+DEFINE_LATENT_AUTOMATION_COMMAND(FSettingsCheckForPIECommand);
+DEFINE_LATENT_AUTOMATION_COMMAND_ONE_PARAMETER(FLogSettingsTestResult, FAutomationTestExecutionInfo*, InExecutionInfo);
+
+//Tests
+IMPLEMENT_SIMPLE_AUTOMATION_TEST(FEditorSettingsKeybindingsTest, "System.Promotion.Editor.Settings.Keybindings", EAutomationTestFlags::EditorContext | EAutomationTestFlags::EngineFilter);
+IMPLEMENT_SIMPLE_AUTOMATION_TEST(FEditorSettingsPreferencesTest, "System.Promotion.Editor.Settings.Preferences", EAutomationTestFlags::EditorContext | EAutomationTestFlags::EngineFilter);
+
+/**
+* Helper functions used by the settings automation tests
+*/
+namespace EditorSettingsTestUtils
+{
+	/**
+	* Exports the current editor keybindings
+	*
+	* @param IniSettings - Settings to export
+	* @param TargetFilename - The name of the file to export to
+	*/
+	static bool ExportSettings(FString& IniSettings, const FString& TargetFilename)
+	{
+		FInputBindingManager::Get().SaveInputBindings();
+		GConfig->Flush(false, IniSettings);
+		ECopyResult Expected = ECopyResult::COPY_OK;
+		return ECopyResult::COPY_OK == IFileManager::Get().Copy(*TargetFilename, *IniSettings);
+	}
+
+	/**
+	* Exports the current editor keybindings
+	*
+	* @param IniSettings - Settings to import to
+	* @param TargetFilename - The name of the file to import
+	*/
+	static bool ImportSettings(FString& IniSettings, const FString& SourceFilename)
+	{
+		GConfig->Flush(true, IniSettings);
+		if (ECopyResult::COPY_OK == IFileManager::Get().Copy(*IniSettings, *SourceFilename))
+		{
+			GConfig->LoadFile(IniSettings);
+		}
+		else
+		{
+			UE_LOG(LogEditorSettingsTests, Error, TEXT("Failed to load settings from %s"), *SourceFilename);
+		}
+	}
+
+	/**
+	* Exports the current editor keybindings
+	*
+	* @param TargetFilename - The name of the file to export to
+	*/
+	static void ExportKeybindings(const FString& TargetFilename)
+	{
+		FInputBindingManager::Get().SaveInputBindings();
+		GConfig->Flush(false, GEditorKeyBindingsIni);
+		IFileManager::Get().Copy(*TargetFilename, *GEditorKeyBindingsIni);
+	}
+
+	/**
+	* Imports new editor keybindings  (Unused)
+	*
+	* @param TargetFilename - The name of the file to import from
+	*/
+	static void ImportKeybindings(const FString& TargetFilename)
+	{
+		GConfig->Flush(true, GEditorKeyBindingsIni);
+		IFileManager::Get().Copy(*GEditorKeyBindingsIni, *TargetFilename);
+		GConfig->LoadFile(GEditorKeyBindingsIni);
+	}
+
+	/**
+	* Exports the current editor settings
+	*
+	* @param TargetFilename - The name of the file to export to
+	*/
+	static void ExportEditorSettings(const FString& TargetFilename)
+	{
+		FInputBindingManager::Get().SaveInputBindings();
+		GConfig->Flush(false, GEditorPerProjectIni);
+		IFileManager::Get().Copy(*TargetFilename, *GEditorPerProjectIni);
+	}
+
+	/**
+	* Imports new editor settings (unused)
+	*
+	* @param TargetFilename - The name of the file to export to
+	*/
+	static void ImportEditorSettings(const FString& TargetFilename)
+	{
+		GConfig->Flush(true, GEditorPerProjectIni);
+		IFileManager::Get().Copy(*GEditorPerProjectIni, *TargetFilename);
+		GConfig->LoadFile(GEditorPerProjectIni);
+	}
+
+	/**
+	* Creates a new keybinding chord and sets it for the supplied command and context
+	*
+	* @param CommandContext - The context of the command
+	* @param Command - The command name to get
+	* @param Key - The keybinding chord key
+	* @param ModifierKey - The keybinding chord modifier key
+	*/
+	static FInputChord SetKeybinding(const FString& CommandContext, const FString& Command, const FKey Key, const EModifierKey::Type ModifierKey)
+	{
+		FInputChord NewChord(Key, ModifierKey);
+		if (!FEditorPromotionTestUtilities::SetEditorKeybinding(CommandContext, Command, NewChord))
+		{
+			// Trigger a failure when used in an automated test
+			UE_LOG(LogEditorSettingsTests, Error, TEXT("Could not find keybinding for %s using context %s"), *Command, *CommandContext);
+		}
+		return NewChord;
+	}
+
+	/**
+	* Retrieves the current keybinding for a command and compares it against the expected binding.
+	* Triggers an automation test failure if keybind cannot be retrieved or does not match expected binding.
+	*
+	* @param CommandContext - The context of the command
+	* @param Command - The command name to get
+	* @param ExpectedChord - The chord value to compare against
+	*/
+	static void CompareKeybindings(const FString& CommandContext, const FString& Command, FInputChord ExpectedChord)
+	{
+		FInputChord CurrentChord = FEditorPromotionTestUtilities::GetEditorKeybinding(CommandContext, Command);
+		if (!CurrentChord.IsValidChord())
+		{
+			UE_LOG(LogEditorSettingsTests, Error, TEXT("Could not find keybinding for %s using context %s"), *Command, *CommandContext);
+		}
+		else
+		{
+			if (CurrentChord == ExpectedChord)
+			{
+				UE_LOG(LogEditorSettingsTests, Display, TEXT("%s keybinding correct."), *Command);
+			}
+			else
+			{
+				UE_LOG(LogEditorSettingsTests, Error, TEXT("%s keybinding incorrect."), *Command);
+			}
+		}
+	}
+}
+
+/**
+* Latent command to check if PIE is running
+*/
+bool FSettingsCheckForPIECommand::Update()
+{
+	if (GEditor->PlayWorld != NULL)
+	{
+		//Success
+		UE_LOG(LogEditorSettingsTests, Display, TEXT("PlayInEditor keyboard shortcut success"));
+		FEditorPromotionTestUtilities::EndPIE();
+	}
+	else
+	{
+		UE_LOG(LogEditorSettingsTests, Error, TEXT("PlayInEditor keyboard shortcut failed"));
+	}
+	return true;
+}
+
+/**
+* Latent command to run the main build promotion test
+*/
+bool FLogSettingsTestResult::Update()
+{
+	if (InExecutionInfo && InExecutionInfo->Errors.Num() > 0)
+	{
+		UE_LOG(LogEditorSettingsTests, Display, TEXT("Test failed!"));
+	}
+	else
+	{
+		UE_LOG(LogEditorSettingsTests, Display, TEXT("Test successful!"));
+	}
+	return true;
+}
+
+/**
+* Automation test that handles setting keybindings
+*/
+bool FEditorSettingsKeybindingsTest::RunTest(const FString& Parameters)
+{
+	UE_LOG(LogEditorSettingsTests, Display, TEXT("Exporting Current keybindings and editor settings"));
+
+	//Export the original keybindings
+	const FString TargetOriginalKeybindFile = FString::Printf(TEXT("%s/BuildPromotion/OriginalKeybindings-%d.ini"), *FPaths::AutomationDir(), FEngineVersion::Current().GetChangelist());
+	EditorSettingsTestUtils::ExportKeybindings(TargetOriginalKeybindFile);
+
+	//Cache original keybinding values, since reimporting settings doesn't reset values w/o a restart
+	const FInputChord OriginalCreateChord = FEditorPromotionTestUtilities::GetEditorKeybinding(TEXT("LayersView"), TEXT("CreateEmptyLayer"));
+	const FInputChord OriginalRenameChord = FEditorPromotionTestUtilities::GetEditorKeybinding(TEXT("LayersView"), TEXT("RequestRenameLayer"));
+	const FInputChord OriginalPIEChord = FEditorPromotionTestUtilities::GetEditorKeybinding(TEXT("PlayWorld"), TEXT("RepeatLastPlay"));
+
+	//New Editor Settings
+	//Bind H to CreateEmptyLayer keybinding
+	UE_LOG(LogEditorSettingsTests, Display, TEXT("Binding create empty layer shortcut"));
+	FInputChord NewCreateChord = EditorSettingsTestUtils::SetKeybinding(TEXT("LayersView"), TEXT("CreateEmptyLayer"), EKeys::H, EModifierKey::None);
+	
+	//Bind J to RequestRenameLayer
+	UE_LOG(LogEditorSettingsTests, Display, TEXT("Binding request rename layer shortcut"));
+	FInputChord NewRenameChord = EditorSettingsTestUtils::SetKeybinding(TEXT("LayersView"), TEXT("RequestRenameLayer"), EKeys::J, EModifierKey::None);
+	
+	// Bind CTRL+L to PIE
+	UE_LOG(LogEditorSettingsTests, Display, TEXT("Binding play shortcut (PIE)"));
+	FInputChord NewPIEChord = EditorSettingsTestUtils::SetKeybinding(TEXT("PlayWorld"), TEXT("RepeatLastPlay"), EKeys::L, EModifierKey::Control);
+	
+	//Export the keybindings
+	const FString TargetKeybindFile = FString::Printf(TEXT("%s/BuildPromotion/Keybindings-%d.ini"), *FPaths::AutomationDir(), FEngineVersion::Current().GetChangelist());
+	UE_LOG(LogEditorSettingsTests, Display, TEXT("Exporting keybind"));
+	EditorSettingsTestUtils::ExportKeybindings(TargetKeybindFile);
+
+	// Verify keybindings were assigned correctly
+	EditorSettingsTestUtils::CompareKeybindings(TEXT("LayersView"), TEXT("CreateEmptyLayer"), NewCreateChord);
+	EditorSettingsTestUtils::CompareKeybindings(TEXT("LayersView"), TEXT("RequestRenameLayer"), NewRenameChord);
+	EditorSettingsTestUtils::CompareKeybindings(TEXT("PlayWorld"), TEXT("RepeatLastPlay"), NewPIEChord);
+
+	//Focus the main editor
+	TArray< TSharedRef<SWindow> > AllWindows;
+	FSlateApplication::Get().GetAllVisibleWindowsOrdered(AllWindows);
+	FSlateApplication::Get().ProcessWindowActivatedEvent(FWindowActivateEvent(FWindowActivateEvent::EA_Activate, AllWindows[0]));
+
+	//Send the PIE event
+	FKeyEvent PIEKeyEvent(EKeys::L, FModifierKeysState(false, false, true, false, false, false, false, false, false), false, 0/*UserIndex*/, 0x4C, 0x4C);
+	FSlateApplication::Get().ProcessKeyDownEvent(PIEKeyEvent);
+	FSlateApplication::Get().ProcessKeyUpEvent(PIEKeyEvent);
+
+	UE_LOG(LogEditorSettingsTests, Display, TEXT("Sent PIE keyboard shortcut"));
+
+	ADD_LATENT_AUTOMATION_COMMAND(FWaitLatentCommand(3.f));
+	ADD_LATENT_AUTOMATION_COMMAND(FSettingsCheckForPIECommand());
+	ADD_LATENT_AUTOMATION_COMMAND(FLogSettingsTestResult(&ExecutionInfo));
+
+	//Import original keybindings and set changed binds back to cached values
+	EditorSettingsTestUtils::ImportKeybindings(TargetOriginalKeybindFile);
+	UE_LOG(LogEditorSettingsTests, Display, TEXT("Reloaded original keybindings"));
+
+	FEditorPromotionTestUtilities::SetEditorKeybinding(TEXT("LayersView"), TEXT("CreateEmptyLayer"), OriginalCreateChord);
+	FEditorPromotionTestUtilities::SetEditorKeybinding(TEXT("LayersView"), TEXT("RequestRenameLayer"), OriginalRenameChord);
+	FEditorPromotionTestUtilities::SetEditorKeybinding(TEXT("PlayWorld"), TEXT("RepeatLastPlay"), OriginalPIEChord);
+
+	return true;
+}
+
+/**
+* Automation test that handles setting editor preferences
+*/
+bool FEditorSettingsPreferencesTest::RunTest(const FString& Parameters)
+{
+	UE_LOG(LogEditorSettingsTests, Display, TEXT("Exporting Current keybindings and editor settings"));
+
+	//Export the original preferences
+	const FString TargetOriginalPreferenceFile = FString::Printf(TEXT("%s/BuildPromotion/OriginalPreferences-%d.ini"), *FPaths::AutomationDir(), FEngineVersion::Current().GetChangelist());
+	EditorSettingsTestUtils::ExportEditorSettings(TargetOriginalPreferenceFile);
+
+	UEditorStyleSettings* EditorStyleSettings = GetMutableDefault<UEditorStyleSettings>();
+	FString OldStyleSetting = FEditorPromotionTestUtilities::GetPropertyByName(EditorStyleSettings, TEXT("bUseSmallToolBarIcons"));
+	
+	FEditorPromotionTestUtilities::SetPropertyByName(EditorStyleSettings, TEXT("bUseSmallToolBarIcons"), TEXT("true"));
+	UE_LOG(LogEditorSettingsTests, Display, TEXT("Enabled UseSmallToolBarIcons"));
+
+	//Export the preferences
+	const FString TargetPreferenceFile = FString::Printf(TEXT("%s/BuildPromotion/Preferences-%d.ini"), *FPaths::AutomationDir(), FEngineVersion::Current().GetChangelist());
+	EditorSettingsTestUtils::ExportEditorSettings(TargetPreferenceFile);
+
+	//Take a screenshot of the small icons
+	FEditorPromotionTestUtilities::TakeScreenshot(TEXT("Small Toolbar Icons"));
+
+	//Change the setting back
+	FEditorPromotionTestUtilities::SetPropertyByName(EditorStyleSettings, TEXT("bUseSmallToolBarIcons"), OldStyleSetting);
+	UE_LOG(LogEditorSettingsTests, Display, TEXT("Set UseSmallToolBarIcons to original state"));
+
+	EditorSettingsTestUtils::ImportEditorSettings(TargetOriginalPreferenceFile);
+	UE_LOG(LogEditorSettingsTests, Display, TEXT("Reloaded original preferences"));
+
+	return true;
+}
+
+#undef LOCTEXT_NAMESPACE
diff --git a/Engine/Source/Editor/UnrealEd/Public/Tests/AutomationEditorPromotionCommon.h b/Engine/Source/Editor/UnrealEd/Public/Tests/AutomationEditorPromotionCommon.h
--- a/Engine/Source/Editor/UnrealEd/Public/Tests/AutomationEditorPromotionCommon.h
+++ b/Engine/Source/Editor/UnrealEd/Public/Tests/AutomationEditorPromotionCommon.h
@@ -72,6 +72,13 @@ public:
 	*/
 	UNREALED_API static void SendCommandToCurrentEditor(const FInputChord& InChord, const FName& WidgetTypeToFocus);
 
+	/**
+	* Gets an object property value by name
+	*
+	* @param TargetObject - The object to modify
+	* @param InVariableName - The name of the property
+	*/
+	UNREALED_API static FString GetPropertyByName(UObject* TargetObject, const FString& InVariableName);
 
 	/**
 	* Sets an object property value by name
@@ -80,4 +87,16 @@ public:
 	* @param InVariableName - The name of the property
 	*/
 	UNREALED_API static void SetPropertyByName(UObject* TargetObject, const FString& InVariableName, const FString& NewValueString);
+
+	/**
+	* Ends a PIE session
+	*/
+	UNREALED_API static void EndPIE();
+
+	/**
+	* Takes an automation screenshot
+	*
+	* @param ScreenshotName - The sub name to use for the screenshot
+	*/
+	UNREALED_API static void TakeScreenshot(const FString& ScreenshotName, bool bUseTopWindow = false);
 };
\ No newline at end of file
