CIS fix #UE4

made EditorSettingsTestUtils::ImportSettings return void rather than bool

diff --git a/Engine/Source/Editor/UnrealEd/Private/Tests/EditorSettingsTests.cpp b/Engine/Source/Editor/UnrealEd/Private/Tests/EditorSettingsTests.cpp
--- a/Engine/Source/Editor/UnrealEd/Private/Tests/EditorSettingsTests.cpp
+++ b/Engine/Source/Editor/UnrealEd/Private/Tests/EditorSettingsTests.cpp
@@ -44,7 +44,7 @@ namespace EditorSettingsTestUtils
 	* @param IniSettings - Settings to import to
 	* @param TargetFilename - The name of the file to import
 	*/
-	static bool ImportSettings(FString& IniSettings, const FString& SourceFilename)
+	static void ImportSettings(FString& IniSettings, const FString& SourceFilename)
 	{
 		GConfig->Flush(true, IniSettings);
 		if (ECopyResult::COPY_OK == IFileManager::Get().Copy(*IniSettings, *SourceFilename))
