#UE4 Clearing the parent and interface function maps on a class during RemoveStaleFunctions to allow them to be regenerated at the next request. This prevents a crash where the parent class map is populated with old function pointers.

diff --git a/Engine/Source/Editor/UnrealEd/Private/Kismet2/BlueprintEditorUtils.cpp b/Engine/Source/Editor/UnrealEd/Private/Kismet2/BlueprintEditorUtils.cpp
--- a/Engine/Source/Editor/UnrealEd/Private/Kismet2/BlueprintEditorUtils.cpp
+++ b/Engine/Source/Editor/UnrealEd/Private/Kismet2/BlueprintEditorUtils.cpp
@@ -1419,6 +1419,9 @@ static void RemoveStaleFunctions(UBlueprintGeneratedClass* Class, UBlueprint* Bl
 		}
 	}
 
+	// Clear function map caches which will be rebuilt the next time functions are searched by name
+	Class->ClearFunctionMapsCaches();
+
 	Blueprint->GeneratedClass->Children = nullptr;
 	Blueprint->GeneratedClass->Bind();
 	Blueprint->GeneratedClass->StaticLink(true);
diff --git a/Engine/Source/Runtime/CoreUObject/Public/UObject/Class.h b/Engine/Source/Runtime/CoreUObject/Public/UObject/Class.h
--- a/Engine/Source/Runtime/CoreUObject/Public/UObject/Class.h
+++ b/Engine/Source/Runtime/CoreUObject/Public/UObject/Class.h
@@ -2021,6 +2021,12 @@ public:
 		FuncMap.Remove(Function->GetFName());
 	}
 
+	void ClearFunctionMapsCaches()
+	{
+		ParentFuncMap.Empty();
+		InterfaceFuncMap.Empty();
+	}
+
 	UFunction* FindFunctionByName(FName InName, EIncludeSuperFlag::Type IncludeSuper = EIncludeSuperFlag::IncludeSuper) const;
 
 	// UObject interface.
