Improved documentation on auto reimport properties

diff --git a/Engine/Source/Editor/UnrealEd/Classes/Settings/EditorLoadingSavingSettings.h b/Engine/Source/Editor/UnrealEd/Classes/Settings/EditorLoadingSavingSettings.h
--- a/Engine/Source/Editor/UnrealEd/Classes/Settings/EditorLoadingSavingSettings.h
+++ b/Engine/Source/Editor/UnrealEd/Classes/Settings/EditorLoadingSavingSettings.h
@@ -40,15 +40,15 @@ struct FAutoReimportDirectoryConfig
 	GENERATED_USTRUCT_BODY()
 
 	/** The source directory to monitor. Either an absolute directory on the file system, or a virtual mounted path */
-	UPROPERTY(EditAnywhere, config, Category=AutoReimport)
+	UPROPERTY(EditAnywhere, config, Category=AutoReimport, meta=(ToolTip="Path to a virtual package path (eg /Game/ or /MyPlugin/), or absolute paths on disk where your source content files reside."))
 	FString SourceDirectory;
 
 	/** Where SourceDirectory points to an ordinary file system path, MountPoint specifies the virtual mounted location to import new files to. */
-	UPROPERTY(EditAnywhere, config, Category=AutoReimport)
+	UPROPERTY(EditAnywhere, config, Category=AutoReimport, meta=(ToolTip="(Optional) Specify a virtual mout point (e.g. /Game/) to map this directory to on disk. Doing so allows auto-creation of assets when a source content file is created in this folder (see below)."))
 	FString MountPoint;
 
 	/** A set of wildcard filters to apply to this directory */
-	UPROPERTY(EditAnywhere, config, Category=AutoReimport)
+	UPROPERTY(EditAnywhere, config, Category=AutoReimport, meta=(DisplayName="Include/Exclude Wildcards", ToolTip="(Optional) Specify a set of wildcards to include or exclude files from this auto-reimporter."))
 	TArray<FAutoReimportWildcard> Wildcards;
 
 	struct UNREALED_API FParseContext
@@ -100,14 +100,14 @@ private:
 public:
 
 	/**Automatically reimports textures when a change to source content is detected */
-	UPROPERTY(EditAnywhere, config, Category=AutoReimport, meta=(DisplayName="Monitor Content Directories", ToolTip="When enabled, changes to made to source content files inside the content directories will automatically be reflected in the content browser."))
+	UPROPERTY(EditAnywhere, config, Category=AutoReimport, meta=(DisplayName="Monitor Content Directories", ToolTip="When enabled, changes to made to source content files inside the content directories will automatically be reflected in the content browser.\nNote that source content files must reside in one of the monitored directories to be eligible for auto-reimport.\nAdvanced setup options are available below."))
 	bool bMonitorContentDirectories;
 
 	UPROPERTY(config)
 	TArray<FString> AutoReimportDirectories_DEPRECATED;
 
 	/** Directories being monitored for Auto Reimport */
-	UPROPERTY(EditAnywhere, config, AdvancedDisplay,Category=AutoReimport, meta=(DisplayName="Directories to Monitor", ToolTip="Lists every directory to monitor for content changes. Can be virtual package paths (eg /Game/ or /MyPlugin/), or absolute paths on disk."))
+	UPROPERTY(EditAnywhere, config, AdvancedDisplay,Category=AutoReimport, meta=(DisplayName="Directories to Monitor", ToolTip="Lists every directory to monitor for content changes. Can be virtual package paths (eg /Game/ or /MyPlugin/), or absolute paths on disk.\nPaths should point to the locations of the source content files (e.g. *.fbx, *.png) you want to be eligible for auto-reimport."))
 	TArray<FAutoReimportDirectoryConfig> AutoReimportDirectorySettings;
 
 	UPROPERTY(EditAnywhere, config, AdvancedDisplay, Category=AutoReimport, meta=(ClampMin=0, ClampMax=60, Units=Seconds, DisplayName="Import Threshold Time", ToolTip="Specifies an amount of time to wait before a specific file change is considered for auto reimport"))
