Auto re-import improvement: When a new file is added alongside an asset that has no source file associated with it, a re-import will be attempted instead of reporting a warning.

This is a partial merge from pull-request 1024

diff --git a/Engine/Source/Editor/UnrealEd/Private/AutoReimport/ContentDirectoryMonitor.cpp b/Engine/Source/Editor/UnrealEd/Private/AutoReimport/ContentDirectoryMonitor.cpp
--- a/Engine/Source/Editor/UnrealEd/Private/AutoReimport/ContentDirectoryMonitor.cpp
+++ b/Engine/Source/Editor/UnrealEd/Private/AutoReimport/ContentDirectoryMonitor.cpp
@@ -180,32 +180,36 @@ void FContentDirectoryMonitor::ProcessAdditions(const IAssetRegistry& Registry,
 
 		const FString FullFilename = Cache.GetDirectory() + Addition.Filename.Get();
 
+		FString NewAssetName = ObjectTools::SanitizeObjectName(FPaths::GetBaseFilename(FullFilename));
+		FString PackagePath = PackageTools::SanitizePackageName(MountedContentPath / FPaths::GetPath(Addition.Filename.Get()) / NewAssetName);
+
 		// Don't create assets for new files if assets already exist for the filename
 		auto ExistingReferences = Utils::FindAssetsPertainingToFile(Registry, FullFilename);
 		if (ExistingReferences.Num() != 0)
 		{
-			if (ExistingReferences.Num() == 1)
-			{
-				FText Message = FText::Format(LOCTEXT("Info_AlreadyImported_Single", "Ignoring new file {0} as it's already imported by {1}."), FText::FromString(FPaths::GetCleanFilename(FullFilename)), FText::FromName(ExistingReferences[0].AssetName));
-				Context.GetMessageLog().Message(EMessageSeverity::Info, Message);
-			}
-			else
-			{
-				FText Message = FText::Format(LOCTEXT("Info_AlreadyImported_Multiple", "Ignoring new file {0} as it's already imported by {1} assets."), FText::FromString(FPaths::GetCleanFilename(FullFilename)), FText::AsNumber(ExistingReferences.Num()));
-				Context.GetMessageLog().Message(EMessageSeverity::Info, Message);
-			}
-
-			// We don't try and import new files that are already associated with an asset
-			Cache.CompleteTransaction(MoveTemp(Addition));
+			// Treat this as a modified file that will attempt to reimport it (if applicable)
+			ModifiedFiles.Add(MoveTemp(Addition));
 			continue;
 		}
-
-		FString NewAssetName = ObjectTools::SanitizeObjectName(FPaths::GetBaseFilename(FullFilename));
-		FString PackagePath = PackageTools::SanitizePackageName(MountedContentPath / FPaths::GetPath(Addition.Filename.Get()) / NewAssetName);
-
-		if (FPackageName::DoesPackageExist(*PackagePath))
+		else if (FPackageName::DoesPackageExist(*PackagePath))
 		{
-			Context.AddMessage(EMessageSeverity::Warning, FText::Format(LOCTEXT("Warning_ExistingAsset", "Can't create a new asset at {0} - one already exists."), FText::FromString(PackagePath)));
+			// Package already exists, so try and import over the top of it, if it doesn't already have a source file path
+			TArray<FAssetData> Assets;
+			if (Registry.GetAssetsByPackageName(*PackagePath, Assets) && Assets.Num() == 1)
+			{
+				if (UObject* ExistingAsset = Assets[0].GetAsset())
+				{
+					// We're only eligible for reimport if the existing asset doesn't reference a source file already
+					const bool bEligibleForReimport = !Utils::ExtractSourceFilePaths(ExistingAsset).ContainsByPredicate([&](const FString& In){
+						return !In.IsEmpty() && In == FullFilename;
+					});
+
+					if (bEligibleForReimport)
+					{
+						ReimportAssetWithNewSource(ExistingAsset, FullFilename, Addition.FileData.FileHash, OutPackagesToSave, Context);
+					}
+				}
+			}
 		}
 		else
 		{
@@ -360,40 +364,14 @@ void FContentDirectoryMonitor::ProcessModifications(const IAssetRegistry& Regist
 				}
 			}
 		}
-		else if (Change.Action == DirectoryWatcher::EFileAction::Modified)
+		else
 		{
+			// Modifications or additions are treated the same by this point
 			for (const auto& AssetData : Utils::FindAssetsPertainingToFile(Registry, FullFilename))
 			{
-				UObject* Asset = AssetData.GetAsset();
-				if (Asset)
+				if (UObject* Asset = AssetData.GetAsset())
 				{
-					TArray<UObject::FAssetRegistryTag> Tags;
-					Asset->GetAssetRegistryTags(Tags);
-					TOptional<FAssetImportInfo> Info = FAssetSourceFilenameCache::ExtractAssetImportInfo(Tags);
-
-					// Check if the source file that this asset last imported was the same as the one we're going to reimport.
-					// If it is, there's no reason to auto-reimport it
-					if (Info.IsSet() && Info->SourceFiles.Num() == 1)
-					{
-						if (Info->SourceFiles[0].FileHash == Change.FileData.FileHash)
-						{
-							continue;
-						}
-					}
-
-					const bool bAssetWasDirty = IsAssetDirty(Asset);
-					if (!ReimportManager->Reimport(Asset, false /* Ask for new file */, false /* Show notification */))
-					{
-						Context.AddMessage(EMessageSeverity::Error, FText::Format(LOCTEXT("Error_FailedToReimportAsset", "Failed to reimport asset {0}."), FText::FromString(Asset->GetName())));
-					}
-					else
-					{
-						Context.AddMessage(EMessageSeverity::Info, FText::Format(LOCTEXT("Success_CreatedNewAsset", "Reimported asset {0} from {1}."), FText::FromString(Asset->GetName()), FText::FromString(FullFilename)));
-						if (!bAssetWasDirty)
-						{
-							OutPackagesToSave.Add(Asset->GetOutermost());
-						}
-					}
+					ReimportAsset(Asset, FullFilename, Change.FileData.FileHash, OutPackagesToSave, Context);
 				}
 			}
 		}
@@ -411,6 +389,46 @@ void FContentDirectoryMonitor::ProcessModifications(const IAssetRegistry& Regist
 	ModifiedFiles.Empty();
 }
 
+void FContentDirectoryMonitor::ReimportAssetWithNewSource(UObject* InAsset, const FString& FullFilename, const FMD5Hash& NewFileHash, TArray<UPackage*>& OutPackagesToSave, FReimportFeedbackContext& Context)
+{
+	TArray<FString> Filenames;
+	Filenames.Add(FullFilename);
+	FReimportManager::Instance()->UpdateReimportPaths(InAsset, Filenames);
+
+	ReimportAsset(InAsset, FullFilename, NewFileHash, OutPackagesToSave, Context);
+}
+
+void FContentDirectoryMonitor::ReimportAsset(UObject* Asset, const FString& FullFilename, const FMD5Hash& NewFileHash, TArray<UPackage*>& OutPackagesToSave, FReimportFeedbackContext& Context)
+{
+	TArray<UObject::FAssetRegistryTag> Tags;
+	Asset->GetAssetRegistryTags(Tags);
+	TOptional<FAssetImportInfo> Info = FAssetSourceFilenameCache::ExtractAssetImportInfo(Tags);
+
+	// Check if the source file that this asset last imported was the same as the one we're going to reimport.
+	// If it is, there's no reason to auto-reimport it
+	if (Info.IsSet() && Info->SourceFiles.Num() == 1)
+	{
+		if (Info->SourceFiles[0].FileHash == NewFileHash)
+		{
+			return;
+		}
+	}
+
+	const bool bAssetWasDirty = IsAssetDirty(Asset);
+	if (!FReimportManager::Instance()->Reimport(Asset, false /* Ask for new file */, false /* Show notification */))
+	{
+		Context.AddMessage(EMessageSeverity::Error, FText::Format(LOCTEXT("Error_FailedToReimportAsset", "Failed to reimport asset {0}."), FText::FromString(Asset->GetName())));
+	}
+	else
+	{
+		Context.AddMessage(EMessageSeverity::Info, FText::Format(LOCTEXT("Success_CreatedNewAsset", "Reimported asset {0} from {1}."), FText::FromString(Asset->GetName()), FText::FromString(FullFilename)));
+		if (!bAssetWasDirty)
+		{
+			OutPackagesToSave.Add(Asset->GetOutermost());
+		}
+	}
+}
+
 void FContentDirectoryMonitor::ExtractAssetsToDelete(const IAssetRegistry& Registry, TArray<FAssetData>& OutAssetsToDelete)
 {
 	for (auto& Deletion : DeletedFiles)
diff --git a/Engine/Source/Editor/UnrealEd/Private/AutoReimport/ContentDirectoryMonitor.h b/Engine/Source/Editor/UnrealEd/Private/AutoReimport/ContentDirectoryMonitor.h
--- a/Engine/Source/Editor/UnrealEd/Private/AutoReimport/ContentDirectoryMonitor.h
+++ b/Engine/Source/Editor/UnrealEd/Private/AutoReimport/ContentDirectoryMonitor.h
@@ -59,6 +59,17 @@ public:
 	/** Get the number of outstanding changes that we potentially have to process (when not already processing) */
 	int32 GetNumUnprocessedChanges() const { return Cache.GetNumOutstandingChanges(); }
 
+private:
+
+	/** Import a new asset into the specified package path, from the specified file */ 
+	void ImportAsset(const FString& PackagePath, const FString& Filename, TArray<UPackage*>& OutPackagesToSave, const TMap<FString, TArray<UFactory*>>& InFactoriesByExtension, FReimportFeedbackContext& Context);
+
+	/** Reimport a specific asset, provided its source content differs from the specified file hash */
+	void ReimportAsset(UObject* InAsset, const FString& FullFilename, const FMD5Hash& NewFileHash, TArray<UPackage*>& OutPackagesToSave, FReimportFeedbackContext& Context);
+
+	/** Set the specified asset to import from the specified file, then attempt to reimport it */
+	void ReimportAssetWithNewSource(UObject* InAsset, const FString& FullFilename, const FMD5Hash& NewFileHash, TArray<UPackage*>& OutPackagesToSave, FReimportFeedbackContext& Context);
+
 private:
 
 	/** The file cache that monitors and reflects the content directory. */
