Moving assets now corrects source reimport paths where possible

This is a partial merge from github pull request 1024, which adds support for updating source filenames when assets are moved, and the referenced file is relative to the asset itself.

diff --git a/Engine/Source/Editor/UnrealEd/Private/AutoReimport/AutoReimportManager.cpp b/Engine/Source/Editor/UnrealEd/Private/AutoReimport/AutoReimportManager.cpp
--- a/Engine/Source/Editor/UnrealEd/Private/AutoReimport/AutoReimportManager.cpp
+++ b/Engine/Source/Editor/UnrealEd/Private/AutoReimport/AutoReimportManager.cpp
@@ -349,28 +349,30 @@ void FAutoReimportManager::OnAssetRenamed(const FAssetData& AssetData, const FSt
 	// This code moves a source content file that reside alongside assets when the assets are renamed. We do this under the following conditions:
 	// 	1. The sourcefile is solely referenced from the the asset that has been moved
 	//	2. Said asset only references a single file
-	//	3. The source file resides in the same folder as the asset
 	//
 	// Additionally, we rename the source file if it matched the name of the asset before the rename/move.
 	//	- If we rename the source file, then we also update the reimport paths for the asset
 
-	
 	TOptional<FAssetImportInfo> ImportInfo = FAssetSourceFilenameCache::ExtractAssetImportInfo(AssetData.TagsAndValues);
-	if (!ImportInfo.IsSet())
+	if (!ImportInfo.IsSet() || ImportInfo->SourceFiles.Num() != 1)
 	{
 		return;
 	}
+	
+	const FString& RelativeFilename = ImportInfo->SourceFiles[0].RelativeFilename;
 
-	const TArray<FAssetImportInfo::FSourceFile>& SourceFileData = ImportInfo->SourceFiles;
+	FString OldPackagePath = FPackageName::GetLongPackagePath(OldPath);
+	FString NewReimportPath;
 
 	// We move the file with the asset provided it is the only file referenced, and sits right beside the uasset file
-	if (SourceFileData.Num() == 1 && !SourceFileData[0].RelativeFilename.GetCharArray().ContainsByPredicate([](const TCHAR Char) { return Char == '/' || Char == '\\'; }))
+	if (!RelativeFilename.GetCharArray().ContainsByPredicate([](const TCHAR Char) { return Char == '/' || Char == '\\'; }))
 	{
-		const FString AbsoluteSrcPath = FPaths::ConvertRelativePathToFull(FPackageName::LongPackageNameToFilename(FPackageName::GetLongPackagePath(OldPath) / TEXT("")));
-		const FString AbsoluteDstPath = FPaths::ConvertRelativePathToFull(FPackageName::LongPackageNameToFilename(AssetData.PackagePath.ToString() / TEXT("")));
+		// File resides in the same folder as the asset, so we can potentially rename the source file too
+		const FString AbsoluteSrcPath = FPaths::ConvertRelativePathToFull(FPackageName::LongPackageNameToFilename(OldPackagePath + TEXT("/")));
+		const FString AbsoluteDstPath = FPaths::ConvertRelativePathToFull(FPackageName::LongPackageNameToFilename(AssetData.PackagePath.ToString() + TEXT("/")));
 
 		const FString OldAssetName = FPackageName::GetLongPackageAssetName(FPackageName::ObjectPathToPackageName(OldPath));
-		FString NewFileName = FPaths::GetBaseFilename(SourceFileData[0].RelativeFilename);
+		FString NewFileName = FPaths::GetBaseFilename(RelativeFilename);
 
 		bool bRequireReimportPathUpdate = false;
 		if (PackageTools::SanitizePackageName(NewFileName) == OldAssetName)
@@ -379,30 +381,44 @@ void FAutoReimportManager::OnAssetRenamed(const FAssetData& AssetData, const FSt
 			bRequireReimportPathUpdate = true;
 		}
 
-		const FString SrcFile = AbsoluteSrcPath / SourceFileData[0].RelativeFilename;
-		const FString DstFile = AbsoluteDstPath / NewFileName + TEXT(".") + FPaths::GetExtension(SourceFileData[0].RelativeFilename);
+		const FString SrcFile = AbsoluteSrcPath / RelativeFilename;
+		const FString DstFile = AbsoluteDstPath / NewFileName + TEXT(".") + FPaths::GetExtension(RelativeFilename);
 
-		// We can't do this if multiple assets reference the same file
+		// We can't do this if multiple assets reference the same file. We should be checking for > 1 referencing asset, but the asset registry
+		// filter lookup won't return the recently renamed package because it will be Empty by now, so we check for *anything* referencing the asset (assuming that we'll never find *this* asset).
 		const IAssetRegistry& Registry = FModuleManager::LoadModuleChecked<FAssetRegistryModule>("AssetRegistry").Get();
-		if (Utils::FindAssetsPertainingToFile(Registry, SrcFile).Num() > 1)
+		if (Utils::FindAssetsPertainingToFile(Registry, SrcFile).Num() == 0)
 		{
-			return;
+			if (!FPlatformFileManager::Get().GetPlatformFile().FileExists(*DstFile) &&
+				IFileManager::Get().Move(*DstFile, *SrcFile, false /*bReplace */, false, true /* attributes */, true /* don't retry */))
+			{
+				IgnoreMovedFile(SrcFile, DstFile);
+
+				if (bRequireReimportPathUpdate)
+				{
+					NewReimportPath = DstFile;
+				}
+			}
 		}
+	}
 
-		if (!FPlatformFileManager::Get().GetPlatformFile().FileExists(*DstFile) &&
-			IFileManager::Get().Move(*DstFile, *SrcFile, false /*bReplace */, false, true /* attributes */, true /* don't retry */))
+	if (NewReimportPath.IsEmpty() && FPackageName::GetLongPackagePath(OldPath) != AssetData.PackagePath.ToString())
+	{
+		// The asset has been moved, try and update its referenced path
+		FString OldSourceFilePath = FPaths::ConvertRelativePathToFull(FPackageName::LongPackageNameToFilename(OldPackagePath), RelativeFilename);
+		if (FPaths::FileExists(OldSourceFilePath))
 		{
-			IgnoreMovedFile(SrcFile, DstFile);
+			NewReimportPath = MoveTemp(OldSourceFilePath);
+		}
+	}
 
-			if (bRequireReimportPathUpdate)
-			{
-				TArray<FString> Paths;
-				Paths.Add(DstFile);
+	if (!NewReimportPath.IsEmpty())
+	{
+		TArray<FString> Paths;
+		Paths.Add(NewReimportPath);
 
-				// Update the reimport file names
-				FReimportManager::Instance()->UpdateReimportPaths(AssetData.GetAsset(), Paths);
-			}
-		}
+		// Update the reimport file names
+		FReimportManager::Instance()->UpdateReimportPaths(AssetData.GetAsset(), Paths);
 	}
 }
 
