"Don't show again" check on auto-reimporter 'detect changes on restart' notification now saves correctly.

This addresses UE-20464

diff --git a/Engine/Source/Editor/UnrealEd/Private/Settings/SettingsClasses.cpp b/Engine/Source/Editor/UnrealEd/Private/Settings/SettingsClasses.cpp
--- a/Engine/Source/Editor/UnrealEd/Private/Settings/SettingsClasses.cpp
+++ b/Engine/Source/Editor/UnrealEd/Private/Settings/SettingsClasses.cpp
@@ -198,7 +198,9 @@ void UEditorLoadingSavingSettings::CheckSourceControlCompatability()
 		Info.bFireAndForget = false;
 
 		Info.CheckBoxStateChanged = FOnCheckStateChanged::CreateLambda([](ECheckBoxState State){
-			GetMutableDefault<UEditorLoadingSavingSettings>()->bEnableSourceControlCompatabilityCheck = (State != ECheckBoxState::Checked);
+			auto* Settings = GetMutableDefault<UEditorLoadingSavingSettings>();
+			Settings->bEnableSourceControlCompatabilityCheck = (State != ECheckBoxState::Checked);
+			Settings->SaveConfig();
 		});
 		Info.CheckBoxText = LOCTEXT("AutoReimport_DontShowAgain", "Don't show again");
 
