Fixes to editing brush builder properties.

diff --git a/Engine/Source/Editor/UnrealEd/Private/EditorBrushBuilder.cpp b/Engine/Source/Editor/UnrealEd/Private/EditorBrushBuilder.cpp
--- a/Engine/Source/Editor/UnrealEd/Private/EditorBrushBuilder.cpp
+++ b/Engine/Source/Editor/UnrealEd/Private/EditorBrushBuilder.cpp
@@ -305,12 +305,12 @@ void UConeBuilder::PostEditChangeProperty(struct FPropertyChangedEvent& Property
 
 		const float ZEpsilon = 0.1f;
 
-		if (PropertyChangedEvent.Property->GetFName() == Name_Z && Z <= CapZ)
+		if (Hollow && PropertyChangedEvent.Property->GetFName() == Name_Z && Z <= CapZ)
 		{
 			Z = CapZ + ZEpsilon;
 		}
 
-		if (PropertyChangedEvent.Property->GetFName() == Name_CapZ && CapZ >= Z)
+		if (Hollow && PropertyChangedEvent.Property->GetFName() == Name_CapZ && CapZ >= Z)
 		{
 			CapZ = FMath::Max(0.0f, Z - ZEpsilon);
 		}
@@ -320,12 +320,12 @@ void UConeBuilder::PostEditChangeProperty(struct FPropertyChangedEvent& Property
 
 		const float RadiusEpsilon = 0.1f;
 
-		if (PropertyChangedEvent.Property->GetFName() == Name_OuterRadius && OuterRadius <= InnerRadius)
+		if (Hollow && PropertyChangedEvent.Property->GetFName() == Name_OuterRadius && OuterRadius <= InnerRadius)
 		{
 			OuterRadius = InnerRadius + RadiusEpsilon;
 		}
 
-		if (PropertyChangedEvent.Property->GetFName() == Name_InnerRadius && InnerRadius >= OuterRadius)
+		if (Hollow && PropertyChangedEvent.Property->GetFName() == Name_InnerRadius && InnerRadius >= OuterRadius)
 		{
 			InnerRadius = FMath::Max(0.0f, OuterRadius - RadiusEpsilon);
 		}
@@ -438,22 +438,22 @@ void UCubeBuilder::PostEditChangeProperty(struct FPropertyChangedEvent& Property
 
 		const float ThicknessEpsilon = 0.1f;
 
-		if (PropertyChangedEvent.Property->GetFName() == Name_X && X <= WallThickness)
+		if (Hollow && PropertyChangedEvent.Property->GetFName() == Name_X && X <= WallThickness)
 		{
 			X = WallThickness + ThicknessEpsilon;
 		}
 
-		if (PropertyChangedEvent.Property->GetFName() == Name_Y && Y <= WallThickness)
+		if (Hollow && PropertyChangedEvent.Property->GetFName() == Name_Y && Y <= WallThickness)
 		{
 			Y = WallThickness + ThicknessEpsilon;
 		}
 
-		if (PropertyChangedEvent.Property->GetFName() == Name_Z && Z <= WallThickness)
+		if (Hollow && PropertyChangedEvent.Property->GetFName() == Name_Z && Z <= WallThickness)
 		{
 			Z = WallThickness + ThicknessEpsilon;
 		}
 
-		if (PropertyChangedEvent.Property->GetFName() == Name_WallThickness && WallThickness >= FMath::Min3(X, Y, Z))
+		if (Hollow && PropertyChangedEvent.Property->GetFName() == Name_WallThickness && WallThickness >= FMath::Min3(X, Y, Z))
 		{
 			WallThickness = FMath::Max(0.0f, FMath::Min3(X, Y, Z) - ThicknessEpsilon);
 		}
@@ -664,12 +664,12 @@ void UCylinderBuilder::PostEditChangeProperty(struct FPropertyChangedEvent& Prop
 
 		const float RadiusEpsilon = 0.1f;
 
-		if (PropertyChangedEvent.Property->GetFName() == Name_OuterRadius && OuterRadius <= InnerRadius)
+		if (Hollow && PropertyChangedEvent.Property->GetFName() == Name_OuterRadius && OuterRadius <= InnerRadius)
 		{
 			OuterRadius = InnerRadius + RadiusEpsilon;
 		}
 
-		if (PropertyChangedEvent.Property->GetFName() == Name_InnerRadius && InnerRadius >= OuterRadius)
+		if (Hollow && PropertyChangedEvent.Property->GetFName() == Name_InnerRadius && InnerRadius >= OuterRadius)
 		{
 			InnerRadius = FMath::Max(0.0f, OuterRadius - RadiusEpsilon);
 		}
