User Defined Enums begin counting at 0, and will no longer increase by two's.

When adding a new enum value, the value will be set to the max value instead of the max value + 1.

The enumerator's max will be zero if there are no enum values.

diff --git a/Engine/Source/Editor/UnrealEd/Private/Kismet2/EnumEditorUtils.cpp b/Engine/Source/Editor/UnrealEd/Private/Kismet2/EnumEditorUtils.cpp
--- a/Engine/Source/Editor/UnrealEd/Private/Kismet2/EnumEditorUtils.cpp
+++ b/Engine/Source/Editor/UnrealEd/Private/Kismet2/EnumEditorUtils.cpp
@@ -96,7 +96,7 @@ void FEnumEditorUtils::AddNewEnumeratorForUserDefinedEnum(UUserDefinedEnum* Enum
 
 	FString EnumNameString = Enum->GenerateNewEnumeratorName();
 	const FString FullNameStr = Enum->GenerateFullEnumName(*EnumNameString);
-	Names.Add(TPairInitializer<FName, uint8>(FName(*FullNameStr), Enum->GetMaxEnumValue() + 1));
+	Names.Add(TPairInitializer<FName, uint8>(FName(*FullNameStr), Enum->GetMaxEnumValue()));
 
 	const UEnum::ECppForm EnumType = Enum->GetCppForm();
 	Enum->SetEnums(Names, EnumType);
diff --git a/Engine/Source/Runtime/Engine/Private/UserDefinedEnum.cpp b/Engine/Source/Runtime/Engine/Private/UserDefinedEnum.cpp
--- a/Engine/Source/Runtime/Engine/Private/UserDefinedEnum.cpp
+++ b/Engine/Source/Runtime/Engine/Private/UserDefinedEnum.cpp
@@ -117,7 +117,8 @@ bool UUserDefinedEnum::SetEnums(TArray<TPair<FName, uint8>>& InNames, ECppForm I
 		const int32 MaxEnumItemIndex = GetValueByName(MaxEnumItem);
 		if ((MaxEnumItemIndex == INDEX_NONE) && (LookupEnumName(MaxEnumItem) == INDEX_NONE))
 		{
-			Names.Add(TPairInitializer<FName, uint8>(MaxEnumItem, GetMaxEnumValue() + 1));
+			int MaxEnumValue = (InNames.Num() == 0)? 0 : GetMaxEnumValue() + 1;
+			Names.Add(TPairInitializer<FName, uint8>(MaxEnumItem, MaxEnumValue));
 			AddNamesToMasterList();
 			return true;
 		}
