Jira UE-20132
Add a new mode with shift + ctrl + left mouse down to add new keys on the editable curves at the mouse location.
Add the menu entry for the new mode

diff --git a/Engine/Source/Editor/UnrealEd/Private/SCurveEditor.cpp b/Engine/Source/Editor/UnrealEd/Private/SCurveEditor.cpp
--- a/Engine/Source/Editor/UnrealEd/Private/SCurveEditor.cpp
+++ b/Engine/Source/Editor/UnrealEd/Private/SCurveEditor.cpp
@@ -1381,7 +1381,8 @@ void SCurveEditor::UpdateCurveToolTip(const FGeometry& InMyGeometry, const FPoin
 	if (Settings->GetShowCurveEditorCurveToolTips())
 	{
 		TSharedPtr<FCurveViewModel> HoveredCurve = HitTestCurves(InMyGeometry, InMouseEvent);
-		if (HoveredCurve.IsValid())
+		//Display the tooltip only when the curve is visible
+		if (HoveredCurve.IsValid() && HoveredCurve->bIsVisible)
 		{
 			FTrackScaleInfo ScaleInfo(ViewMinInput.Get(), ViewMaxInput.Get(), ViewMinOutput.Get(), ViewMaxOutput.Get(), InMyGeometry.Size);
 			const FVector2D HitPosition = InMyGeometry.AbsoluteToLocal(InMouseEvent.GetScreenSpacePosition());
@@ -1844,23 +1845,28 @@ void SCurveEditor::ProcessClick(const FGeometry& InMyGeometry, const FPointerEve
 				TSharedPtr<TArray<TSharedPtr<FCurveViewModel>>> CurvesToAddKeysTo = MakeShareable(new TArray<TSharedPtr<FCurveViewModel>>());
 				TSharedPtr<FCurveViewModel> HoveredCurve = HitTestCurves(InMyGeometry, InMouseEvent);
 				bool bAddKeysInline;
-				if (HoveredCurve.IsValid())
+				//To snap a point on the hovered curve the curve must be visible and unlock
+				if (HoveredCurve.IsValid() && !HoveredCurve->bIsLocked && HoveredCurve->bIsVisible)
 				{
 					CurvesToAddKeysTo->Add(HoveredCurve);
 					bAddKeysInline = true;
 				}
 				else
 				{
-					if (CurveViewModels.Num() == 1)
+					//Add all unlock curves in the editable array
+					for (auto CurveViewModel : CurveViewModels)
 					{
-						CurvesToAddKeysTo->Add(CurveViewModels[0]);
-						bAddKeysInline = false;
-					}
-					else
-					{
-						CurvesToAddKeysTo->Append(CurveViewModels);
-						bAddKeysInline = true;
+						if (!CurveViewModel->bIsLocked)
+						{
+							CurvesToAddKeysTo->Add(CurveViewModel);
+						}
 					}
+
+					//If linear color curve and no show curve, always insert inline
+					//If the user is holding shift-ctrl we snap all curve to the mouse position. (false value)
+					//If the user is holding shift we snap to mouse only if there is only one editable curve. (false value)
+					//In all other case we add key directly on the curve. (true value)
+					bAddKeysInline = (IsLinearColorCurve() && !bAreCurvesVisible) || ((!bControlDown) && (CurvesToAddKeysTo->Num() != 1));
 				}
 				AddNewKey(InMyGeometry, InMouseEvent.GetScreenSpacePosition(), CurvesToAddKeysTo, bAddKeysInline);
 			}
@@ -2496,44 +2502,58 @@ void SCurveEditor::CreateContextMenu(const FGeometry& InMyGeometry, const FPoint
 		FText MenuItemLabel;
 		FText MenuItemToolTip;
 		TSharedPtr<TArray<TSharedPtr<FCurveViewModel>>> CurvesToAddKeysTo = MakeShareable(new TArray<TSharedPtr<FCurveViewModel>>());
-		bool bAddKeysInline;
 
 		FText AddKeyToCurveLabelFormat = LOCTEXT("AddKeyToCurveLabelFormat", "Add key to {0}");
 		FText AddKeyToCurveToolTipFormat = LOCTEXT("AddKeyToCurveToolTipFormat", "Add a new key at the hovered time to the {0} curve.  Keys can also be added with Shift + Click.");
 
+		FVector2D Position = InMouseEvent.GetScreenSpacePosition();
+
 		TSharedPtr<FCurveViewModel> HoveredCurve = HitTestCurves(InMyGeometry, InMouseEvent);
-		if (HoveredCurve.IsValid())
+		//Curve must be visible and unlock to show context menu
+		if (HoveredCurve.IsValid() && !HoveredCurve->bIsLocked && HoveredCurve->bIsVisible)
 		{
 			MenuItemLabel = FText::Format(AddKeyToCurveLabelFormat, FText::FromName(HoveredCurve->CurveInfo.CurveName));
 			MenuItemToolTip = FText::Format(AddKeyToCurveToolTipFormat, FText::FromName(HoveredCurve->CurveInfo.CurveName));
 			CurvesToAddKeysTo->Add(HoveredCurve);
-			bAddKeysInline = true;
+			FUIAction Action = FUIAction(FExecuteAction::CreateSP(this, &SCurveEditor::AddNewKey, InMyGeometry, Position, CurvesToAddKeysTo, true));
+			MenuBuilder.AddMenuEntry(MenuItemLabel, MenuItemToolTip, FSlateIcon(), Action);
 		}
 		else
 		{
-			if (CurveViewModels.Num() == 1)
+			//Get all editable curve
+			for (auto CurveViewModel : CurveViewModels)
 			{
-				MenuItemLabel = FText::Format(AddKeyToCurveLabelFormat, FText::FromName(CurveViewModels[0]->CurveInfo.CurveName));
-				MenuItemToolTip = FText::Format(AddKeyToCurveToolTipFormat, FText::FromName(CurveViewModels[0]->CurveInfo.CurveName));
-				CurvesToAddKeysTo->Add(CurveViewModels[0]);
-				bAddKeysInline = false;
+				if (!CurveViewModel->bIsLocked)
+				{
+					CurvesToAddKeysTo->Add(CurveViewModel);
+				}
 			}
-			else
+			if (CurvesToAddKeysTo->Num() == 1)
+			{
+				MenuItemLabel = FText::Format(AddKeyToCurveLabelFormat, FText::FromName((*CurvesToAddKeysTo)[0]->CurveInfo.CurveName));
+				MenuItemToolTip = FText::Format(AddKeyToCurveToolTipFormat, FText::FromName((*CurvesToAddKeysTo)[0]->CurveInfo.CurveName));
+				FUIAction Action = FUIAction(FExecuteAction::CreateSP(this, &SCurveEditor::AddNewKey, InMyGeometry, Position, CurvesToAddKeysTo, false));
+				MenuBuilder.AddMenuEntry(MenuItemLabel, MenuItemToolTip, FSlateIcon(), Action);
+			}
+			else if (CurvesToAddKeysTo->Num() > 1) //Dont show the menu if we cannot edit any curve
 			{
+				//add key to all curve menu entry
 				MenuItemLabel = LOCTEXT("AddKeyToAllCurves", "Add key to all curves");
 				MenuItemToolTip = LOCTEXT("AddKeyToAllCurveToolTip", "Adds a key at the hovered time to all curves.  Keys can also be added with Shift + Click.");
-				CurvesToAddKeysTo->Append(CurveViewModels);
-				bAddKeysInline = true;
+				FUIAction Action = FUIAction(FExecuteAction::CreateSP(this, &SCurveEditor::AddNewKey, InMyGeometry, Position, CurvesToAddKeysTo, true));
+				MenuBuilder.AddMenuEntry(MenuItemLabel, MenuItemToolTip, FSlateIcon(), Action);
+				
+				//This menu is not required when there is no curve display (color track can hide and show curves)
+				if (bAreCurvesVisible)
+				{
+					//add key and value to all curve menu entry
+					MenuItemLabel = LOCTEXT("AddKeyValueToAllCurves", "Add key & value to all curves");
+					MenuItemToolTip = LOCTEXT("AddKeyValueToAllCurveToolTip", "Adds a key & value at the hovered time to all curves.  Keys can also be added with Shift + ctrl + Click.");
+					Action = FUIAction(FExecuteAction::CreateSP(this, &SCurveEditor::AddNewKey, InMyGeometry, Position, CurvesToAddKeysTo, false));
+					MenuBuilder.AddMenuEntry(MenuItemLabel, MenuItemToolTip, FSlateIcon(), Action);
+				}
 			}
 		}
-
-		FVector2D Position = InMouseEvent.GetScreenSpacePosition();
-		FUIAction Action = FUIAction(FExecuteAction::CreateSP(this, &SCurveEditor::AddNewKey, InMyGeometry, Position, CurvesToAddKeysTo, bAddKeysInline));
-		MenuBuilder.AddMenuEntry(
-			MenuItemLabel,
-			MenuItemToolTip,
-			FSlateIcon(),
-			Action);
 	}
 	MenuBuilder.EndSection();
 
