Added logging for monitoring a GLEO that starts due to PIE and ends with a corruption that is asserting during BP compile.

diff --git a/Engine/Source/Editor/UnrealEd/Private/Kismet2/Kismet2.cpp b/Engine/Source/Editor/UnrealEd/Private/Kismet2/Kismet2.cpp
--- a/Engine/Source/Editor/UnrealEd/Private/Kismet2/Kismet2.cpp
+++ b/Engine/Source/Editor/UnrealEd/Private/Kismet2/Kismet2.cpp
@@ -724,6 +724,18 @@ void FKismetEditorUtilities::CompileBlueprint(UBlueprint* BlueprintObj, bool bIs
 	FCompilerResultsLog LocalResults;
 	FCompilerResultsLog& Results = (pResults != NULL) ? *pResults : LocalResults;
 
+	// Monitoring UE-20486, the OldClass->ClassGeneratedBy is NULL or otherwise not a UBlueprint.
+	if (OldClass)
+	{
+		if (OldClass->ClassGeneratedBy == nullptr)
+		{
+			checkf(OldClass->ClassGeneratedBy == BlueprintObj, TEXT("Generated Class '%s' during compilation has a NULL ClassGeneratedBy for Blueprint '%s'"), *OldClass->GetPathName(), *BlueprintObj->GetPathName());
+		}
+		else
+		{
+			checkf(OldClass->ClassGeneratedBy == BlueprintObj, TEXT("Generated Class '%s' has an invalid ClassGeneratedBy '%s' while the expected is Blueprint '%s'"), *OldClass->GetPathName(), *OldClass->ClassGeneratedBy->GetPathName(), *BlueprintObj->GetPathName());
+		}
+	}
 	auto ReinstanceHelper = FBlueprintCompileReinstancer::Create(OldClass);
 
 	// Suppress errors/warnings in the log if we're recompiling on load on a build machine
diff --git a/Engine/Source/Editor/UnrealEd/Private/PlayLevel.cpp b/Engine/Source/Editor/UnrealEd/Private/PlayLevel.cpp
--- a/Engine/Source/Editor/UnrealEd/Private/PlayLevel.cpp
+++ b/Engine/Source/Editor/UnrealEd/Private/PlayLevel.cpp
@@ -72,6 +72,33 @@ inline FName GetOnlineIdentifier(const FWorldContext& WorldContext)
 
 void UEditorEngine::EndPlayMap()
 {
+	// Monitoring when PIE corrupts references between the World and the PIE generated World for UE-20486
+	{
+		TArray<ULevel*> Levels = EditorWorld->GetLevels();
+
+		for (ULevel* Level : Levels)
+		{
+			TArray<UBlueprint*> LevelBlueprints = Level->GetLevelBlueprints();
+
+			if (LevelBlueprints.Num() > 0)
+			{
+				UBlueprint* LevelScriptBlueprint = LevelBlueprints[0];
+				if (LevelScriptBlueprint && LevelScriptBlueprint->GeneratedClass && LevelScriptBlueprint->GeneratedClass->ClassGeneratedBy)
+				{
+					UE_LOG(LogBlueprintUserMessages, Log, TEXT("Early EndPlayMap Detection: Level '%s' has LevelScriptBlueprint '%s' with GeneratedClass '%s' with ClassGeneratedBy '%s'"), *Level->GetPathName(), *LevelScriptBlueprint->GetPathName(), *LevelScriptBlueprint->GeneratedClass->GetPathName(), *LevelScriptBlueprint->GeneratedClass->ClassGeneratedBy->GetPathName());
+				}
+				else if (LevelScriptBlueprint && LevelScriptBlueprint->GeneratedClass)
+				{
+					UE_LOG(LogBlueprintUserMessages, Log, TEXT("Early EndPlayMap Detection: Level '%s' has LevelScriptBlueprint '%s' with GeneratedClass '%s'"), *Level->GetPathName(), *LevelScriptBlueprint->GetPathName(), *LevelScriptBlueprint->GeneratedClass->GetPathName());
+				}
+				else if (LevelScriptBlueprint)
+				{
+					UE_LOG(LogBlueprintUserMessages, Log, TEXT("Early EndPlayMap Detection: Level '%s' has LevelScriptBlueprint '%s'"), *Level->GetPathName(), *LevelScriptBlueprint->GetPathName());
+				}
+			}
+		}
+	}
+
 	if (GEngine->HMDDevice.IsValid())
 	{
 		GEngine->HMDDevice->OnEndPlay();
@@ -236,6 +263,33 @@ void UEditorEngine::EndPlayMap()
 
 	FGameDelegates::Get().GetEndPlayMapDelegate().Broadcast();
 
+	// Monitoring when PIE corrupts references between the World and the PIE generated World for UE-20486
+	{
+		TArray<ULevel*> Levels = EditorWorld->GetLevels();
+
+		for (ULevel* Level : Levels)
+		{
+			TArray<UBlueprint*> LevelBlueprints = Level->GetLevelBlueprints();
+
+			if (LevelBlueprints.Num() > 0)
+			{
+				UBlueprint* LevelScriptBlueprint = LevelBlueprints[0];
+				if (LevelScriptBlueprint && LevelScriptBlueprint->GeneratedClass && LevelScriptBlueprint->GeneratedClass->ClassGeneratedBy)
+				{
+					UE_LOG(LogBlueprintUserMessages, Log, TEXT("Late EndPlayMap Detection: Level '%s' has LevelScriptBlueprint '%s' with GeneratedClass '%s' with ClassGeneratedBy '%s'"), *Level->GetPathName(), *LevelScriptBlueprint->GetPathName(), *LevelScriptBlueprint->GeneratedClass->GetPathName(), *LevelScriptBlueprint->GeneratedClass->ClassGeneratedBy->GetPathName());
+				}
+				else if (LevelScriptBlueprint && LevelScriptBlueprint->GeneratedClass)
+				{
+					UE_LOG(LogBlueprintUserMessages, Log, TEXT("Late EndPlayMap Detection: Level '%s' has LevelScriptBlueprint '%s' with GeneratedClass '%s'"), *Level->GetPathName(), *LevelScriptBlueprint->GetPathName(), *LevelScriptBlueprint->GeneratedClass->GetPathName());
+				}
+				else if (LevelScriptBlueprint)
+				{
+					UE_LOG(LogBlueprintUserMessages, Log, TEXT("Late EndPlayMap Detection: Level '%s' has LevelScriptBlueprint '%s'"), *Level->GetPathName(), *LevelScriptBlueprint->GetPathName());
+				}
+			}
+		}
+	}
+
 	EditorWorld->bAllowAudioPlayback = true;
 	EditorWorld = NULL;
 
@@ -2151,6 +2205,33 @@ void UEditorEngine::PlayInEditor( UWorld* InWorld, bool bInSimulateInEditor )
 
 	FBlueprintEditorUtils::FindAndSetDebuggableBlueprintInstances();
 
+	// Monitoring when PIE corrupts references between the World and the PIE generated World for UE-20486
+	{
+		TArray<ULevel*> Levels = InWorld->GetLevels();
+
+		for (ULevel* Level : Levels)
+		{
+			TArray<UBlueprint*> LevelBlueprints = Level->GetLevelBlueprints();
+
+			if (LevelBlueprints.Num() > 0)
+			{
+				UBlueprint* LevelScriptBlueprint = LevelBlueprints[0];
+				if (LevelScriptBlueprint && LevelScriptBlueprint->GeneratedClass && LevelScriptBlueprint->GeneratedClass->ClassGeneratedBy)
+				{
+					UE_LOG(LogBlueprintUserMessages, Log, TEXT("Early PlayInEditor Detection: Level '%s' has LevelScriptBlueprint '%s' with GeneratedClass '%s' with ClassGeneratedBy '%s'"), *Level->GetPathName(), *LevelScriptBlueprint->GetPathName(), *LevelScriptBlueprint->GeneratedClass->GetPathName(), *LevelScriptBlueprint->GeneratedClass->ClassGeneratedBy->GetPathName());
+				}
+				else if (LevelScriptBlueprint && LevelScriptBlueprint->GeneratedClass)
+				{
+					UE_LOG(LogBlueprintUserMessages, Log, TEXT("Early PlayInEditor Detection: Level '%s' has LevelScriptBlueprint '%s' with GeneratedClass '%s'"), *Level->GetPathName(), *LevelScriptBlueprint->GetPathName(), *LevelScriptBlueprint->GeneratedClass->GetPathName());
+				}
+				else if (LevelScriptBlueprint)
+				{
+					UE_LOG(LogBlueprintUserMessages, Log, TEXT("Early PlayInEditor Detection: Level '%s' has LevelScriptBlueprint '%s'"), *Level->GetPathName(), *LevelScriptBlueprint->GetPathName());
+				}
+			}
+		}
+	}
+
 	// Broadcast BeginPIE after checks that might block PIE above (PreBeginPIE is broadcast above before the checks)
 	FEditorDelegates::BeginPIE.Broadcast(bInSimulateInEditor);
 
@@ -2326,6 +2407,33 @@ void UEditorEngine::PlayInEditor( UWorld* InWorld, bool bInSimulateInEditor )
 	PlayInSettings->MultipleInstanceLastHeight = PlayInSettings->NewWindowHeight;
 	PlayInSettings->MultipleInstanceLastWidth = PlayInSettings->NewWindowWidth;
 	PlayInSettings->SetPlayNetMode(OrigPlayNetMode);
+
+	// Monitoring when PIE corrupts references between the World and the PIE generated World for UE-20486
+	{
+		TArray<ULevel*> Levels = EditorWorld->GetLevels();
+
+		for (ULevel* Level : Levels)
+		{
+			TArray<UBlueprint*> LevelBlueprints = Level->GetLevelBlueprints();
+
+			if (LevelBlueprints.Num() > 0)
+			{
+				UBlueprint* LevelScriptBlueprint = LevelBlueprints[0];
+				if (LevelScriptBlueprint && LevelScriptBlueprint->GeneratedClass && LevelScriptBlueprint->GeneratedClass->ClassGeneratedBy)
+				{
+					UE_LOG(LogBlueprintUserMessages, Log, TEXT("Late PlayInEditor Detection: Level '%s' has LevelScriptBlueprint '%s' with GeneratedClass '%s' with ClassGeneratedBy '%s'"), *Level->GetPathName(), *LevelScriptBlueprint->GetPathName(), *LevelScriptBlueprint->GeneratedClass->GetPathName(), *LevelScriptBlueprint->GeneratedClass->ClassGeneratedBy->GetPathName());
+				}
+				else if (LevelScriptBlueprint && LevelScriptBlueprint->GeneratedClass)
+				{
+					UE_LOG(LogBlueprintUserMessages, Log, TEXT("Late PlayInEditor Detection: Level '%s' has LevelScriptBlueprint '%s' with GeneratedClass '%s'"), *Level->GetPathName(), *LevelScriptBlueprint->GetPathName(), *LevelScriptBlueprint->GeneratedClass->GetPathName());
+				}
+				else if (LevelScriptBlueprint)
+				{
+					UE_LOG(LogBlueprintUserMessages, Log, TEXT("Late PlayInEditor Detection: Level '%s' has LevelScriptBlueprint '%s'"), *Level->GetPathName(), *LevelScriptBlueprint->GetPathName());
+				}
+			}
+		}
+	}
 }
 
 void UEditorEngine::SpawnIntraProcessPIEWorlds(bool bAnyBlueprintErrors, bool bStartInSpectatorMode)
