PR #1516: Speedup asset discovery by avoiding redundant data copying (Contributed by [REDACTED_PERSON])

diff --git a/Engine/Source/Editor/UnrealEd/Private/ComponentTypeRegistry.cpp b/Engine/Source/Editor/UnrealEd/Private/ComponentTypeRegistry.cpp
--- a/Engine/Source/Editor/UnrealEd/Private/ComponentTypeRegistry.cpp
+++ b/Engine/Source/Editor/UnrealEd/Private/ComponentTypeRegistry.cpp
@@ -282,33 +282,40 @@ void FComponentTypeRegistryData::ForceRefreshComponentList()
 		auto InMemoryClassesSet = TSet<FName>(InMemoryClasses);
 		auto OnDiskClasses = DerivedClassNames.Difference(InMemoryClassesSet);
 
-		// GetAssetsByClass call is a cludget to get the full asset paths for the blueprints we care about, Bob T. thinks 
-		// that the Asset Registry could just keep asset paths:
-		TArray<FAssetData> BlueprintAssetData;
-		AssetRegistryModule.Get().GetAssetsByClass(UBlueprint::StaticClass()->GetFName(), BlueprintAssetData);
-
-		for (auto OnDiskClass : OnDiskClasses)
+		if (OnDiskClasses.Num() > 0)
 		{
-			FAssetData AssetData;
-			FString FixedString = OnDiskClass.ToString();
-			FixedString.RemoveFromEnd(TEXT("_C"));
-			for (auto Blueprint : BlueprintAssetData)
+			// GetAssetsByClass call is a kludge to get the full asset paths for the blueprints we care about, Bob T. thinks 
+			// that the Asset Registry could just keep asset paths:
+			TArray<FAssetData> BlueprintAssetData;
+			AssetRegistryModule.Get().GetAssetsByClass(UBlueprint::StaticClass()->GetFName(), BlueprintAssetData);
+
+			for (auto OnDiskClass : OnDiskClasses)
 			{
-				if (Blueprint.AssetName.ToString() == FixedString)
+				FString FixedString = OnDiskClass.ToString();
+				FixedString.RemoveFromEnd(TEXT("_C"));
+
+				FAssetData AssetData;
 				{
-					AssetData = Blueprint;
-					break;
+					// find asset data for this unloaded class, this makes our logic n^2 and may provide opportunity for further performance improvement:
+					for (auto& Blueprint : BlueprintAssetData)
+					{
+						if (Blueprint.AssetName.ToString() == FixedString)
+						{
+							AssetData = Blueprint;
+							break;
+						}
+					}
 				}
-			}
 
-			FComponentTypeEntry Entry = { FixedString, AssetData.ObjectPath.ToString(), nullptr };
-			ComponentTypeList.Add(Entry);
+				FComponentTypeEntry Entry = { FixedString, AssetData.ObjectPath.ToString(), nullptr };
+				ComponentTypeList.Add(Entry);
 
-			// The blueprint is unloaded, so we need to work out which icon to use for it using its asset data
-			const UClass* BlueprintIconClass = FClassIconFinder::GetIconClassForAssetData(AssetData);
+				// The blueprint is unloaded, so we need to work out which icon to use for it using its asset data
+				const UClass* BlueprintIconClass = FClassIconFinder::GetIconClassForAssetData(AssetData);
 
-			FComponentClassComboEntryPtr NewEntry(new FComponentClassComboEntry(BlueprintComponents, FixedString, AssetData.ObjectPath, BlueprintIconClass, bIncludeInFilter));
-			SortedClassList.Add(NewEntry);
+				FComponentClassComboEntryPtr NewEntry(new FComponentClassComboEntry(BlueprintComponents, FixedString, AssetData.ObjectPath, BlueprintIconClass, bIncludeInFilter));
+				SortedClassList.Add(NewEntry);
+			}
 		}
 	}
 
