Make Mobile HDR log use Display instead of Warning

diff --git a/Engine/Source/Editor/UnrealEd/Private/CookOnTheFlyServer.cpp b/Engine/Source/Editor/UnrealEd/Private/CookOnTheFlyServer.cpp
--- a/Engine/Source/Editor/UnrealEd/Private/CookOnTheFlyServer.cpp
+++ b/Engine/Source/Editor/UnrealEd/Private/CookOnTheFlyServer.cpp
@@ -2772,7 +2772,7 @@ void UCookOnTheFlyServer::Initialize( ECookMode::Type DesiredCookMode, ECookInit
 
 	ITargetPlatformManagerModule& TPM = GetTargetPlatformManagerRef();
 	const TArray<ITargetPlatform*>& Platforms = TPM.GetCookingTargetPlatforms();
-	UE_LOG(LogCook, Warning, TEXT("Mobile HDR setting %d"), IsMobileHDR());
+	UE_LOG(LogCook, Display, TEXT("Mobile HDR setting %d"), IsMobileHDR());
 
 	// always generate the asset registry before starting to cook, for either method
 	GenerateAssetRegistry(Platforms);
