#ANIM : Fix additive base pose display

Merging

//depot/UE4-Orion/Engine/Source/Editor/UnrealEd/Private/Animation/DebugSkelMeshComponent.cpp

to //depot/UE4-Orion/Engine/Source/Editor/UnrealEd/Private/Animation/DebugSkelMeshComponent.cpp

diff --git a/Engine/Source/Editor/UnrealEd/Private/Animation/DebugSkelMeshComponent.cpp b/Engine/Source/Editor/UnrealEd/Private/Animation/DebugSkelMeshComponent.cpp
--- a/Engine/Source/Editor/UnrealEd/Private/Animation/DebugSkelMeshComponent.cpp
+++ b/Engine/Source/Editor/UnrealEd/Private/Animation/DebugSkelMeshComponent.cpp
@@ -431,7 +431,7 @@ void UDebugSkelMeshComponent::RefreshBoneTransforms(FActorComponentTickFunction*
 
 	const bool bIsPreviewInstance = (PreviewInstance && PreviewInstance == AnimScriptInstance);
 
-	BakedAnimationPoses.Empty();
+	BakedAnimationPoses.Reset();
 	if(bDisplayBakedAnimation && bIsPreviewInstance && PreviewInstance->RequiredBones.IsValid())
 	{
 		if(UAnimSequence* Sequence = Cast<UAnimSequence>(PreviewInstance->CurrentAsset))
@@ -448,7 +448,7 @@ void UDebugSkelMeshComponent::RefreshBoneTransforms(FActorComponentTickFunction*
 		}
 	}
 
-	SourceAnimationPoses.Empty();
+	SourceAnimationPoses.Reset();
 	if(bDisplaySourceAnimation && bIsPreviewInstance && PreviewInstance->RequiredBones.IsValid())
 	{
 		if(UAnimSequence* Sequence = Cast<UAnimSequence>(PreviewInstance->CurrentAsset))
@@ -463,7 +463,7 @@ void UDebugSkelMeshComponent::RefreshBoneTransforms(FActorComponentTickFunction*
 		}
 	}
 
-	UncompressedSpaceBases.Empty();
+	UncompressedSpaceBases.Reset();
 	if (bDisplayRawAnimation && AnimScriptInstance && AnimScriptInstance->RequiredBones.IsValid())
 	{
 		UncompressedSpaceBases.AddUninitialized(AnimScriptInstance->RequiredBones.GetNumBones());
@@ -474,7 +474,7 @@ void UDebugSkelMeshComponent::RefreshBoneTransforms(FActorComponentTickFunction*
 	}
 
 	// Non retargeted pose.
-	NonRetargetedSpaceBases.Empty();
+	NonRetargetedSpaceBases.Reset();
 	if( bDisplayNonRetargetedPose && AnimScriptInstance && AnimScriptInstance->RequiredBones.IsValid() )
 	{
 		NonRetargetedSpaceBases.AddUninitialized(AnimScriptInstance->RequiredBones.GetNumBones());
@@ -484,7 +484,7 @@ void UDebugSkelMeshComponent::RefreshBoneTransforms(FActorComponentTickFunction*
 	}
 
 	// Only works in PreviewInstance, and not for anim blueprint. This is intended.
-	AdditiveBasePoses.Empty();
+	AdditiveBasePoses.Reset();
 	if( bDisplayAdditiveBasePose && bIsPreviewInstance && PreviewInstance->RequiredBones.IsValid() )
 	{
 		if (UAnimSequence* Sequence = Cast<UAnimSequence>(PreviewInstance->CurrentAsset)) 
@@ -500,6 +500,7 @@ void UDebugSkelMeshComponent::RefreshBoneTransforms(FActorComponentTickFunction*
 					CSAdditiveBasePose.InitPose(AdditiveBasePose);
 				}
 
+				AdditiveBasePoses.AddUninitialized(PreviewInstance->RequiredBones.GetNumBones());
 				for (int32 i = 0; i < AdditiveBasePoses.Num(); ++i)
 				{
 					FCompactPoseBoneIndex CompactIndex = PreviewInstance->RequiredBones.MakeCompactPoseIndex(FMeshPoseBoneIndex(i));
