Move Blueprint promotion tests into the Kismet module.

diff --git a/Engine/Source/Editor/Kismet/Private/Tests/BlueprintEditorTests.cpp b/Engine/Source/Editor/Kismet/Private/Tests/BlueprintEditorTests.cpp
new file mode 100644
--- /dev/null
+++ b/Engine/Source/Editor/Kismet/Private/Tests/BlueprintEditorTests.cpp
@@ -0,0 +1,1429 @@
+// Copyright 1998-2015 Epic Games, Inc. All Rights Reserved.
+
+#include "BlueprintEditorPrivatePCH.h"
+#include "BlueprintEditor.h"
+#include "BlueprintEditorModes.h"
+#include "BlueprintEditorUtils.h"
+#include "CompilerResultsLog.h"
+#include "ComponentAssetBroker.h"
+#include "EdGraphUtilities.h"
+#include "Engine/Breakpoint.h"
+#include "Engine/LevelScriptBlueprint.h"
+#include "Kismet/KismetStringLibrary.h"
+#include "Kismet/KismetSystemLibrary.h"
+#include "KismetDebugUtilities.h"
+
+#include "AssetSelection.h"
+#include "ScopedTransaction.h"
+#include "ObjectTools.h"
+
+// Automation
+#include "AutomationTest.h"
+#include "AutomationEditorCommon.h"
+#include "AutomationEditorPromotionCommon.h"
+
+#define LOCTEXT_NAMESPACE "BlueprintEditorPromotionTests"
+
+DEFINE_LOG_CATEGORY_STATIC(LogBlueprintEditorPromotionTests, Log, All);
+
+IMPLEMENT_SIMPLE_AUTOMATION_TEST(FBlueprintEditorPromotionTest, "System.Promotion.Editor.Blueprint Editor", EAutomationTestFlags::EditorContext | EAutomationTestFlags::EngineFilter);
+
+/**
+* Helper functions used by the blueprint editor promotion automation test
+*/
+namespace BlueprintEditorPromotionUtils
+{
+	/**
+	* Constants
+	*/
+	static const FString BlueprintNameString = TEXT("BlueprintEditorPromotionBlueprint");
+	static const FName BlueprintStringVariableName(TEXT("MyStringVariable"));
+
+	/**
+	* Gets the full path to the folder on disk
+	*/
+	static FString GetFullPath()
+	{
+		return FPackageName::FilenameToLongPackageName(FPaths::GameContentDir() + TEXT("BlueprintEditorPromotionTest"));
+	}
+
+	/**
+	* Helper class to track once a certain time has passed
+	*/
+	class FDelayHelper
+	{
+	public:
+		/**
+		* Constructor
+		*/
+		FDelayHelper() :
+			bIsRunning(false),
+			StartTime(0.f),
+			Duration(0.f)
+		{
+		}
+
+		/**
+		* Returns if the delay is still running
+		*/
+		bool IsRunning()
+		{
+			return bIsRunning;
+		}
+
+		/**
+		* Sets the helper state to not running
+		*/
+		void Reset()
+		{
+			bIsRunning = false;
+		}
+
+		/**
+		* Starts the delay timer
+		*
+		* @param InDuration - How long to delay for in seconds
+		*/
+		void Start(double InDuration)
+		{
+			bIsRunning = true;
+			StartTime = FPlatformTime::Seconds();
+			Duration = InDuration;
+		}
+
+		/**
+		* Returns true if the desired amount of time has passed
+		*/
+		bool IsComplete()
+		{
+			if (IsRunning())
+			{
+				const double CurrentTime = FPlatformTime::Seconds();
+				return CurrentTime - StartTime >= Duration;
+			}
+			else
+			{
+				return false;
+			}
+		}
+
+	private:
+		/** If true, this delay timer is active */
+		bool bIsRunning;
+		/** The time the delay started */
+		double StartTime;
+		/** How long the timer is for */
+		double Duration;
+	};
+
+	/**
+	* Sends the AssetEditor->SaveAsset UI command
+	*/
+	static void SendBlueprintResetViewCommand()
+	{
+		const FString Context = TEXT("BlueprintEditor");
+		const FString Command = TEXT("ResetCamera");
+
+		FInputChord CurrentSaveChord = FEditorPromotionTestUtilities::GetOrSetUICommand(Context, Command);
+
+		const FName FocusWidgetType(TEXT("SSCSEditorViewport"));
+		FEditorPromotionTestUtilities::SendCommandToCurrentEditor(CurrentSaveChord, FocusWidgetType);
+	}
+
+	/**
+	* Compiles the blueprint
+	*
+	* @param InBlueprint - The blueprint to compile
+	*/
+	static void CompileBlueprint(UBlueprint* InBlueprint)
+	{
+		FBlueprintEditorUtils::RefreshAllNodes(InBlueprint);
+
+		bool bIsRegeneratingOnLoad = false;
+		bool bSkipGarbageCollection = true;
+
+		FKismetEditorUtilities::CompileBlueprint(InBlueprint, bIsRegeneratingOnLoad, bSkipGarbageCollection);
+		if (InBlueprint->Status == EBlueprintStatus::BS_UpToDate)
+		{
+			UE_LOG(LogBlueprintEditorPromotionTests, Display, TEXT("Blueprint compiled successfully (%s)"), *InBlueprint->GetName());
+		}
+		else if (InBlueprint->Status == EBlueprintStatus::BS_UpToDateWithWarnings)
+		{
+			UE_LOG(LogBlueprintEditorPromotionTests, Display, TEXT("Blueprint compiled successfully with warnings(%s)"), *InBlueprint->GetName());
+		}
+		else if (InBlueprint->Status == EBlueprintStatus::BS_Error)
+		{
+			UE_LOG(LogBlueprintEditorPromotionTests, Display, TEXT("Blueprint failed to compile (%s)"), *InBlueprint->GetName());
+		}
+		else
+		{
+			UE_LOG(LogBlueprintEditorPromotionTests, Error, TEXT("Blueprint is in an unexpected state after compiling (%s)"), *InBlueprint->GetName());
+		}
+	}
+
+	/**
+	* Creates a blueprint component based off the supplied asset
+	*
+	* @param InBlueprint - The blueprint to modify
+	* @param InAsset - The asset to use for the component
+	*/
+	static USCS_Node* CreateBlueprintComponent(UBlueprint* InBlueprint, UObject* InAsset)
+	{
+		IAssetEditorInstance* OpenEditor = FAssetEditorManager::Get().FindEditorForAsset(InBlueprint, true);
+		FBlueprintEditor* CurrentBlueprintEditor = (FBlueprintEditor*)OpenEditor;
+		TSubclassOf<UActorComponent> ComponentClass = FComponentAssetBrokerage::GetPrimaryComponentForAsset(InAsset->GetClass());
+
+		USCS_Node* NewNode = InBlueprint->SimpleConstructionScript->CreateNode(ComponentClass);
+		// Assign the asset to the template
+		FComponentAssetBrokerage::AssignAssetToComponent(NewNode->ComponentTemplate, InAsset);
+
+		// Add node to the SCS
+		TArray<USCS_Node*> AllNodes = InBlueprint->SimpleConstructionScript->GetAllNodes();
+		USCS_Node* RootNode = AllNodes.Num() > 0 ? AllNodes[0] : NULL;
+		if (!RootNode || (RootNode == InBlueprint->SimpleConstructionScript->GetDefaultSceneRootNode()))
+		{
+			//New Root
+			InBlueprint->SimpleConstructionScript->AddNode(NewNode);
+		}
+		else
+		{
+			//Add as a child
+			RootNode->AddChildNode(NewNode);
+		}
+
+		// Recompile skeleton because of the new component we added
+		FKismetEditorUtilities::GenerateBlueprintSkeleton(InBlueprint, true);
+
+		CurrentBlueprintEditor->UpdateSCSPreview(true);
+
+		return NewNode;
+	}
+
+	/**
+	* Sets a new component as the root
+	*
+	* @param InBlueprint - The blueprint to modify
+	* @param NewRoot - The new root
+	*/
+	static void SetComponentAsRoot(UBlueprint* InBlueprint, USCS_Node* NewRoot)
+	{
+		// @FIXME: Current usages doesn't guarantee NewRoot is valid!!! Check first!
+		//Get all the construction script nodes
+		TArray<USCS_Node*> AllNodes = InBlueprint->SimpleConstructionScript->GetAllNodes();
+
+		USCS_Node* OldRootNode = AllNodes[0];
+		USCS_Node* OldParent = NULL;
+
+		//Find old parent
+		for (int32 NodeIndex = 0; NodeIndex < AllNodes.Num(); ++NodeIndex)
+		{
+			if (AllNodes[NodeIndex]->ChildNodes.Contains(NewRoot))
+			{
+				OldParent = AllNodes[NodeIndex];
+				break;
+			}
+		}
+
+		//Remove the new root from its old parent and 
+		OldParent->ChildNodes.Remove(NewRoot);
+		NewRoot->Modify();
+		NewRoot->AttachToName = NAME_None;
+
+		//Remove the old root, add the new root, and attach the old root as a child
+		InBlueprint->SimpleConstructionScript->RemoveNode(OldRootNode);
+		InBlueprint->SimpleConstructionScript->AddNode(NewRoot);
+		NewRoot->AddChildNode(OldRootNode);
+	}
+
+	/**
+	* Removes a blueprint component from the simple construction script
+	*
+	* @param InBlueprint - The blueprint to modify
+	* @param InNode - The node of the component to remove
+	*/
+	static void RemoveBlueprintComponent(UBlueprint* InBlueprint, USCS_Node* InNode)
+	{
+		if (InNode != NULL)
+		{
+			// Remove node from SCS tree
+			InNode->GetSCS()->RemoveNodeAndPromoteChildren(InNode);
+
+			// Clear the delegate
+			InNode->SetOnNameChanged(FSCSNodeNameChanged());
+		}
+	}
+
+	/**
+	* Creates a new graph node from a given template
+	*
+	* @param NodeTemplate - The template to use for the node
+	* @param InGraph - The graph to create the new node in
+	* @param GraphLocation - The location to place the node
+	* @param ConnectPin - The pin to connect the node to
+	*/
+	static UEdGraphNode* CreateNewGraphNodeFromTemplate(UK2Node* NodeTemplate, UEdGraph* InGraph, const FVector2D& GraphLocation, UEdGraphPin* ConnectPin = NULL)
+	{
+		const FString EmptyString(TEXT(""));
+		TSharedPtr<FEdGraphSchemaAction_K2NewNode> Action = TSharedPtr<FEdGraphSchemaAction_K2NewNode>(new FEdGraphSchemaAction_K2NewNode(FText::GetEmpty(), FText::GetEmpty(), EmptyString, 0));
+		Action->NodeTemplate = NodeTemplate;
+
+		return Action->PerformAction(InGraph, ConnectPin, GraphLocation, false);
+	}
+
+	/**
+	* Creates an AddComponent action to the blueprint graph
+	*
+	* @param InBlueprint - The blueprint to modify
+	* @param InGraph - The blueprint graph to use
+	* @param InAsset - The asset to use
+	*/
+	static UEdGraphNode* CreateAddComponentActionNode(UBlueprint* InBlueprint, UEdGraph* InGraph, UObject* InAsset)
+	{
+		UEdGraph* TempOuter = NewObject<UEdGraph>((UObject*)InBlueprint);
+		TempOuter->SetFlags(RF_Transient);
+
+		const FScopedTransaction PropertyChanged(LOCTEXT("AddedGraphNode", "Added a graph node"));
+		InGraph->Modify();
+
+		// Make an add component node
+		UK2Node_CallFunction* CallFuncNode = NewObject<UK2Node_AddComponent>(TempOuter);
+		UFunction* AddComponentFn = FindFieldChecked<UFunction>(AActor::StaticClass(), UK2Node_AddComponent::GetAddComponentFunctionName());
+		CallFuncNode->FunctionReference.SetFromField<UFunction>(AddComponentFn, FBlueprintEditorUtils::IsActorBased(InBlueprint));
+
+		UEdGraphNode* NewNode = CreateNewGraphNodeFromTemplate(CallFuncNode, InGraph, FVector2D(200, 0));
+
+		TSubclassOf<UActorComponent> ComponentClass = InAsset ? FComponentAssetBrokerage::GetPrimaryComponentForAsset(InAsset->GetClass()) : NULL;
+		if ((NewNode != NULL) && (InBlueprint != NULL))
+		{
+			UK2Node_AddComponent* AddCompNode = CastChecked<UK2Node_AddComponent>(NewNode);
+
+			ensure(NULL != Cast<UBlueprintGeneratedClass>(InBlueprint->GeneratedClass));
+			// Then create a new template object, and add to array in
+			UActorComponent* NewTemplate = NewObject<UActorComponent>(InBlueprint->GeneratedClass, ComponentClass, NAME_None, RF_ArchetypeObject | RF_Public);
+			InBlueprint->ComponentTemplates.Add(NewTemplate);
+
+			// Set the name of the template as the default for the TemplateName param
+			UEdGraphPin* TemplateNamePin = AddCompNode->GetTemplateNamePinChecked();
+			if (TemplateNamePin)
+			{
+				TemplateNamePin->DefaultValue = NewTemplate->GetName();
+			}
+
+			// Set the return type to be the type of the template
+			UEdGraphPin* ReturnPin = AddCompNode->GetReturnValuePin();
+			if (ReturnPin)
+			{
+				ReturnPin->PinType.PinSubCategoryObject = *ComponentClass;
+			}
+
+			// Set the asset
+			if (InAsset != NULL)
+			{
+				FComponentAssetBrokerage::AssignAssetToComponent(NewTemplate, InAsset);
+			}
+
+			AddCompNode->ReconstructNode();
+		}
+
+		FBlueprintEditorUtils::MarkBlueprintAsModified(InBlueprint);
+		return NewNode;
+	}
+
+	/**
+	* Creates a SetStaticMesh action in the blueprint graph
+	*
+	* @param InBlueprint - The blueprint to modify
+	* @param InGraph - The blueprint graph to use
+	*/
+	static UEdGraphNode* AddSetStaticMeshNode(UBlueprint* InBlueprint, UEdGraph* InGraph)
+	{
+		UEdGraph* TempOuter = NewObject<UEdGraph>((UObject*)InBlueprint);
+		TempOuter->SetFlags(RF_Transient);
+
+		// Make a call function template
+		UK2Node_CallFunction* CallFuncNode = NewObject<UK2Node_CallFunction>(TempOuter);
+		static FName PrintStringFunctionName(TEXT("SetStaticMesh"));
+		UFunction* DelayFn = FindFieldChecked<UFunction>(UStaticMeshComponent::StaticClass(), PrintStringFunctionName);
+		CallFuncNode->FunctionReference.SetFromField<UFunction>(DelayFn, false);
+
+		return CreateNewGraphNodeFromTemplate(CallFuncNode, InGraph, FVector2D(850, 0));
+	}
+
+	/**
+	* Connects two nodes using the supplied pin names
+	*
+	* @param NodeA - The first node to connect
+	* @param PinAName - The name of the pin on the first node
+	* @param NodeB - The second node to connect
+	* @param PinBName - The name of the pin on the second node
+	*/
+	static void ConnectGraphNodes(UEdGraphNode* NodeA, const FString& PinAName, UEdGraphNode* NodeB, const FString& PinBName)
+	{
+		const FScopedTransaction PropertyChanged(LOCTEXT("ConnectedNode", "Connected graph nodes"));
+		NodeA->GetGraph()->Modify();
+
+		UEdGraphPin* PinA = NodeA->FindPin(PinAName);
+		UEdGraphPin* PinB = NodeB->FindPin(PinBName);
+
+		if (PinA && PinB)
+		{
+			PinA->MakeLinkTo(PinB);
+		}
+		else
+		{
+			UE_LOG(LogBlueprintEditorPromotionTests, Error, TEXT("Could not connect pins %s and %s "), *PinAName, *PinBName);
+		}
+	}
+	
+	/**
+	* Promotes a pin to a variable
+	*
+	* @param InBlueprint - The blueprint to modify
+	* @param Node - The node that owns the pin
+	* @param PinName - The name of the pin to promote
+	*/
+	static void PromotePinToVariable(UBlueprint* InBlueprint, UEdGraphNode* Node, const FString& PinName)
+	{
+		IAssetEditorInstance* OpenEditor = FAssetEditorManager::Get().FindEditorForAsset(InBlueprint, true);
+		FBlueprintEditor* CurrentBlueprintEditor = (FBlueprintEditor*)OpenEditor;
+
+		UEdGraphPin* PinToPromote = Node->FindPin(PinName);
+		CurrentBlueprintEditor->DoPromoteToVariable(InBlueprint, PinToPromote, true);
+	}
+
+	/**
+	* Creates a ReceiveBeginPlay event node
+	*
+	* @param InBlueprint - The blueprint to modify
+	* @param InGraph - The graph to use for the new node
+	*/
+	static UEdGraphNode* CreatePostBeginPlayEvent(UBlueprint* InBlueprint, UEdGraph* InGraph)
+	{
+		UEdGraph* TempOuter = NewObject<UEdGraph>((UObject*)InBlueprint);
+		TempOuter->SetFlags(RF_Transient);
+
+		// Make an add component node
+		UK2Node_Event* NewEventNode = NewObject<UK2Node_Event>(TempOuter);
+		NewEventNode->EventReference.SetExternalMember(FName(TEXT("ReceiveBeginPlay")), AActor::StaticClass());
+		NewEventNode->bOverrideFunction = true;
+
+		//Check for existing events
+		UK2Node_Event* ExistingEvent = FBlueprintEditorUtils::FindOverrideForFunction(InBlueprint, NewEventNode->EventReference.GetMemberParentClass(NewEventNode->GetBlueprintClassFromNode()), NewEventNode->EventReference.GetMemberName());
+
+		if (!ExistingEvent)
+		{
+			return CreateNewGraphNodeFromTemplate(NewEventNode, InGraph, FVector2D(200, 0));
+		}
+		return ExistingEvent;
+	}
+
+	/**
+	* Creates a custom event node
+	*
+	* @param InBlueprint - The blueprint to modify
+	* @param InGraph - The graph to use for the new node
+	* @param EventName - The name of the event
+	*/
+	static UEdGraphNode* CreateCustomEvent(UBlueprint* InBlueprint, UEdGraph* InGraph, const FString& EventName)
+	{
+		UEdGraph* TempOuter = NewObject<UEdGraph>((UObject*)InBlueprint);
+		TempOuter->SetFlags(RF_Transient);
+
+		// Make an add component node
+		UK2Node_CustomEvent* NewEventNode = NewObject<UK2Node_CustomEvent>(TempOuter);
+		NewEventNode->CustomFunctionName = "EventName";
+
+		return CreateNewGraphNodeFromTemplate(NewEventNode, InGraph, FVector2D(1200, 0));
+	}
+
+	/**
+	* Creates a node template for a UKismetSystemLibrary function
+	*
+	* @param NodeOuter - The outer to use for the template
+	* @param InGraph - The function to use for the node
+	*/
+	static UK2Node* CreateKismetFunctionTemplate(UObject* NodeOuter, const FName& FunctionName)
+	{
+		// Make a call function template
+		UK2Node_CallFunction* CallFuncNode = NewObject<UK2Node_CallFunction>(NodeOuter);
+		UFunction* Function = FindFieldChecked<UFunction>(UKismetSystemLibrary::StaticClass(), FunctionName);
+		CallFuncNode->FunctionReference.SetFromField<UFunction>(Function, false);
+		return CallFuncNode;
+	}
+
+	/**
+	* Creates a delay node
+	*
+	* @param InBlueprint - The blueprint to modify
+	* @param InGraph - The graph to use for the new node
+	* @param ConnectPin - The pin to connect the new node to
+	*/
+	static UEdGraphNode* AddDelayNode(UBlueprint* InBlueprint, UEdGraph* InGraph, UEdGraphPin* ConnectPin = NULL)
+	{
+		UEdGraph* TempOuter = NewObject<UEdGraph>((UObject*)InBlueprint);
+		TempOuter->SetFlags(RF_Transient);
+
+		const FScopedTransaction PropertyChanged(LOCTEXT("AddedGraphNode", "Added a graph node"));
+		InGraph->Modify();
+
+		// Make a call function template
+		static FName DelayFunctionName(TEXT("Delay"));
+		UK2Node* CallFuncNode = CreateKismetFunctionTemplate(TempOuter, DelayFunctionName);
+
+		//Create the node
+		return CreateNewGraphNodeFromTemplate(CallFuncNode, InGraph, FVector2D(400, 0), ConnectPin);
+	}
+
+	/**
+	* Creates a PrintString node
+	*
+	* @param InBlueprint - The blueprint to modify
+	* @param InGraph - The graph to use for the new node
+	* @param ConnectPin - The pin to connect the new node to
+	*/
+	static UEdGraphNode* AddPrintStringNode(UBlueprint* InBlueprint, UEdGraph* InGraph, UEdGraphPin* ConnectPin = NULL)
+	{
+		UEdGraph* TempOuter = NewObject<UEdGraph>((UObject*)InBlueprint);
+		TempOuter->SetFlags(RF_Transient);
+
+		// Make a call function template
+		static FName PrintStringFunctionName(TEXT("PrintString"));
+		UK2Node* CallFuncNode = CreateKismetFunctionTemplate(TempOuter, PrintStringFunctionName);
+
+		return CreateNewGraphNodeFromTemplate(CallFuncNode, InGraph, FVector2D(680, 0), ConnectPin);
+	}
+
+	/**
+	* Creates a call function node
+	*
+	* @param InBlueprint - The blueprint to modify
+	* @param InGraph - The graph to use for the new node
+	* @param FunctionName - The name of the function to call
+	* @param ConnectPin - The pin to connect the new node to
+	*/
+	static UEdGraphNode* AddCallFunctionGraphNode(UBlueprint* InBlueprint, UEdGraph* InGraph, const FName& FunctionName, UEdGraphPin* ConnectPin = NULL)
+	{
+		UEdGraph* TempOuter = NewObject<UEdGraph>((UObject*)InBlueprint);
+		TempOuter->SetFlags(RF_Transient);
+
+		// Make a call function template
+		UK2Node_CallFunction* CallFuncNode = NewObject<UK2Node_CallFunction>(TempOuter);
+		CallFuncNode->FunctionReference.SetSelfMember(FunctionName);
+
+		return CreateNewGraphNodeFromTemplate(CallFuncNode, InGraph, FVector2D(1200, 0), ConnectPin);
+	}
+
+	/**
+	* Creates Get or Set node
+	*
+	* @param InBlueprint - The blueprint to modify
+	* @param InGraph - The graph to use for the new node
+	* @param VarName - The name of the variable to use
+	* @param bGet - If true, create a Get node.  If false, create a Set node.
+	* @param XOffset - How far to offset the node in the graph
+	*/
+	static UEdGraphNode* AddGetSetNode(UBlueprint* InBlueprint, UEdGraph* InGraph, const FString& VarName, bool bGet, float XOffset = 0.f)
+	{
+		const FScopedTransaction PropertyChanged(LOCTEXT("AddedGraphNode", "Added a graph node"));
+		InGraph->Modify();
+
+		FEdGraphSchemaAction_K2NewNode NodeInfo;
+		// Create get or set node, depending on whether we clicked on an input or output pin
+		UK2Node_Variable* TemplateNode = NULL;
+		if (bGet)
+		{
+			TemplateNode = NewObject<UK2Node_VariableGet>();
+		}
+		else
+		{
+			TemplateNode = NewObject<UK2Node_VariableSet>();
+		}
+
+		TemplateNode->VariableReference.SetSelfMember(FName(*VarName));
+		NodeInfo.NodeTemplate = TemplateNode;
+
+		return NodeInfo.PerformAction(InGraph, NULL, FVector2D(XOffset, 130), true);
+	}
+
+	/**
+	* Sets the default value for a pin
+	*
+	* @param Node - The node that owns the pin to set
+	* @param PinName - The name of the pin
+	* @param PinValue - The new default value
+	*/
+	static void SetPinDefaultValue(UEdGraphNode* Node, const FString& PinName, const FString& PinValue)
+	{
+		UEdGraphPin* Pin = Node->FindPin(PinName);
+		Pin->DefaultValue = PinValue;
+	}
+
+	/**
+	* Sets the default object for a pin
+	*
+	* @param Node - The node that owns the pin to set
+	* @param PinName - The name of the pin
+	* @param PinObject - The new default object
+	*/
+	static void SetPinDefaultObject(UEdGraphNode* Node, const FString& PinName, UObject* PinObject)
+	{
+		UEdGraphPin* Pin = Node->FindPin(PinName);
+		Pin->DefaultObject = PinObject;
+	}
+
+	/**
+	* Adds a string member variable to a blueprint
+	*
+	* @param InBlueprint - The blueprint to modify
+	* @param VariableName - The name of the new string variable
+	*/
+	static void AddStringMemberValue(UBlueprint* InBlueprint, const FName& VariableName)
+	{
+		const UEdGraphSchema_K2* K2Schema = GetDefault<UEdGraphSchema_K2>();
+		FEdGraphPinType StringPinType(K2Schema->PC_String, TEXT(""), NULL, false, false);
+		FBlueprintEditorUtils::AddMemberVariable(InBlueprint, VariableName, StringPinType);
+	}
+
+	/**
+	* Creates a new function graph
+	*
+	* @param InBlueprint - The blueprint to modify
+	* @param FunctionName - The function name to use for the new graph
+	*/
+	static UEdGraph* CreateNewFunctionGraph(UBlueprint* InBlueprint, const FName& FunctionName)
+	{
+		UEdGraph* NewGraph = FBlueprintEditorUtils::CreateNewGraph(InBlueprint, FunctionName, UEdGraph::StaticClass(), UEdGraphSchema_K2::StaticClass());
+		FBlueprintEditorUtils::AddFunctionGraph<UClass>(InBlueprint, NewGraph, /*bIsUserCreated=*/ true, NULL);
+		return NewGraph;
+	}
+}
+
+
+/**
+* Container for items related to the blueprint editor test
+*/
+namespace BlueprintEditorPromotionTestHelper
+{
+	/**
+	* The main build promotion test class
+	*/
+	struct FBlueprintEditorPromotionTestHelper
+	{
+		/** Pointer to running automation test instance */
+		FBlueprintEditorPromotionTest* Test;
+		
+		/** Pointer to the execution info of this test */
+		FAutomationTestExecutionInfo* TestExecutionInfo;
+		
+		/** The number of existing errors, warnings, and logs when the command started */
+		int32 LastErrorCount;
+		int32 LastWarningCount;
+		int32 LastLogCount;
+
+		/** Function definition for the test stage functions */
+		typedef bool(BlueprintEditorPromotionTestHelper::FBlueprintEditorPromotionTestHelper::*TestStageFn)();
+
+		/** List of test stage functions */
+		TArray<TestStageFn> TestStages;
+		TArray<FString> StageNames;
+
+		/** The index of the test stage we are on */
+		int32 CurrentStage;
+
+		/** Pointer to the active editor world */
+		UWorld* CurrentWorld;
+
+		/** Particle System loaded from Automation Settings for Blueprint Pass */
+		UParticleSystem* LoadedParticleSystem;
+
+		/** Meshes to use for the blueprint */
+		UStaticMesh* FirstBlueprintMesh;
+		UStaticMesh* SecondBlueprintMesh;
+
+		/** Objects created by the Blueprint stages */
+		UBlueprint* BlueprintObject;
+		UPackage* BlueprintPackage;
+		UEdGraph* CustomGraph;
+		USCS_Node* MeshNode;
+		USCS_Node* OtherMeshNode;
+		USCS_Node* PSNode;
+		UEdGraphNode* AddMeshNode;
+		UEdGraphNode* PostBeginPlayEventNode;
+		UEdGraphNode* DelayNode;
+		UEdGraphNode* SetNode;
+		UEdGraphNode* GetNode;
+		UEdGraphNode* PrintNode;
+		UEdGraphNode* SetStaticMeshNode;
+		UEdGraphNode* CustomEventNode;
+		UEdGraphNode* AddParticleSystemNode;
+		UEdGraphNode* CallFunctionNode;
+
+		/** Delay helper */
+		BlueprintEditorPromotionUtils::FDelayHelper DelayHelper;
+
+		/** List of skipped tests */
+		TArray<FString> SkippedTests;
+
+		/** summary logs to display at the end */
+		TArray<FString> SummaryLines;
+
+		/** Keeps track of errors generated by building the map */
+		int32 BuildStartErrorCount;
+
+		int32 SectionSuccessCount;
+		int32 SectionTestCount;
+
+		int32 SectionErrorCount;
+		int32 SectionWarningCount;
+		int32 SectionLogCount;
+
+#define ADD_TEST_STAGE(FuncName,StageName) \
+	TestStages.Add(&BlueprintEditorPromotionTestHelper::FBlueprintEditorPromotionTestHelper::FuncName); \
+	StageNames.Add(StageName); 
+
+		/**
+		* Constructor
+		*/
+		FBlueprintEditorPromotionTestHelper(FAutomationTestExecutionInfo* InExecutionInfo) :
+			CurrentStage(0)
+		{
+			FMemory::Memzero(this, sizeof(FBlueprintEditorPromotionTestHelper));
+
+			TestExecutionInfo = InExecutionInfo;
+
+			ADD_TEST_STAGE(Cleanup, TEXT("Pre-start cleanup"));
+			ADD_TEST_STAGE(Setup, TEXT("Setup"));
+			ADD_TEST_STAGE(Blueprint_CreateNewBlueprint_Part1, TEXT("Create a new Blueprint"));
+			ADD_TEST_STAGE(Blueprint_CreateNewBlueprint_Part2, TEXT("Create a new Blueprint"));
+			ADD_TEST_STAGE(Blueprint_DataOnlyBlueprint_Part1, TEXT("Data-only Blueprint"));
+			ADD_TEST_STAGE(Blueprint_DataOnlyBlueprint_Part2, TEXT("Data-only Blueprint"));
+			ADD_TEST_STAGE(Blueprint_DataOnlyBlueprint_Part3, TEXT("Data-only Blueprint"));
+			ADD_TEST_STAGE(Blueprint_ComponentsMode_Part1, TEXT("Components Mode"));
+			ADD_TEST_STAGE(Blueprint_ComponentsMode_Part2, TEXT("Components Mode"));
+			ADD_TEST_STAGE(Blueprint_ConstructionScript, TEXT("Construction Script"));
+			ADD_TEST_STAGE(Blueprint_PromoteVariable_Part1, TEXT("Variable from Component Mode 1"));
+			ADD_TEST_STAGE(Blueprint_PromoteVariable_Part2, TEXT("Variable from Component Mode 2"));
+			//ADD_TEST_STAGE(Blueprint_PromoteVariable_Part3, TEXT("Variable from Component Mode 3"));
+			ADD_TEST_STAGE(Blueprint_EventGraph, TEXT("Event Graph"));
+			ADD_TEST_STAGE(Blueprint_CustomVariable, TEXT("Custom Variables"));
+			ADD_TEST_STAGE(Blueprint_UsingVariables, TEXT("Using Variables"));
+			ADD_TEST_STAGE(Blueprint_RenameCustomEvent, TEXT("Renaming Custom Event"));
+			ADD_TEST_STAGE(Blueprint_NewFunctions, TEXT("New Function"));
+			ADD_TEST_STAGE(Blueprint_CompleteBlueprint, TEXT("Completing the Blueprint"));
+			ADD_TEST_STAGE(Cleanup, TEXT("Teardown"));
+		}
+
+		/**
+		* Runs the current test stage
+		*/
+		bool Update()
+		{
+			bool bStageComplete = (this->*TestStages[CurrentStage])();
+			if (bStageComplete)
+			{
+				//Only add headers if the next section has a different name, or we are the last one
+				if ((CurrentStage + 1 < TestStages.Num() && StageNames[CurrentStage] != StageNames[CurrentStage + 1]) || (CurrentStage + 1 == TestStages.Num()))
+				{
+					TagPreviousLogs(StageNames[CurrentStage]);
+				}
+
+				CurrentStage++;
+			}
+			return CurrentStage >= TestStages.Num();
+		}
+
+	private:
+
+		/**
+		* Handles tagging all logs that have been created since the last time this was called
+		*
+		* @param NewLogTag - The tag to prefix the logs with
+		*/
+		void TagPreviousLogs(const FString& NewLogTag)
+		{
+			if (TestExecutionInfo && NewLogTag.Len() > 0)
+			{
+				for (int32 ErrorIndex = LastErrorCount; ErrorIndex < TestExecutionInfo->Errors.Num(); ++ErrorIndex)
+				{
+					TestExecutionInfo->Errors[ErrorIndex] = FString::Printf(TEXT("%s: %s"), *NewLogTag, *TestExecutionInfo->Errors[ErrorIndex]);
+				}
+				for (int32 WarningIndex = LastWarningCount; WarningIndex < TestExecutionInfo->Warnings.Num(); ++WarningIndex)
+				{
+					TestExecutionInfo->Warnings[WarningIndex] = FString::Printf(TEXT("%s: %s"), *NewLogTag, *TestExecutionInfo->Warnings[WarningIndex]);
+				}
+				for (int32 LogIndex = LastLogCount; LogIndex < TestExecutionInfo->LogItems.Num(); ++LogIndex)
+				{
+					TestExecutionInfo->LogItems[LogIndex] = FString::Printf(TEXT("%s: %s"), *NewLogTag, *TestExecutionInfo->LogItems[LogIndex]);
+				}
+			}
+
+			//This sub tests was a success if we had no new errors
+			if (LastErrorCount == TestExecutionInfo->Errors.Num())
+			{
+				SectionSuccessCount++;
+			}
+			SectionTestCount++;
+
+			LastErrorCount = TestExecutionInfo->Errors.Num();
+			LastWarningCount = TestExecutionInfo->Warnings.Num();
+			LastLogCount = TestExecutionInfo->LogItems.Num();
+
+		}
+
+		bool Cleanup()
+		{
+			//Make sure no editors are open
+			FAssetEditorManager::Get().CloseAllAssetEditors();
+
+			//remove all assets in the build package
+			// Load the asset registry module
+			IAssetRegistry& AssetRegistry = FModuleManager::LoadModuleChecked<FAssetRegistryModule>(TEXT("AssetRegistry")).Get();
+
+			// Form a filter from the paths
+			FARFilter Filter;
+			Filter.bRecursivePaths = true;
+			new (Filter.PackagePaths) FName(*FEditorPromotionTestUtilities::GetGamePath());
+
+			// Query for a list of assets in the selected paths
+			TArray<FAssetData> AssetList;
+			AssetRegistry.GetAssets(Filter, AssetList);
+
+			// Clear and try to delete all assets
+			for (int32 AssetIdx = 0; AssetIdx < AssetList.Num(); ++AssetIdx)
+			{
+				Test->AddLogItem(*FString::Printf(TEXT("Removing asset: %s"), *AssetList[AssetIdx].AssetName.ToString()));
+				if (AssetList[AssetIdx].IsAssetLoaded())
+				{
+					UObject* LoadedAsset = AssetList[AssetIdx].GetAsset();
+					AssetRegistry.AssetDeleted(LoadedAsset);
+
+					bool bSuccessful = ObjectTools::DeleteSingleObject(LoadedAsset, false);
+
+					//If we failed to delete this object manually clear any references and try again
+					if (!bSuccessful)
+					{
+						//Clear references to the object so we can delete it
+						AutomationEditorCommonUtils::NullReferencesToObject(LoadedAsset);
+
+						bSuccessful = ObjectTools::DeleteSingleObject(LoadedAsset, false);
+					}
+				}
+			}
+
+			Test->AddLogItem(*FString::Printf(TEXT("Clearing Path: %s"), *FEditorPromotionTestUtilities::GetGamePath()));
+			AssetRegistry.RemovePath(FEditorPromotionTestUtilities::GetGamePath());
+
+			//Remove the directory
+			bool bEnsureExists = false;
+			bool bDeleteEntireTree = true;
+			FString PackageDirectory = FPaths::GameContentDir() / TEXT("BuildPromotionTest");
+			IFileManager::Get().DeleteDirectory(*PackageDirectory, bEnsureExists, bDeleteEntireTree);
+			Test->AddLogItem(*FString::Printf(TEXT("Deleting Folder: %s"), *PackageDirectory));
+
+			return true;
+		}
+
+		bool Setup()
+		{
+			//Make sure we have the required assets
+			UAutomationTestSettings const* AutomationTestSettings = GetDefault<UAutomationTestSettings>();
+			check(AutomationTestSettings);
+			FAssetData AssetData;
+
+			const FString FirstMeshPath = AutomationTestSettings->BlueprintEditorPromotionTest.FirstMeshPath.FilePath;
+			if (FirstMeshPath.Len() > 0)
+			{
+				AssetData = FEditorAutomationTestUtilities::GetAssetDataFromPackagePath(FirstMeshPath);
+				FirstBlueprintMesh = Cast<UStaticMesh>(AssetData.GetAsset());
+			}
+
+			const FString SecondMeshPath = AutomationTestSettings->BlueprintEditorPromotionTest.SecondMeshPath.FilePath;
+			if (SecondMeshPath.Len() > 0)
+			{
+				AssetData = FEditorAutomationTestUtilities::GetAssetDataFromPackagePath(SecondMeshPath);
+				SecondBlueprintMesh = Cast<UStaticMesh>(AssetData.GetAsset());
+			}
+
+			const FString ParticleSystemPath = AutomationTestSettings->BlueprintEditorPromotionTest.DefaultParticleAsset.FilePath;
+			if (ParticleSystemPath.Len() > 0)
+			{
+				AssetData = FEditorAutomationTestUtilities::GetAssetDataFromPackagePath(ParticleSystemPath);
+				LoadedParticleSystem = Cast<UParticleSystem>(AssetData.GetAsset());
+			}
+
+			if (!(FirstBlueprintMesh && SecondBlueprintMesh && LoadedParticleSystem))
+			{
+				SkippedTests.Add(TEXT("All Blueprint tests. (Missing a required mesh or particle system)"));
+				if (FirstMeshPath.IsEmpty() || SecondMeshPath.IsEmpty())
+				{
+					Test->AddLogItem(TEXT("SKIPPING BLUEPRINT TESTS.  FirstMeshPath or SecondMeshPath not configured in AutomationTestSettings."));
+				}
+				else
+				{
+					Test->AddWarning(TEXT("SKIPPING BLUEPRINT TESTS.  Invalid FirstMeshPath or SecondMeshPath in AutomationTestSettings, or particle system was not created."));
+				}
+			}
+
+			return true;
+		}
+
+		/**
+		* Blueprint Test Stage: Create a new Blueprint (Part 1)
+		*    Creates a new actor based blueprint and opens the editor
+		*/
+		bool Blueprint_CreateNewBlueprint_Part1()
+		{
+			// Exit early if any of the required assets are missing
+			if (!(FirstBlueprintMesh && SecondBlueprintMesh && LoadedParticleSystem))
+			{
+				return true;
+			}
+
+			UBlueprintFactory* Factory = NewObject<UBlueprintFactory>();
+			Factory->ParentClass = AActor::StaticClass();
+
+			const FString PackageName = FEditorPromotionTestUtilities::GetGamePath() + TEXT("/") + BlueprintEditorPromotionUtils::BlueprintNameString;
+			BlueprintPackage = CreatePackage(NULL, *PackageName);
+			EObjectFlags Flags = RF_Public | RF_Standalone;
+
+			UObject* ExistingBlueprint = FindObject<UBlueprint>(BlueprintPackage, *BlueprintEditorPromotionUtils::BlueprintNameString);
+			Test->TestNull(TEXT("Blueprint asset does not already exist (delete blueprint and restart editor)"), ExistingBlueprint);
+			// Exit early since test will not be valid with pre-existing assets
+			if (ExistingBlueprint)
+			{
+				return true;
+			}
+
+			BlueprintObject = Cast<UBlueprint>(Factory->FactoryCreateNew(UBlueprint::StaticClass(), BlueprintPackage, FName(*BlueprintEditorPromotionUtils::BlueprintNameString), Flags, NULL, GWarn));
+			Test->TestNotNull(TEXT("Created new Actor-based blueprint"), BlueprintObject);
+			if (BlueprintObject)
+			{
+				// Update asset registry and mark package dirty
+				FAssetRegistryModule::AssetCreated(BlueprintObject);
+				BlueprintPackage->MarkPackageDirty();
+
+				Test->AddLogItem(TEXT("Opening the blueprint editor for the first time"));
+				FAssetEditorManager::Get().OpenEditorForAsset(BlueprintObject);
+			}
+
+			return true;
+		}
+
+		/**
+		* Blueprint Test Stage: Create a new Blueprint (Part 2)
+		*    Checks that the blueprint editor opened in the correct mode
+		*/
+		bool Blueprint_CreateNewBlueprint_Part2()
+		{
+			if (BlueprintObject)
+			{
+				IAssetEditorInstance* AssetEditor = FAssetEditorManager::Get().FindEditorForAsset(BlueprintObject, true);
+				IBlueprintEditor* BlueprintEditor = (IBlueprintEditor*)AssetEditor;
+				Test->TestTrue(TEXT("Opened correct initial editor"), BlueprintEditor->GetCurrentMode() != FBlueprintEditorApplicationModes::BlueprintDefaultsMode);
+			}
+
+			return true;
+		}
+
+		/**
+		* Blueprint Test Stage: Data-only Blueprint (Part 1)
+		*    Closes the blueprint editor
+		*/
+		bool Blueprint_DataOnlyBlueprint_Part1()
+		{
+			if (BlueprintObject)
+			{
+				Test->AddLogItem(TEXT("Closing the blueprint editor"));
+				FAssetEditorManager::Get().CloseAllAssetEditors();
+			}
+			return true;
+		}
+
+		/**
+		* Blueprint Test Stage: Data-only Blueprint (Part 2)
+		*    Re opens the blueprint editor
+		*/
+		bool Blueprint_DataOnlyBlueprint_Part2()
+		{
+			if (BlueprintObject)
+			{
+				Test->AddLogItem(TEXT("Opening the blueprint editor for the second time"));
+				FAssetEditorManager::Get().OpenEditorForAsset(BlueprintObject);
+			}
+			return true;
+		}
+
+		/**
+		* Blueprint Test Stage: Data-only Blueprint (Part 3)
+		*    Checks that the editor opened in the Defaults mode
+		*/
+		bool Blueprint_DataOnlyBlueprint_Part3()
+		{
+			if (BlueprintObject)
+			{
+				IAssetEditorInstance* AssetEditor = FAssetEditorManager::Get().FindEditorForAsset(BlueprintObject, true);
+				IBlueprintEditor* BlueprintEditor = (IBlueprintEditor*)AssetEditor;
+				
+				bool bCorrectEditorOpened = BlueprintEditor->GetCurrentMode() == FBlueprintEditorApplicationModes::BlueprintDefaultsMode;
+				Test->TestTrue(TEXT("Correct defaults editor opened"), bCorrectEditorOpened);
+
+				if (bCorrectEditorOpened)
+				{
+					Test->AddLogItem(TEXT("Switching to components mode"));
+					BlueprintEditor->SetCurrentMode(FBlueprintEditorApplicationModes::BlueprintComponentsMode);
+				}
+			}
+			return true;
+		}
+
+		/**
+		* Blueprint Test Stage:  Components Mode (Part 1)
+		*   Adds 3 new components to the blueprint, changes the root component, renames the components, and compiles the blueprint
+		*/
+		bool Blueprint_ComponentsMode_Part1()
+		{
+			if (BlueprintObject)
+			{
+				IAssetEditorInstance* AssetEditor = FAssetEditorManager::Get().FindEditorForAsset(BlueprintObject, true);
+				IBlueprintEditor* BlueprintEditor = (IBlueprintEditor*)AssetEditor;
+
+				MeshNode = BlueprintEditorPromotionUtils::CreateBlueprintComponent(BlueprintObject, FirstBlueprintMesh);
+				Test->TestNotNull(TEXT("First mesh component added"), MeshNode);
+
+				OtherMeshNode = BlueprintEditorPromotionUtils::CreateBlueprintComponent(BlueprintObject, SecondBlueprintMesh);
+				Test->TestNotNull(TEXT("Second mesh component added"), OtherMeshNode);
+
+				PSNode = BlueprintEditorPromotionUtils::CreateBlueprintComponent(BlueprintObject, LoadedParticleSystem);
+				Test->TestNotNull(TEXT("Particle system component added"), PSNode);
+
+				//Set the Particle System as the root
+				BlueprintEditorPromotionUtils::SetComponentAsRoot(BlueprintObject, PSNode);
+				Test->TestTrue(TEXT("Particle system set as root"), PSNode->IsRootNode());
+
+				//Set Names
+				const FName MeshName(TEXT("FirstMesh"));
+				FBlueprintEditorUtils::RenameComponentMemberVariable(BlueprintObject, MeshNode, MeshName);
+				Test->AddLogItem(TEXT("Renamed the first mesh component to FirstMesh"));
+
+				const FName OtherMeshName(TEXT("SecondMesh"));
+				FBlueprintEditorUtils::RenameComponentMemberVariable(BlueprintObject, OtherMeshNode, OtherMeshName);
+				Test->AddLogItem(TEXT("Renamed the second mesh component to SecondMesh"));
+
+				const FName PSName(TEXT("ParticleSys"));
+				FBlueprintEditorUtils::RenameComponentMemberVariable(BlueprintObject, PSNode, PSName);
+				Test->AddLogItem(TEXT("Renamed the particle system component to ParticleSys"));
+
+				BlueprintEditorPromotionUtils::CompileBlueprint(BlueprintObject);
+
+				Test->AddLogItem(TEXT("Switched to graph editing mode"));
+				BlueprintEditor->SetCurrentMode(FBlueprintEditorApplicationModes::StandardBlueprintEditorMode);
+			}
+			return true;
+		}
+
+		/**
+		* Blueprint Test Stage: Components Mode (Part 2)
+		*    Removes the 3 components added before and compiles the blueprint
+		*/
+		bool Blueprint_ComponentsMode_Part2()
+		{
+			if (BlueprintObject)
+			{
+				IAssetEditorInstance* AssetEditor = FAssetEditorManager::Get().FindEditorForAsset(BlueprintObject, true);
+				IBlueprintEditor* BlueprintEditor = (IBlueprintEditor*)AssetEditor;
+
+				//FEditorPromotionTestUtilities::TakeScreenshot(TEXT("BlueprintComponentVariables"), true);
+
+				Test->AddLogItem(TEXT("Switched to components mode"));
+				BlueprintEditor->SetCurrentMode(FBlueprintEditorApplicationModes::BlueprintComponentsMode);
+
+				BlueprintEditorPromotionUtils::RemoveBlueprintComponent(BlueprintObject, MeshNode);
+				BlueprintEditorPromotionUtils::RemoveBlueprintComponent(BlueprintObject, OtherMeshNode);
+				BlueprintEditorPromotionUtils::RemoveBlueprintComponent(BlueprintObject, PSNode);
+
+				//There should only be the scene component left
+				Test->TestTrue(TEXT("Successfully removed blueprint components"), BlueprintObject->SimpleConstructionScript->GetAllNodes().Num() == 1);
+
+				MeshNode = NULL;
+				OtherMeshNode = NULL;
+				PSNode = NULL;
+
+				Test->AddLogItem(TEXT("Switched to graph mode"));
+				BlueprintEditor->SetCurrentMode(FBlueprintEditorApplicationModes::StandardBlueprintEditorMode);
+
+				BlueprintEditorPromotionUtils::CompileBlueprint(BlueprintObject);
+			}
+			return true;
+		}
+
+		/**
+		* Blueprint Test Stage: Construction Script
+		*    Adds an AddStaticMeshComponent to the construction graph and links it to the entry node
+		*/
+		bool Blueprint_ConstructionScript()
+		{
+			if (BlueprintObject)
+			{
+				const UEdGraphSchema_K2* K2Schema = GetDefault<UEdGraphSchema_K2>();
+
+				IAssetEditorInstance* AssetEditor = FAssetEditorManager::Get().FindEditorForAsset(BlueprintObject, true);
+				FBlueprintEditor* BlueprintEditor = (FBlueprintEditor*)AssetEditor;
+
+				UEdGraph* ConstructionGraph = FBlueprintEditorUtils::FindUserConstructionScript(BlueprintObject);
+				BlueprintEditor->OpenGraphAndBringToFront(ConstructionGraph);
+
+				AddMeshNode = BlueprintEditorPromotionUtils::CreateAddComponentActionNode(BlueprintObject, ConstructionGraph, FirstBlueprintMesh);
+				Test->TestNotNull(TEXT("Add Static Mesh Component node created"), AddMeshNode);
+
+				GEditor->UndoTransaction();
+				Test->TestTrue(TEXT("Undo add component node completed"), ConstructionGraph->Nodes.Num() == 0 || ConstructionGraph->Nodes[ConstructionGraph->Nodes.Num() - 1] != AddMeshNode);
+				
+				GEditor->RedoTransaction();
+				Test->TestTrue(TEXT("Redo add component node completed"), ConstructionGraph->Nodes.Num() > 0 && ConstructionGraph->Nodes[ConstructionGraph->Nodes.Num() - 1] == AddMeshNode);
+
+				TArray<UK2Node_FunctionEntry*> EntryNodes;
+				ConstructionGraph->GetNodesOfClass(EntryNodes);
+				UEdGraphNode* EntryNode = EntryNodes.Num() > 0 ? EntryNodes[0] : NULL;
+				Test->TestNotNull(TEXT("Found entry node to connect Add Static Mesh to"), EntryNode);
+				if (EntryNode)
+				{
+					BlueprintEditorPromotionUtils::ConnectGraphNodes(AddMeshNode, K2Schema->PN_Execute, EntryNode, K2Schema->PN_Then);
+
+					UEdGraphPin* EntryOutPin = EntryNode->FindPin(K2Schema->PN_Then);
+					UEdGraphPin* AddStaticMeshInPin = AddMeshNode->FindPin(K2Schema->PN_Execute);
+
+					Test->TestTrue(TEXT("Connected entry node to Add Static Mesh node"), EntryOutPin->LinkedTo.Contains(AddStaticMeshInPin));
+
+					GEditor->UndoTransaction();
+					Test->TestTrue(TEXT("Undo connection to Add Static Mesh Node succeeded"), EntryOutPin->LinkedTo.Num() == 0);
+					
+					GEditor->RedoTransaction();
+					Test->TestTrue(TEXT("Redo connection to Add Static Mesh Node succeeded"), EntryOutPin->LinkedTo.Contains(AddStaticMeshInPin));
+				}
+
+				BlueprintEditorPromotionUtils::CompileBlueprint(BlueprintObject);
+			}
+			return true;
+		}
+
+		/**
+		* Saves the blueprint stored in BlueprintObject
+		*/
+		void SaveBlueprint()
+		{
+			if (BlueprintObject && BlueprintPackage)
+			{
+				BlueprintPackage->SetDirtyFlag(true);
+				BlueprintPackage->FullyLoad();
+				const FString PackagePath = FEditorPromotionTestUtilities::GetGamePath() + TEXT("/") + BlueprintEditorPromotionUtils::BlueprintNameString;
+				bool bBlueprintSaved = UPackage::SavePackage(BlueprintPackage, NULL, RF_Standalone, *FPackageName::LongPackageNameToFilename(PackagePath, FPackageName::GetAssetPackageExtension()), GLog, nullptr, false, true, SAVE_None);
+				Test->TestTrue(*FString::Printf(TEXT("Blueprint saved successfully (%s)"), *BlueprintObject->GetName()), bBlueprintSaved);
+			}
+		}
+
+		/**
+		* Blueprint Test Stage: Variable from component mode (Part 1)
+		*    Promotes the return pin of the AddStaticMeshNode to a variable and then renames it
+		*/
+		bool Blueprint_PromoteVariable_Part1()
+		{
+			if (BlueprintObject)
+			{
+				IAssetEditorInstance* AssetEditor = FAssetEditorManager::Get().FindEditorForAsset(BlueprintObject, true);
+				FBlueprintEditor* BlueprintEditor = (FBlueprintEditor*)AssetEditor;
+
+				const UEdGraphSchema_K2* K2Schema = GetDefault<UEdGraphSchema_K2>();
+				BlueprintEditorPromotionUtils::PromotePinToVariable(BlueprintObject, AddMeshNode, K2Schema->PN_ReturnValue);
+
+				Test->AddLogItem(TEXT("Promoted the return pin on the add mesh node to a variable"));
+
+				const FName OldVarName(TEXT("NewVar_0")); // Default variable name
+				const FName NewVarName(TEXT("MyMesh"));
+				FBlueprintEditorUtils::RenameMemberVariable(BlueprintObject, OldVarName, NewVarName);
+				Test->TestNotEqual(TEXT("New variable was renamed"), FBlueprintEditorUtils::FindMemberVariableGuidByName(BlueprintObject, OldVarName), FBlueprintEditorUtils::FindMemberVariableGuidByName(BlueprintObject, NewVarName));
+				
+				BlueprintEditorPromotionUtils::CompileBlueprint(BlueprintObject);
+
+				Test->AddLogItem(TEXT("Switched to graph mode"));
+				BlueprintEditor->SetCurrentMode(FBlueprintEditorApplicationModes::StandardBlueprintEditorMode);
+			}
+			return true;
+		}
+
+		bool Blueprint_PromoteVariable_Part2()
+		{
+			if (BlueprintObject)
+			{
+				IAssetEditorInstance* AssetEditor = FAssetEditorManager::Get().FindEditorForAsset(BlueprintObject, true);
+				FBlueprintEditor* BlueprintEditor = (FBlueprintEditor*)AssetEditor;
+
+				//FEditorPromotionTestUtilities::TakeScreenshot(TEXT("BlueprintMeshVariable"), true);
+
+				Test->AddLogItem(TEXT("Switched to components mode"));
+				BlueprintEditor->SetCurrentMode(FBlueprintEditorApplicationModes::BlueprintComponentsMode);
+
+				BlueprintEditorPromotionUtils::SendBlueprintResetViewCommand();
+			}
+			return true;
+		}
+
+		///**
+		//* Blueprint Test Stage: Variable from component mode (Part 3)
+		//*    Takes a screenshot of the mesh variable
+		//*/
+		//bool Blueprint_PromoteVariable_Part3()
+		//{
+		//	if (BlueprintObject)
+		//	{
+		//		FEditorPromotionTestUtilities::TakeScreenshot(TEXT("BlueprintComponent"), true);
+		//	}
+		//	return true;
+		//}
+
+		/**
+		* Blueprint Test Stage: Event Graph
+		*    Adds a ReceiveBeginPlay and Delay node to the event graph 
+		*/
+		bool Blueprint_EventGraph()
+		{
+			if (BlueprintObject)
+			{
+				IAssetEditorInstance* AssetEditor = FAssetEditorManager::Get().FindEditorForAsset(BlueprintObject, true);
+				FBlueprintEditor* BlueprintEditor = (FBlueprintEditor*)AssetEditor;
+
+				UEdGraph* EventGraph = FBlueprintEditorUtils::FindEventGraph(BlueprintObject);
+				BlueprintEditor->OpenGraphAndBringToFront(EventGraph);
+				Test->AddLogItem(TEXT("Opened the event graph"));
+
+				PostBeginPlayEventNode = BlueprintEditorPromotionUtils::CreatePostBeginPlayEvent(BlueprintObject, EventGraph);
+				Test->TestNotNull(TEXT("Created EventBeginPlay node"), PostBeginPlayEventNode);
+				
+				const UEdGraphSchema_K2* K2Schema = GetDefault<UEdGraphSchema_K2>();
+				UEdGraphPin* PlayThenPin = PostBeginPlayEventNode->FindPin(K2Schema->PN_Then);
+				
+				DelayNode = BlueprintEditorPromotionUtils::AddDelayNode(BlueprintObject, EventGraph, PlayThenPin);
+				Test->TestNotNull(TEXT("Created Delay node"), DelayNode);
+
+				GEditor->UndoTransaction();
+				Test->TestTrue(TEXT("Undo adding Delay node succeeded"), EventGraph->Nodes.Num() == 0 || EventGraph->Nodes[EventGraph->Nodes.Num() - 1] != DelayNode);
+
+				GEditor->RedoTransaction();
+				Test->TestTrue(TEXT("Redo adding Delay node succeeded"), EventGraph->Nodes.Num() > 0 && EventGraph->Nodes[EventGraph->Nodes.Num() - 1] == DelayNode);
+				
+				// Update Delay node's Duration pin with new default value
+				const FString DelayDurationPinName = TEXT("Duration");
+				const FString NewDurationDefaultValue = TEXT("2.0");
+
+				BlueprintEditorPromotionUtils::SetPinDefaultValue(DelayNode, DelayDurationPinName, NewDurationDefaultValue);				
+				Test->TestEqual(TEXT("Delay node default duration set to 2.0 seconds"), DelayNode->FindPin(DelayDurationPinName)->DefaultValue, NewDurationDefaultValue);
+				
+				BlueprintEditorPromotionUtils::CompileBlueprint(BlueprintObject);
+			}
+
+			return true;
+		}
+
+		/**
+		* Blueprint Test Stage: Custom Variable
+		*    Creates a custom string variable and adds Get/Set nodes for it
+		*/
+		bool Blueprint_CustomVariable()
+		{
+			if (BlueprintObject)
+			{
+				UEdGraph* EventGraph = FBlueprintEditorUtils::FindEventGraph(BlueprintObject);
+
+				Test->AddLogItem(TEXT("Added a string member variable"));
+				BlueprintEditorPromotionUtils::AddStringMemberValue(BlueprintObject, BlueprintEditorPromotionUtils::BlueprintStringVariableName);
+
+				SetNode = BlueprintEditorPromotionUtils::AddGetSetNode(BlueprintObject, EventGraph, BlueprintEditorPromotionUtils::BlueprintStringVariableName.ToString(), false);
+				Test->TestNotNull(TEXT("Added Set node for string variable"), SetNode);
+
+				GEditor->UndoTransaction();
+				Test->TestTrue(TEXT("Undo adding Set node succeeded"), EventGraph->Nodes.Num() == 0 || EventGraph->Nodes[EventGraph->Nodes.Num() - 1] != SetNode);
+
+				GEditor->RedoTransaction();
+				Test->TestTrue(TEXT("Redo adding Set node succeeded"), EventGraph->Nodes.Num() > 0 && EventGraph->Nodes[EventGraph->Nodes.Num() - 1] == SetNode);
+				
+				GetNode = BlueprintEditorPromotionUtils::AddGetSetNode(BlueprintObject, EventGraph, BlueprintEditorPromotionUtils::BlueprintStringVariableName.ToString(), true, 400);
+				Test->TestNotNull(TEXT("Added Get node for string variable"), GetNode);
+				
+				GEditor->UndoTransaction();
+				Test->TestTrue(TEXT("Undo adding Get node succeeded"), EventGraph->Nodes.Num() == 0 || EventGraph->Nodes[EventGraph->Nodes.Num() - 1] != GetNode);
+
+				GEditor->RedoTransaction();
+				Test->TestTrue(TEXT("Redo adding Get node succeeded"), EventGraph->Nodes.Num() > 0 && EventGraph->Nodes[EventGraph->Nodes.Num() - 1] == GetNode);
+
+				EventGraph->RemoveNode(SetNode);
+				Test->TestFalse(TEXT("Set node removed from EventGraph"), EventGraph->Nodes.Contains(SetNode));
+				SetNode = NULL;
+			}
+			return true;
+		}
+
+		/**
+		* Blueprint Test Stage:  Using Variables
+		*   Adds a PrintString and SetStaticMesh then connects all the existing nodes
+		*/
+		bool Blueprint_UsingVariables()
+		{
+			if (BlueprintObject)
+			{
+				const bool bVariableIsHidden = false;
+				FBlueprintEditorUtils::SetBlueprintOnlyEditableFlag(BlueprintObject, BlueprintEditorPromotionUtils::BlueprintStringVariableName, bVariableIsHidden);
+				Test->AddLogItem(TEXT("Exposed the blueprint string variable"));
+
+				UEdGraph* EventGraph = FBlueprintEditorUtils::FindEventGraph(BlueprintObject);
+				PrintNode = BlueprintEditorPromotionUtils::AddPrintStringNode(BlueprintObject, EventGraph);
+				Test->TestNotNull(TEXT("Added Print String node"), PrintNode);
+				
+				const UEdGraphSchema_K2* K2Schema = GetDefault<UEdGraphSchema_K2>();
+
+				//Connect Get to printstring
+				UEdGraphPin* GetVarPin = GetNode->FindPin(BlueprintEditorPromotionUtils::BlueprintStringVariableName.ToString());
+				UEdGraphPin* InStringPin = PrintNode->FindPin(TEXT("InString"));
+				GetVarPin->MakeLinkTo(InStringPin);
+				Test->TestTrue(TEXT("Connected string variable Get node to the Print String node"), GetVarPin->LinkedTo.Contains(InStringPin));
+
+
+				//Connect Delay to PrintString
+				UEdGraphPin* DelayExecPin = DelayNode->FindPin(K2Schema->PN_Then);
+				UEdGraphPin* PrintStringPin = PrintNode->FindPin(K2Schema->PN_Execute);
+				DelayExecPin->MakeLinkTo(PrintStringPin);
+				Test->TestTrue(TEXT("Connected Delay nod to Print String node"), DelayExecPin->LinkedTo.Contains(PrintStringPin));
+
+				const FName MyMeshVarName(TEXT("MyMesh"));
+				GetNode = BlueprintEditorPromotionUtils::AddGetSetNode(BlueprintObject, EventGraph, MyMeshVarName.ToString(), true, 680);
+				Test->TestNotNull(TEXT("Added Get node for MyMesh variable"), GetNode);
+				
+				SetStaticMeshNode = BlueprintEditorPromotionUtils::AddSetStaticMeshNode(BlueprintObject, EventGraph);
+				Test->TestNotNull(TEXT("Added Set Static Mesh node"), SetStaticMeshNode);
+				
+ 				UEdGraphPin* GetExecPin = GetNode->FindPin(TEXT("MyMesh"));
+				UEdGraphPin* SetStaticMeshSelfPin = SetStaticMeshNode->FindPin(K2Schema->PN_Self);
+				GetExecPin->MakeLinkTo(SetStaticMeshSelfPin);
+				Test->TestTrue(TEXT("Connected Get MyMesh node to Set Static Mesh node"), GetExecPin->LinkedTo.Contains(SetStaticMeshSelfPin));
+
+				UEdGraphPin* SetStaticMeshMeshPin = SetStaticMeshNode->FindPin(TEXT("NewMesh"));
+				SetStaticMeshMeshPin->DefaultObject = SecondBlueprintMesh;
+				Test->TestEqual(*FString::Printf(TEXT("Set Static Mesh default mesh updated to %s"), *SecondBlueprintMesh->GetName()), Cast<UStaticMesh>(SetStaticMeshMeshPin->DefaultObject), SecondBlueprintMesh);
+
+				//Connect SetStaticMeshMesh to PrintString
+				UEdGraphPin* PrintStringThenPin = PrintNode->FindPin(K2Schema->PN_Then);
+				UEdGraphPin* SetStaticMeshExecPin = SetStaticMeshNode->FindPin(K2Schema->PN_Execute);
+				PrintStringThenPin->MakeLinkTo(SetStaticMeshExecPin);
+				Test->TestTrue(TEXT("Connected Print String node to Set Static Mesh node"), PrintStringThenPin->LinkedTo.Contains(SetStaticMeshExecPin));
+			}
+			return true;
+		}
+
+		/**
+		* Blueprint Test Stage: Renaming custom event
+		*    Creates, renames, and then removes a custom event node
+		*/
+		bool Blueprint_RenameCustomEvent()
+		{
+			if (BlueprintObject)
+			{
+				UEdGraph* EventGraph = FBlueprintEditorUtils::FindEventGraph(BlueprintObject);
+				CustomEventNode = BlueprintEditorPromotionUtils::CreateCustomEvent(BlueprintObject, EventGraph, TEXT("NewEvent"));
+				Test->TestNotNull(TEXT("Custom event node created"), CustomEventNode);
+				if (CustomEventNode)
+				{
+					//Rename the event
+					const FString NewEventNodeName = TEXT("RenamedEvent");
+					CustomEventNode->Rename(*NewEventNodeName);
+					Test->TestEqual(TEXT("Custom event rename succeeded"), CustomEventNode->GetName(), NewEventNodeName);
+					
+					EventGraph->RemoveNode(CustomEventNode);
+					Test->TestFalse(TEXT("Blueprint EventGraph does not contain removed custom event node"), EventGraph->Nodes.Contains(CustomEventNode));
+					CustomEventNode = NULL;
+				}
+			}
+
+			return true;
+		}
+
+		/**
+		* Blueprint Test Stage: New function
+		*    Creates a new function graph and then hooks up a new AddParticleSystem inside it
+		*/
+		bool Blueprint_NewFunctions()
+		{
+			if (BlueprintObject)
+			{
+				const UEdGraphSchema_K2* K2Schema = GetDefault<UEdGraphSchema_K2>();
+
+				CustomGraph = BlueprintEditorPromotionUtils::CreateNewFunctionGraph(BlueprintObject, TEXT("NewFunction"));
+				Test->TestNotNull(TEXT("Created new function graph"), CustomGraph);
+
+				AddParticleSystemNode = BlueprintEditorPromotionUtils::CreateAddComponentActionNode(BlueprintObject, CustomGraph, LoadedParticleSystem);
+				Test->TestNotNull(TEXT("Created Add Particle System node"), AddParticleSystemNode);
+
+				UEdGraphPin* ExecutePin = AddParticleSystemNode ? AddParticleSystemNode->FindPin(K2Schema->PN_Execute) : NULL;
+
+				//Find the input for the function graph
+				TArray<UK2Node_FunctionEntry*> EntryNodes;
+				CustomGraph->GetNodesOfClass(EntryNodes);
+				UEdGraphNode* EntryNode = EntryNodes.Num() > 0 ? EntryNodes[0] : NULL;
+				if (EntryNode && ExecutePin)
+				{
+					UEdGraphPin* EntryPin = EntryNode->FindPin(K2Schema->PN_Then);
+					EntryPin->MakeLinkTo(ExecutePin);
+					Test->TestTrue(TEXT("Connected Add Particle System node to entry node"), EntryPin->LinkedTo.Contains(ExecutePin));
+				}
+
+				BlueprintEditorPromotionUtils::CompileBlueprint(BlueprintObject);
+			}
+			return true;
+		}
+
+		/**
+		* Blueprint Test Stage: Completing the blueprint
+		*    Adds a CallFunction node to call the custom function created in the previous step.
+		*/
+		bool Blueprint_CompleteBlueprint()
+		{
+			if (BlueprintObject)
+			{
+				const UEdGraphSchema_K2* K2Schema = GetDefault<UEdGraphSchema_K2>();
+				UEdGraph* EventGraph = FBlueprintEditorUtils::FindEventGraph(BlueprintObject);
+				UEdGraphPin* SetStaticMeshThenPin = SetStaticMeshNode->FindPin(K2Schema->PN_Then);
+				CallFunctionNode = BlueprintEditorPromotionUtils::AddCallFunctionGraphNode(BlueprintObject, EventGraph, TEXT("NewFunction"), SetStaticMeshThenPin);
+				Test->TestNotNull(TEXT("Created Call Function node"), CallFunctionNode);
+				if (CallFunctionNode)
+				{
+					Test->TestTrue(TEXT("Connected Set Static Mesh node to Call Function node"), SetStaticMeshThenPin->LinkedTo.Contains(CallFunctionNode->FindPin(K2Schema->PN_Execute)));
+				}
+
+				BlueprintEditorPromotionUtils::CompileBlueprint(BlueprintObject);
+
+				SaveBlueprint();
+			}
+			return true;
+		}
+	};
+}
+
+/**
+* Latent command to run the main build promotion test
+*/
+DEFINE_LATENT_AUTOMATION_COMMAND_ONE_PARAMETER(FRunPromotionTestCommand, TSharedPtr<BlueprintEditorPromotionTestHelper::FBlueprintEditorPromotionTestHelper>, BlueprintEditorTestInfo);
+bool FRunPromotionTestCommand::Update()
+{
+	return BlueprintEditorTestInfo->Update();
+}
+
+/**
+* Automation test that handles the blueprint editor promotion process
+*/
+bool FBlueprintEditorPromotionTest::RunTest(const FString& Parameters)
+{
+	TSharedPtr<BlueprintEditorPromotionTestHelper::FBlueprintEditorPromotionTestHelper> BuildPromotionTest = MakeShareable(new BlueprintEditorPromotionTestHelper::FBlueprintEditorPromotionTestHelper(&ExecutionInfo));
+	BuildPromotionTest->Test = this;
+	ADD_LATENT_AUTOMATION_COMMAND(FRunPromotionTestCommand(BuildPromotionTest));
+	return true;
+}
+
+
+#undef LOCTEXT_NAMESPACE
diff --git a/Engine/Source/Editor/UnrealEd/Private/Tests/AutomationEditorPromotionCommon.cpp b/Engine/Source/Editor/UnrealEd/Private/Tests/AutomationEditorPromotionCommon.cpp
--- a/Engine/Source/Editor/UnrealEd/Private/Tests/AutomationEditorPromotionCommon.cpp
+++ b/Engine/Source/Editor/UnrealEd/Private/Tests/AutomationEditorPromotionCommon.cpp
@@ -3,6 +3,7 @@
 #include "UnrealEd.h"
 #include "AutomationEditorPromotionCommon.h"
 #include "AssetRegistryModule.h"
+#include "LevelEditor.h"
 #include "ScopedTransaction.h"
 
 #define LOCTEXT_NAMESPACE "EditorPromotionTestCommon"
@@ -205,6 +206,17 @@ void FEditorPromotionTestUtilities::SetPropertyByName(UObject* TargetObject, con
 	}
 }
 
+/**
+* Starts a PIE session
+*/
+void FEditorPromotionTestUtilities::StartPIE(bool bSimulateInEditor)
+{
+	FLevelEditorModule& LevelEditorModule = FModuleManager::Get().GetModuleChecked<FLevelEditorModule>(TEXT("LevelEditor"));
+	TSharedPtr<class ILevelViewport> ActiveLevelViewport = LevelEditorModule.GetFirstActiveViewport();
+
+	GUnrealEd->RequestPlaySession(false, ActiveLevelViewport, bSimulateInEditor, NULL, NULL, -1, false);
+}
+
 /**
 * Ends a PIE session
 */
diff --git a/Engine/Source/Editor/UnrealEd/Private/Tests/EditorBuildPromotionTests.cpp b/Engine/Source/Editor/UnrealEd/Private/Tests/EditorBuildPromotionTests.cpp
--- a/Engine/Source/Editor/UnrealEd/Private/Tests/EditorBuildPromotionTests.cpp
+++ b/Engine/Source/Editor/UnrealEd/Private/Tests/EditorBuildPromotionTests.cpp
@@ -32,13 +32,10 @@
 //Blueprints
 #include "BlueprintGraphDefinitions.h"
 #include "BlueprintEditor.h"
-#include "BlueprintEditorModes.h"
-#include "CompilerResultsLog.h"
 #include "EdGraphUtilities.h"
 #include "KismetDebugUtilities.h"
 #include "Engine/Breakpoint.h"
 #include "Engine/LevelScriptBlueprint.h"
-#include "Kismet/KismetStringLibrary.h"
 #include "Kismet/KismetSystemLibrary.h"
 #include "BlueprintEditorUtils.h"
 
@@ -74,6 +71,11 @@
 
 DEFINE_LOG_CATEGORY_STATIC(LogEditorBuildPromotionTests, Log, All);
 
+IMPLEMENT_SIMPLE_AUTOMATION_TEST(FBuildPromotionInitialCleanupTest, "System.Promotion.Editor Promotion Pass.Step 1 Main Editor Test.Cleanup old files", EAutomationTestFlags::EditorContext | EAutomationTestFlags::EngineFilter);
+IMPLEMENT_SIMPLE_AUTOMATION_TEST(FEditorPromotionTest, "System.Promotion.Editor Promotion Pass.Step 1 Main Editor Test.General Editor Test", EAutomationTestFlags::EditorContext | EAutomationTestFlags::EngineFilter);
+IMPLEMENT_SIMPLE_AUTOMATION_TEST(FBuildPromotionPIETest, "System.Promotion.Editor Promotion Pass.Step 2 Run Map After Re-launch.Run Map", EAutomationTestFlags::EditorContext | EAutomationTestFlags::EngineFilter);
+IMPLEMENT_SIMPLE_AUTOMATION_TEST(FBuildPromotionCleanupTest, "System.Promotion.Editor Promotion Pass.Step 3 Test Cleanup.Cleanup", EAutomationTestFlags::EditorContext | EAutomationTestFlags::EngineFilter);
+
 /**
 * Helper functions used by the build promotion automation test
 */
@@ -82,7 +84,7 @@ namespace EditorBuildPromotionTestUtils
 	/**
 	* Constants
 	*/
-	static const FString BlueprintNameString = TEXT("BuildPromotionBlueprint");
+	static const FString BlueprintNameString = TEXT("EditorPromotionBlueprint");
 	static const FName BlueprintStringVariableName(TEXT("MyStringVariable"));
 
 	/**
@@ -90,7 +92,7 @@ namespace EditorBuildPromotionTestUtils
 	*/
 	static FString GetFullPath()
 	{
-		return FPackageName::FilenameToLongPackageName(FPaths::GameContentDir() + TEXT("BuildPromotionTest"));
+		return FPackageName::FilenameToLongPackageName(FPaths::GameContentDir() + TEXT("EditorPromotionTest"));
 	}
 
 	/**
@@ -162,54 +164,6 @@ namespace EditorBuildPromotionTestUtils
 		double Duration;
 	};
 
-	///**
-	//* Exports the current editor keybindings
-	//*
-	//* @param TargetFilename - The name of the file to export to
-	//*/
-	//static void ExportKeybindings(const FString& TargetFilename)
-	//{
-	//	FInputBindingManager::Get().SaveInputBindings();
-	//	GConfig->Flush(false, GEditorKeyBindingsIni);
-	//	IFileManager::Get().Copy(*TargetFilename, *GEditorKeyBindingsIni);
-	//}
-
-	///**
-	//* Imports new editor keybindings  (Unused)
-	//*
-	//* @param TargetFilename - The name of the file to import from
-	//*/
-	//static void ImportKeybindings(const FString& TargetFilename)
-	//{
-	//	GConfig->Flush(true, GEditorKeyBindingsIni);
-	//	IFileManager::Get().Copy(*GEditorKeyBindingsIni, *TargetFilename);
-	//	GConfig->LoadFile(GEditorKeyBindingsIni);
-	//}
-
-	///**
-	//* Exports the current editor settings
-	//*
-	//* @param TargetFilename - The name of the file to export to
-	//*/
-	//static void ExportEditorSettings(const FString& TargetFilename)
-	//{
-	//	FInputBindingManager::Get().SaveInputBindings();
-	//	GConfig->Flush(false, GEditorPerProjectIni);
-	//	IFileManager::Get().Copy(*TargetFilename, *GEditorPerProjectIni);
-	//}
-
-	///**
-	//* Imports new editor settings (unused)
-	//*
-	//* @param TargetFilename - The name of the file to export to
-	//*/
-	//static void ImportEditorSettings(const FString& TargetFilename)
-	//{
-	//	GConfig->Flush(true, GEditorPerProjectIni);
-	//	IFileManager::Get().Copy(*GEditorPerProjectIni, *TargetFilename);
-	//	GConfig->LoadFile(GEditorPerProjectIni);
-	//}
-
 	/**
 	* Sends the MaterialEditor->Apply UI command
 	*/
@@ -223,20 +177,6 @@ namespace EditorBuildPromotionTestUtils
 		FEditorPromotionTestUtilities::SendCommandToCurrentEditor(CurrentApplyChord, FocusWidgetType);
 	}
 
-	/**
-	* Sends the AssetEditor->SaveAsset UI command
-	*/
-	static void SendBlueprintResetViewCommand()
-	{
-		const FString Context = TEXT("BlueprintEditor");
-		const FString Command = TEXT("ResetCamera");
-
-		FInputChord CurrentSaveChord = FEditorPromotionTestUtilities::GetOrSetUICommand(Context, Command);
-
-		const FName FocusWidgetType(TEXT("SSCSEditorViewport"));
-		FEditorPromotionTestUtilities::SendCommandToCurrentEditor(CurrentSaveChord, FocusWidgetType);
-	}
-
 	/**
 	* Compiles the blueprint
 	*
@@ -268,98 +208,6 @@ namespace EditorBuildPromotionTestUtils
 		}
 	}
 
-	/**
-	* Creates a blueprint component based off the supplied asset
-	*
-	* @param InBlueprint - The blueprint to modify
-	* @param InAsset - The asset to use for the component
-	*/
-	static USCS_Node* CreateBlueprintComponent(UBlueprint* InBlueprint, UObject* InAsset)
-	{
-		IAssetEditorInstance* OpenEditor = FAssetEditorManager::Get().FindEditorForAsset(InBlueprint, true);
-		FBlueprintEditor* CurrentBlueprintEditor = (FBlueprintEditor*)OpenEditor;
-		TSubclassOf<UActorComponent> ComponentClass = FComponentAssetBrokerage::GetPrimaryComponentForAsset(InAsset->GetClass());
-		
-		USCS_Node* NewNode = InBlueprint->SimpleConstructionScript->CreateNode(ComponentClass);
-		// Assign the asset to the template
-		FComponentAssetBrokerage::AssignAssetToComponent(NewNode->ComponentTemplate, InAsset);
-
-		// Add node to the SCS
-		TArray<USCS_Node*> AllNodes = InBlueprint->SimpleConstructionScript->GetAllNodes();
-		USCS_Node* RootNode = AllNodes.Num() > 0 ? AllNodes[0] : NULL;
-		if (!RootNode || (RootNode == InBlueprint->SimpleConstructionScript->GetDefaultSceneRootNode()))
-		{
-			//New Root
-			InBlueprint->SimpleConstructionScript->AddNode(NewNode);
-		}
-		else
-		{
-			//Add as a child
-			RootNode->AddChildNode(NewNode);
-		}
-
-		// Recompile skeleton because of the new component we added
-		FKismetEditorUtilities::GenerateBlueprintSkeleton(InBlueprint, true);
-
-		CurrentBlueprintEditor->UpdateSCSPreview(true);
-
-		return NewNode;
-	}
-
-	/**
-	* Sets a new component as the root
-	*
-	* @param InBlueprint - The blueprint to modify
-	* @param NewRoot - The new root
-	*/
-	static void SetComponentAsRoot(UBlueprint* InBlueprint, USCS_Node* NewRoot)
-	{
-		// @FIXME: Current usages doesn't guarantee NewRoot is valid!!! Check first!
-		//Get all the construction script nodes
-		TArray<USCS_Node*> AllNodes = InBlueprint->SimpleConstructionScript->GetAllNodes();
-
-		USCS_Node* OldRootNode = AllNodes[0];
-		USCS_Node* OldParent = NULL;
-
-		//Find old parent
-		for (int32 NodeIndex = 0; NodeIndex < AllNodes.Num(); ++NodeIndex)
-		{
-			if (AllNodes[NodeIndex]->ChildNodes.Contains(NewRoot))
-			{
-				OldParent = AllNodes[NodeIndex];
-				break;
-			}
-		}
-
-		//Remove the new root from its old parent and 
-		OldParent->ChildNodes.Remove(NewRoot);
-		NewRoot->Modify();
-		NewRoot->AttachToName = NAME_None;
-
-		//Remove the old root, add the new root, and attach the old root as a child
-		InBlueprint->SimpleConstructionScript->RemoveNode(OldRootNode);
-		InBlueprint->SimpleConstructionScript->AddNode(NewRoot);
-		NewRoot->AddChildNode(OldRootNode);
-	}
-
-	/**
-	* Removes a blueprint component from the simple construction script
-	*
-	* @param InBlueprint - The blueprint to modify
-	* @param InNode - The node of the component to remove
-	*/
-	static void RemoveBlueprintComponent(UBlueprint* InBlueprint, USCS_Node* InNode)
-	{
-		if (InNode != NULL)
-		{
-			// Remove node from SCS tree
-			InNode->GetSCS()->RemoveNodeAndPromoteChildren(InNode);
-
-			// Clear the delegate
-			InNode->SetOnNameChanged(FSCSNodeNameChanged());
-		}
-	}
-
 	/**
 	* Creates a new graph node from a given template
 	*
@@ -377,127 +225,6 @@ namespace EditorBuildPromotionTestUtils
 		return Action->PerformAction(InGraph, ConnectPin, GraphLocation, false);
 	}
 
-	/**
-	* Creates an AddComponent action to the blueprint graph
-	*
-	* @param InBlueprint - The blueprint to modify
-	* @param InGraph - The blueprint graph to use
-	* @param InAsset - The asset to use
-	*/
-	static UEdGraphNode* CreateAddComponentActionNode(UBlueprint* InBlueprint, UEdGraph* InGraph, UObject* InAsset)
-	{
-		UEdGraph* TempOuter = NewObject<UEdGraph>((UObject*)InBlueprint);
-		TempOuter->SetFlags(RF_Transient);
-
-		const FScopedTransaction PropertyChanged(LOCTEXT("AddedGraphNode", "Added a graph node"));
-		InGraph->Modify();
-
-		// Make an add component node
-		UK2Node_CallFunction* CallFuncNode = NewObject<UK2Node_AddComponent>(TempOuter);
-		UFunction* AddComponentFn = FindFieldChecked<UFunction>(AActor::StaticClass(), UK2Node_AddComponent::GetAddComponentFunctionName());
-		CallFuncNode->FunctionReference.SetFromField<UFunction>(AddComponentFn, FBlueprintEditorUtils::IsActorBased(InBlueprint));
-
-		UEdGraphNode* NewNode = CreateNewGraphNodeFromTemplate(CallFuncNode, InGraph, FVector2D(200, 0));
-
-		TSubclassOf<UActorComponent> ComponentClass = InAsset ? FComponentAssetBrokerage::GetPrimaryComponentForAsset(InAsset->GetClass()) : NULL;
-		if ((NewNode != NULL) && (InBlueprint != NULL))
-		{
-			UK2Node_AddComponent* AddCompNode = CastChecked<UK2Node_AddComponent>(NewNode);
-
-			ensure(NULL != Cast<UBlueprintGeneratedClass>(InBlueprint->GeneratedClass));
-			// Then create a new template object, and add to array in
-			UActorComponent* NewTemplate = NewObject<UActorComponent>(InBlueprint->GeneratedClass, ComponentClass, NAME_None, RF_ArchetypeObject|RF_Public);
-			InBlueprint->ComponentTemplates.Add(NewTemplate);
-
-			// Set the name of the template as the default for the TemplateName param
-			UEdGraphPin* TemplateNamePin = AddCompNode->GetTemplateNamePinChecked();
-			if (TemplateNamePin)
-			{
-				TemplateNamePin->DefaultValue = NewTemplate->GetName();
-			}
-
-			// Set the return type to be the type of the template
-			UEdGraphPin* ReturnPin = AddCompNode->GetReturnValuePin();
-			if (ReturnPin)
-			{
-				ReturnPin->PinType.PinSubCategoryObject = *ComponentClass;
-			}
-
-			// Set the asset
-			if (InAsset != NULL)
-			{
-				FComponentAssetBrokerage::AssignAssetToComponent(NewTemplate, InAsset);
-			}
-
-			AddCompNode->ReconstructNode();
-		}
-
-		FBlueprintEditorUtils::MarkBlueprintAsModified(InBlueprint);
-		return NewNode;
-	}
-
-	/**
-	* Creates a SetStaticMesh action in the blueprint graph
-	*
-	* @param InBlueprint - The blueprint to modify
-	* @param InGraph - The blueprint graph to use
-	*/
-	static UEdGraphNode* AddSetStaticMeshNode(UBlueprint* InBlueprint, UEdGraph* InGraph)
-	{
-		UEdGraph* TempOuter = NewObject<UEdGraph>((UObject*)InBlueprint);
-		TempOuter->SetFlags(RF_Transient);
-
-		// Make a call function template
-		UK2Node_CallFunction* CallFuncNode = NewObject<UK2Node_CallFunction>(TempOuter);
-		static FName PrintStringFunctionName(TEXT("SetStaticMesh"));
-		UFunction* DelayFn = FindFieldChecked<UFunction>(UStaticMeshComponent::StaticClass(), PrintStringFunctionName);
-		CallFuncNode->FunctionReference.SetFromField<UFunction>(DelayFn, false);
-
-		return CreateNewGraphNodeFromTemplate(CallFuncNode, InGraph, FVector2D(850, 0));
-	}
-
-	/**
-	* Connects two nodes using the supplied pin names
-	*
-	* @param NodeA - The first node to connect
-	* @param PinAName - The name of the pin on the first node
-	* @param NodeB - The second node to connect
-	* @param PinBName - The name of the pin on the second node
-	*/
-	static void ConnectGraphNodes(UEdGraphNode* NodeA, const FString& PinAName, UEdGraphNode* NodeB, const FString& PinBName)
-	{
-		const FScopedTransaction PropertyChanged(LOCTEXT("ConnectedNode", "Connected graph nodes"));
-		NodeA->GetGraph()->Modify();
-
-		UEdGraphPin* PinA = NodeA->FindPin(PinAName);
-		UEdGraphPin* PinB = NodeB->FindPin(PinBName);
-
-		if (PinA && PinB)
-		{
-			PinA->MakeLinkTo(PinB);
-		}
-		else
-		{
-			UE_LOG(LogEditorBuildPromotionTests, Error, TEXT("Could not connect pins %s and %s "), *PinAName, *PinBName);
-		}
-	}
-
-	/**
-	* Promotes a pin to a variable
-	*
-	* @param InBlueprint - The blueprint to modify
-	* @param Node - The node that owns the pin
-	* @param PinName - The name of the pin to promote
-	*/
-	static void PromotePinToVariable(UBlueprint* InBlueprint, UEdGraphNode* Node, const FString& PinName)
-	{
-		IAssetEditorInstance* OpenEditor = FAssetEditorManager::Get().FindEditorForAsset(InBlueprint, true);
-		FBlueprintEditor* CurrentBlueprintEditor = (FBlueprintEditor*)OpenEditor;
-
-		UEdGraphPin* PinToPromote = Node->FindPin(PinName);
-		CurrentBlueprintEditor->DoPromoteToVariable(InBlueprint, PinToPromote, true);
-	}
-
 	/**
 	* Creates a ReceiveBeginPlay event node
 	*
@@ -558,29 +285,6 @@ namespace EditorBuildPromotionTestUtils
 		return CallFuncNode;
 	}
 
-	/**
-	* Creates a delay node
-	*
-	* @param InBlueprint - The blueprint to modify
-	* @param InGraph - The graph to use for the new node
-	* @param ConnectPin - The pin to connect the new node to
-	*/
-	static UEdGraphNode* AddDelayNode(UBlueprint* InBlueprint, UEdGraph* InGraph, UEdGraphPin* ConnectPin = NULL)
-	{
-		UEdGraph* TempOuter = NewObject<UEdGraph>((UObject*)InBlueprint);
-		TempOuter->SetFlags(RF_Transient);
-
-		const FScopedTransaction PropertyChanged(LOCTEXT("AddedGraphNode", "Added a graph node"));
-		InGraph->Modify();
-
-		// Make a call function template
-		static FName DelayFunctionName(TEXT("Delay"));
-		UK2Node* CallFuncNode = CreateKismetFunctionTemplate(TempOuter, DelayFunctionName);
-
-		//Create the node
-		return CreateNewGraphNodeFromTemplate(CallFuncNode, InGraph, FVector2D(400, 0), ConnectPin);
-	}
-
 	/**
 	* Creates a PrintString node
 	*
@@ -600,110 +304,6 @@ namespace EditorBuildPromotionTestUtils
 		return CreateNewGraphNodeFromTemplate(CallFuncNode, InGraph, FVector2D(680, 0), ConnectPin);
 	}
 
-	/**
-	* Creates a call function node
-	*
-	* @param InBlueprint - The blueprint to modify
-	* @param InGraph - The graph to use for the new node
-	* @param FunctionName - The name of the function to call
-	* @param ConnectPin - The pin to connect the new node to
-	*/
-	static UEdGraphNode* AddCallFunctionGraphNode(UBlueprint* InBlueprint, UEdGraph* InGraph, const FName& FunctionName, UEdGraphPin* ConnectPin = NULL)
-	{
-		UEdGraph* TempOuter = NewObject<UEdGraph>((UObject*)InBlueprint);
-		TempOuter->SetFlags(RF_Transient);
-
-		// Make a call function template
-		UK2Node_CallFunction* CallFuncNode = NewObject<UK2Node_CallFunction>(TempOuter);
-		CallFuncNode->FunctionReference.SetSelfMember(FunctionName);
-
-		return CreateNewGraphNodeFromTemplate(CallFuncNode, InGraph, FVector2D(1200, 0), ConnectPin);
-	}
-
-	/**
-	* Creates Get or Set node
-	*
-	* @param InBlueprint - The blueprint to modify
-	* @param InGraph - The graph to use for the new node
-	* @param VarName - The name of the variable to use
-	* @param bGet - If true, create a Get node.  If false, create a Set node.
-	* @param XOffset - How far to offset the node in the graph
-	*/
-	static UEdGraphNode* AddGetSetNode(UBlueprint* InBlueprint, UEdGraph* InGraph, const FString& VarName, bool bGet, float XOffset = 0.f)
-	{
-		const FScopedTransaction PropertyChanged(LOCTEXT("AddedGraphNode", "Added a graph node"));
-		InGraph->Modify();
-
-		FEdGraphSchemaAction_K2NewNode NodeInfo;
-		// Create get or set node, depending on whether we clicked on an input or output pin
-		UK2Node_Variable* TemplateNode = NULL;
-		if (bGet)
-		{
-			TemplateNode = NewObject<UK2Node_VariableGet>();
-		}
-		else
-		{
-			TemplateNode = NewObject<UK2Node_VariableSet>();
-		}
-
-		TemplateNode->VariableReference.SetSelfMember(FName(*VarName));
-		NodeInfo.NodeTemplate = TemplateNode;
-
-		return NodeInfo.PerformAction(InGraph, NULL, FVector2D(XOffset,130), true);
-	}
-
-	/**
-	* Sets the default value for a pin
-	*
-	* @param Node - The node that owns the pin to set
-	* @param PinName - The name of the pin
-	* @param PinValue - The new default value
-	*/
-	static void SetPinDefaultValue(UEdGraphNode* Node, const FString& PinName, const FString& PinValue)
-	{
-		UEdGraphPin* Pin = Node->FindPin(PinName);
-		Pin->DefaultValue = PinValue;
-	}
-
-	/**
-	* Sets the default object for a pin
-	*
-	* @param Node - The node that owns the pin to set
-	* @param PinName - The name of the pin
-	* @param PinObject - The new default object
-	*/
-	static void SetPinDefaultObject(UEdGraphNode* Node, const FString& PinName, UObject* PinObject)
-	{
-		UEdGraphPin* Pin = Node->FindPin(PinName);
-		Pin->DefaultObject = PinObject;
-	}
-
-	/**
-	* Adds a string member variable to a blueprint
-	*
-	* @param InBlueprint - The blueprint to modify
-	* @param VariableName - The name of the new string variable
-	*/
-	static void AddStringMemberValue(UBlueprint* InBlueprint, const FName& VariableName)
-	{
-		const UEdGraphSchema_K2* K2Schema = GetDefault<UEdGraphSchema_K2>();
-		FEdGraphPinType StringPinType(K2Schema->PC_String, TEXT(""), NULL, false, false);
-		FBlueprintEditorUtils::AddMemberVariable(InBlueprint, VariableName, StringPinType);
-	}
-
-	/**
-	* Creates a new function graph
-	*
-	* @param InBlueprint - The blueprint to modify
-	* @param FunctionName - The function name to use for the new graph
-	*/
-	static UEdGraph* CreateNewFunctionGraph(UBlueprint* InBlueprint, const FName& FunctionName)
-	{
-		UEdGraph* NewGraph = FBlueprintEditorUtils::CreateNewGraph(InBlueprint, FunctionName, UEdGraph::StaticClass(), UEdGraphSchema_K2::StaticClass());
-		FBlueprintEditorUtils::AddFunctionGraph<UClass>(InBlueprint, NewGraph, /*bIsUserCreated=*/ true, NULL);
-		return NewGraph;
-	}
-
 	/**
 	* Starts a lighting build
 	*/
@@ -1328,6 +928,9 @@ namespace BuildPromotionTestHelper
 	*/
 	struct FBuildPromotionTest
 	{
+		/** Pointer to running automation test instance */
+		FEditorPromotionTest* Test;
+
 		/** Pointer to the execution info of this test */
 		FAutomationTestExecutionInfo* TestExecutionInfo;
 
@@ -1377,27 +980,11 @@ namespace BuildPromotionTestHelper
 		/** Helper for opening, modifying, and placing assets */
 		FOpenAssetHelper* OpenAssetHelper;
 
-		/** Meshes to use for the blueprint */
-		UStaticMesh* FirstBlueprintMesh;
-		UStaticMesh* SecondBlueprintMesh;
-
 		/** Objects created by the Blueprint stages */
 		UBlueprint* BlueprintObject;
 		UPackage* BlueprintPackage;
-		UEdGraph* CustomGraph;
-		USCS_Node* MeshNode;
-		USCS_Node* OtherMeshNode;
-		USCS_Node* PSNode;
-		UEdGraphNode* AddMeshNode;
 		UEdGraphNode* PostBeginPlayEventNode;
-		UEdGraphNode* DelayNode;
-		UEdGraphNode* SetNode;
-		UEdGraphNode* GetNode;
 		UEdGraphNode* PrintNode;
-		UEdGraphNode* SetStaticMeshNode;
-		UEdGraphNode* CustomEventNode;
-		UEdGraphNode* AddParticleSystemNode;
-		UEdGraphNode* CallFunctionNode;
 		AActor* PlacedBlueprint;
 
 		/** Source control async helper */
@@ -1462,25 +1049,9 @@ namespace BuildPromotionTestHelper
 			ADD_TEST_STAGE(EndSection, TEXT("Content Browser"));
 
 			// 6) Blueprints
-			ADD_TEST_STAGE(Blueprint_CreateNewBlueprint_Part1,	TEXT("Create a new Blueprint"));
-			ADD_TEST_STAGE(Blueprint_CreateNewBlueprint_Part2,	TEXT("Create a new Blueprint"));
-			ADD_TEST_STAGE(Blueprint_DataOnlyBlueprint_Part1,	TEXT("Data-only Blueprint"));
-			ADD_TEST_STAGE(Blueprint_DataOnlyBlueprint_Part2,	TEXT("Data-only Blueprint"));
-			ADD_TEST_STAGE(Blueprint_DataOnlyBlueprint_Part3,	TEXT("Data-only Blueprint"));
-			ADD_TEST_STAGE(Blueprint_ComponentsMode_Part1,		TEXT("Components Mode"));
-			ADD_TEST_STAGE(Blueprint_ComponentsMode_Part2,		TEXT("Components Mode"));
-			ADD_TEST_STAGE(Blueprint_ConstructionScript,		TEXT("Construction Script"));
-			ADD_TEST_STAGE(Blueprint_PromoteVariable_Part1,		TEXT("Variable from Component Mode"));
-			ADD_TEST_STAGE(Blueprint_PromoteVariable_Part2,		TEXT("Variable from Component Mode"));
-			ADD_TEST_STAGE(Blueprint_PromoteVariable_Part3,		TEXT("Variable from Component Mode"));
-			ADD_TEST_STAGE(Blueprint_EventGraph,				TEXT("Event Graph"));
-			ADD_TEST_STAGE(Blueprint_CustomVariable,			TEXT("Custom Variables"));
-			ADD_TEST_STAGE(Blueprint_UsingVariables,			TEXT("Using Variables"));
-			ADD_TEST_STAGE(Blueprint_RenameCustomEvent,			TEXT("Renaming Custom Event"));
-			ADD_TEST_STAGE(Blueprint_NewFunctions,				TEXT("New Function"));
-			ADD_TEST_STAGE(Blueprint_CompleteBlueprint,			TEXT("Completing the Blueprint"));
+			ADD_TEST_STAGE(Blueprint_Setup,						TEXT("Blueprint setup"));
 			ADD_TEST_STAGE(Blueprint_Placement_Part1,			TEXT("Blueprint Placement"));
-			ADD_TEST_STAGE(Blueprint_Placement_Part2,			TEXT("Blueprint Placement"));
+			//ADD_TEST_STAGE(Blueprint_Placement_Part2,			TEXT("Blueprint Placement"));
 			ADD_TEST_STAGE(Blueprint_Placement_Part3,			TEXT("Blueprint Placement"));
 			ADD_TEST_STAGE(Blueprint_Placement_Part4,			TEXT("Blueprint Placement"));
 			ADD_TEST_STAGE(Blueprint_SetBreakpoint_Part1,		TEXT("Set Breakpoints"));
@@ -2303,809 +1874,120 @@ namespace BuildPromotionTestHelper
 		* Saves the blueprint stored in BlueprintObject
 		*/
 		void SaveBlueprint()
-		{  // @FIXME: Saving BP currently fails.
+		{
 			if (BlueprintObject && BlueprintPackage)
 			{
 				BlueprintPackage->SetDirtyFlag(true);
 				BlueprintPackage->FullyLoad();
 				const FString PackagePath = FEditorPromotionTestUtilities::GetGamePath() + TEXT("/") + EditorBuildPromotionTestUtils::BlueprintNameString;
-				if (UPackage::SavePackage(BlueprintPackage, NULL, RF_Standalone, *FPackageName::LongPackageNameToFilename(PackagePath, FPackageName::GetAssetPackageExtension()), GLog, nullptr, false, true, SAVE_None))
-				{
-					UE_LOG(LogEditorBuildPromotionTests, Display, TEXT("Saved blueprint (%s)"), *BlueprintObject->GetName());
-				}
-				else
-				{
-					UE_LOG(LogEditorBuildPromotionTests, Error, TEXT("Unable to save blueprint (%s)"), *BlueprintObject->GetName());
-				}
+				bool bHasPackageSaved = UPackage::SavePackage(BlueprintPackage, NULL, RF_Standalone, *FPackageName::LongPackageNameToFilename(PackagePath, FPackageName::GetAssetPackageExtension()), GLog, nullptr, false, true, SAVE_None);
+				Test->TestTrue(FString::Printf(TEXT("Saved blueprint (%s)"), *BlueprintObject->GetName()), bHasPackageSaved);
 			}
 		}
 
-		/**
-		* Blueprint Test Stage: Create a new Blueprint (Part 1)
-		*    Creates a new actor based blueprint and opens the editor
-		*/
-		bool Blueprint_CreateNewBlueprint_Part1()
+		bool Blueprint_Setup()
 		{
-			//Make sure we have the required assets
-			UAutomationTestSettings const* AutomationTestSettings = GetDefault<UAutomationTestSettings>();
-			check(AutomationTestSettings);
-			FAssetData AssetData;
+			const FString PackageName = FEditorPromotionTestUtilities::GetGamePath() + TEXT("/") + EditorBuildPromotionTestUtils::BlueprintNameString;
+			
+			// Create blueprint asset
+			UBlueprintFactory* Factory = NewObject<UBlueprintFactory>();
+			Factory->ParentClass = AActor::StaticClass();
+			BlueprintPackage = CreatePackage(NULL, *PackageName);
+			EObjectFlags Flags = RF_Public | RF_Standalone;
 
-			const FString FirstMeshPath = AutomationTestSettings->BuildPromotionTest.BlueprintSettings.FirstMeshPath.FilePath;
-			if (FirstMeshPath.Len() > 0)
+			// Check that conflicting asset doesn't already exist
+			UObject* ExistingBlueprint = FindObject<UBlueprint>(BlueprintPackage, *EditorBuildPromotionTestUtils::BlueprintNameString);
+			Test->TestNull(TEXT("Blueprint asset does not already exist (delete blueprint and restart editor)"), ExistingBlueprint);
+			if (ExistingBlueprint)
 			{
-				AssetData = FEditorAutomationTestUtilities::GetAssetDataFromPackagePath(FirstMeshPath);
-				FirstBlueprintMesh = Cast<UStaticMesh>(AssetData.GetAsset());
+				return true;
 			}
 
-			const FString SecondMeshPath = AutomationTestSettings->BuildPromotionTest.BlueprintSettings.SecondMeshPath.FilePath;
-			if (SecondMeshPath.Len() > 0)
+			//Save blueprint object for reuse in later stages
+			BlueprintObject = Cast<UBlueprint>(Factory->FactoryCreateNew(UBlueprint::StaticClass(), BlueprintPackage, FName(*EditorBuildPromotionTestUtils::BlueprintNameString), Flags, NULL, GWarn));
+			Test->TestNotNull(TEXT("Blueprint test asset created"), BlueprintObject);
+			if (!BlueprintObject)
 			{
-				AssetData = FEditorAutomationTestUtilities::GetAssetDataFromPackagePath(SecondMeshPath);
-				SecondBlueprintMesh = Cast<UStaticMesh>(AssetData.GetAsset());
+				return true;
 			}
+			// Add asset to registry
+			FAssetRegistryModule::AssetCreated(BlueprintObject);
+			BlueprintPackage->MarkPackageDirty();
 
-			const FString ParticleSystemPath = AutomationTestSettings->ParticleEditorPromotionTest.DefaultParticleAsset.FilePath;
-			if (ParticleSystemPath.Len() > 0)
+			//Add BeginPlay event to graph
+			UEdGraph* EventGraph = FBlueprintEditorUtils::FindEventGraph(BlueprintObject);
+			PostBeginPlayEventNode = EditorBuildPromotionTestUtils::CreatePostBeginPlayEvent(BlueprintObject, EventGraph);
+			Test->TestNotNull(TEXT("Event Being Play node added"), PostBeginPlayEventNode);
+			if (!PostBeginPlayEventNode)
 			{
-				AssetData = FEditorAutomationTestUtilities::GetAssetDataFromPackagePath(ParticleSystemPath);
-				LoadedParticleSystem = Cast<UParticleSystem>(AssetData.GetAsset());
+				return true;
 			}
 
-			if (FirstBlueprintMesh && SecondBlueprintMesh && LoadedParticleSystem)
-			{
-				UBlueprintFactory* Factory = NewObject<UBlueprintFactory>();
-				Factory->ParentClass = AActor::StaticClass();
-
-				const FString PackageName = FEditorPromotionTestUtilities::GetGamePath() + TEXT("/") + EditorBuildPromotionTestUtils::BlueprintNameString;
-				BlueprintPackage = CreatePackage(NULL, *PackageName);
-				EObjectFlags Flags = RF_Public | RF_Standalone;
-
-				UObject* ExistingBlueprint = FindObject<UBlueprint>(BlueprintPackage, *EditorBuildPromotionTestUtils::BlueprintNameString);
-				if (ExistingBlueprint)
-				{
-					UE_LOG(LogEditorBuildPromotionTests, Error, TEXT("Found existing blueprint.  Please use the cleanup test between runs."));
-					return true;
-				}
-
-				BlueprintObject = Cast<UBlueprint>(Factory->FactoryCreateNew(UBlueprint::StaticClass(), BlueprintPackage, FName(*EditorBuildPromotionTestUtils::BlueprintNameString), Flags, NULL, GWarn));
-
-				if (BlueprintObject)
-				{
-					UE_LOG(LogEditorBuildPromotionTests, Display, TEXT("Created a new Actor based blueprint (%s)"), *BlueprintObject->GetName());
-
-					// Notify the asset registry
-					FAssetRegistryModule::AssetCreated(BlueprintObject);
-
-					// Mark the package dirty...
-					BlueprintPackage->MarkPackageDirty();
-
-					UE_LOG(LogEditorBuildPromotionTests, Display, TEXT("Opening the blueprint editor for the first time"));
-					FAssetEditorManager::Get().OpenEditorForAsset(BlueprintObject);
-				}
-				else
-				{
-					UE_LOG(LogEditorBuildPromotionTests, Error, TEXT("Failed to create blueprint!"));
-				}
-			}
-			else
-			{
-				SkippedTests.Add(TEXT("All Blueprint tests. (Missing a required mesh or particle system)"));
-				if (FirstMeshPath.IsEmpty() || SecondMeshPath.IsEmpty())
-				{
-					UE_LOG(LogEditorBuildPromotionTests, Log, TEXT("SKIPPING BLUEPRINT TESTS.  FirstMeshPath or SecondMeshPath not configured in AutomationTestSettings."));
-				}
-				else
-				{
-					UE_LOG(LogEditorBuildPromotionTests, Warning, TEXT("SKIPPING BLUEPRINT TESTS.  Invalid FirstMeshPath or SecondMeshPath in AutomationTestSettings, or particle system was not created."));
-				}
-			}
+			const UEdGraphSchema_K2* K2Schema = GetDefault<UEdGraphSchema_K2>();
+			
+			// Add string member variable
+			FEdGraphPinType StringPinType(K2Schema->PC_String, TEXT(""), NULL, false, false);
+			FBlueprintEditorUtils::AddMemberVariable(BlueprintObject, EditorBuildPromotionTestUtils::BlueprintStringVariableName, StringPinType);
+
+			// Add print string node connected to the BeginPlay node; save it for use later
+			UEdGraphPin* PlayThenPin = PostBeginPlayEventNode->FindPin(K2Schema->PN_Then);
+			PrintNode = EditorBuildPromotionTestUtils::AddPrintStringNode(BlueprintObject, EventGraph, PlayThenPin);
+			Test->TestNotNull(TEXT("Print String node added"), PrintNode);
 			
 			return true;
+
 		}
 
 		/**
-		* Blueprint Test Stage: Create a new Blueprint (Part 2)
-		*    Checks that the blueprint editor opened in the correct mode
+		* Blueprint Test Stage: Blueprint placement (Part 1)
+		*    Places the blueprint in the level then starts PIE
 		*/
-		bool Blueprint_CreateNewBlueprint_Part2()
+		bool Blueprint_Placement_Part1()
 		{
 			if (BlueprintObject)
 			{
-				IAssetEditorInstance* AssetEditor = FAssetEditorManager::Get().FindEditorForAsset(BlueprintObject, true);
-				IBlueprintEditor* BlueprintEditor = (IBlueprintEditor*)AssetEditor;
-				if (BlueprintEditor->GetCurrentMode() != FBlueprintEditorApplicationModes::BlueprintDefaultsMode)
-				{
-					UE_LOG(LogEditorBuildPromotionTests, Display, TEXT("Opened the correct initial editor"));
-				}
-				else
+				PlacedBlueprint = FActorFactoryAssetProxy::AddActorForAsset(BlueprintObject);
+				Test->TestNotNull(TEXT("Blueprint instance placed in world"), PlacedBlueprint);
+				if (PlacedBlueprint)
 				{
-					UE_LOG(LogEditorBuildPromotionTests, Error, TEXT("Didn't open the correct editor"));
+					//Set the text
+					const FString NewVariableText = TEXT("Print String works!");
+					FEditorPromotionTestUtilities::SetPropertyByName(PlacedBlueprint, EditorBuildPromotionTestUtils::BlueprintStringVariableName.ToString(), NewVariableText);
+					// TODO: Validate that property updated correctly. Needs to be done in latent command?
+					Test->AddLogItem(TEXT("Updated string variable value"));
 				}
+
+				GEditor->SelectNone(false, true);
+
+				//Start PIE
+				EditorBuildPromotionTestUtils::StartPIE(true);
+
+				//Make sure the timer is reset
+				DelayHelper.Reset();
 			}
-			
 			return true;
 		}
 
+		///**
+		//* Blueprint Test Stage: Blueprint placement (Part 2)
+		//*    Takes a screenshot of the initial state of the blueprint
+		//*/
+		//bool Blueprint_Placement_Part2()
+		//{
+		//	if (BlueprintObject && PlacedBlueprint)
+		//	{
+		//		FEditorPromotionTestUtilities::TakeScreenshot(TEXT("BlueprintPIE_Start"));
+		//	}
+		//	return true;
+		//}
+
 		/**
-		* Blueprint Test Stage: Data-only Blueprint (Part 1)
-		*    Closes the blueprint editor
+		* Blueprint Test Stage: Blueprint placement (Part 3)
+		*    Waits for 2 seconds for the timer to finish
 		*/
-		bool Blueprint_DataOnlyBlueprint_Part1()
+		bool Blueprint_Placement_Part3()
 		{
-			if (BlueprintObject)
-			{
-				UE_LOG(LogEditorBuildPromotionTests, Display, TEXT("Closing the blueprint editor"));
-				FAssetEditorManager::Get().CloseAllAssetEditors();
-			}
-			return true;
-		}
-
-		/**
-		* Blueprint Test Stage: Data-only Blueprint (Part 2)
-		*    Re opens the blueprint editor
-		*/
-		bool Blueprint_DataOnlyBlueprint_Part2()
-		{
-			if (BlueprintObject)
-			{
-				UE_LOG(LogEditorBuildPromotionTests, Display, TEXT("Opening the blueprint editor for the second time"));
-				FAssetEditorManager::Get().OpenEditorForAsset(BlueprintObject);
-			}
-			return true;
-		}
-
-		/**
-		* Blueprint Test Stage: Data-only Blueprint (Part 3)
-		*    Checks that the editor opened in the Defaults mode
-		*/
-		bool Blueprint_DataOnlyBlueprint_Part3()
-		{
-			if (BlueprintObject)
-			{
-				IAssetEditorInstance* AssetEditor = FAssetEditorManager::Get().FindEditorForAsset(BlueprintObject, true);
-				IBlueprintEditor* BlueprintEditor = (IBlueprintEditor*)AssetEditor;
-				if (BlueprintEditor->GetCurrentMode() == FBlueprintEditorApplicationModes::BlueprintDefaultsMode)
-				{
-					//Good.  Now switch to components mode
-					UE_LOG(LogEditorBuildPromotionTests, Display, TEXT("Opened the correct defaults editor"));
-					UE_LOG(LogEditorBuildPromotionTests, Display, TEXT("Switching to components mode"));
-					BlueprintEditor->SetCurrentMode(FBlueprintEditorApplicationModes::BlueprintComponentsMode);
-				}
-				else
-				{
-					UE_LOG(LogEditorBuildPromotionTests, Error, TEXT("Didn't open the correct editor the second time"));
-				}
-			}
-			return true;
-		}
-
-		/**
-		* Blueprint Test Stage:  Components Mode (Part 1)
-		*   Adds 3 new components to the blueprint, changes the root component, renames the components, and compiles the blueprint
-		*/
-		bool Blueprint_ComponentsMode_Part1()
-		{
-			if (BlueprintObject)
-			{
-				IAssetEditorInstance* AssetEditor = FAssetEditorManager::Get().FindEditorForAsset(BlueprintObject, true);
-				IBlueprintEditor* BlueprintEditor = (IBlueprintEditor*)AssetEditor;
-
-				MeshNode = EditorBuildPromotionTestUtils::CreateBlueprintComponent(BlueprintObject, FirstBlueprintMesh);
-				if (MeshNode)
-				{
-					UE_LOG(LogEditorBuildPromotionTests, Display, TEXT("Added a mesh component"));
-				}
-				else
-				{
-					UE_LOG(LogEditorBuildPromotionTests, Error, TEXT("Failed to create the first mesh component"));
-				}
-
-				OtherMeshNode = EditorBuildPromotionTestUtils::CreateBlueprintComponent(BlueprintObject, SecondBlueprintMesh);
-				if (OtherMeshNode)
-				{
-					UE_LOG(LogEditorBuildPromotionTests, Display, TEXT("Added a second mesh component"));
-				}
-				else
-				{
-					UE_LOG(LogEditorBuildPromotionTests, Error, TEXT("Failed to create the second mesh component"));
-				}
-
-				PSNode = EditorBuildPromotionTestUtils::CreateBlueprintComponent(BlueprintObject, LoadedParticleSystem);
-				if (PSNode)
-				{
-					UE_LOG(LogEditorBuildPromotionTests, Display, TEXT("Added a particle system component"));
-				}
-				else
-				{
-					UE_LOG(LogEditorBuildPromotionTests, Error, TEXT("Failed to create the particle system component"));
-				}
-
-				//Set the Particle System as the root
-				// @FIXME: This will probably do bad things if PSNode doesn't exist?
-				EditorBuildPromotionTestUtils::SetComponentAsRoot(BlueprintObject, PSNode);
-				UE_LOG(LogEditorBuildPromotionTests, Display, TEXT("Set the particle system component as the new root"));
-
-				//Set Names
-				const FName MeshName(TEXT("FirstMesh"));
-				FBlueprintEditorUtils::RenameComponentMemberVariable(BlueprintObject, MeshNode, MeshName);
-				UE_LOG(LogEditorBuildPromotionTests, Display, TEXT("Renamed the first mesh component to FirstMesh"));
-
-				const FName OtherMeshName(TEXT("SecondMesh"));
-				FBlueprintEditorUtils::RenameComponentMemberVariable(BlueprintObject, OtherMeshNode, OtherMeshName);
-				UE_LOG(LogEditorBuildPromotionTests, Display, TEXT("Renamed the second mesh component to SecondMesh"));
-
-				// @FIXME: This will probably also do bad things if PSNode doesn't exist?
-				const FName PSName(TEXT("ParticleSys"));
-				FBlueprintEditorUtils::RenameComponentMemberVariable(BlueprintObject, PSNode, PSName);
-				UE_LOG(LogEditorBuildPromotionTests, Display, TEXT("Renamed the particle system component to ParticleSys"));
-
-				EditorBuildPromotionTestUtils::CompileBlueprint(BlueprintObject);
-
-				UE_LOG(LogEditorBuildPromotionTests, Display, TEXT("Switched to graph editing mode"));
-				BlueprintEditor->SetCurrentMode(FBlueprintEditorApplicationModes::StandardBlueprintEditorMode);
-			}
-			return true;
-		}
-
-		/**
-		* Blueprint Test Stage: Components Mode (Part 2)
-		*    Removes the 3 components added before and compiles the blueprint
-		*/
-		bool Blueprint_ComponentsMode_Part2()
-		{
-			if (BlueprintObject)
-			{
-				IAssetEditorInstance* AssetEditor = FAssetEditorManager::Get().FindEditorForAsset(BlueprintObject, true);
-				IBlueprintEditor* BlueprintEditor = (IBlueprintEditor*)AssetEditor;
-
-				FEditorPromotionTestUtilities::TakeScreenshot(TEXT("BlueprintComponentVariables"), true);
-
-				UE_LOG(LogEditorBuildPromotionTests, Display, TEXT("Switched to components mode"));
-				BlueprintEditor->SetCurrentMode(FBlueprintEditorApplicationModes::BlueprintComponentsMode);
-
-				EditorBuildPromotionTestUtils::RemoveBlueprintComponent(BlueprintObject, MeshNode);
-				EditorBuildPromotionTestUtils::RemoveBlueprintComponent(BlueprintObject, OtherMeshNode);
-				EditorBuildPromotionTestUtils::RemoveBlueprintComponent(BlueprintObject, PSNode);
-
-				//There should only be the scene component left
-				if (BlueprintObject->SimpleConstructionScript->GetAllNodes().Num() == 1)
-				{
-					UE_LOG(LogEditorBuildPromotionTests, Display, TEXT("Successfully removed the blueprint components"));
-				}
-				else
-				{
-					UE_LOG(LogEditorBuildPromotionTests, Error, TEXT("Failed to remove the blueprint components"));
-				}
-
-				MeshNode = NULL;
-				OtherMeshNode = NULL;
-				PSNode = NULL;
-
-				UE_LOG(LogEditorBuildPromotionTests, Display, TEXT("Switched to graph mode"));
-				BlueprintEditor->SetCurrentMode(FBlueprintEditorApplicationModes::StandardBlueprintEditorMode);
-
-				EditorBuildPromotionTestUtils::CompileBlueprint(BlueprintObject);
-			}
-			return true;
-		}
-
-		/**
-		* Blueprint Test Stage: Construction Script
-		*    Adds an AddStaticMeshComponent to the construction graph and links it to the entry node
-		*/
-		bool Blueprint_ConstructionScript()
-		{
-			if (BlueprintObject)
-			{
-				const UEdGraphSchema_K2* K2Schema = GetDefault<UEdGraphSchema_K2>();
-
-				IAssetEditorInstance* AssetEditor = FAssetEditorManager::Get().FindEditorForAsset(BlueprintObject, true);
-				FBlueprintEditor* BlueprintEditor = (FBlueprintEditor*)AssetEditor;
-
-				UEdGraph* ConstructionGraph = FBlueprintEditorUtils::FindUserConstructionScript(BlueprintObject);
-				BlueprintEditor->OpenGraphAndBringToFront(ConstructionGraph);
-
-				AddMeshNode = EditorBuildPromotionTestUtils::CreateAddComponentActionNode(BlueprintObject, ConstructionGraph, FirstBlueprintMesh);
-				if (AddMeshNode)
-				{
-					UE_LOG(LogEditorBuildPromotionTests, Display, TEXT("Created an Add Static Mesh Component node"));
-				}
-				else
-				{
-					UE_LOG(LogEditorBuildPromotionTests, Error, TEXT("Failed to create an Add Static Mesh Component node"));
-				}
-
-				GEditor->UndoTransaction();
-				if (ConstructionGraph->Nodes.Num() == 0 || ConstructionGraph->Nodes[ConstructionGraph->Nodes.Num() - 1] != AddMeshNode)
-				{
-					UE_LOG(LogEditorBuildPromotionTests, Display, TEXT("Undo add component node successful"));
-				}
-				else
-				{
-					UE_LOG(LogEditorBuildPromotionTests, Error, TEXT("Failed to undo add component node"));
-				}
-
-				GEditor->RedoTransaction();
-				if (ConstructionGraph->Nodes.Num() > 0 && ConstructionGraph->Nodes[ConstructionGraph->Nodes.Num() - 1] == AddMeshNode)
-				{
-					UE_LOG(LogEditorBuildPromotionTests, Display, TEXT("Redo add component node successful"));
-				}
-				else
-				{
-					UE_LOG(LogEditorBuildPromotionTests, Error, TEXT("Failed to redo add component node"));
-				}
-
-
-				TArray<UK2Node_FunctionEntry*> EntryNodes;
-				ConstructionGraph->GetNodesOfClass(EntryNodes);
-				UEdGraphNode* EntryNode = EntryNodes.Num() > 0 ? EntryNodes[0] : NULL;
-				if (EntryNode)
-				{
-					EditorBuildPromotionTestUtils::ConnectGraphNodes(AddMeshNode, K2Schema->PN_Execute, EntryNode, K2Schema->PN_Then);
-					UE_LOG(LogEditorBuildPromotionTests, Display, TEXT("Connected the entry node to the Add Static Mesh node"));
-
-					UEdGraphPin* EntryOutPin = EntryNode->FindPin(K2Schema->PN_Then);
-					UEdGraphPin* AddStaticMeshInPin = AddMeshNode->FindPin(K2Schema->PN_Execute);
-
-					GEditor->UndoTransaction();
-					if (EntryOutPin->LinkedTo.Num() == 0)
-					{
-						UE_LOG(LogEditorBuildPromotionTests, Display, TEXT("Undo connection successful"));
-					}
-					else
-					{
-						UE_LOG(LogEditorBuildPromotionTests, Error, TEXT("Failed to undo connection to Add Static Mesh node"));
-					}
-
-
-					GEditor->RedoTransaction();
-					if (EntryOutPin->LinkedTo.Contains(AddStaticMeshInPin))
-					{
-						UE_LOG(LogEditorBuildPromotionTests, Display, TEXT("Redo connection successful"));
-					}
-					else
-					{
-						UE_LOG(LogEditorBuildPromotionTests, Error, TEXT("Failed to redo connection to Add Static Mesh node"));
-					}
-				}
-				else
-				{
-					UE_LOG(LogEditorBuildPromotionTests, Error, TEXT("Failed to connect the entry node to the Add Static Mesh node"));
-				}
-
-				EditorBuildPromotionTestUtils::CompileBlueprint(BlueprintObject);
-			}
-			return true;
-		}
-
-		/**
-		* Blueprint Test Stage: Variable from component mode (Part 1)
-		*    Promotes the return pin of the AddStaticMeshNode to a variable and then renames it
-		*/
-		bool Blueprint_PromoteVariable_Part1()
-		{
-			if (BlueprintObject)
-			{
-				IAssetEditorInstance* AssetEditor = FAssetEditorManager::Get().FindEditorForAsset(BlueprintObject, true);
-				FBlueprintEditor* BlueprintEditor = (FBlueprintEditor*)AssetEditor;
-
-				const UEdGraphSchema_K2* K2Schema = GetDefault<UEdGraphSchema_K2>();
-				EditorBuildPromotionTestUtils::PromotePinToVariable(BlueprintObject, AddMeshNode, K2Schema->PN_ReturnValue);
-
-				UE_LOG(LogEditorBuildPromotionTests, Display, TEXT("Promoted the return pin on the add mesh node to a variable"));
-
-				const FName OldVarName(TEXT("NewVar")); //<- Default variable name
-				const FName NewVarName(TEXT("MyMesh"));
-				FBlueprintEditorUtils::RenameMemberVariable(BlueprintObject, OldVarName, NewVarName);
-				UE_LOG(LogEditorBuildPromotionTests, Display, TEXT("Renamed the new variable to MyMesh"));
-
-				EditorBuildPromotionTestUtils::CompileBlueprint(BlueprintObject);
-
-				UE_LOG(LogEditorBuildPromotionTests, Display, TEXT("Switched to graph mode"));
-				BlueprintEditor->SetCurrentMode(FBlueprintEditorApplicationModes::StandardBlueprintEditorMode);
-			}
-			return true;
-		}
-
-		bool Blueprint_PromoteVariable_Part2()
-		{
-			if (BlueprintObject)
-			{
-				IAssetEditorInstance* AssetEditor = FAssetEditorManager::Get().FindEditorForAsset(BlueprintObject, true);
-				FBlueprintEditor* BlueprintEditor = (FBlueprintEditor*)AssetEditor;
-
-				FEditorPromotionTestUtilities::TakeScreenshot(TEXT("BlueprintMeshVariable"), true);
-
-				UE_LOG(LogEditorBuildPromotionTests, Display, TEXT("Switched to components mode"));
-				BlueprintEditor->SetCurrentMode(FBlueprintEditorApplicationModes::BlueprintComponentsMode);
-
-				EditorBuildPromotionTestUtils::SendBlueprintResetViewCommand();
-			}
-			return true;
-		}
-
-		/**
-		* Blueprint Test Stage: Variable from component mode (Part 3)
-		*    Takes a screenshot of the mesh variable
-		*/
-		bool Blueprint_PromoteVariable_Part3()
-		{
-			if (BlueprintObject)
-			{
-				FEditorPromotionTestUtilities::TakeScreenshot(TEXT("BlueprintComponent"), true);
-			}
-			return true;
-		}
-
-		/**
-		* Blueprint Test Stage: Event Graph
-		*    Adds a ReceiveBeginPlay and Delay node to the event graph 
-		*/
-		bool Blueprint_EventGraph()
-		{
-			if (BlueprintObject)
-			{
-				IAssetEditorInstance* AssetEditor = FAssetEditorManager::Get().FindEditorForAsset(BlueprintObject, true);
-				FBlueprintEditor* BlueprintEditor = (FBlueprintEditor*)AssetEditor;
-
-				UEdGraph* EventGraph = FBlueprintEditorUtils::FindEventGraph(BlueprintObject);
-				BlueprintEditor->OpenGraphAndBringToFront(EventGraph);
-				UE_LOG(LogEditorBuildPromotionTests, Display, TEXT("Opened the event graph"));
-
-				PostBeginPlayEventNode = EditorBuildPromotionTestUtils::CreatePostBeginPlayEvent(BlueprintObject, EventGraph);
-				if (PostBeginPlayEventNode)
-				{
-					UE_LOG(LogEditorBuildPromotionTests, Display, TEXT("Created a ReceiveBeginPlay node"));
-				}
-				else
-				{
-					UE_LOG(LogEditorBuildPromotionTests, Error, TEXT("Failed to create a ReceiveBeginPlay node"));
-				}
-
-				const UEdGraphSchema_K2* K2Schema = GetDefault<UEdGraphSchema_K2>();
-				UEdGraphPin* PlayThenPin = PostBeginPlayEventNode->FindPin(K2Schema->PN_Then);
-				
-				DelayNode = EditorBuildPromotionTestUtils::AddDelayNode(BlueprintObject, EventGraph, PlayThenPin);
-				if (DelayNode)
-				{
-					UE_LOG(LogEditorBuildPromotionTests, Display, TEXT("Created a Delay node"));
-				}
-				else
-				{
-					UE_LOG(LogEditorBuildPromotionTests, Error, TEXT("Failed to create a Delay node"));
-				}
-
-				GEditor->UndoTransaction();
-				if (EventGraph->Nodes.Num() == 0 || EventGraph->Nodes[EventGraph->Nodes.Num() - 1] != DelayNode)
-				{
-					UE_LOG(LogEditorBuildPromotionTests, Display, TEXT("Undo add Delay node successful"));
-				}
-				else
-				{
-					UE_LOG(LogEditorBuildPromotionTests, Error, TEXT("Failed to undo add delay node"));
-				}
-
-				GEditor->RedoTransaction();
-				if (EventGraph->Nodes.Num() > 0 && EventGraph->Nodes[EventGraph->Nodes.Num() - 1] == DelayNode)
-				{
-					UE_LOG(LogEditorBuildPromotionTests, Display, TEXT("Redo add Delay node successful"));
-				}
-				else
-				{
-					UE_LOG(LogEditorBuildPromotionTests, Error, TEXT("Failed to redo add delay node"));
-				}
-
-				EditorBuildPromotionTestUtils::SetPinDefaultValue(DelayNode, TEXT("Duration"), TEXT("2.0"));
-				UE_LOG(LogEditorBuildPromotionTests, Display, TEXT("Set the duration of the delay node to 2.0 seconds"));
-
-				EditorBuildPromotionTestUtils::CompileBlueprint(BlueprintObject);
-			}
-
-			return true;
-		}
-
-		/**
-		* Blueprint Test Stage: Custom Variable
-		*    Creates a custom string variable and adds Get/Set nodes for it
-		*/
-		bool Blueprint_CustomVariable()
-		{
-			if (BlueprintObject)
-			{
-				UEdGraph* EventGraph = FBlueprintEditorUtils::FindEventGraph(BlueprintObject);
-
-				UE_LOG(LogEditorBuildPromotionTests, Display, TEXT("Added a string member variable"));
-				EditorBuildPromotionTestUtils::AddStringMemberValue(BlueprintObject, EditorBuildPromotionTestUtils::BlueprintStringVariableName);
-
-				UE_LOG(LogEditorBuildPromotionTests, Display, TEXT("Added a new Set node for the string variable"));
-				SetNode = EditorBuildPromotionTestUtils::AddGetSetNode(BlueprintObject, EventGraph, EditorBuildPromotionTestUtils::BlueprintStringVariableName.ToString(), false);
-
-				GEditor->UndoTransaction();
-				if (EventGraph->Nodes.Num() == 0 || EventGraph->Nodes[EventGraph->Nodes.Num() - 1] != SetNode)
-				{
-					UE_LOG(LogEditorBuildPromotionTests, Display, TEXT("Undo add Set node successful"));
-				}
-				else
-				{
-					UE_LOG(LogEditorBuildPromotionTests, Error, TEXT("Failed to undo add Set node"));
-				}
-
-				GEditor->RedoTransaction();
-				if (EventGraph->Nodes.Num() > 0 && EventGraph->Nodes[EventGraph->Nodes.Num() - 1] == SetNode)
-				{
-					UE_LOG(LogEditorBuildPromotionTests, Display, TEXT("Redo add Set node successful"));
-				}
-				else
-				{
-					UE_LOG(LogEditorBuildPromotionTests, Error, TEXT("Failed to redo add set node"));
-				}
-
-				UE_LOG(LogEditorBuildPromotionTests, Display, TEXT("Adding a Get node for the string variable"));
-				GetNode = EditorBuildPromotionTestUtils::AddGetSetNode(BlueprintObject, EventGraph, EditorBuildPromotionTestUtils::BlueprintStringVariableName.ToString(), true, 400);
-				
-				GEditor->UndoTransaction();
-				if (EventGraph->Nodes.Num() == 0 || EventGraph->Nodes[EventGraph->Nodes.Num() - 1] != GetNode)
-				{
-					UE_LOG(LogEditorBuildPromotionTests, Display, TEXT("Undo add get node"));
-				}
-				else
-				{
-					UE_LOG(LogEditorBuildPromotionTests, Error, TEXT("Failed to undo add get node"));
-				}
-
-				GEditor->RedoTransaction();
-				if (EventGraph->Nodes.Num() > 0 && EventGraph->Nodes[EventGraph->Nodes.Num() - 1] == GetNode)
-				{
-					UE_LOG(LogEditorBuildPromotionTests, Display, TEXT("Redo add get node"));
-				}
-				else
-				{
-					UE_LOG(LogEditorBuildPromotionTests, Error, TEXT("Failed to redo add get node"));
-				}
-
-				UE_LOG(LogEditorBuildPromotionTests, Display, TEXT("Removed the Set node"));
-				EventGraph->RemoveNode(SetNode);
-				SetNode = NULL;
-			}
-			return true;
-		}
-
-		/**
-		* Blueprint Test Stage:  Using Variables
-		*   Adds a PrintString and SetStaticMesh then connects all the existing nodes
-		*/
-		bool Blueprint_UsingVariables()
-		{
-			if (BlueprintObject)
-			{
-				const bool bVariableIsHidden = false;
-				FBlueprintEditorUtils::SetBlueprintOnlyEditableFlag(BlueprintObject, EditorBuildPromotionTestUtils::BlueprintStringVariableName, bVariableIsHidden);
-				UE_LOG(LogEditorBuildPromotionTests, Display, TEXT("Exposed the blueprint string variable"));
-
-				UEdGraph* EventGraph = FBlueprintEditorUtils::FindEventGraph(BlueprintObject);
-				PrintNode = EditorBuildPromotionTestUtils::AddPrintStringNode(BlueprintObject, EventGraph);
-				if (PrintNode)
-				{
-					UE_LOG(LogEditorBuildPromotionTests, Display, TEXT("Added a PrintString node"));
-				}
-				else
-				{
-					UE_LOG(LogEditorBuildPromotionTests, Error, TEXT("Failed to create a PrintString node"));
-				}
-
-				const UEdGraphSchema_K2* K2Schema = GetDefault<UEdGraphSchema_K2>();
-
-				//Connect Get to printstring
-				UEdGraphPin* GetVarPin = GetNode->FindPin(EditorBuildPromotionTestUtils::BlueprintStringVariableName.ToString());
-				UEdGraphPin* InStringPin = PrintNode->FindPin(TEXT("InString"));
-				GetVarPin->MakeLinkTo(InStringPin);
-				UE_LOG(LogEditorBuildPromotionTests, Display, TEXT("Connected the string variable Get node to the PrintString node"));
-
-				//Connect Delay to PrintString
-				UEdGraphPin* DelayExecPin = DelayNode->FindPin(K2Schema->PN_Then);
-				UEdGraphPin* PrintStringPin = PrintNode->FindPin(K2Schema->PN_Execute);
-				DelayExecPin->MakeLinkTo(PrintStringPin);
-				UE_LOG(LogEditorBuildPromotionTests, Display, TEXT("Connected the Delay node to the PrintString node"));
-
-				const FName MyMeshVarName(TEXT("MyMesh"));
-				GetNode = EditorBuildPromotionTestUtils::AddGetSetNode(BlueprintObject, EventGraph, MyMeshVarName.ToString(), true, 680);
-				if (GetNode)
-				{
-					UE_LOG(LogEditorBuildPromotionTests, Display, TEXT("Added a Get node for the MyMesh variable"));
-				}
-				else
-				{
-					UE_LOG(LogEditorBuildPromotionTests, Error, TEXT("Failed to create a Get node for the MyMesh variable"));
-				}
-
-				SetStaticMeshNode = EditorBuildPromotionTestUtils::AddSetStaticMeshNode(BlueprintObject, EventGraph);
-				if (SetStaticMeshNode)
-				{
-					UE_LOG(LogEditorBuildPromotionTests, Display, TEXT("Added a SetStaticMesh node"));
-				}
-				else
-				{
-					UE_LOG(LogEditorBuildPromotionTests, Error, TEXT("Failed to create a SetStaticMesh node"));
-				}
-
-				UEdGraphPin* GetExecPin = GetNode->FindPin(TEXT("MyMesh"));
-				UEdGraphPin* SetStaticMeshSelfPin = SetStaticMeshNode->FindPin(K2Schema->PN_Self);
-				GetExecPin->MakeLinkTo(SetStaticMeshSelfPin);
-				UE_LOG(LogEditorBuildPromotionTests, Display, TEXT("Connected the Get MyMesh node to the SetStaticMesh node"));
-
-				UEdGraphPin* SetStaticMeshMeshPin = SetStaticMeshNode->FindPin(TEXT("NewMesh"));
-				SetStaticMeshMeshPin->DefaultObject = SecondBlueprintMesh;
-				UE_LOG(LogEditorBuildPromotionTests, Display, TEXT("Set the satic mesh on the SetStaticMesh node to: %s"), *SecondBlueprintMesh->GetName());
-
-				//Connect SetStaticMeshMesh to PrintString
-				UEdGraphPin* PrintStringThenPin = PrintNode->FindPin(K2Schema->PN_Then);
-				UEdGraphPin* SetStaticMeshExecPin = SetStaticMeshNode->FindPin(K2Schema->PN_Execute);
-				PrintStringThenPin->MakeLinkTo(SetStaticMeshExecPin);
-				UE_LOG(LogEditorBuildPromotionTests, Display, TEXT("Connected the PrintString node to the SetStaticMesh node"));
-			}
-			return true;
-		}
-
-		/**
-		* Blueprint Test Stage: Renaming custom event
-		*    Creates, renames, and then removes a custom event node
-		*/
-		bool Blueprint_RenameCustomEvent()
-		{
-			if (BlueprintObject)
-			{
-				UEdGraph* EventGraph = FBlueprintEditorUtils::FindEventGraph(BlueprintObject);
-				CustomEventNode = EditorBuildPromotionTestUtils::CreateCustomEvent(BlueprintObject, EventGraph, TEXT("NewEvent"));
-				if (CustomEventNode)
-				{
-					UE_LOG(LogEditorBuildPromotionTests, Display, TEXT("Created a custom event node"));
-					//Rename the event
-					CustomEventNode->OnRenameNode(TEXT("RenamedEvent"));
-					UE_LOG(LogEditorBuildPromotionTests, Display, TEXT("Renamed the custom event node to RenamedEvent"));
-					EventGraph->RemoveNode(CustomEventNode);
-					UE_LOG(LogEditorBuildPromotionTests, Display, TEXT("Removed the custom event node"));
-					CustomEventNode = NULL;
-				}
-				else
-				{
-					UE_LOG(LogEditorBuildPromotionTests, Error, TEXT("Failed to create a custom event node"));
-				}
-
-			}
-
-			return true;
-		}
-
-		/**
-		* Blueprint Test Stage: New function
-		*    Creates a new function graph and then hooks up a new AddParticleSystem inside it
-		*/
-		bool Blueprint_NewFunctions()
-		{
-			if (BlueprintObject)
-			{
-				const UEdGraphSchema_K2* K2Schema = GetDefault<UEdGraphSchema_K2>();
-
-				CustomGraph = EditorBuildPromotionTestUtils::CreateNewFunctionGraph(BlueprintObject, TEXT("NewFunction"));
-				if (CustomGraph)
-				{
-					UE_LOG(LogEditorBuildPromotionTests, Display, TEXT("Created a new function graph (NewFunction)"));
-				}
-				else
-				{
-					UE_LOG(LogEditorBuildPromotionTests, Error, TEXT("Failed to create a new function graph"));
-				}
-
-				AddParticleSystemNode = EditorBuildPromotionTestUtils::CreateAddComponentActionNode(BlueprintObject, CustomGraph, LoadedParticleSystem);
-				if (AddParticleSystemNode)
-				{
-					UE_LOG(LogEditorBuildPromotionTests, Display, TEXT("Created an AddParticleSystem node"));
-				}
-				else
-				{
-					UE_LOG(LogEditorBuildPromotionTests, Error, TEXT("Failed to create an AddParticleSystem node"));
-				}
-
-				UEdGraphPin* ExecutePin = AddParticleSystemNode ? AddParticleSystemNode->FindPin(K2Schema->PN_Execute) : NULL;
-
-				//Find the input for the function graph
-				TArray<UK2Node_FunctionEntry*> EntryNodes;
-				CustomGraph->GetNodesOfClass(EntryNodes);
-				UEdGraphNode* EntryNode = EntryNodes.Num() > 0 ? EntryNodes[0] : NULL;
-				if (EntryNode && ExecutePin)
-				{
-					UEdGraphPin* EntryPin = EntryNode->FindPin(K2Schema->PN_Then);
-					EntryPin->MakeLinkTo(ExecutePin);
-					UE_LOG(LogEditorBuildPromotionTests, Display, TEXT("Connected the AddParticleSystem node to the entry node"));
-				}
-
-				EditorBuildPromotionTestUtils::CompileBlueprint(BlueprintObject);
-			}
-			return true;
-		}
-
-		/**
-		* Blueprint Test Stage: Completing the blueprint
-		*    Adds a CallFunction node to call the custom function created in the previous step.
-		*/
-		bool Blueprint_CompleteBlueprint()
-		{
-			if (BlueprintObject)
-			{
-				const UEdGraphSchema_K2* K2Schema = GetDefault<UEdGraphSchema_K2>();
-				UEdGraph* EventGraph = FBlueprintEditorUtils::FindEventGraph(BlueprintObject);
-				UEdGraphPin* SetStaticMeshThenPin = SetStaticMeshNode->FindPin(K2Schema->PN_Then);
-				CallFunctionNode = EditorBuildPromotionTestUtils::AddCallFunctionGraphNode(BlueprintObject, EventGraph, TEXT("NewFunction"), SetStaticMeshThenPin);
-				if (CallFunctionNode)
-				{
-					UE_LOG(LogEditorBuildPromotionTests, Display, TEXT("Created a CallFunction node to call NewFunction"));
-					UE_LOG(LogEditorBuildPromotionTests, Display, TEXT("Hooked the CallFunction node to the output of the SetStaticMesh node"));
-				}
-				else
-				{
-					UE_LOG(LogEditorBuildPromotionTests, Error, TEXT("Failed to create a CallFunction node"));
-				}
-
-				EditorBuildPromotionTestUtils::CompileBlueprint(BlueprintObject);
-
-				SaveBlueprint();
-			}
-			return true;
-		}
-
-		/**
-		* Blueprint Test Stage: Blueprint placement (Part 1)
-		*    Places the blueprint in the level then starts PIE
-		*/
-		bool Blueprint_Placement_Part1()
-		{
-			if (BlueprintObject)
-			{
-				PlacedBlueprint = FActorFactoryAssetProxy::AddActorForAsset(BlueprintObject);
-
-				if (PlacedBlueprint)
-				{
-					//Set the text
-					FEditorPromotionTestUtilities::SetPropertyByName(PlacedBlueprint, EditorBuildPromotionTestUtils::BlueprintStringVariableName.ToString(), TEXT("Print String works!"));
-					UE_LOG(LogEditorBuildPromotionTests, Display, TEXT("Placed the blueprint and set the display string to \"Print String works!\""));
-				}
-				else
-				{
-					UE_LOG(LogEditorBuildPromotionTests, Error, TEXT("Failed to place blueprint"));
-				}
-
-				GEditor->SelectNone(false, true);
-
-				//Start PIE
-				EditorBuildPromotionTestUtils::StartPIE(true);
-
-				//Make sure the timer is reset
-				DelayHelper.Reset();
-			}
-			return true;
-		}
-
-		/**
-		* Blueprint Test Stage: Blueprint placement (Part 2)
-		*    Takes a screenshot of the initial state of the blueprint
-		*/
-		bool Blueprint_Placement_Part2()
-		{
-			if (BlueprintObject)
-			{
-				FEditorPromotionTestUtilities::TakeScreenshot(TEXT("BlueprintPIE_Start"));
-			}
-			return true;
-		}
-
-		/**
-		* Blueprint Test Stage: Blueprint placement (Part 3)
-		*    Waits for 2 seconds for the timer to finish
-		*/
-		bool Blueprint_Placement_Part3()
-		{
-			if (BlueprintObject)
+			if (BlueprintObject && PlacedBlueprint)
 			{
 				//Set a timeout in case PIE doesn't work
 				if (!DelayHelper.IsRunning())
@@ -3115,7 +1997,7 @@ namespace BuildPromotionTestHelper
 				else if (DelayHelper.IsComplete())
 				{
 					//FAILED to hit breakpoint in time
-					UE_LOG(LogEditorBuildPromotionTests, Error, TEXT("Timed out waiting for PIE to start"));
+					Test->AddError(TEXT("Timed out waiting for PIE to start"));
 					DelayHelper.Reset();
 					return true;
 				}
@@ -3141,12 +2023,12 @@ namespace BuildPromotionTestHelper
 		*/
 		bool Blueprint_Placement_Part4()
 		{
-			if (BlueprintObject)
+			if (BlueprintObject && PlacedBlueprint)
 			{
-				if (GEditor->PlayWorld != NULL)
-				{
-					FEditorPromotionTestUtilities::TakeScreenshot(TEXT("BlueprintPIE_End"));
-				}
+				//if (GEditor->PlayWorld != NULL)
+				//{
+				//	FEditorPromotionTestUtilities::TakeScreenshot(TEXT("BlueprintPIE_End"));
+				//}
 				FEditorPromotionTestUtilities::EndPIE();
 			}
 			return true;
@@ -3160,7 +2042,6 @@ namespace BuildPromotionTestHelper
 		{
 			if (BlueprintObject)
 			{
-
 				//Add a breakpoint
 				UBreakpoint* NewBreakpoint = NewObject<UBreakpoint>(BlueprintObject);
 				FKismetDebugUtilities::SetBreakpointEnabled(NewBreakpoint, true);
@@ -3168,7 +2049,7 @@ namespace BuildPromotionTestHelper
 				BlueprintObject->Breakpoints.Add(NewBreakpoint);
 				BlueprintObject->MarkPackageDirty();
 
-				UE_LOG(LogEditorBuildPromotionTests, Display, TEXT("Set a breakpoint on the PrintString node"));
+				Test->AddLogItem(TEXT("Set a breakpoint on the PrintString node"));
 				EditorBuildPromotionTestUtils::StartPIE(true);
 
 				//Make sure the timer is reset
@@ -3193,7 +2074,7 @@ namespace BuildPromotionTestHelper
 				else if (DelayHelper.IsComplete())
 				{
 					//FAILED to hit breakpoint in time
-					UE_LOG(LogEditorBuildPromotionTests, Error, TEXT("Failed to hit the breakpoint after 10 seconds"));
+					Test->AddError(TEXT("Failed to hit the breakpoint after 10 seconds"));
 					DelayHelper.Reset();
 					return true;
 				}
@@ -3209,7 +2090,7 @@ namespace BuildPromotionTestHelper
 				if (CurrentBreakpointNode == PrintNode)
 				{
 					//Success!
-					UE_LOG(LogEditorBuildPromotionTests, Display, TEXT("Hit the PrintString breakpoint"));
+					Test->AddLogItem(TEXT("Hit the PrintString breakpoint"));
 					DelayHelper.Reset();
 					return true;
 				}
@@ -3226,7 +2107,7 @@ namespace BuildPromotionTestHelper
 		{
 			if (BlueprintObject)
 			{
-				UE_LOG(LogEditorBuildPromotionTests, Display, TEXT("Clearing the breakpoint"));
+				Test->AddLogItem(TEXT("Clearing the breakpoint"));
 				FKismetDebugUtilities::ClearBreakpoints(BlueprintObject);
 				FEditorPromotionTestUtilities::EndPIE();
 				CurrentWorld = PlacedBlueprint->GetWorld();
@@ -3245,7 +2126,7 @@ namespace BuildPromotionTestHelper
 				//Open the level script blueprint
 				ULevelScriptBlueprint* LSB = PlacedBlueprint->GetLevel()->GetLevelScriptBlueprint(false);
 				FAssetEditorManager::Get().OpenEditorForAsset(LSB);
-				UE_LOG(LogEditorBuildPromotionTests, Display, TEXT("Opened the level script blueprint"));
+				Test->AddLogItem(TEXT("Opened the level script blueprint"));
 			}
 			return true;
 		}
@@ -3272,7 +2153,7 @@ namespace BuildPromotionTestHelper
 					NodesToExport.Add(EventGraph->Nodes[i]);
 				}
 
-				//Saving the starting log counts so we can clear any entries created from the copy/paste action.  (Lots of expected spam)
+				//Saving the starting log counts so we can clear any entries created from the copy/paste action. (Lots of expected spam)
 				const int32 PreCopyErrorCount = TestExecutionInfo->Errors.Num();
 				const int32 PreCopyWarningCount = TestExecutionInfo->Warnings.Num();
 				const int32 PreCopyLogCount = TestExecutionInfo->LogItems.Num();
@@ -3292,8 +2173,8 @@ namespace BuildPromotionTestHelper
 
 
 				//Note: These are a little out of order because logs are disabled above
-				UE_LOG(LogEditorBuildPromotionTests, Display, TEXT("Copied the blueprint event nodes"));
-				UE_LOG(LogEditorBuildPromotionTests, Display, TEXT("Pasted the nodes into to level script"));
+				Test->AddLogItem(TEXT("Copied the blueprint event nodes"));
+				Test->AddLogItem(TEXT("Pasted the nodes into to level script"));
 				
 
 				//We are expecting errors here.  So copy any new errors to the logs section
@@ -3311,49 +2192,6 @@ namespace BuildPromotionTestHelper
 					TestExecutionInfo->Errors.RemoveAt(NumErrors, TestExecutionInfo->Errors.Num() - NumErrors);
 				}
 
-				//Remove the "Get" nodes
-				TArray<UK2Node_VariableGet*> AllGetNodes;
-				LevelEventGraph->GetNodesOfClass<UK2Node_VariableGet>(AllGetNodes);
-				for (int32 i = 0; i < AllGetNodes.Num(); ++i)
-				{
-					LevelEventGraph->RemoveNode(AllGetNodes[i]);
-				}
-				UE_LOG(LogEditorBuildPromotionTests, Display, TEXT("Removed get nodes from the level script"));
-
-				//Remove all the call function nodes except the PrintString and Delay nodes
-				TArray<UK2Node_CallFunction*> AllCallFunctionNodes;
-				LevelEventGraph->GetNodesOfClass<UK2Node_CallFunction>(AllCallFunctionNodes);
-				for (int32 i = 0; i < AllCallFunctionNodes.Num(); ++i)
-				{
-					UK2Node_CallFunction* NodeIt = AllCallFunctionNodes[i];
-					const FString MemberName = NodeIt->FunctionReference.GetMemberName().ToString();
-					if (MemberName == TEXT("PrintString"))
-					{
-						UE_LOG(LogEditorBuildPromotionTests, Display, TEXT("Setting the level PrintString to \"Level kismet string text\""));
-						EditorBuildPromotionTestUtils::SetPinDefaultValue(NodeIt, TEXT("InString"), TEXT("Level kismet string text"));
-					}
-					else if (MemberName == TEXT("Delay"))
-					{
-						UE_LOG(LogEditorBuildPromotionTests, Display, TEXT("Setting the level script Delay to 4 seconds"));
-						EditorBuildPromotionTestUtils::SetPinDefaultValue(NodeIt, TEXT("Duration"), TEXT("4.0"));
-					}
-					else
-					{
-						if (NodeIt == CallFunctionNode)
-						{
-							UE_LOG(LogEditorBuildPromotionTests, Display, TEXT("Removed CallFunction node from level script"));
-							CallFunctionNode = NULL;
-						}
-						else if (NodeIt == SetStaticMeshNode)
-						{
-							UE_LOG(LogEditorBuildPromotionTests, Display, TEXT("Removed SetStaticMesh node from level script"));
-							SetStaticMeshNode = NULL;
-						}
-						//Remove 
-						LevelEventGraph->RemoveNode(NodeIt);
-					}
-				}
-
 				//Compile the blueprint
 				EditorBuildPromotionTestUtils::CompileBlueprint(LSB);
 
@@ -3382,7 +2220,7 @@ namespace BuildPromotionTestHelper
 				else if (DelayHelper.IsComplete())
 				{
 					//FAILED to hit breakpoint in time
-					UE_LOG(LogEditorBuildPromotionTests, Error, TEXT("Timed out waiting for PIE to start"));
+					Test->AddError(TEXT("Timed out waiting for PIE to start"));
 					DelayHelper.Reset();
 					return true;
 				}
@@ -3410,11 +2248,11 @@ namespace BuildPromotionTestHelper
 		{
 			if (BlueprintObject)
 			{
-				if (GEditor->PlayWorld != NULL)
-				{
-					//Take a screenshot and end the PIE session
-					FEditorPromotionTestUtilities::TakeScreenshot(TEXT("LevelBlueprint"), false);
-				}
+				//if (GEditor->PlayWorld != NULL)
+				//{
+				//	//Take a screenshot and end the PIE session
+				//	FEditorPromotionTestUtilities::TakeScreenshot(TEXT("LevelBlueprint"), false);
+				//}
 				FEditorPromotionTestUtilities::EndPIE();
 			}
 			return true;
@@ -3428,10 +2266,10 @@ namespace BuildPromotionTestHelper
 		{
 			if (BlueprintObject)
 			{
-				UE_LOG(LogEditorBuildPromotionTests, Display, TEXT("Closing the blueprint editor"));
+				Test->AddLogItem(TEXT("Closing the blueprint editor"));
 				FAssetEditorManager::Get().CloseAllAssetEditors();
-				UE_LOG(LogEditorBuildPromotionTests, Display, TEXT("Saving the blueprint"));
-				SaveBlueprint();
+				//UE_LOG(LogEditorBuildPromotionTests, Display, TEXT("Saving the blueprint"));
+				//SaveBlueprint();
 			}
 			return true;
 		}
@@ -3555,7 +2393,6 @@ namespace BuildPromotionTestHelper
 /**
 * Automation test that handles cleanup of the build promotion test
 */
-IMPLEMENT_SIMPLE_AUTOMATION_TEST(FBuildPromotionInitialCleanupTest, "System.Promotion.Editor Promotion Pass.Step 1 Main Editor Test.Cleanup old files", EAutomationTestFlags::EditorContext | EAutomationTestFlags::EngineFilter);
 bool FBuildPromotionInitialCleanupTest::RunTest(const FString& Parameters)
 {
 	EditorBuildPromotionTestUtils::PerformCleanup();
@@ -3591,11 +2428,10 @@ bool FRunBuildPromotionTestCommand::Update()
 /**
 * Automation test that handles the build promotion process
 */
-IMPLEMENT_SIMPLE_AUTOMATION_TEST(FBuildPromotionTest, "System.Promotion.Editor Promotion Pass.Step 1 Main Editor Test.General Editor Test", EAutomationTestFlags::EditorContext | EAutomationTestFlags::EngineFilter);
-
-bool FBuildPromotionTest::RunTest(const FString& Parameters)
+bool FEditorPromotionTest::RunTest(const FString& Parameters)
 {
 	TSharedPtr<BuildPromotionTestHelper::FBuildPromotionTest> BuildPromotionTest = MakeShareable(new BuildPromotionTestHelper::FBuildPromotionTest(&ExecutionInfo));
+	BuildPromotionTest->Test = this;
 	ADD_LATENT_AUTOMATION_COMMAND(FRunBuildPromotionTestCommand(BuildPromotionTest));
 	return true;
 }
@@ -3631,7 +2467,6 @@ bool FPIEExecCommand::Update()
 /**
 * Execute the loading of one map to verify PIE works
 */
-IMPLEMENT_SIMPLE_AUTOMATION_TEST(FBuildPromotionPIETest, "System.Promotion.Editor Promotion Pass.Step 2 Run Map After Re-launch.Run Map", EAutomationTestFlags::EditorContext | EAutomationTestFlags::EngineFilter)
 bool FBuildPromotionPIETest::RunTest(const FString& Parameters)
 {
 	UAutomationTestSettings const* AutomationTestSettings = GetDefault<UAutomationTestSettings>();
@@ -3678,7 +2513,6 @@ bool FBuildPromotionPIETest::RunTest(const FString& Parameters)
 /**
 * Automation test that handles cleanup of the build promotion test
 */
-IMPLEMENT_SIMPLE_AUTOMATION_TEST(FBuildPromotionCleanupTest, "System.Promotion.Editor Promotion Pass.Step 3 Test Cleanup.Cleanup", EAutomationTestFlags::EditorContext | EAutomationTestFlags::EngineFilter);
 bool FBuildPromotionCleanupTest::RunTest(const FString& Parameters)
 {
 	EditorBuildPromotionTestUtils::PerformCleanup();
diff --git a/Engine/Source/Editor/UnrealEd/Public/Tests/AutomationEditorCommon.h b/Engine/Source/Editor/UnrealEd/Public/Tests/AutomationEditorCommon.h
--- a/Engine/Source/Editor/UnrealEd/Public/Tests/AutomationEditorCommon.h
+++ b/Engine/Source/Editor/UnrealEd/Public/Tests/AutomationEditorCommon.h
@@ -58,7 +58,7 @@ namespace AutomationEditorCommonUtils
 	*
 	* @param InObject - Object to null references to
 	*/
-	void NullReferencesToObject(UObject* InObject);
+	UNREALED_API void NullReferencesToObject(UObject* InObject);
 
 	/**
 	* gets a factory class based off an asset file extension
diff --git a/Engine/Source/Editor/UnrealEd/Public/Tests/AutomationEditorPromotionCommon.h b/Engine/Source/Editor/UnrealEd/Public/Tests/AutomationEditorPromotionCommon.h
--- a/Engine/Source/Editor/UnrealEd/Public/Tests/AutomationEditorPromotionCommon.h
+++ b/Engine/Source/Editor/UnrealEd/Public/Tests/AutomationEditorPromotionCommon.h
@@ -88,6 +88,11 @@ public:
 	*/
 	UNREALED_API static void SetPropertyByName(UObject* TargetObject, const FString& InVariableName, const FString& NewValueString);
 
+	/**
+	* Starts a PIE session
+	*/
+	UNREALED_API static void StartPIE(bool bSimulateInEditor);
+
 	/**
 	* Ends a PIE session
 	*/
diff --git a/Engine/Source/Runtime/Engine/Classes/Tests/AutomationTestSettings.h b/Engine/Source/Runtime/Engine/Classes/Tests/AutomationTestSettings.h
--- a/Engine/Source/Runtime/Engine/Classes/Tests/AutomationTestSettings.h
+++ b/Engine/Source/Runtime/Engine/Classes/Tests/AutomationTestSettings.h
@@ -207,23 +207,6 @@ struct FBuildPromotionOpenAssetSettings
 	FFilePath TextureAsset;
 };
 
-/**
-* Holds settings for the blueprint stage of the build promotion test
-*/
-USTRUCT()
-struct FBuildPromotionBlueprintSettings
-{
-	GENERATED_USTRUCT_BODY()
-
-	/** The starting mesh for the blueprint **/
-	UPROPERTY(EditAnywhere, Category = Automation, meta = (FilePathFilter = "uasset"))
-	FFilePath FirstMeshPath;
-
-	/** The mesh to set on the blueprint after the delay **/
-	UPROPERTY(EditAnywhere, Category = Automation, meta = (FilePathFilter = "uasset"))
-	FFilePath SecondMeshPath;
-};
-
 /**
 * Holds settings for the new project stage of the build promotion test
 */
@@ -277,6 +260,28 @@ struct FParticleEditorPromotionSettings
 	FFilePath DefaultParticleAsset;
 };
 
+/**
+* Holds settings for the blueprint editor build promotion tests
+*/
+USTRUCT()
+struct FBlueprintEditorPromotionSettings
+{
+	GENERATED_USTRUCT_BODY()
+
+	/** The starting mesh for the blueprint **/
+	UPROPERTY(EditAnywhere, Category = Automation, meta = (FilePathFilter = "uasset"))
+	FFilePath FirstMeshPath;
+
+	/** The mesh to set on the blueprint after the delay **/
+	UPROPERTY(EditAnywhere, Category = Automation, meta = (FilePathFilter = "uasset"))
+		FFilePath SecondMeshPath;
+
+	/** Default particle asset to use for tests*/
+	UPROPERTY(EditAnywhere, Category = Automation, meta = (FilePathFilter = "uasset"))
+	FFilePath DefaultParticleAsset;
+};
+
+
 /**
 * Holds settings for the editor build promotion test
 */
@@ -297,10 +302,6 @@ struct FBuildPromotionTestSettings
 	UPROPERTY(EditAnywhere, Category = Automation)
 	FBuildPromotionOpenAssetSettings	OpenAssets;
 
-	/** Blueprint settings **/
-	UPROPERTY(EditAnywhere, Category = Automation)
-	FBuildPromotionBlueprintSettings	BlueprintSettings;
-
 	/** New project settings **/
 	UPROPERTY(EditAnywhere, Category = Automation)
 	FBuildPromotionNewProjectSettings NewProjectSettings;
@@ -384,6 +385,12 @@ public:
 	UPROPERTY(EditAnywhere, config, Category = Automation)
 	FParticleEditorPromotionSettings ParticleEditorPromotionTest;
 
+	/**
+	* Blueprint editor promotion test settings
+	*/
+	UPROPERTY(EditAnywhere, config, Category = Automation)
+	FBlueprintEditorPromotionSettings BlueprintEditorPromotionTest;
+
 	/**
 	* Modules to load that have engine tests
 	*/
