Use SizeSquared comparisons instead of Size where trivially replaceable

diff --git a/Engine/Source/Developer/LogVisualizer/Private/LogVisualizer.cpp b/Engine/Source/Developer/LogVisualizer/Private/LogVisualizer.cpp
--- a/Engine/Source/Developer/LogVisualizer/Private/LogVisualizer.cpp
+++ b/Engine/Source/Developer/LogVisualizer/Private/LogVisualizer.cpp
@@ -172,7 +172,7 @@ void FLogVisualizer::UpdateCameraPosition(FName RowName, int32 ItemIndes)
 
 
 	const float DefaultCameraDistance = ULogVisualizerSettings::StaticClass()->GetDefaultObject<ULogVisualizerSettings>()->DefaultCameraDistance;
-	Extent = Extent.Size() < DefaultCameraDistance ? FVector(1) * DefaultCameraDistance : Extent;
+	Extent = Extent.SizeSquared() < FMath::Square(DefaultCameraDistance) ? FVector(DefaultCameraDistance) : Extent;
 
 #if WITH_EDITOR
 	UEditorEngine *EEngine = Cast<UEditorEngine>(GEngine);
diff --git a/Engine/Source/Developer/MeshUtilities/Private/MeshUtilities.cpp b/Engine/Source/Developer/MeshUtilities/Private/MeshUtilities.cpp
--- a/Engine/Source/Developer/MeshUtilities/Private/MeshUtilities.cpp
+++ b/Engine/Source/Developer/MeshUtilities/Private/MeshUtilities.cpp
@@ -299,7 +299,7 @@ void FMeshDistanceFieldAsyncTask::DoWork()
 {
 	FMeshBuildDataProvider kDOPDataProvider(*kDopTree);
 	const FVector DistanceFieldVoxelSize(VolumeBounds.GetSize() / FVector(VolumeDimensions.X, VolumeDimensions.Y, VolumeDimensions.Z));
-	const float VoxelDiameter = DistanceFieldVoxelSize.Size();
+	const float VoxelDiameterSqr = DistanceFieldVoxelSize.SizeSquared();
 
 	for (int32 YIndex = 0; YIndex < VolumeDimensions.Y; YIndex++)
 	{
@@ -359,7 +359,7 @@ void FMeshDistanceFieldAsyncTask::DoWork()
 
 			// If we are very close to a surface and nearly all of our rays hit backfaces, treat as inside
 			// This is important for one sided planes
-			if (UnsignedDistance < VoxelDiameter && HitBack > .95f * Hit)
+			if (FMath::Square(UnsignedDistance) < VoxelDiameterSqr && HitBack > .95f * Hit)
 			{
 				MinDistance = -UnsignedDistance;
 			}
diff --git a/Engine/Source/Developer/Profiler/Private/Widgets/SProfilerMiniView.cpp b/Engine/Source/Developer/Profiler/Private/Widgets/SProfilerMiniView.cpp
--- a/Engine/Source/Developer/Profiler/Private/Widgets/SProfilerMiniView.cpp
+++ b/Engine/Source/Developer/Profiler/Private/Widgets/SProfilerMiniView.cpp
@@ -331,7 +331,7 @@ int32 SProfilerMiniView::OnPaint( const FPaintArgs& Args, const FGeometry& Allot
 
 void SProfilerMiniView::DrawText( const FString& Text, const FSlateFontInfo& FontInfo, FVector2D Position, const FColor& TextColor, const FColor& ShadowColor, FVector2D ShadowOffset ) const
 {
-	if( ShadowOffset.Size() > 0.0f )
+	if( ShadowOffset.SizeSquared() > 0.0f )
 	{
 		FSlateDrawElement::MakeText
 		(
diff --git a/Engine/Source/Developer/Profiler/Private/Widgets/SProfilerThreadView.cpp b/Engine/Source/Developer/Profiler/Private/Widgets/SProfilerThreadView.cpp
--- a/Engine/Source/Developer/Profiler/Private/Widgets/SProfilerThreadView.cpp
+++ b/Engine/Source/Developer/Profiler/Private/Widgets/SProfilerThreadView.cpp
@@ -464,7 +464,7 @@ void SProfilerThreadView::DrawText( const FString& Text, const FSlateFontInfo& F
 {
 	check( PaintState );
 
-	if( ShadowOffset.Size() > 0.0f )
+	if( ShadowOffset.SizeSquared() > 0.0f )
 	{
 		FSlateDrawElement::MakeText
 		(
diff --git a/Engine/Source/Editor/Cascade/Private/CascadeEmitterCanvasClient.cpp b/Engine/Source/Editor/Cascade/Private/CascadeEmitterCanvasClient.cpp
--- a/Engine/Source/Editor/Cascade/Private/CascadeEmitterCanvasClient.cpp
+++ b/Engine/Source/Editor/Cascade/Private/CascadeEmitterCanvasClient.cpp
@@ -671,7 +671,7 @@ void FCascadeEmitterCanvasClient::CapturedMouseMove(FViewport* Viewport, int32 X
 
 		int32 MoveThresh = CascadePtr.Pin()->GetEditorOptions() ? CascadePtr.Pin()->GetEditorOptions()->Cascade_MouseMoveThreshold : 4;
 		MoveThresh = FMath::Max<int32>(4,MoveThresh);
-		if (TotalMouseMove.Size() > MoveThresh)
+		if (TotalMouseMove.SizeSquared() > FMath::Square(MoveThresh))
 		{
 			if ((SelectedModuleIndex == INDEX_REQUIREDMODULE) ||
 				(SelectedModuleIndex == INDEX_SPAWNMODULE))
diff --git a/Engine/Source/Editor/GeometryMode/Private/GeometryModifiers.cpp b/Engine/Source/Editor/GeometryMode/Private/GeometryModifiers.cpp
--- a/Engine/Source/Editor/GeometryMode/Private/GeometryModifiers.cpp
+++ b/Engine/Source/Editor/GeometryMode/Private/GeometryModifiers.cpp
@@ -811,7 +811,7 @@ void UGeomModifier_Extrude::Apply(int32 InLength, int32 InSegments)
 			{
 				FORCEINLINE bool operator()( const FPoly& A, const FPoly& B ) const
 				{
-					return (B.Normal - A.Normal).Size() < 0;
+					return (B.Normal - A.Normal).SizeSquared() < 0.f;
 				}
 			};
 			Polygons.Sort( FCompareFPolyNormal() );
diff --git a/Engine/Source/Editor/GraphEditor/Private/SNodePanel.cpp b/Engine/Source/Editor/GraphEditor/Private/SNodePanel.cpp
--- a/Engine/Source/Editor/GraphEditor/Private/SNodePanel.cpp
+++ b/Engine/Source/Editor/GraphEditor/Private/SNodePanel.cpp
@@ -1474,7 +1474,7 @@ bool SNodePanel::ScrollToLocation(const FGeometry& MyGeometry, FVector2D Desired
 	ViewOffset = NewPosition - HalfOFScreenInGraphSpace;
 
 	// If within 1 pixel of target, stop interpolating
-	return ((NewPosition - DesiredCenterPosition).Size() < 1.f);
+	return ((NewPosition - DesiredCenterPosition).SizeSquared() < 1.f);
 }
 
 bool SNodePanel::ZoomToLocation(const FVector2D& CurrentSizeWithoutZoom, const FVector2D& DesiredSize, bool bDoneScrolling)
diff --git a/Engine/Source/Editor/GraphEditor/Public/SNodePanel.h b/Engine/Source/Editor/GraphEditor/Public/SNodePanel.h
--- a/Engine/Source/Editor/GraphEditor/Public/SNodePanel.h
+++ b/Engine/Source/Editor/GraphEditor/Public/SNodePanel.h
@@ -392,7 +392,7 @@ public:
 				const FArrangedWidget& CurWidget = ArrangedChildren[ ChildIndex ];
 				FSlateRect ChildClipRect = MyClippingRect.IntersectionWith(CurWidget.Geometry.GetClippingRect());
 
-				if( ChildClipRect.GetSize().Size() > 0 )
+				if( ChildClipRect.GetSize().SizeSquared() > 0.f )
 				{
 					const FSlateRect ClipRect = Children[ ChildIndex ].Zone == ENodeZone::Center ? ChildClipRect : MyClippingRect;
 					const int32 CurWidgetsMaxLayerId = CurWidget.Widget->Paint( Args.WithNewParent( this ), CurWidget.Geometry, ClipRect, OutDrawElements, LayerId, InWidgetStyle, ShouldBeEnabled( bParentEnabled ));
diff --git a/Engine/Source/Editor/LandscapeEditor/Private/LandscapeEdModeXYOffset.cpp b/Engine/Source/Editor/LandscapeEditor/Private/LandscapeEdModeXYOffset.cpp
--- a/Engine/Source/Editor/LandscapeEditor/Private/LandscapeEdModeXYOffset.cpp
+++ b/Engine/Source/Editor/LandscapeEditor/Private/LandscapeEdModeXYOffset.cpp
@@ -533,7 +533,7 @@ public:
 				float Sin = FMath::Abs(TransformedLocal.Y) / TransformedLocal.Size();
 				float RatioX = FalloffRadius > 0.0f ? 1.0f - FMath::Clamp((FMath::Abs(TransformedLocal.X) - Cos*SquareRadius) / FalloffRadius, 0.0f, 1.0f) : 1.0f;
 				float RatioY = FalloffRadius > 0.0f ? 1.0f - FMath::Clamp((FMath::Abs(TransformedLocal.Y) - Sin*SquareRadius) / FalloffRadius, 0.0f, 1.0f) : 1.0f;
-				float Ratio = TransformedLocal.Size() > SquareRadius ? RatioX * RatioY : 1.0f; //TransformedLocal.X / LW * TransformedLocal.Y / LW;
+				float Ratio = TransformedLocal.SizeSquared() > FMath::Square(SquareRadius) ? RatioX * RatioY : 1.0f; //TransformedLocal.X / LW * TransformedLocal.Y / LW;
 				float PaintAmount = Ratio*Ratio*(3 - 2 * Ratio);
 
 				XYOffsetVectorData[Index] = FMath::Lerp(XYOffsetVectorData[Index], NewXYOffset[Index], PaintAmount);
diff --git a/Engine/Source/Editor/Sequencer/Private/Tools/DelayedDrag.h b/Engine/Source/Editor/Sequencer/Private/Tools/DelayedDrag.h
--- a/Engine/Source/Editor/Sequencer/Private/Tools/DelayedDrag.h
+++ b/Engine/Source/Editor/Sequencer/Private/Tools/DelayedDrag.h
@@ -27,7 +27,7 @@ public:
 	{
 		if (!bHasInitiatedDrag && MouseEvent.IsMouseButtonDown(EffectiveKey))
 		{
-			DistanceDragged += FMath::Abs( MouseEvent.GetCursorDelta().Size() );
+			DistanceDragged += MouseEvent.GetCursorDelta().Size();
 			if (!bHasInitiatedDrag && DistanceDragged > FSlateApplication::Get().GetDragTriggerDistance())
 			{
 				ForceDragStart();
diff --git a/Engine/Source/Editor/UMGEditor/Private/Designer/SDesignSurface.cpp b/Engine/Source/Editor/UMGEditor/Private/Designer/SDesignSurface.cpp
--- a/Engine/Source/Editor/UMGEditor/Private/Designer/SDesignSurface.cpp
+++ b/Engine/Source/Editor/UMGEditor/Private/Designer/SDesignSurface.cpp
@@ -379,7 +379,7 @@ bool SDesignSurface::ScrollToLocation(const FGeometry& MyGeometry, FVector2D Des
 	ViewOffset = NewPosition - HalfOFScreenInGraphSpace;
 
 	// If within 1 pixel of target, stop interpolating
-	return ( ( NewPosition - DesiredCenterPosition ).Size() < 1.f );
+	return ( ( NewPosition - DesiredCenterPosition ).SizeSquared() < 1.f );
 }
 
 bool SDesignSurface::ZoomToLocation(const FVector2D& CurrentSizeWithoutZoom, const FVector2D& DesiredSize, bool bDoneScrolling)
diff --git a/Engine/Source/Editor/UMGEditor/Private/Designer/STransformHandle.cpp b/Engine/Source/Editor/UMGEditor/Private/Designer/STransformHandle.cpp
--- a/Engine/Source/Editor/UMGEditor/Private/Designer/STransformHandle.cpp
+++ b/Engine/Source/Editor/UMGEditor/Private/Designer/STransformHandle.cpp
@@ -314,7 +314,7 @@ ETransformAction STransformHandle::ComputeActionAtLocation(const FGeometry& MyGe
 {
 	FVector2D LocalPosition = MyGeometry.AbsoluteToLocal(MouseEvent.GetScreenSpacePosition());
 	FVector2D GrabOriginOffset = LocalPosition - DragOrigin;
-	if ( GrabOriginOffset.Size() < 6 )
+	if ( GrabOriginOffset.SizeSquared() < 36.f )
 	{
 		return ETransformAction::Primary;
 	}
diff --git a/Engine/Source/Editor/UnrealEd/Private/CameraController.cpp b/Engine/Source/Editor/UnrealEd/Private/CameraController.cpp
--- a/Engine/Source/Editor/UnrealEd/Private/CameraController.cpp
+++ b/Engine/Source/Editor/UnrealEd/Private/CameraController.cpp
@@ -245,14 +245,14 @@ void FEditorCameraController::UpdatePosition( const FCameraControllerUserImpulse
 
 
 	// Constrain maximum movement speed
-	if( MovementVelocity.Size() > Config.MaximumMovementSpeed * MovementSpeedScale )
+	if( MovementVelocity.SizeSquared() > FMath::Square(Config.MaximumMovementSpeed * MovementSpeedScale) )
 	{
 		MovementVelocity = MovementVelocity.GetUnsafeNormal() * Config.MaximumMovementSpeed * MovementSpeedScale;
 	}
 
 
 	// Clamp velocity to a reasonably small number
-	if( MovementVelocity.Size() < KINDA_SMALL_NUMBER )
+	if( MovementVelocity.SizeSquared() < FMath::Square(KINDA_SMALL_NUMBER) )
 	{
 		MovementVelocity = FVector::ZeroVector;
 	}
diff --git a/Engine/Source/Editor/UnrealEd/Private/Editor/ActorPositioning.cpp b/Engine/Source/Editor/UnrealEd/Private/Editor/ActorPositioning.cpp
--- a/Engine/Source/Editor/UnrealEd/Private/Editor/ActorPositioning.cpp
+++ b/Engine/Source/Editor/UnrealEd/Private/Editor/ActorPositioning.cpp
@@ -111,14 +111,14 @@ FActorPositionTraceResult FActorPositioning::TraceWorldForPosition(const UWorld&
 		}
 
 		// Go through all hits and find closest
-		float ClosestHitDistance = TNumericLimits<float>::Max();
+		float ClosestHitDistanceSqr = TNumericLimits<float>::Max();
 
 		for (const auto& Hit : Hits)
 		{
-			const float DistanceToHit = (Hit.ImpactPoint - RayStart).Size();
-			if (DistanceToHit < ClosestHitDistance)
+			const float DistanceToHitSqr = (Hit.ImpactPoint - RayStart).SizeSquared();
+			if (DistanceToHitSqr < ClosestHitDistanceSqr)
 			{
-				ClosestHitDistance = DistanceToHit;
+				ClosestHitDistanceSqr = DistanceToHitSqr;
 				Results.Location = Hit.Location;
 				Results.SurfaceNormal = Hit.Normal.GetSafeNormal();
 				Results.State = FActorPositionTraceResult::HitSuccess;
diff --git a/Engine/Source/Editor/UnrealEd/Private/GeomTools.cpp b/Engine/Source/Editor/UnrealEd/Private/GeomTools.cpp
--- a/Engine/Source/Editor/UnrealEd/Private/GeomTools.cpp
+++ b/Engine/Source/Editor/UnrealEd/Private/GeomTools.cpp
@@ -263,24 +263,24 @@ static float EdgeMatchTolerance = 0.01f;
 /** Util to look for best next edge start from Start. Returns false if no good next edge found. Edge is removed from InEdgeSet when found. */
 static bool FindNextEdge(FUtilEdge2D& OutNextEdge, const FVector2D& Start, TArray<FUtilEdge2D>& InEdgeSet)
 {
-	float ClosestDist = BIG_NUMBER;
+	float ClosestDistSqr = BIG_NUMBER;
 	FUtilEdge2D OutEdge;
 	int32 OutEdgeIndex = INDEX_NONE;
 	// Search set of edges for one that starts closest to Start
 	for(int32 i=0; i<InEdgeSet.Num(); i++)
 	{
-		float Dist = (InEdgeSet[i].V0 - Start).Size();
-		if(Dist < ClosestDist)
+		float DistSqr = (InEdgeSet[i].V0 - Start).SizeSquared();
+		if(DistSqr < ClosestDistSqr)
 		{
-			ClosestDist = Dist;
+			ClosestDistSqr = DistSqr;
 			OutNextEdge = InEdgeSet[i];
 			OutEdgeIndex = i;
 		}
 
-		Dist = (InEdgeSet[i].V1 - Start).Size();
-		if(Dist < ClosestDist)
+		DistSqr = (InEdgeSet[i].V1 - Start).SizeSquared();
+		if(DistSqr < ClosestDistSqr)
 		{
-			ClosestDist = Dist;
+			ClosestDistSqr = DistSqr;
 			OutNextEdge = InEdgeSet[i];
 			Swap(OutNextEdge.V0, OutNextEdge.V1);
 			OutEdgeIndex = i;
@@ -288,17 +288,15 @@ static bool FindNextEdge(FUtilEdge2D& OutNextEdge, const FVector2D& Start, TArra
 	}
 
 	// If next edge starts close enough return it
-	if(ClosestDist < EdgeMatchTolerance)
+	if(ClosestDistSqr < FMath::Square(EdgeMatchTolerance))
 	{
 		check(OutEdgeIndex != INDEX_NONE);
 		InEdgeSet.RemoveAt(OutEdgeIndex);
 		return true;
 	}
+
 	// No next edge found.
-	else
-	{
-		return false;
-	}
+	return false;
 }
 
 /** 
@@ -315,14 +313,14 @@ static void FixPolyWinding(FUtilPoly2D& Poly)
 		int32 BIndex = i;
 		int32 CIndex = (i+1)%Poly.Verts.Num();
 
-		float ABDist = (Poly.Verts[BIndex].Pos - Poly.Verts[AIndex].Pos).Size();
+		float ABDistSqr = (Poly.Verts[BIndex].Pos - Poly.Verts[AIndex].Pos).SizeSquared();
 		FVector2D ABEdge = (Poly.Verts[BIndex].Pos - Poly.Verts[AIndex].Pos).GetSafeNormal();
 
-		float BCDist = (Poly.Verts[CIndex].Pos - Poly.Verts[BIndex].Pos).Size();
+		float BCDistSqr = (Poly.Verts[CIndex].Pos - Poly.Verts[BIndex].Pos).SizeSquared();
 		FVector2D BCEdge = (Poly.Verts[CIndex].Pos - Poly.Verts[BIndex].Pos).GetSafeNormal();
 
 		// See if points are co-incident or edges are co-linear - if so, remove.
-		if(ABDist < 0.01f || BCDist < 0.01f || ABEdge.Equals(BCEdge, 0.01f))
+		if(ABDistSqr < 0.0001f || BCDistSqr < 0.0001f || ABEdge.Equals(BCEdge, 0.01f))
 		{
 			Poly.Verts.RemoveAt(i);
 		}
@@ -373,8 +371,8 @@ void Buid2DPolysFromEdges(TArray<FUtilPoly2D>& OutPolys, const TArray<FUtilEdge2
 		}
 
 		// After walking edges see if we have a closed polygon.
-		float CloseDist = (NewPoly.Verts[0].Pos - NewPoly.Verts[ NewPoly.Verts.Num()-1 ].Pos).Size();
-		if(NewPoly.Verts.Num() >= 4 && CloseDist < EdgeMatchTolerance)
+		float CloseDistSqr = (NewPoly.Verts[0].Pos - NewPoly.Verts[ NewPoly.Verts.Num()-1 ].Pos).SizeSquared();
+		if(NewPoly.Verts.Num() >= 4 && CloseDistSqr < FMath::Square(EdgeMatchTolerance))
 		{
 			// Remove last vert - its basically a duplicate of the first.
 			NewPoly.Verts.RemoveAt( NewPoly.Verts.Num()-1 );
diff --git a/Engine/Source/Editor/UnrealEd/Private/PhysicsAssetUtils.cpp b/Engine/Source/Editor/UnrealEd/Private/PhysicsAssetUtils.cpp
--- a/Engine/Source/Editor/UnrealEd/Private/PhysicsAssetUtils.cpp
+++ b/Engine/Source/Editor/UnrealEd/Private/PhysicsAssetUtils.cpp
@@ -339,7 +339,7 @@ bool CreateCollisionFromBone( UBodySetup* bs, USkeletalMesh* skelMesh, int32 Bon
 			FVector ChildPos = RelTM.GetOrigin();
 
 			// Check that child is not on top of parent. If it is - we can't make an orientation
-			if(ChildPos.Size() > KINDA_SMALL_NUMBER)
+			if(ChildPos.SizeSquared() > FMath::Square(KINDA_SMALL_NUMBER))
 			{
 				// ZAxis for collision geometry lies down axis to child bone.
 				FVector ZAxis = ChildPos.GetSafeNormal();
diff --git a/Engine/Source/Editor/UnrealEd/Private/SCurveEditor.cpp b/Engine/Source/Editor/UnrealEd/Private/SCurveEditor.cpp
--- a/Engine/Source/Editor/UnrealEd/Private/SCurveEditor.cpp
+++ b/Engine/Source/Editor/UnrealEd/Private/SCurveEditor.cpp
@@ -1494,7 +1494,7 @@ void SCurveEditor::TryStartDrag(const FGeometry& InMyGeometry, const FPointerEve
 
 	FVector2D MousePosition = InMyGeometry.AbsoluteToLocal(InMouseEvent.GetScreenSpacePosition());
 	FVector2D DragVector = MousePosition - MouseDownLocation;
-	if (DragVector.Size() >= DragThreshold)
+	if (DragVector.SizeSquared() >= FMath::Square(DragThreshold))
 	{
 		if (bShiftDown)
 		{
diff --git a/Engine/Source/Editor/UnrealEd/Private/Tests/CollisionAutomationTests.cpp b/Engine/Source/Editor/UnrealEd/Private/Tests/CollisionAutomationTests.cpp
--- a/Engine/Source/Editor/UnrealEd/Private/Tests/CollisionAutomationTests.cpp
+++ b/Engine/Source/Editor/UnrealEd/Private/Tests/CollisionAutomationTests.cpp
@@ -98,9 +98,8 @@ namespace CollisionAutomationTests
 
 	void CheckVector( FVector ResultVector, FVector ExpectedVector, FString TestName, FString ParameterName, int32 TestIndex, float Tolerance = KINDA_SMALL_NUMBER )
 	{
-		FVector Delta = ExpectedVector - ResultVector;
-		float Diff = Delta.Size();
-		if (Diff > Tolerance)
+		const FVector Delta = ExpectedVector - ResultVector;
+		if (Delta.SizeSquared() > FMath::Square(Tolerance))
 		{
 			//UE_LOG(CollisionAutomationTestLog, Log, TEXT("%d:HitResult=(%s)"), iTest+1, *OutHits[iHits].ToString());
 			TestBase->AddError(FString::Printf(TEXT("Test %d:%s %s mismatch. Should be %s but is actually %s."), TestIndex, *TestName, *ParameterName, *ExpectedVector.ToString(), *ResultVector.ToString()));
diff --git a/Engine/Source/Editor/WorldBrowser/Private/Tiles/SWorldComposition.cpp b/Engine/Source/Editor/WorldBrowser/Private/Tiles/SWorldComposition.cpp
--- a/Engine/Source/Editor/WorldBrowser/Private/Tiles/SWorldComposition.cpp
+++ b/Engine/Source/Editor/WorldBrowser/Private/Tiles/SWorldComposition.cpp
@@ -456,7 +456,7 @@ public:
 							{
 								const FVector2D AnchorNodeOldPos = NodeBeingDragged->GetPosition();
 								const FVector2D DeltaPos = AnchorNodeNewPos - AnchorNodeOldPos;
-								if (DeltaPos.Size() > KINDA_SMALL_NUMBER)
+								if (DeltaPos.SizeSquared() > FMath::Square(KINDA_SMALL_NUMBER))
 								{
 									MoveSelectedNodes(NodeBeingDragged, AnchorNodeNewPos);
 								}
diff --git a/Engine/Source/Editor/WorldBrowser/Private/Tiles/SWorldTileItem.cpp b/Engine/Source/Editor/WorldBrowser/Private/Tiles/SWorldTileItem.cpp
--- a/Engine/Source/Editor/WorldBrowser/Private/Tiles/SWorldTileItem.cpp
+++ b/Engine/Source/Editor/WorldBrowser/Private/Tiles/SWorldTileItem.cpp
@@ -244,7 +244,7 @@ FSlateRect SWorldTileItem::GetItemRect() const
 TSharedPtr<IToolTip> SWorldTileItem::GetToolTip()
 {
 	// Hide tooltip in case item is being dragged now
-	if (TileModel->GetLevelTranslationDelta().Size() > KINDA_SMALL_NUMBER)
+	if (TileModel->GetLevelTranslationDelta().SizeSquared() > FMath::Square(KINDA_SMALL_NUMBER))
 	{
 		return NULL;
 	}
diff --git a/Engine/Source/Editor/WorldBrowser/Private/Tiles/WorldTileCollectionModel.cpp b/Engine/Source/Editor/WorldBrowser/Private/Tiles/WorldTileCollectionModel.cpp
--- a/Engine/Source/Editor/WorldBrowser/Private/Tiles/WorldTileCollectionModel.cpp
+++ b/Engine/Source/Editor/WorldBrowser/Private/Tiles/WorldTileCollectionModel.cpp
@@ -909,7 +909,7 @@ void FWorldTileCollectionModel::UpdateStreamingPreview(FVector2D InLocation, boo
 	{
 		FVector NewPreviewLocation = FVector(InLocation, 0);
 		
-		if ((PreviewLocation-NewPreviewLocation).Size() > KINDA_SMALL_NUMBER)
+		if ((PreviewLocation-NewPreviewLocation).SizeSquared() > FMath::Square(KINDA_SMALL_NUMBER))
 		{
 			PreviewLocation = NewPreviewLocation;
 			PreviewVisibleTiles.Empty();
diff --git a/Engine/Source/Programs/UnrealLightmass/Private/Lighting/PrecomputedVisibility.cpp b/Engine/Source/Programs/UnrealLightmass/Private/Lighting/PrecomputedVisibility.cpp
--- a/Engine/Source/Programs/UnrealLightmass/Private/Lighting/PrecomputedVisibility.cpp
+++ b/Engine/Source/Programs/UnrealLightmass/Private/Lighting/PrecomputedVisibility.cpp
@@ -573,10 +573,10 @@ void FStaticLightingSystem::SetupPrecomputedVisibility()
 				MeshBounds += VisibilityMeshes[VisibilityMeshIndex].Meshes[OriginalMeshIndex]->BoundingBox;
 			}
 
-			const float MeshBoundingRadius = MeshBounds.GetExtent().Size();
+			const float MeshBoundingRadiusSqr = MeshBounds.GetExtent().SizeSquared();
 
 			// Only put the mesh in a group if its radius is small enough to keep the group effective
-			bool bPutInGroup = MeshBoundingRadius < GridCellBoundingRadius * MeshGroupingCellRadiusThreshold;
+			bool bPutInGroup = MeshBoundingRadiusSqr < FMath::Square(GridCellBoundingRadius * MeshGroupingCellRadiusThreshold);
 
 			if (bPutInGroup)
 			{
diff --git a/Engine/Source/Runtime/AnimGraphRuntime/Private/BoneControllers/AnimNode_SpringBone.cpp b/Engine/Source/Runtime/AnimGraphRuntime/Private/BoneControllers/AnimNode_SpringBone.cpp
--- a/Engine/Source/Runtime/AnimGraphRuntime/Private/BoneControllers/AnimNode_SpringBone.cpp
+++ b/Engine/Source/Runtime/AnimGraphRuntime/Private/BoneControllers/AnimNode_SpringBone.cpp
@@ -166,7 +166,7 @@ void FAnimNode_SpringBone::EvaluateBoneTransforms(USkeletalMeshComponent* SkelCo
 		{
 			FVector CurrentDisp = BoneLocation - TargetPos;
 			// Too far away - project back onto sphere around target.
-			if (CurrentDisp.Size() > MaxDisplacement)
+			if (CurrentDisp.SizeSquared() > FMath::Square(MaxDisplacement))
 			{
 				FVector DispDir = CurrentDisp.GetSafeNormal();
 				BoneLocation = TargetPos + (MaxDisplacement * DispDir);
diff --git a/Engine/Source/Runtime/AnimGraphRuntime/Private/BoneControllers/AnimNode_TwoBoneIK.cpp b/Engine/Source/Runtime/AnimGraphRuntime/Private/BoneControllers/AnimNode_TwoBoneIK.cpp
--- a/Engine/Source/Runtime/AnimGraphRuntime/Private/BoneControllers/AnimNode_TwoBoneIK.cpp
+++ b/Engine/Source/Runtime/AnimGraphRuntime/Private/BoneControllers/AnimNode_TwoBoneIK.cpp
@@ -120,11 +120,11 @@ void FAnimNode_TwoBoneIK::EvaluateBoneTransforms(USkeletalMeshComponent* SkelCom
 
 	FVector	JointTargetPos = JointTargetTransform.GetTranslation();
 	FVector JointTargetDelta = JointTargetPos - RootPos;
-	float JointTargetLength = JointTargetDelta.Size();
+	const float JointTargetLengthSqr = JointTargetDelta.SizeSquared();
 
 	// Same check as above, to cover case when JointTarget position is the same as RootPos.
 	FVector JointPlaneNormal, JointBendDir;
-	if (JointTargetLength < (float)KINDA_SMALL_NUMBER)
+	if (JointTargetLengthSqr < FMath::Square((float)KINDA_SMALL_NUMBER))
 	{
 		JointBendDir = FVector(0,1,0);
 		JointPlaneNormal = FVector(0,0,1);
@@ -135,7 +135,7 @@ void FAnimNode_TwoBoneIK::EvaluateBoneTransforms(USkeletalMeshComponent* SkelCom
 
 		// If we are trying to point the limb in the same direction that we are supposed to displace the joint in, 
 		// we have to just pick 2 random vector perp to DesiredDir and each other.
-		if (JointPlaneNormal.Size() < (float)KINDA_SMALL_NUMBER)
+		if (JointPlaneNormal.SizeSquared() < FMath::Square((float)KINDA_SMALL_NUMBER))
 		{
 			DesiredDir.FindBestAxisVectors(JointPlaneNormal, JointBendDir);
 		}
diff --git a/Engine/Source/Runtime/Core/Private/Math/UnrealMath.cpp b/Engine/Source/Runtime/Core/Private/Math/UnrealMath.cpp
--- a/Engine/Source/Runtime/Core/Private/Math/UnrealMath.cpp
+++ b/Engine/Source/Runtime/Core/Private/Math/UnrealMath.cpp
@@ -1721,7 +1721,7 @@ static bool ComputeProjectedSphereShaft(
 uint32 FMath::ComputeProjectedSphereScissorRect(FIntRect& InOutScissorRect, FVector SphereOrigin, float Radius, FVector ViewOrigin, const FMatrix& ViewMatrix, const FMatrix& ProjMatrix)
 {
 	// Calculate a scissor rectangle for the light's radius.
-	if((SphereOrigin - ViewOrigin).Size() > Radius)
+	if((SphereOrigin - ViewOrigin).SizeSquared() > FMath::Square(Radius))
 	{
 		FVector LightVector = ViewMatrix.TransformPosition(SphereOrigin);
 
@@ -2575,13 +2575,13 @@ void FVector::GenerateClusterCenters(TArray<FVector>& Clusters, const TArray<FVe
 
 			// Iterate over all clusters to find closes one
 			int32 NearestClusterIndex = INDEX_NONE;
-			float NearestClusterDist = BIG_NUMBER;
+			float NearestClusterDistSqr = BIG_NUMBER;
 			for(int32 j=0; j<Clusters.Num() ; j++)
 			{
-				const float Dist = (Pos - Clusters[j]).Size();
-				if(Dist < NearestClusterDist)
+				const float DistSqr = (Pos - Clusters[j]).SizeSquared();
+				if(DistSqr < NearestClusterDistSqr)
 				{
-					NearestClusterDist = Dist;
+					NearestClusterDistSqr = DistSqr;
 					NearestClusterIndex = j;
 				}
 			}
diff --git a/Engine/Source/Runtime/Engine/Private/Animation/AnimCompress_PerTrackCompression.cpp b/Engine/Source/Runtime/Engine/Private/Animation/AnimCompress_PerTrackCompression.cpp
--- a/Engine/Source/Runtime/Engine/Private/Animation/AnimCompress_PerTrackCompression.cpp
+++ b/Engine/Source/Runtime/Engine/Private/Animation/AnimCompress_PerTrackCompression.cpp
@@ -466,7 +466,7 @@ protected:
 		const bool bHasY = (FMath::Abs(KeyBounds.Max.Y) >= ZeroingThreshold) || (FMath::Abs(KeyBounds.Min.Y) >= ZeroingThreshold);
 		const bool bHasZ = (FMath::Abs(KeyBounds.Max.Z) >= ZeroingThreshold) || (FMath::Abs(KeyBounds.Min.Z) >= ZeroingThreshold);
 
-		if ((!bHasX && !bHasY && !bHasZ) || (Range.Size() > 4.0f))
+		if ((!bHasX && !bHasY && !bHasZ) || (Range.SizeSquared() > 16.0f))
 		{
 			// If there are no components, then there is no point in using this over the identity encoding
 			// If the range is insane, error out early (error metric will be high)
diff --git a/Engine/Source/Runtime/Engine/Private/Animation/AnimCompress_RemoveLinearKeys.cpp b/Engine/Source/Runtime/Engine/Private/Animation/AnimCompress_RemoveLinearKeys.cpp
--- a/Engine/Source/Runtime/Engine/Private/Animation/AnimCompress_RemoveLinearKeys.cpp
+++ b/Engine/Source/Runtime/Engine/Private/Animation/AnimCompress_RemoveLinearKeys.cpp
@@ -749,8 +749,8 @@ void UAnimCompress_RemoveLinearKeys::ProcessAnimationTracks(
 									Adjustment= EnforceShortestArc(FQuat::Identity, Adjustment);
 
 									const FVector Test = Adjustment.RotateVector(CurrentHeading);
-									const float Delta = (Test - DesiredHeading).Size();
-									if (Delta < 0.001f)
+									const float DeltaSqr = (Test - DesiredHeading).SizeSquared();
+									if (DeltaSqr < FMath::Square(0.001f))
 									{
 										FQuat NewKey = Adjustment * Key;
 										NewKey.Normalize();
diff --git a/Engine/Source/Runtime/Engine/Private/Animation/AnimSequence.cpp b/Engine/Source/Runtime/Engine/Private/Animation/AnimSequence.cpp
--- a/Engine/Source/Runtime/Engine/Private/Animation/AnimSequence.cpp
+++ b/Engine/Source/Runtime/Engine/Private/Animation/AnimSequence.cpp
@@ -1116,7 +1116,7 @@ void UAnimSequence::GetBonePose(FCompactPose& OutPose, FBlendedCurve& OutCurve,
 
 			// @todo - precache that in FBoneContainer when we have SkeletonIndex->TrackIndex mapping. So we can just apply scale right away.
 			float const SourceTranslationLength = AuthoredOnRefSkeleton[SkeletonBoneIndex].GetTranslation().Size();
-			if (FMath::Abs(SourceTranslationLength) > KINDA_SMALL_NUMBER)
+			if (SourceTranslationLength > KINDA_SMALL_NUMBER)
 			{
 				float const TargetTranslationLength = RequiredBones.GetRefPoseTransform(BoneIndex).GetTranslation().Size();
 				OutPose[BoneIndex].ScaleTranslation(TargetTranslationLength / SourceTranslationLength);
@@ -1247,7 +1247,7 @@ void UAnimSequence::RetargetBoneTransform(FTransform& BoneTransform, const int32
 		// @todo - precache that in FBoneContainer when we have SkeletonIndex->TrackIndex mapping. So we can just apply scale right away.
 		const TArray<FTransform>& SkeletonRefPoseArray = GetSkeleton()->GetRefLocalPoses(RetargetSource);
 		const float SourceTranslationLength = SkeletonRefPoseArray[SkeletonBoneIndex].GetTranslation().Size();
-		if (FMath::Abs(SourceTranslationLength) > KINDA_SMALL_NUMBER)
+		if (SourceTranslationLength > KINDA_SMALL_NUMBER)
 		{
 			const float TargetTranslationLength = RequiredBones.GetRefPoseTransform(BoneIndex).GetTranslation().Size();
 			BoneTransform.ScaleTranslation(TargetTranslationLength / SourceTranslationLength);
@@ -1594,7 +1594,7 @@ bool UAnimSequence::CompressRawAnimSequenceTrack(FRawAnimSequenceTrack& RawTrack
 		bool bFramesIdentical = true;
 		for(int32 j=1; j<RawTrack.PosKeys.Num() && bFramesIdentical; j++)
 		{
-			if( (FirstPos - RawTrack.PosKeys[j]).Size() > MaxPosDiff )
+			if( (FirstPos - RawTrack.PosKeys[j]).SizeSquared() > FMath::Square(MaxPosDiff) )
 			{
 				bFramesIdentical = false;
 			}
@@ -1642,7 +1642,7 @@ bool UAnimSequence::CompressRawAnimSequenceTrack(FRawAnimSequenceTrack& RawTrack
 		bool bFramesIdentical = true;
 		for(int32 j=1; j<RawTrack.ScaleKeys.Num() && bFramesIdentical; j++)
 		{
-			if( (FirstScale - RawTrack.ScaleKeys[j]).Size() > MaxScaleDiff )
+			if( (FirstScale - RawTrack.ScaleKeys[j]).SizeSquared() > FMath::Square(MaxScaleDiff) )
 			{
 				bFramesIdentical = false;
 			}
diff --git a/Engine/Source/Runtime/Engine/Private/Collision/KAggregateGeom.cpp b/Engine/Source/Runtime/Engine/Private/Collision/KAggregateGeom.cpp
--- a/Engine/Source/Runtime/Engine/Private/Collision/KAggregateGeom.cpp
+++ b/Engine/Source/Runtime/Engine/Private/Collision/KAggregateGeom.cpp
@@ -182,20 +182,20 @@ static bool EnsureHullIsValid(TArray<FVector>& InVerts)
 	const FVector FirstVert = InVerts[0];
 
 	// Now find vert furthest from this one.
-	float FurthestDist = 0.f;
+	float FurthestDistSqr = 0.f;
 	int32 FurthestVertIndex = INDEX_NONE;
 	for(int32 i=1; i<InVerts.Num(); i++)
 	{
-		const float TestDist = (InVerts[i] - FirstVert).Size();
-		if(TestDist > FurthestDist)
+		const float TestDistSqr = (InVerts[i] - FirstVert).SizeSquared();
+		if(TestDistSqr > FurthestDistSqr)
 		{
-			FurthestDist = TestDist;
+			FurthestDistSqr = TestDistSqr;
 			FurthestVertIndex = i;
 		}
 	}
 
 	// If smallest dimension is too small - hull is invalid.
-	if(FurthestVertIndex == INDEX_NONE || FurthestDist < MIN_HULL_VALID_DIMENSION)
+	if(FurthestVertIndex == INDEX_NONE || FurthestDistSqr < FMath::Square(MIN_HULL_VALID_DIMENSION))
 	{
 		return false;
 	}
diff --git a/Engine/Source/Runtime/Engine/Private/Components/SphereComponent.cpp b/Engine/Source/Runtime/Engine/Private/Components/SphereComponent.cpp
--- a/Engine/Source/Runtime/Engine/Private/Components/SphereComponent.cpp
+++ b/Engine/Source/Runtime/Engine/Private/Components/SphereComponent.cpp
@@ -106,8 +106,8 @@ FPrimitiveSceneProxy* USphereComponent::CreateSceneProxy()
 					const FLinearColor DrawSphereColor = GetViewSelectionColor(SphereColor, *View, IsSelected(), IsHovered(), false, IsIndividuallySelected() );
 
 					// Taking into account the min and maximum drawing distance
-					const float Distance = (View->ViewMatrices.ViewOrigin - LocalToWorld.GetOrigin()).Size();
-					if (Distance < GetMinDrawDistance() || Distance > GetMaxDrawDistance() )
+					const float DistanceSqr = (View->ViewMatrices.ViewOrigin - LocalToWorld.GetOrigin()).SizeSquared();
+					if (DistanceSqr < FMath::Square(GetMinDrawDistance()) || DistanceSqr > FMath::Square(GetMaxDrawDistance()) )
 					{
 						continue;
 					}
diff --git a/Engine/Source/Runtime/Engine/Private/ContentStreaming.cpp b/Engine/Source/Runtime/Engine/Private/ContentStreaming.cpp
--- a/Engine/Source/Runtime/Engine/Private/ContentStreaming.cpp
+++ b/Engine/Source/Runtime/Engine/Private/ContentStreaming.cpp
@@ -1431,7 +1431,7 @@ void IStreamingManager::SetupViewInfos( float DeltaTime )
 			bool bFound = false;
 			for ( int32 PrevView=0; PrevView < GPrevViewLocations.Num(); ++PrevView )
 			{
-				if ( (ViewInfo.ViewOrigin - GPrevViewLocations(PrevView).ViewOrigin).Size() < 100.0f )
+				if ( (ViewInfo.ViewOrigin - GPrevViewLocations(PrevView).ViewOrigin).SizeSquared() < 10000.0f )
 				{
 					bFound = true;
 					break;
@@ -1450,7 +1450,7 @@ void IStreamingManager::SetupViewInfos( float DeltaTime )
 			for ( int32 ViewIndex = 0; ViewIndex < CurrentViewInfos.Num(); ++ViewIndex )
 			{
 				FStreamingViewInfo& ViewInfo = CurrentViewInfos( ViewIndex );
-				if ( (ViewInfo.ViewOrigin - GPrevViewLocations(PrevView).ViewOrigin).Size() < 100.0f )
+				if ( (ViewInfo.ViewOrigin - GPrevViewLocations(PrevView).ViewOrigin).SizeSquared() < 10000.0f )
 				{
 					bFound = true;
 					break;
diff --git a/Engine/Source/Runtime/Engine/Private/GameFramework/TouchInterface.cpp b/Engine/Source/Runtime/Engine/Private/GameFramework/TouchInterface.cpp
--- a/Engine/Source/Runtime/Engine/Private/GameFramework/TouchInterface.cpp
+++ b/Engine/Source/Runtime/Engine/Private/GameFramework/TouchInterface.cpp
@@ -38,7 +38,7 @@ void UTouchInterface::Activate(TSharedPtr<SVirtualJoystick> VirtualJoystick)
 			SlateControl->Center = Control.Center;
 			SlateControl->VisualSize = Control.VisualSize;
 			SlateControl->ThumbSize = Control.ThumbSize;
-			if (Control.InputScale.Size() > DELTA)
+			if (Control.InputScale.SizeSquared() > FMath::Square(DELTA))
 			{
 				SlateControl->InputScale = Control.InputScale;
 			}
diff --git a/Engine/Source/Runtime/Engine/Private/GestureRecognizer.cpp b/Engine/Source/Runtime/Engine/Private/GestureRecognizer.cpp
--- a/Engine/Source/Runtime/Engine/Private/GestureRecognizer.cpp
+++ b/Engine/Source/Runtime/Engine/Private/GestureRecognizer.cpp
@@ -70,7 +70,7 @@ void FGestureRecognizer::DetectGestures(const FVector (&Touches)[EKeys::NUM_TOUC
 		else if (PreviousTouchCount >= 1 && TouchCount == 0)
 		{
 			// must be a fast flick
-			if (FlickTime < 0.25f && (FlickCurrent - AnchorPoints[0]).Size() > 100.f)
+			if (FlickTime < 0.25f && (FlickCurrent - AnchorPoints[0]).SizeSquared() > 10000.f)
 			{
 				// this is the angle from +X in screen space, meaning right is 0, up is 90, left is 180, down is 270
 				float Angle = FMath::Atan2(-(FlickCurrent.Y - AnchorPoints[0].Y), FlickCurrent.X - AnchorPoints[0].X) * 180.f / PI;
diff --git a/Engine/Source/Runtime/Engine/Private/MorphTools.cpp b/Engine/Source/Runtime/Engine/Private/MorphTools.cpp
--- a/Engine/Source/Runtime/Engine/Private/MorphTools.cpp
+++ b/Engine/Source/Runtime/Engine/Private/MorphTools.cpp
@@ -119,11 +119,11 @@ void UMorphTarget::CreateMorphMeshStreams( const FMorphMeshRawSource& BaseSource
 					FVector NormalDeltaZ (VTarget.TanZ - VBase.TanZ);
 
 					// check if position actually changed much
-					if( PositionDelta.Size() > THRESH_POINTS_ARE_NEAR || 
+					if( PositionDelta.SizeSquared() > FMath::Square(THRESH_POINTS_ARE_NEAR) || 
 						// since we can't get imported morphtarget normal from FBX
 						// we can't compare normal unless it's calculated
 						// this is special flag to ignore normal diff
-						( bCompareNormal && NormalDeltaZ.Size() > 0.1f) )
+						( bCompareNormal && NormalDeltaZ.SizeSquared() > 0.01f) )
 					{
 						// create a new entry
 						FVertexAnimDelta NewVertex;
diff --git a/Engine/Source/Runtime/Engine/Private/NetworkDriver.cpp b/Engine/Source/Runtime/Engine/Private/NetworkDriver.cpp
--- a/Engine/Source/Runtime/Engine/Private/NetworkDriver.cpp
+++ b/Engine/Source/Runtime/Engine/Private/NetworkDriver.cpp
@@ -2738,11 +2738,11 @@ void UNetDriver::DrawNetDriverDebug()
 		return;
 	}
 
-	const float CullDist = CVarNetDormancyDrawCullDistance.GetValueOnGameThread();
+	const float CullDistSqr = FMath::Square(CVarNetDormancyDrawCullDistance.GetValueOnGameThread());
 
 	for (FActorIterator It(GetWorld()); It; ++It)
 	{
-		if ((It->GetActorLocation() - LocalPlayer->LastViewLocation).Size() > CullDist)
+		if ((It->GetActorLocation() - LocalPlayer->LastViewLocation).SizeSquared() > CullDistSqr)
 		{
 			continue;
 		}
diff --git a/Engine/Source/Runtime/Engine/Private/NiagaraEffectRenderer.cpp b/Engine/Source/Runtime/Engine/Private/NiagaraEffectRenderer.cpp
--- a/Engine/Source/Runtime/Engine/Private/NiagaraEffectRenderer.cpp
+++ b/Engine/Source/Runtime/Engine/Private/NiagaraEffectRenderer.cpp
@@ -576,7 +576,7 @@ FNiagaraDynamicDataBase *NiagaraEffectRendererRibbon::GenerateVertexData(const F
 
 		const FVector ParticlePos = PosPtr[Index1];
 		FVector ParticleDir = PosPtr[Index2] - ParticlePos;
-		if (ParticleDir.Size() <= SMALL_NUMBER)
+		if (ParticleDir.SizeSquared() <= FMath::Square(SMALL_NUMBER))
 		{
 			ParticleDir = PrevDir*0.1f;
 		}
diff --git a/Engine/Source/Runtime/Engine/Private/Particles/ParticleComponents.cpp b/Engine/Source/Runtime/Engine/Private/Particles/ParticleComponents.cpp
--- a/Engine/Source/Runtime/Engine/Private/Particles/ParticleComponents.cpp
+++ b/Engine/Source/Runtime/Engine/Private/Particles/ParticleComponents.cpp
@@ -5149,26 +5149,21 @@ int32 UParticleSystemComponent::DetermineLODLevelForLocation(const FVector& Effe
 		}
 
 		// This will now put everything in LODLevel 0 (high detail) by default
-		float LODDistance = 0.0f;
-		if (PlayerViewLocations.Num())
+		float LODDistanceSqr = (PlayerViewLocations.Num() ? FMath::Square(WORLD_MAX) : 0.0f);
+		for (const FVector& ViewLocation : PlayerViewLocations)
 		{
-			LODDistance = WORLD_MAX;
-			int32 NumLocations = PlayerViewLocations.Num();
-			for (int32 i = 0; i < NumLocations; ++i)
+			const float DistanceToEffectSqr = FVector(ViewLocation - EffectLocation).SizeSquared();
+			if (DistanceToEffectSqr < LODDistanceSqr)
 			{
-				float DistanceToEffect = FVector(PlayerViewLocations[i] - EffectLocation).Size();
-				if (DistanceToEffect < LODDistance)
-				{
-					LODDistance = DistanceToEffect;
-				}
+				LODDistanceSqr = DistanceToEffectSqr;
 			}
 		}
 
 		// Find appropriate LOD based on distance
 		Retval = Template->LODDistances.Num() - 1;
-		for (int32 LODIdx = 1; LODIdx < Template->LODDistances.Num(); LODIdx++)
+		for (int32 LODIdx = 1; LODIdx < Template->LODDistances.Num(); ++LODIdx)
 		{
-			if (LODDistance < Template->LODDistances[LODIdx])
+			if (LODDistanceSqr < FMath::Square(Template->LODDistances[LODIdx]))
 			{
 				Retval = LODIdx - 1;
 				break;
diff --git a/Engine/Source/Runtime/Engine/Private/Particles/ParticleModules.cpp b/Engine/Source/Runtime/Engine/Private/Particles/ParticleModules.cpp
--- a/Engine/Source/Runtime/Engine/Private/Particles/ParticleModules.cpp
+++ b/Engine/Source/Runtime/Engine/Private/Particles/ParticleModules.cpp
@@ -3528,7 +3528,7 @@ void UParticleModuleAttractorLine::Update(FParticleEmitterInstance* Owner, int32
 	FVector Line = EndPoint1 - EndPoint0;
 
 	// if both end points are the same, we end up with NaNs in the results of the update
-	if (Line.Size() == 0.0f)
+	if (Line.SizeSquared() == 0.0f)
 	{
 		Line = FVector(SMALL_NUMBER, SMALL_NUMBER, SMALL_NUMBER);
 	}
diff --git a/Engine/Source/Runtime/Engine/Private/PhysicsEngine/BodyInstance.cpp b/Engine/Source/Runtime/Engine/Private/PhysicsEngine/BodyInstance.cpp
--- a/Engine/Source/Runtime/Engine/Private/PhysicsEngine/BodyInstance.cpp
+++ b/Engine/Source/Runtime/Engine/Private/PhysicsEngine/BodyInstance.cpp
@@ -1672,7 +1672,7 @@ FVector GetInitialLinearVelocity(const AActor* OwningActor, bool& bComponentAwak
 	{
 		InitialLinVel = OwningActor->GetVelocity();
 
-		if (InitialLinVel.Size() > KINDA_SMALL_NUMBER)
+		if (InitialLinVel.SizeSquared() > FMath::Square(KINDA_SMALL_NUMBER))
 		{
 			bComponentAwake = true;
 		}
diff --git a/Engine/Source/Runtime/Engine/Private/PhysicsEngine/PhysicsSpring.cpp b/Engine/Source/Runtime/Engine/Private/PhysicsEngine/PhysicsSpring.cpp
--- a/Engine/Source/Runtime/Engine/Private/PhysicsEngine/PhysicsSpring.cpp
+++ b/Engine/Source/Runtime/Engine/Private/PhysicsEngine/PhysicsSpring.cpp
@@ -60,8 +60,8 @@ UPrimitiveComponent* UPhysicsSpringComponent::GetSpringCollision(const FVector&
 	UPrimitiveComponent* CollidedComponent = nullptr;
 	
 	const FVector Delta = End - Start;
-	const float DeltaSize = Delta.Size();
-	if (DeltaSize > SMALL_NUMBER)
+	const float DeltaSizeSqr = Delta.SizeSquared();
+	if (DeltaSizeSqr > FMath::Square(SMALL_NUMBER))
 	{
 		if (bool bBlockingHit = World->SweepSingleByChannel(Hit, Start, End, FQuat::Identity, SpringChannel, FCollisionShape::MakeSphere(SpringRadius), QueryParams))
 		{
diff --git a/Engine/Source/Runtime/Engine/Private/PlayerCameraManager.cpp b/Engine/Source/Runtime/Engine/Private/PlayerCameraManager.cpp
--- a/Engine/Source/Runtime/Engine/Private/PlayerCameraManager.cpp
+++ b/Engine/Source/Runtime/Engine/Private/PlayerCameraManager.cpp
@@ -1241,7 +1241,7 @@ float APlayerCameraManager::CalcRadialShakeScale(APlayerCameraManager* Camera, F
 	else
 	{
 		// ignore OuterRadius and do a cliff falloff at InnerRadius
-		return ((Epicenter - POVLoc).Size() < InnerRadius) ? 1.f : 0.f;
+		return ((Epicenter - POVLoc).SizeSquared() < FMath::Square(InnerRadius)) ? 1.f : 0.f;
 	}
 }
 
diff --git a/Engine/Source/Runtime/Engine/Private/UserInterface/CanvasItem.cpp b/Engine/Source/Runtime/Engine/Private/UserInterface/CanvasItem.cpp
--- a/Engine/Source/Runtime/Engine/Private/UserInterface/CanvasItem.cpp
+++ b/Engine/Source/Runtime/Engine/Private/UserInterface/CanvasItem.cpp
@@ -830,7 +830,7 @@ void FCanvasTextItem::Draw( class FCanvas* InCanvas )
 	}
 
 	bool bHasShadow = FontRenderInfo.bEnableShadow;
-	if( bHasShadow && ShadowOffset.Size() == 0.0f )
+	if( bHasShadow && ShadowOffset.SizeSquared() == 0.0f )
 	{
 		// EnableShadow will set a default ShadowOffset value
 		EnableShadow( FLinearColor::Black );
diff --git a/Engine/Source/Runtime/Engine/Private/WorldComposition.cpp b/Engine/Source/Runtime/Engine/Private/WorldComposition.cpp
--- a/Engine/Source/Runtime/Engine/Private/WorldComposition.cpp
+++ b/Engine/Source/Runtime/Engine/Private/WorldComposition.cpp
@@ -768,7 +768,7 @@ void UWorldComposition::EvaluateWorldOriginLocation(const FVector& ViewLocation)
 	}
 		
 	// Request to shift world in case current view is quite far from current origin
-	if (Location.Size() > RebaseOriginDistance)
+	if (Location.SizeSquared() > FMath::Square(RebaseOriginDistance))
 	{
 		OwningWorld->RequestNewWorldOrigin(FIntVector(Location.X, Location.Y, Location.Z) + OwningWorld->OriginLocation);
 	}
diff --git a/Engine/Source/Runtime/Engine/Public/CanvasItem.h b/Engine/Source/Runtime/Engine/Public/CanvasItem.h
--- a/Engine/Source/Runtime/Engine/Public/CanvasItem.h
+++ b/Engine/Source/Runtime/Engine/Public/CanvasItem.h
@@ -364,7 +364,7 @@ public:
 	{
 		ShadowOffset = InOffset;
 		ShadowColor = InColor;
-		FontRenderInfo.bEnableShadow = ShadowOffset.Size() != 0.0f;
+		FontRenderInfo.bEnableShadow = ShadowOffset.SizeSquared() != 0.0f;
 	}
 	
 	/** 
diff --git a/Engine/Source/Runtime/Engine/Public/GenericQuadTree.h b/Engine/Source/Runtime/Engine/Public/GenericQuadTree.h
--- a/Engine/Source/Runtime/Engine/Public/GenericQuadTree.h
+++ b/Engine/Source/Runtime/Engine/Public/GenericQuadTree.h
@@ -249,7 +249,7 @@ void TQuadTree<ElementType, NodeCapacity>::InsertElementRecursive(const ElementT
 		// It's possible that all elements in the leaf are bigger than the leaf or that more elements than NodeCapacity exist outside the top level quad
 		// In either case, we can get into an endless spiral of splitting
 		static const float MinimumQuadSize = 100.f;
-		const bool bCanSplitTree = TreeBox.GetSize().Size() > MinimumQuadSize;
+		const bool bCanSplitTree = TreeBox.GetSize().SizeSquared() > FMath::Square(MinimumQuadSize);
 		if (!bCanSplitTree || Nodes.Num() < NodeCapacity)
 		{
 			Nodes.Add(FNode(Element, Box));
diff --git a/Engine/Source/Runtime/Slate/Private/Framework/Application/SlateApplication.cpp b/Engine/Source/Runtime/Slate/Private/Framework/Application/SlateApplication.cpp
--- a/Engine/Source/Runtime/Slate/Private/Framework/Application/SlateApplication.cpp
+++ b/Engine/Source/Runtime/Slate/Private/Framework/Application/SlateApplication.cpp
@@ -4458,7 +4458,7 @@ bool FSlateApplication::RoutePointerMoveEvent(const FWidgetPath& WidgetsUnderPoi
 	if( DragDetector.DetectDragForWidget.IsValid() )
 	{	
 		const FVector2D DragDelta = (DragDetector.DetectDragStartLocation - PointerEvent.GetScreenSpacePosition());
-		bDragDetected = ( DragDelta.Size() > FSlateApplication::Get().GetDragTriggerDistance() );
+		bDragDetected = ( DragDelta.SizeSquared() > FMath::Square(FSlateApplication::Get().GetDragTriggerDistance()) );
 		if (bDragDetected)
 		{
 			FWidgetPath DragDetectPath = DragDetector.DetectDragForWidget.ToWidgetPath(FWeakWidgetPath::EInterruptedPathHandling::ReturnInvalid);
diff --git a/Engine/Source/Runtime/Slate/Private/Framework/Text/SlateTextRun.cpp b/Engine/Source/Runtime/Slate/Private/Framework/Text/SlateTextRun.cpp
--- a/Engine/Source/Runtime/Slate/Private/Framework/Text/SlateTextRun.cpp
+++ b/Engine/Source/Runtime/Slate/Private/Framework/Text/SlateTextRun.cpp
@@ -73,7 +73,7 @@ int32 FSlateTextRun::OnPaint( const FPaintArgs& Args, const FTextLayout::FLineVi
 	const FSlateRect ClippingRect = AllottedGeometry.GetClippingRect().IntersectionWith(MyClippingRect);
 	const ESlateDrawEffect::Type DrawEffects = bParentEnabled ? ESlateDrawEffect::None : ESlateDrawEffect::DisabledEffect;
 
-	const bool ShouldDropShadow = Style.ShadowOffset.Size() > 0 && Style.ShadowColorAndOpacity.A > 0;
+	const bool ShouldDropShadow = Style.ShadowColorAndOpacity.A > 0.f && Style.ShadowOffset.SizeSquared() > 0.f;
 	const FVector2D BlockLocationOffset = Block->GetLocationOffset();
 	const FTextRange BlockRange = Block->GetTextRange();
 	
diff --git a/Engine/Source/Runtime/Slate/Private/Widgets/Input/SVirtualJoystick.cpp b/Engine/Source/Runtime/Slate/Private/Widgets/Input/SVirtualJoystick.cpp
--- a/Engine/Source/Runtime/Slate/Private/Widgets/Input/SVirtualJoystick.cpp
+++ b/Engine/Source/Runtime/Slate/Private/Widgets/Input/SVirtualJoystick.cpp
@@ -292,7 +292,7 @@ bool SVirtualJoystick::HandleTouch(int32 ControlIndex, const FVector2D& LocalCoo
 	else
 	{
 		// clamp to the ellipse of the stick (snaps to the visual size, so, the art should go all the way to the edge of the texture)
-		float DistanceToTouch = Offset.Size();
+		float DistanceToTouchSqr = Offset.SizeSquared();
 		float Angle = FMath::Atan2(Offset.Y, Offset.X);
 
 		// length along line to ellipse: L = 1.0 / sqrt(((sin(T)/Rx)^2 + (cos(T)/Ry)^2))
@@ -303,7 +303,7 @@ bool SVirtualJoystick::HandleTouch(int32 ControlIndex, const FVector2D& LocalCoo
 		float DistanceToEdge = FMath::InvSqrt(XTerm * XTerm + YTerm * YTerm);
 
 		// only clamp 
-		if (DistanceToTouch > DistanceToEdge)
+		if (DistanceToTouchSqr > FMath::Square(DistanceToEdge))
 		{
 			Control.ThumbPosition = FVector2D(DistanceToEdge * CosAngle,  DistanceToEdge * SinAngle);
 		}
diff --git a/Engine/Source/Runtime/Slate/Private/Widgets/Text/STextBlock.cpp b/Engine/Source/Runtime/Slate/Private/Widgets/Text/STextBlock.cpp
--- a/Engine/Source/Runtime/Slate/Private/Widgets/Text/STextBlock.cpp
+++ b/Engine/Source/Runtime/Slate/Private/Widgets/Text/STextBlock.cpp
@@ -153,7 +153,7 @@ int32 STextBlock::OnPaint( const FPaintArgs& Args, const FGeometry& AllottedGeom
 		const ESlateDrawEffect::Type DrawEffects = ShouldBeEnabled(bParentEnabled) ? ESlateDrawEffect::None : ESlateDrawEffect::DisabledEffect;
 		const FLinearColor CurShadowColor = GetShadowColorAndOpacity();
 		const FVector2D CurShadowOffset = GetShadowOffset();
-		const bool ShouldDropShadow = CurShadowOffset.Size() > 0 && CurShadowColor.A > 0;
+		const bool ShouldDropShadow = CurShadowColor.A > 0.f && CurShadowOffset.SizeSquared() > 0.f;
 		const FSlateFontInfo FontInfo = GetFont();
 		const FText& TextToDraw = BoundText.Get(FText::GetEmpty());
 
diff --git a/Engine/Source/Runtime/SlateCore/Private/Widgets/SPanel.cpp b/Engine/Source/Runtime/SlateCore/Private/Widgets/SPanel.cpp
--- a/Engine/Source/Runtime/SlateCore/Private/Widgets/SPanel.cpp
+++ b/Engine/Source/Runtime/SlateCore/Private/Widgets/SPanel.cpp
@@ -34,7 +34,7 @@ int32 SPanel::PaintArrangedChildren( const FPaintArgs& Args, const FArrangedChil
 		const FArrangedWidget& CurWidget = ArrangedChildren[ChildIndex];
 		FSlateRect ChildClipRect = MyClippingRect.IntersectionWith(CurWidget.Geometry.GetClippingRect());
 		
-		if (ChildClipRect.GetSize().Size() > 0)
+		if (ChildClipRect.GetSize().SizeSquared() > 0.f)
 		{
 			const int32 CurWidgetsMaxLayerId = CurWidget.Widget->Paint(Args.WithNewParent(this), CurWidget.Geometry, ChildClipRect, OutDrawElements, LayerId, InWidgetStyle, ShouldBeEnabled(bParentEnabled));
 			MaxLayerId = FMath::Max(MaxLayerId, CurWidgetsMaxLayerId);
diff --git a/Engine/Source/Runtime/SlateCore/Public/Layout/SlateRect.h b/Engine/Source/Runtime/SlateCore/Public/Layout/SlateRect.h
--- a/Engine/Source/Runtime/SlateCore/Public/Layout/SlateRect.h
+++ b/Engine/Source/Runtime/SlateCore/Public/Layout/SlateRect.h
@@ -64,7 +64,7 @@ public:
 	 */
 	bool IsEmpty() const
 	{
-		return GetSize().Size() == 0.0f;
+		return GetSize().SizeSquared() == 0.0f;
 	}
 
 	/**
diff --git a/Templates/TP_2DSideScroller/Source/TP_2DSideScroller/TP_2DSideScrollerCharacter.cpp b/Templates/TP_2DSideScroller/Source/TP_2DSideScroller/TP_2DSideScrollerCharacter.cpp
--- a/Templates/TP_2DSideScroller/Source/TP_2DSideScroller/TP_2DSideScrollerCharacter.cpp
+++ b/Templates/TP_2DSideScroller/Source/TP_2DSideScroller/TP_2DSideScrollerCharacter.cpp
@@ -95,10 +95,10 @@ ATP_2DSideScrollerCharacter::ATP_2DSideScrollerCharacter()
 void ATP_2DSideScrollerCharacter::UpdateAnimation()
 {
 	const FVector PlayerVelocity = GetVelocity();
-	const float PlayerSpeed = PlayerVelocity.Size();
+	const float PlayerSpeedSqr = PlayerVelocity.SizeSquared();
 
 	// Are we moving or standing still?
-	UPaperFlipbook* DesiredAnimation = (PlayerSpeed > 0.0f) ? RunningAnimation : IdleAnimation;
+	UPaperFlipbook* DesiredAnimation = (PlayerSpeedSqr > 0.0f) ? RunningAnimation : IdleAnimation;
 	if( GetSprite()->GetFlipbook() != DesiredAnimation 	)
 	{
 		GetSprite()->SetFlipbook(DesiredAnimation);
