jira UE-10126
Integrate github pull request #832
- I did not convert geometry, ai path and lighting build to use the generic system. But users will be able to add custom build step using the public interface.

diff --git a/Engine/Source/Editor/LevelEditor/Private/LevelEditorActions.cpp b/Engine/Source/Editor/LevelEditor/Private/LevelEditorActions.cpp
--- a/Engine/Source/Editor/LevelEditor/Private/LevelEditorActions.cpp
+++ b/Engine/Source/Editor/LevelEditor/Private/LevelEditorActions.cpp
@@ -526,7 +526,7 @@ void FLevelEditorActionCallbacks::Build_Execute()
 	ConfigureLightingBuildOptions( FLightingBuildOptions() );
 
 	// Build everything!
-	FEditorBuildUtils::EditorBuild( GetWorld(), EBuildOptions::BuildAll );
+	FEditorBuildUtils::EditorBuild( GetWorld(), FBuildOptions::BuildAll );
 }
 
 
@@ -544,7 +544,7 @@ void FLevelEditorActionCallbacks::BuildLightingOnly_Execute()
 
 	// Build lighting!
 	const bool bAllowLightingDialog = false;
-	FEditorBuildUtils::EditorBuild( GetWorld(), EBuildOptions::BuildLighting, bAllowLightingDialog );
+	FEditorBuildUtils::EditorBuild( GetWorld(), FBuildOptions::BuildLighting, bAllowLightingDialog );
 }
 
 bool FLevelEditorActionCallbacks::BuildLighting_CanExecute()
@@ -568,7 +568,7 @@ void FLevelEditorActionCallbacks::BuildLightingOnly_VisibilityOnly_Execute()
 
 	// Build lighting!
 	const bool bAllowLightingDialog = false;
-	FEditorBuildUtils::EditorBuild( GetWorld(), EBuildOptions::BuildLighting, bAllowLightingDialog );
+	FEditorBuildUtils::EditorBuild( GetWorld(), FBuildOptions::BuildLighting, bAllowLightingDialog );
 
 	// Reset build options
 	ConfigureLightingBuildOptions( FLightingBuildOptions() );
@@ -610,27 +610,27 @@ void FLevelEditorActionCallbacks::LightingBuildOptions_ShowLightingStats_Toggled
 void FLevelEditorActionCallbacks::BuildGeometryOnly_Execute()
 {
 	// Build geometry!
-	FEditorBuildUtils::EditorBuild( GetWorld(), EBuildOptions::BuildVisibleGeometry );
+	FEditorBuildUtils::EditorBuild( GetWorld(), FBuildOptions::BuildVisibleGeometry );
 }
 
 
 void FLevelEditorActionCallbacks::BuildGeometryOnly_OnlyCurrentLevel_Execute()
 {
 	// Build geometry (current level)!
-	FEditorBuildUtils::EditorBuild( GetWorld(), EBuildOptions::BuildGeometry );
+	FEditorBuildUtils::EditorBuild( GetWorld(), FBuildOptions::BuildGeometry );
 }
 
 
 void FLevelEditorActionCallbacks::BuildPathsOnly_Execute()
 {
 	// Build paths!
-	FEditorBuildUtils::EditorBuild( GetWorld(), EBuildOptions::BuildAIPaths );
+	FEditorBuildUtils::EditorBuild( GetWorld(), FBuildOptions::BuildAIPaths );
 }
 
 void FLevelEditorActionCallbacks::BuildLODsOnly_Execute()
 {
 	// Build HLOD
-	FEditorBuildUtils::EditorBuild(GetWorld(), EBuildOptions::BuildHierarchicalLOD);
+	FEditorBuildUtils::EditorBuild(GetWorld(), FBuildOptions::BuildHierarchicalLOD);
 }
 
 bool FLevelEditorActionCallbacks::IsLightingQualityChecked( ELightingBuildQuality TestQuality )
diff --git a/Engine/Source/Editor/UnrealEd/Private/EditorBuildUtils.cpp b/Engine/Source/Editor/UnrealEd/Private/EditorBuildUtils.cpp
--- a/Engine/Source/Editor/UnrealEd/Private/EditorBuildUtils.cpp
+++ b/Engine/Source/Editor/UnrealEd/Private/EditorBuildUtils.cpp
@@ -30,7 +30,52 @@ extern UNREALED_API bool GLightmassDebugMode;
 extern UNREALED_API bool GLightmassStatsMode;
 extern FSwarmDebugOptions GSwarmDebugOptions;
 
+const FName FBuildOptions::BuildGeometry(TEXT("BuildGeometry"));
+const FName FBuildOptions::BuildVisibleGeometry(TEXT("BuildVisibleGeometry"));
+const FName FBuildOptions::BuildLighting(TEXT("BuildLighting"));
+const FName FBuildOptions::BuildAIPaths(TEXT("BuildAIPaths"));
+const FName FBuildOptions::BuildSelectedAIPaths(TEXT("BuildSelectedAIPaths"));
+const FName FBuildOptions::BuildAll(TEXT("BuildAll"));
+const FName FBuildOptions::BuildAllSubmit(TEXT("BuildAllSubmit"));
+const FName FBuildOptions::BuildAllOnlySelectedPaths(TEXT("BuildAllOnlySelectedPaths"));
+const FName FBuildOptions::BuildHierarchicalLOD(TEXT("BuildHierarchicalLOD"));
+
+
 bool FEditorBuildUtils::bBuildingNavigationFromUserRequest = false;
+TMap<FName, FEditorBuildUtils::FCustomBuildType> FEditorBuildUtils::CustomBuildTypes;
+FName FEditorBuildUtils::InProgressBuildId;
+
+/**
+ * Class that handles potentially-async Build All requests.
+ */
+class FBuildAllHandler
+{
+public:
+	void StartBuild(UWorld* World, FName BuildId, const TWeakPtr<SBuildProgressWidget>& BuildProgressWidget);
+	void ResumeBuild();
+	void AddCustomBuildStep(FName Id, FName InsertBefore);
+	void RemoveCustomBuildStep(FName Id);
+
+	static FBuildAllHandler& Get()
+	{
+		static FBuildAllHandler Instance;
+		return Instance;
+	}
+
+private:
+	FBuildAllHandler();
+	FBuildAllHandler(const FBuildAllHandler&);
+
+	void ProcessBuild(const TWeakPtr<SBuildProgressWidget>& BuildProgressWidget);
+	void BuildFinished();
+
+	TArray<FName> BuildSteps;
+	int32 CurrentStep;
+
+	UWorld* CurrentWorld;
+	FName CurrentBuildId;
+};
+
 /** Constructor */
 FEditorBuildUtils::FEditorAutomatedBuildSettings::FEditorAutomatedBuildSettings()
 :	BuildErrorBehavior( ABB_PromptOnError ),
@@ -67,7 +112,7 @@ bool FEditorBuildUtils::EditorAutomatedBuildAndSubmit( const FEditorAutomatedBui
 	// If the preparation went smoothly, attempt the actual map building process
 	if ( bBuildSuccessful )
 	{
-		bBuildSuccessful = EditorBuild( GWorld, EBuildOptions::BuildAllSubmit );
+		bBuildSuccessful = EditorBuild( GWorld, FBuildOptions::BuildAllSubmit );
 
 		// If the map build failed, log the error
 		if ( !bBuildSuccessful )
@@ -162,7 +207,7 @@ static bool IsBuildCancelled()
  *
  * @return	true if the build completed successfully; false if it did not (or was manually canceled)
  */
-bool FEditorBuildUtils::EditorBuild( UWorld* InWorld, EBuildOptions::Type Id, const bool bAllowLightingDialog )
+bool FEditorBuildUtils::EditorBuild( UWorld* InWorld, FName Id, const bool bAllowLightingDialog )
 {
 	FMessageLog("MapCheck").NewPage(LOCTEXT("MapCheckNewPage", "Map Check"));
 
@@ -183,7 +228,7 @@ bool FEditorBuildUtils::EditorBuild( UWorld* InWorld, EBuildOptions::Type Id, co
 
 	// Show option dialog first, before showing the DlgBuildProgress window.
 	FLightingBuildOptions LightingBuildOptions;
-	if ( Id == EBuildOptions::BuildLighting )
+	if ( Id == FBuildOptions::BuildLighting )
 	{
 		// Retrieve settings from ini.
 		GConfig->GetBool( TEXT("LightingBuildOptions"), TEXT("OnlyBuildSelected"),		LightingBuildOptions.bOnlyBuildSelected,			GEditorPerProjectIni );
@@ -201,181 +246,143 @@ bool FEditorBuildUtils::EditorBuild( UWorld* InWorld, EBuildOptions::Type Id, co
 	// Show the build progress dialog.
 	SBuildProgressWidget::EBuildType BuildType = SBuildProgressWidget::BUILDTYPE_Geometry;
 
-	switch (Id)
+	if (Id == FBuildOptions::BuildGeometry ||
+		Id == FBuildOptions::BuildVisibleGeometry ||
+		Id == FBuildOptions::BuildAll ||
+		Id == FBuildOptions::BuildAllOnlySelectedPaths)
 	{
-	case EBuildOptions::BuildGeometry:
-	case EBuildOptions::BuildVisibleGeometry:
-	case EBuildOptions::BuildAll:
-	case EBuildOptions::BuildAllOnlySelectedPaths:
 		BuildType = SBuildProgressWidget::BUILDTYPE_Geometry;
-		break;
-	case EBuildOptions::BuildLighting:
+	}
+	else if (Id == FBuildOptions::BuildLighting)
+	{
 		BuildType = SBuildProgressWidget::BUILDTYPE_Lighting;
-		break;
-	case EBuildOptions::BuildAIPaths:
-	case EBuildOptions::BuildSelectedAIPaths:
+	}
+	else if (Id == FBuildOptions::BuildAIPaths ||
+		Id == FBuildOptions::BuildSelectedAIPaths)
+	{
 		BuildType = SBuildProgressWidget::BUILDTYPE_Paths;
-		break;
-	case EBuildOptions::BuildHierarchicalLOD:
+	}
+	else if (Id == FBuildOptions::BuildHierarchicalLOD)
+	{
 		BuildType = SBuildProgressWidget::BUILDTYPE_LODs;
-		break;
-	default:
+	}
+	else
+	{
 		BuildType = SBuildProgressWidget::BUILDTYPE_Unknown;	
-		break;
 	}
 
 	TWeakPtr<class SBuildProgressWidget> BuildProgressWidget = GWarn->ShowBuildProgressWindow();
 	BuildProgressWidget.Pin()->SetBuildType(BuildType);
 
 	bool bShouldMapCheck = true;
-	switch( Id )
+	if (Id == FBuildOptions::BuildGeometry)
 	{
-	case EBuildOptions::BuildGeometry:
+		// We can't set the busy cursor for all windows, because lighting
+		// needs a cursor for the lighting options dialog.
+		const FScopedBusyCursor BusyCursor;
+
+		GUnrealEd->Exec( InWorld, TEXT("MAP REBUILD") );
+
+		if (GetDefault<ULevelEditorMiscSettings>()->bNavigationAutoUpdate)
+		{
+			TriggerNavigationBuilder(InWorld, Id);
+		}
+
+		// No need to dirty the persient level if we're building BSP for a sub-level.
+		bDirtyPersistentLevel = false;
+	}
+	else if (Id == FBuildOptions::BuildVisibleGeometry)
+	{
+		// If any levels are hidden, prompt the user about how to proceed
+		bDoBuild = GEditor->WarnAboutHiddenLevels( InWorld, true );
+		if ( bDoBuild )
 		{
 			// We can't set the busy cursor for all windows, because lighting
 			// needs a cursor for the lighting options dialog.
 			const FScopedBusyCursor BusyCursor;
 
-			GUnrealEd->Exec( InWorld, TEXT("MAP REBUILD") );
+			GUnrealEd->Exec( InWorld, TEXT("MAP REBUILD ALLVISIBLE") );
 
 			if (GetDefault<ULevelEditorMiscSettings>()->bNavigationAutoUpdate)
 			{
 				TriggerNavigationBuilder(InWorld, Id);
 			}
-
-			// No need to dirty the persient level if we're building BSP for a sub-level.
-			bDirtyPersistentLevel = false;
-			break;
 		}
-	case EBuildOptions::BuildVisibleGeometry:
+	}
+	else if (Id == FBuildOptions::BuildLighting)
+	{
+		if( bDoBuild )
 		{
-			// If any levels are hidden, prompt the user about how to proceed
-			bDoBuild = GEditor->WarnAboutHiddenLevels( InWorld, true );
-			if ( bDoBuild )
-			{
-				// We can't set the busy cursor for all windows, because lighting
-				// needs a cursor for the lighting options dialog.
-				const FScopedBusyCursor BusyCursor;
+			// We can't set the busy cursor for all windows, because lighting
+			// needs a cursor for the lighting options dialog.
+			const FScopedBusyCursor BusyCursor;
 
-				GUnrealEd->Exec( InWorld, TEXT("MAP REBUILD ALLVISIBLE") );
+			// BSP export to lightmass relies on current BSP state
+			GUnrealEd->Exec( InWorld, TEXT("MAP REBUILD ALLVISIBLE") );
 
-				if (GetDefault<ULevelEditorMiscSettings>()->bNavigationAutoUpdate)
-				{
-					TriggerNavigationBuilder(InWorld, Id);
-				}
-			}
-			break;
+			GUnrealEd->BuildLighting( LightingBuildOptions );
+			bShouldMapCheck = false;
 		}
-
-	case EBuildOptions::BuildLighting:
+	}
+	else if (Id == FBuildOptions::BuildAIPaths)
+	{
+		bDoBuild = GEditor->WarnAboutHiddenLevels( InWorld, false );
+		if ( bDoBuild )
 		{
-			if( bDoBuild )
-			{
-				// We can't set the busy cursor for all windows, because lighting
-				// needs a cursor for the lighting options dialog.
-				const FScopedBusyCursor BusyCursor;
+			GEditor->ResetTransaction( NSLOCTEXT("UnrealEd", "RebuildNavigation", "Rebuilding Navigation") );
 
-				// BSP export to lightmass relies on current BSP state
-				GUnrealEd->Exec( InWorld, TEXT("MAP REBUILD ALLVISIBLE") );
+			// We can't set the busy cursor for all windows, because lighting
+			// needs a cursor for the lighting options dialog.
+			const FScopedBusyCursor BusyCursor;
 
-				GUnrealEd->BuildLighting( LightingBuildOptions );
-				bShouldMapCheck = false;
-			}
-			break;
+			TriggerNavigationBuilder(InWorld, Id);
 		}
+	}
+	else if (CustomBuildTypes.Contains(Id))
+	{
+		const auto& BuildType = CustomBuildTypes.FindChecked(Id);
+		check(BuildType.DoBuild.IsBound());
 
-	case EBuildOptions::BuildAIPaths:
-		{
-			bDoBuild = GEditor->WarnAboutHiddenLevels( InWorld, false );
-			if ( bDoBuild )
-			{
-				GEditor->ResetTransaction( NSLOCTEXT("UnrealEd", "RebuildNavigation", "Rebuilding Navigation") );
-
-				// We can't set the busy cursor for all windows, because lighting
-				// needs a cursor for the lighting options dialog.
-				const FScopedBusyCursor BusyCursor;
+		// Invoke custom build.
+		auto Result = BuildType.DoBuild.Execute(InWorld, Id);
 
-				TriggerNavigationBuilder(InWorld, Id);
-			}
+		bDoBuild = Result != EEditorBuildResult::Skipped;
+		bShouldMapCheck = Result == EEditorBuildResult::Success;
+		bDirtyPersistentLevel = Result == EEditorBuildResult::Success;
 
-			break;
+		if (Result == EEditorBuildResult::InProgress)
+		{
+			InProgressBuildId = Id;
 		}
-
-	case EBuildOptions::BuildHierarchicalLOD:
+	}
+	else if (Id == FBuildOptions::BuildHierarchicalLOD)
+	{
+		bDoBuild = GEditor->WarnAboutHiddenLevels( InWorld, false );
+		if ( bDoBuild )
 		{
-			bDoBuild = GEditor->WarnAboutHiddenLevels( InWorld, false );
-			if ( bDoBuild )
-			{
 				GEditor->ResetTransaction( NSLOCTEXT("UnrealEd", "BuildHLODMeshes", "Building Hierarchical LOD Meshes") );
 
-				// We can't set the busy cursor for all windows, because lighting
-				// needs a cursor for the lighting options dialog.
-				const FScopedBusyCursor BusyCursor;
-
-				TriggerHierarchicalLODBuilder(InWorld, Id);
-			}
+			// We can't set the busy cursor for all windows, because lighting
+			// needs a cursor for the lighting options dialog.
+			const FScopedBusyCursor BusyCursor;
 
-			break;
+			TriggerHierarchicalLODBuilder(InWorld, Id);
 		}
-
-	case EBuildOptions::BuildAll:
-	case EBuildOptions::BuildAllSubmit:
+	}
+	else if (Id == FBuildOptions::BuildAll || Id == FBuildOptions::BuildAllSubmit)
+	{
+		// TODO: WarnIfLightingBuildIsCurrentlyRunning should check with FBuildAllHandler
+		bDoBuild = GEditor->WarnAboutHiddenLevels( InWorld, true );
+		bool bLightingAlreadyRunning = GUnrealEd->WarnIfLightingBuildIsCurrentlyRunning();
+		if ( bDoBuild && !bLightingAlreadyRunning )
 		{
-			bDoBuild = GEditor->WarnAboutHiddenLevels( InWorld, true );
-			bool bLightingAlreadyRunning = GUnrealEd->WarnIfLightingBuildIsCurrentlyRunning();
-			if ( bDoBuild && !bLightingAlreadyRunning )
-			{
-				// We can't set the busy cursor for all windows, because lighting
-				// needs a cursor for the lighting options dialog.
-				const FScopedBusyCursor BusyCursor;
-
-				GUnrealEd->Exec( InWorld, TEXT("MAP REBUILD ALLVISIBLE") );
-
- 				{
- 					BuildProgressWidget.Pin()->SetBuildType(SBuildProgressWidget::BUILDTYPE_LODs);
-					TriggerHierarchicalLODBuilder(InWorld, Id);
- 				}
-
-				{
-					BuildProgressWidget.Pin()->SetBuildType(SBuildProgressWidget::BUILDTYPE_Paths);
-					TriggerNavigationBuilder(InWorld, Id);
-				}
-
-				//Do a canceled check before moving on to the next step of the build.
-				if( GEditor->GetMapBuildCancelled() )
-				{
-					break;
-				}
-				else
-				{
-					BuildProgressWidget.Pin()->SetBuildType(SBuildProgressWidget::BUILDTYPE_Lighting);
-
-					FLightingBuildOptions LightingOptions;
-
-					int32 QualityLevel;
-
-					// Force automated builds to always use production lighting
-					if ( Id == EBuildOptions::BuildAllSubmit )
-					{
-						QualityLevel = Quality_Production;
-					}
-					else
-					{
-						GConfig->GetInt( TEXT("LightingBuildOptions"), TEXT("QualityLevel"), QualityLevel, GEditorPerProjectIni);
-						QualityLevel = FMath::Clamp<int32>(QualityLevel, Quality_Preview, Quality_Production);
-					}
-					LightingOptions.QualityLevel = (ELightingBuildQuality)QualityLevel;
-
-					GUnrealEd->BuildLighting(LightingOptions);
-					bShouldMapCheck = false;
-				}
-			}
-			break;
+			FBuildAllHandler::Get().StartBuild(InWorld, Id, BuildProgressWidget);
 		}
-
-	default:
-		UE_LOG(LogEditorBuildUtils, Warning, TEXT("Invalid build Id"));
-		break;
+	}
+	else
+	{
+		UE_LOG(LogEditorBuildUtils, Warning, TEXT("Invalid build Id: %s"), *Id.ToString());
+		bDoBuild = false;
 	}
 
 	// Check map for errors (only if build operation happened)
@@ -769,35 +776,240 @@ void FEditorBuildUtils::SubmitPackagesForAutomatedBuild( const TSet<UPackage*>&
 	SourceControlProvider.Execute( CheckInOperation, LevelsToSubmit, EConcurrency::Synchronous );
 }
 
-void FEditorBuildUtils::TriggerNavigationBuilder(UWorld* InWorld, EBuildOptions::Type Id)
+void FEditorBuildUtils::TriggerNavigationBuilder(UWorld* InWorld, FName Id)
 {
 	if( InWorld->GetWorldSettings()->bEnableNavigationSystem &&
 		InWorld->GetNavigationSystem() )
 	{
-		switch( Id )
+		if (Id == FBuildOptions::BuildAIPaths ||
+			Id == FBuildOptions::BuildSelectedAIPaths ||
+			Id == FBuildOptions::BuildAllOnlySelectedPaths ||
+			Id == FBuildOptions::BuildAll ||
+			Id == FBuildOptions::BuildAllSubmit)
+		{
+			bBuildingNavigationFromUserRequest = true;
+		}
+		else
 		{
-		case EBuildOptions::BuildAIPaths:
-		case EBuildOptions::BuildSelectedAIPaths:
-		case EBuildOptions::BuildAllOnlySelectedPaths:
-		case EBuildOptions::BuildAll:
-		case EBuildOptions::BuildAllSubmit:
+			bBuildingNavigationFromUserRequest = false;
+		}
+
+		// Invoke navmesh generator
+		InWorld->GetNavigationSystem()->Build();
+	}
+}
+
+/**
+ * Call this when an async custom build step has completed (successfully or not).
+ */
+void FEditorBuildUtils::AsyncBuildCompleted()
+{
+	check(InProgressBuildId != NAME_None);
+
+	// Reset in-progress id before resuming build all do we don't overwrite something that's just been set.
+	auto BuildId = InProgressBuildId;
+	InProgressBuildId = NAME_None;
+
+	if (BuildId == FBuildOptions::BuildAll || BuildId == FBuildOptions::BuildAllSubmit)
+	{
+		FBuildAllHandler::Get().ResumeBuild();
+	}
+}
+
+/**
+ * Is there currently an (async) build in progress?
+ */
+bool FEditorBuildUtils::IsBuildCurrentlyRunning()
+{
+	return InProgressBuildId != NAME_None;
+}
+
+/**
+ * Register a custom build type.
+ * @param Id The identifier to use for this build type.
+ * @param DoBuild The delegate to execute to run this build.
+ * @param BuildAllExtensionPoint If a valid name, run this build *before* running the build with this id when performing a Build All.
+ */
+void FEditorBuildUtils::RegisterCustomBuildType(FName Id, const FDoEditorBuildDelegate& DoBuild, FName BuildAllExtensionPoint)
+{
+	check(!CustomBuildTypes.Contains(Id));
+	CustomBuildTypes.Add(Id, FCustomBuildType(DoBuild, BuildAllExtensionPoint));
+
+	if (BuildAllExtensionPoint != NAME_None)
+	{
+		FBuildAllHandler::Get().AddCustomBuildStep(Id, BuildAllExtensionPoint);
+	}
+}
+
+/**
+ * Unregister a custom build type.
+ * @param Id The identifier of the build type to unregister.
+ */
+void FEditorBuildUtils::UnregisterCustomBuildType(FName Id)
+{
+	CustomBuildTypes.Remove(Id);
+	FBuildAllHandler::Get().RemoveCustomBuildStep(Id);
+}
+
+/**
+ * Initialise Build All handler.
+ */
+FBuildAllHandler::FBuildAllHandler()
+	: CurrentStep(0)
+{
+	// Add built in build steps.
+	BuildSteps.Add(FBuildOptions::BuildGeometry);
+	BuildSteps.Add(FBuildOptions::BuildHierarchicalLOD);
+	BuildSteps.Add(FBuildOptions::BuildAIPaths);
+	
+	//Lighting must always be the last one when doing a build all
+	BuildSteps.Add(FBuildOptions::BuildLighting);
+}
+
+/**
+ * Add a custom Build All step.
+ */
+void FBuildAllHandler::AddCustomBuildStep(FName Id, FName InsertBefore)
+{
+	const int32 InsertionPoint = BuildSteps.Find(InsertBefore);
+	if (InsertionPoint != INDEX_NONE)
+	{
+		BuildSteps.Insert(Id, InsertionPoint);
+	}
+}
+
+/**
+ * Remove a custom Build All step.
+ */
+void FBuildAllHandler::RemoveCustomBuildStep(FName Id)
+{
+	BuildSteps.Remove(Id);
+}
+
+/**
+ * Commence a new Build All operation.
+ */
+void FBuildAllHandler::StartBuild(UWorld* World, FName BuildId, const TWeakPtr<SBuildProgressWidget>& BuildProgressWidget)
+{
+	check(CurrentStep == 0);
+	check(CurrentWorld == nullptr);
+	check(CurrentBuildId == NAME_None);
+
+	CurrentWorld = World;
+	CurrentBuildId = BuildId;
+	ProcessBuild(BuildProgressWidget);
+}
+
+/**
+ * Resume a Build All build from where it was left off.
+ */
+void FBuildAllHandler::ResumeBuild()
+{
+	// Resuming from async operation, may be about to do slow stuff again so show the progress window again.
+	TWeakPtr<SBuildProgressWidget> BuildProgressWidget = GWarn->ShowBuildProgressWindow();
+
+	ProcessBuild(BuildProgressWidget);
+
+	// Synchronous part completed, hide the build progress dialog.
+	GWarn->CloseBuildProgressWindow();
+}
+
+/**
+ * Internal method that actual does the build.
+ */
+void FBuildAllHandler::ProcessBuild(const TWeakPtr<SBuildProgressWidget>& BuildProgressWidget)
+{
+	const FScopedBusyCursor BusyCursor;
+
+	// Loop until we finish, or we start an async step.
+	while (true)
+	{
+		if (GEditor->GetMapBuildCancelled())
+		{
+			// Build cancelled, so bail.
+			BuildFinished();
+			break;
+		}
+
+		check(BuildSteps.IsValidIndex(CurrentStep));
+		FName StepId = BuildSteps[CurrentStep];
+
+		if (StepId == FBuildOptions::BuildGeometry)
+		{
+			BuildProgressWidget.Pin()->SetBuildType(SBuildProgressWidget::BUILDTYPE_Geometry);
+			GUnrealEd->Exec(CurrentWorld, TEXT("MAP REBUILD ALLVISIBLE") );
+		}
+		else if (StepId == FBuildOptions::BuildHierarchicalLOD)
+		{
+			BuildProgressWidget.Pin()->SetBuildType(SBuildProgressWidget::BUILDTYPE_LODs);
+			FEditorBuildUtils::TriggerHierarchicalLODBuilder(CurrentWorld, CurrentBuildId);
+		}
+		else if (StepId == FBuildOptions::BuildAIPaths)
+		{
+			BuildProgressWidget.Pin()->SetBuildType(SBuildProgressWidget::BUILDTYPE_Paths);
+			FEditorBuildUtils::TriggerNavigationBuilder(CurrentWorld, CurrentBuildId);
+		}
+		else if (StepId == FBuildOptions::BuildLighting)
+		{
+			BuildProgressWidget.Pin()->SetBuildType(SBuildProgressWidget::BUILDTYPE_Lighting);
+
+			FLightingBuildOptions LightingOptions;
+
+			int32 QualityLevel;
+
+			// Force automated builds to always use production lighting
+			if ( CurrentBuildId == FBuildOptions::BuildAllSubmit )
 			{
-				bBuildingNavigationFromUserRequest = true;
-				break;
+				QualityLevel = Quality_Production;
+			}
+			else
+			{
+				GConfig->GetInt(TEXT("LightingBuildOptions"), TEXT("QualityLevel"), QualityLevel, GEditorPerProjectIni);
+				QualityLevel = FMath::Clamp<int32>(QualityLevel, Quality_Preview, Quality_Production);
 			}
-		default:
+			LightingOptions.QualityLevel = (ELightingBuildQuality)QualityLevel;
+
+			GUnrealEd->BuildLighting(LightingOptions);
+
+			// TODO!
+			//bShouldMapCheck = false;
+
+			// Lighting is always the last step (Lightmass isn't set up to resume builds).
+			BuildFinished();
+			break;
+		}
+		else
+		{
+			auto& CustomBuildType = FEditorBuildUtils::CustomBuildTypes[StepId];
+			auto Result = CustomBuildType.DoBuild.Execute(CurrentWorld, CurrentBuildId);
+
+			if (Result == EEditorBuildResult::InProgress)
 			{
-				bBuildingNavigationFromUserRequest = false;
+				// Build & Submit builds must be synchronous.
+				check(CurrentBuildId != FBuildOptions::BuildAllSubmit);
+
+				// Build step is running asynchronously, so let it run.
+				FEditorBuildUtils::InProgressBuildId = CurrentBuildId;
 				break;
 			}
 		}
 
-		// Invoke navmesh generator
-		InWorld->GetNavigationSystem()->Build();
+		// Next go around we want to do the next step.
+		CurrentStep++;
 	}
 }
 
-void FEditorBuildUtils::TriggerHierarchicalLODBuilder(UWorld* InWorld, EBuildOptions::Type Id)
+/**
+ * Called when a build is finished (successfully or not).
+ */
+void FBuildAllHandler::BuildFinished()
+{
+	CurrentStep = 0;
+	CurrentWorld = nullptr;
+	CurrentBuildId = NAME_None;
+}
+
+void FEditorBuildUtils::TriggerHierarchicalLODBuilder(UWorld* InWorld, FName Id)
 {
 	// Invoke HLOD generator, with either preview or full build
 	InWorld->HierarchicalLODBuilder->BuildMeshesForLODActors();
diff --git a/Engine/Source/Editor/UnrealEd/Private/StaticLightingSystem/StaticLightingSystem.cpp b/Engine/Source/Editor/UnrealEd/Private/StaticLightingSystem/StaticLightingSystem.cpp
--- a/Engine/Source/Editor/UnrealEd/Private/StaticLightingSystem/StaticLightingSystem.cpp
+++ b/Engine/Source/Editor/UnrealEd/Private/StaticLightingSystem/StaticLightingSystem.cpp
@@ -21,6 +21,7 @@ FSwarmDebugOptions GSwarmDebugOptions;
 #include "LightMap.h"
 #include "ShadowMap.h"
 #include "RendererInterface.h"
+#include "EditorBuildUtils.h"
 
 DEFINE_LOG_CATEGORY(LogStaticLightingSystem);
 
@@ -2181,6 +2182,19 @@ bool UEditorEngine::WarnIfLightingBuildIsCurrentlyRunning()
 			Notification->SetCompletionState(SNotificationItem::CS_Fail);
 		}
 	}
+	else if (FEditorBuildUtils::IsBuildCurrentlyRunning())
+	{
+		// Another, non-lighting editor build is running.
+		FNotificationInfo Info( LOCTEXT("EditorBuildUnderwayWarning", "A build process is currently underway! Please cancel it to proceed!") );
+		Info.ExpireDuration = 5.0f;
+		TSharedPtr<SNotificationItem> Notification = FSlateNotificationManager::Get().AddNotification(Info);
+		if (Notification.IsValid())
+		{
+			Notification->SetCompletionState(SNotificationItem::CS_Fail);
+		}
+
+		bFailure = true;
+	}
 	return bFailure;
 }
 
diff --git a/Engine/Source/Editor/UnrealEd/Private/UnrealEdSrv.cpp b/Engine/Source/Editor/UnrealEd/Private/UnrealEdSrv.cpp
--- a/Engine/Source/Editor/UnrealEd/Private/UnrealEdSrv.cpp
+++ b/Engine/Source/Editor/UnrealEd/Private/UnrealEdSrv.cpp
@@ -505,12 +505,12 @@ bool UUnrealEdEngine::HandleDumpSelectionCommand( const TCHAR* Str, FOutputDevic
 
 bool UUnrealEdEngine::HandleBuildLightingCommand( const TCHAR* Str, FOutputDevice& Ar, UWorld* InWorld )
 {
-	return FEditorBuildUtils::EditorBuild(InWorld, EBuildOptions::BuildLighting);	
+	return FEditorBuildUtils::EditorBuild(InWorld, FBuildOptions::BuildLighting);	
 }
 
 bool UUnrealEdEngine::HandleBuildPathsCommand( const TCHAR* Str, FOutputDevice& Ar, UWorld* InWorld )
 {
-	return FEditorBuildUtils::EditorBuild(InWorld, EBuildOptions::BuildAIPaths);
+	return FEditorBuildUtils::EditorBuild(InWorld, FBuildOptions::BuildAIPaths);
 }
 
 bool UUnrealEdEngine::HandleUpdateLandscapeEditorDataCommand( const TCHAR* Str, FOutputDevice& Ar, UWorld* InWorld )
diff --git a/Engine/Source/Editor/UnrealEd/Public/EditorBuildUtils.h b/Engine/Source/Editor/UnrealEd/Public/EditorBuildUtils.h
--- a/Engine/Source/Editor/UnrealEd/Public/EditorBuildUtils.h
+++ b/Engine/Source/Editor/UnrealEd/Public/EditorBuildUtils.h
@@ -7,30 +7,45 @@
 #ifndef __EditorBuildUtils_h__
 #define __EditorBuildUtils_h__
 
-namespace EBuildOptions
+/** Names of the built-in editor build types. */
+struct FBuildOptions
 {
-	enum Type
-	{
-		/** Build all geometry */
-		BuildGeometry,
-		/** Build only visible geometry */
-		BuildVisibleGeometry,
-		/** Build lighting */
-		BuildLighting,
-		/** Build all AI paths */
-		BuildAIPaths,
-		/** Build only selected AI paths */
-		BuildSelectedAIPaths,
-		/** Build everything */
-		BuildAll,
-		/** Build everything and submit to source control */
-		BuildAllSubmit,
-		/** Build everything except for paths only build seleced */
-		BuildAllOnlySelectedPaths,
-		/** Build Hierarchical LOD system - need WorldSetting setup*/
-		BuildHierarchicalLOD
-	};
-}
+	/** Build all geometry */
+	UNREALED_API static const FName BuildGeometry;
+	/** Build only visible geometry */
+	UNREALED_API static const FName BuildVisibleGeometry;
+	/** Build lighting */
+	UNREALED_API static const FName BuildLighting;
+	/** Build all AI paths */
+	UNREALED_API static const FName BuildAIPaths;
+	/** Build only selected AI paths */
+	UNREALED_API static const FName BuildSelectedAIPaths;
+	/** Build everything */
+	UNREALED_API static const FName BuildAll;
+	/** Build everything and submit to source control */
+	UNREALED_API static const FName BuildAllSubmit;
+	/** Build everything except for paths only build selected */
+	UNREALED_API static const FName BuildAllOnlySelectedPaths;
+	/** Build Hierarchical LOD system - need WorldSetting setup*/
+	UNREALED_API static const FName BuildHierarchicalLOD;
+};
+/**
+ * Result of a custom editor build.
+ */
+enum class EEditorBuildResult
+{
+	Success,		// The build step completed successfully
+	Skipped,		// The build step was skipped for some reason (e.g. cancelled)
+	InProgress,		// The build step is running asynchronously
+};
+
+/**
+ * Delegate for performing a custom editor build.
+ * @param UWorld* The world to run the build on.
+ * @param FName The Id of the build being run (either the registered build Id, or one of the BuildAll types).
+ * @return Status of the build step.
+ */
+DECLARE_DELEGATE_RetVal_TwoParams(EEditorBuildResult, FDoEditorBuildDelegate, UWorld*, FName);
 
 /** Utility class to hold functionality for building within the editor */
 class FEditorBuildUtils
@@ -101,7 +116,7 @@ public:
 	 *
 	 * @return	true if the build completed successfully; false if it did not (or was manually canceled)
 	 */
-	static UNREALED_API bool EditorBuild( UWorld* InWorld, EBuildOptions::Type Id, const bool bAllowLightingDialog = true );
+	static UNREALED_API bool EditorBuild( UWorld* InWorld, FName Id, const bool bAllowLightingDialog = true );
 
 	/** 
 	* check if navigation build was was triggered from editor as user request
@@ -114,6 +129,30 @@ public:
 	* call it to notify that navigation builder finished building 
 	*/
 	static UNREALED_API void PathBuildingFinished() { bBuildingNavigationFromUserRequest = false; }
+
+	/**
+	 * Call this when an async custom build step has completed (successfully or not).
+	 */
+	static UNREALED_API void AsyncBuildCompleted();
+
+	/**
+	 * Is there currently an (async) build in progress?
+	 */
+	static UNREALED_API bool IsBuildCurrentlyRunning();
+
+	/**
+	 * Register a custom build type.
+	 * @param Id The identifier to use for this build type.
+	 * @param DoBuild The delegate to execute to run this build.
+	 * @param BuildAllExtensionPoint If a valid name, run this build *before* running the build with this id when performing a Build All.
+	 */
+	static UNREALED_API void RegisterCustomBuildType(FName Id, const FDoEditorBuildDelegate& DoBuild, FName BuildAllExtensionPoint);
+
+	/**
+	 * Unregister a custom build type.
+	 * @param Id The identifier of the build type to unregister.
+	 */
+	static UNREALED_API void UnregisterCustomBuildType(FName Id);
 private:
 
 	/**
@@ -166,7 +205,7 @@ private:
 	 * @param	InWorld			WorldContext
 	 * @param	BuildSettings	Build settings that will be used for the editor build
 	 */
-	static void TriggerNavigationBuilder(UWorld* InWorld, EBuildOptions::Type Id);
+	static void TriggerNavigationBuilder(UWorld* InWorld, FName Id);
 
 	/** 
 	 * Trigger LOD builder to (re)generate LODActors
@@ -174,7 +213,7 @@ private:
 	 * @param	InWorld			WorldContext
 	 * @param	BuildSettings	Build settings that will be used for the editor build
 	 */
-	static void TriggerHierarchicalLODBuilder(UWorld* InWorld, EBuildOptions::Type Id);
+	static void TriggerHierarchicalLODBuilder(UWorld* InWorld, FName Id);
 
 	/** Intentionally hide constructors, etc. to prevent instantiation */
 	FEditorBuildUtils();
@@ -182,8 +221,31 @@ private:
 	FEditorBuildUtils( const FEditorBuildUtils& );
 	FEditorBuildUtils operator=( const FEditorBuildUtils& );
 
+	// Allow the async build all handler to access custom build type info.
+	friend class FBuildAllHandler;
+
 	/** static variable to cache data about user request. navigation builder works in a background and we have to cache this information */
 	static bool bBuildingNavigationFromUserRequest;
+
+	/**
+	 * Struct containing data for a custom build type.
+	 */
+	struct FCustomBuildType
+	{
+		FDoEditorBuildDelegate DoBuild;
+		FName BuildAllExtensionPoint;
+
+		FCustomBuildType(const FDoEditorBuildDelegate& InDoBuild, FName InBuildAllExtensionPoint)
+			: DoBuild(InDoBuild)
+			, BuildAllExtensionPoint(InBuildAllExtensionPoint)
+		{}
+	};
+
+	/** Map of custom build types registered with us. */
+	static TMap<FName, FCustomBuildType> CustomBuildTypes;
+
+	/** Set to a valid build type if an async build is in progress. */
+	static FName InProgressBuildId;
 };
 
 #endif // #ifndef __EditorBuildUtils_h__
