Fix for CIS shadow variable warning

diff --git a/Engine/Source/Editor/UnrealEd/Private/EditorBuildUtils.cpp b/Engine/Source/Editor/UnrealEd/Private/EditorBuildUtils.cpp
--- a/Engine/Source/Editor/UnrealEd/Private/EditorBuildUtils.cpp
+++ b/Engine/Source/Editor/UnrealEd/Private/EditorBuildUtils.cpp
@@ -340,11 +340,11 @@ bool FEditorBuildUtils::EditorBuild( UWorld* InWorld, FName Id, const bool bAllo
 	}
 	else if (CustomBuildTypes.Contains(Id))
 	{
-		const auto& BuildType = CustomBuildTypes.FindChecked(Id);
-		check(BuildType.DoBuild.IsBound());
+		const auto& CustomBuild = CustomBuildTypes.FindChecked(Id);
+		check(CustomBuild.DoBuild.IsBound());
 
 		// Invoke custom build.
-		auto Result = BuildType.DoBuild.Execute(InWorld, Id);
+		auto Result = CustomBuild.DoBuild.Execute(InWorld, Id);
 
 		bDoBuild = Result != EEditorBuildResult::Skipped;
 		bShouldMapCheck = Result == EEditorBuildResult::Success;
