UE-12379 : Fix the preview viewport trying to move the camera actor when FOV is outside clamped range

diff --git a/Engine/Source/Editor/UnrealEd/Private/EditorViewportClient.cpp b/Engine/Source/Editor/UnrealEd/Private/EditorViewportClient.cpp
--- a/Engine/Source/Editor/UnrealEd/Private/EditorViewportClient.cpp
+++ b/Engine/Source/Editor/UnrealEd/Private/EditorViewportClient.cpp
@@ -1503,9 +1503,13 @@ void FEditorViewportClient::UpdateCameraMovement( float DeltaTime )
 			NewViewRotation = FRotator::MakeFromEuler( NewViewEuler );
 		}
 
-		if( !NewViewLocation.Equals( GetViewLocation(), SMALL_NUMBER ) ||
-			NewViewRotation != GetViewRotation() ||
-			!FMath::IsNearlyEqual( NewViewFOV, ViewFOV, float(SMALL_NUMBER) ) )
+		// See if translation/rotation have changed
+		const bool bTransformDifferent = !NewViewLocation.Equals(GetViewLocation(), SMALL_NUMBER) || NewViewRotation != GetViewRotation();
+		// See if FOV has changed
+		const bool bFOVDifferent = !FMath::IsNearlyEqual( NewViewFOV, ViewFOV, float(SMALL_NUMBER) );
+
+		// If something has changed, tell the actor
+		if(bTransformDifferent || bFOVDifferent)
 		{
 			// Something has changed!
 			const bool bInvalidateChildViews=true;
@@ -1518,9 +1522,12 @@ void FEditorViewportClient::UpdateCameraMovement( float DeltaTime )
 			ViewFOV = NewViewFOV;
 
 			// Actually move/rotate the camera
-			MoveViewportPerspectiveCamera(
-				NewViewLocation - GetViewLocation(),
-				NewViewRotation - GetViewRotation() );
+			if(bTransformDifferent)
+			{
+				MoveViewportPerspectiveCamera(
+					NewViewLocation - GetViewLocation(),
+					NewViewRotation - GetViewRotation() );
+			}
 
 			// Invalidate the viewport widget
 			if (EditorViewportWidget.IsValid())
