Stopped packages which failed to cook from being saved in the asset regsitry.

diff --git a/Engine/Source/Editor/UnrealEd/Classes/CookOnTheSide/CookOnTheFlyServer.h b/Engine/Source/Editor/UnrealEd/Classes/CookOnTheSide/CookOnTheFlyServer.h
--- a/Engine/Source/Editor/UnrealEd/Classes/CookOnTheSide/CookOnTheFlyServer.h
+++ b/Engine/Source/Editor/UnrealEd/Classes/CookOnTheSide/CookOnTheFlyServer.h
@@ -170,7 +170,7 @@ public:
 	/** cooked file requests which includes platform which file is requested for */
 	struct FFilePlatformRequest
 	{
-	private:
+	protected:
 		FName Filename;
 		TArray<FName> Platformnames;
 	public:
@@ -182,29 +182,15 @@ public:
 
 
 		FFilePlatformRequest( const FName& InFileName, const FName& InPlatformName ) : Filename( InFileName )
-		{
-			Platformnames.Add( InPlatformName );
-		}
-
+		{ Platformnames.Add( InPlatformName ); }
 
 		FFilePlatformRequest( const FName& InFilename, const TArray<FName>& InPlatformname ) : Filename( InFilename )
-		{
-			Platformnames = InPlatformname;
-		}
+		{ Platformnames = InPlatformname; }
 
-		FFilePlatformRequest( const FName& InFilename, TArray<FName>&& InPlatformname ) : Filename( InFilename ), Platformnames(MoveTemp(InPlatformname))
-		{
-		}
-
-		FFilePlatformRequest( const FFilePlatformRequest& InFilePlatformRequest ) : Filename( InFilePlatformRequest.Filename ), Platformnames( InFilePlatformRequest.Platformnames )
-		{
-		}
+		FFilePlatformRequest( const FName& InFilename, TArray<FName>&& InPlatformname ) : Filename( InFilename ), Platformnames(MoveTemp(InPlatformname)) { }
+		FFilePlatformRequest( const FFilePlatformRequest& InFilePlatformRequest ) : Filename( InFilePlatformRequest.Filename ), Platformnames( InFilePlatformRequest.Platformnames ) { }
+		FFilePlatformRequest( FFilePlatformRequest&& InFilePlatformRequest ) : Filename( MoveTemp(InFilePlatformRequest.Filename) ), Platformnames( MoveTemp(InFilePlatformRequest.Platformnames) ) { }
 		
-		FFilePlatformRequest( FFilePlatformRequest&& InFilePlatformRequest ) : Filename( MoveTemp(InFilePlatformRequest.Filename) ), Platformnames( MoveTemp(InFilePlatformRequest.Platformnames) )
-		{
-		}
-		
-
 		void SetFilename( const FString &InFilename ) 
 		{
 			Filename = FName(*InFilename);
@@ -282,12 +268,26 @@ public:
 
 private:
 
+	struct FFilePlatformCookedPackage : public FFilePlatformRequest
+	{
+	private:
+		bool bSucceededSavePackage;
+	public:
+		FFilePlatformCookedPackage(const FFilePlatformRequest& InFilePlatformRequest, bool bInSuccededSavePackage) : FFilePlatformRequest(InFilePlatformRequest.GetFilename(), InFilePlatformRequest.GetPlatformnames()), bSucceededSavePackage(bInSuccededSavePackage) { }
+		FFilePlatformCookedPackage(const FName& InFilename, const TArray<FName>& InPlatformname, bool bInSuccededSavePackage) : FFilePlatformRequest(InFilename, InPlatformname), bSucceededSavePackage(bInSuccededSavePackage) { }
+		FFilePlatformCookedPackage(const FName& InFilename, TArray<FName>&& InPlatformname, bool bInSuccededSavePackage) : FFilePlatformRequest(InFilename, MoveTemp(InPlatformname)), bSucceededSavePackage(bInSuccededSavePackage) { }
+		FFilePlatformCookedPackage(const FFilePlatformCookedPackage& InFilePlatformRequest) : FFilePlatformRequest(InFilePlatformRequest.Filename, InFilePlatformRequest.Platformnames), bSucceededSavePackage(InFilePlatformRequest.bSucceededSavePackage) { }
+		FFilePlatformCookedPackage(FFilePlatformCookedPackage&& InFilePlatformRequest) : FFilePlatformRequest(MoveTemp(InFilePlatformRequest.Filename), MoveTemp(InFilePlatformRequest.Platformnames)), bSucceededSavePackage(InFilePlatformRequest.bSucceededSavePackage) { }
+
+		inline const bool HasSucceededSavePackage() const { return bSucceededSavePackage; }
+	};
+
 	/** Helper list of all files which have been cooked */
 	struct FThreadSafeFilenameSet
 	{
 	private:
 		mutable FCriticalSection	SynchronizationObject;
-		TMap<FName, FFilePlatformRequest> FilesProcessed;
+		TMap<FName, FFilePlatformCookedPackage> FilesProcessed;
 	public:
 
 		void Lock()
@@ -299,7 +299,7 @@ private:
 			SynchronizationObject.Unlock();
 		}
 
-		void Add(const FFilePlatformRequest& Request)
+		void Add(const FFilePlatformCookedPackage& Request)
 		{
 			FScopeLock ScopeLock(&SynchronizationObject);
 			check(Request.IsValid());
@@ -318,11 +318,11 @@ private:
 			else
 				FilesProcessed.Add(Request.GetFilename(), Request);
 		}
-		bool Exists(const FFilePlatformRequest& Request) const
+		bool Exists(const FFilePlatformRequest & Request) const
 		{
 			FScopeLock ScopeLock(&SynchronizationObject);
 
-			const FFilePlatformRequest* OurRequest = FilesProcessed.Find( Request.GetFilename() );
+			const FFilePlatformCookedPackage* OurRequest = FilesProcessed.Find(Request.GetFilename());
 			
 			if (!OurRequest)
 			{
@@ -378,14 +378,13 @@ private:
 		bool GetCookedPlatforms( const FName& Filename, TArray<FName>& PlatformList ) const
 		{
 			FScopeLock ScopeLock( &SynchronizationObject );
-			const FFilePlatformRequest* Request = FilesProcessed.Find(Filename);
+			const FFilePlatformCookedPackage* Request = FilesProcessed.Find(Filename);
 			if ( Request )
 			{
 				PlatformList = Request->GetPlatformnames();
 				return true;
 			}
 			return false;
-			
 		}
 		int RemoveFile( const FName& Filename )
 		{
@@ -396,7 +395,7 @@ private:
 		bool RemoveFileForPlatform( const FName& Filename, const FName& PlatformName )
 		{
 			FScopeLock ScopeLock( &SynchronizationObject );
-			FFilePlatformRequest* ProcessedFile = FilesProcessed.Find(Filename);
+			FFilePlatformCookedPackage* ProcessedFile = FilesProcessed.Find(Filename);
 			if( ProcessedFile )
 			{
 				ProcessedFile->RemovePlatform(PlatformName);
@@ -405,13 +404,16 @@ private:
 			return false;
 		}
 
-		void GetCookedFilesForPlatform(const FName& PlatformName, TArray<FName>& CookedFiles)
+		void GetCookedFilesForPlatform(const FName& PlatformName, TArray<FName>& CookedFiles, bool bGetFailedCookedPackages = true, bool bGetSuccessfulCookedPackages = true)
 		{
 			FScopeLock ScopeLock(&SynchronizationObject);
 			for (const auto& CookedFile : FilesProcessed)
 			{
 				if (CookedFile.Value.HasPlatform(PlatformName))
 				{
+					bool bHasSucceededSavePackage = CookedFile.Value.HasSucceededSavePackage();
+					if ((bHasSucceededSavePackage && bGetSuccessfulCookedPackages) ||
+						((bHasSucceededSavePackage == false) && bGetFailedCookedPackages) )
 					CookedFiles.Add(CookedFile.Value.GetFilename());
 				}
 			}
diff --git a/Engine/Source/Editor/UnrealEd/Private/CookOnTheFlyServer.cpp b/Engine/Source/Editor/UnrealEd/Private/CookOnTheFlyServer.cpp
--- a/Engine/Source/Editor/UnrealEd/Private/CookOnTheFlyServer.cpp
+++ b/Engine/Source/Editor/UnrealEd/Private/CookOnTheFlyServer.cpp
@@ -1841,7 +1841,13 @@ uint32 UCookOnTheFlyServer::TickCookOnTheSide( const float TimeSlice, uint32 &Co
 #if DEBUG_COOKONTHEFLY
 			UE_LOG(LogCook, Display, TEXT("Not cooking package %s"), *ToBuild.GetFilename().ToString());
 #endif
-			CookedPackages.Add( FFilePlatformRequest( ToBuild.GetFilename(), TargetPlatformNames ) );
+			// did not cook this package 
+#if DO_CHECK
+			// make sure this package doesn't exist
+			const FString SandboxFilename = ConvertToFullSandboxPath(ToBuild.GetFilename().ToString(), false);
+			check(IFileManager::Get().FileExists(*SandboxFilename) == false);
+#endif
+			CookedPackages.Add( FFilePlatformCookedPackage( ToBuild.GetFilename(), TargetPlatformNames, false) );
 			continue;
 		}
 
@@ -2166,13 +2172,16 @@ uint32 UCookOnTheFlyServer::TickCookOnTheSide( const float TimeSlice, uint32 &Co
 
 				ESavePackageResult SavePackageResult = ESavePackageResult::Error;
 				END_PACKAGE_STAT(SavePackageCacheForCookedPlatformData);
+				bool bSucceededSavePackage = false;
 				{
 					SCOPE_PACKAGE_STAT(SaveCookedPackage);
 
 					SCOPE_TIMER(SaveCookedPackage);
-					SavePackageResult = SaveCookedPackage(Package, SAVE_KeepGUID | (bShouldSaveAsync ? SAVE_Async : SAVE_None) | (IsCookFlagSet(ECookInitializationFlags::Unversioned) ? SAVE_Unversioned : 0), bWasUpToDate, AllTargetPlatformNames);
+					uint32 SaveFlags = SAVE_KeepGUID | (bShouldSaveAsync ? SAVE_Async : SAVE_None) | (IsCookFlagSet(ECookInitializationFlags::Unversioned) ? SAVE_Unversioned : 0);
+					SavePackageResult = SaveCookedPackage(Package, SaveFlags, bWasUpToDate, AllTargetPlatformNames);
 					if (SavePackageResult == ESavePackageResult::Success)
 					{
+						bSucceededSavePackage = true;
 						// Update flags used to determine garbage collection.
 						if (Package->ContainsMap())
 						{
@@ -2214,7 +2223,7 @@ uint32 UCookOnTheFlyServer::TickCookOnTheSide( const float TimeSlice, uint32 &Co
 				if (StandardFilename != NAME_None)
 				{
 					// mark the package as cooked
-					FFilePlatformRequest FileRequest(StandardFilename, AllTargetPlatformNames);
+					FFilePlatformCookedPackage FileRequest(StandardFilename, AllTargetPlatformNames, bSucceededSavePackage);
 					if (SavePackageResult != ESavePackageResult::ReferencedOnlyByEditorOnlyData)
 					{
 						CookedPackages.Add(FileRequest);
@@ -2246,7 +2255,7 @@ uint32 UCookOnTheFlyServer::TickCookOnTheSide( const float TimeSlice, uint32 &Co
 		// our original package request will fail however
 		if ( !CookedPackages.Exists(ToBuild) )
 		{
-			CookedPackages.Add( ToBuild );
+			CookedPackages.Add( FFilePlatformCookedPackage(ToBuild, false) );
 		}
 
 		// TODO: Daniel: this is reference code needs to be reimplemented on the callee side.
@@ -3165,7 +3174,9 @@ void UCookOnTheFlyServer::PopulateCookedPackagesFromDisk( const TArray<ITargetPl
 
 				if (Diff >= 0.0)
 				{
-					CookedPackages.Add(FFilePlatformRequest( FName(*StandardCookedFilename), PlatformFName) );
+					TArray<FName> PlatformNames;
+					PlatformNames.Add(PlatformFName);
+					CookedPackages.Add(FFilePlatformCookedPackage( FName(*StandardCookedFilename), MoveTemp(PlatformNames), true) );
 					UE_LOG(LogCook, Display, TEXT("Package %s for %s up to date, avoiding cook"), *StandardCookedFilename, *PlatformFName.ToString());
 				}
 			}
@@ -3884,17 +3895,27 @@ void UCookOnTheFlyServer::CookByTheBookFinished()
 			Manifest.Value->SaveManifests(SandboxFile.GetOwnedPointer());
 
 			const FName& PlatformName = Manifest.Key;
-			const TArray<FName>& PackageFilenames = CookByTheBookOptions->BasedOnReleaseCookedPackages.FindRef(PlatformName);
-			TArray<FName> LongPackageNames;
-			LongPackageNames.Empty(PackageFilenames.Num());
-			for (const auto& PackageFilename : PackageFilenames)
+			TArray<FName> IgnorePackageFilenames;
+
+			// ignore any packages which failed to cook
+			CookedPackages.GetCookedFilesForPlatform(PlatformName, IgnorePackageFilenames, true, false);
+			// ignore any packages which were in a previous release of cooked content
+			IgnorePackageFilenames.Append( CookByTheBookOptions->BasedOnReleaseCookedPackages.FindRef(PlatformName) );
+
+
+			TArray<FName> IgnorePackageNames;
+			IgnorePackageNames.Empty(IgnorePackageFilenames.Num());
+			for (const auto& PackageFilename : IgnorePackageFilenames)
 			{
 				FString LongPackageName;
 				verify(FPackageName::TryConvertFilenameToLongPackageName(PackageFilename.ToString(), LongPackageName));
-				LongPackageNames.Add(FName(*LongPackageName));
+				IgnorePackageNames.Add(FName(*LongPackageName));
 			}
 
-			Manifest.Value->SaveAssetRegistry(SandboxRegistryFilename, &LongPackageNames);
+			
+
+
+			Manifest.Value->SaveAssetRegistry(SandboxRegistryFilename, &IgnorePackageNames);
 
 			Manifest.Value->SaveCookedPackageAssetRegistry(SandboxCookedAssetRegistryFilename, true);
 
@@ -4292,7 +4313,7 @@ void UCookOnTheFlyServer::StartCookByTheBook( const FCookByTheBookStartupOptions
 				PlatformNames.Add(PlatformName);
 				for (const auto& PackageFilename : PackageList)
 				{
-					CookedPackages.Add( MoveTemp( FFilePlatformRequest( PackageFilename, PlatformNames ) ) );
+					CookedPackages.Add( MoveTemp( FFilePlatformCookedPackage( PackageFilename, PlatformNames, true ) ) );
 				}
 			}
 			CookByTheBookOptions->BasedOnReleaseCookedPackages.Add(PlatformName, MoveTemp(PackageList));
@@ -4632,7 +4653,7 @@ void UCookOnTheFlyServer::StartChildCookers(int32 NumCookersToSpawn, const TArra
 			FPaths::MakeStandardFilename(PackageName);
 			FName StandardPackageName = FName(*PackageName);*/
 			
-			CookedPackages.Add(FFilePlatformRequest(StandardPackageName, TargetPlatformNames));
+			CookedPackages.Add(FFilePlatformCookedPackage(StandardPackageName, TargetPlatformNames, true));
 		}
 		DistributeStandardFilenames.RemoveAt(0, NumFilesForCooker);
 
