Fix issues with UserDefinedEnums storing non-contiguous values.

diff --git a/Engine/Source/Editor/UnrealEd/Private/Kismet2/EnumEditorUtils.cpp b/Engine/Source/Editor/UnrealEd/Private/Kismet2/EnumEditorUtils.cpp
--- a/Engine/Source/Editor/UnrealEd/Private/Kismet2/EnumEditorUtils.cpp
+++ b/Engine/Source/Editor/UnrealEd/Private/Kismet2/EnumEditorUtils.cpp
@@ -57,7 +57,7 @@ void FEnumEditorUtils::UpdateAfterPathChanged(UEnum* Enum)
 	{
 		const FString ShortName = Enum->GetEnumName(Index);
 		const FString NewFullName = Enum->GenerateFullEnumName(*ShortName);
-		NewEnumeratorsNames.Add(TPairInitializer<FName, uint8>(FName(*NewFullName), Enum->GetValueByIndex(Index)));
+		NewEnumeratorsNames.Add(TPairInitializer<FName, uint8>(FName(*NewFullName), Index));
 	}
 
 	Enum->SetEnums(NewEnumeratorsNames, UEnum::ECppForm::Namespaced);
@@ -98,6 +98,11 @@ void FEnumEditorUtils::AddNewEnumeratorForUserDefinedEnum(UUserDefinedEnum* Enum
 	const FString FullNameStr = Enum->GenerateFullEnumName(*EnumNameString);
 	Names.Add(TPairInitializer<FName, uint8>(FName(*FullNameStr), Enum->GetMaxEnumValue()));
 
+	// Clean up enum values.
+	for (int32 i = 0; i < Names.Num(); ++i)
+	{
+		Names[i].Value = i;
+	}
 	const UEnum::ECppForm EnumType = Enum->GetCppForm();
 	Enum->SetEnums(Names, EnumType);
 	EnsureAllDisplayNamesExist(Enum);
@@ -124,6 +129,11 @@ void FEnumEditorUtils::RemoveEnumeratorFromUserDefinedEnum(UUserDefinedEnum* Enu
 
 	Enum->RemoveMetaData(FEnumEditorUtilsHelper::DisplayName(), EnumeratorIndex);
 
+	// Clean up enum values.
+	for (int32 i = 0; i < Names.Num(); ++i)
+	{
+		Names[i].Value = i;
+	}
 	const UEnum::ECppForm EnumType = Enum->GetCppForm();
 	Enum->SetEnums(Names, EnumType);
 	EnsureAllDisplayNamesExist(Enum);
@@ -158,6 +168,11 @@ void FEnumEditorUtils::MoveEnumeratorInUserDefinedEnum(UUserDefinedEnum* Enum, i
 		Names.Swap(EnumeratorIndex, EnumeratorIndex + 1);
 	}
 
+	// Clean up enum values.
+	for (int32 i = 0; i < Names.Num(); ++i)
+	{
+		Names[i].Value = i;
+	}
 	const UEnum::ECppForm EnumType = Enum->GetCppForm();
 	Enum->SetEnums(Names, EnumType);
 	EnsureAllDisplayNamesExist(Enum);
diff --git a/Engine/Source/Runtime/Engine/Classes/Engine/UserDefinedEnum.h b/Engine/Source/Runtime/Engine/Classes/Engine/UserDefinedEnum.h
--- a/Engine/Source/Runtime/Engine/Classes/Engine/UserDefinedEnum.h
+++ b/Engine/Source/Runtime/Engine/Classes/Engine/UserDefinedEnum.h
@@ -34,6 +34,9 @@ class ENGINE_API UUserDefinedEnum : public UEnum
 	TArray<FText> DisplayNames;
 
 public:
+	// UObject interface.
+	virtual void Serialize(FArchive& Ar) override;
+	// End of UObject interface.
 	/**
 	 * Generates full enum name give enum name.
 	 * For UUserDefinedEnum full enumerator name has form: '<enumeration path>::<short, user defined enumerator name>'
diff --git a/Engine/Source/Runtime/Engine/Private/UserDefinedEnum.cpp b/Engine/Source/Runtime/Engine/Private/UserDefinedEnum.cpp
--- a/Engine/Source/Runtime/Engine/Private/UserDefinedEnum.cpp
+++ b/Engine/Source/Runtime/Engine/Private/UserDefinedEnum.cpp
@@ -12,6 +12,21 @@ UUserDefinedEnum::UUserDefinedEnum(const FObjectInitializer& ObjectInitializer)
 
 }
 
+void UUserDefinedEnum::Serialize(FArchive& Ar)
+{
+	Super::Serialize(Ar);
+
+#if WITH_EDITOR
+	if (Ar.IsLoading() && Ar.IsPersistent())
+	{
+		for (int32 i = 0; i < Names.Num(); ++i)
+		{
+			Names[i].Value = i;
+		}
+	}
+#endif // WITH_EDITOR
+}
+
 FString UUserDefinedEnum::GenerateFullEnumName(const TCHAR* InEnumName) const
 {
 	check(CppForm == ECppForm::Namespaced);
@@ -54,6 +69,11 @@ void UUserDefinedEnum::PostLoad()
 	Super::PostLoad();
 	FEnumEditorUtils::UpdateAfterPathChanged(this);
 	FEnumEditorUtils::EnsureAllDisplayNamesExist(this);
+
+	for (int32 i = 0; i < Names.Num(); ++i)
+	{
+		Names[i].Value = i;
+	}
 }
 
 void UUserDefinedEnum::PostEditUndo()
