Removing an interface and keeping the functions will correctly make the pins on the function available in the details panel.

diff --git a/Engine/Source/Editor/BlueprintGraph/Classes/K2Node_FunctionResult.h b/Engine/Source/Editor/BlueprintGraph/Classes/K2Node_FunctionResult.h
--- a/Engine/Source/Editor/BlueprintGraph/Classes/K2Node_FunctionResult.h
+++ b/Engine/Source/Editor/BlueprintGraph/Classes/K2Node_FunctionResult.h
@@ -29,9 +29,13 @@ class UK2Node_FunctionResult : public UK2Node_FunctionTerminator
 	virtual void PostPasteNode() override;
 	// End UK2Node interface
 
-	// Begin K2Node_FunctionTerminator interface
+	// Begin UK2Node_EditablePinBase interface
 	virtual bool CanCreateUserDefinedPin(const FEdGraphPinType& InPinType, EEdGraphPinDirection InDesiredDirection, FText& OutErrorMessage) override;
 	virtual UEdGraphPin* CreatePinFromUserDefinition(const TSharedPtr<FUserPinInfo> NewPinInfo) override;
+	// End UK2Node_EditablePinBase interface
+
+	// Begin K2Node_FunctionTerminator interface
+	virtual void PromoteFromInterfaceOverride(bool bIsPrimaryTerminator = true);
 	// End K2Node_FunctionTerminator interface
 
 	BLUEPRINTGRAPH_API TArray<UK2Node_FunctionResult*> GetAllResultNodes() const;
diff --git a/Engine/Source/Editor/BlueprintGraph/Classes/K2Node_FunctionTerminator.h b/Engine/Source/Editor/BlueprintGraph/Classes/K2Node_FunctionTerminator.h
--- a/Engine/Source/Editor/BlueprintGraph/Classes/K2Node_FunctionTerminator.h
+++ b/Engine/Source/Editor/BlueprintGraph/Classes/K2Node_FunctionTerminator.h
@@ -36,5 +36,8 @@ class UK2Node_FunctionTerminator : public UK2Node_EditablePinBase
 	// Begin UK2Node_EditablePinBase interface
 	virtual bool CanCreateUserDefinedPin(const FEdGraphPinType& InPinType, EEdGraphPinDirection InDesiredDirection, FText& OutErrorMessage) override;
 	// End UK2Node_EditablePinBase interface
+
+	/** Promotes the node from being a part of an interface override to a full function that allows for parameter and result pin additions */
+	virtual void PromoteFromInterfaceOverride(bool bIsPrimaryTerminator = true);
 };
 
diff --git a/Engine/Source/Editor/BlueprintGraph/Private/K2Node_FunctionResult.cpp b/Engine/Source/Editor/BlueprintGraph/Private/K2Node_FunctionResult.cpp
--- a/Engine/Source/Editor/BlueprintGraph/Private/K2Node_FunctionResult.cpp
+++ b/Engine/Source/Editor/BlueprintGraph/Private/K2Node_FunctionResult.cpp
@@ -304,4 +304,20 @@ void UK2Node_FunctionResult::ValidateNodeDuringCompilation(class FCompilerResult
 			}
 		}
 	}
+}
+
+void UK2Node_FunctionResult::PromoteFromInterfaceOverride(bool bIsPrimaryTerminator/* = true*/)
+{
+	// For non-primary terminators, we want to sync with the primary one and reconstruct.
+	if (bIsPrimaryTerminator)
+	{
+		Super::PromoteFromInterfaceOverride();
+	}
+	else
+	{
+		SignatureClass = nullptr;
+		SyncWithPrimaryResultNode();
+		const UEdGraphSchema_K2* Schema = GetDefault<UEdGraphSchema_K2>();
+		Schema->ReconstructNode(*this, true);
+	}
 }
\ No newline at end of file
diff --git a/Engine/Source/Editor/BlueprintGraph/Private/K2Node_FunctionTerminator.cpp b/Engine/Source/Editor/BlueprintGraph/Private/K2Node_FunctionTerminator.cpp
--- a/Engine/Source/Editor/BlueprintGraph/Private/K2Node_FunctionTerminator.cpp
+++ b/Engine/Source/Editor/BlueprintGraph/Private/K2Node_FunctionTerminator.cpp
@@ -65,4 +65,19 @@ bool UK2Node_FunctionTerminator::HasExternalDependencies(TArray<class UStruct*>*
 	return bSuperResult || bResult;
 }
 
+void UK2Node_FunctionTerminator::PromoteFromInterfaceOverride(bool bIsPrimaryTerminator)
+{
+	// Remove the signature class, that is not relevant.
+	SignatureClass = nullptr;
+	for (const UEdGraphPin* Pin : Pins)
+	{
+		if (Pin->PinType.PinCategory != UEdGraphSchema_K2::PC_Exec)
+		{
+			CreateUserDefinedPin(Pin->PinName, Pin->PinType, (Pin->Direction == EGPD_Input)? EGPD_Output : EGPD_Input, false);
+		}
+	}
+	const UEdGraphSchema_K2* Schema = GetDefault<UEdGraphSchema_K2>();
+	Schema->ReconstructNode(*this, true);
+}
+
 #undef LOCTEXT_NAMESPACE
diff --git a/Engine/Source/Editor/UnrealEd/Private/Kismet2/BlueprintEditorUtils.cpp b/Engine/Source/Editor/UnrealEd/Private/Kismet2/BlueprintEditorUtils.cpp
--- a/Engine/Source/Editor/UnrealEd/Private/Kismet2/BlueprintEditorUtils.cpp
+++ b/Engine/Source/Editor/UnrealEd/Private/Kismet2/BlueprintEditorUtils.cpp
@@ -5448,6 +5448,33 @@ void FBlueprintEditorUtils::RemoveInterface(UBlueprint* Blueprint, const FName&
 				Schema->MarkFunctionEntryAsEditable(CurrentGraph, true);
 
 				Blueprint->FunctionGraphs.Add(CurrentGraph);
+
+				// Move all non-exec pins from the function entry node to being user defined pins
+				TArray<UK2Node_FunctionEntry*> FunctionEntryNodes;
+				CurrentGraph->GetNodesOfClass(FunctionEntryNodes);
+				if (FunctionEntryNodes.Num() > 0)
+				{
+					UK2Node_FunctionEntry* FunctionEntry = FunctionEntryNodes[0];
+					FunctionEntry->PromoteFromInterfaceOverride();
+				}
+
+				// Move all non-exec pins from the function result node to being user defined pins
+				TArray<UK2Node_FunctionResult*> FunctionResultNodes;
+				CurrentGraph->GetNodesOfClass(FunctionResultNodes);
+				if (FunctionResultNodes.Num() > 0)
+				{
+					UK2Node_FunctionResult* PrimaryFunctionResult = FunctionResultNodes[0];
+					PrimaryFunctionResult->PromoteFromInterfaceOverride();
+
+					// Reconstruct all result nodes so they update their pins accordingly
+					for (UK2Node_FunctionResult* FunctionResult : FunctionResultNodes)
+					{
+						if (PrimaryFunctionResult != FunctionResult)
+						{
+							FunctionResult->PromoteFromInterfaceOverride(false);
+						}
+					}
+				}
 			}
 			else
 			{
