Removing unnecessary include

diff --git a/Engine/Source/Editor/UnrealEd/Public/GraphEditor.h b/Engine/Source/Editor/UnrealEd/Public/GraphEditor.h
--- a/Engine/Source/Editor/UnrealEd/Public/GraphEditor.h
+++ b/Engine/Source/Editor/UnrealEd/Public/GraphEditor.h
@@ -2,7 +2,6 @@
 
 #pragma once
 
-#include "UnrealEd.h"
 #include "SlateBasics.h"
 #include "BlueprintUtilities.h"
 
