[INTEGRATE] Compile times: Working on eliminating <Module>Classes.h header files
- This checkin eliminates use of all <Module>Classes.h files in Engine, with the exception of UnrealEdClasses.h
- Compilation speed-ups for all of the modules affected, now that they are including what they use
- No effective changes other than C++ includes

diff --git a/Engine/Plugins/2D/Paper2D/Source/Paper2D/Private/Paper2DPrivatePCH.h b/Engine/Plugins/2D/Paper2D/Source/Paper2D/Private/Paper2DPrivatePCH.h
--- a/Engine/Plugins/2D/Paper2D/Source/Paper2D/Private/Paper2DPrivatePCH.h
+++ b/Engine/Plugins/2D/Paper2D/Source/Paper2D/Private/Paper2DPrivatePCH.h
@@ -2,12 +2,8 @@
 
 #pragma once
 
-// Uncommenting this like causes the FTransform linkage issue
-//#include "Core.h"
-
 #include "Engine.h"
 
 DECLARE_LOG_CATEGORY_EXTERN(LogPaper2D, Verbose, All);
 
-#include "Paper2DClasses.h"
 #include "Paper2DModule.h"
\ No newline at end of file
diff --git a/Engine/Plugins/2D/Paper2D/Source/Paper2DEditor/Private/Paper2DEditorPrivatePCH.h b/Engine/Plugins/2D/Paper2D/Source/Paper2DEditor/Private/Paper2DEditorPrivatePCH.h
--- a/Engine/Plugins/2D/Paper2D/Source/Paper2DEditor/Private/Paper2DEditorPrivatePCH.h
+++ b/Engine/Plugins/2D/Paper2D/Source/Paper2DEditor/Private/Paper2DEditorPrivatePCH.h
@@ -3,9 +3,7 @@
 #pragma once
 
 #include "UnrealEd.h"
-#include "Paper2DClasses.h"
 #include "Paper2DModule.h"
-#include "Paper2DEditorClasses.h"
 
 DECLARE_LOG_CATEGORY_EXTERN(LogPaper2DEditor, Log, All);
 
diff --git a/Engine/Plugins/2D/Paper2D/Source/PaperSpriteSheetImporter/Private/PaperSpriteSheetImporterPrivatePCH.h b/Engine/Plugins/2D/Paper2D/Source/PaperSpriteSheetImporter/Private/PaperSpriteSheetImporterPrivatePCH.h
--- a/Engine/Plugins/2D/Paper2D/Source/PaperSpriteSheetImporter/Private/PaperSpriteSheetImporterPrivatePCH.h
+++ b/Engine/Plugins/2D/Paper2D/Source/PaperSpriteSheetImporter/Private/PaperSpriteSheetImporterPrivatePCH.h
@@ -3,7 +3,6 @@
 #pragma once
 
 #include "UnrealEd.h"
-//#include "Paper2DClasses.h"
 
 DECLARE_LOG_CATEGORY_EXTERN(LogPaperSpriteSheetImporter, Verbose, All);
 
diff --git a/Engine/Plugins/2D/Paper2D/Source/PaperTiledImporter/Private/PaperTiledImporterPrivatePCH.h b/Engine/Plugins/2D/Paper2D/Source/PaperTiledImporter/Private/PaperTiledImporterPrivatePCH.h
--- a/Engine/Plugins/2D/Paper2D/Source/PaperTiledImporter/Private/PaperTiledImporterPrivatePCH.h
+++ b/Engine/Plugins/2D/Paper2D/Source/PaperTiledImporter/Private/PaperTiledImporterPrivatePCH.h
@@ -6,4 +6,4 @@
 
 DECLARE_LOG_CATEGORY_EXTERN(LogPaperTiledImporter, Verbose, All);
 
-#include "PaperTiledImporterClasses.h"
+#include "Classes/PaperTiledImporterFactory.h"
diff --git a/Engine/Plugins/Experimental/BlueprintStats/Source/BlueprintStats/Private/BlueprintStats.cpp b/Engine/Plugins/Experimental/BlueprintStats/Source/BlueprintStats/Private/BlueprintStats.cpp
--- a/Engine/Plugins/Experimental/BlueprintStats/Source/BlueprintStats/Private/BlueprintStats.cpp
+++ b/Engine/Plugins/Experimental/BlueprintStats/Source/BlueprintStats/Private/BlueprintStats.cpp
@@ -3,6 +3,8 @@
 #include "BlueprintStatsPrivatePCH.h"
 #include "BlueprintStats.h"
 #include "BlueprintEditorUtils.h"
+#include "K2Node_CallFunction.h"
+#include "K2Node_MacroInstance.h"
 
 #define LOCTEXT_NAMESPACE "BlueprintStats"
 
diff --git a/Engine/Plugins/Experimental/BlueprintStats/Source/BlueprintStats/Private/BlueprintStatsPrivatePCH.h b/Engine/Plugins/Experimental/BlueprintStats/Source/BlueprintStats/Private/BlueprintStatsPrivatePCH.h
--- a/Engine/Plugins/Experimental/BlueprintStats/Source/BlueprintStats/Private/BlueprintStatsPrivatePCH.h
+++ b/Engine/Plugins/Experimental/BlueprintStats/Source/BlueprintStats/Private/BlueprintStatsPrivatePCH.h
@@ -1,4 +1,3 @@
 // Copyright 1998-2015 Epic Games, Inc. All Rights Reserved.
 
 #include "UnrealEd.h"
-#include "BlueprintGraphClasses.h"
diff --git a/Engine/Plugins/NetcodeUnitTest/NetcodeUnitTest/Source/NetcodeUnitTest/Classes/Net/UnitTestNetConnection.h b/Engine/Plugins/NetcodeUnitTest/NetcodeUnitTest/Source/NetcodeUnitTest/Classes/Net/UnitTestNetConnection.h
--- a/Engine/Plugins/NetcodeUnitTest/NetcodeUnitTest/Source/NetcodeUnitTest/Classes/Net/UnitTestNetConnection.h
+++ b/Engine/Plugins/NetcodeUnitTest/NetcodeUnitTest/Source/NetcodeUnitTest/Classes/Net/UnitTestNetConnection.h
@@ -2,8 +2,7 @@
 
 #pragma once
 
-#include "OnlineSubsystemUtilsClasses.h"
-
+#include "IpConnection.h"
 #include "UnitTestNetConnection.generated.h"
 
 
diff --git a/Engine/Plugins/NetcodeUnitTest/NetcodeUnitTest/Source/NetcodeUnitTest/Classes/Net/UnitTestNetDriver.h b/Engine/Plugins/NetcodeUnitTest/NetcodeUnitTest/Source/NetcodeUnitTest/Classes/Net/UnitTestNetDriver.h
--- a/Engine/Plugins/NetcodeUnitTest/NetcodeUnitTest/Source/NetcodeUnitTest/Classes/Net/UnitTestNetDriver.h
+++ b/Engine/Plugins/NetcodeUnitTest/NetcodeUnitTest/Source/NetcodeUnitTest/Classes/Net/UnitTestNetDriver.h
@@ -2,8 +2,7 @@
 
 #pragma once
 
-#include "OnlineSubsystemUtilsClasses.h"
-
+#include "IpNetDriver.h"
 #include "UnitTestNetDriver.generated.h"
 
 
diff --git a/Engine/Plugins/NetcodeUnitTest/NetcodeUnitTest/Source/NetcodeUnitTest/Private/ClientUnitTest.cpp b/Engine/Plugins/NetcodeUnitTest/NetcodeUnitTest/Source/NetcodeUnitTest/Private/ClientUnitTest.cpp
--- a/Engine/Plugins/NetcodeUnitTest/NetcodeUnitTest/Source/NetcodeUnitTest/Private/ClientUnitTest.cpp
+++ b/Engine/Plugins/NetcodeUnitTest/NetcodeUnitTest/Source/NetcodeUnitTest/Private/ClientUnitTest.cpp
@@ -20,6 +20,7 @@
 #include "SLogWindow.h"
 #include "SLogWidget.h"
 
+#include "OnlineBeaconClient.h"
 
 #include "Engine/ActorChannel.h"
 #include "Regex.h"
diff --git a/Engine/Plugins/NetcodeUnitTest/NetcodeUnitTest/Source/NetcodeUnitTest/Private/Net/NUTUtilNet.cpp b/Engine/Plugins/NetcodeUnitTest/NetcodeUnitTest/Source/NetcodeUnitTest/Private/Net/NUTUtilNet.cpp
--- a/Engine/Plugins/NetcodeUnitTest/NetcodeUnitTest/Source/NetcodeUnitTest/Private/Net/NUTUtilNet.cpp
+++ b/Engine/Plugins/NetcodeUnitTest/NetcodeUnitTest/Source/NetcodeUnitTest/Private/Net/NUTUtilNet.cpp
@@ -14,7 +14,7 @@
 #include "UnrealNetwork.h"
 #include "EngineVersion.h"
 #include "DataChannel.h"
-
+#include "OnlineBeaconClient.h"
 
 // Forward declarations
 class FWorldTickHook;
diff --git a/Engine/Plugins/NetcodeUnitTest/NetcodeUnitTest/Source/NetcodeUnitTest/Private/Net/UnitTestNetConnection.cpp b/Engine/Plugins/NetcodeUnitTest/NetcodeUnitTest/Source/NetcodeUnitTest/Private/Net/UnitTestNetConnection.cpp
--- a/Engine/Plugins/NetcodeUnitTest/NetcodeUnitTest/Source/NetcodeUnitTest/Private/Net/UnitTestNetConnection.cpp
+++ b/Engine/Plugins/NetcodeUnitTest/NetcodeUnitTest/Source/NetcodeUnitTest/Private/Net/UnitTestNetConnection.cpp
@@ -4,7 +4,6 @@
 
 #include "Engine.h"
 #include "Net/UnrealNetwork.h"
-#include "OnlineSubsystemUtilsClasses.h"
 
 #include "NUTUtil.h"
 #include "Net/NUTUtilNet.h"
diff --git a/Engine/Plugins/Runtime/Analytics/AnalyticsBlueprintLibrary/Source/AnalyticsBlueprintLibrary/Private/AnalyticsBlueprintLibraryPrivatePCH.h b/Engine/Plugins/Runtime/Analytics/AnalyticsBlueprintLibrary/Source/AnalyticsBlueprintLibrary/Private/AnalyticsBlueprintLibraryPrivatePCH.h
--- a/Engine/Plugins/Runtime/Analytics/AnalyticsBlueprintLibrary/Source/AnalyticsBlueprintLibrary/Private/AnalyticsBlueprintLibraryPrivatePCH.h
+++ b/Engine/Plugins/Runtime/Analytics/AnalyticsBlueprintLibrary/Source/AnalyticsBlueprintLibrary/Private/AnalyticsBlueprintLibraryPrivatePCH.h
@@ -7,4 +7,4 @@
 
 #include "AnalyticsBlueprintLibraryModule.h"
 
-#include "AnalyticsBlueprintLibraryClasses.h"
+#include "Classes/AnalyticsBlueprintLibrary.h"
diff --git a/Engine/Plugins/Runtime/SoundMod/Source/SoundModImporter/Private/SoundModImporterPrivatePCH.h b/Engine/Plugins/Runtime/SoundMod/Source/SoundModImporter/Private/SoundModImporterPrivatePCH.h
--- a/Engine/Plugins/Runtime/SoundMod/Source/SoundModImporter/Private/SoundModImporterPrivatePCH.h
+++ b/Engine/Plugins/Runtime/SoundMod/Source/SoundModImporter/Private/SoundModImporterPrivatePCH.h
@@ -7,4 +7,4 @@
 
 DECLARE_LOG_CATEGORY_EXTERN(LogSoundModImporter, Verbose, All);
 
-#include "SoundModImporterClasses.h"
+#include "Classes/SoundModImporterFactory.h"
diff --git a/Engine/Source/Developer/LogVisualizer/Private/LogVisualizer.cpp b/Engine/Source/Developer/LogVisualizer/Private/LogVisualizer.cpp
--- a/Engine/Source/Developer/LogVisualizer/Private/LogVisualizer.cpp
+++ b/Engine/Source/Developer/LogVisualizer/Private/LogVisualizer.cpp
@@ -10,7 +10,6 @@
 #	include "Editor/UnrealEd/Public/EditorReimportHandler.h"
 #	include "Editor/UnrealEd/Public/TexAlignTools.h"
 #	include "Editor/UnrealEd/Public/TickableEditorObject.h"
-#	include "UnrealEdClasses.h"
 #	include "Editor/UnrealEd/Public/Editor.h"
 #	include "Editor/UnrealEd/Public/EditorViewportClient.h"
 #endif
diff --git a/Engine/Source/Developer/LogVisualizer/Private/STimelineBar.cpp b/Engine/Source/Developer/LogVisualizer/Private/STimelineBar.cpp
--- a/Engine/Source/Developer/LogVisualizer/Private/STimelineBar.cpp
+++ b/Engine/Source/Developer/LogVisualizer/Private/STimelineBar.cpp
@@ -9,7 +9,6 @@
 #	include "Editor/UnrealEd/Public/EditorReimportHandler.h"
 #	include "Editor/UnrealEd/Public/TexAlignTools.h"
 #	include "Editor/UnrealEd/Public/TickableEditorObject.h"
-#	include "UnrealEdClasses.h"
 #	include "Editor/UnrealEd/Public/Editor.h"
 #	include "Editor/UnrealEd/Public/EditorViewportClient.h"
 #endif
diff --git a/Engine/Source/Developer/LogVisualizer/Private/SVisualLogger.cpp b/Engine/Source/Developer/LogVisualizer/Private/SVisualLogger.cpp
--- a/Engine/Source/Developer/LogVisualizer/Private/SVisualLogger.cpp
+++ b/Engine/Source/Developer/LogVisualizer/Private/SVisualLogger.cpp
@@ -14,7 +14,6 @@
 #	include "Editor/UnrealEd/Public/EditorReimportHandler.h"
 #	include "Editor/UnrealEd/Public/TexAlignTools.h"
 #	include "Editor/UnrealEd/Public/TickableEditorObject.h"
-#	include "UnrealEdClasses.h"
 #	include "Editor/UnrealEd/Public/Editor.h"
 #	include "Editor/UnrealEd/Public/EditorViewportClient.h"
 #endif
diff --git a/Engine/Source/Developer/LogVisualizer/Private/SVisualLoggerReport.cpp b/Engine/Source/Developer/LogVisualizer/Private/SVisualLoggerReport.cpp
--- a/Engine/Source/Developer/LogVisualizer/Private/SVisualLoggerReport.cpp
+++ b/Engine/Source/Developer/LogVisualizer/Private/SVisualLoggerReport.cpp
@@ -9,7 +9,6 @@
 #include "Editor/UnrealEd/Public/EditorReimportHandler.h"
 #include "Editor/UnrealEd/Public/TexAlignTools.h"
 #include "Editor/UnrealEd/Public/TickableEditorObject.h"
-#include "UnrealEdClasses.h"
 #include "Editor/UnrealEd/Public/Editor.h"
 #include "Editor/UnrealEd/Public/EditorViewportClient.h"
 #include "VisualLoggerTypes.h"
diff --git a/Engine/Source/Editor/AIGraph/Private/AIGraphPrivatePCH.h b/Engine/Source/Editor/AIGraph/Private/AIGraphPrivatePCH.h
--- a/Engine/Source/Editor/AIGraph/Private/AIGraphPrivatePCH.h
+++ b/Engine/Source/Editor/AIGraph/Private/AIGraphPrivatePCH.h
@@ -6,4 +6,8 @@
 #include "GraphEditor.h"
 
 #include "AIGraphModule.h"
-#include "AIGraphClasses.h"
+#include "Classes/AIGraph.h"
+#include "Classes/AIGraphTypes.h"
+#include "Classes/AIGraphNode.h"
+#include "Classes/AIGraphSchema.h"
+
diff --git a/Engine/Source/Editor/BehaviorTreeEditor/Private/BehaviorTreeEditor.cpp b/Engine/Source/Editor/BehaviorTreeEditor/Private/BehaviorTreeEditor.cpp
--- a/Engine/Source/Editor/BehaviorTreeEditor/Private/BehaviorTreeEditor.cpp
+++ b/Engine/Source/Editor/BehaviorTreeEditor/Private/BehaviorTreeEditor.cpp
@@ -26,6 +26,7 @@
 #include "BehaviorTreeEditorUtils.h"
 #include "BehaviorTree/BehaviorTree.h"
 #include "BehaviorTree/BlackboardData.h"
+#include "BehaviorTreeGraphNode_SubtreeTask.h"
 #include "DetailCustomizations/BlackboardDataDetails.h"
 #include "SBehaviorTreeBlackboardView.h"
 #include "SBehaviorTreeBlackboardEditor.h"
diff --git a/Engine/Source/Editor/BehaviorTreeEditor/Private/BehaviorTreeEditorPrivatePCH.h b/Engine/Source/Editor/BehaviorTreeEditor/Private/BehaviorTreeEditorPrivatePCH.h
--- a/Engine/Source/Editor/BehaviorTreeEditor/Private/BehaviorTreeEditorPrivatePCH.h
+++ b/Engine/Source/Editor/BehaviorTreeEditor/Private/BehaviorTreeEditorPrivatePCH.h
@@ -14,7 +14,22 @@
 #include "BehaviorTree/BTService.h"
 #include "BehaviorTree/BTTaskNode.h"
 
-#include "BehaviorTreeEditorClasses.h"
+#include "Classes/BehaviorTreeDecoratorGraph.h"
+#include "Classes/BehaviorTreeDecoratorGraphNode.h"
+#include "Classes/BehaviorTreeEditorTypes.h"
+#include "Classes/BehaviorTreeDecoratorGraphNode_Decorator.h"
+#include "Classes/BehaviorTreeDecoratorGraphNode_Logic.h"
+#include "Classes/BehaviorTreeGraph.h"
+#include "Classes/BehaviorTreeGraphNode.h"
+#include "Classes/BehaviorTreeGraphNode_Composite.h"
+#include "Classes/BehaviorTreeGraphNode_CompositeDecorator.h"
+#include "Classes/BehaviorTreeGraphNode_Decorator.h"
+#include "Classes/BehaviorTreeGraphNode_Root.h"
+#include "Classes/BehaviorTreeGraphNode_Service.h"
+#include "Classes/BehaviorTreeGraphNode_Task.h"
+#include "Classes/EdGraphSchema_BehaviorTree.h"
+#include "Classes/EdGraphSchema_BehaviorTreeDecorator.h"
+
 #include "AssetTypeActions_BehaviorTree.h"
 #include "AssetTypeActions_Blackboard.h"
 
diff --git a/Engine/Source/Editor/BehaviorTreeEditor/Private/BehaviorTreeFactory.cpp b/Engine/Source/Editor/BehaviorTreeEditor/Private/BehaviorTreeFactory.cpp
--- a/Engine/Source/Editor/BehaviorTreeEditor/Private/BehaviorTreeFactory.cpp
+++ b/Engine/Source/Editor/BehaviorTreeEditor/Private/BehaviorTreeFactory.cpp
@@ -1,6 +1,7 @@
 // Copyright 1998-2015 Epic Games, Inc. All Rights Reserved.
 
 #include "BehaviorTreeEditorPrivatePCH.h"
+#include "Classes/BehaviorTreeFactory.h"
 #include "BehaviorTreeEditorModule.h"
 #include "BehaviorTree/BehaviorTree.h"
 
diff --git a/Engine/Source/Editor/BehaviorTreeEditor/Private/BehaviorTreeGraph.cpp b/Engine/Source/Editor/BehaviorTreeEditor/Private/BehaviorTreeGraph.cpp
--- a/Engine/Source/Editor/BehaviorTreeEditor/Private/BehaviorTreeGraph.cpp
+++ b/Engine/Source/Editor/BehaviorTreeEditor/Private/BehaviorTreeGraph.cpp
@@ -8,6 +8,8 @@
 #include "BehaviorTree/BehaviorTree.h"
 #include "BehaviorTree/Composites/BTComposite_SimpleParallel.h"
 #include "BehaviorTree/Tasks/BTTask_Wait.h"
+#include "Classes/BehaviorTreeGraphNode_SimpleParallel.h"
+#include "Classes/BehaviorTreeGraphNode_SubtreeTask.h"
 
 //////////////////////////////////////////////////////////////////////////
 // BehaviorTreeGraph
diff --git a/Engine/Source/Editor/BehaviorTreeEditor/Private/BehaviorTreeGraphNode_SimpleParallel.cpp b/Engine/Source/Editor/BehaviorTreeEditor/Private/BehaviorTreeGraphNode_SimpleParallel.cpp
--- a/Engine/Source/Editor/BehaviorTreeEditor/Private/BehaviorTreeGraphNode_SimpleParallel.cpp
+++ b/Engine/Source/Editor/BehaviorTreeEditor/Private/BehaviorTreeGraphNode_SimpleParallel.cpp
@@ -1,6 +1,7 @@
 // Copyright 1998-2015 Epic Games, Inc. All Rights Reserved.
 
 #include "BehaviorTreeEditorPrivatePCH.h"
+#include "Classes/BehaviorTreeGraphNode_SimpleParallel.h"
 #include "Internationalization.h" // for LOCTEXT()
 
 #define LOCTEXT_NAMESPACE "BehaviorTreeEditor"
diff --git a/Engine/Source/Editor/BehaviorTreeEditor/Private/BehaviorTreeGraphNode_SubtreeTask.cpp b/Engine/Source/Editor/BehaviorTreeEditor/Private/BehaviorTreeGraphNode_SubtreeTask.cpp
--- a/Engine/Source/Editor/BehaviorTreeEditor/Private/BehaviorTreeGraphNode_SubtreeTask.cpp
+++ b/Engine/Source/Editor/BehaviorTreeEditor/Private/BehaviorTreeGraphNode_SubtreeTask.cpp
@@ -2,6 +2,7 @@
 
 #include "BehaviorTreeEditorPrivatePCH.h"
 #include "BehaviorTree/Tasks/BTTask_RunBehavior.h"
+#include "BehaviorTreeGraphNode_SubtreeTask.h"
 
 UBehaviorTreeGraphNode_SubtreeTask::UBehaviorTreeGraphNode_SubtreeTask(const FObjectInitializer& ObjectInitializer) : Super(ObjectInitializer)
 {
diff --git a/Engine/Source/Editor/BehaviorTreeEditor/Private/BlackboardDataFactory.cpp b/Engine/Source/Editor/BehaviorTreeEditor/Private/BlackboardDataFactory.cpp
--- a/Engine/Source/Editor/BehaviorTreeEditor/Private/BlackboardDataFactory.cpp
+++ b/Engine/Source/Editor/BehaviorTreeEditor/Private/BlackboardDataFactory.cpp
@@ -4,7 +4,7 @@
 #include "BehaviorTree/BlackboardData.h"
 #include "BlackboardDataFactory.h"
 
-#define LOCTEXT_NAMESPACE "BehaviorTreeFactory"
+#define LOCTEXT_NAMESPACE "BlackboardDataFactory"
 
 UBlackboardDataFactory::UBlackboardDataFactory(const FObjectInitializer& ObjectInitializer)
 	: Super(ObjectInitializer)
diff --git a/Engine/Source/Editor/BehaviorTreeEditor/Private/EdGraphSchema_BehaviorTree.cpp b/Engine/Source/Editor/BehaviorTreeEditor/Private/EdGraphSchema_BehaviorTree.cpp
--- a/Engine/Source/Editor/BehaviorTreeEditor/Private/EdGraphSchema_BehaviorTree.cpp
+++ b/Engine/Source/Editor/BehaviorTreeEditor/Private/EdGraphSchema_BehaviorTree.cpp
@@ -13,6 +13,8 @@
 #include "BehaviorTree/BTTaskNode.h"
 #include "BehaviorTree/Tasks/BTTask_RunBehavior.h"
 #include "BehaviorTree/Composites/BTComposite_SimpleParallel.h"
+#include "BehaviorTreeGraphNode_SimpleParallel.h"
+#include "BehaviorTreeGraphNode_SubtreeTask.h"
 #include "GenericCommands.h"
 
 #define LOCTEXT_NAMESPACE "BehaviorTreeEditor"
diff --git a/Engine/Source/Editor/BlueprintGraph/Private/K2Node_MathExpression.cpp b/Engine/Source/Editor/BlueprintGraph/Private/K2Node_MathExpression.cpp
--- a/Engine/Source/Editor/BlueprintGraph/Private/K2Node_MathExpression.cpp
+++ b/Engine/Source/Editor/BlueprintGraph/Private/K2Node_MathExpression.cpp
@@ -2,7 +2,6 @@
 
 #include "BlueprintGraphPrivatePCH.h"
 #include "K2Node_MathExpression.h"
-#include "BlueprintGraphClasses.h"
 #include "Kismet2NameValidators.h"
 #include "EdGraphUtilities.h"
 #include "BasicTokenParser.h"
@@ -11,6 +10,7 @@
 #include "BlueprintActionDatabaseRegistrar.h"
 #include "DiffResults.h"
 #include "MathExpressionHandler.h"
+#include "BlueprintNodeSpawner.h"
 
 #define LOCTEXT_NAMESPACE "K2Node"
 
diff --git a/Engine/Source/Editor/Blutility/Private/AssetTypeActions_EditorUtilityBlueprint.cpp b/Engine/Source/Editor/Blutility/Private/AssetTypeActions_EditorUtilityBlueprint.cpp
--- a/Engine/Source/Editor/Blutility/Private/AssetTypeActions_EditorUtilityBlueprint.cpp
+++ b/Engine/Source/Editor/Blutility/Private/AssetTypeActions_EditorUtilityBlueprint.cpp
@@ -1,7 +1,6 @@
 // Copyright 1998-2015 Epic Games, Inc. All Rights Reserved.
 
 #include "BlutilityPrivatePCH.h"
-#include "BlutilityClasses.h"
 #include "AssetTypeActions_EditorUtilityBlueprint.h"
 #include "Toolkits/AssetEditorManager.h"
 #include "BlueprintEditorModule.h"
diff --git a/Engine/Source/Editor/Blutility/Private/BlutilityDetailsPanel.cpp b/Engine/Source/Editor/Blutility/Private/BlutilityDetailsPanel.cpp
--- a/Engine/Source/Editor/Blutility/Private/BlutilityDetailsPanel.cpp
+++ b/Engine/Source/Editor/Blutility/Private/BlutilityDetailsPanel.cpp
@@ -5,7 +5,6 @@
 
 #include "Editor/PropertyEditor/Public/PropertyEditing.h"
 
-#include "BlutilityClasses.h"
 #include "ScopedTransaction.h"
 
 /////////////////////////////////////////////////////
diff --git a/Engine/Source/Editor/Blutility/Private/BlutilityModule.cpp b/Engine/Source/Editor/Blutility/Private/BlutilityModule.cpp
--- a/Engine/Source/Editor/Blutility/Private/BlutilityModule.cpp
+++ b/Engine/Source/Editor/Blutility/Private/BlutilityModule.cpp
@@ -1,7 +1,6 @@
 // Copyright 1998-2015 Epic Games, Inc. All Rights Reserved.
 
 #include "BlutilityPrivatePCH.h"
-#include "BlutilityClasses.h"
 
 
 #include "AssetTypeActions_EditorUtilityBlueprint.h"
diff --git a/Engine/Source/Editor/Blutility/Private/BlutilityPrivatePCH.h b/Engine/Source/Editor/Blutility/Private/BlutilityPrivatePCH.h
--- a/Engine/Source/Editor/Blutility/Private/BlutilityPrivatePCH.h
+++ b/Engine/Source/Editor/Blutility/Private/BlutilityPrivatePCH.h
@@ -6,5 +6,10 @@
 #include "UnrealEd.h"
 
 #include "IBlutilityModule.h"
+#include "Classes/EditorUtilityBlueprint.h"
+#include "Classes/EditorUtilityBlueprintFactory.h"
+#include "Classes/GlobalEditorUtilityBase.h"
+#include "Classes/PlacedEditorUtilityBase.h"
+
 
 #endif // __BlutilityPrivatePCH_h__
diff --git a/Engine/Source/Editor/Blutility/Private/BlutilityShelf.cpp b/Engine/Source/Editor/Blutility/Private/BlutilityShelf.cpp
--- a/Engine/Source/Editor/Blutility/Private/BlutilityShelf.cpp
+++ b/Engine/Source/Editor/Blutility/Private/BlutilityShelf.cpp
@@ -1,7 +1,6 @@
 // Copyright 1998-2015 Epic Games, Inc. All Rights Reserved.
 
 #include "BlutilityPrivatePCH.h"
-#include "BlutilityClasses.h"
 #include "AssetToolsModule.h"
 #include "BlutilityShelf.h"
 
diff --git a/Engine/Source/Editor/Blutility/Private/EditorUtilityBlueprintFactory.cpp b/Engine/Source/Editor/Blutility/Private/EditorUtilityBlueprintFactory.cpp
--- a/Engine/Source/Editor/Blutility/Private/EditorUtilityBlueprintFactory.cpp
+++ b/Engine/Source/Editor/Blutility/Private/EditorUtilityBlueprintFactory.cpp
@@ -3,7 +3,6 @@
 #include "BlutilityPrivatePCH.h"
 #include "ClassViewerModule.h"
 #include "ClassViewerFilter.h"
-#include "BlutilityClasses.h"
 #include "Kismet2/KismetEditorUtilities.h"
 #include "Kismet2/SClassPickerDialog.h"
 #include "Engine/BlueprintGeneratedClass.h"
diff --git a/Engine/Source/Editor/Blutility/Private/GlobalBlutilityDialog.cpp b/Engine/Source/Editor/Blutility/Private/GlobalBlutilityDialog.cpp
--- a/Engine/Source/Editor/Blutility/Private/GlobalBlutilityDialog.cpp
+++ b/Engine/Source/Editor/Blutility/Private/GlobalBlutilityDialog.cpp
@@ -1,7 +1,6 @@
 // Copyright 1998-2015 Epic Games, Inc. All Rights Reserved.
 
 #include "BlutilityPrivatePCH.h"
-#include "BlutilityClasses.h"
 #include "AssetTypeActions_EditorUtilityBlueprint.h"
 #include "Toolkits/AssetEditorManager.h"
 #include "BlueprintEditorModule.h"
diff --git a/Engine/Source/Editor/DestructibleMeshEditor/Private/DestructibleMeshEditorPrivatePCH.h b/Engine/Source/Editor/DestructibleMeshEditor/Private/DestructibleMeshEditorPrivatePCH.h
--- a/Engine/Source/Editor/DestructibleMeshEditor/Private/DestructibleMeshEditorPrivatePCH.h
+++ b/Engine/Source/Editor/DestructibleMeshEditor/Private/DestructibleMeshEditorPrivatePCH.h
@@ -8,4 +8,4 @@
 #include "DestructibleMeshEditor.h"
 #include "SDestructibleMeshEditorViewport.h"
 #include "DestructibleChunkParamsProxy.h"
-#include "DestructibleMeshEditorClasses.h"
+
diff --git a/Engine/Source/Editor/EnvironmentQueryEditor/Private/EnvironmentQueryEditorPrivatePCH.h b/Engine/Source/Editor/EnvironmentQueryEditor/Private/EnvironmentQueryEditorPrivatePCH.h
--- a/Engine/Source/Editor/EnvironmentQueryEditor/Private/EnvironmentQueryEditorPrivatePCH.h
+++ b/Engine/Source/Editor/EnvironmentQueryEditor/Private/EnvironmentQueryEditorPrivatePCH.h
@@ -15,7 +15,13 @@
 #include "SGraphNode.h"
 #include "Editor/UnrealEd/Public/Kismet2/BlueprintEditorUtils.h"
 
-#include "EnvironmentQueryEditorClasses.h"
+#include "Classes/EdGraphSchema_EnvironmentQuery.h"
+#include "Classes/EnvironmentQueryGraph.h"
+#include "Classes/EnvironmentQueryGraphNode.h"
+#include "Classes/EnvironmentQueryGraphNode_Option.h"
+#include "Classes/EnvironmentQueryGraphNode_Root.h"
+#include "Classes/EnvironmentQueryGraphNode_Test.h"
+
 #include "AssetTypeActions_EnvironmentQuery.h"
 
 #include "EnvironmentQueryEditorModule.h"
diff --git a/Engine/Source/Editor/EnvironmentQueryEditor/Private/EnvironmentQueryFactory.cpp b/Engine/Source/Editor/EnvironmentQueryEditor/Private/EnvironmentQueryFactory.cpp
--- a/Engine/Source/Editor/EnvironmentQueryEditor/Private/EnvironmentQueryFactory.cpp
+++ b/Engine/Source/Editor/EnvironmentQueryEditor/Private/EnvironmentQueryFactory.cpp
@@ -1,6 +1,7 @@
 // Copyright 1998-2015 Epic Games, Inc. All Rights Reserved.
 
 #include "EnvironmentQueryEditorPrivatePCH.h"
+#include "EnvironmentQueryFactory.h"
 #include "EnvironmentQueryEditorModule.h"
 #include "EnvironmentQuery/EnvQuery.h"
 
diff --git a/Engine/Source/Editor/GameProjectGeneration/Private/GameProjectGenerationPrivatePCH.h b/Engine/Source/Editor/GameProjectGeneration/Private/GameProjectGenerationPrivatePCH.h
--- a/Engine/Source/Editor/GameProjectGeneration/Private/GameProjectGenerationPrivatePCH.h
+++ b/Engine/Source/Editor/GameProjectGeneration/Private/GameProjectGenerationPrivatePCH.h
@@ -9,7 +9,8 @@
 
 DECLARE_LOG_CATEGORY_EXTERN(LogGameProjectGeneration, Verbose, All);
 
-#include "GameProjectGenerationClasses.h"
+#include "Classes/TemplateProjectDefs.h"
+#include "Classes/DefaultTemplateProjectDefs.h"
 #include "GameProjectGenerationModule.h"
 #include "GameProjectUtils.h"
 #include "SProjectBrowser.h"
diff --git a/Engine/Source/Editor/GameplayTagsEditor/Classes/GameplayTagsK2Node_MultiCompareGameplayTagAssetInterface.h b/Engine/Source/Editor/GameplayTagsEditor/Classes/GameplayTagsK2Node_MultiCompareGameplayTagAssetInterface.h
--- a/Engine/Source/Editor/GameplayTagsEditor/Classes/GameplayTagsK2Node_MultiCompareGameplayTagAssetInterface.h
+++ b/Engine/Source/Editor/GameplayTagsEditor/Classes/GameplayTagsK2Node_MultiCompareGameplayTagAssetInterface.h
@@ -3,6 +3,7 @@
 
 #pragma once
 
+#include "GameplayTagsK2Node_MultiCompareBase.h"
 #include "GameplayTagsK2Node_MultiCompareGameplayTagAssetInterface.generated.h"
 
 UCLASS()
diff --git a/Engine/Source/Editor/GameplayTagsEditor/Classes/GameplayTagsK2Node_MultiCompareGameplayTagAssetInterfaceSingleTags.h b/Engine/Source/Editor/GameplayTagsEditor/Classes/GameplayTagsK2Node_MultiCompareGameplayTagAssetInterfaceSingleTags.h
--- a/Engine/Source/Editor/GameplayTagsEditor/Classes/GameplayTagsK2Node_MultiCompareGameplayTagAssetInterfaceSingleTags.h
+++ b/Engine/Source/Editor/GameplayTagsEditor/Classes/GameplayTagsK2Node_MultiCompareGameplayTagAssetInterfaceSingleTags.h
@@ -1,7 +1,8 @@
 // Copyright 1998-2014 Epic Games, Inc. All Rights Reserved.
 
-
 #pragma once
+
+#include "GameplayTagsK2Node_MultiCompareBase.h"
 #include "GameplayTagsK2Node_MultiCompareGameplayTagAssetInterfaceSingleTags.generated.h"
 
 
diff --git a/Engine/Source/Editor/GameplayTagsEditor/Classes/GameplayTagsK2Node_MultiCompareGameplayTagContainer.h b/Engine/Source/Editor/GameplayTagsEditor/Classes/GameplayTagsK2Node_MultiCompareGameplayTagContainer.h
--- a/Engine/Source/Editor/GameplayTagsEditor/Classes/GameplayTagsK2Node_MultiCompareGameplayTagContainer.h
+++ b/Engine/Source/Editor/GameplayTagsEditor/Classes/GameplayTagsK2Node_MultiCompareGameplayTagContainer.h
@@ -1,7 +1,8 @@
 // Copyright 1998-2014 Epic Games, Inc. All Rights Reserved.
 
-
 #pragma once
+
+#include "GameplayTagsK2Node_MultiCompareBase.h"
 #include "GameplayTagsK2Node_MultiCompareGameplayTagContainer.generated.h"
 
 
diff --git a/Engine/Source/Editor/GameplayTagsEditor/Classes/GameplayTagsK2Node_MultiCompareGameplayTagContainerSingleTags.h b/Engine/Source/Editor/GameplayTagsEditor/Classes/GameplayTagsK2Node_MultiCompareGameplayTagContainerSingleTags.h
--- a/Engine/Source/Editor/GameplayTagsEditor/Classes/GameplayTagsK2Node_MultiCompareGameplayTagContainerSingleTags.h
+++ b/Engine/Source/Editor/GameplayTagsEditor/Classes/GameplayTagsK2Node_MultiCompareGameplayTagContainerSingleTags.h
@@ -1,7 +1,8 @@
 // Copyright 1998-2015 Epic Games, Inc. All Rights Reserved.
 
-
 #pragma once
+
+#include "GameplayTagsK2Node_MultiCompareBase.h"
 #include "GameplayTagsK2Node_MultiCompareGameplayTagContainerSingleTags.generated.h"
 
 
diff --git a/Engine/Source/Editor/GameplayTagsEditor/Private/GameplayTagsEditorModulePrivatePCH.h b/Engine/Source/Editor/GameplayTagsEditor/Private/GameplayTagsEditorModulePrivatePCH.h
--- a/Engine/Source/Editor/GameplayTagsEditor/Private/GameplayTagsEditorModulePrivatePCH.h
+++ b/Engine/Source/Editor/GameplayTagsEditor/Private/GameplayTagsEditorModulePrivatePCH.h
@@ -13,8 +13,6 @@
 #include "Core.h"
 #include "Engine.h"
 #include "BlueprintGraphDefinitions.h"
-#include "GameplayTagsEditorClasses.h"
-
 #include "GraphEditor.h"
 #include "SNodePanel.h"
 #include "SGraphNode.h"
diff --git a/Engine/Source/Editor/GameplayTagsEditor/Private/GameplayTagsK2Node_LiteralGameplayTag.cpp b/Engine/Source/Editor/GameplayTagsEditor/Private/GameplayTagsK2Node_LiteralGameplayTag.cpp
--- a/Engine/Source/Editor/GameplayTagsEditor/Private/GameplayTagsK2Node_LiteralGameplayTag.cpp
+++ b/Engine/Source/Editor/GameplayTagsEditor/Private/GameplayTagsK2Node_LiteralGameplayTag.cpp
@@ -1,6 +1,7 @@
 // Copyright 1998-2015 Epic Games, Inc. All Rights Reserved.
 
 #include "GameplayTagsEditorModulePrivatePCH.h"
+#include "GameplayTagsK2Node_LiteralGameplayTag.h"
 #include "KismetCompiler.h"
 #include "GameplayTags.h"
 #include "BlueprintNodeSpawner.h"
diff --git a/Engine/Source/Editor/GameplayTagsEditor/Private/GameplayTagsK2Node_MultiCompareBase.cpp b/Engine/Source/Editor/GameplayTagsEditor/Private/GameplayTagsK2Node_MultiCompareBase.cpp
--- a/Engine/Source/Editor/GameplayTagsEditor/Private/GameplayTagsK2Node_MultiCompareBase.cpp
+++ b/Engine/Source/Editor/GameplayTagsEditor/Private/GameplayTagsK2Node_MultiCompareBase.cpp
@@ -6,6 +6,7 @@
 #include "GameplayTagContainer.h"
 #include "BlueprintGameplayTagLibrary.h"
 #include "GameplayTagsModule.h"
+#include "GameplayTagsK2Node_MultiCompareBase.h"
 #include "GameplayTagsK2Node_MultiCompareGameplayTagContainer.h"
 #include "KismetCompiler.h"
 
diff --git a/Engine/Source/Editor/GameplayTagsEditor/Private/SGraphNode_MultiCompareGameplayTag.h b/Engine/Source/Editor/GameplayTagsEditor/Private/SGraphNode_MultiCompareGameplayTag.h
--- a/Engine/Source/Editor/GameplayTagsEditor/Private/SGraphNode_MultiCompareGameplayTag.h
+++ b/Engine/Source/Editor/GameplayTagsEditor/Private/SGraphNode_MultiCompareGameplayTag.h
@@ -8,7 +8,7 @@ public:
 	SLATE_BEGIN_ARGS(SGraphNode_MultiCompareGameplayTag){}
 	SLATE_END_ARGS()
 
-	void Construct(const FArguments& InArgs, UGameplayTagsK2Node_MultiCompareBase* InNode);
+	void Construct(const FArguments& InArgs, class UGameplayTagsK2Node_MultiCompareBase* InNode);
 
 protected:
 
diff --git a/Engine/Source/Editor/GeometryMode/Private/GeometryEdMode.cpp b/Engine/Source/Editor/GeometryMode/Private/GeometryEdMode.cpp
--- a/Engine/Source/Editor/GeometryMode/Private/GeometryEdMode.cpp
+++ b/Engine/Source/Editor/GeometryMode/Private/GeometryEdMode.cpp
@@ -3,6 +3,20 @@
 #include "GeometryModePrivatePCH.h"
 #include "DynamicMeshBuilder.h"
 #include "Engine/Selection.h"
+#include "Classes/GeomModifier.h"
+#include "Classes/GeomModifier_Edit.h"
+#include "Classes/GeomModifier_Clip.h"
+#include "Classes/GeomModifier_Create.h"
+#include "Classes/GeomModifier_Delete.h"
+#include "Classes/GeomModifier_Extrude.h"
+#include "Classes/GeomModifier_Flip.h"
+#include "Classes/GeomModifier_Lathe.h"
+#include "Classes/GeomModifier_Pen.h"
+#include "Classes/GeomModifier_Split.h"
+#include "Classes/GeomModifier_Triangulate.h"
+#include "Classes/GeomModifier_Optimize.h"
+#include "Classes/GeomModifier_Turn.h"
+#include "Classes/GeomModifier_Weld.h"
 
 IMPLEMENT_MODULE( FGeometryModeModule, GeometryMode );
 
diff --git a/Engine/Source/Editor/GeometryMode/Private/GeometryMode.cpp b/Engine/Source/Editor/GeometryMode/Private/GeometryMode.cpp
--- a/Engine/Source/Editor/GeometryMode/Private/GeometryMode.cpp
+++ b/Engine/Source/Editor/GeometryMode/Private/GeometryMode.cpp
@@ -3,6 +3,7 @@
 #include "GeometryModePrivatePCH.h"
 #include "Editor/PropertyEditor/Public/PropertyEditorModule.h"
 #include "Editor/PropertyEditor/Public/IDetailsView.h"
+#include "Classes/GeomModifier.h"
 
 #define LOCTEXT_NAMESPACE "GeometryMode"
 
diff --git a/Engine/Source/Editor/GeometryMode/Private/GeometryMode.h b/Engine/Source/Editor/GeometryMode/Private/GeometryMode.h
--- a/Engine/Source/Editor/GeometryMode/Private/GeometryMode.h
+++ b/Engine/Source/Editor/GeometryMode/Private/GeometryMode.h
@@ -3,6 +3,7 @@
 #pragma once
 
 
+class UGeomModifier;
 
 /** Geometry Mode widget for controls */
 class SGeometryModeControls : public SCompoundWidget, public FNotifyHook
diff --git a/Engine/Source/Editor/GeometryMode/Private/GeometryModePrivatePCH.h b/Engine/Source/Editor/GeometryMode/Private/GeometryModePrivatePCH.h
--- a/Engine/Source/Editor/GeometryMode/Private/GeometryModePrivatePCH.h
+++ b/Engine/Source/Editor/GeometryMode/Private/GeometryModePrivatePCH.h
@@ -12,7 +12,6 @@
 
 #include "Editor/BspMode/Public/IBspModeModule.h"
 
-#include "GeometryModeClasses.h"
 #include "GeometryMode.h"
 #include "GeometryEdMode.h"
 #include "EditorGeometry.h"
diff --git a/Engine/Source/Editor/GeometryMode/Private/GeometryModifiers.cpp b/Engine/Source/Editor/GeometryMode/Private/GeometryModifiers.cpp
--- a/Engine/Source/Editor/GeometryMode/Private/GeometryModifiers.cpp
+++ b/Engine/Source/Editor/GeometryMode/Private/GeometryModifiers.cpp
@@ -13,6 +13,20 @@
 #include "Engine/Polys.h"
 #include "Engine/Selection.h"
 #include "Components/BrushComponent.h"
+#include "Classes/GeomModifier.h"
+#include "Classes/GeomModifier_Edit.h"
+#include "Classes/GeomModifier_Clip.h"
+#include "Classes/GeomModifier_Create.h"
+#include "Classes/GeomModifier_Delete.h"
+#include "Classes/GeomModifier_Extrude.h"
+#include "Classes/GeomModifier_Flip.h"
+#include "Classes/GeomModifier_Lathe.h"
+#include "Classes/GeomModifier_Pen.h"
+#include "Classes/GeomModifier_Split.h"
+#include "Classes/GeomModifier_Triangulate.h"
+#include "Classes/GeomModifier_Optimize.h"
+#include "Classes/GeomModifier_Turn.h"
+#include "Classes/GeomModifier_Weld.h"
 
 DEFINE_LOG_CATEGORY_STATIC(LogGeomModifier, Log, All);
 
diff --git a/Engine/Source/Editor/Kismet/Private/SBlueprintFavoritesPalette.cpp b/Engine/Source/Editor/Kismet/Private/SBlueprintFavoritesPalette.cpp
--- a/Engine/Source/Editor/Kismet/Private/SBlueprintFavoritesPalette.cpp
+++ b/Engine/Source/Editor/Kismet/Private/SBlueprintFavoritesPalette.cpp
@@ -3,7 +3,7 @@
 #include "BlueprintEditorPrivatePCH.h"
 #include "SBlueprintFavoritesPalette.h"
 #include "Engine.h"
-#include "KismetClasses.h"
+#include "Classes/BlueprintPaletteFavorites.h"
 #include "BlueprintActionMenuUtils.h"
 #include "BlueprintActionMenuBuilder.h"
 #include "BlueprintActionFilter.h" // for FBlueprintActionContext
diff --git a/Engine/Source/Editor/KismetWidgets/Private/KismetWidgetsPrivatePCH.h b/Engine/Source/Editor/KismetWidgets/Private/KismetWidgetsPrivatePCH.h
--- a/Engine/Source/Editor/KismetWidgets/Private/KismetWidgetsPrivatePCH.h
+++ b/Engine/Source/Editor/KismetWidgets/Private/KismetWidgetsPrivatePCH.h
@@ -6,7 +6,6 @@
 #include "UnrealEd.h"
 #include "ScopedTransaction.h"
 #include "SKismetLinearExpression.h"
-#include "BlueprintGraphClasses.h"
 
 
 #endif // __KismetWidgetsPrivatePCH_h__
diff --git a/Engine/Source/Editor/KismetWidgets/Private/SKismetLinearExpression.cpp b/Engine/Source/Editor/KismetWidgets/Private/SKismetLinearExpression.cpp
--- a/Engine/Source/Editor/KismetWidgets/Private/SKismetLinearExpression.cpp
+++ b/Engine/Source/Editor/KismetWidgets/Private/SKismetLinearExpression.cpp
@@ -2,6 +2,7 @@
 
 #include "KismetWidgetsPrivatePCH.h"
 #include "SKismetLinearExpression.h"
+#include "K2Node_VariableGet.h"
 
 #define LOCTEXT_NAMESPACE "KismetLinearExpression"
 
diff --git a/Engine/Source/Editor/Matinee/Private/Matinee.cpp b/Engine/Source/Editor/Matinee/Private/Matinee.cpp
--- a/Engine/Source/Editor/Matinee/Private/Matinee.cpp
+++ b/Engine/Source/Editor/Matinee/Private/Matinee.cpp
@@ -21,7 +21,8 @@
 #include "DistCurveEditorModule.h"
 #include "IDistCurveEditor.h"
 
-#include "MatineeClasses.h"
+#include "Classes/MatineeOptions.h"
+#include "Classes/MatineeTransBuffer.h"
 
 #include "CameraController.h"
 #include "MatineeConstants.h"
diff --git a/Engine/Source/Editor/Matinee/Private/MatineeTools.cpp b/Engine/Source/Editor/Matinee/Private/MatineeTools.cpp
--- a/Engine/Source/Editor/Matinee/Private/MatineeTools.cpp
+++ b/Engine/Source/Editor/Matinee/Private/MatineeTools.cpp
@@ -2,8 +2,7 @@
 
 #include "MatineeModule.h"
 #include "Matinee.h"
-#include "MatineeClasses.h"
-
+#include "Classes/InterpTrackHelper.h"
 #include "Matinee/MatineeActor.h"
 #include "Matinee/MatineeActorCameraAnim.h"
 #include "Matinee/InterpTrackInst.h"
diff --git a/Engine/Source/Editor/Matinee/Private/MatineeTrackHelpers.cpp b/Engine/Source/Editor/Matinee/Private/MatineeTrackHelpers.cpp
--- a/Engine/Source/Editor/Matinee/Private/MatineeTrackHelpers.cpp
+++ b/Engine/Source/Editor/Matinee/Private/MatineeTrackHelpers.cpp
@@ -2,7 +2,19 @@
 
 #include "MatineeModule.h"
 #include "Matinee.h"
-#include "MatineeClasses.h"
+#include "Classes/InterpTrackHelper.h"
+#include "Classes/MatineeTrackAnimControlHelper.h"
+#include "Classes/MatineeTrackBoolPropHelper.h"
+#include "Classes/MatineeTrackDirectorHelper.h"
+#include "Classes/MatineeTrackEventHelper.h"
+#include "Classes/MatineeTrackFloatPropHelper.h"
+#include "Classes/MatineeTrackParticleReplayHelper.h"
+#include "Classes/MatineeTrackSoundHelper.h"
+#include "Classes/MatineeTrackToggleHelper.h"
+#include "Classes/MatineeTrackVectorPropHelper.h"
+#include "Classes/MatineeTrackColorPropHelper.h"
+#include "Classes/MatineeTrackLinearColorPropHelper.h"
+#include "Classes/MatineeTrackVisibilityHelper.h"
 #include "AssetRegistryModule.h"
 #include "SoundDefinitions.h"
 #include "MatineeUtils.h"
diff --git a/Engine/Source/Editor/ReferenceViewer/Private/ReferenceViewerPrivatePCH.h b/Engine/Source/Editor/ReferenceViewer/Private/ReferenceViewerPrivatePCH.h
--- a/Engine/Source/Editor/ReferenceViewer/Private/ReferenceViewerPrivatePCH.h
+++ b/Engine/Source/Editor/ReferenceViewer/Private/ReferenceViewerPrivatePCH.h
@@ -1,7 +1,10 @@
 // Copyright 1998-2015 Epic Games, Inc. All Rights Reserved.
 
 #include "UnrealEd.h"
-#include "ReferenceViewerClasses.h"
+
+#include "Classes/EdGraph_ReferenceViewer.h"
+#include "Classes/EdGraphNode_Reference.h"
+#include "Classes/ReferenceViewerSchema.h"
 
 #include "GraphEditor.h"
 #include "SNodePanel.h"
diff --git a/Engine/Source/Editor/UnrealEd/Private/Tests/EditorAutomationTests.cpp b/Engine/Source/Editor/UnrealEd/Private/Tests/EditorAutomationTests.cpp
--- a/Engine/Source/Editor/UnrealEd/Private/Tests/EditorAutomationTests.cpp
+++ b/Engine/Source/Editor/UnrealEd/Private/Tests/EditorAutomationTests.cpp
@@ -16,7 +16,6 @@
 
 #include "Kismet2/KismetEditorUtilities.h"
 #include "Kismet2/BlueprintEditorUtils.h"
-#include "BlueprintGraphClasses.h"
 
 #include "AutomationCommon.h"
 #include "AutomationEditorCommon.h"
@@ -31,6 +30,13 @@
 #include "Engine/StaticMeshActor.h"
 #include "EngineUtils.h"
 
+#include "K2Node_Event.h"
+#include "K2Node_CustomEvent.h"
+#include "K2Node_Timeline.h"
+#include "K2Node_MacroInstance.h"
+#include "K2Node_CallFunction.h"
+#include "K2Node_Composite.h"
+
 
 /**
 * Change the attributes for a point light in the level.
diff --git a/Engine/Source/Editor/UnrealEd/Public/UnrealEd.h b/Engine/Source/Editor/UnrealEd/Public/UnrealEd.h
--- a/Engine/Source/Editor/UnrealEd/Public/UnrealEd.h
+++ b/Engine/Source/Editor/UnrealEd/Public/UnrealEd.h
@@ -18,7 +18,297 @@
 
 #include "TickableEditorObject.h"
 
-#include "UnrealEdClasses.h"
+// @todo Classes.h
+#include "Classes/ActorFactories/ActorFactory.h"
+#include "Classes/ActorFactories/ActorFactoryAmbientSound.h"
+#include "Classes/ActorFactories/ActorFactoryAtmosphericFog.h"
+#include "Classes/ActorFactories/ActorFactoryBlueprint.h"
+#include "Classes/ActorFactories/ActorFactoryBoxReflectionCapture.h"
+#include "Classes/ActorFactories/ActorFactoryBoxVolume.h"
+#include "Classes/ActorFactories/ActorFactoryCameraActor.h"
+#include "Classes/ActorFactories/ActorFactoryCharacter.h"
+#include "Classes/ActorFactories/ActorFactoryClass.h"
+#include "Classes/ActorFactories/ActorFactoryCylinderVolume.h"
+#include "Classes/ActorFactories/ActorFactoryDeferredDecal.h"
+#include "Classes/ActorFactories/ActorFactoryDestructible.h"
+#include "Classes/ActorFactories/ActorFactoryDirectionalLight.h"
+#include "Classes/ActorFactories/ActorFactoryEmitter.h"
+#include "Classes/ActorFactories/ActorFactoryEmptyActor.h"
+#include "Classes/ActorFactories/ActorFactoryPawn.h"
+#include "Classes/ActorFactories/ActorFactoryExponentialHeightFog.h"
+#include "Classes/ActorFactories/ActorFactoryMatineeActor.h"
+#include "Classes/ActorFactories/ActorFactoryNiagara.h"
+#include "Classes/ActorFactories/ActorFactoryNote.h"
+#include "Classes/ActorFactories/ActorFactoryPhysicsAsset.h"
+#include "Classes/ActorFactories/ActorFactoryPlaneReflectionCapture.h"
+#include "Classes/ActorFactories/ActorFactoryPlayerStart.h"
+#include "Classes/ActorFactories/ActorFactoryPointLight.h"
+#include "Classes/ActorFactories/ActorFactorySkeletalMesh.h"
+#include "Classes/ActorFactories/ActorFactoryAnimationAsset.h"
+#include "Classes/ActorFactories/ActorFactorySkyLight.h"
+#include "Classes/ActorFactories/ActorFactorySphereReflectionCapture.h"
+#include "Classes/ActorFactories/ActorFactorySphereVolume.h"
+#include "Classes/ActorFactories/ActorFactorySpotLight.h"
+#include "Classes/ActorFactories/ActorFactoryStaticMesh.h"
+#include "Classes/ActorFactories/ActorFactoryBasicShape.h"
+#include "Classes/ActorFactories/ActorFactoryInteractiveFoliage.h"
+#include "Classes/ActorFactories/ActorFactoryTargetPoint.h"
+#include "Classes/ActorFactories/ActorFactoryTextRender.h"
+#include "Classes/ActorFactories/ActorFactoryTriggerBox.h"
+#include "Classes/ActorFactories/ActorFactoryTriggerCapsule.h"
+#include "Classes/ActorFactories/ActorFactoryTriggerSphere.h"
+#include "Classes/ActorFactories/ActorFactoryVectorFieldVolume.h"
+#include "Classes/Exporters/AnimSequenceExporterFBX.h"
+#include "Classes/Commandlets/AudioTestCommandlet.h"
+#include "Classes/Preferences/CascadeOptions.h"
+#include "Classes/Commandlets/CompressAnimationsCommandlet.h"
+#include "Classes/Settings/ContentBrowserSettings.h"
+#include "Classes/Commandlets/CookCommandlet.h"
+#include "Classes/CookOnTheSide/CookOnTheFlyServer.h"
+#include "Classes/Preferences/CurveEdOptions.h"
+#include "Classes/Animation/DebugSkelMeshComponent.h"
+#include "Classes/Editor/UnrealEdTypes.h"
+#include "Classes/MaterialEditor/DEditorParameterValue.h"
+#include "Classes/MaterialEditor/DEditorFontParameterValue.h"
+#include "Classes/MaterialEditor/DEditorScalarParameterValue.h"
+#include "Classes/MaterialEditor/DEditorStaticComponentMaskParameterValue.h"
+#include "Classes/MaterialEditor/DEditorStaticSwitchParameterValue.h"
+#include "Classes/MaterialEditor/DEditorTextureParameterValue.h"
+#include "Classes/MaterialEditor/DEditorVectorParameterValue.h"
+#include "Classes/Commandlets/DerivedDataCacheCommandlet.h"
+#include "Classes/Settings/DestructableMeshEditorSettings.h"
+#include "Classes/Commandlets/DiffAssetsCommandlet.h"
+#include "Classes/Commandlets/DiffPackagesCommandlet.h"
+#include "Classes/Commandlets/DumpBlueprintsInfoCommandlet.h"
+#include "Classes/Commandlets/DumpHiddenCategoriesCommandlet.h"
+#include "Classes/MaterialGraph/MaterialGraph.h"
+#include "Classes/MaterialGraph/MaterialGraphNode_Comment.h"
+#include "Classes/Niagara/EdGraphSchema_Niagara.h"
+#include "Classes/Animation/EditorAnimBaseObj.h"
+#include "Classes/Animation/EditorAnimCompositeSegment.h"
+#include "Classes/Animation/EditorAnimSegment.h"
+#include "Classes/Animation/EditorCompositeSection.h"
+#include "Classes/Animation/EditorNotifyObject.h"
+#include "Classes/Builders/EditorBrushBuilder.h"
+#include "Classes/Builders/ConeBuilder.h"
+#include "Classes/Builders/CubeBuilder.h"
+#include "Classes/Builders/CurvedStairBuilder.h"
+#include "Classes/Builders/CylinderBuilder.h"
+#include "Classes/Builders/LinearStairBuilder.h"
+#include "Classes/Builders/SheetBuilder.h"
+#include "Classes/Builders/SpiralStairBuilder.h"
+#include "Classes/Builders/TetrahedronBuilder.h"
+#include "Classes/Builders/VolumetricBuilder.h"
+#include "Classes/Editor/EditorPerProjectUserSettings.h"
+#include "Classes/Editor/Transactor.h"
+#include "Classes/Settings/LevelEditorPlaySettings.h"
+#include "Classes/Settings/LevelEditorViewportSettings.h"
+#include "Classes/Editor/EditorEngine.h"
+#include "Classes/Editor/UnrealEdEngine.h"
+#include "Classes/Settings/EditorExperimentalSettings.h"
+#include "Classes/Settings/EditorLoadingSavingSettings.h"
+#include "Classes/Settings/EditorMiscSettings.h"
+#include "Classes/Animation/EditorParentPlayerListObj.h"
+#include "Classes/Animation/EditorSkeletonNotifyObj.h"
+#include "Classes/Exporters/ExportTextContainer.h"
+#include "Classes/Factories/Factory.h"
+#include "Classes/Factories/AnimBlueprintFactory.h"
+#include "Classes/Factories/AnimCompositeFactory.h"
+#include "Classes/Factories/AnimMontageFactory.h"
+#include "Classes/Factories/AnimSequenceFactory.h"
+#include "Classes/Factories/BlendSpaceFactory1D.h"
+#include "Classes/Factories/AimOffsetBlendSpaceFactory1D.h"
+#include "Classes/Factories/BlendSpaceFactoryNew.h"
+#include "Classes/Factories/AimOffsetBlendSpaceFactoryNew.h"
+#include "Classes/Factories/BlueprintFactory.h"
+#include "Classes/Factories/BlueprintFunctionLibraryFactory.h"
+#include "Classes/Factories/BlueprintMacroFactory.h"
+#include "Classes/Factories/BlueprintInterfaceFactory.h"
+#include "Classes/Factories/CameraAnimFactory.h"
+#include "Classes/Factories/CSVImportFactory.h"
+#include "Classes/Factories/ReimportCurveFactory.h"
+#include "Classes/Factories/ReimportCurveTableFactory.h"
+#include "Classes/Factories/ReimportDataTableFactory.h"
+#include "Classes/Factories/CurveFactory.h"
+#include "Classes/Factories/CurveImportFactory.h"
+#include "Classes/Factories/DataAssetFactory.h"
+#include "Classes/Factories/DataTableFactory.h"
+#include "Classes/Factories/DestructibleMeshFactory.h"
+#include "Classes/Factories/ReimportDestructibleMeshFactory.h"
+#include "Classes/Factories/DialogueVoiceFactory.h"
+#include "Classes/Factories/DialogueWaveFactory.h"
+#include "Classes/Factories/EnumFactory.h"
+#include "Classes/Factories/FbxFactory.h"
+#include "Classes/Factories/ReimportFbxAnimSequenceFactory.h"
+#include "Classes/Factories/ReimportFbxSkeletalMeshFactory.h"
+#include "Classes/Factories/ReimportFbxStaticMeshFactory.h"
+#include "Classes/Factories/FontFactory.h"
+#include "Classes/Factories/FontFileImportFactory.h"
+#include "Classes/Factories/ForceFeedbackEffectFactory.h"
+#include "Classes/Factories/InterpDataFactoryNew.h"
+#include "Classes/Factories/LevelFactory.h"
+#include "Classes/Factories/MaterialFactoryNew.h"
+#include "Classes/Factories/MaterialFunctionFactoryNew.h"
+#include "Classes/Factories/MaterialInstanceConstantFactoryNew.h"
+#include "Classes/Factories/MaterialParameterCollectionFactoryNew.h"
+#include "Classes/Factories/ModelFactory.h"
+#include "Classes/Factories/NiagaraEffectFactoryNew.h"
+#include "Classes/Factories/NiagaraScriptFactoryNew.h"
+#include "Classes/Factories/ObjectLibraryFactory.h"
+#include "Classes/Factories/PackageFactory.h"
+#include "Classes/Factories/PackFactory.h"
+#include "Classes/Factories/ParticleSystemFactoryNew.h"
+#include "Classes/Factories/PhysicalMaterialFactoryNew.h"
+#include "Classes/Factories/PolysFactory.h"
+#include "Classes/Factories/ReverbEffectFactory.h"
+#include "Classes/Factories/SkeletonFactory.h"
+#include "Classes/Factories/SlateBrushAssetFactory.h"
+#include "Classes/Factories/SlateWidgetStyleAssetFactory.h"
+#include "Classes/Factories/SoundAttenuationFactory.h"
+#include "Classes/Factories/SoundClassFactory.h"
+#include "Classes/Factories/SoundCueFactoryNew.h"
+#include "Classes/Factories/SoundFactory.h"
+#include "Classes/Factories/ReimportSoundFactory.h"
+#include "Classes/Factories/SoundMixFactory.h"
+#include "Classes/Factories/SoundSurroundFactory.h"
+#include "Classes/Factories/ReimportSoundSurroundFactory.h"
+#include "Classes/Factories/StructureFactory.h"
+#include "Classes/Factories/SubsurfaceProfileFactory.h"
+#include "Classes/Factories/Texture2dFactoryNew.h"
+#include "Classes/Factories/TextureFactory.h"
+#include "Classes/Factories/ReimportTextureFactory.h"
+#include "Classes/Factories/TrueTypeFontFactory.h"
+#include "Classes/Factories/TextureRenderTargetCubeFactoryNew.h"
+#include "Classes/Factories/TextureRenderTargetFactoryNew.h"
+#include "Classes/Factories/TouchInterfaceFactory.h"
+#include "Classes/Factories/VectorFieldStaticFactory.h"
+#include "Classes/Factories/ReimportVectorFieldStaticFactory.h"
+#include "Classes/Factories/WorldFactory.h"
+#include "Classes/Factories/FbxAssetImportData.h"
+#include "Classes/Factories/FbxAnimSequenceImportData.h"
+#include "Classes/Factories/FbxMeshImportData.h"
+#include "Classes/Factories/FbxSkeletalMeshImportData.h"
+#include "Classes/Factories/FbxStaticMeshImportData.h"
+#include "Classes/Factories/FbxTextureImportData.h"
+#include "Classes/Factories/FbxImportUI.h"
+#include "Classes/Commandlets/FileServerCommandlet.h"
+#include "Classes/Commandlets/FixupRedirectsCommandlet.h"
+#include "Classes/Commandlets/GatherTextCommandletBase.h"
+#include "Classes/Commandlets/GatherTextCommandlet.h"
+#include "Classes/Commandlets/GatherTextFromAssetsCommandlet.h"
+#include "Classes/Commandlets/GatherTextFromMetadataCommandlet.h"
+#include "Classes/Commandlets/GatherTextFromSourceCommandlet.h"
+#include "Classes/Commandlets/GenerateGatherArchiveCommandlet.h"
+#include "Classes/Commandlets/GenerateGatherManifestCommandlet.h"
+#include "Classes/Commandlets/GenerateTextLocalizationReportCommandlet.h"
+#include "Classes/Commandlets/GenerateTextLocalizationResourceCommandlet.h"
+#include "Classes/Commandlets/InternationalizationConditioningCommandlet.h"
+#include "Classes/Commandlets/InternationalizationExportCommandlet.h"
+#include "Classes/Commandlets/RepairLocalizationDataCommandlet.h"
+#include "Classes/Commandlets/GenerateBlueprintAPICommandlet.h"
+#include "Classes/Commandlets/GenerateDistillFileSetsCommandlet.h"
+#include "Classes/Editor/GroupActor.h"
+#include "Classes/Settings/LevelEditorMiscSettings.h"
+#include "Classes/Exporters/LevelExporterFBX.h"
+#include "Classes/Exporters/LevelExporterLOD.h"
+#include "Classes/Exporters/LevelExporterOBJ.h"
+#include "Classes/Exporters/LevelExporterSTL.h"
+#include "Classes/Exporters/LevelExporterT3D.h"
+#include "Classes/Preferences/LightmassOptionsObject.h"
+#include "Classes/Commandlets/ListMaterialsUsedWithMeshEmittersCommandlet.h"
+#include "Classes/Commandlets/ListStaticMeshesImportedFromSpeedTreesCommandlet.h"
+#include "Classes/Commandlets/LoadPackageCommandlet.h"
+#include "Classes/MaterialEditor/MaterialEditorInstanceConstant.h"
+#include "Classes/MaterialEditor/MaterialEditorMeshComponent.h"
+#include "Classes/Preferences/MaterialEditorOptions.h"
+#include "Classes/MaterialGraph/MaterialGraphNode_Base.h"
+#include "Classes/MaterialGraph/MaterialGraphNode.h"
+#include "Classes/MaterialGraph/MaterialGraphNode_Root.h"
+#include "Classes/MaterialGraph/MaterialGraphSchema.h"
+#include "Classes/Exporters/ModelExporterT3D.h"
+#include "Classes/Niagara/NiagaraGraph.h"
+#include "Classes/Niagara/NiagaraNode.h"
+#include "Classes/Niagara/NiagaraNodeConstant.h"
+#include "Classes/Niagara/NiagaraNodeFunctionCall.h"
+#include "Classes/Niagara/NiagaraNodeGetAttr.h"
+#include "Classes/Niagara/NiagaraNodeInput.h"
+#include "Classes/Niagara/NiagaraNodeOp.h"
+#include "Classes/Niagara/NiagaraNodeOutput.h"
+#include "Classes/Niagara/NiagaraNodeOutputUpdate.h"
+#include "Classes/Niagara/NiagaraNodeReadDataSet.h"
+#include "Classes/Niagara/NiagaraNodeWriteDataSet.h"
+#include "Classes/Niagara/NiagaraScriptSource.h"
+#include "Classes/Exporters/ObjectExporterT3D.h"
+#include "Classes/Commandlets/ParticleSystemAuditCommandlet.h"
+#include "Classes/Preferences/PersonaOptions.h"
+#include "Classes/Preferences/PhATSimOptions.h"
+#include "Classes/Commandlets/PkgInfoCommandlet.h"
+#include "Classes/Exporters/PolysExporterOBJ.h"
+#include "Classes/Exporters/PolysExporterT3D.h"
+#include "Classes/MaterialEditor/PreviewMaterial.h"
+#include "Classes/Settings/ProjectPackagingSettings.h"
+#include "Classes/Editor/PropertyEditorTestObject.h"
+#include "Classes/Commandlets/ReplaceActorCommandlet.h"
+#include "Classes/Commandlets/ResavePackagesCommandlet.h"
+#include "Classes/ThumbnailRendering/SceneThumbnailInfo.h"
+#include "Classes/ThumbnailRendering/ThumbnailManager.h"
+#include "Classes/ThumbnailRendering/SceneThumbnailInfoWithPrimitive.h"
+#include "Classes/ThumbnailRendering/WorldThumbnailInfo.h"
+#include "Classes/Exporters/SequenceExporterT3D.h"
+#include "Classes/Exporters/SkeletalMeshExporterFBX.h"
+#include "Classes/SoundClassGraph/SoundClassGraph.h"
+#include "Classes/SoundClassGraph/SoundClassGraphNode.h"
+#include "Classes/SoundClassGraph/SoundClassGraphSchema.h"
+#include "Classes/SoundCueGraph/SoundCueGraph.h"
+#include "Classes/SoundCueGraph/SoundCueGraphNode_Base.h"
+#include "Classes/SoundCueGraph/SoundCueGraphNode.h"
+#include "Classes/SoundCueGraph/SoundCueGraphNode_Root.h"
+#include "Classes/SoundCueGraph/SoundCueGraphSchema.h"
+#include "Classes/Exporters/SoundExporterOGG.h"
+#include "Classes/Exporters/SoundExporterWAV.h"
+#include "Classes/Exporters/SoundSurroundExporterWAV.h"
+#include "Classes/Exporters/StaticMeshExporterFBX.h"
+#include "Classes/Exporters/StaticMeshExporterOBJ.h"
+#include "Classes/Editor/TemplateMapMetadata.h"
+#include "Classes/TexAligner/TexAligner.h"
+#include "Classes/TexAligner/TexAlignerBox.h"
+#include "Classes/TexAligner/TexAlignerDefault.h"
+#include "Classes/TexAligner/TexAlignerFit.h"
+#include "Classes/TexAligner/TexAlignerPlanar.h"
+#include "Classes/Exporters/TextBufferExporterTXT.h"
+#include "Classes/Exporters/TextureCubeExporterHDR.h"
+#include "Classes/Exporters/TextureExporterBMP.h"
+#include "Classes/Exporters/TextureExporterHDR.h"
+#include "Classes/Exporters/TextureExporterPCX.h"
+#include "Classes/Exporters/TextureExporterTGA.h"
+#include "Classes/ThumbnailRendering/ThumbnailRenderer.h"
+#include "Classes/ThumbnailRendering/DefaultSizedThumbnailRenderer.h"
+#include "Classes/ThumbnailRendering/AnimBlueprintThumbnailRenderer.h"
+#include "Classes/ThumbnailRendering/AnimSequenceThumbnailRenderer.h"
+#include "Classes/ThumbnailRendering/BlendSpaceThumbnailRenderer.h"
+#include "Classes/ThumbnailRendering/BlueprintThumbnailRenderer.h"
+#include "Classes/ThumbnailRendering/ClassThumbnailRenderer.h"
+#include "Classes/ThumbnailRendering/DestructibleMeshThumbnailRenderer.h"
+#include "Classes/ThumbnailRendering/LevelThumbnailRenderer.h"
+#include "Classes/ThumbnailRendering/MaterialFunctionThumbnailRenderer.h"
+#include "Classes/ThumbnailRendering/MaterialInstanceThumbnailRenderer.h"
+#include "Classes/ThumbnailRendering/SkeletalMeshThumbnailRenderer.h"
+#include "Classes/ThumbnailRendering/SlateBrushThumbnailRenderer.h"
+#include "Classes/ThumbnailRendering/StaticMeshThumbnailRenderer.h"
+#include "Classes/ThumbnailRendering/WorldThumbnailRenderer.h"
+#include "Classes/ThumbnailRendering/SoundWaveThumbnailRenderer.h"
+#include "Classes/ThumbnailRendering/TextureThumbnailRenderer.h"
+#include "Classes/ThumbnailRendering/FontThumbnailRenderer.h"
+#include "Classes/ThumbnailRendering/ParticleSystemThumbnailRenderer.h"
+#include "Classes/ThumbnailRendering/SubsurfaceProfileRenderer.h"
+#include "Classes/Factories/TextureCubeThumbnailRenderer.h"
+#include "Classes/Editor/TransBuffer.h"
+#include "Classes/Preferences/UnrealEdKeyBindings.h"
+#include "Classes/Preferences/UnrealEdOptions.h"
+#include "Classes/Commandlets/UpdateGameProjectCommandlet.h"
+#include "Classes/UserDefinedStructure/UserDefinedStructEditorData.h"
+#include "Classes/Commandlets/WrangleContentCommandlet.h"
+//////////////////////	@todo Classes.h
 
 #include "Editor.h"
 
diff --git a/Engine/Source/Runtime/AnimGraphRuntime/Private/AnimGraphRuntimePrivatePCH.h b/Engine/Source/Runtime/AnimGraphRuntime/Private/AnimGraphRuntimePrivatePCH.h
--- a/Engine/Source/Runtime/AnimGraphRuntime/Private/AnimGraphRuntimePrivatePCH.h
+++ b/Engine/Source/Runtime/AnimGraphRuntime/Private/AnimGraphRuntimePrivatePCH.h
@@ -4,7 +4,3 @@
 
 #include "Engine.h"
 
-// DECLARE_LOG_CATEGORY_EXTERN(LogAnimation, Verbose, All);
-
-// #include "Paper2DClasses.h"
-// #include "Paper2DModule.h"
\ No newline at end of file
diff --git a/Engine/Source/Runtime/Engine/Private/VisualLogger/VisualLogger.cpp b/Engine/Source/Runtime/Engine/Private/VisualLogger/VisualLogger.cpp
--- a/Engine/Source/Runtime/Engine/Private/VisualLogger/VisualLogger.cpp
+++ b/Engine/Source/Runtime/Engine/Private/VisualLogger/VisualLogger.cpp
@@ -9,7 +9,6 @@
 #	include "Editor/UnrealEd/Public/EditorReimportHandler.h"
 #	include "Editor/UnrealEd/Public/TexAlignTools.h"
 #	include "Editor/UnrealEd/Public/TickableEditorObject.h"
-#	include "UnrealEdClasses.h"
 #	include "Editor/UnrealEd/Public/Editor.h"
 #	include "Editor/UnrealEd/Public/EditorViewportClient.h"
 #endif
diff --git a/Engine/Source/Runtime/GameplayTags/Private/GameplayTagsModulePrivatePCH.h b/Engine/Source/Runtime/GameplayTags/Private/GameplayTagsModulePrivatePCH.h
--- a/Engine/Source/Runtime/GameplayTags/Private/GameplayTagsModulePrivatePCH.h
+++ b/Engine/Source/Runtime/GameplayTags/Private/GameplayTagsModulePrivatePCH.h
@@ -14,7 +14,12 @@
 #include "Core.h"
 #include "Engine.h"
 
-#include "GameplayTagsClasses.h"
+#include "Classes/GameplayTagContainer.h"
+#include "Classes/GameplayTagAssetInterface.h"
+#include "Classes/BlueprintGameplayTagLibrary.h"
+#include "Classes/GameplayTagsManager.h"
+#include "Classes/GameplayTagsSettings.h"
+
 #include "GameplayTagsModule.h"
 
 DECLARE_LOG_CATEGORY_EXTERN(LogGameplayTags, Log, All);
\ No newline at end of file
diff --git a/Engine/Source/Runtime/InputCore/Public/InputCore.h b/Engine/Source/Runtime/InputCore/Public/InputCore.h
--- a/Engine/Source/Runtime/InputCore/Public/InputCore.h
+++ b/Engine/Source/Runtime/InputCore/Public/InputCore.h
@@ -4,4 +4,4 @@
 
 #include "InputCoreModule.h"
 #include "../Classes/InputCoreTypes.h"
-#include "InputCoreClasses.h"
+#include "InputCoreTypes.h"
diff --git a/Engine/Source/Runtime/MovieSceneTracks/Private/MovieSceneTracksPrivatePCH.h b/Engine/Source/Runtime/MovieSceneTracks/Private/MovieSceneTracksPrivatePCH.h
--- a/Engine/Source/Runtime/MovieSceneTracks/Private/MovieSceneTracksPrivatePCH.h
+++ b/Engine/Source/Runtime/MovieSceneTracks/Private/MovieSceneTracksPrivatePCH.h
@@ -18,7 +18,3 @@
 #include "IMovieSceneTracksModule.h"
 
 
-/* Private includes
- *****************************************************************************/
-
-#include "MovieSceneTracksClasses.h"
diff --git a/Engine/Source/Runtime/Online/OnlineSubsystemSteam/Classes/SteamNetConnection.h b/Engine/Source/Runtime/Online/OnlineSubsystemSteam/Classes/SteamNetConnection.h
--- a/Engine/Source/Runtime/Online/OnlineSubsystemSteam/Classes/SteamNetConnection.h
+++ b/Engine/Source/Runtime/Online/OnlineSubsystemSteam/Classes/SteamNetConnection.h
@@ -5,6 +5,7 @@
 //
 
 #pragma once
+#include "IpConnection.h"
 #include "SteamNetConnection.generated.h"
 
 UCLASS(transient, config=Engine)
diff --git a/Engine/Source/Runtime/Online/OnlineSubsystemSteam/Classes/SteamNetDriver.h b/Engine/Source/Runtime/Online/OnlineSubsystemSteam/Classes/SteamNetDriver.h
--- a/Engine/Source/Runtime/Online/OnlineSubsystemSteam/Classes/SteamNetDriver.h
+++ b/Engine/Source/Runtime/Online/OnlineSubsystemSteam/Classes/SteamNetDriver.h
@@ -5,6 +5,7 @@
 //
 
 #pragma once
+#include "IpNetDriver.h"
 #include "SteamNetDriver.generated.h"
 
 UCLASS(transient, config=Engine)
diff --git a/Engine/Source/Runtime/Online/OnlineSubsystemSteam/Private/OnlineSubsystemSteam.cpp b/Engine/Source/Runtime/Online/OnlineSubsystemSteam/Private/OnlineSubsystemSteam.cpp
--- a/Engine/Source/Runtime/Online/OnlineSubsystemSteam/Private/OnlineSubsystemSteam.cpp
+++ b/Engine/Source/Runtime/Online/OnlineSubsystemSteam/Private/OnlineSubsystemSteam.cpp
@@ -5,8 +5,6 @@
 #include "ModuleManager.h"
 
 #include "IPAddressSteam.h"
-#include "OnlineSubsystemUtilsClasses.h"
-#include "OnlineSubsystemSteamClasses.h"
 
 
 #include "SocketSubsystemSteam.h"
diff --git a/Engine/Source/Runtime/Online/OnlineSubsystemSteam/Private/OnlineSubsystemSteamPrivatePCH.h b/Engine/Source/Runtime/Online/OnlineSubsystemSteam/Private/OnlineSubsystemSteamPrivatePCH.h
--- a/Engine/Source/Runtime/Online/OnlineSubsystemSteam/Private/OnlineSubsystemSteamPrivatePCH.h
+++ b/Engine/Source/Runtime/Online/OnlineSubsystemSteam/Private/OnlineSubsystemSteamPrivatePCH.h
@@ -6,7 +6,6 @@
 #include "OnlineSubsystemSteamModule.h"
 #include "OnlineSubsystemModule.h"
 #include "OnlineSubsystemUtils.h"
-#include "OnlineSubsystemUtilsClasses.h"
 #include "OnlineSubsystem.h"
 #include "OnlineSubsystemUtils.h"
 #include "SocketSubsystem.h"
diff --git a/Engine/Source/Runtime/Online/OnlineSubsystemSteam/Private/SocketSubsystemSteam.cpp b/Engine/Source/Runtime/Online/OnlineSubsystemSteam/Private/SocketSubsystemSteam.cpp
--- a/Engine/Source/Runtime/Online/OnlineSubsystemSteam/Private/SocketSubsystemSteam.cpp
+++ b/Engine/Source/Runtime/Online/OnlineSubsystemSteam/Private/SocketSubsystemSteam.cpp
@@ -7,10 +7,8 @@
 #include "OnlineSessionInterfaceSteam.h"
 #include "SocketSubsystemModule.h"
 #include "ModuleManager.h"
-
+#include "Classes/SteamNetConnection.h"
 #include "Engine.h"
-#include "OnlineSubsystemUtilsClasses.h"
-#include "OnlineSubsystemSteamClasses.h"
 
 FSocketSubsystemSteam* FSocketSubsystemSteam::SocketSingleton = NULL;
 
diff --git a/Engine/Source/Runtime/Online/OnlineSubsystemSteam/Private/SteamNetConnection.cpp b/Engine/Source/Runtime/Online/OnlineSubsystemSteam/Private/SteamNetConnection.cpp
--- a/Engine/Source/Runtime/Online/OnlineSubsystemSteam/Private/SteamNetConnection.cpp
+++ b/Engine/Source/Runtime/Online/OnlineSubsystemSteam/Private/SteamNetConnection.cpp
@@ -1,12 +1,11 @@
 // Copyright 1998-2015 Epic Games, Inc. All Rights Reserved.
 
 #include "OnlineSubsystemSteamPrivatePCH.h"
-#include "Engine.h"
+#include "Classes/SteamNetConnection.h"
 #include "IPAddressSteam.h"
-#include "OnlineSubsystemUtilsClasses.h"
-#include "OnlineSubsystemSteamClasses.h"
 #include "SocketSubsystemSteam.h"
 #include "SocketsSteam.h"
+#include "SteamNetDriver.h"
 
 USteamNetConnection::USteamNetConnection(const FObjectInitializer& ObjectInitializer) :
 	Super(ObjectInitializer),
diff --git a/Engine/Source/Runtime/Online/OnlineSubsystemSteam/Private/SteamNetDriver.cpp b/Engine/Source/Runtime/Online/OnlineSubsystemSteam/Private/SteamNetDriver.cpp
--- a/Engine/Source/Runtime/Online/OnlineSubsystemSteam/Private/SteamNetDriver.cpp
+++ b/Engine/Source/Runtime/Online/OnlineSubsystemSteam/Private/SteamNetDriver.cpp
@@ -1,13 +1,13 @@
 // Copyright 1998-2015 Epic Games, Inc. All Rights Reserved.
 
 #include "OnlineSubsystemSteamPrivatePCH.h"
-#include "Engine.h"
+#include "SteamNetDriver.h"
 #include "IPAddressSteam.h"
-#include "OnlineSubsystemUtilsClasses.h"
-#include "OnlineSubsystemSteamClasses.h"
 #include "OnlineSubsystemSteam.h"
 #include "SocketSubsystemSteam.h"
 #include "SocketsSteam.h"
+#include "SteamNetDriver.h"
+#include "SteamNetConnection.h"
 
 USteamNetDriver::USteamNetDriver(const FObjectInitializer& ObjectInitializer) :
 	Super(ObjectInitializer),
diff --git a/Engine/Source/Runtime/Online/OnlineSubsystemUtils/Classes/AchievementQueryCallbackProxy.h b/Engine/Source/Runtime/Online/OnlineSubsystemUtils/Classes/AchievementQueryCallbackProxy.h
--- a/Engine/Source/Runtime/Online/OnlineSubsystemUtils/Classes/AchievementQueryCallbackProxy.h
+++ b/Engine/Source/Runtime/Online/OnlineSubsystemUtils/Classes/AchievementQueryCallbackProxy.h
@@ -1,6 +1,7 @@
 // Copyright 1998-2015 Epic Games, Inc. All Rights Reserved.
 #pragma once
 
+#include "OnlineBlueprintCallProxyBase.h"
 #include "Runtime/Online/OnlineSubsystem/Public/Interfaces/OnlineAchievementsInterface.h"
 #include "AchievementQueryCallbackProxy.generated.h"
 
diff --git a/Engine/Source/Runtime/Online/OnlineSubsystemUtils/Classes/AchievementWriteCallbackProxy.h b/Engine/Source/Runtime/Online/OnlineSubsystemUtils/Classes/AchievementWriteCallbackProxy.h
--- a/Engine/Source/Runtime/Online/OnlineSubsystemUtils/Classes/AchievementWriteCallbackProxy.h
+++ b/Engine/Source/Runtime/Online/OnlineSubsystemUtils/Classes/AchievementWriteCallbackProxy.h
@@ -1,6 +1,7 @@
 // Copyright 1998-2015 Epic Games, Inc. All Rights Reserved.
 #pragma once
 
+#include "OnlineBlueprintCallProxyBase.h"
 #include "Runtime/Online/OnlineSubsystem/Public/Interfaces/OnlineAchievementsInterface.h"
 #include "AchievementWriteCallbackProxy.generated.h"
 
diff --git a/Engine/Source/Runtime/Online/OnlineSubsystemUtils/Classes/ConnectionCallbackProxy.h b/Engine/Source/Runtime/Online/OnlineSubsystemUtils/Classes/ConnectionCallbackProxy.h
--- a/Engine/Source/Runtime/Online/OnlineSubsystemUtils/Classes/ConnectionCallbackProxy.h
+++ b/Engine/Source/Runtime/Online/OnlineSubsystemUtils/Classes/ConnectionCallbackProxy.h
@@ -2,6 +2,7 @@
 
 #pragma once
 
+#include "OnlineBlueprintCallProxyBase.h"
 #include "Runtime/Online/OnlineSubsystem/Public/Interfaces/OnlineIdentityInterface.h"
 #include "ConnectionCallbackProxy.generated.h"
 
diff --git a/Engine/Source/Runtime/Online/OnlineSubsystemUtils/Classes/CreateSessionCallbackProxy.h b/Engine/Source/Runtime/Online/OnlineSubsystemUtils/Classes/CreateSessionCallbackProxy.h
--- a/Engine/Source/Runtime/Online/OnlineSubsystemUtils/Classes/CreateSessionCallbackProxy.h
+++ b/Engine/Source/Runtime/Online/OnlineSubsystemUtils/Classes/CreateSessionCallbackProxy.h
@@ -1,6 +1,7 @@
 // Copyright 1998-2015 Epic Games, Inc. All Rights Reserved.
 #pragma once
 
+#include "OnlineBlueprintCallProxyBase.h"
 #include "OnlineSessionInterface.h"
 #include "CreateSessionCallbackProxy.generated.h"
 
diff --git a/Engine/Source/Runtime/Online/OnlineSubsystemUtils/Classes/DestroySessionCallbackProxy.h b/Engine/Source/Runtime/Online/OnlineSubsystemUtils/Classes/DestroySessionCallbackProxy.h
--- a/Engine/Source/Runtime/Online/OnlineSubsystemUtils/Classes/DestroySessionCallbackProxy.h
+++ b/Engine/Source/Runtime/Online/OnlineSubsystemUtils/Classes/DestroySessionCallbackProxy.h
@@ -1,6 +1,7 @@
 // Copyright 1998-2015 Epic Games, Inc. All Rights Reserved.
 #pragma once
 
+#include "OnlineBlueprintCallProxyBase.h"
 #include "OnlineSessionInterface.h"
 #include "DestroySessionCallbackProxy.generated.h"
 
diff --git a/Engine/Source/Runtime/Online/OnlineSubsystemUtils/Classes/EndMatchCallbackProxy.h b/Engine/Source/Runtime/Online/OnlineSubsystemUtils/Classes/EndMatchCallbackProxy.h
--- a/Engine/Source/Runtime/Online/OnlineSubsystemUtils/Classes/EndMatchCallbackProxy.h
+++ b/Engine/Source/Runtime/Online/OnlineSubsystemUtils/Classes/EndMatchCallbackProxy.h
@@ -2,6 +2,7 @@
 
 #pragma once
 
+#include "OnlineBlueprintCallProxyBase.h"
 #include "OnlineTurnBasedInterface.h"
 #include "EndMatchCallbackProxy.generated.h"
 
diff --git a/Engine/Source/Runtime/Online/OnlineSubsystemUtils/Classes/EndTurnCallbackProxy.h b/Engine/Source/Runtime/Online/OnlineSubsystemUtils/Classes/EndTurnCallbackProxy.h
--- a/Engine/Source/Runtime/Online/OnlineSubsystemUtils/Classes/EndTurnCallbackProxy.h
+++ b/Engine/Source/Runtime/Online/OnlineSubsystemUtils/Classes/EndTurnCallbackProxy.h
@@ -2,6 +2,7 @@
 
 #pragma once
 
+#include "OnlineBlueprintCallProxyBase.h"
 #include "OnlineTurnBasedInterface.h"
 #include "TurnBasedMatchInterface.h"
 #include "EndTurnCallbackProxy.generated.h"
diff --git a/Engine/Source/Runtime/Online/OnlineSubsystemUtils/Classes/FindSessionsCallbackProxy.h b/Engine/Source/Runtime/Online/OnlineSubsystemUtils/Classes/FindSessionsCallbackProxy.h
--- a/Engine/Source/Runtime/Online/OnlineSubsystemUtils/Classes/FindSessionsCallbackProxy.h
+++ b/Engine/Source/Runtime/Online/OnlineSubsystemUtils/Classes/FindSessionsCallbackProxy.h
@@ -1,6 +1,7 @@
 // Copyright 1998-2015 Epic Games, Inc. All Rights Reserved.
 #pragma once
 
+#include "OnlineBlueprintCallProxyBase.h"
 #include "OnlineSessionInterface.h"
 #include "FindSessionsCallbackProxy.generated.h"
 
diff --git a/Engine/Source/Runtime/Online/OnlineSubsystemUtils/Classes/FindTurnBasedMatchCallbackProxy.h b/Engine/Source/Runtime/Online/OnlineSubsystemUtils/Classes/FindTurnBasedMatchCallbackProxy.h
--- a/Engine/Source/Runtime/Online/OnlineSubsystemUtils/Classes/FindTurnBasedMatchCallbackProxy.h
+++ b/Engine/Source/Runtime/Online/OnlineSubsystemUtils/Classes/FindTurnBasedMatchCallbackProxy.h
@@ -2,6 +2,7 @@
 
 #pragma once
 
+#include "OnlineBlueprintCallProxyBase.h"
 #include "OnlineTurnBasedInterface.h"
 #include "FindTurnBasedMatchCallbackProxy.generated.h"
 
diff --git a/Engine/Source/Runtime/Online/OnlineSubsystemUtils/Classes/JoinSessionCallbackProxy.h b/Engine/Source/Runtime/Online/OnlineSubsystemUtils/Classes/JoinSessionCallbackProxy.h
--- a/Engine/Source/Runtime/Online/OnlineSubsystemUtils/Classes/JoinSessionCallbackProxy.h
+++ b/Engine/Source/Runtime/Online/OnlineSubsystemUtils/Classes/JoinSessionCallbackProxy.h
@@ -1,6 +1,7 @@
 // Copyright 1998-2015 Epic Games, Inc. All Rights Reserved.
 #pragma once
 
+#include "OnlineBlueprintCallProxyBase.h"
 #include "OnlineSessionInterface.h"
 #include "JoinSessionCallbackProxy.generated.h"
 
diff --git a/Engine/Source/Runtime/Online/OnlineSubsystemUtils/Classes/LeaderboardFlushCallbackProxy.cpp b/Engine/Source/Runtime/Online/OnlineSubsystemUtils/Classes/LeaderboardFlushCallbackProxy.cpp
--- a/Engine/Source/Runtime/Online/OnlineSubsystemUtils/Classes/LeaderboardFlushCallbackProxy.cpp
+++ b/Engine/Source/Runtime/Online/OnlineSubsystemUtils/Classes/LeaderboardFlushCallbackProxy.cpp
@@ -1,6 +1,7 @@
 // Copyright 1998-2015 Epic Games, Inc. All Rights Reserved.
 
 #include "OnlineSubsystemUtilsPrivatePCH.h"
+#include "Classes/LeaderboardFlushCallbackProxy.h"
 
 //////////////////////////////////////////////////////////////////////////
 // ULeaderboardFlushCallbackProxy
diff --git a/Engine/Source/Runtime/Online/OnlineSubsystemUtils/Classes/LeaderboardQueryCallbackProxy.cpp b/Engine/Source/Runtime/Online/OnlineSubsystemUtils/Classes/LeaderboardQueryCallbackProxy.cpp
--- a/Engine/Source/Runtime/Online/OnlineSubsystemUtils/Classes/LeaderboardQueryCallbackProxy.cpp
+++ b/Engine/Source/Runtime/Online/OnlineSubsystemUtils/Classes/LeaderboardQueryCallbackProxy.cpp
@@ -1,6 +1,7 @@
 // Copyright 1998-2015 Epic Games, Inc. All Rights Reserved.
 
 #include "OnlineSubsystemUtilsPrivatePCH.h"
+#include "Classes/LeaderboardQueryCallbackProxy.h"
 
 //////////////////////////////////////////////////////////////////////////
 // ULeaderboardQueryCallbackProxy
diff --git a/Engine/Source/Runtime/Online/OnlineSubsystemUtils/Classes/QuitMatchCallbackProxy.h b/Engine/Source/Runtime/Online/OnlineSubsystemUtils/Classes/QuitMatchCallbackProxy.h
--- a/Engine/Source/Runtime/Online/OnlineSubsystemUtils/Classes/QuitMatchCallbackProxy.h
+++ b/Engine/Source/Runtime/Online/OnlineSubsystemUtils/Classes/QuitMatchCallbackProxy.h
@@ -2,6 +2,7 @@
 
 #pragma once
 
+#include "OnlineBlueprintCallProxyBase.h"
 #include "OnlineTurnBasedInterface.h"
 #include "QuitMatchCallbackProxy.generated.h"
 
diff --git a/Engine/Source/Runtime/Online/OnlineSubsystemUtils/Private/AchievementQueryCallbackProxy.cpp b/Engine/Source/Runtime/Online/OnlineSubsystemUtils/Private/AchievementQueryCallbackProxy.cpp
--- a/Engine/Source/Runtime/Online/OnlineSubsystemUtils/Private/AchievementQueryCallbackProxy.cpp
+++ b/Engine/Source/Runtime/Online/OnlineSubsystemUtils/Private/AchievementQueryCallbackProxy.cpp
@@ -1,6 +1,7 @@
 // Copyright 1998-2015 Epic Games, Inc. All Rights Reserved.
 
 #include "OnlineSubsystemUtilsPrivatePCH.h"
+#include "Classes/AchievementQueryCallbackProxy.h"
 
 //////////////////////////////////////////////////////////////////////////
 // UAchievementQueryCallbackProxy
diff --git a/Engine/Source/Runtime/Online/OnlineSubsystemUtils/Private/AchievementWriteCallbackProxy.cpp b/Engine/Source/Runtime/Online/OnlineSubsystemUtils/Private/AchievementWriteCallbackProxy.cpp
--- a/Engine/Source/Runtime/Online/OnlineSubsystemUtils/Private/AchievementWriteCallbackProxy.cpp
+++ b/Engine/Source/Runtime/Online/OnlineSubsystemUtils/Private/AchievementWriteCallbackProxy.cpp
@@ -1,6 +1,7 @@
 // Copyright 1998-2015 Epic Games, Inc. All Rights Reserved.
 
 #include "OnlineSubsystemUtilsPrivatePCH.h"
+#include "AchievementWriteCallbackProxy.h"
 
 //////////////////////////////////////////////////////////////////////////
 // ULeaderboardQueryCallbackProxy
diff --git a/Engine/Source/Runtime/Online/OnlineSubsystemUtils/Private/ConnectionCallbackProxy.cpp b/Engine/Source/Runtime/Online/OnlineSubsystemUtils/Private/ConnectionCallbackProxy.cpp
--- a/Engine/Source/Runtime/Online/OnlineSubsystemUtils/Private/ConnectionCallbackProxy.cpp
+++ b/Engine/Source/Runtime/Online/OnlineSubsystemUtils/Private/ConnectionCallbackProxy.cpp
@@ -1,6 +1,7 @@
 // Copyright 1998-2015 Epic Games, Inc. All Rights Reserved.
 
 #include "OnlineSubsystemUtilsPrivatePCH.h"
+#include "ConnectionCallbackProxy.h"
 
 //////////////////////////////////////////////////////////////////////////
 // UConnectionCallbackProxy
diff --git a/Engine/Source/Runtime/Online/OnlineSubsystemUtils/Private/EndMatchCallbackProxy.cpp b/Engine/Source/Runtime/Online/OnlineSubsystemUtils/Private/EndMatchCallbackProxy.cpp
--- a/Engine/Source/Runtime/Online/OnlineSubsystemUtils/Private/EndMatchCallbackProxy.cpp
+++ b/Engine/Source/Runtime/Online/OnlineSubsystemUtils/Private/EndMatchCallbackProxy.cpp
@@ -1,6 +1,7 @@
 // Copyright 1998-2015 Epic Games, Inc. All Rights Reserved.
 
 #include "OnlineSubsystemUtilsPrivatePCH.h"
+#include "EndMatchCallbackProxy.h"
 #include "RepLayout.h"
 #include "TurnBasedMatchInterface.h"
 
diff --git a/Engine/Source/Runtime/Online/OnlineSubsystemUtils/Private/EndTurnCallbackProxy.cpp b/Engine/Source/Runtime/Online/OnlineSubsystemUtils/Private/EndTurnCallbackProxy.cpp
--- a/Engine/Source/Runtime/Online/OnlineSubsystemUtils/Private/EndTurnCallbackProxy.cpp
+++ b/Engine/Source/Runtime/Online/OnlineSubsystemUtils/Private/EndTurnCallbackProxy.cpp
@@ -1,4 +1,5 @@
 #include "OnlineSubsystemUtilsPrivatePCH.h"
+#include "EndTurnCallbackProxy.h"
 #include "GameFramework/Actor.h"
 #include "RepLayout.h"
 
diff --git a/Engine/Source/Runtime/Online/OnlineSubsystemUtils/Private/FindTurnBasedMatchCallbackProxy.cpp b/Engine/Source/Runtime/Online/OnlineSubsystemUtils/Private/FindTurnBasedMatchCallbackProxy.cpp
--- a/Engine/Source/Runtime/Online/OnlineSubsystemUtils/Private/FindTurnBasedMatchCallbackProxy.cpp
+++ b/Engine/Source/Runtime/Online/OnlineSubsystemUtils/Private/FindTurnBasedMatchCallbackProxy.cpp
@@ -1,6 +1,7 @@
 // Copyright 1998-2015 Epic Games, Inc. All Rights Reserved.
 
 #include "OnlineSubsystemUtilsPrivatePCH.h"
+#include "FindTurnBasedMatchCallbackProxy.h"
 #include "RepLayout.h"
 #include "TurnBasedMatchInterface.h"
 
diff --git a/Engine/Source/Runtime/Online/OnlineSubsystemUtils/Private/InAppPurchaseCallbackProxy.cpp b/Engine/Source/Runtime/Online/OnlineSubsystemUtils/Private/InAppPurchaseCallbackProxy.cpp
--- a/Engine/Source/Runtime/Online/OnlineSubsystemUtils/Private/InAppPurchaseCallbackProxy.cpp
+++ b/Engine/Source/Runtime/Online/OnlineSubsystemUtils/Private/InAppPurchaseCallbackProxy.cpp
@@ -1,6 +1,7 @@
 // Copyright 1998-2015 Epic Games, Inc. All Rights Reserved.
 
 #include "OnlineSubsystemUtilsPrivatePCH.h"
+#include "InAppPurchaseCallbackProxy.h"
 
 //////////////////////////////////////////////////////////////////////////
 // UInAppPurchaseCallbackProxy
diff --git a/Engine/Source/Runtime/Online/OnlineSubsystemUtils/Private/InAppPurchaseQueryCallbackProxy.cpp b/Engine/Source/Runtime/Online/OnlineSubsystemUtils/Private/InAppPurchaseQueryCallbackProxy.cpp
--- a/Engine/Source/Runtime/Online/OnlineSubsystemUtils/Private/InAppPurchaseQueryCallbackProxy.cpp
+++ b/Engine/Source/Runtime/Online/OnlineSubsystemUtils/Private/InAppPurchaseQueryCallbackProxy.cpp
@@ -1,6 +1,7 @@
 // Copyright 1998-2015 Epic Games, Inc. All Rights Reserved.
 
 #include "OnlineSubsystemUtilsPrivatePCH.h"
+#include "InAppPurchaseQueryCallbackProxy.h"
 
 //////////////////////////////////////////////////////////////////////////
 // UInAppPurchaseQueryCallbackProxy
diff --git a/Engine/Source/Runtime/Online/OnlineSubsystemUtils/Private/InAppPurchaseRestoreCallbackProxy.cpp b/Engine/Source/Runtime/Online/OnlineSubsystemUtils/Private/InAppPurchaseRestoreCallbackProxy.cpp
--- a/Engine/Source/Runtime/Online/OnlineSubsystemUtils/Private/InAppPurchaseRestoreCallbackProxy.cpp
+++ b/Engine/Source/Runtime/Online/OnlineSubsystemUtils/Private/InAppPurchaseRestoreCallbackProxy.cpp
@@ -1,6 +1,7 @@
 // Copyright 1998-2015 Epic Games, Inc. All Rights Reserved.
 
 #include "OnlineSubsystemUtilsPrivatePCH.h"
+#include "InAppPurchaseRestoreCallbackProxy.h"
 
 //////////////////////////////////////////////////////////////////////////
 // UInAppPurchaseRestoreCallbackProxy
diff --git a/Engine/Source/Runtime/Online/OnlineSubsystemUtils/Private/LogoutCallbackProxy.cpp b/Engine/Source/Runtime/Online/OnlineSubsystemUtils/Private/LogoutCallbackProxy.cpp
--- a/Engine/Source/Runtime/Online/OnlineSubsystemUtils/Private/LogoutCallbackProxy.cpp
+++ b/Engine/Source/Runtime/Online/OnlineSubsystemUtils/Private/LogoutCallbackProxy.cpp
@@ -1,6 +1,7 @@
 // Copyright 1998-2015 Epic Games, Inc. All Rights Reserved.
 
 #include "OnlineSubsystemUtilsPrivatePCH.h"
+#include "LogoutCallbackProxy.h"
 #include "OnlineExternalUIInterface.h"
 
 ULogoutCallbackProxy::ULogoutCallbackProxy(const FObjectInitializer& ObjectInitializer)
diff --git a/Engine/Source/Runtime/Online/OnlineSubsystemUtils/Private/OnlineBeacon.cpp b/Engine/Source/Runtime/Online/OnlineSubsystemUtils/Private/OnlineBeacon.cpp
--- a/Engine/Source/Runtime/Online/OnlineSubsystemUtils/Private/OnlineBeacon.cpp
+++ b/Engine/Source/Runtime/Online/OnlineSubsystemUtils/Private/OnlineBeacon.cpp
@@ -1,6 +1,7 @@
 // Copyright 1998-2015 Epic Games, Inc. All Rights Reserved.
 
 #include "OnlineSubsystemUtilsPrivatePCH.h"
+#include "OnlineBeacon.h"
 #include "Net/DataChannel.h"
 
 DEFINE_LOG_CATEGORY(LogBeacon);
diff --git a/Engine/Source/Runtime/Online/OnlineSubsystemUtils/Private/OnlineBeaconClient.cpp b/Engine/Source/Runtime/Online/OnlineSubsystemUtils/Private/OnlineBeaconClient.cpp
--- a/Engine/Source/Runtime/Online/OnlineSubsystemUtils/Private/OnlineBeaconClient.cpp
+++ b/Engine/Source/Runtime/Online/OnlineSubsystemUtils/Private/OnlineBeaconClient.cpp
@@ -3,6 +3,7 @@
 #include "OnlineSubsystemUtilsPrivatePCH.h"
 #include "OnlineBeaconClient.h"
 #include "OnlineBeaconHost.h"
+#include "OnlineBeaconHostObject.h"
 #include "Net/DataChannel.h"
 #include "Net/UnrealNetwork.h"
 
diff --git a/Engine/Source/Runtime/Online/OnlineSubsystemUtils/Private/OnlineBeaconHost.cpp b/Engine/Source/Runtime/Online/OnlineSubsystemUtils/Private/OnlineBeaconHost.cpp
--- a/Engine/Source/Runtime/Online/OnlineSubsystemUtils/Private/OnlineBeaconHost.cpp
+++ b/Engine/Source/Runtime/Online/OnlineSubsystemUtils/Private/OnlineBeaconHost.cpp
@@ -4,6 +4,8 @@
 #include "OnlineBeaconHost.h"
 #include "Net/UnrealNetwork.h"
 #include "Net/DataChannel.h"
+#include "OnlineBeaconClient.h"
+#include "OnlineBeaconHostObject.h"
 
 AOnlineBeaconHost::AOnlineBeaconHost(const FObjectInitializer& ObjectInitializer) :
 	Super(ObjectInitializer)
diff --git a/Engine/Source/Runtime/Online/OnlineSubsystemUtils/Private/OnlineBeaconHostObject.cpp b/Engine/Source/Runtime/Online/OnlineSubsystemUtils/Private/OnlineBeaconHostObject.cpp
--- a/Engine/Source/Runtime/Online/OnlineSubsystemUtils/Private/OnlineBeaconHostObject.cpp
+++ b/Engine/Source/Runtime/Online/OnlineSubsystemUtils/Private/OnlineBeaconHostObject.cpp
@@ -3,6 +3,7 @@
 #include "OnlineSubsystemUtilsPrivatePCH.h"
 #include "OnlineBeaconHostObject.h"
 #include "OnlineBeaconHost.h"
+#include "OnlineBeaconClient.h"
 
 AOnlineBeaconHostObject::AOnlineBeaconHostObject(const FObjectInitializer& ObjectInitializer) :
 	Super(ObjectInitializer),
diff --git a/Engine/Source/Runtime/Online/OnlineSubsystemUtils/Private/OnlineSessionClient.cpp b/Engine/Source/Runtime/Online/OnlineSubsystemUtils/Private/OnlineSessionClient.cpp
--- a/Engine/Source/Runtime/Online/OnlineSubsystemUtils/Private/OnlineSessionClient.cpp
+++ b/Engine/Source/Runtime/Online/OnlineSubsystemUtils/Private/OnlineSessionClient.cpp
@@ -6,7 +6,7 @@
 =============================================================================*/
 
 #include "OnlineSubsystemUtilsPrivatePCH.h"
-
+#include "OnlineSessionClient.h"
 #include "Engine/GameInstance.h"
 
 UOnlineSessionClient::UOnlineSessionClient(const FObjectInitializer& ObjectInitializer)
diff --git a/Engine/Source/Runtime/Online/OnlineSubsystemUtils/Private/OnlineSubsystemUtilsPrivatePCH.h b/Engine/Source/Runtime/Online/OnlineSubsystemUtils/Private/OnlineSubsystemUtilsPrivatePCH.h
--- a/Engine/Source/Runtime/Online/OnlineSubsystemUtils/Private/OnlineSubsystemUtilsPrivatePCH.h
+++ b/Engine/Source/Runtime/Online/OnlineSubsystemUtils/Private/OnlineSubsystemUtilsPrivatePCH.h
@@ -11,7 +11,8 @@
 #include "OnlineSubsystemUtilsModule.h"
 #include "SocketSubsystem.h"
 #include "ModuleManager.h"
-#include "OnlineSubsystemUtilsClasses.h"
+#include "Classes/IpConnection.h"
+#include "Classes/IpNetDriver.h"
 
 class FUniqueNetId;
 
diff --git a/Engine/Source/Runtime/Online/OnlineSubsystemUtils/Private/ShowLoginUICallbackProxy.cpp b/Engine/Source/Runtime/Online/OnlineSubsystemUtils/Private/ShowLoginUICallbackProxy.cpp
--- a/Engine/Source/Runtime/Online/OnlineSubsystemUtils/Private/ShowLoginUICallbackProxy.cpp
+++ b/Engine/Source/Runtime/Online/OnlineSubsystemUtils/Private/ShowLoginUICallbackProxy.cpp
@@ -1,6 +1,7 @@
 // Copyright 1998-2015 Epic Games, Inc. All Rights Reserved.
 
 #include "OnlineSubsystemUtilsPrivatePCH.h"
+#include "Classes/ShowLoginUICallbackProxy.h"
 #include "OnlineExternalUIInterface.h"
 
 UShowLoginUICallbackProxy::UShowLoginUICallbackProxy(const FObjectInitializer& ObjectInitializer)
diff --git a/Engine/Source/Runtime/Online/OnlineSubsystemUtils/Private/TestBeaconHost.cpp b/Engine/Source/Runtime/Online/OnlineSubsystemUtils/Private/TestBeaconHost.cpp
--- a/Engine/Source/Runtime/Online/OnlineSubsystemUtils/Private/TestBeaconHost.cpp
+++ b/Engine/Source/Runtime/Online/OnlineSubsystemUtils/Private/TestBeaconHost.cpp
@@ -1,6 +1,7 @@
 // Copyright 1998-2015 Epic Games, Inc. All Rights Reserved.
 
 #include "OnlineSubsystemUtilsPrivatePCH.h"
+#include "TestBeaconHost.h"
 #include "TestBeaconClient.h"
 
 ATestBeaconHost::ATestBeaconHost(const FObjectInitializer& ObjectInitializer) :
diff --git a/Engine/Source/Runtime/Slate/Public/SlateBasics.h b/Engine/Source/Runtime/Slate/Public/SlateBasics.h
--- a/Engine/Source/Runtime/Slate/Public/SlateBasics.h
+++ b/Engine/Source/Runtime/Slate/Public/SlateBasics.h
@@ -6,7 +6,6 @@
 #include "CoreUObject.h"
 #include "Json.h"
 #include "SlateCore.h"
-#include "SlateClasses.h"
 #include "SlateOptMacros.h"
 
 #include "IPlatformTextField.h"
