Fixed a crash when KismetDebugger reads a local function variable while reinstancing.

diff --git a/Engine/Source/Editor/UnrealEd/Private/Kismet2/KismetDebugUtilities.cpp b/Engine/Source/Editor/UnrealEd/Private/Kismet2/KismetDebugUtilities.cpp
--- a/Engine/Source/Editor/UnrealEd/Private/Kismet2/KismetDebugUtilities.cpp
+++ b/Engine/Source/Editor/UnrealEd/Private/Kismet2/KismetDebugUtilities.cpp
@@ -945,7 +945,8 @@ FKismetDebugUtilities::EWatchTextResult FKismetDebugUtilities::GetWatchText(FStr
 			UFunction* OuterFunction = Cast<UFunction>(Property->GetOuter());
 			if(!PropertyBase && OuterFunction)
 			{
-				if(UBlueprintGeneratedClass* BPGC = Cast<UBlueprintGeneratedClass>(Blueprint->GeneratedClass))
+				UBlueprintGeneratedClass* BPGC = Cast<UBlueprintGeneratedClass>(Blueprint->GeneratedClass);
+				if (BPGC && ActiveObject->IsA(BPGC))
 				{
 					PropertyBase = BPGC->GetPersistentUberGraphFrame(ActiveObject, OuterFunction);
 				}
