Fixed issue with SaveCookerStats crashing in cook on the fly.

diff --git a/Engine/Source/Editor/UnrealEd/Private/CookOnTheFlyServer.cpp b/Engine/Source/Editor/UnrealEd/Private/CookOnTheFlyServer.cpp
--- a/Engine/Source/Editor/UnrealEd/Private/CookOnTheFlyServer.cpp
+++ b/Engine/Source/Editor/UnrealEd/Private/CookOnTheFlyServer.cpp
@@ -3826,12 +3826,24 @@ const FString UCookOnTheFlyServer::GetCookedAssetRegistryFilename(const FString&
 
 void UCookOnTheFlyServer::SaveCookerStats() const
 {
-	static const FName CookingStatsName("CookingStats");
-	FCookingStatsModule& CookingStatsModule = FModuleManager::LoadModuleChecked<FCookingStatsModule>(CookingStatsName);
-	ICookingStats& CookingStats = CookingStatsModule.Get();
+	static ICookingStats* CookingStats = nullptr;
+	static bool bInitialized = false;
+	if (!bInitialized)
+	{
+		static const FName CookingStatsName("CookingStats");
+		FCookingStatsModule* CookingStatsModule = FModuleManager::LoadModulePtr<FCookingStatsModule>(CookingStatsName);
+		if (CookingStatsModule)
+		{
+			CookingStats = &CookingStatsModule->Get();
+		}
+		bInitialized = true;
+	}
 
-	const FString StatsFilename = GetStatsFilename(CookByTheBookOptions->ChildCookFilename);
-	CookingStats.SaveStatsAsCSV(StatsFilename);
+	if (CookingStats)
+	{
+		const FString StatsFilename = GetStatsFilename(CookByTheBookOptions->ChildCookFilename);
+		CookingStats->SaveStatsAsCSV(StatsFilename);
+	}
 }
 
 void UCookOnTheFlyServer::CookByTheBookFinished()
