#IMPORT: Fix issue with crash happening when invalid FBX file is being imported

Merging

//depot/UE4-Orion/Engine/Source/Editor/UnrealEd/Private/Fbx/FbxMainImport.cpp

to //depot/UE4/Engine/Source/Editor/UnrealEd/Private/Fbx/FbxMainImport.cpp

diff --git a/Engine/Source/Editor/UnrealEd/Private/Fbx/FbxMainImport.cpp b/Engine/Source/Editor/UnrealEd/Private/Fbx/FbxMainImport.cpp
--- a/Engine/Source/Editor/UnrealEd/Private/Fbx/FbxMainImport.cpp
+++ b/Engine/Source/Editor/UnrealEd/Private/Fbx/FbxMainImport.cpp
@@ -416,7 +416,7 @@ int32 FFbxImporter::GetImportType(const FString& InFilename)
 
 			bHasAnimation = SceneInfo.bHasAnimation;
 		}
-		else
+		else // back up old solution
 		{
 			for(ItemIndex = 0; ItemIndex < Statistics.GetNbItems(); ItemIndex++)
 			{
@@ -452,14 +452,17 @@ int32 FFbxImporter::GetImportType(const FString& InFilename)
 		CurPhase = NOTSTARTED;
 
 		// In case no Geometry was found, check for animation (FBX can still contain mesh data though)
-		if ( Result == -1 )
-		{
-			// If animation data is found, set the result to 2, otherwise default to static mesh
-			Result = bHasAnimation ? 2 : 0;
-		}
-		else if ( Result == 0 && bHasAnimation )
+		if (bHasAnimation)
 		{
-			Result = 1;
+			if ( Result == -1)
+			{
+				Result = 2;
+			}
+			// by default detects as skeletalmesh since it has animation curves
+			else if (Result == 0)
+			{
+				Result = 1;
+			}
 		}
 	}
 	
