#IMPORT: Remove old solution, doesn't make sense - meaning if GetSceneInfo failed, that means something bad bad happened, so the backup solution won't work either.

Merging

//depot/UE4-Orion/Engine/Source/Editor/UnrealEd/Private/Fbx/FbxMainImport.cpp

to //depot/UE4-Orion/Engine/Source/Editor/UnrealEd/Private/Fbx/FbxMainImport.cpp

diff --git a/Engine/Source/Editor/UnrealEd/Private/Fbx/FbxMainImport.cpp b/Engine/Source/Editor/UnrealEd/Private/Fbx/FbxMainImport.cpp
--- a/Engine/Source/Editor/UnrealEd/Private/Fbx/FbxMainImport.cpp
+++ b/Engine/Source/Editor/UnrealEd/Private/Fbx/FbxMainImport.cpp
@@ -416,32 +416,6 @@ int32 FFbxImporter::GetImportType(const FString& InFilename)
 
 			bHasAnimation = SceneInfo.bHasAnimation;
 		}
-		else // back up old solution
-		{
-			for(ItemIndex = 0; ItemIndex < Statistics.GetNbItems(); ItemIndex++)
-			{
-				Statistics.GetItemPair(ItemIndex, ItemName, ItemCount);
-				const char* NameBuffer = ItemName.Buffer();
-				if(ItemName == "Deformer" && ItemCount > 0)
-				{
-					// if SkeletalMesh is found, just return
-					Result = 1;
-					break;
-				}
-				// if Geometry is found, sets it, but it can be overwritten by Deformer
-				else if(ItemName == "Geometry" && ItemCount > 0)
-				{
-					// let it still loop through even if Geometry is found
-					// Deformer can overwrite this information
-					Result = 0;
-				}
-				// Check for animation data. It can be overwritten by Geometry or Deformer
-				else if((ItemName == "AnimationCurve" || ItemName == "AnimationCurveNode") && ItemCount > 0)
-				{
-					bHasAnimation = true;
-				}
-			}
-		}
 
 		if (Importer)
 		{
