Cook Commandlet doesn't use FReplaceConvertedAssetManager anymore

diff --git a/Engine/Config/BaseEditor.ini b/Engine/Config/BaseEditor.ini
--- a/Engine/Config/BaseEditor.ini
+++ b/Engine/Config/BaseEditor.ini
@@ -345,9 +345,6 @@ bBlueprintIsNotBlueprintType=false
 SkipByteCodeSerialization=true
 BreakSerializationRecursion=true
 
-[BlueprintCppConversion]
-ReplaceConvertedAssetWhileCooking=false
-
 [ScriptErrorLog]
 ; per object
 MaxNumOfAccessViolation=16
diff --git a/Engine/Source/Editor/UnrealEd/Private/Commandlets/CookCommandlet.cpp b/Engine/Source/Editor/UnrealEd/Private/Commandlets/CookCommandlet.cpp
--- a/Engine/Source/Editor/UnrealEd/Private/Commandlets/CookCommandlet.cpp
+++ b/Engine/Source/Editor/UnrealEd/Private/Commandlets/CookCommandlet.cpp
@@ -24,7 +24,6 @@
 #include "CookerSettings.h"
 #include "ShaderCompiler.h"
 #include "MemoryMisc.h"
-#include "BlueprintSupport.h"
 
 DEFINE_LOG_CATEGORY_STATIC(LogCookCommandlet, Log, All);
 
@@ -517,9 +516,6 @@ int32 UCookCommandlet::Main(const FString& CmdLineParams)
 		}
 	}
 
-	static const FBoolConfigValueHelper ReplaceConvertedAsset(TEXT("BlueprintCppConversion"), TEXT("ReplaceConvertedAssetWhileCooking"));
-	FReplaceConvertedAssetManager::Get().SetEnabled(ReplaceConvertedAsset);
-
 	if ( bCookOnTheFly )
 	{
 		// parse instance identifier
