Replaced all calls to UWorld::GetFirstPlayerController in the engine with a more readily available player controller, or UGameInstance::GetFirstLocalPlayerController, since that appears to be the actual intent. The first controller in the world's list may not be considered local, such as when a client is recording a replay.
Also changed UConsole::ConsoleCommand to use the first local player controller instead of the first player controller with a non-null player. This fixes commands like ToggleDebugCamera during client replay recording.

--------
Integrated using branch UE4-Orion-To-UE4 of change#2688431 by [REDACTED_PERSON] on 2015/09/11 14:24:23.

diff --git a/Engine/Source/Developer/FunctionalTesting/Private/FunctionalTest.cpp b/Engine/Source/Developer/FunctionalTesting/Private/FunctionalTest.cpp
--- a/Engine/Source/Developer/FunctionalTesting/Private/FunctionalTest.cpp
+++ b/Engine/Source/Developer/FunctionalTesting/Private/FunctionalTest.cpp
@@ -292,9 +292,9 @@ void AFunctionalTest::GoToObservationPoint()
 	}
 
 	UWorld* World = GetWorld();
-	if (World)
+	if (World && World->GetGameInstance())
 	{
-		APlayerController* PC = World->GetFirstPlayerController();
+		APlayerController* PC = World->GetGameInstance()->GetFirstLocalPlayerController();
 		if (PC && PC->GetPawn())
 		{
 			PC->GetPawn()->TeleportTo(ObservationPoint->GetActorLocation(), ObservationPoint->GetActorRotation(), /*bIsATest=*/false, /*bNoCheck=*/true);
diff --git a/Engine/Source/Developer/GameplayDebugger/Private/GameplayDebugger.cpp b/Engine/Source/Developer/GameplayDebugger/Private/GameplayDebugger.cpp
--- a/Engine/Source/Developer/GameplayDebugger/Private/GameplayDebugger.cpp
+++ b/Engine/Source/Developer/GameplayDebugger/Private/GameplayDebugger.cpp
@@ -447,7 +447,7 @@ bool FGameplayDebugger::Exec(UWorld* InWorld, const TCHAR* Cmd, FOutputDevice& A
 
 	if (FParse::Command(&Cmd, TEXT("RunEQS")) && InWorld)
 	{
-		APlayerController* MyPC = InWorld->GetFirstPlayerController();
+		APlayerController* MyPC = InWorld->GetGameInstance() ? InWorld->GetGameInstance()->GetFirstLocalPlayerController() : nullptr;
 		UAISystem* AISys = UAISystem::GetCurrent(*InWorld);
 
 		UEnvQueryManager* EQS = AISys ? AISys->GetEnvironmentQueryManager() : NULL;
@@ -548,7 +548,7 @@ bool FGameplayDebugger::Exec(UWorld* InWorld, const TCHAR* Cmd, FOutputDevice& A
 
 		if (!LocalPC && MyWorld->GetNetMode() != NM_DedicatedServer)
 		{
-			LocalPC = MyWorld->GetFirstPlayerController();
+			LocalPC = MyWorld->GetGameInstance() ? MyWorld->GetGameInstance()->GetFirstLocalPlayerController() : nullptr;
 		}
 	}
 
diff --git a/Engine/Source/Editor/UnrealEd/Private/DebugToolExec.cpp b/Engine/Source/Editor/UnrealEd/Private/DebugToolExec.cpp
--- a/Engine/Source/Editor/UnrealEd/Private/DebugToolExec.cpp
+++ b/Engine/Source/Editor/UnrealEd/Private/DebugToolExec.cpp
@@ -194,7 +194,7 @@ bool FDebugToolExec::Exec( UWorld* InWorld, const TCHAR* Cmd, FOutputDevice& Ar
 
 		if (FParse::Command(&Cmd, TEXT("TRACE")))
 		{
-			APlayerController* PlayerController = InWorld->GetFirstPlayerController();
+			APlayerController* PlayerController = InWorld->GetGameInstance() ? InWorld->GetGameInstance()->GetFirstLocalPlayerController() : nullptr;
 			if (PlayerController != NULL)
 			{
 				// Do a trace in the player's facing direction and edit anything that's hit.
@@ -213,7 +213,7 @@ bool FDebugToolExec::Exec( UWorld* InWorld, const TCHAR* Cmd, FOutputDevice& Ar
 			
 			// Look for the closest actor of this class to the player.
 			FVector PlayerLocation(0.0f);
-			APlayerController* PlayerController = InWorld->GetFirstPlayerController();
+			APlayerController* PlayerController = InWorld->GetGameInstance() ? InWorld->GetGameInstance()->GetFirstLocalPlayerController() : nullptr;
 			if (PlayerController != NULL)
 			{
 				FRotator DummyRotation;
diff --git a/Engine/Source/Editor/UnrealEd/Private/PlayLevel.cpp b/Engine/Source/Editor/UnrealEd/Private/PlayLevel.cpp
--- a/Engine/Source/Editor/UnrealEd/Private/PlayLevel.cpp
+++ b/Engine/Source/Editor/UnrealEd/Private/PlayLevel.cpp
@@ -3282,45 +3282,48 @@ void UEditorEngine::ToggleBetweenPIEandSIE( bool bNewSession )
 
 				UWorld* World = GameViewport->GetWorld();
 				AGameMode* AuthGameMode = World->GetAuthGameMode();
-				if (AuthGameMode)	// If there is no GameMode, we are probably the client and cannot RestartPlayer.
+				if (AuthGameMode && GameViewport->GetGameInstance())	// If there is no GameMode, we are probably the client and cannot RestartPlayer.
 				{
-					APlayerController* PC = World->GetFirstPlayerController();
-					AuthGameMode->RemovePlayerControllerFromPlayerCount(PC);
-					PC->PlayerState->bOnlySpectator = false;
-					AuthGameMode->NumPlayers++;
-
-					bool bNeedsRestart = true;
-					if (PC->GetPawn() == NULL)
+					APlayerController* PC = GameViewport->GetGameInstance()->GetFirstLocalPlayerController();
+					if (PC != nullptr)
 					{
-						// Use the "auto-possess" pawn in the world, if there is one.
-						for (FConstPawnIterator Iterator = World->GetPawnIterator(); Iterator; ++Iterator)
+						AuthGameMode->RemovePlayerControllerFromPlayerCount(PC);
+						PC->PlayerState->bOnlySpectator = false;
+						AuthGameMode->NumPlayers++;
+
+						bool bNeedsRestart = true;
+						if (PC->GetPawn() == NULL)
 						{
-							APawn* Pawn = *Iterator;
-							if (Pawn && Pawn->AutoPossessPlayer == EAutoReceiveInput::Player0)
+							// Use the "auto-possess" pawn in the world, if there is one.
+							for (FConstPawnIterator Iterator = World->GetPawnIterator(); Iterator; ++Iterator)
 							{
-								if (Pawn->Controller == nullptr)
+								APawn* Pawn = *Iterator;
+								if (Pawn && Pawn->AutoPossessPlayer == EAutoReceiveInput::Player0)
 								{
-									PC->Possess(Pawn);
-									bNeedsRestart = false;
+									if (Pawn->Controller == nullptr)
+									{
+										PC->Possess(Pawn);
+										bNeedsRestart = false;
+									}
+									break;
 								}
-								break;
 							}
 						}
-					}
-
-					if (bNeedsRestart)
-					{
-						AuthGameMode->RestartPlayer(PC);
 
-						if (PC->GetPawn())
+						if (bNeedsRestart)
 						{
-							// If there was no player start, then try to place the pawn where the camera was.						
-							if (PC->StartSpot == nullptr || Cast<AWorldSettings>(PC->StartSpot.Get()))
+							AuthGameMode->RestartPlayer(PC);
+
+							if (PC->GetPawn())
 							{
-								const FVector Location = EditorViewportClient.GetViewLocation();
-								const FRotator Rotation = EditorViewportClient.GetViewRotation();
-								PC->SetControlRotation(Rotation);
-								PC->GetPawn()->TeleportTo(Location, Rotation);
+								// If there was no player start, then try to place the pawn where the camera was.						
+								if (PC->StartSpot == nullptr || Cast<AWorldSettings>(PC->StartSpot.Get()))
+								{
+									const FVector Location = EditorViewportClient.GetViewLocation();
+									const FRotator Rotation = EditorViewportClient.GetViewRotation();
+									PC->SetControlRotation(Rotation);
+									PC->GetPawn()->TeleportTo(Location, Rotation);
+								}
 							}
 						}
 					}
diff --git a/Engine/Source/Runtime/AIModule/Private/AISystem.cpp b/Engine/Source/Runtime/AIModule/Private/AISystem.cpp
--- a/Engine/Source/Runtime/AIModule/Private/AISystem.cpp
+++ b/Engine/Source/Runtime/AIModule/Private/AISystem.cpp
@@ -115,9 +115,9 @@ void UAISystem::AIIgnorePlayers()
 void UAISystem::AILoggingVerbose()
 {
 	UWorld* OuterWorld = GetOuterWorld();
-	if (OuterWorld)
+	if (OuterWorld && OuterWorld->GetGameInstance())
 	{
-		APlayerController* PC = OuterWorld->GetFirstPlayerController();
+		APlayerController* PC = OuterWorld->GetGameInstance()->GetFirstLocalPlayerController();
 		if (PC)
 		{
 			PC->ConsoleCommand(TEXT("log lognavigation verbose | log logpathfollowing verbose | log LogCharacter verbose | log LogBehaviorTree verbose | log LogPawnAction verbose|"));
@@ -129,12 +129,12 @@ void UAISystem::RunEQS(const FString& QueryName, UObject* Target)
 {
 #if !(UE_BUILD_SHIPPING || UE_BUILD_TEST)
 	UWorld* OuterWorld = GetOuterWorld();
-	if (OuterWorld == NULL)
+	if (OuterWorld == NULL || OuterWorld->GetGameInstance() == NULL)
 	{
 		return;
 	}
 
-	APlayerController* MyPC = OuterWorld->GetFirstPlayerController();
+	APlayerController* MyPC = OuterWorld->GetGameInstance()->GetFirstLocalPlayerController();
 	UEnvQueryManager* EQS = GetEnvironmentQueryManager();
 
 	if (Target && MyPC && EQS)
diff --git a/Engine/Source/Runtime/Engine/Private/DrawDebugHelpers.cpp b/Engine/Source/Runtime/Engine/Private/DrawDebugHelpers.cpp
--- a/Engine/Source/Runtime/Engine/Private/DrawDebugHelpers.cpp
+++ b/Engine/Source/Runtime/Engine/Private/DrawDebugHelpers.cpp
@@ -934,7 +934,7 @@ void DrawDebugFloatHistory(UWorld const & WorldRef, FDebugFloatHistory const & F
 
 void DrawDebugFloatHistory(UWorld const & WorldRef, FDebugFloatHistory const & FloatHistory, FVector const & DrawLocation, FVector2D const & DrawSize, FColor const & DrawColor, bool const & bPersistent, float const & LifeTime, uint8 const & DepthPriority)
 {
-	APlayerController * PlayerController = WorldRef.GetFirstPlayerController();
+	APlayerController * PlayerController = WorldRef.GetGameInstance() != nullptr ? WorldRef.GetGameInstance()->GetFirstLocalPlayerController() : nullptr;
 	FRotator const DrawRotation = (PlayerController && PlayerController->PlayerCameraManager) ? PlayerController->PlayerCameraManager->CameraCache.POV.Rotation : FRotator(0, 0, 0);
 
 	FTransform const DrawTransform(DrawRotation, DrawLocation);
diff --git a/Engine/Source/Runtime/Engine/Private/GameViewportClient.cpp b/Engine/Source/Runtime/Engine/Private/GameViewportClient.cpp
--- a/Engine/Source/Runtime/Engine/Private/GameViewportClient.cpp
+++ b/Engine/Source/Runtime/Engine/Private/GameViewportClient.cpp
@@ -567,9 +567,9 @@ bool UGameViewportClient::RequiresUncapturedAxisInput() const
 		{
 			bRequired = true;
 		}
-		else if (GetWorld() && GetWorld()->GetFirstPlayerController())
+		else if (GameInstance && GameInstance->GetFirstLocalPlayerController())
 		{
-			bRequired = GetWorld()->GetFirstPlayerController()->ShouldShowMouseCursor();
+			bRequired = GameInstance->GetFirstLocalPlayerController()->ShouldShowMouseCursor();
 		}
 	}
 
@@ -610,9 +610,9 @@ EMouseCursor::Type UGameViewportClient::GetCursor(FViewport* InViewport, int32 X
 	}
 	else if ( /*(!bIsPlayingMovie) && */(InViewport->IsFullscreen() || !bIsWithinTitleBar) ) //bIsPlayingMovie has always false value
 	{
-		if (GetWorld() && GetWorld()->GetFirstPlayerController())
+		if (GameInstance && GameInstance->GetFirstLocalPlayerController())
 		{
-			return GetWorld()->GetFirstPlayerController()->GetMouseCursor();
+			return GameInstance->GetFirstLocalPlayerController()->GetMouseCursor();
 		}
 
 		return EMouseCursor::None;
@@ -751,9 +751,9 @@ bool UGameViewportClient::ShouldForceFullscreenViewport() const
 		{
 			bResult = true;
 		}
-		else
+		else if ( GameInstance )
 		{
-			APlayerController* PlayerController = GetWorld()->GetFirstPlayerController();
+			APlayerController* PlayerController = GameInstance->GetFirstLocalPlayerController();
 			if( ( PlayerController ) && ( PlayerController->bCinematicMode ) )
 			{
 				bResult = true;
diff --git a/Engine/Source/Runtime/Engine/Private/GameplayStatics.cpp b/Engine/Source/Runtime/Engine/Private/GameplayStatics.cpp
--- a/Engine/Source/Runtime/Engine/Private/GameplayStatics.cpp
+++ b/Engine/Source/Runtime/Engine/Private/GameplayStatics.cpp
@@ -159,8 +159,8 @@ void UGameplayStatics::SetGlobalTimeDilation(UObject* WorldContextObject, float
 
 bool UGameplayStatics::SetGamePaused(UObject* WorldContextObject, bool bPaused)
 {
-	UWorld* const World = GEngine->GetWorldFromContextObject( WorldContextObject );
-	APlayerController* const PC = World ? World->GetFirstPlayerController() : nullptr;
+	UGameInstance* const GameInstance = GetGameInstance( WorldContextObject );
+	APlayerController* const PC = GameInstance ? GameInstance->GetFirstLocalPlayerController() : nullptr;
 	return PC ? PC->SetPause(bPaused) : false;
 }
 
diff --git a/Engine/Source/Runtime/Engine/Private/Particles/ParticleComponents.cpp b/Engine/Source/Runtime/Engine/Private/Particles/ParticleComponents.cpp
--- a/Engine/Source/Runtime/Engine/Private/Particles/ParticleComponents.cpp
+++ b/Engine/Source/Runtime/Engine/Private/Particles/ParticleComponents.cpp
@@ -3427,7 +3427,11 @@ void UParticleSystemComponent::OrientZAxisTowardCamera()
 	ForceAsyncWorkCompletion(ENSURE_AND_STALL);
 
 	//@TODO: CAMERA: How does this work for stereo and/or split-screen?
-	APlayerController* PlayerController = GetWorld()->GetFirstPlayerController();
+	APlayerController* PlayerController = nullptr;
+	if (GetWorld() && GetWorld()->GetGameInstance())
+	{
+		PlayerController = GetWorld()->GetGameInstance()->GetFirstLocalPlayerController();
+	}
 
 	// Orient the Z axis toward the camera
 	if (PlayerController && PlayerController->PlayerCameraManager)
diff --git a/Engine/Source/Runtime/Engine/Private/PlayerController.cpp b/Engine/Source/Runtime/Engine/Private/PlayerController.cpp
--- a/Engine/Source/Runtime/Engine/Private/PlayerController.cpp
+++ b/Engine/Source/Runtime/Engine/Private/PlayerController.cpp
@@ -2894,7 +2894,12 @@ void APlayerController::LevelStreamingStatusChanged(ULevelStreaming* LevelObject
 void APlayerController::ClientPrepareMapChange_Implementation(FName LevelName, bool bFirst, bool bLast)
 {
 	// Only call on the first local player controller to handle it being called on multiple PCs for splitscreen.
-	APlayerController* PlayerController = GetWorld()->GetFirstPlayerController();
+	if (GetWorld()->GetGameInstance() == nullptr)
+	{
+		return;
+	}
+
+	APlayerController* PlayerController = GetWorld()->GetGameInstance()->GetFirstLocalPlayerController();
 	if( PlayerController != this )
 	{
 		return;
diff --git a/Engine/Source/Runtime/Engine/Private/Slate/SceneViewport.cpp b/Engine/Source/Runtime/Engine/Private/Slate/SceneViewport.cpp
--- a/Engine/Source/Runtime/Engine/Private/Slate/SceneViewport.cpp
+++ b/Engine/Source/Runtime/Engine/Private/Slate/SceneViewport.cpp
@@ -395,12 +395,12 @@ FReply FSceneViewport::OnMouseButtonDown( const FGeometry& InGeometry, const FPo
 			CurrentReplyState.SetUserFocus(ViewportWidgetRef, EFocusCause::SetDirectly, true);
 			
 			UWorld* World = ViewportClient->GetWorld();
-			if (World && World->IsGameWorld() && World->GetFirstPlayerController())
+			if (World && World->IsGameWorld() && World->GetGameInstance() && World->GetGameInstance()->GetFirstLocalPlayerController())
 			{
 				CurrentReplyState.CaptureMouse(ViewportWidgetRef);
 				CurrentReplyState.LockMouseToWidget(ViewportWidgetRef);
 
-				bool bShouldShowMouseCursor = World->GetFirstPlayerController()->ShouldShowMouseCursor();
+				bool bShouldShowMouseCursor = World->GetGameInstance()->GetFirstLocalPlayerController()->ShouldShowMouseCursor();
 				if (ViewportClient->HideCursorDuringCapture() && bShouldShowMouseCursor)
 				{
 					bCursorHiddenDueToCapture = true;
@@ -825,9 +825,9 @@ FReply FSceneViewport::OnFocusReceived(const FFocusEvent& InFocusEvent)
 			if (IsForegroundWindow())
 			{
 				bool bIsCursorForcedVisible = false;
-				if (ViewportClient->GetWorld() && ViewportClient->GetWorld()->GetFirstPlayerController())
+				if (ViewportClient->GetWorld() && ViewportClient->GetWorld()->GetGameInstance() && ViewportClient->GetWorld()->GetGameInstance()->GetFirstLocalPlayerController())
 				{
-					bIsCursorForcedVisible = ViewportClient->GetWorld()->GetFirstPlayerController()->GetMouseCursor() != EMouseCursor::None;
+					bIsCursorForcedVisible = ViewportClient->GetWorld()->GetGameInstance()->GetFirstLocalPlayerController()->GetMouseCursor() != EMouseCursor::None;
 				}
 
 				const bool bPlayInEditorCapture = !bIsPlayInEditorViewport || InFocusEvent.GetCause() != EFocusCause::SetDirectly || bPlayInEditorGetsMouseControl;
diff --git a/Engine/Source/Runtime/Engine/Private/UserInterface/Console.cpp b/Engine/Source/Runtime/Engine/Private/UserInterface/Console.cpp
--- a/Engine/Source/Runtime/Engine/Private/UserInterface/Console.cpp
+++ b/Engine/Source/Runtime/Engine/Private/UserInterface/Console.cpp
@@ -449,24 +449,17 @@ void UConsole::ConsoleCommand(const FString& Command)
 
 	OutputText(FString::Printf(TEXT("\n>>> %s <<<"), *Command));
 
-	UWorld *World = GetOuterUGameViewportClient()->GetWorld();
+	UGameInstance* GameInstance = GetOuterUGameViewportClient()->GetGameInstance();
 	if(ConsoleTargetPlayer != NULL)
 	{
 		// If there is a console target player, execute the command in the player's context.
 		ConsoleTargetPlayer->PlayerController->ConsoleCommand(Command);
 	}
-	else if(World && World->GetPlayerControllerIterator())
+	else if(GameInstance && GameInstance->GetFirstLocalPlayerController())
 	{
-		// If there are any players, execute the command in the first player's context that has a non-null Player.
-		for (auto PCIter = World->GetPlayerControllerIterator(); PCIter; ++PCIter)
-		{
-			APlayerController* PC = *PCIter;
-			if (PC && PC->Player)
-			{
-				PC->ConsoleCommand(Command);
-				break;
-			}
-		}
+		// If there are any players, execute the command in the first local player's context.
+		APlayerController* PC = GameInstance->GetFirstLocalPlayerController();
+		PC->ConsoleCommand(Command);
 	}
 	else
 	{
