UE-20991 - fix for Launch On instruments crash on Mac, currently disable the run.

diff --git a/Engine/Source/Editor/UnrealEd/Private/PlayLevel.cpp b/Engine/Source/Editor/UnrealEd/Private/PlayLevel.cpp
--- a/Engine/Source/Editor/UnrealEd/Private/PlayLevel.cpp
+++ b/Engine/Source/Editor/UnrealEd/Private/PlayLevel.cpp
@@ -1759,6 +1759,15 @@ void UEditorEngine::PlayUsingLauncher()
 		LauncherProfile->SetEditorExe(FUnrealEdMisc::Get().GetExecutableForCommandlets());
 
 		const FString DummyDeviceName(FString::Printf(TEXT("All_iOS_On_%s"), FPlatformProcess::ComputerName()));
+
+        // TODO: remove below once instruments has been fixed to run from the terminal again
+        if (PlayUsingLauncherDeviceId.Left(PlayUsingLauncherDeviceId.Find(TEXT("@"))) == TEXT("IOS"))
+        {
+            PlayUsingLauncherDeviceName = DummyDeviceName;
+            LauncherProfile->SetLaunchMode(ELauncherProfileLaunchModes::DoNotLaunch);
+        }
+        // TODO: remove above once instruments has been fixed to run from the terminal again
+
 		if (PlayUsingLauncherDeviceId.Left(PlayUsingLauncherDeviceId.Find(TEXT("@"))) != TEXT("IOS") || !PlayUsingLauncherDeviceName.Contains(DummyDeviceName))
 		{
 			LauncherProfile->SetLaunchMode(ELauncherProfileLaunchModes::DefaultRole);
