Fixed issue where SplineComponents added dynamically through the in-world details panel were not editable with the visualizer.

diff --git a/Engine/Source/Editor/ComponentVisualizers/Private/SplineComponentVisualizer.cpp b/Engine/Source/Editor/ComponentVisualizers/Private/SplineComponentVisualizer.cpp
--- a/Engine/Source/Editor/ComponentVisualizers/Private/SplineComponentVisualizer.cpp
+++ b/Engine/Source/Editor/ComponentVisualizers/Private/SplineComponentVisualizer.cpp
@@ -354,7 +354,7 @@ bool FSplineComponentVisualizer::VisProxyHandleClick(FLevelEditorViewportClient*
 		const USplineComponent* SplineComp = CastChecked<const USplineComponent>(VisProxy->Component.Get());
 
 		SplineCompPropName = GetComponentPropertyName(SplineComp);
-		if(SplineCompPropName != NAME_None)
+		if(SplineCompPropName.IsValid())
 		{
 			AActor* OldSplineOwningActor = SplineOwningActor.Get();
 			SplineOwningActor = SplineComp->GetOwner();
@@ -681,7 +681,7 @@ bool FSplineComponentVisualizer::HandleInputKey(FEditorViewportClient* ViewportC
 void FSplineComponentVisualizer::EndEditing()
 {
 	SplineOwningActor = NULL;
-	SplineCompPropName = NAME_None;
+	SplineCompPropName.Clear();
 	ChangeSelectionState(INDEX_NONE, false);
 	SelectedSegmentIndex = INDEX_NONE;
 	SelectedTangentHandle = INDEX_NONE;
diff --git a/Engine/Source/Editor/ComponentVisualizers/Private/SplineComponentVisualizer.h b/Engine/Source/Editor/ComponentVisualizers/Private/SplineComponentVisualizer.h
--- a/Engine/Source/Editor/ComponentVisualizers/Private/SplineComponentVisualizer.h
+++ b/Engine/Source/Editor/ComponentVisualizers/Private/SplineComponentVisualizer.h
@@ -119,7 +119,7 @@ private:
 	TWeakObjectPtr<AActor> SplineOwningActor;
 
 	/** Name of property on the actor that references the spline we are editing */
-	FName SplineCompPropName;
+	FPropertyNameAndIndex SplineCompPropName;
 
 	/** Index of keys we have selected */
 	TSet<int32> SelectedKeys;
diff --git a/Engine/Source/Editor/ComponentVisualizers/Private/SplineMeshComponentVisualizer.cpp b/Engine/Source/Editor/ComponentVisualizers/Private/SplineMeshComponentVisualizer.cpp
--- a/Engine/Source/Editor/ComponentVisualizers/Private/SplineMeshComponentVisualizer.cpp
+++ b/Engine/Source/Editor/ComponentVisualizers/Private/SplineMeshComponentVisualizer.cpp
@@ -108,7 +108,7 @@ bool FSplineMeshComponentVisualizer::VisProxyHandleClick(FLevelEditorViewportCli
 		const USplineMeshComponent* SplineMeshComp = CastChecked<const USplineMeshComponent>(VisProxy->Component.Get());
 
 		SplineMeshCompPropName = GetComponentPropertyName(SplineMeshComp);
-		if (SplineMeshCompPropName != NAME_None)
+		if (SplineMeshCompPropName.IsValid())
 		{
 			SplineMeshOwningActor = SplineMeshComp->GetOwner();
 
@@ -358,7 +358,7 @@ bool FSplineMeshComponentVisualizer::HandleInputKey(FEditorViewportClient* Viewp
 void FSplineMeshComponentVisualizer::EndEditing()
 {
 	SplineMeshOwningActor = NULL;
-	SplineMeshCompPropName = NAME_None;
+	SplineMeshCompPropName.Clear();
 	SelectedKey = INDEX_NONE;
 	SelectedTangentHandle = INDEX_NONE;
 	SelectedTangentHandleType = ESelectedTangentHandle::None;
diff --git a/Engine/Source/Editor/ComponentVisualizers/Private/SplineMeshComponentVisualizer.h b/Engine/Source/Editor/ComponentVisualizers/Private/SplineMeshComponentVisualizer.h
--- a/Engine/Source/Editor/ComponentVisualizers/Private/SplineMeshComponentVisualizer.h
+++ b/Engine/Source/Editor/ComponentVisualizers/Private/SplineMeshComponentVisualizer.h
@@ -77,7 +77,7 @@ private:
 	TWeakObjectPtr<AActor> SplineMeshOwningActor;
 
 	/** Name of property on the actor that references the spline we are editing */
-	FName SplineMeshCompPropName;
+	FPropertyNameAndIndex SplineMeshCompPropName;
 
 	/** Index of the key we selected */
 	int32 SelectedKey;
diff --git a/Engine/Source/Editor/UnrealEd/Private/ComponentVisualizer.cpp b/Engine/Source/Editor/UnrealEd/Private/ComponentVisualizer.cpp
--- a/Engine/Source/Editor/UnrealEd/Private/ComponentVisualizer.cpp
+++ b/Engine/Source/Editor/UnrealEd/Private/ComponentVisualizer.cpp
@@ -6,46 +6,76 @@
 IMPLEMENT_HIT_PROXY(HComponentVisProxy, HHitProxy);
 
 
-FName FComponentVisualizer::GetComponentPropertyName(const UActorComponent* Component)
+FComponentVisualizer::FPropertyNameAndIndex FComponentVisualizer::GetComponentPropertyName(const UActorComponent* Component)
 {
 	if (Component)
 	{
 		const AActor* CompOwner = Component->GetOwner();
-		if(CompOwner)
+		if (CompOwner)
 		{
-			// Iterate over fields of this actor
+			// Iterate over UObject* fields of this actor
 			UClass* ActorClass = CompOwner->GetClass();
 			for (TFieldIterator<UObjectProperty> It(ActorClass); It; ++It)
 			{
 				// See if this property points to the component in question
 				UObjectProperty* ObjectProp = *It;
-				UObject* Object = ObjectProp->GetObjectPropertyValue(ObjectProp->ContainerPtrToValuePtr<void>(CompOwner));
-				if (Object == Component)
+				for (int32 Index = 0; Index < ObjectProp->ArrayDim; ++Index)
 				{
-					// It does! Return this name
-					return ObjectProp->GetFName();
+					UObject* Object = ObjectProp->GetObjectPropertyValue(ObjectProp->ContainerPtrToValuePtr<void>(CompOwner, Index));
+					if (Object == Component)
+					{
+						// It does! Return this name
+						return FPropertyNameAndIndex(ObjectProp->GetFName(), Index);
+					}
+				}
+			}
+
+			// If nothing found, look in TArray<UObject*> fields
+			for (TFieldIterator<UArrayProperty> It(ActorClass); It; ++It)
+			{
+				UArrayProperty* ArrayProp = *It;
+				if (UObjectProperty* InnerProp = Cast<UObjectProperty>(It->Inner))
+				{
+					FScriptArrayHelper ArrayHelper(ArrayProp, ArrayProp->ContainerPtrToValuePtr<void>(CompOwner));
+					for (int32 Index = 0; Index < ArrayHelper.Num(); ++Index)
+					{
+						UObject* Object = InnerProp->GetObjectPropertyValue(ArrayHelper.GetRawPtr(Index));
+						if (Object == Component)
+						{
+							return FPropertyNameAndIndex(ArrayProp->GetFName(), Index);
+						}
+					}
 				}
 			}
 		}
 	}
 
 	// Didn't find actor property referencing this component
-	return NAME_None;
+	return FPropertyNameAndIndex();
 }
 
-UActorComponent* FComponentVisualizer::GetComponentFromPropertyName(const AActor* CompOwner, FName PropertyName)
+UActorComponent* FComponentVisualizer::GetComponentFromPropertyName(const AActor* CompOwner, const FPropertyNameAndIndex& Property)
 {
 	UActorComponent* ResultComp = NULL;
-	if(CompOwner && PropertyName != NAME_None)
+	if(CompOwner && Property.IsValid())
 	{
 		UClass* ActorClass = CompOwner->GetClass();
-		UObjectProperty* ObjectProp = FindField<UObjectProperty>(ActorClass, PropertyName);
-		if(ObjectProp)
+		UProperty* Prop = FindField<UProperty>(ActorClass, Property.Name);
+		if (UObjectProperty* ObjectProp = Cast<UObjectProperty>(Prop))
 		{
-			UObject* Object = ObjectProp->GetObjectPropertyValue(ObjectProp->ContainerPtrToValuePtr<void>(CompOwner));
+			UObject* Object = ObjectProp->GetObjectPropertyValue(ObjectProp->ContainerPtrToValuePtr<void>(CompOwner, Property.Index));
 			ResultComp = Cast<UActorComponent>(Object);
 		}
-
+		else if (UArrayProperty* ArrayProp = Cast<UArrayProperty>(Prop))
+		{
+			if (UObjectProperty* InnerProp = Cast<UObjectProperty>(ArrayProp->Inner))
+			{
+				FScriptArrayHelper ArrayHelper(ArrayProp, ArrayProp->ContainerPtrToValuePtr<void>(CompOwner));
+				UObject* Object = InnerProp->GetObjectPropertyValue(ArrayHelper.GetRawPtr(Property.Index));
+				ResultComp = Cast<UActorComponent>(Object);
+			}
+		}
 	}
+
 	return ResultComp;
 }
diff --git a/Engine/Source/Editor/UnrealEd/Public/ComponentVisualizer.h b/Engine/Source/Editor/UnrealEd/Public/ComponentVisualizer.h
--- a/Engine/Source/Editor/UnrealEd/Public/ComponentVisualizer.h
+++ b/Engine/Source/Editor/UnrealEd/Public/ComponentVisualizer.h
@@ -49,8 +49,32 @@ public:
 	/** */
 	virtual TSharedPtr<SWidget> GenerateContextMenu() const { return TSharedPtr<SWidget>(); }
 
+	struct FPropertyNameAndIndex
+	{
+		FPropertyNameAndIndex()
+			: Name(NAME_None)
+			, Index(INDEX_NONE)
+		{}
+
+		explicit FPropertyNameAndIndex(FName InName, int32 InIndex = 0)
+			: Name(InName)
+			, Index(InIndex)
+		{}
+
+		bool IsValid() const { return Name != NAME_None && Index != INDEX_NONE; }
+
+		void Clear()
+		{
+			Name = NAME_None;
+			Index = INDEX_NONE;
+		}
+
+		FName Name;
+		int32 Index;
+	};
+
 	/** Find the name of the property that points to this component */
-	static FName GetComponentPropertyName(const UActorComponent* Component);
+	static FPropertyNameAndIndex GetComponentPropertyName(const UActorComponent* Component);
 	/** Get a component pointer from the property name */
-	static UActorComponent* GetComponentFromPropertyName(const AActor* CompOwner, FName PropertyName);
+	static UActorComponent* GetComponentFromPropertyName(const AActor* CompOwner, const FPropertyNameAndIndex& Property);
 };
