Fixed assert when using "Convert Selected Components to Blueprint Class".

diff --git a/Engine/Source/Editor/UnrealEd/Private/Kismet2/Kismet2.cpp b/Engine/Source/Editor/UnrealEd/Private/Kismet2/Kismet2.cpp
--- a/Engine/Source/Editor/UnrealEd/Private/Kismet2/Kismet2.cpp
+++ b/Engine/Source/Editor/UnrealEd/Private/Kismet2/Kismet2.cpp
@@ -1200,9 +1200,10 @@ void FKismetEditorUtilities::AddComponentsToBlueprint(UBlueprint* Blueprint, con
 					USCS_Node* ParentSCSNode = nullptr;
 					for (UBlueprint* ParentBlueprint : ParentBPStack)
 					{
-						ParentSCSNode = ParentBlueprint->SimpleConstructionScript->FindSCSNode(SceneComponent->AttachParent->GetFName());
-						if (ParentSCSNode)
+						USCS_Node** ParentSCSNodePtr = InstanceComponentToNodeMap.Find(SceneComponent->AttachParent);
+						if (ParentSCSNodePtr)
 						{
+							ParentSCSNode = *ParentSCSNodePtr;
 							break;
 						}
 					}
