UE-21132 - Spelling

diff --git a/Engine/Source/Editor/UnrealEd/Classes/Settings/EditorExperimentalSettings.h b/Engine/Source/Editor/UnrealEd/Classes/Settings/EditorExperimentalSettings.h
--- a/Engine/Source/Editor/UnrealEd/Classes/Settings/EditorExperimentalSettings.h
+++ b/Engine/Source/Editor/UnrealEd/Classes/Settings/EditorExperimentalSettings.h
@@ -90,7 +90,7 @@ public:
 	bool bDisableCookInEditor;
 
 	/** Enable -iterate for launch on */
-	UPROPERTY(EditAnywhere, config, Category = Cooking, meta = (DisplayName = "Iterative cooking for builds launched form the editor (launch on)"))
+	UPROPERTY(EditAnywhere, config, Category = Cooking, meta = (DisplayName = "Iterative cooking for builds launched from the editor (launch on)"))
 	bool bIterativeCookingForLaunchOn;
 
 	UPROPERTY(EditAnywhere, config, Category = Cooking, meta = (DisplayName = "Use multiple processes when cooking (only affects File -> Package)"))
