Editor: Fixed session name for launch-on feature

diff --git a/Engine/Source/Editor/UnrealEd/Private/PlayLevel.cpp b/Engine/Source/Editor/UnrealEd/Private/PlayLevel.cpp
--- a/Engine/Source/Editor/UnrealEd/Private/PlayLevel.cpp
+++ b/Engine/Source/Editor/UnrealEd/Private/PlayLevel.cpp
@@ -1695,7 +1695,7 @@ void UEditorEngine::PlayUsingLauncher()
 
 		const ULevelEditorPlaySettings* PlayInSettings = GetDefault<ULevelEditorPlaySettings>();
 		// Setup launch profile, keep the setting here to a minimum.
-		ILauncherProfileRef LauncherProfile = LauncherServicesModule.CreateProfile(TEXT("Play On Device"));
+		ILauncherProfileRef LauncherProfile = LauncherServicesModule.CreateProfile(TEXT("Launch On Device"));
 		EPlayOnBuildMode bBuildType = PlayInSettings->BuildGameBeforeLaunch;
 		if ((bBuildType == EPlayOnBuildMode::PlayOnBuild_Always) || (bBuildType == PlayOnBuild_Default && (bPlayUsingLauncherHasCode) && bPlayUsingLauncherHasCompiler))
 		{
