SessionFrontend: Fixed "Play in Standalone Game" games not showing up in session list

diff --git a/Engine/Source/Editor/UnrealEd/Private/PlayLevel.cpp b/Engine/Source/Editor/UnrealEd/Private/PlayLevel.cpp
--- a/Engine/Source/Editor/UnrealEd/Private/PlayLevel.cpp
+++ b/Engine/Source/Editor/UnrealEd/Private/PlayLevel.cpp
@@ -1254,7 +1254,7 @@ void UEditorEngine::PlayStandaloneLocalPc(FString MapNameOverride, FIntPoint* Wi
 		GameNameOrProjectFile = FApp::GetGameName();
 	}
 
-	FString AdditionalParameters(TEXT(""));
+	FString AdditionalParameters(TEXT(" -messaging -SessionName=\"Play in Standalone Game\""));
 	bool bRunningDebug = FParse::Param(FCommandLine::Get(), TEXT("debug"));
 	if (bRunningDebug)
 	{
