Fix the fbx export to support correctly the hierarchy when there is parenting.
1. The order of the export was not sort correctly
2. When a parent was not exported it was crashing, now it is using the absolute transform so the child is export with the expected transform.

diff --git a/Engine/Source/Editor/UnrealEd/Private/Fbx/FbxMainExport.cpp b/Engine/Source/Editor/UnrealEd/Private/Fbx/FbxMainExport.cpp
--- a/Engine/Source/Editor/UnrealEd/Private/Fbx/FbxMainExport.cpp
+++ b/Engine/Source/Editor/UnrealEd/Private/Fbx/FbxMainExport.cpp
@@ -238,6 +238,32 @@ void FFbxExporter::CreateAnimatableUserProperty(FbxNode* Node, float Value, cons
 	IntensityProp.ModifyFlag(FbxPropertyFlags::eAnimatable, true);
 }
 
+/**
+*	Sorts actors such that parent actors will appear before children actors in the list
+*	Stable sort
+*/
+static void SortActorsHierarchy(TArray<AActor*>& Actors)
+{
+	auto CalcAttachDepth = [](AActor* InActor) -> int32 {
+		int32 Depth = MAX_int32;
+		if (InActor)
+		{
+			Depth = 0;
+			if (InActor->GetRootComponent())
+			{
+				for (const USceneComponent* Test = InActor->GetRootComponent()->AttachParent; Test != nullptr; Test = Test->AttachParent, Depth++);
+			}
+		}
+		return Depth;
+	};
+
+	// Unfortunately TArray.StableSort assumes no null entries in the array
+	// So it forces me to use internal unrestricted version
+	StableSortInternal(Actors.GetData(), Actors.Num(), [&](AActor* L, AActor* R) {
+		return CalcAttachDepth(L) < CalcAttachDepth(R);
+	});
+}
+
 /**
  * Exports the basic scene information to the FBX document.
  */
@@ -280,42 +306,54 @@ void FFbxExporter::ExportLevelMesh( ULevel* InLevel, AMatineeActor* InMatineeAct
 		World = CastChecked<UWorld>( InLevel->GetOuter() );
 	}
 	check(World);
+
+	TArray<AActor*> ActorToExport;
 	int32 ActorCount = InLevel->Actors.Num();
 	for (int32 ActorIndex = 0; ActorIndex < ActorCount; ++ActorIndex)
 	{
 		AActor* Actor = InLevel->Actors[ActorIndex];
-		if ( Actor != NULL && ( !bSelectedOnly || ( bSelectedOnly && Actor->IsSelected() ) ) )
+		if (Actor != NULL && (!bSelectedOnly || (bSelectedOnly && Actor->IsSelected())))
 		{
-			if (Actor->IsA(ALight::StaticClass()))
-			{
-				ExportLight((ALight*) Actor, InMatineeActor );
-			}
-			else if (Actor->IsA(AStaticMeshActor::StaticClass()))
-			{
-				ExportStaticMesh( Actor, CastChecked<AStaticMeshActor>(Actor)->GetStaticMeshComponent(), InMatineeActor );
-			}
-			else if (Actor->IsA(ALandscapeProxy::StaticClass()))
-			{
-				ExportLandscape(CastChecked<ALandscapeProxy>(Actor), false);
-			}
-			else if (Actor->IsA(ABrush::StaticClass()))
-			{
-				// All brushes should be included within the world geometry exported above.
-				ExportBrush((ABrush*) Actor, NULL, 0 );
-			}
-			else if (Actor->IsA(AEmitter::StaticClass()))
-			{
-				ExportActor( Actor, InMatineeActor ); // Just export the placement of the particle emitter.
-			}
-			else if(Actor->IsA(ACameraActor::StaticClass()))
-			{
-				ExportCamera(CastChecked<ACameraActor>(Actor), InMatineeActor, true); // Just export the placement of the particle emitter.
-			}
-			else if( Actor != NULL )
-			{
-				// Export blueprint actors and all their components
-				ExportActor( Actor, InMatineeActor, true );
-			}
+			ActorToExport.Add(Actor);
+		}
+	}
+
+	//Sort the hierarchy to make sure parent come first
+	SortActorsHierarchy(ActorToExport);
+
+	ActorCount = ActorToExport.Num();
+	for (int32 ActorIndex = 0; ActorIndex < ActorCount; ++ActorIndex)
+	{
+		AActor* Actor = ActorToExport[ActorIndex];
+		if (Actor->IsA(ALight::StaticClass()))
+		{
+			ExportLight((ALight*) Actor, InMatineeActor );
+		}
+		else if (Actor->IsA(AStaticMeshActor::StaticClass()))
+		{
+			ExportStaticMesh( Actor, CastChecked<AStaticMeshActor>(Actor)->GetStaticMeshComponent(), InMatineeActor );
+		}
+		else if (Actor->IsA(ALandscapeProxy::StaticClass()))
+		{
+			ExportLandscape(CastChecked<ALandscapeProxy>(Actor), false);
+		}
+		else if (Actor->IsA(ABrush::StaticClass()))
+		{
+			// All brushes should be included within the world geometry exported above.
+			ExportBrush((ABrush*) Actor, NULL, 0 );
+		}
+		else if (Actor->IsA(AEmitter::StaticClass()))
+		{
+			ExportActor( Actor, InMatineeActor ); // Just export the placement of the particle emitter.
+		}
+		else if(Actor->IsA(ACameraActor::StaticClass()))
+		{
+			ExportCamera(CastChecked<ACameraActor>(Actor), InMatineeActor, true); // Just export the placement of the particle emitter.
+		}
+		else if( Actor != NULL )
+		{
+			// Export blueprint actors and all their components
+			ExportActor( Actor, InMatineeActor, true );
 		}
 	}
 }
@@ -1067,14 +1105,13 @@ FbxNode* FFbxExporter::ExportActor(AActor* Actor, AMatineeActor* InMatineeActor,
 		
 
 		AActor* ParentActor = Actor->GetAttachParentActor();
-		FbxNode* ParentNode;
+		// this doesn't work with skeletalmeshcomponent
+		FbxNode* ParentNode = FindActor(ParentActor);
 		FVector ActorLocation, ActorRotation, ActorScale;
-		if (bKeepHierarchy && ParentActor)
+		
+		//If the parent is the root or is not export use the root node as the parent
+		if (bKeepHierarchy && ParentNode)
 		{
-			// this doesn't work with skeletalmeshcomponent
-			ParentNode = FindActor(ParentActor);
-			check (ParentNode);
-
 			// Set the default position of the actor on the transforms
 			// The transformation is different from FBX's Z-up: invert the Y-axis for translations and the Y/Z angle values in rotations.
 			const FTransform RelativeTransform = Actor->GetTransform().GetRelativeTransform(ParentActor->GetTransform());
@@ -1087,9 +1124,20 @@ FbxNode* FFbxExporter::ExportActor(AActor* Actor, AMatineeActor* InMatineeActor,
 			ParentNode = Scene->GetRootNode();
 			// Set the default position of the actor on the transforms
 			// The transformation is different from FBX's Z-up: invert the Y-axis for translations and the Y/Z angle values in rotations.
-			ActorLocation = Actor->GetActorLocation();
-			ActorRotation = Actor->GetActorRotation().Euler();
-			ActorScale = Actor->GetRootComponent() ? Actor->GetRootComponent()->RelativeScale3D : FVector(1.f,1.f,1.f);
+			if (ParentActor != NULL)
+			{
+				//In case the parent was not export, get the absolute transform
+				const FTransform AbsoluteTransform = Actor->GetTransform();
+				ActorLocation = AbsoluteTransform.GetTranslation();
+				ActorRotation = AbsoluteTransform.GetRotation().Euler();
+				ActorScale = AbsoluteTransform.GetScale3D();
+			}
+			else
+			{
+				ActorLocation = Actor->GetActorLocation();
+				ActorRotation = Actor->GetActorRotation().Euler();
+				ActorScale = Actor->GetRootComponent() ? Actor->GetRootComponent()->RelativeScale3D : FVector(1.f, 1.f, 1.f);
+			}
 		}
 
 		ParentNode->AddChild(ActorNode);
