jira: UE-20578
When we focus on an actor we now prevent changing the preview camera transform.

diff --git a/Engine/Source/Editor/UnrealEd/Private/EditorServer.cpp b/Engine/Source/Editor/UnrealEd/Private/EditorServer.cpp
--- a/Engine/Source/Editor/UnrealEd/Private/EditorServer.cpp
+++ b/Engine/Source/Editor/UnrealEd/Private/EditorServer.cpp
@@ -6735,7 +6735,9 @@ void UEditorEngine::MoveViewportCamerasToBox(const FBox& BoundingBox, bool bActi
 			for (auto ViewportIt = LevelViewportClients.CreateConstIterator(); ViewportIt; ++ViewportIt)
 			{
 				FLevelEditorViewportClient* LinkedViewportClient = *ViewportIt;
-				LinkedViewportClient->FocusViewportOnBox(BoundingBox);
+				//Dont move camera attach to an actor
+				if (!LinkedViewportClient->IsAnyActorLocked())
+					LinkedViewportClient->FocusViewportOnBox(BoundingBox);
 			}
 		}
 	}
