Suspending the undo buffer during duplication of a Blueprint's CDO during reinstancing until the CDO is fully resolved.

During duplication of the CDO, the duplicated CDO's class still does not point to the CDO when it is placed in the transaction buffer. Any native sub object components that are placed in the transaction buffer are serialized incorrectly. These are fixed up during a compile which leads to a different path of serialization during undo and a mismatch buffer size. Delaying CDO placement into the undo buffer until it is finished resolves this.

diff --git a/Engine/Source/Editor/UnrealEd/Private/Kismet2/KismetReinstanceUtilities.cpp b/Engine/Source/Editor/UnrealEd/Private/Kismet2/KismetReinstanceUtilities.cpp
--- a/Engine/Source/Editor/UnrealEd/Private/Kismet2/KismetReinstanceUtilities.cpp
+++ b/Engine/Source/Editor/UnrealEd/Private/Kismet2/KismetReinstanceUtilities.cpp
@@ -179,12 +179,20 @@ FBlueprintCompileReinstancer::FBlueprintCompileReinstancer(UClass* InClassToRein
 		DuplicatedClass->StaticLink(true);
 
 		// Copy over the ComponentNametoDefaultObjectMap, which tells CopyPropertiesForUnrelatedObjects which components are instanced and which aren't
-		
+
+		// Temporarily suspend the undo buffer; we don't need to record the duplicated CDO until it is fully resolved
+ 		ITransaction* CurrentTransaction = GUndo;
+ 		GUndo = NULL;
 		DuplicatedClass->ClassDefaultObject = GetClassCDODuplicate(ClassToReinstance, DuplicatedClass->GetDefaultObjectName());
 
+		// Restore the undo buffer
+		GUndo = CurrentTransaction;
 		DuplicatedClass->ClassDefaultObject->SetFlags(RF_ClassDefaultObject);
 		DuplicatedClass->ClassDefaultObject->SetClass(DuplicatedClass);
 
+		// The CDO is fully duplicated and ready to be placed in the undo buffer
+		DuplicatedClass->ClassDefaultObject->Modify();
+
 		ClassToReinstance->GetDefaultObject()->SetClass(DuplicatedClass);
 		ObjectsThatShouldUseOldStuff.Add(DuplicatedClass); //CDO of REINST_ class can be used as archetype
 
