Can no longer reparent level Blueprints to other level Blueprints, native classes are still an option.

diff --git a/Engine/Source/Editor/UnrealEd/Private/Kismet2/BlueprintEditorUtils.cpp b/Engine/Source/Editor/UnrealEd/Private/Kismet2/BlueprintEditorUtils.cpp
--- a/Engine/Source/Editor/UnrealEd/Private/Kismet2/BlueprintEditorUtils.cpp
+++ b/Engine/Source/Editor/UnrealEd/Private/Kismet2/BlueprintEditorUtils.cpp
@@ -6943,6 +6943,13 @@ public:
 	/** Classes to never show in this class viewer. */
 	TSet< const UClass* > DisallowedClasses;
 
+	/** Will limit the results to only native classes */
+	bool bShowNativeOnly;
+
+	FBlueprintReparentFilter()
+		: bShowNativeOnly(false)
+	{}
+
 	virtual bool IsClassAllowed(const FClassViewerInitializationOptions& InInitOptions, const UClass* InClass, TSharedRef< FClassViewerFilterFuncs > InFilterFuncs ) override
 	{
 		// If it appears on the allowed child-of classes list (or there is nothing on that list)
@@ -6951,7 +6958,8 @@ public:
 		return InFilterFuncs->IfInChildOfClassesSet( AllowedChildrenOfClasses, InClass) != EFilterReturn::Failed && 
 			InFilterFuncs->IfInChildOfClassesSet(DisallowedChildrenOfClasses, InClass) != EFilterReturn::Passed && 
 			InFilterFuncs->IfInClassesSet(DisallowedClasses, InClass) != EFilterReturn::Passed &&
-			!InClass->HasAnyClassFlags(CLASS_Deprecated);
+			!InClass->HasAnyClassFlags(CLASS_Deprecated) &&
+			(bShowNativeOnly && InClass->HasAnyClassFlags(CLASS_Native)) || !bShowNativeOnly;
 	}
 
 	virtual bool IsUnloadedClassAllowed(const FClassViewerInitializationOptions& InInitOptions, const TSharedRef< const IUnloadedBlueprintData > InUnloadedClassData, TSharedRef< FClassViewerFilterFuncs > InFilterFuncs) override
@@ -6962,7 +6970,8 @@ public:
 		return InFilterFuncs->IfInChildOfClassesSet( AllowedChildrenOfClasses, InUnloadedClassData) != EFilterReturn::Failed && 
 			InFilterFuncs->IfInChildOfClassesSet(DisallowedChildrenOfClasses, InUnloadedClassData) != EFilterReturn::Passed && 
 			InFilterFuncs->IfInClassesSet(DisallowedClasses, InUnloadedClassData) != EFilterReturn::Passed &&
-			!InUnloadedClassData->HasAnyClassFlags(CLASS_Deprecated);
+			!InUnloadedClassData->HasAnyClassFlags(CLASS_Deprecated)  &&
+			(bShowNativeOnly && InUnloadedClassData->HasAnyClassFlags(CLASS_Native)) || !bShowNativeOnly;
 	}
 };
 
@@ -7014,6 +7023,7 @@ TSharedRef<SWidget> FBlueprintEditorUtils::ConstructBlueprintParentClassPicker(
 		{
 			// Don't allow conversion outside of the LevelScriptActor hierarchy
 			Filter->AllowedChildrenOfClasses.Add( ALevelScriptActor::StaticClass() );
+			Filter->bShowNativeOnly = true;
 		}
 		else
 		{
