Fixes some dubious logic from missing parenthesis in CL# 2697190

diff --git a/Engine/Source/Editor/UnrealEd/Private/Kismet2/BlueprintEditorUtils.cpp b/Engine/Source/Editor/UnrealEd/Private/Kismet2/BlueprintEditorUtils.cpp
--- a/Engine/Source/Editor/UnrealEd/Private/Kismet2/BlueprintEditorUtils.cpp
+++ b/Engine/Source/Editor/UnrealEd/Private/Kismet2/BlueprintEditorUtils.cpp
@@ -6959,7 +6959,7 @@ public:
 			InFilterFuncs->IfInChildOfClassesSet(DisallowedChildrenOfClasses, InClass) != EFilterReturn::Passed && 
 			InFilterFuncs->IfInClassesSet(DisallowedClasses, InClass) != EFilterReturn::Passed &&
 			!InClass->HasAnyClassFlags(CLASS_Deprecated) &&
-			(bShowNativeOnly && InClass->HasAnyClassFlags(CLASS_Native)) || !bShowNativeOnly;
+			((bShowNativeOnly && InClass->HasAnyClassFlags(CLASS_Native)) || !bShowNativeOnly);
 	}
 
 	virtual bool IsUnloadedClassAllowed(const FClassViewerInitializationOptions& InInitOptions, const TSharedRef< const IUnloadedBlueprintData > InUnloadedClassData, TSharedRef< FClassViewerFilterFuncs > InFilterFuncs) override
@@ -6971,7 +6971,7 @@ public:
 			InFilterFuncs->IfInChildOfClassesSet(DisallowedChildrenOfClasses, InUnloadedClassData) != EFilterReturn::Passed && 
 			InFilterFuncs->IfInClassesSet(DisallowedClasses, InUnloadedClassData) != EFilterReturn::Passed &&
 			!InUnloadedClassData->HasAnyClassFlags(CLASS_Deprecated)  &&
-			(bShowNativeOnly && InUnloadedClassData->HasAnyClassFlags(CLASS_Native)) || !bShowNativeOnly;
+			((bShowNativeOnly && InUnloadedClassData->HasAnyClassFlags(CLASS_Native)) || !bShowNativeOnly);
 	}
 };
 
