Enabled saving stats while running cook on the fly.

diff --git a/Engine/Source/Editor/UnrealEd/Private/CookOnTheFlyServer.cpp b/Engine/Source/Editor/UnrealEd/Private/CookOnTheFlyServer.cpp
--- a/Engine/Source/Editor/UnrealEd/Private/CookOnTheFlyServer.cpp
+++ b/Engine/Source/Editor/UnrealEd/Private/CookOnTheFlyServer.cpp
@@ -867,6 +867,10 @@ UCookOnTheFlyServer::UCookOnTheFlyServer(const FObjectInitializer& ObjectInitial
 
 UCookOnTheFlyServer::~UCookOnTheFlyServer()
 {
+	if (IsCookOnTheFlyMode())
+	{
+		SaveCookerStats();
+	}
 	check(TickChildCookers() == true);
 	if ( CookByTheBookOptions )
 	{		
@@ -1438,7 +1442,8 @@ void UCookOnTheFlyServer::CleanUpChildCookers()
 {
 	if (IsCookByTheBookMode())
 	{
-		const FString MasterCookerStatsFilename = GetStatsFilename(FString());
+		check(!IsChildCooker());
+		const FString MasterCookerStatsFilename = GetStatsFilename(CookByTheBookOptions->ChildCookFilename);
 
 		FString AllStats;
 		ensure(FFileHelper::LoadFileToString(AllStats, *MasterCookerStatsFilename));
@@ -3838,12 +3843,20 @@ void UCookOnTheFlyServer::SaveCookerStats() const
 		}
 		bInitialized = true;
 	}
-
 	if (CookingStats)
 	{
-		const FString StatsFilename = GetStatsFilename(CookByTheBookOptions->ChildCookFilename);
-		CookingStats->SaveStatsAsCSV(StatsFilename);
+		if (IsCookByTheBookMode())
+		{
+			const FString StatsFilename = GetStatsFilename(CookByTheBookOptions->ChildCookFilename);
+			CookingStats->SaveStatsAsCSV(StatsFilename);
+		}
+		else if ( IsCookOnTheFlyMode() )
+		{
+			const FString StatsFilename = GetStatsFilename(FString(TEXT("")));
+			CookingStats->SaveStatsAsCSV(StatsFilename);
+		}
 	}
+	
 }
 
 void UCookOnTheFlyServer::CookByTheBookFinished()
