#Fix crash on missing adjacency index buffer data of existing LODs after reimport
Merging

//depot/UE4-Orion/Engine/Source/Editor/UnrealEd/...

to //depot/UE4-Orion/Engine/Source/Editor/UnrealEd/...

diff --git a/Engine/Source/Editor/UnrealEd/Private/Factories/SkeletalMeshImport.cpp b/Engine/Source/Editor/UnrealEd/Private/Factories/SkeletalMeshImport.cpp
--- a/Engine/Source/Editor/UnrealEd/Private/Factories/SkeletalMeshImport.cpp
+++ b/Engine/Source/Editor/UnrealEd/Private/Factories/SkeletalMeshImport.cpp
@@ -729,6 +729,10 @@ ExistingSkelMeshData* SaveExistingSkelMeshData(USkeletalMesh* ExistingSkelMesh)
 				LODModel.MultiSizeIndexContainer.GetIndexBufferData( ExistingData );
 				ExistingMeshDataPtr->ExistingIndexBufferData.Add( ExistingData );
 
+				FMultiSizeIndexContainerData ExistingAdjacencyData;
+				LODModel.AdjacencyMultiSizeIndexContainer.GetIndexBufferData(ExistingAdjacencyData);
+				ExistingMeshDataPtr->ExistingAdjacencyIndexBufferData.Add(ExistingAdjacencyData);
+
 			}
 
 			ExistingMeshDataPtr->ExistingLODInfo = ExistingSkelMesh->LODInfo;
@@ -885,6 +889,7 @@ void RestoreExistingSkelMeshData(ExistingSkelMeshData* MeshData, USkeletalMesh*
 					FStaticLODModel* NewLODModel = new(SkeletalMesh->GetImportedResource()->LODModels) FStaticLODModel( LODModel );
 			
 					NewLODModel->MultiSizeIndexContainer.RebuildIndexBuffer( MeshData->ExistingIndexBufferData[i] );
+					NewLODModel->AdjacencyMultiSizeIndexContainer.RebuildIndexBuffer( MeshData->ExistingAdjacencyIndexBufferData[i] );
 
 					SkeletalMesh->LODInfo.Add( LODInfo );
 				}
diff --git a/Engine/Source/Editor/UnrealEd/Public/SkelImport.h b/Engine/Source/Editor/UnrealEd/Public/SkelImport.h
--- a/Engine/Source/Editor/UnrealEd/Public/SkelImport.h
+++ b/Engine/Source/Editor/UnrealEd/Public/SkelImport.h
@@ -125,6 +125,7 @@ struct ExistingSkelMeshData
 	TIndirectArray<FStaticLODModel>			ExistingLODModels;
 	TArray<FSkeletalMeshLODInfo>			ExistingLODInfo;
 	TArray<FMultiSizeIndexContainerData>	ExistingIndexBufferData;
+	TArray<FMultiSizeIndexContainerData>	ExistingAdjacencyIndexBufferData;
 	FReferenceSkeleton						ExistingRefSkeleton;
 	TArray<FSkeletalMaterial>				ExistingMaterials;
 	TArray<UMorphTarget*>					ExistingMorphTargets;
