[UE-21284] Don't delete derived class CDOs on force delete of Blueprint Class assets in the content browser.

diff --git a/Engine/Source/Editor/UnrealEd/Private/ObjectTools.cpp b/Engine/Source/Editor/UnrealEd/Private/ObjectTools.cpp
--- a/Engine/Source/Editor/UnrealEd/Private/ObjectTools.cpp
+++ b/Engine/Source/Editor/UnrealEd/Private/ObjectTools.cpp
@@ -1889,6 +1889,12 @@ namespace ObjectTools
 				{
 					UObject* CurrentInstance = *InstanceItr;
 
+					// Don't include derived class CDOs.
+					if(CurrentInstance->HasAnyFlags(RF_ClassDefaultObject))
+					{
+						continue;
+					}
+
 					AActor* CurrentInstanceAsActor = Cast<AActor>( CurrentInstance );
 					UActorComponent* CurrentInstanceAsComponent = Cast<UActorComponent>(CurrentInstance);
 					if ( CurrentInstanceAsActor )
