Fixed a crash dereferecing a null ptr

diff --git a/Engine/Source/Editor/UnrealEd/Private/LevelEditorViewport.cpp b/Engine/Source/Editor/UnrealEd/Private/LevelEditorViewport.cpp
--- a/Engine/Source/Editor/UnrealEd/Private/LevelEditorViewport.cpp
+++ b/Engine/Source/Editor/UnrealEd/Private/LevelEditorViewport.cpp
@@ -1720,9 +1720,10 @@ void FLevelEditorViewportClient::BeginCameraMovement(bool bHasMovement)
 	{
 		if (!bIsCameraMoving)
 		{
-			if (!bIsCameraMovingOnTick)
+			AActor* ActorLock = GetActiveActorLock().Get();
+			if (!bIsCameraMovingOnTick && ActorLock)
 			{
-				GEditor->BroadcastBeginCameraMovement(*GetActiveActorLock());
+				GEditor->BroadcastBeginCameraMovement(*ActorLock);
 			}
 			bIsCameraMoving = true;
 		}
@@ -1736,11 +1737,11 @@ void FLevelEditorViewportClient::BeginCameraMovement(bool bHasMovement)
 void FLevelEditorViewportClient::EndCameraMovement()
 {
 	// If there was movement and it has now stopped, broadcast it
-	if (bIsCameraMovingOnTick)
+	if (bIsCameraMovingOnTick && !bIsCameraMoving)
 	{
-		if (!bIsCameraMoving)
+		if (AActor* ActorLock = GetActiveActorLock().Get())
 		{
-			GEditor->BroadcastEndCameraMovement(*GetActiveActorLock());
+			GEditor->BroadcastEndCameraMovement(*ActorLock);
 		}
 	}
 }
