Better detection of PIE when duplicating blueprints

diff --git a/Engine/Source/Editor/UnrealEd/Private/Kismet2/BlueprintEditorUtils.cpp b/Engine/Source/Editor/UnrealEd/Private/Kismet2/BlueprintEditorUtils.cpp
--- a/Engine/Source/Editor/UnrealEd/Private/Kismet2/BlueprintEditorUtils.cpp
+++ b/Engine/Source/Editor/UnrealEd/Private/Kismet2/BlueprintEditorUtils.cpp
@@ -1877,12 +1877,12 @@ void FBlueprintEditorUtils::PropagateParentBlueprintDefaults(UClass* ClassToProp
 
 UNREALED_API FSecondsCounterData BlueprintCompileAndLoadTimerData;
 
-void FBlueprintEditorUtils::PostDuplicateBlueprint(UBlueprint* Blueprint)
+void FBlueprintEditorUtils::PostDuplicateBlueprint(UBlueprint* Blueprint, bool bDuplicateForPIE)
 {
 	FSecondsCounterScope Timer(BlueprintCompileAndLoadTimerData); 
 	
 	// Only recompile after duplication if this isn't PIE
-	if (!GIsPlayInEditorWorld)
+	if (!bDuplicateForPIE)
 	{
 		check(Blueprint->GeneratedClass != NULL);
 		{
diff --git a/Engine/Source/Editor/UnrealEd/Public/Kismet2/BlueprintEditorUtils.h b/Engine/Source/Editor/UnrealEd/Public/Kismet2/BlueprintEditorUtils.h
--- a/Engine/Source/Editor/UnrealEd/Public/Kismet2/BlueprintEditorUtils.h
+++ b/Engine/Source/Editor/UnrealEd/Public/Kismet2/BlueprintEditorUtils.h
@@ -157,7 +157,7 @@ public:
 	static void PropagateParentBlueprintDefaults(UClass* ClassToPropagate);
 
 	/** Called on a Blueprint after it has been duplicated */
-	static void PostDuplicateBlueprint(UBlueprint* Blueprint);
+	static void PostDuplicateBlueprint(UBlueprint* Blueprint, bool bDuplicateForPIE);
 
 	/** Consigns the blueprint's generated classes to oblivion */
 	static void RemoveGeneratedClasses(UBlueprint* Blueprint);
diff --git a/Engine/Source/Runtime/Engine/Private/Blueprint.cpp b/Engine/Source/Runtime/Engine/Private/Blueprint.cpp
--- a/Engine/Source/Runtime/Engine/Private/Blueprint.cpp
+++ b/Engine/Source/Runtime/Engine/Private/Blueprint.cpp
@@ -418,7 +418,7 @@ void UBlueprint::PostDuplicate(bool bDuplicateForPIE)
 	Super::PostDuplicate(bDuplicateForPIE);
 	if( !bDuplicatingReadOnly )
 	{
-		FBlueprintEditorUtils::PostDuplicateBlueprint(this);
+		FBlueprintEditorUtils::PostDuplicateBlueprint(this, bDuplicateForPIE);
 	}
 }
 
