Add ability to choose Test or Shipping configuration for Launch On Device

diff --git a/Engine/Source/Editor/UnrealEd/Classes/Settings/LevelEditorPlaySettings.h b/Engine/Source/Editor/UnrealEd/Classes/Settings/LevelEditorPlaySettings.h
--- a/Engine/Source/Editor/UnrealEd/Classes/Settings/LevelEditorPlaySettings.h
+++ b/Engine/Source/Editor/UnrealEd/Classes/Settings/LevelEditorPlaySettings.h
@@ -101,6 +101,10 @@ enum EPlayOnLaunchConfiguration
 	LaunchConfig_Debug UMETA(DisplayName = "Debug"),
 	/** Launch on device with a Development build configuration. */
 	LaunchConfig_Development UMETA(DisplayName = "Development"),
+	/** Launch on device with a Test build configuration. */
+	LaunchConfig_Test UMETA(DisplayName = "Test"),
+	/** Launch on device with a Shipping build configuration. */
+	LaunchConfig_Shipping UMETA(DisplayName = "Shipping"),
 };
 
 /**
diff --git a/Engine/Source/Editor/UnrealEd/Private/PlayLevel.cpp b/Engine/Source/Editor/UnrealEd/Private/PlayLevel.cpp
--- a/Engine/Source/Editor/UnrealEd/Private/PlayLevel.cpp
+++ b/Engine/Source/Editor/UnrealEd/Private/PlayLevel.cpp
@@ -1692,6 +1692,12 @@ void UEditorEngine::PlayUsingLauncher()
 		case LaunchConfig_Development:
 			LauncherProfile->SetBuildConfiguration(EBuildConfigurations::Development);
 			break;
+		case LaunchConfig_Test:
+			LauncherProfile->SetBuildConfiguration(EBuildConfigurations::Test);
+			break;
+		case LaunchConfig_Shipping:
+			LauncherProfile->SetBuildConfiguration(EBuildConfigurations::Shipping);
+			break;
 		default:
 			// same as the running editor
 			FString ExeName = FUnrealEdMisc::Get().GetExecutableForCommandlets();
