Changed FEditorFileUtils::FindAllSubmittablePackageFiles to use FPackageName::TryConvertFilenameToLongPackageName to avoid a crash if it finds a bad package path

diff --git a/Engine/Source/Editor/UnrealEd/Private/FileHelpers.cpp b/Engine/Source/Editor/UnrealEd/Private/FileHelpers.cpp
--- a/Engine/Source/Editor/UnrealEd/Private/FileHelpers.cpp
+++ b/Engine/Source/Editor/UnrealEd/Private/FileHelpers.cpp
@@ -3264,7 +3264,15 @@ void FEditorFileUtils::FindAllSubmittablePackageFiles(TMap<FString, FSourceContr
 	for (TArray<FString>::TConstIterator PackageIter(Packages); PackageIter; ++PackageIter)
 	{
 		const FString Filename = *PackageIter;
-		const FString PackageName = FPackageName::FilenameToLongPackageName(Filename);
+
+		FString PackageName;
+		FString FailureReason;
+		if (!FPackageName::TryConvertFilenameToLongPackageName(Filename, PackageName, &FailureReason))
+		{
+			UE_LOG(LogFileHelpers, Warning, TEXT("%s"), *FailureReason);
+			continue;
+		}
+
 		FSourceControlStatePtr SourceControlState = SourceControlProvider.GetState(FPaths::ConvertRelativePathToFull(Filename), EStateCacheUsage::Use);
 
 		// Only include non-map packages that are currently checked out or packages not under source control
@@ -3272,7 +3280,7 @@ void FEditorFileUtils::FindAllSubmittablePackageFiles(TMap<FString, FSourceContr
 			(SourceControlState->IsCheckedOut() || SourceControlState->IsAdded() || (!SourceControlState->IsSourceControlled() && SourceControlState->CanAdd())) &&
 			(bIncludeMaps || !IsMapPackageAsset(*Filename)))
 		{
-			OutPackages.Add(PackageName, SourceControlState);
+			OutPackages.Add(MoveTemp(PackageName), MoveTemp(SourceControlState));
 		}
 	}
 }
