UE-21469 - Editor - Updating the C++ project wizard auto installion of visual studio to be setup to install VS2015 instead of 2013.

diff --git a/Engine/Config/BaseEditor.ini b/Engine/Config/BaseEditor.ini
--- a/Engine/Config/BaseEditor.ini
+++ b/Engine/Config/BaseEditor.ini
@@ -301,7 +301,7 @@ DocumentationURL="http://udn.unrealengine.com/docs/ue4/INT/index.html"
 AskAQuestionURL="http://udn.unrealengine.com/questions/ask.html"
 SearchForAnswersURL="http://udn.unrealengine.com/search.html"
 SupportWebsiteURL="http://udn.unrealengine.com/index.html"
-SourceCodeIDEURL_Windows="http://go.microsoft.com/fwlink/?LinkId=525621"
+SourceCodeIDEURL_Windows="http://go.microsoft.com/fwlink/?LinkId=532606"
 SourceCodeIDEURL_Mac="https://developer.apple.com/xcode"
 SourceCodeIDEURL_Other="http://en.wikipedia.org/wiki/Integrated_development_environment"
 EpicGamesURL="http://UnrealEngine.com"
diff --git a/Engine/Plugins/Developer/VisualStudioSourceCodeAccess/Source/VisualStudioSourceCodeAccess/Private/VisualStudioSourceCodeAccessor.cpp b/Engine/Plugins/Developer/VisualStudioSourceCodeAccess/Source/VisualStudioSourceCodeAccess/Private/VisualStudioSourceCodeAccessor.cpp
--- a/Engine/Plugins/Developer/VisualStudioSourceCodeAccess/Source/VisualStudioSourceCodeAccess/Private/VisualStudioSourceCodeAccessor.cpp
+++ b/Engine/Plugins/Developer/VisualStudioSourceCodeAccess/Source/VisualStudioSourceCodeAccess/Private/VisualStudioSourceCodeAccessor.cpp
@@ -80,7 +80,7 @@ void FVisualStudioSourceCodeAccessor::Startup()
 	GetSolutionPath();
 
 	// Preferential order of VS versions
-	AddVisualStudioVersion(12); // Visual Studio 2013
+	AddVisualStudioVersion(14); // Visual Studio 2015
 }
 
 void FVisualStudioSourceCodeAccessor::Shutdown()
diff --git a/Engine/Source/Editor/UnrealEd/Private/SourceCodeNavigation.cpp b/Engine/Source/Editor/UnrealEd/Private/SourceCodeNavigation.cpp
--- a/Engine/Source/Editor/UnrealEd/Private/SourceCodeNavigation.cpp
+++ b/Engine/Source/Editor/UnrealEd/Private/SourceCodeNavigation.cpp
@@ -1494,7 +1494,7 @@ FText FSourceCodeNavigation::GetSuggestedSourceCodeIDE(bool bShortIDEName)
 	}
 	else
 	{
-		return LOCTEXT("SuggestedCodeIDE_Windows", "Visual Studio 2013");
+		return LOCTEXT("SuggestedCodeIDE_Windows", "Visual Studio 2015");
 	}
 #elif PLATFORM_MAC
 	return LOCTEXT("SuggestedCodeIDE_Mac", "Xcode");
@@ -1827,7 +1827,7 @@ FString FSourceCodeNavigationImpl::GetSuggestedIDEInstallerFileName()
 void FSourceCodeNavigationImpl::LaunchIDEInstaller(const FString& Filepath)
 {
 #if PLATFORM_WINDOWS
-	auto Params = TEXT("/PromptRestart /ChainingPackage EpicGames_UE4");
+	auto Params = TEXT("/PromptRestart /InstallSelectableItems NativeLanguageSupport_Group /ChainingPackage EpicGames_UE4");
 	FPlatformProcess::CreateProc(*Filepath, Params, true, false, false, nullptr, 0, nullptr, nullptr);
 #endif
 }
diff --git a/Engine/Source/Runtime/Core/Private/Windows/WindowsPlatformMisc.cpp b/Engine/Source/Runtime/Core/Private/Windows/WindowsPlatformMisc.cpp
--- a/Engine/Source/Runtime/Core/Private/Windows/WindowsPlatformMisc.cpp
+++ b/Engine/Source/Runtime/Core/Private/Windows/WindowsPlatformMisc.cpp
@@ -2422,7 +2422,7 @@ bool FWindowsPlatformMisc::QueryRegKey( const HKEY InKey, const TCHAR* InSubKey,
 
 bool FWindowsPlatformMisc::GetVSComnTools(int32 Version, FString& OutData)
 {
-	checkf(11 <= Version && Version <= 12, L"Not supported Visual Studio version.");
+	checkf(11 <= Version && Version <= 14, L"Not supported Visual Studio version.");
 
 	const TCHAR* PossibleRegPaths[] = {
 		L"Wow6432Node\\Microsoft\\VisualStudio",	// Non-express VS2013 on 64-bit machine.
