Rebuild Lightmaps Commandlet

Runs a lighting build for specific, or all, map(s). Reusing the ResavePackages commandlet. The commandline for the commandlet is:
ProjectName -run=resavepackages -buildlighting -allowcommandletrendering -map=optionalmapname

Script added to UAT to automate the process. The commandline for this is:
RebuildLightmaps -project="My:/Absolute/Project/Path.uproject" -MapsToRebuildLightMaps=OptionalMapName

diff --git a/Engine/Source/Editor/UnrealEd/Classes/Commandlets/ResavePackagesCommandlet.h b/Engine/Source/Editor/UnrealEd/Classes/Commandlets/ResavePackagesCommandlet.h
--- a/Engine/Source/Editor/UnrealEd/Classes/Commandlets/ResavePackagesCommandlet.h
+++ b/Engine/Source/Editor/UnrealEd/Classes/Commandlets/ResavePackagesCommandlet.h
@@ -72,18 +72,20 @@ class UNREALED_API UResavePackagesCommandlet : public UCommandlet
 
 	// List of files to submit
 	TArray<FString> FilesToSubmit;
+
+	// The list of switches that were passed on the commandline
+	TArray<FString> Switches;
 protected:
 	/**
 	 * Evaluates the command-line to determine which maps to check.  By default all maps are checked
 	 * Provides child classes with a chance to initialize any variables, parse the command line, etc.
 	 *
 	 * @param	Tokens			the list of tokens that were passed to the commandlet
-	 * @param	Switches		the list of switches that were passed on the commandline
 	 * @param	MapPathNames	receives the list of path names for the maps that will be checked.
 	 *
 	 * @return	0 to indicate that the commandlet should continue; otherwise, the error code that should be returned by Main()
 	 */
-	virtual int32 InitializeResaveParameters( const TArray<FString>& Tokens, const TArray<FString>& Switches, TArray<FString>& MapPathNames );
+	virtual int32 InitializeResaveParameters( const TArray<FString>& Tokens, TArray<FString>& MapPathNames );
 
 	// Loads and saves a single package
 	virtual void LoadAndSaveOnePackage(const FString& Filename);
@@ -118,6 +120,15 @@ protected:
 	 */
 	virtual void PerformAdditionalOperations( class UPackage* Package, bool& bSavePackage );
 
+	/**
+	* Allows the commandlet to perform any additional operations on the world before it is resaved.
+	*
+	* @param	World			the world that is currently being processed
+	* @param	bSavePackage	[in]	indicates whether the package is currently going to be saved
+	*							[out]	set to true to resave the package
+	*/
+	virtual void PerformAdditionalOperations(class UWorld* World, bool& bSavePackage);
+
 	/**
 	 * Removes any UClass exports from packages which aren't script packages.
 	 *
diff --git a/Engine/Source/Editor/UnrealEd/Private/Commandlets/ContentCommandlets.cpp b/Engine/Source/Editor/UnrealEd/Private/Commandlets/ContentCommandlets.cpp
--- a/Engine/Source/Editor/UnrealEd/Private/Commandlets/ContentCommandlets.cpp
+++ b/Engine/Source/Editor/UnrealEd/Private/Commandlets/ContentCommandlets.cpp
@@ -27,6 +27,10 @@ DEFINE_LOG_CATEGORY_STATIC(LogContentCommandlet, Log, All);
 #include "Engine/LevelStreaming.h"
 #include "Engine/StaticMesh.h"
 
+// for UResavePackagesCommandlet::PerformAdditionalOperations building lighting code
+#include "Engine/WorldComposition.h"
+#include "LightingBuildOptions.h"
+
 
 /**-----------------------------------------------------------------------------
  *	UResavePackages commandlet.
@@ -47,7 +51,7 @@ UResavePackagesCommandlet::UResavePackagesCommandlet(const FObjectInitializer& O
 }
 
 
-int32 UResavePackagesCommandlet::InitializeResaveParameters( const TArray<FString>& Tokens, const TArray<FString>& Switches, TArray<FString>& PackageNames )
+int32 UResavePackagesCommandlet::InitializeResaveParameters( const TArray<FString>& Tokens, TArray<FString>& PackageNames )
 {
 	Verbosity = VERY_VERBOSE;
 
@@ -297,6 +301,13 @@ void UResavePackagesCommandlet::LoadAndSaveOnePackage(const FString& Filename)
 				}
 			}
 
+			{
+				UWorld* World = UWorld::FindWorldInPackage(Package);
+				if (World)
+				{
+					PerformAdditionalOperations(World, bSavePackage);
+				}
+			}
 
 			// hook to allow performing additional checks without lumping everything into this one function
 			PerformAdditionalOperations(Package,bSavePackage);
@@ -485,7 +496,7 @@ void UResavePackagesCommandlet::LoadAndSaveOnePackage(const FString& Filename)
 int32 UResavePackagesCommandlet::Main( const FString& Params )
 {
 	const TCHAR* Parms = *Params;
-	TArray<FString> Tokens, Switches;
+	TArray<FString> Tokens;
 	ParseCommandLine(Parms, Tokens, Switches);
 
 	// Ensure source control is initialized and shut down properly
@@ -505,7 +516,7 @@ int32 UResavePackagesCommandlet::Main( const FString& Params )
 	bAutoCheckIn = bAutoCheckOut && Switches.Contains(TEXT("AutoCheckIn"));
 
 	TArray<FString> PackageNames;
-	int32 ResultCode = InitializeResaveParameters(Tokens, Switches, PackageNames);
+	int32 ResultCode = InitializeResaveParameters(Tokens, PackageNames);
 	if ( ResultCode != 0 )
 	{
 		return ResultCode;
@@ -642,7 +653,106 @@ bool UResavePackagesCommandlet::PerformPreloadOperations( FLinkerLoad* PackageLi
 
 	return bResult;
 }
+void UResavePackagesCommandlet::PerformAdditionalOperations(class UWorld* World, bool& bSavePackage)
+{
+	check(World);
+	static bool bBuildLighting = FParse::Param(FCommandLine::Get(), TEXT("buildlighting")) == true;
+
+	if (bBuildLighting)
+	{
+		TArray<FString> CmdLineMapEntries;
+		for (const FString& Switch : Switches)
+		{
+			auto GetSwitchValueElements = [&Switch](const FString SwitchKey) -> TArray < FString >
+			{
+				TArray<FString> ValueElements;
+				if (Switch.StartsWith(SwitchKey + TEXT("=")) == true)
+				{
+					FString ValuesList = Switch.Right(Switch.Len() - (SwitchKey + TEXT("=")).Len());
+
+					// Allow support for -KEY=Value1+Value2+Value3 as well as -KEY=Value1 -KEY=Value2
+					for (int32 PlusIdx = ValuesList.Find(TEXT("+")); PlusIdx != INDEX_NONE; PlusIdx = ValuesList.Find(TEXT("+")))
+					{
+						const FString ValueElement = ValuesList.Left(PlusIdx);
+						ValueElements.Add(ValueElement);
+
+						ValuesList = ValuesList.Right(ValuesList.Len() - (PlusIdx + 1));
+					}
+					ValueElements.Add(ValuesList);
+				}
+				return ValueElements;
+			};
+
+			CmdLineMapEntries += GetSwitchValueElements(TEXT("MAP"));
+		}
+
+		// If specific maps were specified, then only continue if we are processing one that was required.
+		bool bShouldBuildLightmapsForWorld = CmdLineMapEntries.Num() == 0 || CmdLineMapEntries.Contains(World->GetMapName());
+		if (bShouldBuildLightmapsForWorld)
+		{
+			World->AddToRoot();
+			if (!World->bIsWorldInitialized)
+			{
+				UWorld::InitializationValues IVS;
+				IVS.RequiresHitProxies(false);
+				IVS.ShouldSimulatePhysics(false);
+				IVS.EnableTraceCollision(false);
+				IVS.CreateNavigation(false);
+				IVS.CreateAISystem(false);
+				IVS.AllowAudioPlayback(false);
+				IVS.CreatePhysicsScene(false);
+
+				World->InitWorld(IVS);
+				World->PersistentLevel->UpdateModelComponents();
+				World->UpdateWorldComponents(true, false);
+			}
+
+			// load all the sublevels
 
+			if (World->StreamingLevels.Num())
+			{
+				//World->LoadSecondaryLevels(true, &PreviouslyCookedPackages);
+				World->LoadSecondaryLevels(true, NULL);
+			}
+
+			GWorld = World;
+
+			TArray<FString> SubPackages;
+
+			// force load all of the map to be loaded so we can get good rebuilt lighting 
+			if (World->WorldComposition)
+			{
+				World->WorldComposition->CollectTilesToCook(SubPackages);
+			}
+
+			for (const auto& SubPackage : SubPackages)
+			{
+				UPackage* Package = LoadPackage(nullptr, *SubPackage, 0);
+				check(Package->IsFullyLoaded());
+			}
+
+			GRedirectCollector.ResolveStringAssetReference();
+
+
+			FLightingBuildOptions LightingOptions;
+
+			int32 QualityLevel = Quality_Production;
+			GConfig->GetInt(TEXT("LightingBuildOptions"), TEXT("QualityLevel"), QualityLevel, GEditorPerProjectIni);
+			QualityLevel = FMath::Clamp<int32>(QualityLevel, Quality_Preview, Quality_Production);
+			
+			LightingOptions.QualityLevel = (ELightingBuildQuality)QualityLevel;
+
+			GEditor->BuildLighting(LightingOptions);
+			while (GEditor->IsLightingBuildCurrentlyRunning())
+			{
+				GEditor->UpdateBuildLighting();
+			}
+
+			World->RemoveFromRoot();
+		}
+
+	}
+}
 
 void UResavePackagesCommandlet::PerformAdditionalOperations( class UObject* Object, bool& bSavePackage )
 {
diff --git a/Engine/Source/Editor/UnrealEd/Private/StaticLightingSystem/StaticLightingSystem.cpp b/Engine/Source/Editor/UnrealEd/Private/StaticLightingSystem/StaticLightingSystem.cpp
--- a/Engine/Source/Editor/UnrealEd/Private/StaticLightingSystem/StaticLightingSystem.cpp
+++ b/Engine/Source/Editor/UnrealEd/Private/StaticLightingSystem/StaticLightingSystem.cpp
@@ -784,7 +784,7 @@ void FStaticLightingSystem::GatherStaticLightingInfo(bool bRebuildDirtyGeometryF
 			{
 				// This will go ahead and clean up lighting on all dirty levels (not just this one)
 				UE_LOG(LogStaticLightingSystem, Warning, TEXT("WARNING: Lighting build automatically rebuilding geometry.") );
-				GUnrealEd->Exec( World, TEXT("MAP REBUILD ALLDIRTYFORLIGHTING") );
+				GEditor->Exec(World, TEXT("MAP REBUILD ALLDIRTYFORLIGHTING"));
 			}
 		}
 
@@ -1861,20 +1861,24 @@ bool FStaticLightingSystem::InitiateLightmassProcessor()
 			bSuccessful = true;
 			CurrentBuildStage = FStaticLightingSystem::AmortizedExport;
 
-			// Crash tracker interferes with performance during export only.
-			// Disable it only for export, for everything else it shouldn't matter.
-			// This is a very special case, and doing this sort of thing
-			// is almost never recommended, especially without profiling heavily.
-			// The reason it works here is because amortized export flushes the render
-			// commands every tick, which is highly detrimental to the crash tracker's operation.
-			// ALSO NOTE: The reason this is set here rather than be a common API in the crashtracker
-			// module is to discourage people from doing this sort of thing all over the place.
-			ICrashTrackerModule* CrashTracker = FModuleManager::LoadModulePtr<ICrashTrackerModule>("CrashTracker");
-			if (CrashTracker)
+			if (!IsRunningCommandlet())
 			{
-				bCrashTrackerOriginallyEnabled = CrashTracker->IsCurrentlyCapturing();
-				CrashTracker->SetCrashTrackingEnabled(false);
+				// Crash tracker interferes with performance during export only.
+				// Disable it only for export, for everything else it shouldn't matter.
+				// This is a very special case, and doing this sort of thing
+				// is almost never recommended, especially without profiling heavily.
+				// The reason it works here is because amortized export flushes the render
+				// commands every tick, which is highly detrimental to the crash tracker's operation.
+				// ALSO NOTE: The reason this is set here rather than be a common API in the crashtracker
+				// module is to discourage people from doing this sort of thing all over the place.
+				ICrashTrackerModule* CrashTracker = FModuleManager::LoadModulePtr<ICrashTrackerModule>("CrashTracker");
+				if (CrashTracker)
+				{
+					bCrashTrackerOriginallyEnabled = CrashTracker->IsCurrentlyCapturing();
+					CrashTracker->SetCrashTrackingEnabled(false);
+				}
 			}
+
 		}
 	}
 	
@@ -2085,9 +2089,9 @@ bool FStaticLightingSystem::CanAutoApplyLighting() const
 	const bool bAutoApplyEnabled = GetDefault<ULevelEditorMiscSettings>()->bAutoApplyLightingEnable;
 	const bool bSlowTask = GIsSlowTask;
 	const bool bInterpEditMode = GLevelEditorModeTools().IsModeActive( FBuiltinEditorModes::EM_InterpEdit );
-	const bool bPlayWorldValid = GUnrealEd->PlayWorld != nullptr;
-	const bool bAnyMenusVisible = FSlateApplication::Get().AnyMenusVisible();
-	//const bool bIsInteratcting = false;// FSlateApplication::Get().GetMouseCaptor().IsValid() || GUnrealEd->IsUserInteracting();
+	const bool bPlayWorldValid = GEditor->PlayWorld != nullptr;
+	const bool bAnyMenusVisible = (FSlateApplication::IsInitialized() && FSlateApplication::Get().AnyMenusVisible());
+	//const bool bIsInteratcting = false;// FSlateApplication::Get().GetMouseCaptor().IsValid() || GEditor->IsUserInteracting();
 	const bool bHasGameOrProjectLoaded = FApp::HasGameName();
 
 	return ( bAutoApplyEnabled && !bSlowTask && !bInterpEditMode && !bPlayWorldValid && !bAnyMenusVisible/* && !bIsInteratcting */&& !GIsDemoMode && bHasGameOrProjectLoaded );
diff --git a/Engine/Source/Programs/AutomationTool/AutomationUtils/CommandletUtils.cs b/Engine/Source/Programs/AutomationTool/AutomationUtils/CommandletUtils.cs
--- a/Engine/Source/Programs/AutomationTool/AutomationUtils/CommandletUtils.cs
+++ b/Engine/Source/Programs/AutomationTool/AutomationUtils/CommandletUtils.cs
@@ -72,6 +72,25 @@ namespace AutomationTool
             RunCommandlet(ProjectName, UE4Exe, "DerivedDataCache", String.Format("{0} -TargetPlatform={1} {2}", MapsToCook, TargetPlatform, Parameters));
         }
 
+		/// <summary>
+		/// Runs RebuildLightMaps commandlet.
+		/// </summary>
+		/// <param name="ProjectName">Project name.</param>
+		/// <param name="UE4Exe">The name of the UE4 Editor executable to use.</param>
+		/// <param name="Maps">List of maps to rebuild light maps for. Can be null in which case -MapIniSection=AllMaps is used.</param>
+		/// <param name="Parameters">List of additional parameters.</param>
+		public static void RebuildLightMapsCommandlet(FileReference ProjectName, string UE4Exe = "UE4Editor-Cmd.exe", string[] Maps = null, string Parameters = "")
+		{
+			string MapsToRebuildLighting = "";
+			if (!IsNullOrEmpty(Maps))
+			{
+				MapsToRebuildLighting = "-Map=" + CombineCommandletParams(Maps);
+				MapsToRebuildLighting.Trim();
+			}
+
+			RunCommandlet(ProjectName, UE4Exe, "ResavePackages", String.Format("-buildlighting -AllowCommandletRendering {0} {1}", MapsToRebuildLighting, Parameters));
+		}
+
         /// <summary>
         /// Runs GenerateDistillFileSets commandlet.
         /// </summary>
diff --git a/Engine/Source/Programs/AutomationTool/AutomationUtils/ProjectParams.cs b/Engine/Source/Programs/AutomationTool/AutomationUtils/ProjectParams.cs
--- a/Engine/Source/Programs/AutomationTool/AutomationUtils/ProjectParams.cs
+++ b/Engine/Source/Programs/AutomationTool/AutomationUtils/ProjectParams.cs
@@ -404,7 +404,8 @@ namespace AutomationTool
 			bool? IgnoreCookErrors = null,
 			bool? CodeSign = null,
 			bool? UploadSymbols = null,
-			string Provision = null
+			string Provision = null,
+			ParamList<string> InMapsToRebuildLightMaps = null
 			)
 		{
 			//
@@ -704,6 +705,28 @@ namespace AutomationTool
             {
                 this.MapsToCook.Add(this.MapToRun);
             }
+
+			if (InMapsToRebuildLightMaps == null)
+			{
+				if (Command != null)
+				{
+					this.MapsToRebuildLightMaps = new ParamList<string>();
+
+					var MapsString = Command.ParseParamValue("MapsToRebuildLightMaps");
+					if (String.IsNullOrEmpty(MapsString) == false)
+					{
+						var MapNames = new ParamList<string>(MapsString.Split('+'));
+						foreach (var M in MapNames)
+						{
+							this.MapsToRebuildLightMaps.Add(M);
+						}
+					}
+				}
+			}
+			else
+			{
+				this.MapsToRebuildLightMaps = InMapsToRebuildLightMaps;
+			}
             
 			if (ServerConfigsToBuild == null)
 			{
@@ -1473,6 +1496,16 @@ namespace AutomationTool
 
 		#endregion
 
+		#region Misc
+		
+		[Help("MapsToRebuildLightMaps", "List of maps that need light maps rebuilding")]
+		public ParamList<string> MapsToRebuildLightMaps = new ParamList<string>();
+
+		[Help("IgnoreLightMapErrors", "Whether Light Map errors should be treated as critical")]
+		public bool IgnoreLightMapErrors { get; set; }
+
+		#endregion
+
 		#region Initialization
 
 		private Dictionary<TargetRules.TargetType, SingleTargetProperties> DetectedTargets;
diff --git a/Engine/Source/Programs/AutomationTool/Scripts/AutomationScripts.Automation.csproj b/Engine/Source/Programs/AutomationTool/Scripts/AutomationScripts.Automation.csproj
--- a/Engine/Source/Programs/AutomationTool/Scripts/AutomationScripts.Automation.csproj
+++ b/Engine/Source/Programs/AutomationTool/Scripts/AutomationScripts.Automation.csproj
@@ -73,6 +73,7 @@
     <Compile Include="DeployCommand.Automation.cs" />
     <Compile Include="FixupRedirects.Automation.cs" />
     <Compile Include="GenerateDSYM.Automation.cs" />
+    <Compile Include="RebuildLightMapsCommand.Automation.cs" />
     <Compile Include="RocketBuild.Automation.cs" />
     <Compile Include="IPhonePackager.Automation.cs" />
     <Compile Include="LauncherLocalization.Automation.cs" />
diff --git a/Engine/Source/Programs/AutomationTool/Scripts/RebuildLightMapsCommand.Automation.cs b/Engine/Source/Programs/AutomationTool/Scripts/RebuildLightMapsCommand.Automation.cs
new file mode 100644
--- /dev/null
+++ b/Engine/Source/Programs/AutomationTool/Scripts/RebuildLightMapsCommand.Automation.cs
@@ -0,0 +1,107 @@
+﻿// Copyright 1998-2015 Epic Games, Inc. All Rights Reserved.
+using System;
+using System.Collections.Generic;
+using System.IO;
+using System.Threading;
+using System.Reflection;
+using System.Linq;
+using AutomationTool;
+using UnrealBuildTool;
+
+/// <summary>
+/// Helper command used for rebuilding a projects light maps.
+/// </summary>
+/// <remarks>
+/// Command line parameters used by this command:
+/// -
+/// </remarks>
+namespace AutomationScripts.Automation
+{
+	public class RebuildLightMaps : BuildCommand
+	{
+		#region RebuildLightMaps Command
+
+		public override void ExecuteBuild()
+		{
+			var Params = new ProjectParams
+			(
+				Command: this,
+				// Shared
+				RawProjectPath: ProjectPath
+			);
+
+
+			Log("********** REBUILD LIGHT MAPS COMMAND STARTED **********");
+
+			string UE4EditorExe = HostPlatform.Current.GetUE4ExePath(Params.UE4Exe);
+			if (!FileExists(UE4EditorExe))
+			{
+				throw new AutomationException("Missing " + UE4EditorExe + " executable. Needs to be built first.");
+			}
+
+			try
+			{
+				var CommandletParams = IsBuildMachine ? "-buildmachine -fileopenlog" : "-fileopenlog";
+				RebuildLightMapsCommandlet(Params.RawProjectPath, Params.UE4Exe, Params.MapsToRebuildLightMaps.ToArray(), CommandletParams);
+			}
+			catch (Exception Ex)
+			{
+				if (Params.IgnoreLightMapErrors)
+				{
+					LogWarning("Ignoring Rebuild Light Map failure.");
+				}
+				else
+				{
+					// Abandon this run, dont check in any updated files, etc.
+					Log("Rebuild Light Maps has failed.");
+					//CleanupCookedData(PlatformsToCook.ToList(), Params);
+					throw new AutomationException(ExitCode.Error_UnknownCookFailure, Ex, "RebuildLightMaps failed.");
+				}
+			}
+			Log("********** COOK COMMAND COMPLETED **********");
+		}
+
+		#endregion
+
+
+		private FileReference ProjectFullPath;
+		public virtual FileReference ProjectPath
+		{
+			get
+			{
+				if (ProjectFullPath == null)
+				{
+					var OriginalProjectName = ParseParamValue("project", "");
+					var ProjectName = OriginalProjectName;
+					ProjectName = ProjectName.Trim(new char[] { '\"' });
+					if (ProjectName.IndexOfAny(new char[] { '\\', '/' }) < 0)
+					{
+						ProjectName = CombinePaths(CmdEnv.LocalRoot, ProjectName, ProjectName + ".uproject");
+					}
+					else if (!FileExists_NoExceptions(ProjectName))
+					{
+						ProjectName = CombinePaths(CmdEnv.LocalRoot, ProjectName);
+					}
+					if (FileExists_NoExceptions(ProjectName))
+					{
+						ProjectFullPath = new FileReference(ProjectName);
+					}
+					else
+					{
+						var Branch = new BranchInfo(new List<UnrealTargetPlatform> { UnrealBuildTool.BuildHostPlatform.Current.Platform });
+						var GameProj = Branch.FindGame(OriginalProjectName);
+						if (GameProj != null)
+						{
+							ProjectFullPath = GameProj.FilePath;
+						}
+						if (!FileExists_NoExceptions(ProjectFullPath.FullName))
+						{
+							throw new AutomationException("Could not find a project file {0}.", ProjectName);
+						}
+					}
+				}
+				return ProjectFullPath;
+			}
+		}
+	}
+}
\ No newline at end of file
