PR #1413: Apply transforms correctly to LODs (Contributed by [REDACTED_PERSON])

diff --git a/Engine/Source/Editor/UnrealEd/Private/FbxMeshUtils.cpp b/Engine/Source/Editor/UnrealEd/Private/FbxMeshUtils.cpp
--- a/Engine/Source/Editor/UnrealEd/Private/FbxMeshUtils.cpp
+++ b/Engine/Source/Editor/UnrealEd/Private/FbxMeshUtils.cpp
@@ -81,6 +81,7 @@ namespace FbxMeshUtils
 
 		// don't import materials
 		UnFbx::FBXImportOptions* ImportOptions = FFbxImporter->GetImportOptions();
+		UFbxStaticMeshImportData* ImportData = Cast<UFbxStaticMeshImportData>(BaseStaticMesh->AssetImportData);
 		ImportOptions->bImportMaterials = false;
 		ImportOptions->bImportTextures = false;
 
@@ -93,6 +94,7 @@ namespace FbxMeshUtils
 		else
 		{
 			FFbxImporter->FlushToTokenizedErrorMessage(EMessageSeverity::Warning);
+			FFbxImporter->ApplyTransformSettingsToFbxNode(FFbxImporter->Scene->GetRootNode(), ImportData);
 
 			bool bUseLODs = true;
 			int32 MaxLODLevel = 0;
@@ -131,7 +133,7 @@ namespace FbxMeshUtils
 			{
 				// Import mesh
 				UStaticMesh* TempStaticMesh = NULL;
-				TempStaticMesh = (UStaticMesh*)FFbxImporter->ImportStaticMeshAsSingle(GetTransientPackage(), *(LODNodeList[bUseLODs? LODLevel: 0]), NAME_None, RF_NoFlags, NULL, BaseStaticMesh, LODLevel);
+				TempStaticMesh = (UStaticMesh*)FFbxImporter->ImportStaticMeshAsSingle(GetTransientPackage(), *(LODNodeList[bUseLODs? LODLevel: 0]), NAME_None, RF_NoFlags, ImportData, BaseStaticMesh, LODLevel);
 
 				// Add imported mesh to existing model
 				if( TempStaticMesh )
