Fix crash when retargeting an asset that has a NULL Skeleton.

diff --git a/Engine/Source/Editor/UnrealEd/Private/SSkeletonWidget.cpp b/Engine/Source/Editor/UnrealEd/Private/SSkeletonWidget.cpp
--- a/Engine/Source/Editor/UnrealEd/Private/SSkeletonWidget.cpp
+++ b/Engine/Source/Editor/UnrealEd/Private/SSkeletonWidget.cpp
@@ -410,6 +410,7 @@ void SAnimationRemapSkeleton::Construct( const FArguments& InArgs )
 	bConvertSpaces = false;
 	bShowOnlyCompatibleSkeletons = false;
 	OnRetargetAnimationDelegate = InArgs._OnRetargetDelegate;
+	bShowDuplicateAssetOption = InArgs._ShowDuplicateAssetOption;
 
 	TSharedRef<SVerticalBox> RetargetWidget = SNew(SVerticalBox);
 
@@ -479,7 +480,7 @@ void SAnimationRemapSkeleton::Construct( const FArguments& InArgs )
 			RetargetWidget
 		];
 
-	if (InArgs._ShowDuplicateAssetOption)
+	if (bShowDuplicateAssetOption)
 	{
 		TSharedRef<SVerticalBox> NameOptionWidget = SNew(SVerticalBox);
 
@@ -900,7 +901,7 @@ FReply SAnimationRemapSkeleton::OnApply()
 {
 	if (OnRetargetAnimationDelegate.IsBound())
 	{
-		OnRetargetAnimationDelegate.Execute(OldSkeleton, NewSkeleton, bRemapReferencedAssets, bConvertSpaces, &NameDuplicateRule);
+		OnRetargetAnimationDelegate.Execute(OldSkeleton, NewSkeleton, bRemapReferencedAssets, bConvertSpaces, bShowDuplicateAssetOption ? &NameDuplicateRule : nullptr);
 	}
 
 	CloseWindow();
diff --git a/Engine/Source/Editor/UnrealEd/Public/SSkeletonWidget.h b/Engine/Source/Editor/UnrealEd/Public/SSkeletonWidget.h
--- a/Engine/Source/Editor/UnrealEd/Public/SSkeletonWidget.h
+++ b/Engine/Source/Editor/UnrealEd/Public/SSkeletonWidget.h
@@ -471,6 +471,8 @@ public:
 	static UNREALED_API void ShowWindow(USkeleton* OldSkeleton, const FText& WarningMessage, bool bDuplicateAssets, FOnRetargetAnimation RetargetDelegate);
 
 	static TSharedPtr<SWindow> DialogWindow;
+
+	bool bShowDuplicateAssetOption;
 };
 
 ////////////////////////////////////////////////////
