Added new setting for localization import/export in order to persist comments in PO files that aren't generated from manifest/archive and are likely from a third party service like OneSky.

diff --git a/Engine/Source/Editor/Localization/Private/LocalizationConfigurationScript.cpp b/Engine/Source/Editor/Localization/Private/LocalizationConfigurationScript.cpp
--- a/Engine/Source/Editor/Localization/Private/LocalizationConfigurationScript.cpp
+++ b/Engine/Source/Editor/Localization/Private/LocalizationConfigurationScript.cpp
@@ -323,12 +323,12 @@ namespace LocalizationConfigurationScript
 				ConfigSection.Add( TEXT("CulturesToGenerate"), Target->Settings.SupportedCulturesStatistics[Index].CultureName );
 			};
 
-			// Export for a specific culture.
+			// Import for a specific culture.
 			if (CultureName.IsSet())
 			{
 				ConfigSection.Add( TEXT("CulturesToGenerate"), CultureName.GetValue() );
 			}
-			// Export for all cultures.
+			// Import for all cultures.
 			else
 			{
 				for (const FCultureStatistics& CultureStatistics : Target->Settings.SupportedCulturesStatistics)
@@ -395,6 +395,11 @@ namespace LocalizationConfigurationScript
 
 			ConfigSection.Add( TEXT("bExportLoc"), TEXT("true") );
 
+			if (Target->Settings.ExportSettings.ShouldPersistCommentsOnExport)
+			{
+				ConfigSection.Add(TEXT("ShouldPersistComments"), "true");
+			}
+
 			const FString SourcePath = ContentDirRelativeToGameDir / TEXT("Localization") / Target->Settings.Name;
 			ConfigSection.Add( TEXT("SourcePath"), SourcePath );
 
@@ -451,7 +456,6 @@ namespace LocalizationConfigurationScript
 				ConfigSection.Add( TEXT("bUseCultureDirectory"), "false" );
 			}
 
-
 			ConfigSection.Add( TEXT("ManifestName"), FPaths::GetCleanFilename(GetManifestPath(Target)) );
 			ConfigSection.Add( TEXT("ArchiveName"), FPaths::GetCleanFilename(GetArchivePath(Target, FString())) );
 			FString POFileName;
diff --git a/Engine/Source/Editor/Localization/Public/LocalizationTargetTypes.h b/Engine/Source/Editor/Localization/Public/LocalizationTargetTypes.h
--- a/Engine/Source/Editor/Localization/Public/LocalizationTargetTypes.h
+++ b/Engine/Source/Editor/Localization/Public/LocalizationTargetTypes.h
@@ -204,13 +204,23 @@ struct FGatherTextFromMetaDataConfiguration
 	UPROPERTY(config, EditAnywhere, Category = "MetaData")
 	TArray<FMetaDataKeyGatherSpecification> KeySpecifications;
 
-	/* If enable, data that is specified as editor-only may be processed for gathering. */
+	/* If enabled, data that is specified as editor-only may be processed for gathering. */
 	UPROPERTY(config, EditAnywhere, Category = "Filter")
 	bool ShouldGatherFromEditorOnlyData;
 
 	LOCALIZATION_API bool Validate(const FString& RootDirectory, FText& OutError) const;
 };
 
+USTRUCT()
+struct FLocalizationExportingSettings
+{
+	GENERATED_BODY()
+
+	/* Should user comments in existing PO files be persisted after export? Useful if using a third party service that stores editor/translator notes in the PO format's comment fields. */
+	UPROPERTY(config, EditAnywhere, Category = "Comments")
+	bool ShouldPersistCommentsOnExport;
+};
+
 USTRUCT()
 struct FCultureStatistics
 {
@@ -301,6 +311,10 @@ struct FLocalizationTargetSettings
 	UPROPERTY(config, EditAnywhere, Category = "Gather")
 	FGatherTextFromMetaDataConfiguration GatherFromMetaData;
 
+	/* Settings for import/export of translations. */
+	UPROPERTY(config, EditAnywhere, Category = "Import & Export")
+	FLocalizationExportingSettings ExportSettings;
+
 	/* The index of the native culture among the supported cultures. */
 	UPROPERTY(config, EditAnywhere, Category = "Cultures")
 	int32 NativeCultureIndex;
diff --git a/Engine/Source/Editor/UnrealEd/Classes/Commandlets/InternationalizationExportCommandlet.h b/Engine/Source/Editor/UnrealEd/Classes/Commandlets/InternationalizationExportCommandlet.h
--- a/Engine/Source/Editor/UnrealEd/Classes/Commandlets/InternationalizationExportCommandlet.h
+++ b/Engine/Source/Editor/UnrealEd/Classes/Commandlets/InternationalizationExportCommandlet.h
@@ -5,15 +5,31 @@
 #include "Commandlets/Commandlet.h"
 #include "InternationalizationExportCommandlet.generated.h"
 
+class FPortableObjectFormatDOM;
+
+struct FPortableObjectEntryIdentity
+{
+	FString MsgCtxt;
+	FString MsgId;
+	FString MsgIdPlural;
+};
+bool operator==(const FPortableObjectEntryIdentity& LHS, const FPortableObjectEntryIdentity& RHS);
+uint32 GetTypeHash(const FPortableObjectEntryIdentity& ID);
+
 /**
  *	UInternationalizationExportCommandlet: Commandlet used to export internationalization data to various standard formats. 
  */
 UCLASS()
 class UInternationalizationExportCommandlet : public UGatherTextCommandletBase
 {
-    GENERATED_UCLASS_BODY()
+    GENERATED_BODY()
 
 public:
+	UInternationalizationExportCommandlet(const FObjectInitializer& ObjectInitializer)
+		: Super(ObjectInitializer)
+		, HasPreservedComments(false)
+	{}
+
 	//~ Begin UCommandlet Interface
 	virtual int32 Main(const FString& Params) override;
 	//~ End UCommandlet Interface
@@ -21,8 +37,14 @@ private:
 	bool DoExport(const FString& SourcePath, const FString& DestinationPath, const FString& Filename);
 	bool DoImport(const FString& SourcePath, const FString& DestinationPath, const FString& Filename);
 
+	bool LoadPOFile(const FString& POFilePath, FPortableObjectFormatDOM& OutPortableObject);
+	void PreserveExtractedCommentsForPersistence(FPortableObjectFormatDOM& PortableObject);
+
 	TArray<FString> CulturesToGenerate;
 	FString ConfigPath;
 	FString SectionName;
 
+	bool ShouldPersistComments;
+	bool HasPreservedComments;
+	TMap< FPortableObjectEntryIdentity, TArray<FString> > POEntryToCommentMap;
 };
\ No newline at end of file
diff --git a/Engine/Source/Editor/UnrealEd/Private/Commandlets/InternationalizationExportCommandlet.cpp b/Engine/Source/Editor/UnrealEd/Private/Commandlets/InternationalizationExportCommandlet.cpp
--- a/Engine/Source/Editor/UnrealEd/Private/Commandlets/InternationalizationExportCommandlet.cpp
+++ b/Engine/Source/Editor/UnrealEd/Private/Commandlets/InternationalizationExportCommandlet.cpp
@@ -14,6 +14,18 @@ static const TCHAR* NewLineDelimiter = TEXT("\n");
 /**
 *	Helper Functions
 */
+bool operator==(const FPortableObjectEntryIdentity& LHS, const FPortableObjectEntryIdentity& RHS)
+{
+	return LHS.MsgCtxt == RHS.MsgCtxt && LHS.MsgId == RHS.MsgId && LHS.MsgIdPlural == RHS.MsgIdPlural;
+}
+
+uint32 GetTypeHash(const FPortableObjectEntryIdentity& ID)
+{
+	const uint32 HashA = HashCombine(GetTypeHash(ID.MsgCtxt), GetTypeHash(ID.MsgId));
+	const uint32 HashB = GetTypeHash(ID.MsgIdPlural);
+	return HashCombine(HashA, HashB);
+}
+
 namespace
 {
 	FString ConditionIdentityForPOMsgCtxt(const FString& Namespace, const FString& Key, const TSharedPtr<FLocMetadataObject>& KeyMetaData)
@@ -229,17 +241,26 @@ namespace
 	// @TODO: Note, we assume the source location format here but it could be arbitrary.
 	return InSrcLocation.Replace(TEXT(" - line "), TEXT(":"));
 	}
+
+	FString GetConditionedKeyForExtractedComment(const FString& Key)
+	{
+		return FString::Printf(TEXT("Key:\t%s"), *Key);
+	}
+
+	FString GetConditionedReferenceForExtractedComment(const FString& PORefString)
+	{
+		return FString::Printf(TEXT("SourceLocation:\t%s"), *PORefString);
+	}
+
+	FString GetConditionedInfoMetaDataForExtractedComment(const FString& KeyName, const FString& ValueString)
+	{
+		return FString::Printf(TEXT("InfoMetaData:\t\"%s\" : \"%s\""), *KeyName, *ValueString);
+	}
 }
 
 /**
 *	UInternationalizationExportCommandlet
 */
-UInternationalizationExportCommandlet::UInternationalizationExportCommandlet(const FObjectInitializer& ObjectInitializer)
-	: Super(ObjectInitializer)
-{
-}
-
-
 bool UInternationalizationExportCommandlet::DoExport( const FString& SourcePath, const FString& DestinationPath, const FString& Filename )
 {
 	// Get native culture.
@@ -273,7 +294,6 @@ bool UInternationalizationExportCommandlet::DoExport( const FString& SourcePath,
 		bUseCultureDirectory = true;
 	}
 
-
 	TSharedRef< FInternationalizationManifest > InternationalizationManifest = MakeShareable( new FInternationalizationManifest );
 	// Load the manifest info
 	{
@@ -334,17 +354,17 @@ bool UInternationalizationExportCommandlet::DoExport( const FString& SourcePath,
 			ArchiveSerializer.DeserializeArchive( ArchiveJsonObject.ToSharedRef(), InternationalizationArchive );
 
 			{
-				FPortableObjectFormatDOM PortableObj;
+				FPortableObjectFormatDOM NewPortableObject;
 
 				FString LocLang;
-				if( !PortableObj.SetLanguage( CultureName ) )
+				if( !NewPortableObject.SetLanguage( CultureName ) )
 				{
 					UE_LOG( LogInternationalizationExportCommandlet, Error, TEXT("Skipping export of loc language %s because it is not recognized."), *LocLang );
 					continue;
 				}
 
-				PortableObj.SetProjectName( FPaths::GetBaseFilename( ManifestName ) );
-				PortableObj.CreateNewHeader();
+				NewPortableObject.SetProjectName( FPaths::GetBaseFilename( ManifestName ) );
+				NewPortableObject.CreateNewHeader();
 
 				{
 					for(TManifestEntryBySourceTextContainer::TConstIterator ManifestIter = InternationalizationManifest->GetEntriesBySourceTextIterator(); ManifestIter; ++ManifestIter)
@@ -374,10 +394,10 @@ bool UInternationalizationExportCommandlet::DoExport( const FString& SourcePath,
 									PoEntry->MsgStr.Add( ConditionedArchiveTranslation );
 
 									const FString PORefString = ConvertSrcLocationToPORef( Context.SourceLocation );
-									PoEntry->AddReference( PORefString ); // Source location.
+									PoEntry->AddReference(PORefString); // Source location.
 
-									PoEntry->AddExtractedComment( FString::Printf(TEXT("Key:\t%s"), *Context.Key) ); // "Notes from Programmer" in the form of the Key.
-									PoEntry->AddExtractedComment( FString::Printf(TEXT("SourceLocation:\t%s"), *PORefString) ); // "Notes from Programmer" in the form of the Source Location, since this comes in handy too and OneSky doesn't properly show references, only comments.
+									PoEntry->AddExtractedComment( GetConditionedKeyForExtractedComment(Context.Key) ); // "Notes from Programmer" in the form of the Key.
+									PoEntry->AddExtractedComment( GetConditionedReferenceForExtractedComment(PORefString) ); // "Notes from Programmer" in the form of the Source Location, since this comes in handy too and OneSky doesn't properly show references, only comments.
 									TArray<FString> InfoMetaDataStrings;
 									if (Context.InfoMetadataObj.IsValid())
 									{
@@ -385,7 +405,7 @@ bool UInternationalizationExportCommandlet::DoExport( const FString& SourcePath,
 										{
 											const FString KeyName = InfoMetaDataPair.Key;
 											const TSharedPtr<FLocMetadataValue> Value = InfoMetaDataPair.Value;
-											InfoMetaDataStrings.Add(FString::Printf(TEXT("InfoMetaData:\t\"%s\" : \"%s\""), *KeyName, *Value->AsString()));
+											InfoMetaDataStrings.Add(GetConditionedInfoMetaDataForExtractedComment(KeyName, Value->AsString()));
 										}
 									}
 									if (InfoMetaDataStrings.Num())
@@ -393,7 +413,7 @@ bool UInternationalizationExportCommandlet::DoExport( const FString& SourcePath,
 										PoEntry->AddExtractedComments(InfoMetaDataStrings);
 									}
 
-									PortableObj.AddEntry( PoEntry );
+									NewPortableObject.AddEntry( PoEntry );
 								}
 							}
 
@@ -449,7 +469,7 @@ bool UInternationalizationExportCommandlet::DoExport( const FString& SourcePath,
 												PoEntry->AddExtractedComments(InfoMetaDataStrings);
 											}
 
-											PortableObj.AddEntry( PoEntry );
+											NewPortableObject.AddEntry( PoEntry );
 										}
 									}
 								}
@@ -460,9 +480,7 @@ bool UInternationalizationExportCommandlet::DoExport( const FString& SourcePath,
 
 				// Write out the Portable Object to .po file.
 				{
-					PortableObj.SortEntries();
-					FString OutputString = PortableObj.ToString();
-					FString OutputFileName = "";
+					FString OutputFileName;
 					if (bUseCultureDirectory)
 					{
 						OutputFileName = DestinationPath / CultureName / Filename;
@@ -472,6 +490,35 @@ bool UInternationalizationExportCommandlet::DoExport( const FString& SourcePath,
 						OutputFileName = DestinationPath / Filename;
 					}
 
+					// Persist comments if requested.
+					if (ShouldPersistComments)
+					{
+						// Preserve comments from the specified file now, if they haven't already been.
+						if (!HasPreservedComments)
+						{
+							FPortableObjectFormatDOM ExistingPortableObject;
+							const bool HasLoadedPOFile = LoadPOFile(OutputFileName, ExistingPortableObject);
+							if (!HasLoadedPOFile)
+							{
+								return false;
+							}
+
+							PreserveExtractedCommentsForPersistence(ExistingPortableObject);
+						}
+
+						// Persist the comments into the new portable object we're going to be saving.
+						for (const auto& Pair : POEntryToCommentMap)
+						{
+							const TSharedPtr<FPortableObjectEntry> FoundEntry = NewPortableObject.FindEntry(Pair.Key.MsgId, Pair.Key.MsgIdPlural, Pair.Key.MsgCtxt);
+							if (FoundEntry.IsValid())
+							{
+								FoundEntry->AddExtractedComments(Pair.Value);
+							}
+						}
+					}
+
+					NewPortableObject.SortEntries();
+
 					if( SourceControlInfo.IsValid() )
 					{
 						FText SCCErrorText;
@@ -483,7 +530,8 @@ bool UInternationalizationExportCommandlet::DoExport( const FString& SourcePath,
 					}
 
 					//@TODO We force UTF8 at the moment but we want this to be based on the format found in the header info.
-					if( !FFileHelper::SaveStringToFile(OutputString, *OutputFileName, FFileHelper::EEncodingOptions::ForceUTF8) )
+					const FString OutputString = NewPortableObject.ToString();
+					if (!FFileHelper::SaveStringToFile(OutputString, *OutputFileName, FFileHelper::EEncodingOptions::ForceUTF8))
 					{
 						UE_LOG( LogInternationalizationExportCommandlet, Error, TEXT("Could not write file %s"), *OutputFileName );
 						return false;
@@ -535,32 +583,23 @@ bool UInternationalizationExportCommandlet::DoImport(const FString& SourcePath,
 			POFilePath = SourcePath / Filename;
 		}
 
-		if( !FPaths::FileExists(POFilePath) )
-		{
-			UE_LOG( LogInternationalizationExportCommandlet, Warning, TEXT("Could not find file %s"), *POFilePath );
-			continue;
-		}
-
-		FString POFileContents;
-		if ( !FFileHelper::LoadFileToString( POFileContents, *POFilePath ) )
+		FPortableObjectFormatDOM PortableObject;
+		const bool HasLoadedPOFile = LoadPOFile(POFilePath, PortableObject);
+		if (!HasLoadedPOFile)
 		{
-			UE_LOG( LogInternationalizationExportCommandlet, Error, TEXT("Failed to load file %s."), *POFilePath);
 			continue;
 		}
 
-		FPortableObjectFormatDOM PortableObject;
-		if( !PortableObject.FromString( POFileContents ) )
+		if (ShouldPersistComments)
 		{
-			UE_LOG( LogInternationalizationExportCommandlet, Error, TEXT("Failed to parse Portable Object file %s."), *POFilePath);
-			continue;
+			PreserveExtractedCommentsForPersistence(PortableObject);
 		}
 
-		if( PortableObject.GetProjectName() != ManifestName.Replace(TEXT(".manifest"), TEXT("")) )
+		if (PortableObject.GetProjectName() != ManifestName.Replace(TEXT(".manifest"), TEXT("")))
 		{
-			UE_LOG( LogInternationalizationExportCommandlet, Warning, TEXT("The project name (%s) in the file (%s) did not match the target manifest project (%s)."), *POFilePath, *PortableObject.GetProjectName(), *ManifestName.Replace(TEXT(".manifest"), TEXT("")));
+			UE_LOG(LogInternationalizationExportCommandlet, Warning, TEXT("The project name (%s) in the file (%s) did not match the target manifest project (%s)."), *POFilePath, *PortableObject.GetProjectName(), *ManifestName.Replace(TEXT(".manifest"), TEXT("")));
 		}
 
-
 		const FString ManifestFileName = DestinationPath / ManifestName;
 
 		TSharedPtr< FJsonObject > ManifestJsonObject = NULL;
@@ -737,9 +776,18 @@ int32 UInternationalizationExportCommandlet::Main( const FString& Params )
 
 	bool bDoExport = false;
 	bool bDoImport = false;
+	ShouldPersistComments = false;
 
 	GetBoolFromConfig( *SectionName, TEXT("bImportLoc"), bDoImport, ConfigPath );
 	GetBoolFromConfig( *SectionName, TEXT("bExportLoc"), bDoExport, ConfigPath );
+	GetBoolFromConfig(*SectionName, TEXT("ShouldPersistComments"), ShouldPersistComments, ConfigPath);
+
+	// Reject the ShouldPersistComments flag and warn if not exporting - we're not writing to anything, so we can't persist.
+	if (ShouldPersistComments && !bDoExport)
+	{
+		UE_LOG(LogInternationalizationExportCommandlet, Warning, TEXT("ShouldPersistComments is true, but bExportLoc is false - can't persist comments if not writing PO files."));
+		ShouldPersistComments = false;
+	}
 
 	if( !bDoImport && !bDoExport )
 	{
@@ -768,3 +816,47 @@ int32 UInternationalizationExportCommandlet::Main( const FString& Params )
 	return 0;
 }
 
+bool UInternationalizationExportCommandlet::LoadPOFile(const FString& POFilePath, FPortableObjectFormatDOM& OutPortableObject)
+{
+	if (!FPaths::FileExists(POFilePath))
+	{
+		UE_LOG(LogInternationalizationExportCommandlet, Warning, TEXT("Could not find file %s"), *POFilePath);
+		return false;
+	}
+
+	FString POFileContents;
+	if (!FFileHelper::LoadFileToString(POFileContents, *POFilePath))
+	{
+		UE_LOG(LogInternationalizationExportCommandlet, Error, TEXT("Failed to load file %s."), *POFilePath);
+		return false;
+	}
+
+	if (!OutPortableObject.FromString(POFileContents))
+	{
+		UE_LOG(LogInternationalizationExportCommandlet, Error, TEXT("Failed to parse Portable Object file %s."), *POFilePath);
+		return false;
+	}
+
+	return true;
+}
+
+void UInternationalizationExportCommandlet::PreserveExtractedCommentsForPersistence(FPortableObjectFormatDOM& PortableObject)
+{
+	// Preserve comments for later.
+	for (auto EntriesIterator = PortableObject.GetEntriesIterator(); EntriesIterator; ++EntriesIterator)
+	{
+		const TSharedPtr< FPortableObjectEntry >& Entry = *EntriesIterator;
+
+		// Preserve only non-procedurally generated extracted comments.
+		const TArray<FString> CommentsToPreserve = Entry->ExtractedComments.FilterByPredicate([=](const FString& ExtractedComment) -> bool
+		{
+			return !ExtractedComment.StartsWith("Key:") && !ExtractedComment.StartsWith("SourceLocation:") && !ExtractedComment.StartsWith("InfoMetaData:");
+		});
+
+		if (CommentsToPreserve.Num())
+		{
+			POEntryToCommentMap.Add(FPortableObjectEntryIdentity{ Entry->MsgCtxt, Entry->MsgId, Entry->MsgIdPlural }, CommentsToPreserve);
+		}
+	}
+	HasPreservedComments = true;
+}
\ No newline at end of file
