Sequencer: Fix auto frame so that it happens regardless of the curve visibility setting.

diff --git a/Engine/Source/Editor/Sequencer/Private/SSequencerCurveEditor.cpp b/Engine/Source/Editor/Sequencer/Private/SSequencerCurveEditor.cpp
--- a/Engine/Source/Editor/Sequencer/Private/SSequencerCurveEditor.cpp
+++ b/Engine/Source/Editor/Sequencer/Private/SSequencerCurveEditor.cpp
@@ -173,6 +173,11 @@ void SSequencerCurveEditor::NodeTreeSelectionChanged()
 			UpdateCurveOwner();
 		}
 
+		if (GetAllowAutoFrame())
+		{
+			ZoomToFit();
+		}
+
 		UpdateCurveViewModelSelection();
 	}
 }
@@ -191,6 +196,30 @@ void SSequencerCurveEditor::OnCurveEditorCurveVisibilityChanged()
 	UpdateCurveOwner();
 }
 
+TArray<FRichCurve*> SSequencerCurveEditor::GetCurvesToFit() const
+{
+	TArray<FRichCurve*> FitCurves;
+
+	for(auto CurveViewModel : CurveViewModels)
+	{
+		if (CurveViewModel->bIsVisible)
+		{
+			if (CurveOwner->GetSelectedCurves().Contains(CurveViewModel->CurveInfo.CurveToEdit))
+			{
+				FitCurves.Add(CurveViewModel->CurveInfo.CurveToEdit);
+			}
+		}
+	}
+
+	if (FitCurves.Num() > 0)
+	{
+		return FitCurves;
+	}
+
+	return SCurveEditor::GetCurvesToFit();
+}
+
+
 SSequencerCurveEditor::~SSequencerCurveEditor()
 {
 	if ( Sequencer.IsValid() )
diff --git a/Engine/Source/Editor/Sequencer/Private/SSequencerCurveEditor.h b/Engine/Source/Editor/Sequencer/Private/SSequencerCurveEditor.h
--- a/Engine/Source/Editor/Sequencer/Private/SSequencerCurveEditor.h
+++ b/Engine/Source/Editor/Sequencer/Private/SSequencerCurveEditor.h
@@ -56,6 +56,9 @@ private:
 	/** Update the curve view model selection */
 	void UpdateCurveViewModelSelection();
 
+	/* Get the curves to will be used during a fit operation */
+	virtual TArray<FRichCurve*> GetCurvesToFit()const override;
+
 private:
 
 	/** The sequencer which owns this widget. */
diff --git a/Engine/Source/Editor/UnrealEd/Public/SCurveEditor.h b/Engine/Source/Editor/UnrealEd/Public/SCurveEditor.h
--- a/Engine/Source/Editor/UnrealEd/Public/SCurveEditor.h
+++ b/Engine/Source/Editor/UnrealEd/Public/SCurveEditor.h
@@ -240,6 +240,14 @@ public:
 	UNREALED_API bool GetAllowAutoFrame() { return bAllowAutoFrame; }
 	UNREALED_API void SetAllowAutoFrame(bool bInAllowAutoFrame) { bAllowAutoFrame = bInAllowAutoFrame; }
 
+	/* Get the curves to will be used during a fit operation */
+	UNREALED_API virtual TArray<FRichCurve*> GetCurvesToFit()const;
+
+	/** Zoom to fit */
+	UNREALED_API void ZoomToFitHorizontal();
+	UNREALED_API void ZoomToFitVertical();
+	UNREALED_API void ZoomToFit();
+
 private:
 	/** Used to track a key and the curve that owns it */
 	struct FSelectedCurveKey
@@ -372,13 +380,6 @@ private:
 	/** Function to check whether the current track is editable */
 	bool IsEditingEnabled() const;
 
-	/* Get the curves to will be used during a fit operation */
-	TArray<FRichCurve*> GetCurvesToFit()const;
-
-	void ZoomToFitHorizontal();
-	void ZoomToFitVertical();
-	void ZoomToFit();
-
 	FReply ZoomToFitHorizontalClicked();
 	FReply ZoomToFitVerticalClicked();
 
