Added ability to opt out of adding reference comments as extracted comments when exporting localization to PO files.

diff --git a/Engine/Source/Editor/Localization/Private/LocalizationConfigurationScript.cpp b/Engine/Source/Editor/Localization/Private/LocalizationConfigurationScript.cpp
--- a/Engine/Source/Editor/Localization/Private/LocalizationConfigurationScript.cpp
+++ b/Engine/Source/Editor/Localization/Private/LocalizationConfigurationScript.cpp
@@ -393,11 +393,12 @@ namespace LocalizationConfigurationScript
 			// CommandletClass
 			ConfigSection.Add( TEXT("CommandletClass"), TEXT("InternationalizationExport") );
 
-			ConfigSection.Add( TEXT("bExportLoc"), TEXT("true") );
+			ConfigSection.Add(TEXT("bExportLoc"), TEXT("true"));
 
-			if (Target->Settings.ExportSettings.ShouldPersistCommentsOnExport)
+			// Export-specific settings.
 			{
-				ConfigSection.Add(TEXT("ShouldPersistComments"), "true");
+				ConfigSection.Add(TEXT("ShouldPersistCommentsOnExport"), Target->Settings.ExportSettings.ShouldPersistCommentsOnExport ? TEXT("true") : TEXT("false"));
+				ConfigSection.Add(TEXT("ShouldAddSourceLocationsAsComments"), Target->Settings.ExportSettings.ShouldAddSourceLocationsAsComments ? TEXT("true") : TEXT("false"));
 			}
 
 			const FString SourcePath = ContentDirRelativeToGameDir / TEXT("Localization") / Target->Settings.Name;
diff --git a/Engine/Source/Editor/Localization/Public/LocalizationTargetTypes.h b/Engine/Source/Editor/Localization/Public/LocalizationTargetTypes.h
--- a/Engine/Source/Editor/Localization/Public/LocalizationTargetTypes.h
+++ b/Engine/Source/Editor/Localization/Public/LocalizationTargetTypes.h
@@ -216,9 +216,18 @@ struct FLocalizationExportingSettings
 {
 	GENERATED_BODY()
 
+	FLocalizationExportingSettings()
+	: ShouldPersistCommentsOnExport(false)
+	, ShouldAddSourceLocationsAsComments(true)
+	{}
+
 	/* Should user comments in existing PO files be persisted after export? Useful if using a third party service that stores editor/translator notes in the PO format's comment fields. */
 	UPROPERTY(config, EditAnywhere, Category = "Comments")
 	bool ShouldPersistCommentsOnExport;
+
+	/* Should source locations be added to PO file entries as comments? Useful if a third party service doesn't expose PO file reference comments, which typically store the source location. */
+	UPROPERTY(config, EditAnywhere, Category = "Comments")
+	bool ShouldAddSourceLocationsAsComments;
 };
 
 USTRUCT()
diff --git a/Engine/Source/Editor/UnrealEd/Private/Commandlets/InternationalizationExportCommandlet.cpp b/Engine/Source/Editor/UnrealEd/Private/Commandlets/InternationalizationExportCommandlet.cpp
--- a/Engine/Source/Editor/UnrealEd/Private/Commandlets/InternationalizationExportCommandlet.cpp
+++ b/Engine/Source/Editor/UnrealEd/Private/Commandlets/InternationalizationExportCommandlet.cpp
@@ -294,6 +294,9 @@ bool UInternationalizationExportCommandlet::DoExport( const FString& SourcePath,
 		bUseCultureDirectory = true;
 	}
 
+	bool ShouldAddSourceLocationsAsComments = true;
+	GetBoolFromConfig(*SectionName, TEXT("ShouldAddSourceLocationsAsComments"), ShouldAddSourceLocationsAsComments, ConfigPath);
+
 	TSharedRef< FInternationalizationManifest > InternationalizationManifest = MakeShareable( new FInternationalizationManifest );
 	// Load the manifest info
 	{
@@ -397,7 +400,12 @@ bool UInternationalizationExportCommandlet::DoExport( const FString& SourcePath,
 									PoEntry->AddReference(PORefString); // Source location.
 
 									PoEntry->AddExtractedComment( GetConditionedKeyForExtractedComment(Context.Key) ); // "Notes from Programmer" in the form of the Key.
-									PoEntry->AddExtractedComment( GetConditionedReferenceForExtractedComment(PORefString) ); // "Notes from Programmer" in the form of the Source Location, since this comes in handy too and OneSky doesn't properly show references, only comments.
+
+									if (ShouldAddSourceLocationsAsComments)
+									{
+										PoEntry->AddExtractedComment(GetConditionedReferenceForExtractedComment(PORefString)); // "Notes from Programmer" in the form of the Source Location, since this comes in handy too and OneSky doesn't properly show references, only comments.
+									}
+
 									TArray<FString> InfoMetaDataStrings;
 									if (Context.InfoMetadataObj.IsValid())
 									{
