Fixed unsafe array index access in Portable Object entry sorting logic.

Integration of CL# 2713824.

diff --git a/Engine/Source/Editor/UnrealEd/Private/Localization/PortableObjectFormatDOM.cpp b/Engine/Source/Editor/UnrealEd/Private/Localization/PortableObjectFormatDOM.cpp
--- a/Engine/Source/Editor/UnrealEd/Private/Localization/PortableObjectFormatDOM.cpp
+++ b/Engine/Source/Editor/UnrealEd/Private/Localization/PortableObjectFormatDOM.cpp
@@ -832,11 +832,26 @@ void FPortableObjectFormatDOM::SortEntries()
 		const int32 ExtractedCommentCount = FMath::Max(A->ExtractedComments.Num(), B->ExtractedComments.Num());
 		for (int32 i = 0; i < ExtractedCommentCount; ++i)
 		{
-			if (A->ExtractedComments[i] < B->ExtractedComments[i] || (!A->ExtractedComments.IsValidIndex(i) && B->ExtractedComments.IsValidIndex(i)))
+			// If A has no more comments, it is before B.
+			if (!A->ExtractedComments.IsValidIndex(i) && B->ExtractedComments.IsValidIndex(i))
 			{
 				return true;
 			}
-			else if (A->ExtractedComments[i] > B->ExtractedComments[i] || (A->ExtractedComments.IsValidIndex(i) && !B->ExtractedComments.IsValidIndex(i)))
+			// If B has no more comments, it is before A.
+			if (A->ExtractedComments.IsValidIndex(i) && !B->ExtractedComments.IsValidIndex(i))
+			{
+				return false;
+			}
+
+			check(A->ExtractedComments.IsValidIndex(i) && B->ExtractedComments.IsValidIndex(i));
+
+			// If A's comment is lexicographically less, it is before B.
+			if (A->ExtractedComments[i] < B->ExtractedComments[i])
+			{
+				return true;
+			}
+			// If B's comment is lexicographically less, it is before A.
+			if (A->ExtractedComments[i] > B->ExtractedComments[i])
 			{
 				return false;
 			}
