Sequencer: Followup fix for autoframing on curves in the sequencer.

diff --git a/Engine/Source/Editor/Sequencer/Private/SSequencerCurveEditor.cpp b/Engine/Source/Editor/Sequencer/Private/SSequencerCurveEditor.cpp
--- a/Engine/Source/Editor/Sequencer/Private/SSequencerCurveEditor.cpp
+++ b/Engine/Source/Editor/Sequencer/Private/SSequencerCurveEditor.cpp
@@ -173,7 +173,7 @@ void SSequencerCurveEditor::NodeTreeSelectionChanged()
 			UpdateCurveOwner();
 		}
 
-		if (GetAllowAutoFrame())
+		if (GetAutoFrame())
 		{
 			ZoomToFit();
 		}
diff --git a/Engine/Source/Editor/UnrealEd/Private/SCurveEditor.cpp b/Engine/Source/Editor/UnrealEd/Private/SCurveEditor.cpp
--- a/Engine/Source/Editor/UnrealEd/Private/SCurveEditor.cpp
+++ b/Engine/Source/Editor/UnrealEd/Private/SCurveEditor.cpp
@@ -1173,7 +1173,7 @@ void SCurveEditor::SetCurveOwner(FCurveOwnerInterface* InCurveOwner, bool bCanEd
 
 	ValidateSelection();
 
-	if (Settings->GetAutoFrameCurveEditor() && GetAllowAutoFrame())
+	if (GetAutoFrame())
 	{
 		if( bZoomToFitVertical )
 		{
@@ -2286,6 +2286,11 @@ void SCurveEditor::ValidateSelection()
 	}
 }
 
+bool SCurveEditor::GetAutoFrame() const
+{
+	return Settings->GetAutoFrameCurveEditor() && GetAllowAutoFrame();
+}
+
 TArray<FRichCurve*> SCurveEditor::GetCurvesToFit() const
 {
 	TArray<FRichCurve*> FitCurves;
diff --git a/Engine/Source/Editor/UnrealEd/Public/SCurveEditor.h b/Engine/Source/Editor/UnrealEd/Public/SCurveEditor.h
--- a/Engine/Source/Editor/UnrealEd/Public/SCurveEditor.h
+++ b/Engine/Source/Editor/UnrealEd/Public/SCurveEditor.h
@@ -237,9 +237,12 @@ public:
 	UNREALED_API TSharedPtr<FUICommandList> GetCommands();
 
 	/** Gets or sets whether autoframing is allowed */
-	UNREALED_API bool GetAllowAutoFrame() { return bAllowAutoFrame; }
+	UNREALED_API bool GetAllowAutoFrame() const { return bAllowAutoFrame; }
 	UNREALED_API void SetAllowAutoFrame(bool bInAllowAutoFrame) { bAllowAutoFrame = bInAllowAutoFrame; }
 
+	/** Gets whether autoframe will be invoked (combination of allow auto frame and curve editor auto frame setting) */
+	UNREALED_API bool GetAutoFrame() const;
+
 	/* Get the curves to will be used during a fit operation */
 	UNREALED_API virtual TArray<FRichCurve*> GetCurvesToFit()const;
 
