Crash fixed when doing a "Find Reference" on local variables without any placed nodes for the variable.

-------------
Merging

//UE4/Dev-Blueprints/Engine/Source/Editor/UnrealEd/Private/Kismet2/BlueprintEditorUtils.cpp

to //UE4/Main/Engine/Source/Editor/UnrealEd/Private/Kismet2/BlueprintEditorUtils.cpp

diff --git a/Engine/Source/Editor/UnrealEd/Private/Kismet2/BlueprintEditorUtils.cpp b/Engine/Source/Editor/UnrealEd/Private/Kismet2/BlueprintEditorUtils.cpp
--- a/Engine/Source/Editor/UnrealEd/Private/Kismet2/BlueprintEditorUtils.cpp
+++ b/Engine/Source/Editor/UnrealEd/Private/Kismet2/BlueprintEditorUtils.cpp
@@ -5009,7 +5009,9 @@ void FBlueprintEditorUtils::ChangeLocalVariableType(UBlueprint* InBlueprint, con
 				FBlueprintEditorUtils::MarkBlueprintAsStructurallyModified(InBlueprint);
 
 				TSharedPtr<IToolkit> FoundAssetEditor = FToolkitManager::Get().FindEditorForAsset(InBlueprint);
-				if (FoundAssetEditor.IsValid())
+
+				// No need to submit a search query if there are no nodes.
+				if (FoundAssetEditor.IsValid() && VariableNodes.Num())
 				{
 					TSharedRef<IBlueprintEditor> BlueprintEditor = StaticCastSharedRef<IBlueprintEditor>(FoundAssetEditor.ToSharedRef());
 
