Fixed issue with ddc commandlet killing texture source data then trying to cache again.

diff --git a/Engine/Source/Editor/UnrealEd/Private/Commandlets/DerivedDataCacheCommandlet.cpp b/Engine/Source/Editor/UnrealEd/Private/Commandlets/DerivedDataCacheCommandlet.cpp
--- a/Engine/Source/Editor/UnrealEd/Private/Commandlets/DerivedDataCacheCommandlet.cpp
+++ b/Engine/Source/Editor/UnrealEd/Private/Commandlets/DerivedDataCacheCommandlet.cpp
@@ -226,9 +226,12 @@ int32 UDerivedDataCacheCommandlet::Main( const FString& Params )
 				// cache all the resources for this platform
 				for ( TObjectIterator<UObject> It; It; ++It )
 				{
-					for ( auto Platform : Platforms )
+					if (ProcessedPackages.Contains(It->GetOutermost()->GetName()))
 					{
-						It->BeginCacheForCookedPlatformData( Platform );
+						for (auto Platform : Platforms)
+						{
+							It->BeginCacheForCookedPlatformData(Platform);
+						}
 					}
 				}
 
