Fix for crash when accessing stale function in a ParentFuncMap

diff --git a/Engine/Source/Editor/UnrealEd/Private/Kismet2/KismetReinstanceUtilities.cpp b/Engine/Source/Editor/UnrealEd/Private/Kismet2/KismetReinstanceUtilities.cpp
--- a/Engine/Source/Editor/UnrealEd/Private/Kismet2/KismetReinstanceUtilities.cpp
+++ b/Engine/Source/Editor/UnrealEd/Private/Kismet2/KismetReinstanceUtilities.cpp
@@ -645,6 +645,8 @@ void FBlueprintCompileReinstancer::UpdateBytecodeReferences()
 				continue;
 			}
 
+			BPClass->ClearFunctionMapsCaches();
+
 			bool bBPWasChanged = false;
 			// For each function defined in this blueprint, run through the bytecode, and update any refs from the old properties to the new
 			for( TFieldIterator<UFunction> FuncIter(BPClass, EFieldIteratorFlags::ExcludeSuper); FuncIter; ++FuncIter )
