[UE-22013] - Remove Blueprint compiler note/warning from custom event nodes bound to delegates with 'const&' and TArray input parameters.

change summary:
- added a new object version for fixup during serialization (VER_UE4_BLUEPRINT_CUSTOM_EVENT_CONST_INPUT)
- modified FBlueprintGraphArgumentLayout::OnRefCheckStateChanged() to implicitly set PinType.bIsConst 'true' for Custom Event nodes when PinType.bIsArray or PinType.bIsReference is also 'true'; this mirrors how delegate signatures are handled by UHT
- modified UK2Node_CustomEvent::Serialize() to fix up PinType.bIsConst for array and non-array reference parameter types
- in BlueprintEditorUtils.h, added FCompilerRelevantNodeLink and updated array typedef for the GetCompilerRelevantNodes() API (expanded out to include which pin it's linked to in addition to the owning node)
- modified the FBlueprintEditorUtils::GetCompilerRelevantNodes() API (renamed) to track linked pins in addition to owning nodes in the result. also modified FindFirstCompilerRelevantNode() to match.
- added FBlueprintEditorUtils::FindFirstCompilerRelevantLinkedPin() to allow queries for the linked pin instead of the compiler-relevant node itself. this is now used by the Custom Event ReconstructNode() logic to "pass-through" reroute nodes on delegate pin connections to find the "correct" target pin
- modified UK2Node_CustomEvent::ReconstructNode() to properly handle the case when the delegate pin might be connected to a reroute node (in which case the LinkedTo pin might resolve to the Custom Event's previously-compiled function, rather than the delegate signature function it's bound to)

notes:
- also addresses UE-20613

diff --git a/Engine/Source/Editor/BlueprintGraph/Private/K2Node_CustomEvent.cpp b/Engine/Source/Editor/BlueprintGraph/Private/K2Node_CustomEvent.cpp
--- a/Engine/Source/Editor/BlueprintGraph/Private/K2Node_CustomEvent.cpp
+++ b/Engine/Source/Editor/BlueprintGraph/Private/K2Node_CustomEvent.cpp
@@ -114,6 +114,44 @@ void UK2Node_CustomEvent::Serialize(FArchive& Ar)
 	if (Ar.IsLoading())
 	{
 		CachedNodeTitle.MarkDirty();
+
+		if (Ar.UE4Ver() < VER_UE4_BLUEPRINT_CUSTOM_EVENT_CONST_INPUT)
+		{
+			const UEdGraphSchema_K2* K2Schema = GetDefault<UEdGraphSchema_K2>();
+
+			// Ensure that array type inputs and non-array type pass-by-reference inputs are also marked 'const' - since arrays
+			// are implicitly passed by reference, and since events do not have return values/outputs, this equates to marking
+			// the parameter as 'const Type&' in native code. Also note that since UHT already blocks non-const reference types
+			// from being compiled into a MC delegate signature, any existing custom event param pins that were implicitly
+			// created via "Assign" in the Blueprint editor's context menu will previously have had 'const' set for its pin type.
+			//
+			// This ensures that (a) we don't emit the "no reference will be returned" note on custom event nodes with array
+			// inputs added by the user via the Details panel, and (b) we don't emit the "no reference will be returned" warning
+			// on custom event nodes with struct/object inputs added by the user in the Details tab and also set to pass-by-reference.
+			//
+			// Blueprint details customization now sets 'bIsConst' in new custom event node placements - see OnRefCheckStateChanged().
+			for (auto Pin : Pins)
+			{
+				if (Pin)
+				{
+					Ar.Preload(Pin);
+
+					if (Pin->Direction == EGPD_Output
+						&& !K2Schema->IsExecPin(*Pin)
+						&& !K2Schema->IsDelegateCategory(Pin->PinType.PinCategory))
+					{
+						for (auto PinInfo : UserDefinedPins)
+						{
+							if (PinInfo->PinName == Pin->PinName)
+							{
+								Pin->PinType.bIsConst = PinInfo->PinType.bIsConst = PinInfo->PinType.bIsArray || PinInfo->PinType.bIsReference;
+								break;
+							}
+						}
+					}
+				}
+			}
+		}
 	}
 }
 
@@ -301,7 +339,7 @@ void UK2Node_CustomEvent::ReconstructNode()
 	const UEdGraphSchema_K2* K2Schema = GetDefault<UEdGraphSchema_K2>();
 
 	const UEdGraphPin* DelegateOutPin = FindPin(DelegateOutputName);
-	const UEdGraphPin* LinkedPin = ( DelegateOutPin && DelegateOutPin->LinkedTo.Num() && DelegateOutPin->LinkedTo[0] ) ? DelegateOutPin->LinkedTo[0] : nullptr;
+	const UEdGraphPin* LinkedPin = ( DelegateOutPin && DelegateOutPin->LinkedTo.Num() && DelegateOutPin->LinkedTo[0] ) ? FBlueprintEditorUtils::FindFirstCompilerRelevantLinkedPin(DelegateOutPin->LinkedTo[0]) : nullptr;
 
 	const UFunction* DelegateSignature = nullptr;
 
diff --git a/Engine/Source/Editor/Kismet/Private/BlueprintDetailsCustomization.cpp b/Engine/Source/Editor/Kismet/Private/BlueprintDetailsCustomization.cpp
--- a/Engine/Source/Editor/Kismet/Private/BlueprintDetailsCustomization.cpp
+++ b/Engine/Source/Editor/Kismet/Private/BlueprintDetailsCustomization.cpp
@@ -2365,6 +2365,7 @@ void FBlueprintGraphArgumentLayout::OnRefCheckStateChanged(ECheckBoxState InStat
 {
 	FEdGraphPinType PinType = OnGetPinInfo();
 	PinType.bIsReference = (InState == ECheckBoxState::Checked)? true : false;
+	PinType.bIsConst = (PinType.bIsArray || PinType.bIsReference) && TargetNode && TargetNode->IsA<UK2Node_CustomEvent>();
 	PinInfoChanged(PinType);
 }
 
diff --git a/Engine/Source/Editor/UnrealEd/Private/Kismet2/BlueprintEditorUtils.cpp b/Engine/Source/Editor/UnrealEd/Private/Kismet2/BlueprintEditorUtils.cpp
--- a/Engine/Source/Editor/UnrealEd/Private/Kismet2/BlueprintEditorUtils.cpp
+++ b/Engine/Source/Editor/UnrealEd/Private/Kismet2/BlueprintEditorUtils.cpp
@@ -739,7 +739,7 @@ void FBlueprintEditorUtils::GetAllGraphNames(const UBlueprint* Blueprint, TArray
 	}
 }
 
-void FBlueprintEditorUtils::GetCompilerRelevantNodes(const UEdGraphPin* FromPin, FCompilerRelevantNodesArray& OutNodes)
+void FBlueprintEditorUtils::GetCompilerRelevantNodeLinks(UEdGraphPin* FromPin, FCompilerRelevantNodeLinkArray& OutNodeLinks)
 {
 	if(FromPin)
 	{
@@ -757,24 +757,32 @@ void FBlueprintEditorUtils::GetCompilerRelevantNodes(const UEdGraphPin* FromPin,
 					// Recursively check each link for a compiler-relevant node that will "pass through" this node at compile time
 					for(auto LinkedPin : FromPin->LinkedTo)
 					{
-						GetCompilerRelevantNodes(LinkedPin, OutNodes);
+						GetCompilerRelevantNodeLinks(LinkedPin, OutNodeLinks);
 					}
 				}
 			}
 			else
 			{
-				OutNodes.Add(OwningNode);
+				OutNodeLinks.Add(FCompilerRelevantNodeLink(OwningNode, FromPin));
 			}
 		}		
 	}
 }
 
-UK2Node* FBlueprintEditorUtils::FindFirstCompilerRelevantNode(const UEdGraphPin* FromPin)
+UK2Node* FBlueprintEditorUtils::FindFirstCompilerRelevantNode(UEdGraphPin* FromPin)
 {
-	FCompilerRelevantNodesArray RelevantNodes;
-	GetCompilerRelevantNodes(FromPin, RelevantNodes);
-	
-	return RelevantNodes.Num() > 0 ? RelevantNodes[0] : nullptr;
+	FCompilerRelevantNodeLinkArray RelevantNodeLinks;
+	GetCompilerRelevantNodeLinks(FromPin, RelevantNodeLinks);
+
+	return RelevantNodeLinks.Num() > 0 ? RelevantNodeLinks[0].Node : nullptr;
+}
+
+UEdGraphPin* FBlueprintEditorUtils::FindFirstCompilerRelevantLinkedPin(UEdGraphPin* FromPin)
+{
+	FCompilerRelevantNodeLinkArray RelevantNodeLinks;
+	GetCompilerRelevantNodeLinks(FromPin, RelevantNodeLinks);
+
+	return RelevantNodeLinks.Num() > 0 ? RelevantNodeLinks[0].LinkedPin : nullptr;
 }
 
 /** 
diff --git a/Engine/Source/Editor/UnrealEd/Public/Kismet2/BlueprintEditorUtils.h b/Engine/Source/Editor/UnrealEd/Public/Kismet2/BlueprintEditorUtils.h
--- a/Engine/Source/Editor/UnrealEd/Public/Kismet2/BlueprintEditorUtils.h
+++ b/Engine/Source/Editor/UnrealEd/Public/Kismet2/BlueprintEditorUtils.h
@@ -94,8 +94,20 @@ private:
 
 DECLARE_CYCLE_STAT_EXTERN(TEXT("Notify Blueprint Changed"), EKismetCompilerStats_NotifyBlueprintChanged, STATGROUP_KismetCompiler, );
 
-/** Array type for GetCompilerRelevantNodes()                                                                                                                                                                                                                                                                                                                                                                                                                                                                                                                                                                                 */
-typedef TArray<UK2Node*, TInlineAllocator<4> > FCompilerRelevantNodesArray;
+struct UNREALED_API FCompilerRelevantNodeLink
+{
+	UK2Node* Node;
+	UEdGraphPin* LinkedPin;
+
+	FCompilerRelevantNodeLink(UK2Node* InNode, UEdGraphPin* InLinkedPin)
+		: Node(InNode)
+		, LinkedPin(InLinkedPin)
+	{
+	}
+};
+
+/** Array type for GetCompilerRelevantNodeLinks() */
+typedef TArray<FCompilerRelevantNodeLink, TInlineAllocator<4> > FCompilerRelevantNodeLinkArray;
 
 class UNREALED_API FBlueprintEditorUtils
 {
@@ -519,12 +531,12 @@ public:
 	static void GetAllGraphNames(const UBlueprint* Blueprint, TArray<FName>& GraphNames);
 
 	/**
-	 * Gets the compiler-relevant (i.e. non-ignorable) nodes from the given pin.
+	 * Gets the compiler-relevant (i.e. non-ignorable) node links from the given pin.
 	 *
 	 * @param			FromPin			The pin to start searching from.
-	 * @param			OutNodes		Will contain the given pin's owning node if compiler-relevant, or all nodes linked to the owning node at the matching "pass-through" pin that are compiler-relevant. Empty if no compiler-relevant nodes can be found from the given pin.
+	 * @param			OutNodeLinks	Will contain the given pin + owning node if compiler-relevant, or all nodes linked to the owning node at the matching "pass-through" pin that are compiler-relevant. Empty if no compiler-relevant node links can be found from the given pin.
 	 */
-	static void GetCompilerRelevantNodes(const UEdGraphPin* FromPin, FCompilerRelevantNodesArray& OutNodes);
+	static void GetCompilerRelevantNodeLinks(UEdGraphPin* FromPin, FCompilerRelevantNodeLinkArray& OutNodeLinks);
 
 	/**
 	 * Finds the first compiler-relevant (i.e. non-ignorable) node from the given pin.
@@ -533,8 +545,16 @@ public:
 	 *
 	 * @return			The given pin's owning node if compiler-relevant, or the first node linked to the owning node at the matching "pass-through" pin that is compiler-relevant. May be NULL if no compiler-relevant nodes can be found from the given pin.
 	 */
-	static UK2Node* FindFirstCompilerRelevantNode(const UEdGraphPin* FromPin);
+	static UK2Node* FindFirstCompilerRelevantNode(UEdGraphPin* FromPin);
 
+	/**
+	 * Finds the first compiler-relevant (i.e. non-ignorable) node from the given pin and returns the owned pin that's linked.
+	 *
+	 * @param			FromPin			The pin to start searching from.
+	 *
+	 * @return			The given pin if its owning node is compiler-relevant, or the first pin linked to the owning node at the matching "pass-through" pin that is owned by a compiler-relevant node. May be NULL if no compiler-relevant nodes can be found from the given pin.
+	 */
+	static UEdGraphPin* FindFirstCompilerRelevantLinkedPin(UEdGraphPin* FromPin);
 
 	//////////////////////////////////////////////////////////////////////////
 	// Functions
diff --git a/Engine/Source/Runtime/Core/Public/UObject/ObjectVersion.h b/Engine/Source/Runtime/Core/Public/UObject/ObjectVersion.h
--- a/Engine/Source/Runtime/Core/Public/UObject/ObjectVersion.h
+++ b/Engine/Source/Runtime/Core/Public/UObject/ObjectVersion.h
@@ -576,6 +576,8 @@ enum EUnrealEngineObjectUE4Version
 	VER_UE4_SOUND_CONCURRENCY_PACKAGE,
 	// Changing the default value for focusable user widgets to false
 	VER_UE4_USERWIDGET_DEFAULT_FOCUSABLE_FALSE,
+	// Custom event nodes implicitly set 'const' on array and non-array pass-by-reference input params
+	VER_UE4_BLUEPRINT_CUSTOM_EVENT_CONST_INPUT,
 	// -----<new versions can be added before this line>-------------------------------------------------
 	// - this needs to be the last line (see note below)
 	VER_UE4_AUTOMATIC_VERSION_PLUS_ONE,
