Fix reference to const string that will go out of scope (clang reports as compile error)

diff --git a/Engine/Source/Editor/UnrealEd/Private/Fbx/FbxMainImport.cpp b/Engine/Source/Editor/UnrealEd/Private/Fbx/FbxMainImport.cpp
--- a/Engine/Source/Editor/UnrealEd/Private/Fbx/FbxMainImport.cpp
+++ b/Engine/Source/Editor/UnrealEd/Private/Fbx/FbxMainImport.cpp
@@ -559,8 +559,8 @@ bool FFbxImporter::GetSceneInfo(FString Filename, FbxSceneInfo& SceneInfo)
 						}
 					}
 
-					MeshInfo.SkeletonRoot = Link? MakeName(Link->GetName()) : "None";
-					MeshInfo.SkeletonElemNum = Link? Link->GetChildCount(true) : 0;
+					MeshInfo.SkeletonRoot = Link ? MakeName(Link->GetName()) : MakeName("None");
+					MeshInfo.SkeletonElemNum = Link ? Link->GetChildCount(true) : 0;
 
 					if (Link)
 					{
@@ -846,14 +846,14 @@ bool FFbxImporter::ImportFromFile(const FString& Filename, const FString& Type)
 
 ANSICHAR* FFbxImporter::MakeName(const ANSICHAR* Name)
 {
-	int SpecialChars[] = {'.', ',', '/', '`', '%'};
+	const int SpecialChars[] = {'.', ',', '/', '`', '%'};
 
-	int len = FCStringAnsi::Strlen(Name);
+	const int len = FCStringAnsi::Strlen(Name);
 	ANSICHAR* TmpName = new ANSICHAR[len+1];
 	
 	FCStringAnsi::Strcpy(TmpName, len + 1, Name);
 
-	for ( int32 i = 0; i < 5; i++ )
+	for ( int32 i = 0; i < ARRAY_COUNT(SpecialChars); i++ )
 	{
 		ANSICHAR* CharPtr = TmpName;
 		while ( (CharPtr = FCStringAnsi::Strchr(CharPtr,SpecialChars[i])) != NULL )
