Copying //UE4/Dev-Editor to Dev-Main (//UE4/Dev-Main)

Fixing some incoming CIS errors due to Non-unity build failures.

diff --git a/Engine/Source/Editor/UnrealEd/Private/Fbx/FbxSceneImportFactory.cpp b/Engine/Source/Editor/UnrealEd/Private/Fbx/FbxSceneImportFactory.cpp
--- a/Engine/Source/Editor/UnrealEd/Private/Fbx/FbxSceneImportFactory.cpp
+++ b/Engine/Source/Editor/UnrealEd/Private/Fbx/FbxSceneImportFactory.cpp
@@ -6,13 +6,14 @@
 #include "SSkeletonWidget.h"
 #include "Dialogs/DlgPickPath.h"
 
-#include "FbxImporter.h"
 #include "AssetSelection.h"
 
+#include "FbxImporter.h"
 #include "FbxErrors.h"
 #include "AssetRegistryModule.h"
 #include "Engine/StaticMesh.h"
 #include "Animation/SkeletalMeshActor.h"
+#include "PackageTools.h"
 
 #include "FbxSceneOptionWindow.h"
 #include "MainFrame.h"
@@ -22,6 +23,8 @@
 
 #define LOCTEXT_NAMESPACE "FBXSceneImportFactory"
 
+using namespace UnFbx;
+
 //////////////////////////////////////////////////////////////////////////
 // TODO list
 // -. When we export the material is set in the actor when we import the mesh,
diff --git a/Engine/Source/Editor/UnrealEd/Private/FbxSceneOptionWindow.h b/Engine/Source/Editor/UnrealEd/Private/FbxSceneOptionWindow.h
--- a/Engine/Source/Editor/UnrealEd/Private/FbxSceneOptionWindow.h
+++ b/Engine/Source/Editor/UnrealEd/Private/FbxSceneOptionWindow.h
@@ -8,6 +8,8 @@
 #include "Fbx/SFbxSceneListView.h"
 #include "Factories/FbxSceneImportFactory.h"
 
+class IDetailsView;
+
 class SFbxSceneOptionWindow : public SCompoundWidget
 {
 public:
