Non-unity compilation fixes #UE4

diff --git a/Engine/Source/Editor/UnrealEd/Private/Fbx/FbxSceneOptionWindow.cpp b/Engine/Source/Editor/UnrealEd/Private/Fbx/FbxSceneOptionWindow.cpp
--- a/Engine/Source/Editor/UnrealEd/Private/Fbx/FbxSceneOptionWindow.cpp
+++ b/Engine/Source/Editor/UnrealEd/Private/Fbx/FbxSceneOptionWindow.cpp
@@ -1,6 +1,7 @@
 // Copyright 1998-2015 Epic Games, Inc. All Rights Reserved.
 #include "UnrealEd.h"
 #include "UObjectGlobals.h"
+#include "FbxImporter.h"
 #include "FbxSceneOptionWindow.h"
 #include "Editor/ContentBrowser/Public/ContentBrowserModule.h"
 #include "IDocumentation.h"
diff --git a/Engine/Source/Editor/UnrealEd/Private/Fbx/SFbxSceneListView.cpp b/Engine/Source/Editor/UnrealEd/Private/Fbx/SFbxSceneListView.cpp
--- a/Engine/Source/Editor/UnrealEd/Private/Fbx/SFbxSceneListView.cpp
+++ b/Engine/Source/Editor/UnrealEd/Private/Fbx/SFbxSceneListView.cpp
@@ -1,7 +1,10 @@
 // Copyright 1998-2015 Epic Games, Inc. All Rights Reserved.
 
 #include "UnrealEd.h"
+#include "FbxImporter.h"
+#include "Widgets/Input/STextComboBox.h"
 #include "SFbxSceneListView.h"
+#include "FbxSceneOptionWindow.h"
 #include "ClassIconFinder.h"
 
 #define LOCTEXT_NAMESPACE "SFbxSceneStaticMeshListView"
