Copying //UE4/Orion-Staging to Dev-Main (Originating from //Orion/Main at CL-2777663)

diff --git a/Engine/Source/Editor/Sequencer/Public/ISequencerModule.h b/Engine/Source/Editor/Sequencer/Public/ISequencerModule.h
--- a/Engine/Source/Editor/Sequencer/Public/ISequencerModule.h
+++ b/Engine/Source/Editor/Sequencer/Public/ISequencerModule.h
@@ -66,6 +66,9 @@ struct FSequencerInitParams
 
 	/** Whether or not sequencer should be edited within the level editor */
 	bool bEditWithinLevelEditor;
+
+	/** Domain-specific spawn register for the movie scene */
+	TSharedPtr<IMovieSceneSpawnRegister> SpawnRegister;
 };
 
 
diff --git a/Engine/Source/Editor/Sequencer/Public/ISequencerObjectChangeListener.h b/Engine/Source/Editor/Sequencer/Public/ISequencerObjectChangeListener.h
--- a/Engine/Source/Editor/Sequencer/Public/ISequencerObjectChangeListener.h
+++ b/Engine/Source/Editor/Sequencer/Public/ISequencerObjectChangeListener.h
@@ -8,7 +8,9 @@ class IPropertyHandle;
 
 DECLARE_MULTICAST_DELEGATE_OneParam( FOnAnimatablePropertyChanged, const FPropertyChangedParams& );
 
-DECLARE_MULTICAST_DELEGATE_OneParam ( FOnPropagateObjectChanges, UObject* );
+DECLARE_MULTICAST_DELEGATE_OneParam( FOnObjectPropertyChanged, UObject& );
+
+DECLARE_MULTICAST_DELEGATE_OneParam( FOnPropagateObjectChanges, UObject* );
 
 /**
  * Listens for changes objects and calls delegates when those objects change
@@ -25,7 +27,17 @@ public:
 	 * A delegate for when object changes should be propagated to/from puppet actors
 	 */
 	virtual FOnPropagateObjectChanges& GetOnPropagateObjectChanges() = 0;
+
+	/**
+	 * A delegate for when any property of a specific object is changed.
+	 */
+	virtual FOnObjectPropertyChanged& GetOnAnyPropertyChanged(UObject& Object) = 0;
 	
+	/**
+	 * Report that an object is about to be destroyed (removes any object change delegates bound to the object)
+	 */
+	virtual void ReportObjectDestroyed(UObject& Object) = 0;
+
 	/**
 	 * Triggers all properties as changed for the passed in object.
 	 */
diff --git a/Engine/Source/Editor/Sequencer/Public/ISequencerSection.h b/Engine/Source/Editor/Sequencer/Public/ISequencerSection.h
--- a/Engine/Source/Editor/Sequencer/Public/ISequencerSection.h
+++ b/Engine/Source/Editor/Sequencer/Public/ISequencerSection.h
@@ -128,8 +128,9 @@ public:
 	 * Builds up the section context menu for the outliner
 	 *
 	 * @param MenuBuilder	The menu builder to change
+	 * @param ObjectBinding The object guid bound to this section
 	 */
-	virtual void BuildSectionContextMenu(FMenuBuilder& MenuBuilder) {}
+	virtual void BuildSectionContextMenu(FMenuBuilder& MenuBuilder, const FGuid& ObjectBinding) {}
 
 	/**
 	 * Called when the user requests that a category from this section be deleted. 
diff --git a/Engine/Source/Editor/UMGEditor/Private/Animation/MarginTrackEditor.cpp b/Engine/Source/Editor/UMGEditor/Private/Animation/MarginTrackEditor.cpp
--- a/Engine/Source/Editor/UMGEditor/Private/Animation/MarginTrackEditor.cpp
+++ b/Engine/Source/Editor/UMGEditor/Private/Animation/MarginTrackEditor.cpp
@@ -84,20 +84,13 @@ void FMarginTrackEditor::GenerateKeysFromPropertyChanged( const FPropertyChanged
 	FName ChannelName = PropertyChangedParams.StructPropertyNameToKey;
 	FMargin Margin = PropertyChangedParams.GetPropertyValue<FMargin>();
 
-	if ( ChannelName == NAME_None || ChannelName == LeftName )
-	{
-		GeneratedKeys.Add( FMarginKey( EKeyMarginChannel::Left, Margin.Left ) );
-	}
-	if ( ChannelName == NAME_None || ChannelName == TopName )
-	{
-		GeneratedKeys.Add( FMarginKey( EKeyMarginChannel::Top, Margin.Top ) );
-	}
-	if ( ChannelName == NAME_None || ChannelName == RightName )
-	{
-		GeneratedKeys.Add( FMarginKey( EKeyMarginChannel::Right, Margin.Right ) );
-	}
-	if ( ChannelName == NAME_None || ChannelName == BottomName )
-	{
-		GeneratedKeys.Add( FMarginKey( EKeyMarginChannel::Bottom, Margin.Bottom ) );
-	}
+	GeneratedKeys.Add( FMarginKey( EKeyMarginChannel::Left, Margin.Left, ChannelName == NAME_None || ChannelName == LeftName ? EKeyMarginValueType::Key : EKeyMarginValueType::Default ) );
+	GeneratedKeys.Add( FMarginKey( EKeyMarginChannel::Top, Margin.Top, ChannelName == NAME_None || ChannelName == TopName ? EKeyMarginValueType::Key : EKeyMarginValueType::Default ) );
+	GeneratedKeys.Add( FMarginKey( EKeyMarginChannel::Right, Margin.Right, ChannelName == NAME_None || ChannelName == RightName ? EKeyMarginValueType::Key : EKeyMarginValueType::Default ) );
+	GeneratedKeys.Add( FMarginKey( EKeyMarginChannel::Bottom, Margin.Bottom, ChannelName == NAME_None || ChannelName == BottomName ? EKeyMarginValueType::Key : EKeyMarginValueType::Default ) );
+}
+
+bool FMarginTrackEditor::ShouldAddKey(UMovieSceneMarginTrack* InTrack, FMarginKey InKey, FKeyParams InKeyParams) const
+{
+	return FPropertyTrackEditor::ShouldAddKey(InTrack, InKey, InKeyParams) && InKey.ValueType == EKeyMarginValueType::Key;
 }
diff --git a/Engine/Source/Editor/UMGEditor/Private/Animation/MarginTrackEditor.h b/Engine/Source/Editor/UMGEditor/Private/Animation/MarginTrackEditor.h
--- a/Engine/Source/Editor/UMGEditor/Private/Animation/MarginTrackEditor.h
+++ b/Engine/Source/Editor/UMGEditor/Private/Animation/MarginTrackEditor.h
@@ -41,6 +41,9 @@ protected:
 	virtual TSharedRef<FPropertySection> MakePropertySectionInterface( UMovieSceneSection& SectionObject, UMovieSceneTrack& Track ) override;
 	virtual void GenerateKeysFromPropertyChanged( const FPropertyChangedParams& PropertyChangedParams, TArray<FMarginKey>& GeneratedKeys ) override;
 
+	// FKeyframeTrackEditor interface
+	virtual bool ShouldAddKey(UMovieSceneMarginTrack* InTrack, FMarginKey InKey, FKeyParams InKeyParams) const override;
+
 private:
 
 	static FName LeftName;
diff --git a/Engine/Source/Editor/UMGEditor/Private/Animation/Sequencer2DTransformTrackEditor.cpp b/Engine/Source/Editor/UMGEditor/Private/Animation/Sequencer2DTransformTrackEditor.cpp
--- a/Engine/Source/Editor/UMGEditor/Private/Animation/Sequencer2DTransformTrackEditor.cpp
+++ b/Engine/Source/Editor/UMGEditor/Private/Animation/Sequencer2DTransformTrackEditor.cpp
@@ -125,23 +125,17 @@ void F2DTransformTrackEditor::GenerateKeysFromPropertyChanged( const FPropertyCh
 {
 	FName ChannelName = PropertyChangedParams.StructPropertyNameToKey;
 	FWidgetTransform Transform = PropertyChangedParams.GetPropertyValue<FWidgetTransform>();
-	if ( ChannelName == NAME_None || ChannelName == TranslationName )
-	{
-		GeneratedKeys.Add( F2DTransformKey(EKey2DTransformChannel::Translation, EKey2DTransformAxis::X, Transform.Translation.X ) );
-		GeneratedKeys.Add( F2DTransformKey( EKey2DTransformChannel::Translation, EKey2DTransformAxis::Y, Transform.Translation.Y ) );
-	}
-	if ( ChannelName == NAME_None || ChannelName == ScaleName )
-	{
-		GeneratedKeys.Add( F2DTransformKey( EKey2DTransformChannel::Scale, EKey2DTransformAxis::X, Transform.Scale.X ) );
-		GeneratedKeys.Add( F2DTransformKey( EKey2DTransformChannel::Scale, EKey2DTransformAxis::Y, Transform.Scale.Y ) );
-	}
-	if ( ChannelName == NAME_None || ChannelName == ShearName )
-	{
-		GeneratedKeys.Add( F2DTransformKey( EKey2DTransformChannel::Shear, EKey2DTransformAxis::X, Transform.Shear.X ) );
-		GeneratedKeys.Add( F2DTransformKey( EKey2DTransformChannel::Shear, EKey2DTransformAxis::Y, Transform.Shear.Y ) );
-	}
-	if ( ChannelName == NAME_None || ChannelName == AngleName )
-	{
-		GeneratedKeys.Add( F2DTransformKey( EKey2DTransformChannel::Rotation, EKey2DTransformAxis::None, Transform.Angle ) );
-	}
+
+	GeneratedKeys.Add( F2DTransformKey(EKey2DTransformChannel::Translation, EKey2DTransformAxis::X, Transform.Translation.X, ChannelName == NAME_None || ChannelName == TranslationName ? EKey2DTransformValueType::Key : EKey2DTransformValueType::Default ) );
+	GeneratedKeys.Add( F2DTransformKey( EKey2DTransformChannel::Translation, EKey2DTransformAxis::Y, Transform.Translation.Y, ChannelName == NAME_None || ChannelName == TranslationName ? EKey2DTransformValueType::Key : EKey2DTransformValueType::Default ) );
+	GeneratedKeys.Add( F2DTransformKey( EKey2DTransformChannel::Scale, EKey2DTransformAxis::X, Transform.Scale.X, ChannelName == NAME_None || ChannelName == ScaleName ? EKey2DTransformValueType::Key : EKey2DTransformValueType::Default ) );
+	GeneratedKeys.Add( F2DTransformKey( EKey2DTransformChannel::Scale, EKey2DTransformAxis::Y, Transform.Scale.Y, ChannelName == NAME_None || ChannelName == ScaleName ? EKey2DTransformValueType::Key : EKey2DTransformValueType::Default ) );
+	GeneratedKeys.Add( F2DTransformKey( EKey2DTransformChannel::Shear, EKey2DTransformAxis::X, Transform.Shear.X, ChannelName == NAME_None || ChannelName == ShearName ? EKey2DTransformValueType::Key : EKey2DTransformValueType::Default ) );
+	GeneratedKeys.Add( F2DTransformKey( EKey2DTransformChannel::Shear, EKey2DTransformAxis::Y, Transform.Shear.Y, ChannelName == NAME_None || ChannelName == ShearName ? EKey2DTransformValueType::Key : EKey2DTransformValueType::Default ) );
+	GeneratedKeys.Add( F2DTransformKey( EKey2DTransformChannel::Rotation, EKey2DTransformAxis::None, Transform.Angle, ChannelName == NAME_None || ChannelName == AngleName ? EKey2DTransformValueType::Key : EKey2DTransformValueType::Default ) );
+}
+
+bool F2DTransformTrackEditor::ShouldAddKey(UMovieScene2DTransformTrack* InTrack, F2DTransformKey InKey, FKeyParams InKeyParams) const
+{
+	return FPropertyTrackEditor::ShouldAddKey(InTrack, InKey, InKeyParams) && InKey.ValueType == EKey2DTransformValueType::Key;
 }
diff --git a/Engine/Source/Editor/UMGEditor/Private/Animation/Sequencer2DTransformTrackEditor.h b/Engine/Source/Editor/UMGEditor/Private/Animation/Sequencer2DTransformTrackEditor.h
--- a/Engine/Source/Editor/UMGEditor/Private/Animation/Sequencer2DTransformTrackEditor.h
+++ b/Engine/Source/Editor/UMGEditor/Private/Animation/Sequencer2DTransformTrackEditor.h
@@ -37,6 +37,9 @@ protected:
 
 	virtual void GenerateKeysFromPropertyChanged( const FPropertyChangedParams& PropertyChangedParams, TArray<F2DTransformKey>& GeneratedKeys ) override;
 
+	// FKeyframeTrackEditor interface
+	virtual bool ShouldAddKey(UMovieScene2DTransformTrack* InTrack, F2DTransformKey InKey, FKeyParams InKeyParams) const override;
+
 private:
 	static FName TranslationName;
 	static FName ScaleName;
diff --git a/Engine/Source/Editor/UMGEditor/Private/WidgetBlueprintEditor.cpp b/Engine/Source/Editor/UMGEditor/Private/WidgetBlueprintEditor.cpp
--- a/Engine/Source/Editor/UMGEditor/Private/WidgetBlueprintEditor.cpp
+++ b/Engine/Source/Editor/UMGEditor/Private/WidgetBlueprintEditor.cpp
@@ -5,6 +5,7 @@
 #include "SKismetInspector.h"
 #include "WidgetBlueprintEditor.h"
 #include "MovieScene.h"
+#include "MovieSceneSequenceInstance.h"
 #include "Editor/Sequencer/Public/ISequencerModule.h"
 #include "ObjectEditorUtils.h"
 
@@ -803,7 +804,7 @@ void FWidgetBlueprintEditor::OnGetAnimationAddMenuContent(FMenuBuilder& MenuBuil
 	{
 		for (FObjectAndDisplayName& BindableObject : BindableObjects)
 		{
-			FGuid BoundObjectGuid = CurrentAnimation->FindObjectId(*BindableObject.Object);
+			FGuid BoundObjectGuid = Sequencer->GetFocusedMovieSceneSequenceInstance()->FindObjectId(*BindableObject.Object);
 			if (BoundObjectGuid.IsValid() == false)
 			{
 				FUIAction AddMenuAction(FExecuteAction::CreateSP(this, &FWidgetBlueprintEditor::AddObjectToAnimation, BindableObject.Object));
diff --git a/Engine/Source/Editor/UnrealEd/Classes/Commandlets/ResavePackagesCommandlet.h b/Engine/Source/Editor/UnrealEd/Classes/Commandlets/ResavePackagesCommandlet.h
--- a/Engine/Source/Editor/UnrealEd/Classes/Commandlets/ResavePackagesCommandlet.h
+++ b/Engine/Source/Editor/UnrealEd/Classes/Commandlets/ResavePackagesCommandlet.h
@@ -64,6 +64,9 @@ class UNREALED_API UResavePackagesCommandlet : public UCommandlet
 	/** if we should auto checkin packages that were checked out**/
 	bool bAutoCheckIn;
 
+	/** Should we build lighting for the packages we are saving? **/
+	bool bShouldBuildLighting;
+
 	// Running count of packages that got modified and will need to be resaved
 	int32 PackagesRequiringResave;
 
diff --git a/Engine/Source/Editor/UnrealEd/Private/Commandlets/ContentCommandlets.cpp b/Engine/Source/Editor/UnrealEd/Private/Commandlets/ContentCommandlets.cpp
--- a/Engine/Source/Editor/UnrealEd/Private/Commandlets/ContentCommandlets.cpp
+++ b/Engine/Source/Editor/UnrealEd/Private/Commandlets/ContentCommandlets.cpp
@@ -577,6 +577,8 @@ int32 UResavePackagesCommandlet::Main( const FString& Params )
 	bAutoCheckOut = Switches.Contains(TEXT("AutoCheckOutPackages"));
 	/** if we should auto checkin packages that were checked out**/
 	bAutoCheckIn = bAutoCheckOut && Switches.Contains(TEXT("AutoCheckIn"));
+	/** determine if we are building lighting for the map packages on the pass. **/
+	bShouldBuildLighting = Switches.Contains(TEXT("buildlighting"));
 
 	TArray<FString> PackageNames;
 	int32 ResultCode = InitializeResaveParameters(Tokens, PackageNames);
@@ -672,8 +674,7 @@ FText UResavePackagesCommandlet::GetChangelistDescription() const
 {
 	FText ChangelistDescription;
 
-	static bool bBuildLighting = FParse::Param(FCommandLine::Get(), TEXT("buildlighting")) == true;
-	if (bBuildLighting)
+	if (bShouldBuildLighting)
 	{
 		ChangelistDescription = NSLOCTEXT("ContentCmdlets", "ChangelistDescriptionBuildLighting", "Rebuild lightmaps.");
 	}
@@ -747,9 +748,8 @@ bool UResavePackagesCommandlet::PerformPreloadOperations( FLinkerLoad* PackageLi
 void UResavePackagesCommandlet::PerformAdditionalOperations(class UWorld* World, bool& bSavePackage)
 {
 	check(World);
-	static bool bBuildLighting = FParse::Param(FCommandLine::Get(), TEXT("buildlighting")) == true;
 
-	if (bBuildLighting)
+	if (bShouldBuildLighting)
 	{
 		static bool bHasLoadedStartupPackages = false;
 		if (bHasLoadedStartupPackages == false)
@@ -758,6 +758,8 @@ void UResavePackagesCommandlet::PerformAdditionalOperations(class UWorld* World,
 			bHasLoadedStartupPackages = FStartupPackages::LoadAll();
 		}
 
+		// Setup the world.
+		World->WorldType = EWorldType::Editor;
 		World->AddToRoot();
 		if (!World->bIsWorldInitialized)
 		{
@@ -774,30 +776,24 @@ void UResavePackagesCommandlet::PerformAdditionalOperations(class UWorld* World,
 			World->PersistentLevel->UpdateModelComponents();
 			World->UpdateWorldComponents(true, false);
 		}
-
-		// load all the sublevels
+		FWorldContext &WorldContext = GEditor->GetEditorWorldContext(true);
+		WorldContext.SetCurrentWorld(World);
+		GWorld = World;
 
 		if (World->StreamingLevels.Num())
 		{
-			//World->LoadSecondaryLevels(true, &PreviouslyCookedPackages);
 			World->LoadSecondaryLevels(true, NULL);
-		}
-
-		GWorld = World;
-
-		TArray<FString> SubPackages;
 
-		// force load all of the map to be loaded so we can get good rebuilt lighting 
-		if (World->WorldComposition)
-		{
-			World->WorldComposition->CollectTilesToCook(SubPackages);
-		}
+			for (ULevelStreaming* NextStreamingLevel : World->StreamingLevels)
+			{
+				NextStreamingLevel->bShouldBeVisible = true;
+				NextStreamingLevel->bShouldBeLoaded = true;
+			}
 
-		for (const auto& SubPackage : SubPackages)
-		{
-			UPackage* Package = LoadPackage(nullptr, *SubPackage, 0);
-			check(Package->IsFullyLoaded());
+			World->FlushLevelStreaming(EFlushLevelStreamingType::Full);
 		}
+		// We need any deferred commands added when loading to be executed before we start building lighting.
+		GEngine->TickDeferredCommands();
 
 		GRedirectCollector.ResolveStringAssetReference();
 
diff --git a/Engine/Source/Editor/UnrealEd/Private/Editor.cpp b/Engine/Source/Editor/UnrealEd/Private/Editor.cpp
--- a/Engine/Source/Editor/UnrealEd/Private/Editor.cpp
+++ b/Engine/Source/Editor/UnrealEd/Private/Editor.cpp
@@ -911,11 +911,11 @@ namespace EditorUtilities
 		CopySinglePropertyRecursive(InSourceObject, InTargetObject, InTargetObject, InProperty);
 	}
 
-	int32 CopyActorProperties( AActor* SourceActor, AActor* TargetActor, const ECopyOptions::Type Options )
+	int32 CopyActorProperties( AActor* SourceActor, AActor* TargetActor, const FCopyOptions& Options )
 	{
 		check( SourceActor != nullptr && TargetActor != nullptr );
 
-		const bool bIsPreviewing = ( Options & ECopyOptions::PreviewOnly ) != 0;
+		const bool bIsPreviewing = ( Options.Flags & ECopyOptions::PreviewOnly ) != 0;
 
 		int32 CopiedPropertyCount = 0;
 
@@ -925,7 +925,7 @@ namespace EditorUtilities
 
 		// Get archetype instances for propagation (if requested)
 		TArray<AActor*> ArchetypeInstances;
-		if( Options & ECopyOptions::PropagateChangesToArchetypeInstances )
+		if( Options.Flags & ECopyOptions::PropagateChangesToArchetypeInstances )
 		{
 			TArray<UObject*> ObjectArchetypeInstances;
 			TargetActor->GetArchetypeInstances(ObjectArchetypeInstances);
@@ -935,7 +935,7 @@ namespace EditorUtilities
 				if (AActor* ActorArchetype = Cast<AActor>(ObjectArchetype))
 				{
 					ArchetypeInstances.Add(ActorArchetype);
-				}			
+				}
 			}
 		}
 
@@ -949,14 +949,19 @@ namespace EditorUtilities
 			const bool bIsTransient = !!( Property->PropertyFlags & CPF_Transient );
 			const bool bIsComponentContainer = !!( Property->PropertyFlags & CPF_ContainsInstancedReference );
 			const bool bIsComponentProp = !!( Property->PropertyFlags & ( CPF_InstancedReference | CPF_ContainsInstancedReference ) );
-			const bool bIsBlueprintReadonly = !!(Options & ECopyOptions::FilterBlueprintReadOnly) && !!( Property->PropertyFlags & CPF_BlueprintReadOnly );
+			const bool bIsBlueprintReadonly = !!(Options.Flags & ECopyOptions::FilterBlueprintReadOnly) && !!( Property->PropertyFlags & CPF_BlueprintReadOnly );
 			const bool bIsIdentical = Property->Identical_InContainer( SourceActor, TargetActor );
 
 			if( !bIsTransient && !bIsIdentical && !bIsComponentContainer && !bIsComponentProp && !bIsBlueprintReadonly)
 			{
-				const bool bIsSafeToCopy = !( Options & ECopyOptions::OnlyCopyEditOrInterpProperties ) || ( Property->HasAnyPropertyFlags( CPF_Edit | CPF_Interp ) );
+				const bool bIsSafeToCopy = !( Options.Flags & ECopyOptions::OnlyCopyEditOrInterpProperties ) || ( Property->HasAnyPropertyFlags( CPF_Edit | CPF_Interp ) );
 				if( bIsSafeToCopy )
 				{
+					if (!Options.CanCopyProperty(*Property, *SourceActor))
+					{
+						continue;
+					}
+
 					if( !bIsPreviewing )
 					{
 						if( !ModifiedObjects.Contains(TargetActor) )
@@ -966,14 +971,14 @@ namespace EditorUtilities
 							ModifiedObjects.Add(TargetActor);
 						}
 
-						if( Options & ECopyOptions::CallPostEditChangeProperty )
+						if( Options.Flags & ECopyOptions::CallPostEditChangeProperty )
 						{
 							TargetActor->PreEditChange( Property );
 						}
 
 						// Determine which archetype instances match the current property value of the target actor (before it gets changed). We only want to propagate the change to those instances.
 						TArray<UObject*> ArchetypeInstancesToChange;
-						if( Options & ECopyOptions::PropagateChangesToArchetypeInstances )
+						if( Options.Flags & ECopyOptions::PropagateChangesToArchetypeInstances )
 						{
 							for( AActor* ArchetypeInstance : ArchetypeInstances )
 							{
@@ -986,13 +991,13 @@ namespace EditorUtilities
 
 						CopySingleProperty(SourceActor, TargetActor, Property);
 
-						if( Options & ECopyOptions::CallPostEditChangeProperty )
+						if( Options.Flags & ECopyOptions::CallPostEditChangeProperty )
 						{
 							FPropertyChangedEvent PropertyChangedEvent( Property );
 							TargetActor->PostEditChangeProperty( PropertyChangedEvent );
 						}
 
-						if( Options & ECopyOptions::PropagateChangesToArchetypeInstances )
+						if( Options.Flags & ECopyOptions::PropagateChangesToArchetypeInstances )
 						{
 							for( int32 InstanceIndex = 0; InstanceIndex < ArchetypeInstancesToChange.Num(); ++InstanceIndex )
 							{
@@ -1041,7 +1046,7 @@ namespace EditorUtilities
 
 				// Build a list of matching component archetype instances for propagation (if requested)
 				TArray<UActorComponent*> ComponentArchetypeInstances;
-				if( Options & ECopyOptions::PropagateChangesToArchetypeInstances )
+				if( Options.Flags & ECopyOptions::PropagateChangesToArchetypeInstances )
 				{
 					for( AActor* ArchetypeInstance : ArchetypeInstances )
 					{
@@ -1073,9 +1078,14 @@ namespace EditorUtilities
 					if( !bIsTransient && !bIsIdentical && !bIsComponent && !SourceUCSModifiedProperties.Contains(Property)
 						&& ( !bIsTransform || SourceComponent != SourceActor->GetRootComponent() || ( !SourceActor->HasAnyFlags( RF_ClassDefaultObject | RF_ArchetypeObject ) && !TargetActor->HasAnyFlags( RF_ClassDefaultObject | RF_ArchetypeObject ) ) ) )
 					{
-						const bool bIsSafeToCopy = !( Options & ECopyOptions::OnlyCopyEditOrInterpProperties ) || ( Property->HasAnyPropertyFlags( CPF_Edit | CPF_Interp ) );
+						const bool bIsSafeToCopy = !( Options.Flags & ECopyOptions::OnlyCopyEditOrInterpProperties ) || ( Property->HasAnyPropertyFlags( CPF_Edit | CPF_Interp ) );
 						if( bIsSafeToCopy )
 						{
+							if (!Options.CanCopyProperty(*Property, *SourceActor))
+							{
+								continue;
+							}
+							
 							if( !bIsPreviewing )
 							{
 								if( !ModifiedObjects.Contains(TargetComponent) )
@@ -1085,7 +1095,7 @@ namespace EditorUtilities
 									ModifiedObjects.Add(TargetComponent);
 								}
 
-								if( Options & ECopyOptions::CallPostEditChangeProperty )
+								if( Options.Flags & ECopyOptions::CallPostEditChangeProperty )
 								{
 									// @todo simulate: Should we be calling this on the component instead?
 									TargetActor->PreEditChange( Property );
@@ -1093,7 +1103,7 @@ namespace EditorUtilities
 
 								// Determine which component archetype instances match the current property value of the target component (before it gets changed). We only want to propagate the change to those instances.
 								TArray<UActorComponent*> ComponentArchetypeInstancesToChange;
-								if( Options & ECopyOptions::PropagateChangesToArchetypeInstances )
+								if( Options.Flags & ECopyOptions::PropagateChangesToArchetypeInstances )
 								{
 									for (UActorComponent* ComponentArchetypeInstance : ComponentArchetypeInstances)
 									{
@@ -1125,13 +1135,13 @@ namespace EditorUtilities
 
 								CopySingleProperty(SourceComponent, TargetComponent, Property);
 
-								if( Options & ECopyOptions::CallPostEditChangeProperty )
+								if( Options.Flags & ECopyOptions::CallPostEditChangeProperty )
 								{
 									FPropertyChangedEvent PropertyChangedEvent( Property );
 									TargetActor->PostEditChangeProperty( PropertyChangedEvent );
 								}
 
-								if( Options & ECopyOptions::PropagateChangesToArchetypeInstances )
+								if( Options.Flags & ECopyOptions::PropagateChangesToArchetypeInstances )
 								{
 									for( int32 InstanceIndex = 0; InstanceIndex < ComponentArchetypeInstancesToChange.Num(); ++InstanceIndex )
 									{
@@ -1187,7 +1197,7 @@ namespace EditorUtilities
 		// If one of the changed properties was part of the actor's transformation, then we'll call PostEditMove too.
 		if( !bIsPreviewing && bTransformChanged )
 		{
-			if( Options & ECopyOptions::CallPostEditMove )
+			if( Options.Flags & ECopyOptions::CallPostEditMove )
 			{
 				const bool bFinishedMove = true;
 				TargetActor->PostEditMove( bFinishedMove );
diff --git a/Engine/Source/Editor/UnrealEd/Private/Kismet2/KismetDebugUtilities.cpp b/Engine/Source/Editor/UnrealEd/Private/Kismet2/KismetDebugUtilities.cpp
--- a/Engine/Source/Editor/UnrealEd/Private/Kismet2/KismetDebugUtilities.cpp
+++ b/Engine/Source/Editor/UnrealEd/Private/Kismet2/KismetDebugUtilities.cpp
@@ -736,7 +736,7 @@ UBreakpoint* FKismetDebugUtilities::FindBreakpointForNode(UBlueprint* Blueprint,
 		{
 			Blueprint->Breakpoints.RemoveAtSwap(Index);
 			Blueprint->MarkPackageDirty();
-			UE_LOG(LogBlueprintDebug, Warning, TEXT("Encountered an invalid blueprint breakpoint (this should not happen... if you know how your blueprint got in this state, then please notify the Engine-Blueprints team)"));
+			UE_LOG(LogBlueprintDebug, Warning, TEXT("Encountered an invalid blueprint breakpoint in %s (this should not happen... if you know how your blueprint got in this state, then please notify the Engine-Blueprints team)"), *Blueprint->GetPathName());
 			continue;
 		}
 
@@ -745,7 +745,7 @@ UBreakpoint* FKismetDebugUtilities::FindBreakpointForNode(UBlueprint* Blueprint,
 		{
 			Blueprint->Breakpoints.RemoveAtSwap(Index);
 			Blueprint->MarkPackageDirty();
-			UE_LOG(LogBlueprintDebug, Warning, TEXT("Encountered a blueprint breakpoint without an associated node (this should not happen... if you know how your blueprint got in this state, then please notify the Engine-Blueprints team)"));
+			UE_LOG(LogBlueprintDebug, Warning, TEXT("Encountered a blueprint breakpoint in %s without an associated node (this should not happen... if you know how your blueprint got in this state, then please notify the Engine-Blueprints team)"), *Blueprint->GetPathName());
 			continue;
 		}
 
diff --git a/Engine/Source/Editor/UnrealEd/Private/SCurveEditor.cpp b/Engine/Source/Editor/UnrealEd/Private/SCurveEditor.cpp
--- a/Engine/Source/Editor/UnrealEd/Private/SCurveEditor.cpp
+++ b/Engine/Source/Editor/UnrealEd/Private/SCurveEditor.cpp
@@ -3156,6 +3156,8 @@ void SCurveEditor::MoveTangents(FTrackScaleInfo& ScaleInfo, FVector2D Delta)
 				RichKey.LeaveTangent = NewLeaveTangent;
 			}
 		}	
+
+		RichKey.InterpMode = RCIM_Cubic;
 	}
 }
 
diff --git a/Engine/Source/Editor/UnrealEd/Public/Editor.h b/Engine/Source/Editor/UnrealEd/Public/Editor.h
--- a/Engine/Source/Editor/UnrealEd/Public/Editor.h
+++ b/Engine/Source/Editor/UnrealEd/Public/Editor.h
@@ -600,6 +600,25 @@ namespace EditorUtilities
 	}
 
 
+	/** Copy options structure for CopyActorProperties */
+	struct FCopyOptions
+	{
+		/** Implicit construction for an options enumeration */
+		FCopyOptions(const ECopyOptions::Type InFlags) : Flags(InFlags) {}
+
+		/** Check whether we can copy the specified property */
+		bool CanCopyProperty(UProperty& Property, UObject& Object) const
+		{
+			return !PropertyFilter || PropertyFilter(Property, Object);
+		}
+
+		/** User-specified flags for the copy */
+		ECopyOptions::Type Flags;
+
+		/** User-specified custom property filter predicate */
+		TFunction<bool(UProperty&, UObject&)> PropertyFilter;
+	};
+
 	/**
 	 * Copies properties from one actor to another.  Designed for propagating changes made to PIE actors back to their EditorWorld
 	 * counterpart, or pushing spawned actor changes back to a Blueprint CDO object.  You can pass the 'PreviewOnly' option to
@@ -611,7 +630,7 @@ namespace EditorUtilities
 	 *
 	 * @return	The number of properties that were copied over (properties that were filtered out, or were already identical, are not counted.)
 	 */
-	UNREALED_API int32 CopyActorProperties( AActor* SourceActor, AActor* TargetActor, const ECopyOptions::Type Options = ECopyOptions::Default );
+	UNREALED_API int32 CopyActorProperties( AActor* SourceActor, AActor* TargetActor, const FCopyOptions& Options = FCopyOptions(ECopyOptions::Default) );
 }
 
 
diff --git a/Engine/Source/Programs/AutomationTool/GUBP/GUBP.cs b/Engine/Source/Programs/AutomationTool/GUBP/GUBP.cs
--- a/Engine/Source/Programs/AutomationTool/GUBP/GUBP.cs
+++ b/Engine/Source/Programs/AutomationTool/GUBP/GUBP.cs
@@ -1164,7 +1164,6 @@ public partial class GUBP : BuildCommand
 			if(ParseParam("ResetCounter"))
 			{
 				NewIndex = 0;
-				NowMinutes = 0;
 			}
 			{
 				string Line = String.Format("{0} {1}", NewIndex, NowMinutes);
diff --git a/Engine/Source/Programs/AutomationTool/GUBP/LegacyBranchSetup.cs b/Engine/Source/Programs/AutomationTool/GUBP/LegacyBranchSetup.cs
--- a/Engine/Source/Programs/AutomationTool/GUBP/LegacyBranchSetup.cs
+++ b/Engine/Source/Programs/AutomationTool/GUBP/LegacyBranchSetup.cs
@@ -177,10 +177,11 @@ partial class GUBP
             /// </summary>
             public readonly Dictionary<string, string> GameNameToBuildShareMapping = new Dictionary<string,string>();
 			/// <summary>
-			/// Allows a branch to override the "Localise" commandlet settings.
+			/// Controls which branches should build the Engine localization.
+			/// If you need localization for your own project, add it in your project specific automation scripts (see BuildLocalizationNode and BuildEngineLocalizationNode as an example).
 			/// </summary>
-			public List<string> LocalizationNodesToBuild = new List<string>();
-			public string LocalizationBranchSuffix = "";
+			public bool bBuildEngineLocalization = false;
+			public string EngineLocalizationBranchSuffix = "";
         }
         public virtual void ModifyOptions(GUBP bp, ref BranchOptions Options, string Branch)
         {
@@ -1122,16 +1123,9 @@ partial class GUBP
 			}
 		}
 
-		if (HostPlatforms.Contains(UnrealTargetPlatform.Win64) && BranchConfig.HasNode(RootEditorNode.StaticGetFullName(UnrealTargetPlatform.Win64)))
+		if (HostPlatforms.Contains(UnrealTargetPlatform.Win64) && BranchConfig.HasNode(RootEditorNode.StaticGetFullName(UnrealTargetPlatform.Win64)) && BranchOptions.bBuildEngineLocalization)
 		{
-			foreach (var LocalizationNodeToBuild in BranchOptions.LocalizationNodesToBuild)
-			{
-				var LocalizationNode = BuildLocalizationNode.GetLocalizationNode(LocalizationNodeToBuild, BranchOptions.LocalizationBranchSuffix, BranchConfig);
-				if (LocalizationNode != null)
-				{
-					BranchConfig.AddNode(LocalizationNode);
-				}
-			}
+			BranchConfig.AddNode(new BuildEngineLocalizationNode(BranchOptions.EngineLocalizationBranchSuffix));
 		}
 
         BranchConfig.AddNode(new WaitForTestShared(this));
@@ -1252,7 +1246,7 @@ partial class GUBP
 						}
 						else
 						{
-							throw new AutomationException("Couldn't find referenced node {0} in graph", Aggregate.GetFullName());
+							CommandUtils.LogWarning("Couldn't find referenced frequency barrier node {0}", SearchNodes[Idx]);
 						}
 					}
 				}
diff --git a/Engine/Source/Programs/AutomationTool/GUBP/LegacyNodes.cs b/Engine/Source/Programs/AutomationTool/GUBP/LegacyNodes.cs
--- a/Engine/Source/Programs/AutomationTool/GUBP/LegacyNodes.cs
+++ b/Engine/Source/Programs/AutomationTool/GUBP/LegacyNodes.cs
@@ -2648,7 +2648,7 @@ partial class GUBP
 
 	public abstract class BuildLocalizationNode : HostPlatformNode
 	{
-		public BuildLocalizationNode(string InLocalizationBranchSuffix, GUBPBranchConfig InBranchConfig)
+		public BuildLocalizationNode(string InLocalizationBranchSuffix)
 			: base(UnrealTargetPlatform.Win64)
 		{
 			LocalizationBranchSuffix = InLocalizationBranchSuffix;
@@ -2709,11 +2709,6 @@ partial class GUBP
 			return base.CISFrequencyQuantumShift(BranchConfig) + 6;
 		}
 
-		public virtual string GetLocalizationId()
-		{
-			throw new AutomationException("Unimplemented GetLocalizationId.");
-		}
-
 		protected virtual string GetUEProjectDirectory()
 		{
 			throw new AutomationException("Unimplemented GetUEProjectDirectory.");
@@ -2739,83 +2734,17 @@ partial class GUBP
 			throw new AutomationException("Unimplemented GetOneSkyProjectNames.");
 		}
 
-		public static BuildLocalizationNode GetLocalizationNode(string InLocalizationId, string InLocalizationBranchSuffix, GUBPBranchConfig InBranchConfig)
-		{
-			if (CachedLocalizationNodeTypes == null)
-			{
-				// Find all types that derive from BuildLocalizationNode in any of our DLLs
-				CachedLocalizationNodeTypes = new Dictionary<string, Type>();
-				var LoadedAssemblies = AppDomain.CurrentDomain.GetAssemblies();
-				foreach (var Dll in LoadedAssemblies)
-				{
-					var AllTypes = Dll.GetTypes();
-					foreach (var PotentialLocalizationNodeType in AllTypes)
-					{
-						if (PotentialLocalizationNodeType != typeof(BuildLocalizationNode) && typeof(BuildLocalizationNode).IsAssignableFrom(PotentialLocalizationNodeType))
-						{
-							// Types should implement a static StaticGetLocalizationId method
-							var Method = PotentialLocalizationNodeType.GetMethod("StaticGetLocalizationId");
-							if (Method != null)
-							{
-								try
-								{
-									var NodeLocalizationId = Method.Invoke(null, null) as string;
-									CachedLocalizationNodeTypes.Add(NodeLocalizationId, PotentialLocalizationNodeType);
-								}
-								catch
-								{
-									BuildCommand.LogWarning("Type '{0}' threw when calling its StaticGetLocalizationId method.", PotentialLocalizationNodeType.FullName);
-								}
-							}
-							else
-							{
-								BuildCommand.LogWarning("Type '{0}' derives from BuildLocalizationNode but is missing its StaticGetLocalizationId method.", PotentialLocalizationNodeType.FullName);
-							}
-						}
-					}
-				}
-			}
-
-			Type LocalizationNodeType;
-			CachedLocalizationNodeTypes.TryGetValue(InLocalizationId, out LocalizationNodeType);
-			if (LocalizationNodeType != null)
-			{
-				try
-				{
-					return Activator.CreateInstance(LocalizationNodeType, new object[] { InLocalizationBranchSuffix, InBranchConfig }) as BuildLocalizationNode;
-				}
-				catch
-				{
-					BuildCommand.LogWarning("Unable to create an instance of the type '{0}'", LocalizationNodeType.FullName);
-				}
-			}
-
-			return null;
-		}
-
-		private static Dictionary<string, Type> CachedLocalizationNodeTypes;
-
 		protected string LocalizationBranchSuffix;
 	}
 
 	public class BuildEngineLocalizationNode : BuildLocalizationNode
 	{
-		public BuildEngineLocalizationNode(string InLocalizationBranchSuffix, GUBPBranchConfig InBranchConfig)
-			: base(InLocalizationBranchSuffix, InBranchConfig)
+		public BuildEngineLocalizationNode(string InLocalizationBranchSuffix)
+			: base(InLocalizationBranchSuffix)
 		{
 			AddDependency(RootEditorNode.StaticGetFullName(HostPlatform));
 		}
 
-		public static string StaticGetLocalizationId()
-		{
-			return "Engine";
-		}
-
-		public override string GetLocalizationId()
-		{
-			return StaticGetLocalizationId();
-		}
-
 		public static string StaticGetFullName()
 		{
 			return "BuildEngineLocalization";
diff --git a/Engine/Source/Programs/AutomationTool/Scripts/CopyBuildToStagingDirectory.Automation.cs b/Engine/Source/Programs/AutomationTool/Scripts/CopyBuildToStagingDirectory.Automation.cs
--- a/Engine/Source/Programs/AutomationTool/Scripts/CopyBuildToStagingDirectory.Automation.cs
+++ b/Engine/Source/Programs/AutomationTool/Scripts/CopyBuildToStagingDirectory.Automation.cs
@@ -925,7 +925,18 @@ public partial class Project : CommandUtils
 
 		if (!bCopiedExistingPak)
 		{
-			RunUnrealPak(UnrealPakResponseFile, OutputLocation, Params.SignPak, PakOrderFileLocation, SC.StageTargetPlatform.GetPlatformPakCommandLine(), Params.Compressed, PatchSourceContentPath );
+			if (File.Exists(OutputLocation))
+			{
+				string UnrealPakResponseFileName = CombinePaths(CmdEnv.LogFolder, "PakList_" + PakName + ".txt");
+				if (File.GetLastWriteTimeUtc(OutputLocation) > File.GetLastWriteTimeUtc(UnrealPakResponseFileName))
+				{
+					bCopiedExistingPak = true;
+				}
+			}
+			if (!bCopiedExistingPak)
+			{
+				RunUnrealPak(UnrealPakResponseFile, OutputLocation, Params.SignPak, PakOrderFileLocation, SC.StageTargetPlatform.GetPlatformPakCommandLine(), Params.Compressed, PatchSourceContentPath);
+			}
 		}
 
 
@@ -1173,21 +1184,21 @@ public partial class Project : CommandUtils
 		StagedFilesDir.GetFiles("*.pak", SearchOption.AllDirectories).ToList().ForEach(File => File.Delete());
 	}
 
-	public static void ApplyStagingManifest(ProjectParams Params, DeploymentContext SC)
+	public static void CleanStagingDirectory(ProjectParams Params, DeploymentContext SC)
 	{
 		MaybeConvertToLowerCase(Params, SC);
 		Log("Cleaning Stage Directory: {0}", SC.StageDirectory);
 		if (SC.Stage && !Params.NoCleanStage && !Params.SkipStage && !Params.IterativeDeploy)
 		{
-            try
-            {
-                DeleteDirectory(SC.StageDirectory);
-            }
-            catch (Exception Ex)
-            {
-                // Delete cooked data (if any) as it may be incomplete / corrupted.
-                throw new AutomationException(ExitCode.Error_FailedToDeleteStagingDirectory, Ex, "Stage Failed. Failed to delete staging directory " + SC.StageDirectory);
-            }
+			try
+			{
+				DeleteDirectory(SC.StageDirectory);
+			}
+			catch (Exception Ex)
+			{
+				// Delete cooked data (if any) as it may be incomplete / corrupted.
+				throw new AutomationException(ExitCode.Error_FailedToDeleteStagingDirectory, Ex, "Stage Failed. Failed to delete staging directory " + SC.StageDirectory);
+			}
 		}
 		else
 		{
@@ -1196,12 +1207,16 @@ public partial class Project : CommandUtils
 				// delete old pak files
 				DeletePakFiles(SC.StageDirectory);
 			}
-            catch (Exception Ex)
-            {
-                // Delete cooked data (if any) as it may be incomplete / corrupted.
-                throw new AutomationException(ExitCode.Error_FailedToDeleteStagingDirectory, Ex, "Stage Failed. Failed to delete pak files in " + SC.StageDirectory);
-            }
+			catch (Exception Ex)
+			{
+				// Delete cooked data (if any) as it may be incomplete / corrupted.
+				throw new AutomationException(ExitCode.Error_FailedToDeleteStagingDirectory, Ex, "Stage Failed. Failed to delete pak files in " + SC.StageDirectory);
+			}
 		}
+	}
+
+	public static void ApplyStagingManifest(ProjectParams Params, DeploymentContext SC)
+	{
 		if (ShouldCreatePak(Params, SC))
 		{
 			if (Params.Manifests && DoesChunkPakManifestExist(Params, SC))
@@ -1595,14 +1610,8 @@ public partial class Project : CommandUtils
 			PlatformsToStage = Params.ServerTargetPlatforms;
         }
 
-        bool prefixArchiveDir = false;
-        if (PlatformsToStage.Contains(UnrealTargetPlatform.Win32) && PlatformsToStage.Contains(UnrealTargetPlatform.Win64))
-        {
-            prefixArchiveDir = true;
-        }
-
-		List<DeploymentContext> DeploymentContexts = new List<DeploymentContext>();
-		foreach (var StagePlatform in PlatformsToStage)
+ 		List<DeploymentContext> DeploymentContexts = new List<DeploymentContext>();
+			foreach (var StagePlatform in PlatformsToStage)
 		{
 			// Get the platform to get cooked data from, may differ from the stage platform
 			UnrealTargetPlatform CookedDataPlatform = Params.GetCookedDataPlatformForClientTarget(StagePlatform);
@@ -1631,18 +1640,6 @@ public partial class Project : CommandUtils
 
 			string StageDirectory = ((ShouldCreatePak(Params) || (Params.Stage)) || !String.IsNullOrEmpty(Params.StageDirectoryParam)) ? Params.BaseStageDirectory : "";
             string ArchiveDirectory = (Params.Archive || !String.IsNullOrEmpty(Params.ArchiveDirectoryParam)) ? Params.BaseArchiveDirectory : "";
-            if (prefixArchiveDir && (StagePlatform == UnrealTargetPlatform.Win32 || StagePlatform == UnrealTargetPlatform.Win64))
-            {
-                if (Params.Stage)
-                {
-                    StageDirectory = CombinePaths(Params.BaseStageDirectory, StagePlatform.ToString());
-                }
-                if (Params.Archive)
-                {
-                    ArchiveDirectory = CombinePaths(Params.BaseArchiveDirectory, StagePlatform.ToString());
-                }
-            }
-
 			string EngineDir = CommandUtils.CombinePaths(CommandUtils.CmdEnv.LocalRoot, "Engine");
 
 			List<StageTarget> TargetsToStage = new List<StageTarget>();
@@ -1748,11 +1745,17 @@ public partial class Project : CommandUtils
 			if (!Params.NoClient)
 			{
 				var DeployContextList = CreateDeploymentContext(Params, false, true);
+
+				// clean the staging directories first
 				foreach (var SC in DeployContextList)
 				{
 					// write out the commandline file now so it can go into the manifest
 					WriteStageCommandline(Params, SC);
 					CreateStagingManifest(Params, SC);
+					CleanStagingDirectory(Params, SC);
+				}
+				foreach (var SC in DeployContextList)
+				{
 					ApplyStagingManifest(Params, SC);
 
 					if (Params.Deploy)
diff --git a/Engine/Source/Programs/AutomationTool/Scripts/Localisation.Automation.cs b/Engine/Source/Programs/AutomationTool/Scripts/Localisation.Automation.cs
--- a/Engine/Source/Programs/AutomationTool/Scripts/Localisation.Automation.cs
+++ b/Engine/Source/Programs/AutomationTool/Scripts/Localisation.Automation.cs
@@ -106,6 +106,14 @@ class Localise : BuildCommand
 
 		var RootWorkingDirectory = CombinePaths(CmdEnv.LocalRoot, UEProjectDirectory);
 
+		// Make sure the Localization configs and content is up-to-date to ensure we don't get errors later on
+		if (P4Enabled)
+		{
+			Log("Sync necessary content to head revision");
+			P4.Sync(P4Env.BuildRootP4 + "/" + UEProjectDirectory + "/Config/Localization/...");
+			P4.Sync(P4Env.BuildRootP4 + "/" + UEProjectDirectory + "/Content/Localization/...");
+		}
+
 		// Generate the info we need to gather for each project
 		var ProjectInfos = new List<ProjectInfo>();
 		foreach (var ProjectName in OneSkyProjectNames)
@@ -126,7 +134,7 @@ class Localise : BuildCommand
 		// Export all text from OneSky
 		foreach (var ProjectInfo in ProjectInfos)
 		{
-			ExportOneSkyProjectToDirectory(RootWorkingDirectory, OneSkyService, OneSkyProjectGroup, ProjectInfo);
+			ExportOneSkyProjectToDirectory(RootWorkingDirectory, OneSkyService, OneSkyProjectGroup, OneSkyBranchSuffix, ProjectInfo);
 		}
 
 		// Submit changelist for backed up POs from OneSky.
@@ -311,10 +319,11 @@ class Localise : BuildCommand
 		return OneSkyProject;
     }
 
-	private void ExportOneSkyProjectToDirectory(string RootWorkingDirectory, OneSkyService OneSkyService, ProjectGroup OneSkyProjectGroup, ProjectInfo ProjectInfo)
+	private void ExportOneSkyProjectToDirectory(string RootWorkingDirectory, OneSkyService OneSkyService, ProjectGroup OneSkyProjectGroup, string OneSkyBranchSuffix, ProjectInfo ProjectInfo)
 	{
+		var OneSkyFileName = GetOneSkyFilename(ProjectInfo.ImportInfo.PortableObjectName, OneSkyBranchSuffix);
 		var OneSkyProject = GetOneSkyProject(OneSkyService, OneSkyProjectGroup.Name, ProjectInfo.ProjectName);
-		var OneSkyFile = OneSkyProject.UploadedFiles.FirstOrDefault(f => f.Filename == ProjectInfo.ImportInfo.PortableObjectName);
+		var OneSkyFile = OneSkyProject.UploadedFiles.FirstOrDefault(f => f.Filename == OneSkyFileName);
 
 		//Export
 		if (OneSkyFile != null)
@@ -329,11 +338,11 @@ class Localise : BuildCommand
 				}
 			}
 
-			ExportOneSkyFileToDirectory(OneSkyFile, new DirectoryInfo(CombinePaths(RootWorkingDirectory, ProjectInfo.ImportInfo.DestinationPath)), CulturesToExport, ProjectInfo.ImportInfo.bUseCultureDirectory);
+			ExportOneSkyFileToDirectory(OneSkyFile, new DirectoryInfo(CombinePaths(RootWorkingDirectory, ProjectInfo.ImportInfo.DestinationPath)), ProjectInfo.ImportInfo.PortableObjectName, CulturesToExport, ProjectInfo.ImportInfo.bUseCultureDirectory);
 		}
 	}
 
-	private void ExportOneSkyFileToDirectory(UploadedFile OneSkyFile, DirectoryInfo DestinationDirectory, IEnumerable<string> Cultures, bool bUseCultureDirectory)
+	private void ExportOneSkyFileToDirectory(UploadedFile OneSkyFile, DirectoryInfo DestinationDirectory, string DestinationFilename, IEnumerable<string> Cultures, bool bUseCultureDirectory)
     {
 		foreach (var Culture in Cultures)
         {
@@ -345,11 +354,11 @@ class Localise : BuildCommand
 
             using (var MemoryStream = new MemoryStream())
             {
-				var ExportFile = new FileInfo(Path.Combine(CultureDirectory.FullName, OneSkyFile.Filename));
-
 				var ExportTranslationState = OneSkyFile.ExportTranslation(Culture, MemoryStream).Result;
 				if (ExportTranslationState == UploadedFile.ExportTranslationState.Success)
 				{
+					var ExportFile = new FileInfo(Path.Combine(CultureDirectory.FullName, DestinationFilename));
+
 					// Write out the updated PO file so that the gather commandlet will import the new data from it
 					{
 						var ExportFileWasReadOnly = false;
@@ -364,7 +373,7 @@ class Localise : BuildCommand
 						using (Stream FileStream = ExportFile.OpenWrite())
 						{
 							MemoryStream.CopyTo(FileStream);
-							Console.WriteLine("[SUCCESS] Exporting: '{0}' ({1})", ExportFile.FullName, Culture);
+							Console.WriteLine("[SUCCESS] Exporting: '{0}' as '{1}' ({2})", OneSkyFile.Filename, ExportFile.FullName, Culture);
 						}
 
 						if (ExportFileWasReadOnly)
@@ -401,11 +410,11 @@ class Localise : BuildCommand
 				}
 				else if (ExportTranslationState == UploadedFile.ExportTranslationState.NoContent)
                 {
-					Console.WriteLine("[WARNING] Exporting: '{0}' ({1}) has no translations!", ExportFile.FullName, Culture);
+					Console.WriteLine("[WARNING] Exporting: '{0}' ({1}) has no translations!", OneSkyFile.Filename, Culture);
                 }
                 else
                 {
-					Console.WriteLine("[FAILED] Exporting: '{0}' ({1})", ExportFile.FullName, Culture);
+					Console.WriteLine("[FAILED] Exporting: '{0}' ({1})", OneSkyFile.Filename, Culture);
                 }
             }
         }
@@ -455,15 +464,9 @@ class Localise : BuildCommand
 				FileStream.Position = 0;
 			}
 
-			Console.WriteLine("Uploading: '{0}' ({1})", FileToUpload.FullName, Culture);
+			var OneSkyFileName = GetOneSkyFilename(Path.GetFileName(FileToUpload.FullName), OneSkyBranchSuffix);
 
-			var OneSkyFileName = Path.GetFileName(FileToUpload.FullName);
-			if (!String.IsNullOrEmpty(OneSkyBranchSuffix))
-			{
-				// Apply the branch suffix. OneSky will take care of merging the files from different branches together.
-				var OneSkyFileNameWithSuffix = Path.GetFileNameWithoutExtension(OneSkyFileName) + "_" + OneSkyBranchSuffix + Path.GetExtension(OneSkyFileName);
-				OneSkyFileName = OneSkyFileNameWithSuffix;
-			}
+			Console.WriteLine("Uploading: '{0}' as '{1}' ({2})", FileToUpload.FullName, OneSkyFileName, Culture);
 
 			var UploadedFile = OneSkyProject.Upload(OneSkyFileName, FileStream, Culture).Result;
 
@@ -477,4 +480,16 @@ class Localise : BuildCommand
 			}
 		}
 	}
+
+	private static string GetOneSkyFilename(string BaseFilename, string OneSkyBranchSuffix)
+	{
+		var OneSkyFileName = BaseFilename;
+		if (!String.IsNullOrEmpty(OneSkyBranchSuffix))
+		{
+			// Apply the branch suffix. OneSky will take care of merging the files from different branches together.
+			var OneSkyFileNameWithSuffix = Path.GetFileNameWithoutExtension(OneSkyFileName) + "_" + OneSkyBranchSuffix + Path.GetExtension(OneSkyFileName);
+			OneSkyFileName = OneSkyFileNameWithSuffix;
+		}
+		return OneSkyFileName;
+	}
 }
diff --git a/Engine/Source/Programs/CrashReporter/CrashReportClient/Private/CrashReportClientConfig.cpp b/Engine/Source/Programs/CrashReporter/CrashReportClient/Private/CrashReportClientConfig.cpp
--- a/Engine/Source/Programs/CrashReporter/CrashReportClient/Private/CrashReportClientConfig.cpp
+++ b/Engine/Source/Programs/CrashReporter/CrashReportClient/Private/CrashReportClientConfig.cpp
@@ -36,6 +36,12 @@ FCrashReportClientConfig::FCrashReportClientConfig()
 			bSendLogFile = true;
 		}
 	}
+
+	if (!GConfig->GetBool(TEXT("CrashReportClient"), TEXT("bHideLogFilesOption"), bHideLogFilesOption, GEngineIni))
+	{
+		// Default to false (show the option) when config is missing.
+		bHideLogFilesOption = false;
+	}
 	
 	ReadFullCrashDumpConfigurations();
 
diff --git a/Engine/Source/Programs/CrashReporter/CrashReportClient/Private/CrashReportClientConfig.h b/Engine/Source/Programs/CrashReporter/CrashReportClient/Private/CrashReportClientConfig.h
--- a/Engine/Source/Programs/CrashReporter/CrashReportClient/Private/CrashReportClientConfig.h
+++ b/Engine/Source/Programs/CrashReporter/CrashReportClient/Private/CrashReportClientConfig.h
@@ -68,6 +68,11 @@ struct FCrashReportClientConfig
 		return bSendLogFile;
 	}
 
+	const bool& GetHideLogFilesOption() const
+	{
+		return bHideLogFilesOption;
+	}
+
 	void SetAllowToBeContacted( bool bNewValue );
 	void SetSendLogFile( bool bNewValue );
 
@@ -104,4 +109,7 @@ protected:
 
 	/** Whether the user allowed us to send the log file. */
 	bool bSendLogFile;
+
+	/** Whether the user is shown the option to enable/disable sending the log file. */
+	bool bHideLogFilesOption;
 };
diff --git a/Engine/Source/Programs/CrashReporter/CrashReportClient/Private/SCrashReportClient.cpp b/Engine/Source/Programs/CrashReporter/CrashReportClient/Private/SCrashReportClient.cpp
--- a/Engine/Source/Programs/CrashReporter/CrashReportClient/Private/SCrashReportClient.cpp
+++ b/Engine/Source/Programs/CrashReporter/CrashReportClient/Private/SCrashReportClient.cpp
@@ -191,7 +191,7 @@ void SCrashReportClient::Construct(const FArguments& InArgs, TSharedRef<FCrashRe
 			.Padding( FMargin( 4, 12, 4, 4 ) )
 			[
 				SNew( SHorizontalBox )
-
+				.Visibility( FCrashReportClientConfig::Get().GetHideLogFilesOption() ? EVisibility::Collapsed : EVisibility::Visible )
 				+ SHorizontalBox::Slot()
 				.AutoWidth()
 				.VAlign( VAlign_Center )
diff --git a/Engine/Source/Programs/UnrealBuildTool/Configuration/UEBuildPlatform.cs b/Engine/Source/Programs/UnrealBuildTool/Configuration/UEBuildPlatform.cs
--- a/Engine/Source/Programs/UnrealBuildTool/Configuration/UEBuildPlatform.cs
+++ b/Engine/Source/Programs/UnrealBuildTool/Configuration/UEBuildPlatform.cs
@@ -104,6 +104,15 @@ namespace UnrealBuildTool
 			return "";
 		}
 
+		/// <summary>
+		/// Get name for architecture-specific directories (can be shorter than architecture name itself)
+		/// </summary>
+		public virtual string GetActiveArchitectureFolderName()
+		{
+			// by default, use the architecture name
+			return GetActiveArchitecture();
+		}
+
 		/// <summary>
 		/// Setup the configuration environment for building
 		/// </summary>
diff --git a/Engine/Source/Programs/UnrealBuildTool/Linux/UEBuildLinux.cs b/Engine/Source/Programs/UnrealBuildTool/Linux/UEBuildLinux.cs
--- a/Engine/Source/Programs/UnrealBuildTool/Linux/UEBuildLinux.cs
+++ b/Engine/Source/Programs/UnrealBuildTool/Linux/UEBuildLinux.cs
@@ -22,6 +22,23 @@ namespace UnrealBuildTool
 			return LinuxPlatform.DefaultArchitecture;
 		}
 
+		/// <summary>
+		/// Get name for architecture-specific directories (can be shorter than architecture name itself)
+		/// </summary>
+		public override string GetActiveArchitectureFolderName()
+		{
+			// shorten the string (heuristically)
+			string Arch = GetActiveArchitecture();
+			uint Sum = 0;
+			int Len = Arch.Length;
+			for (int Index = 0; Index < Len; ++Index)
+			{
+				Sum += (uint)(Arch[Index]);
+				Sum <<= 1;	// allowed to overflow
+			}
+			return Sum.ToString("X");
+		}
+
 		/// <summary>
 		/// Modify the rules for a newly created module, in a target that's being built for this platform.
 		/// This is not required - but allows for hiding details of a particular platform.
diff --git a/Engine/Source/Programs/UnrealBuildTool/System/UnrealBuildTool.cs b/Engine/Source/Programs/UnrealBuildTool/System/UnrealBuildTool.cs
--- a/Engine/Source/Programs/UnrealBuildTool/System/UnrealBuildTool.cs
+++ b/Engine/Source/Programs/UnrealBuildTool/System/UnrealBuildTool.cs
@@ -1304,7 +1304,7 @@ namespace UnrealBuildTool
 								{
 									// Setup environment wasn't called, so set the flag
 									BuildConfiguration.bDeployAfterCompile = BuildPlatform.RequiresDeployPrepAfterCompile();
-									BuildConfiguration.PlatformIntermediateFolder = Path.Combine(BuildConfiguration.BaseIntermediateFolder, CheckPlatform.ToString(), BuildPlatform.CreateContext(ProjectFile).GetActiveArchitecture());
+									BuildConfiguration.PlatformIntermediateFolder = Path.Combine(BuildConfiguration.BaseIntermediateFolder, CheckPlatform.ToString(), BuildPlatform.CreateContext(ProjectFile).GetActiveArchitectureFolderName());
 								}
 							}
 
@@ -1548,7 +1548,7 @@ namespace UnrealBuildTool
 			BuildPlatformContext.ResetBuildConfiguration(ResetConfiguration);
 
 			// now that we have the platform, we can set the intermediate path to include the platform/architecture name
-			BuildConfiguration.PlatformIntermediateFolder = Path.Combine(BuildConfiguration.BaseIntermediateFolder, ResetPlatform.ToString(), BuildPlatformContext.GetActiveArchitecture());
+			BuildConfiguration.PlatformIntermediateFolder = Path.Combine(BuildConfiguration.BaseIntermediateFolder, ResetPlatform.ToString(), BuildPlatformContext.GetActiveArchitectureFolderName());
 
 			string ExecutorName = "Unknown";
 			ECompilationResult BuildResult = ECompilationResult.Succeeded;
diff --git a/Engine/Source/Programs/UnrealSwarm/Agent/Channels.cs b/Engine/Source/Programs/UnrealSwarm/Agent/Channels.cs
--- a/Engine/Source/Programs/UnrealSwarm/Agent/Channels.cs
+++ b/Engine/Source/Programs/UnrealSwarm/Agent/Channels.cs
@@ -1071,7 +1071,7 @@ namespace Agent
 						{
 							if( LocalChannelHashValue[Index] != RemoteChannelHashValue[Index] )
 							{
-								Log( EVerbosityLevel.Informative, ELogColour.Orange, String.Format( "[ValidateChannel] Hash value mismatch for {0}", ChannelName ) );
+								//Log( EVerbosityLevel.Informative, ELogColour.Orange, String.Format( "[ValidateChannel] Hash value mismatch for {0}", ChannelName ) );
 								return false;
 							}
 						}
diff --git a/Engine/Source/Runtime/Core/Private/Misc/ConfigCacheIni.cpp b/Engine/Source/Runtime/Core/Private/Misc/ConfigCacheIni.cpp
--- a/Engine/Source/Runtime/Core/Private/Misc/ConfigCacheIni.cpp
+++ b/Engine/Source/Runtime/Core/Private/Misc/ConfigCacheIni.cpp
@@ -2749,7 +2749,13 @@ static void GetSourceIniHierarchyFilenames(const TCHAR* InBaseIniName, const TCH
 
 	// [[[[ ENGINE DEFAULTS ]]]]
 	// Engine/Config/Base.ini (included in every ini type, required)
-	OutHierarchy.Add(EConfigFileHierarchy::AbsoluteBase, FIniFilename(FString::Printf(TEXT("%sBase.ini"), EngineConfigDir), true));
+	// @todo: ChrisW - this is a temporary measure to allow standalone tools to work when engine config files are in a pak
+#if IS_PROGRAM
+	const bool BaseIniRequired = false;
+#else
+	const bool BaseIniRequired = true;
+#endif
+	OutHierarchy.Add(EConfigFileHierarchy::AbsoluteBase, FIniFilename(FString::Printf(TEXT("%sBase.ini"), EngineConfigDir), BaseIniRequired));
 	// Engine/Config/Base* ini
 	OutHierarchy.Add(EConfigFileHierarchy::EngineDirBase, FIniFilename(FString::Printf(TEXT("%sBase%s.ini"), EngineConfigDir, InBaseIniName), false));
 	// Engine/Config/NotForLicensees/Base* ini
diff --git a/Engine/Source/Runtime/Core/Public/GenericPlatform/GenericPlatformMath.h b/Engine/Source/Runtime/Core/Public/GenericPlatform/GenericPlatformMath.h
--- a/Engine/Source/Runtime/Core/Public/GenericPlatform/GenericPlatformMath.h
+++ b/Engine/Source/Runtime/Core/Public/GenericPlatform/GenericPlatformMath.h
@@ -233,7 +233,20 @@ struct FGenericPlatformMath
 	*			So for example Fmod(2.8f, 2) gives .8f as you would expect, however, Fmod(-2.8f, 2) gives -.8f, NOT 1.2f 
 	* Use Floor instead when snapping positions that can be negative to a grid
 	*/
-	static FORCEINLINE float Fmod( float X, float Y ) { return fmodf(X, Y); }
+	static FORCEINLINE_DEBUGGABLE float Fmod(float X, float Y)
+	{
+#ifdef PLATFORM_WINDOWS
+		// There's a compiler bug on Windows, where fmodf will start returning NaNs randomly with valid inputs.
+		// Until this is resolved, we implement our own version.
+		float IntPortion = TruncToFloat(X / Y);
+		float Result = X - Y * IntPortion;
+
+		return Result;
+#else
+		return fmodf(X, Y);
+#endif		
+	}
+
 	static FORCEINLINE float Sin( float Value ) { return sinf(Value); }
 	static FORCEINLINE float Asin( float Value ) { return asinf( (Value<-1.f) ? -1.f : ((Value<1.f) ? Value : 1.f) ); }
 	static FORCEINLINE float Cos( float Value ) { return cosf(Value); }
diff --git a/Engine/Source/Runtime/CoreUObject/Private/UObject/LinkerLoad.cpp b/Engine/Source/Runtime/CoreUObject/Private/UObject/LinkerLoad.cpp
--- a/Engine/Source/Runtime/CoreUObject/Private/UObject/LinkerLoad.cpp
+++ b/Engine/Source/Runtime/CoreUObject/Private/UObject/LinkerLoad.cpp
@@ -2388,7 +2388,6 @@ bool FLinkerLoad::VerifyImportInner(const int32 ImportIndex, FString& WarningSuf
 		return false;
 	}
 
-
 	bool SafeReplace = false;
 	UObject* Pkg=NULL;
 	UPackage* TmpPkg=NULL;
@@ -2396,7 +2395,7 @@ bool FLinkerLoad::VerifyImportInner(const int32 ImportIndex, FString& WarningSuf
 	// Find or load the linker load that contains the FObjectExport for this import
 	if (Import.OuterIndex.IsNull() && Import.ClassName!=NAME_Package )
 	{
-		UE_LOG(LogLinker, Warning, TEXT("%s has an inappropriate outermost, it was probably saved with a deprecated outer (file: %s)"), *Import.ObjectName.ToString(), *Filename);
+		UE_LOG(LogLinker, Error, TEXT("%s has an inappropriate outermost, it was probably saved with a deprecated outer (file: %s)"), *Import.ObjectName.ToString(), *Filename);
 		Import.SourceLinker = NULL;
 		return false;
 	}
diff --git a/Engine/Source/Runtime/Engine/Classes/AI/Navigation/RecastNavMesh.h b/Engine/Source/Runtime/Engine/Classes/AI/Navigation/RecastNavMesh.h
--- a/Engine/Source/Runtime/Engine/Classes/AI/Navigation/RecastNavMesh.h
+++ b/Engine/Source/Runtime/Engine/Classes/AI/Navigation/RecastNavMesh.h
@@ -880,10 +880,10 @@ public:
 	uint32 GetPolyAreaID(NavNodeRef PolyID) const;
 
 	/** Sets area ID for the specified polygon. */
-	void SetPolyAreaID(NavNodeRef PolyID, uint8 AreaID);
+	void SetPolyArea(NavNodeRef PolyID, TSubclassOf<UNavArea> AreaClass);
 
 	/** Sets area ID for the specified polygons */
-	void SetPolyArrayAreaID(const TArray<FNavPoly>& Polys, uint8 AreaID);
+	void SetPolyArrayArea(const TArray<FNavPoly>& Polys, TSubclassOf<UNavArea> AreaClass);
 
 	/** Retrieves poly and area flags for specified polygon */
 	bool GetPolyFlags(NavNodeRef PolyID, uint16& PolyFlags, uint16& AreaFlags) const;
@@ -1025,6 +1025,13 @@ private:
 private:
 	static const FRecastQueryFilter* NamedFilters[ERecastNamedFilter::NamedFiltersCount];
 
+	// DEPRECATED FUNCTIONS
+public:
+	DEPRECATED(4.11, "Function was changed to SetPolyArea, using AreaClass pointer instead of Id")
+	void SetPolyAreaID(NavNodeRef PolyID, uint8 AreaID);
+	DEPRECATED(4.11, "Function was changed to SetPolyArrayArea, using AreaClass pointer instead of Id")
+	void SetPolyArrayAreaID(const TArray<FNavPoly>& Polys, uint8 AreaID);
+
 #endif // WITH_RECAST
 };
 
diff --git a/Engine/Source/Runtime/Engine/Classes/Animation/AnimNode_StateMachine.h b/Engine/Source/Runtime/Engine/Classes/Animation/AnimNode_StateMachine.h
--- a/Engine/Source/Runtime/Engine/Classes/Animation/AnimNode_StateMachine.h
+++ b/Engine/Source/Runtime/Engine/Classes/Animation/AnimNode_StateMachine.h
@@ -172,6 +172,8 @@ private:
 
 	TArray<FPoseContext*> StateCachedPoses;
 
+	TArray<FGraphTraversalCounter> StateCacheBoneCounters;
+
 public:
 	FAnimNode_StateMachine()
 		: MaxTransitionsPerFrame(3)
@@ -189,6 +191,8 @@ public:
 	virtual void GatherDebugData(FNodeDebugData& DebugData) override;
 	// End of FAnimNode_Base interface
 
+	void ConditionallyCacheBonesForState(int32 StateIndex, FAnimationBaseContext Context);
+
 	// Returns the blend weight of the specified state, as calculated by the last call to Update()
 	float GetStateWeight(int32 StateIndex) const;
 
diff --git a/Engine/Source/Runtime/Engine/Classes/Animation/AnimNotifies/AnimNotify_PlayParticleEffect.h b/Engine/Source/Runtime/Engine/Classes/Animation/AnimNotifies/AnimNotify_PlayParticleEffect.h
--- a/Engine/Source/Runtime/Engine/Classes/Animation/AnimNotifies/AnimNotify_PlayParticleEffect.h
+++ b/Engine/Source/Runtime/Engine/Classes/Animation/AnimNotifies/AnimNotify_PlayParticleEffect.h
@@ -5,7 +5,7 @@
 #include "AnimNotify.h"
 #include "AnimNotify_PlayParticleEffect.generated.h"
 
-UCLASS(const, hidecategories=Object, collapsecategories, meta=(DisplayName="Play Particle Effect"))
+UCLASS(const, hidecategories=Object, collapsecategories, meta=(DisplayName="Play Particle Effect"), MinimalAPI)
 class UAnimNotify_PlayParticleEffect : public UAnimNotify
 {
 	GENERATED_BODY()
diff --git a/Engine/Source/Runtime/Engine/Classes/Animation/AnimNotifies/AnimNotify_PlaySound.h b/Engine/Source/Runtime/Engine/Classes/Animation/AnimNotifies/AnimNotify_PlaySound.h
--- a/Engine/Source/Runtime/Engine/Classes/Animation/AnimNotifies/AnimNotify_PlaySound.h
+++ b/Engine/Source/Runtime/Engine/Classes/Animation/AnimNotifies/AnimNotify_PlaySound.h
@@ -5,7 +5,7 @@
 #include "AnimNotify.h"
 #include "AnimNotify_PlaySound.generated.h"
 
-UCLASS(const, hidecategories=Object, collapsecategories, meta=(DisplayName="Play Sound"))
+UCLASS(const, hidecategories=Object, collapsecategories, meta=(DisplayName="Play Sound"), MinimalAPI)
 class UAnimNotify_PlaySound : public UAnimNotify
 {
 	GENERATED_BODY()
diff --git a/Engine/Source/Runtime/Engine/Classes/Camera/CameraShake.h b/Engine/Source/Runtime/Engine/Classes/Camera/CameraShake.h
--- a/Engine/Source/Runtime/Engine/Classes/Camera/CameraShake.h
+++ b/Engine/Source/Runtime/Engine/Classes/Camera/CameraShake.h
@@ -249,6 +249,9 @@ public:
 	virtual void PlayShake(class APlayerCameraManager* Camera, float Scale, ECameraAnimPlaySpace::Type InPlaySpace, FRotator UserPlaySpaceRot = FRotator::ZeroRotator);
 	virtual bool IsFinished() const;
 	virtual void StopShake();
+
+	// Returns true if this camera shake will loop forever
+	bool IsLooping() const;
 };
 
 
diff --git a/Engine/Source/Runtime/Engine/Classes/Curves/CurveBase.h b/Engine/Source/Runtime/Engine/Classes/Curves/CurveBase.h
--- a/Engine/Source/Runtime/Engine/Classes/Curves/CurveBase.h
+++ b/Engine/Source/Runtime/Engine/Classes/Curves/CurveBase.h
@@ -290,12 +290,13 @@ public:
 	FKeyHandle SetKeyTime(FKeyHandle KeyHandle, float NewTime);
 
 	/**
-	 * Finds the key at InTime, and updates its value. If it can't find the key, it adds one at that time.
+	 * Finds the key at InTime, and updates its value. If it can't find the key within the KeyTimeTolerance, it adds one at that time.
 	 *
 	 * @param InTime The time at which the key should be added or updated.
 	 * @param InValue The value of the key.
+	 * @param KeyTimeTolerance The tolerance used for key time equality.
 	 */
-	FKeyHandle UpdateOrAddKey(float InTime, const FName& InValue);
+	FKeyHandle UpdateOrAddKey(float InTime, const FName& InValue, float KeyTimeTolerance = KINDA_SMALL_NUMBER);
 
 public:
 
@@ -534,8 +535,8 @@ public:
 	/** Remove the specified key from the curve.*/
 	void DeleteKey(FKeyHandle KeyHandle);
 
-	/** Finds the key at InTime, and updates its value. If it can't find the key, it adds one at that time */
-	FKeyHandle UpdateOrAddKey(float InTime, float InValue, const bool bUnwindRotation = false);
+	/** Finds the key at InTime, and updates its value. If it can't find the key within the KeyTimeTolerance, it adds one at that time */
+	FKeyHandle UpdateOrAddKey(float InTime, float InValue, const bool bUnwindRotation = false, float KeyTimeTolerance = KINDA_SMALL_NUMBER);
 
 	/** Move a key to a new time. This may change the index of the key, so the new key index is returned. */
 	FKeyHandle SetKeyTime(FKeyHandle KeyHandle, float NewTime);
@@ -877,8 +878,8 @@ public:
 	/** Remove the specified key from the curve.*/
 	void DeleteKey(FKeyHandle KeyHandle);
 	
-	/** Finds the key at InTime, and updates its value. If it can't find the key, it adds one at that time */
-	FKeyHandle UpdateOrAddKey( float Time, int32 Value );
+	/** Finds the key at InTime, and updates its value. If it can't find the key within the KeyTimeTolerance, it adds one at that time */
+	FKeyHandle UpdateOrAddKey( float InTime, int32 Value, float KeyTimeTolerance = KINDA_SMALL_NUMBER );
 	
 	/** Move a key to a new time. This may change the index of the key, so the new key index is returned. */
 	FKeyHandle SetKeyTime(FKeyHandle KeyHandle, float NewTime);
@@ -910,7 +911,7 @@ public:
 	FIntegralKey& GetKey(FKeyHandle KeyHandle);
 	FIntegralKey GetKey(FKeyHandle KeyHandle) const;
 
-	FKeyHandle FindKey(float KeyTime) const;
+	FKeyHandle FindKey(float KeyTime, float KeyTimeTolerance = KINDA_SMALL_NUMBER) const;
 
 	/** Gets the handle for the last key which is at or before the time requested.  If there are no keys at or before the requested time, an invalid handle is returned. */
 	FKeyHandle FindKeyBeforeOrAt(float KeyTime) const;
diff --git a/Engine/Source/Runtime/Engine/Classes/Engine/NetDriver.h b/Engine/Source/Runtime/Engine/Classes/Engine/NetDriver.h
--- a/Engine/Source/Runtime/Engine/Classes/Engine/NetDriver.h
+++ b/Engine/Source/Runtime/Engine/Classes/Engine/NetDriver.h
@@ -17,14 +17,26 @@ class FObjectReplicator;
 //
 #define DO_ENABLE_NET_TEST !(UE_BUILD_SHIPPING || UE_BUILD_TEST)
 
-#if DO_ENABLE_NET_TEST
+
 /** Holds the packet simulation settings in one place */
-struct FPacketSimulationSettings
+USTRUCT()
+struct ENGINE_API FPacketSimulationSettings
 {
+	GENERATED_USTRUCT_BODY()
+
+	UPROPERTY(EditAnywhere, Category="Simulation Settings")
 	int32	PktLoss;
+
+	UPROPERTY(EditAnywhere, Category="Simulation Settings")
 	int32	PktOrder;
+
+	UPROPERTY(EditAnywhere, Category="Simulation Settings")
 	int32	PktDup;
+	
+	UPROPERTY(EditAnywhere, Category="Simulation Settings")
 	int32	PktLag;
+	
+	UPROPERTY(EditAnywhere, Category="Simulation Settings")
 	int32	PktLagVariance;
 
 	/** Ctor. Zeroes the settings */
@@ -59,7 +71,6 @@ struct FPacketSimulationSettings
 	 */
 	bool ParseSettings(const TCHAR* Stream);
 };
-#endif
 
 //
 // Priority sortable list.
@@ -327,6 +338,8 @@ public:
 
 #if DO_ENABLE_NET_TEST
 	FPacketSimulationSettings	PacketSimulationSettings;
+
+	ENGINE_API void SetPacketSimulationSettings(FPacketSimulationSettings NewSettings);
 #endif
 
 	// Constructors.
diff --git a/Engine/Source/Runtime/Engine/Classes/Particles/EmitterCameraLensEffectBase.h b/Engine/Source/Runtime/Engine/Classes/Particles/EmitterCameraLensEffectBase.h
--- a/Engine/Source/Runtime/Engine/Classes/Particles/EmitterCameraLensEffectBase.h
+++ b/Engine/Source/Runtime/Engine/Classes/Particles/EmitterCameraLensEffectBase.h
@@ -66,9 +66,14 @@ public:
 	/** This will actually activate the lens Effect.  We want this separated from PostInitializeComponents so we can cache these emitters **/
 	virtual void ActivateLensEffect();
 	
+	/** Deactivtes the particle system. If bDestroyOnSystemFinish is true, actor will die after particles are all dead. */
+	virtual void DeactivateLensEffect();
+	
 	/** Given updated camera information, adjust this effect to display appropriately. */
 	virtual void UpdateLocation(const FVector& CamLoc, const FRotator& CamRot, float CamFOVDeg);
 
+	/** Returns true if either particle system would loop forever when played */
+	bool IsLooping() const;
 private:
 	/** DEPRECATED(4.11) */
 	UPROPERTY()
diff --git a/Engine/Source/Runtime/Engine/Classes/Particles/Light/ParticleModuleLight.h b/Engine/Source/Runtime/Engine/Classes/Particles/Light/ParticleModuleLight.h
--- a/Engine/Source/Runtime/Engine/Classes/Particles/Light/ParticleModuleLight.h
+++ b/Engine/Source/Runtime/Engine/Classes/Particles/Light/ParticleModuleLight.h
@@ -49,6 +49,12 @@ class UParticleModuleLight : public UParticleModuleLightBase
 	UPROPERTY(EditAnywhere, Category=Light)
 	struct FRawDistributionFloat LightExponent;
 
+	UPROPERTY(EditAnywhere, Category = Light)
+	bool bHighQualityLights;
+
+	UPROPERTY(EditAnywhere, Category = Light)
+	bool bShadowCastingLights;	
+
 	/** Initializes the default values for this property */
 	void InitializeDefaults();
 
@@ -72,6 +78,11 @@ class UParticleModuleLight : public UParticleModuleLightBase
 	void SpawnEx(FParticleEmitterInstance* Owner, int32 Offset, float SpawnTime, struct FRandomStream* InRandomStream, FBaseParticle* ParticleBase);
 
 	virtual bool CanTickInAnyThread() override;
+
+private:
+		
+	uint64 SpawnHQLight(const FLightParticlePayload& Payload, const FBaseParticle& Particle, FParticleEmitterInstance* Owner);	
+	void UpdateHQLight(UPointLightComponent* PointLightComponent, const FLightParticlePayload& Payload, const FBaseParticle& Particle, int32 ScreenAlignment, FVector ComponentScale, bool bLocalSpace, FSceneInterface* OwnerScene, bool bDoRTUpdate);
 };
 
 
diff --git a/Engine/Source/Runtime/Engine/Classes/Particles/ParticleSystem.h b/Engine/Source/Runtime/Engine/Classes/Particles/ParticleSystem.h
--- a/Engine/Source/Runtime/Engine/Classes/Particles/ParticleSystem.h
+++ b/Engine/Source/Runtime/Engine/Classes/Particles/ParticleSystem.h
@@ -467,6 +467,12 @@ public:
 	UFUNCTION(BlueprintCallable, Category = "Particle System")
 	bool ContainsEmitterType(UClass* TypeData);
 
+	/** Returns true if the particle system is looping (contains one or more looping emitters) */
+	bool IsLooping() const { return bAnyEmitterLoopsForever; }
+
+private:
+	/** Does any emitter loop forever (computed during load and when emitters are built) */
+	bool bAnyEmitterLoopsForever;
 };
 
 
diff --git a/Engine/Source/Runtime/Engine/Private/AI/Navigation/NavAreas/NavArea_Null.cpp b/Engine/Source/Runtime/Engine/Private/AI/Navigation/NavAreas/NavArea_Null.cpp
--- a/Engine/Source/Runtime/Engine/Private/AI/Navigation/NavAreas/NavArea_Null.cpp
+++ b/Engine/Source/Runtime/Engine/Private/AI/Navigation/NavAreas/NavArea_Null.cpp
@@ -6,4 +6,5 @@
 UNavArea_Null::UNavArea_Null(const FObjectInitializer& ObjectInitializer) : Super(ObjectInitializer)
 {
 	DefaultCost = BIG_NUMBER;
+	AreaFlags = 0;
 }
diff --git a/Engine/Source/Runtime/Engine/Private/AI/Navigation/NavigationSystem.cpp b/Engine/Source/Runtime/Engine/Private/AI/Navigation/NavigationSystem.cpp
--- a/Engine/Source/Runtime/Engine/Private/AI/Navigation/NavigationSystem.cpp
+++ b/Engine/Source/Runtime/Engine/Private/AI/Navigation/NavigationSystem.cpp
@@ -2141,7 +2141,6 @@ UNavigationSystem* UNavigationSystem::CreateNavigationSystem(UWorld* WorldOwner)
 {
 	UNavigationSystem* NavSys = NULL;
 
-#if WITH_SERVER_CODE || WITH_EDITOR
 	// create navigation system for editor and server targets, but remove it from game clients
 	if (WorldOwner && (*GEngine->NavigationSystemClass != nullptr) 
 		&& (GEngine->NavigationSystemClass->GetDefaultObject<UNavigationSystem>()->bAllowClientSideNavigation || WorldOwner->GetNetMode() != NM_Client))
@@ -2153,7 +2152,6 @@ UNavigationSystem* UNavigationSystem::CreateNavigationSystem(UWorld* WorldOwner)
 			WorldOwner->SetNavigationSystem(NavSys);
 		}
 	}
-#endif
 
 	return NavSys;
 }
diff --git a/Engine/Source/Runtime/Engine/Private/AI/Navigation/RecastNavMesh.cpp b/Engine/Source/Runtime/Engine/Private/AI/Navigation/RecastNavMesh.cpp
--- a/Engine/Source/Runtime/Engine/Private/AI/Navigation/RecastNavMesh.cpp
+++ b/Engine/Source/Runtime/Engine/Private/AI/Navigation/RecastNavMesh.cpp
@@ -1274,20 +1274,60 @@ uint32 ARecastNavMesh::GetPolyAreaID(NavNodeRef PolyID) const
 	return AreaID;
 }
 
+void ARecastNavMesh::SetPolyArea(NavNodeRef PolyID, TSubclassOf<UNavArea> AreaClass)
+{
+	if (AreaClass && RecastNavMeshImpl)
+	{
+		dtNavMesh* NavMesh = RecastNavMeshImpl->GetRecastMesh();
+		const int32 AreaId = GetAreaID(AreaClass);
+		const uint16 AreaFlags = AreaClass->GetDefaultObject<UNavArea>()->GetAreaFlags();
+		
+		if (AreaId != INDEX_NONE && NavMesh)
+		{
+			NavMesh->setPolyArea(PolyID, AreaId);
+			NavMesh->setPolyFlags(PolyID, AreaId);
+		}
+	}
+}
+
+void ARecastNavMesh::SetPolyArrayArea(const TArray<FNavPoly>& Polys, TSubclassOf<UNavArea> AreaClass)
+{
+	if (AreaClass && RecastNavMeshImpl)
+	{
+		dtNavMesh* NavMesh = RecastNavMeshImpl->GetRecastMesh();
+		const int32 AreaId = GetAreaID(AreaClass);
+		const uint16 AreaFlags = AreaClass->GetDefaultObject<UNavArea>()->GetAreaFlags();
+
+		if (AreaId != INDEX_NONE && NavMesh)
+		{
+			for (int32 Idx = 0; Idx < Polys.Num(); Idx++)
+			{
+				NavMesh->setPolyArea(Polys[Idx].Ref, AreaId);
+				NavMesh->setPolyFlags(Polys[Idx].Ref, AreaFlags);
+			}
+		}
+	}
+}
+
+// deprecated
 void ARecastNavMesh::SetPolyAreaID(NavNodeRef PolyID, uint8 AreaID)
 {
-	if (RecastNavMeshImpl)
+	const UClass* NavAreaClass = GetAreaClass(AreaID);
+	if (NavAreaClass)
 	{
-		RecastNavMeshImpl->SetPolyAreaID(PolyID, AreaID);
+		TSubclassOf<UNavArea> UseAreaClass((UClass*)NavAreaClass);
+		SetPolyArea(PolyID, UseAreaClass);
 	}
 }
 
+// deprecated
 void ARecastNavMesh::SetPolyArrayAreaID(const TArray<FNavPoly>& Polys, uint8 AreaID)
 {
-	dtNavMesh* NavMesh = RecastNavMeshImpl->GetRecastMesh();
-	for(const auto& Poly : Polys)
+	const UClass* NavAreaClass = GetAreaClass(AreaID);
+	if (NavAreaClass)
 	{
-		NavMesh->setPolyArea(Poly.Ref, AreaID);
+		TSubclassOf<UNavArea> UseAreaClass((UClass*)NavAreaClass);
+		SetPolyArrayArea(Polys, UseAreaClass);
 	}
 }
 
diff --git a/Engine/Source/Runtime/Engine/Private/Actor.cpp b/Engine/Source/Runtime/Engine/Private/Actor.cpp
--- a/Engine/Source/Runtime/Engine/Private/Actor.cpp
+++ b/Engine/Source/Runtime/Engine/Private/Actor.cpp
@@ -29,6 +29,7 @@
 DEFINE_LOG_CATEGORY(LogActor);
 
 DEFINE_STAT(STAT_GetComponentsTime);
+DECLARE_CYCLE_STAT(TEXT("PostActorConstruction"), STAT_PostActorConstruction, STATGROUP_Engine);
 
 #if UE_BUILD_SHIPPING
 #define DEBUG_CALLSPACE(Format, ...)
@@ -2631,7 +2632,11 @@ void AActor::FinishSpawning(const FTransform& UserTransform, bool bIsDefaultTran
 		}
 
 		ExecuteConstruction(FinalRootComponentTransform, nullptr, bIsDefaultTransform);
-		PostActorConstruction();
+
+		{
+			SCOPE_CYCLE_COUNTER(STAT_PostActorConstruction);
+			PostActorConstruction();
+		}
 	}
 }
 
@@ -2711,6 +2716,7 @@ void AActor::PostActorConstruction()
 
 			if (World->HasBegunPlay() && !deferBeginPlayAndUpdateOverlaps)
 			{
+				SCOPE_CYCLE_COUNTER(STAT_ActorBeginPlay);
 				BeginPlay();
 			}
 		}
@@ -2798,6 +2804,7 @@ void AActor::PostNetInit()
 		const UWorld* MyWorld = GetWorld();
 		if (MyWorld && MyWorld->HasBegunPlay())
 		{
+			SCOPE_CYCLE_COUNTER(STAT_ActorBeginPlay);
 			BeginPlay();
 		}
 	}
diff --git a/Engine/Source/Runtime/Engine/Private/Animation/AnimInstanceProxy.cpp b/Engine/Source/Runtime/Engine/Private/Animation/AnimInstanceProxy.cpp
--- a/Engine/Source/Runtime/Engine/Private/Animation/AnimInstanceProxy.cpp
+++ b/Engine/Source/Runtime/Engine/Private/Animation/AnimInstanceProxy.cpp
@@ -297,6 +297,17 @@ void FAnimInstanceProxy::TickAssetPlayerInstances(float DeltaSeconds)
 			// we found leader
 			SyncGroup.Finalize(PreviousGroup);
 
+			if (TickContext.CanUseMarkerPosition())
+			{
+				const FMarkerSyncAnimPosition& MarkerStart = TickContext.MarkerTickContext.GetMarkerSyncStartPosition();
+				FName SyncGroupName = AnimBlueprintClass->SyncGroupNames[GroupIndex];
+				FAnimTickRecord& GroupLeader = SyncGroup.ActivePlayers[SyncGroup.GroupLeaderIndex];
+				FString LeaderAnimName = GroupLeader.SourceAsset->GetName();
+
+				checkf(MarkerStart.PreviousMarkerName == NAME_None || SyncGroup.ValidMarkers.Contains(MarkerStart.PreviousMarkerName), TEXT("Prev Marker name not valid for sync group. Marker %s : SyncGroupName %s : Leader %s"), *MarkerStart.PreviousMarkerName.ToString(), *SyncGroupName.ToString(), *LeaderAnimName);
+				checkf(MarkerStart.NextMarkerName == NAME_None || SyncGroup.ValidMarkers.Contains(MarkerStart.NextMarkerName), TEXT("Next Marker name not valid for sync group. Marker %s : SyncGroupName %s : Leader %s"), *MarkerStart.PreviousMarkerName.ToString(), *SyncGroupName.ToString(), *LeaderAnimName);
+			}
+
 			// Update everything else to follow the leader, if there is more followers
 			if (SyncGroup.ActivePlayers.Num() > GroupLeaderIndex + 1)
 			{
diff --git a/Engine/Source/Runtime/Engine/Private/Animation/AnimMontage.cpp b/Engine/Source/Runtime/Engine/Private/Animation/AnimMontage.cpp
--- a/Engine/Source/Runtime/Engine/Private/Animation/AnimMontage.cpp
+++ b/Engine/Source/Runtime/Engine/Private/Animation/AnimMontage.cpp
@@ -1962,7 +1962,6 @@ void FAnimMontageInstance::SetMatineeAnimPositionInner(FName SlotName, USkeletal
 			{
 				SingleNodeInst->SetAnimationAsset(InAnimSequence, bLooping);
 				SingleNodeInst->SetPosition(0.0f);
-				SingleNodeInst->SetPlaying(false);
 			}
 
 			if(SingleNodeInst->IsLooping()!=bLooping)
diff --git a/Engine/Source/Runtime/Engine/Private/Animation/AnimNode_StateMachine.cpp b/Engine/Source/Runtime/Engine/Private/Animation/AnimNode_StateMachine.cpp
--- a/Engine/Source/Runtime/Engine/Private/Animation/AnimNode_StateMachine.cpp
+++ b/Engine/Source/Runtime/Engine/Private/Animation/AnimNode_StateMachine.cpp
@@ -259,6 +259,12 @@ void FAnimNode_StateMachine::Initialize(const FAnimationInitializeContext& Conte
 			// Reset transition related variables
 			StatesUpdated.Reset();
 			ActiveTransitionArray.Reset();
+
+			if (StateCacheBoneCounters.Num() != Machine->States.Num())
+			{
+				StateCacheBoneCounters.Reset(Machine->States.Num());
+				StateCacheBoneCounters.AddDefaulted(Machine->States.Num());
+			}
 		
 			// Move to the default state
 			SetState(Context, Machine->InitialState);
@@ -275,9 +281,9 @@ void FAnimNode_StateMachine::CacheBones(const FAnimationCacheBonesContext& Conte
 	{
 		for (int32 StateIndex = 0; StateIndex < Machine->States.Num(); ++StateIndex)
 		{
-			if( GetStateWeight(StateIndex) > 0.f) 
+			if (GetStateWeight(StateIndex) > 0.f)
 			{
-				StatePoseLinks[StateIndex].CacheBones(Context);
+				ConditionallyCacheBonesForState(StateIndex, Context);
 			}
 		}
 	}
@@ -285,6 +291,20 @@ void FAnimNode_StateMachine::CacheBones(const FAnimationCacheBonesContext& Conte
 	// @TODO GetStateWeight is O(N) transitions.
 }
 
+void FAnimNode_StateMachine::ConditionallyCacheBonesForState(int32 StateIndex, FAnimationBaseContext Context)
+{
+	// Only call CacheBones when needed.
+	check(StateCacheBoneCounters.IsValidIndex(StateIndex));
+	if (!StateCacheBoneCounters[StateIndex].IsSynchronizedWith(Context.AnimInstanceProxy->GetCachedBonesCounter()))
+	{
+		// keep track of states that have had CacheBones called on.
+		StateCacheBoneCounters[StateIndex].SynchronizeWith(Context.AnimInstanceProxy->GetCachedBonesCounter());
+
+		FAnimationCacheBonesContext CacheBoneContext(Context.AnimInstanceProxy);
+		StatePoseLinks[StateIndex].CacheBones(CacheBoneContext);
+	}
+}
+
 const FBakedAnimationState& FAnimNode_StateMachine::GetStateInfo() const
 {
 	return PRIVATE_MachineDescription->States[CurrentState];
@@ -709,6 +729,9 @@ void FAnimNode_StateMachine::Evaluate(FPoseContext& Output)
 	}
 	else if (!IsAConduitState(CurrentState))
 	{
+		// Make sure CacheBones has been called before evaluating.
+		ConditionallyCacheBonesForState(CurrentState, Output);
+
 		// Evaluate the current state
 		StatePoseLinks[CurrentState].Evaluate(Output);
 	}
@@ -823,6 +846,7 @@ void FAnimNode_StateMachine::SetStateInternal(int32 NewStateIndex)
 	checkSlow(PRIVATE_MachineDescription);
 	ensure(!IsAConduitState(NewStateIndex));
 	CurrentState = FMath::Clamp<int32>(NewStateIndex, 0, PRIVATE_MachineDescription->States.Num() - 1);
+	check(CurrentState == NewStateIndex);
 	ElapsedTime = 0.0f;
 }
 
@@ -859,9 +883,8 @@ void FAnimNode_StateMachine::SetState(const FAnimationBaseContext& Context, int3
 			FAnimationInitializeContext InitContext(Context.AnimInstanceProxy);
 			StatePoseLinks[NewStateIndex].Initialize(InitContext);
 
-			// Also update BoneCaching.
-			FAnimationCacheBonesContext CacheBoneContext(Context.AnimInstanceProxy);
-			StatePoseLinks[NewStateIndex].CacheBones(CacheBoneContext);
+			// Also call cache bones if needed
+			ConditionallyCacheBonesForState(NewStateIndex, Context);
 		}
 
 		if(CurrentState != INDEX_NONE && CurrentState < OnGraphStatesEntered.Num())
@@ -950,6 +973,9 @@ const FPoseContext& FAnimNode_StateMachine::EvaluateState(int32 StateIndex, cons
 
 		if (!IsAConduitState(StateIndex))
 		{
+			// Make sure CacheBones has been called before evaluating.
+			ConditionallyCacheBonesForState(StateIndex, Context);
+
 			StatePoseLinks[StateIndex].Evaluate(*CachePosePtr);
 		}
 	}
diff --git a/Engine/Source/Runtime/Engine/Private/Animation/AnimSequence.cpp b/Engine/Source/Runtime/Engine/Private/Animation/AnimSequence.cpp
--- a/Engine/Source/Runtime/Engine/Private/Animation/AnimSequence.cpp
+++ b/Engine/Source/Runtime/Engine/Private/Animation/AnimSequence.cpp
@@ -3933,6 +3933,7 @@ void AdvanceMarkerBackwards(int32& Marker, FName MarkerToFind, bool bLooping, co
 
 bool MarkerMatchesPosition(const UAnimSequence* Sequence, int32 MarkerIndex, FName CorrectMarker)
 {
+	checkf(MarkerIndex != MarkerIndexSpecialValues::Unitialized, TEXT("Uninitialized marker supplied to MarkerMatchesPosition. Anim: %s Expecting marker %s"), *Sequence->GetName(), *CorrectMarker.ToString());
 	return MarkerIndex == MarkerIndexSpecialValues::AnimationBoundary || CorrectMarker == Sequence->AuthoredSyncMarkers[MarkerIndex].MarkerName;
 }
 
@@ -4183,7 +4184,8 @@ void UAnimSequence::GetMarkerIndicesForPosition(const FMarkerSyncAnimPosition& S
 				return;
 			}
 		}
-		check(false); // Should have found a marker above!
+		// Should have found a marker above!
+		checkf(false, TEXT("Next Marker not found in GetMarkerIndicesForPosition. Anim: %s Expecting marker %s"), *GetName(), *SyncPosition.NextMarkerName.ToString());
 	}
 
 	if (SyncPosition.NextMarkerName == NAME_None)
@@ -4201,7 +4203,8 @@ void UAnimSequence::GetMarkerIndicesForPosition(const FMarkerSyncAnimPosition& S
 				return;
 			}
 		}
-		check(false); // Should have found a marker above!
+		// Should have found a marker above!
+		checkf(false, TEXT("Previous Marker not found in GetMarkerIndicesForPosition. Anim: %s Expecting marker %s"), *GetName(), *SyncPosition.PreviousMarkerName.ToString());
 	}
 
 	float DiffToCurrentTime = FLT_MAX;
@@ -4243,6 +4246,9 @@ void UAnimSequence::GetMarkerIndicesForPosition(const FMarkerSyncAnimPosition& S
 			}
 		}
 	}
+	// Should have found a markers above!
+	checkf(OutPrevMarker.MarkerIndex != MarkerIndexSpecialValues::Unitialized, TEXT("Prev Marker not found in GetMarkerIndicesForPosition. Anim: %s Expecting marker %s"), *GetName(), *SyncPosition.PreviousMarkerName.ToString());
+	checkf(OutNextMarker.MarkerIndex != MarkerIndexSpecialValues::Unitialized, TEXT("Next Marker not found in GetMarkerIndicesForPosition. Anim: %s Expecting marker %s"), *GetName(), *SyncPosition.NextMarkerName.ToString());
 }
 
 float UAnimSequence::GetFirstMatchingPosFromMarkerSyncPos(const FMarkerSyncAnimPosition& InMarkerSyncGroupPosition) const
diff --git a/Engine/Source/Runtime/Engine/Private/AnimationAsset.cpp b/Engine/Source/Runtime/Engine/Private/AnimationAsset.cpp
--- a/Engine/Source/Runtime/Engine/Private/AnimationAsset.cpp
+++ b/Engine/Source/Runtime/Engine/Private/AnimationAsset.cpp
@@ -104,7 +104,7 @@ void FAnimGroupInstance::Prepare(const FAnimGroupInstance* PreviousGroup)
 						// Found previous record for "us"
 						if (PrevRecord.SourceAsset != Candidate.SourceAsset)
 						{
-							Candidate.MarkerTickRecord->Reset(); // Changed animation, clPear our cached data
+							Candidate.MarkerTickRecord->Reset(); // Changed animation, clear our cached data
 						}
 						bCandidateFound = true;
 						break;
diff --git a/Engine/Source/Runtime/Engine/Private/Blueprint.cpp b/Engine/Source/Runtime/Engine/Private/Blueprint.cpp
--- a/Engine/Source/Runtime/Engine/Private/Blueprint.cpp
+++ b/Engine/Source/Runtime/Engine/Private/Blueprint.cpp
@@ -749,6 +749,40 @@ void UBlueprint::GetAssetRegistryTags(TArray<FAssetRegistryTag>& OutTags) const
 
 		OutTags.Add( FAssetRegistryTag("FiBData", FFindInBlueprintSearchManager::Get().QuerySingleBlueprint((UBlueprint*)this, false), FAssetRegistryTag::TT_Hidden) );
 	}
+
+	// Only show for strict blueprints (not animation or widget blueprints)
+	// Note: Can't be an Actor specific check on the gen class, as CB only queries the CDO for the majority type to determine which columns are shown in the view
+	if (ExactCast<UBlueprint>(this) != nullptr)
+	{
+		// Determine how many inherited native components exist
+		int32 NumNativeComponents = 0;
+		if (BlueprintClass != nullptr)
+		{
+			TArray<UObject*> PotentialComponents;
+			BlueprintClass->GetDefaultObjectSubobjects(/*out*/ PotentialComponents);
+
+			for (UObject* TestSubObject : PotentialComponents)
+			{
+				if (Cast<UActorComponent>(TestSubObject) != nullptr)
+				{
+					++NumNativeComponents;
+				}
+			}
+		}
+		OutTags.Add(FAssetRegistryTag("NativeComponents", FString::FromInt(NumNativeComponents), UObject::FAssetRegistryTag::TT_Numerical));
+
+		// Determine how many components are added via a SimpleConstructionScript (both newly introduced and inherited from parent BPs)
+		int32 NumAddedComponents = 0;
+		for (UBlueprintGeneratedClass* TestBPClass = BlueprintClass; TestBPClass != nullptr; TestBPClass = Cast<UBlueprintGeneratedClass>(TestBPClass->GetSuperClass()))
+		{
+			const UBlueprint* AssociatedBP = Cast<const UBlueprint>(TestBPClass->ClassGeneratedBy);
+			if (AssociatedBP->SimpleConstructionScript != nullptr)
+			{
+				NumAddedComponents += AssociatedBP->SimpleConstructionScript->GetAllNodesConst().Num();
+			}
+		}
+		OutTags.Add(FAssetRegistryTag("BlueprintComponents", FString::FromInt(NumAddedComponents), UObject::FAssetRegistryTag::TT_Numerical));
+	}
 }
 
 FString UBlueprint::GetFriendlyName() const
diff --git a/Engine/Source/Runtime/Engine/Private/Camera/CameraAnimInst.cpp b/Engine/Source/Runtime/Engine/Private/Camera/CameraAnimInst.cpp
--- a/Engine/Source/Runtime/Engine/Private/Camera/CameraAnimInst.cpp
+++ b/Engine/Source/Runtime/Engine/Private/Camera/CameraAnimInst.cpp
@@ -162,6 +162,8 @@ void UCameraAnimInst::SetScale(float NewScale)
 	BasePlayScale = NewScale;
 }
 
+static const FName NAME_CameraComponentFieldOfViewPropertyName(TEXT("CameraComponent.FieldOfView"));
+
 void UCameraAnimInst::Play(UCameraAnim* Anim, class AActor* CamActor, float InRate, float InScale, float InBlendInTime, float InBlendOutTime, bool bInLooping, bool bRandomStartTime, float Duration)
 {
 	if (Anim && Anim->CameraInterpGroup)
@@ -219,12 +221,17 @@ void UCameraAnimInst::Play(UCameraAnim* Anim, class AActor* CamActor, float InRa
 			for (int32 Idx = 0; Idx < InterpGroupInst->TrackInst.Num(); ++Idx)
 			{
 				UInterpTrackFloatProp* const FloatTrack = Cast<UInterpTrackFloatProp>(CamAnim->CameraInterpGroup->InterpTracks[Idx]);
-				if ( FloatTrack && (FloatTrack->PropertyName == TEXT("CameraComponent.FieldOfView")) )
+				if (FloatTrack && (FloatTrack->PropertyName == NAME_CameraComponentFieldOfViewPropertyName))
 				{
 					InitialFOV = FloatTrack->EvalSub(0, 0.f);
 				}
 			}
-
+		}
+		else
+		{
+			// make sure these are set in cases where there is no move track
+			InitialCamToWorld = FTransform::Identity;
+			InitialFOV = Anim->BaseFOV;
 		}
 	}
 }
diff --git a/Engine/Source/Runtime/Engine/Private/Camera/CameraShake.cpp b/Engine/Source/Runtime/Engine/Private/Camera/CameraShake.cpp
--- a/Engine/Source/Runtime/Engine/Private/Camera/CameraShake.cpp
+++ b/Engine/Source/Runtime/Engine/Private/Camera/CameraShake.cpp
@@ -293,4 +293,8 @@ bool UCameraShake::ReceiveIsFinished_Implementation() const
 	return true;
 }
 
+bool UCameraShake::IsLooping() const
+{
+	return OscillationDuration < 0.0f;
+}
 
diff --git a/Engine/Source/Runtime/Engine/Private/CurveBase.cpp b/Engine/Source/Runtime/Engine/Private/CurveBase.cpp
--- a/Engine/Source/Runtime/Engine/Private/CurveBase.cpp
+++ b/Engine/Source/Runtime/Engine/Private/CurveBase.cpp
@@ -346,26 +346,26 @@ void FNameCurve::ScaleCurve(float ScaleOrigin, float ScaleFactor, TSet<FKeyHandl
 }
 
 
-FKeyHandle FNameCurve::UpdateOrAddKey(float InTime, const FName& InValue)
+FKeyHandle FNameCurve::UpdateOrAddKey(float InTime, const FName& InValue, float KeyTimeTolerance)
 {
 	// Search for a key that already exists at the time and if found, update its value
 	for (int32 KeyIndex = 0; KeyIndex < Keys.Num(); ++KeyIndex)
 	{
 		float KeyTime = Keys[KeyIndex].Time;
 
+		if (FMath::IsNearlyEqual(KeyTime, InTime, KeyTimeTolerance))
+		{
+			Keys[KeyIndex].Value = InValue;
+
+			return GetKeyHandle(KeyIndex);
+		}
+
 		if (KeyTime > InTime)
 		{
 			// All the rest of the keys exist after the key we want to add
 			// so there is no point in searching
 			break;
 		}
-
-		if (FMath::IsNearlyEqual(KeyTime, InTime))
-		{
-			Keys[KeyIndex].Value = InValue;
-
-			return GetKeyHandle(KeyIndex);
-		}
 	}
 
 	// A key wasnt found, add it now
@@ -674,26 +674,26 @@ void FRichCurve::DeleteKey(FKeyHandle InKeyHandle)
 }
 
 
-FKeyHandle FRichCurve::UpdateOrAddKey(float InTime, float InValue, const bool bUnwindRotation)
+FKeyHandle FRichCurve::UpdateOrAddKey(float InTime, float InValue, const bool bUnwindRotation, float KeyTimeTolerance)
 {
 	// Search for a key that already exists at the time and if found, update its value
 	for (int32 KeyIndex = 0; KeyIndex < Keys.Num(); ++KeyIndex)
 	{
 		float KeyTime = Keys[KeyIndex].Time;
 
+		if (FMath::IsNearlyEqual(KeyTime, InTime, KeyTimeTolerance))
+		{
+			Keys[KeyIndex].Value = InValue;
+
+			return GetKeyHandle(KeyIndex);
+		}
+
 		if (KeyTime > InTime)
 		{
 			// All the rest of the keys exist after the key we want to add
 			// so there is no point in searching
 			break;
 		}
-
-		if (FMath::IsNearlyEqual(KeyTime, InTime))
-		{
-			Keys[KeyIndex].Value = InValue;
-
-			return GetKeyHandle(KeyIndex);
-		}
 	}
 
 	// A key wasnt found, add it now
@@ -1725,28 +1725,28 @@ void FIntegralCurve::DeleteKey(FKeyHandle InKeyHandle)
 }
 
 
-FKeyHandle FIntegralCurve::UpdateOrAddKey( float Time, int32 Value )
+FKeyHandle FIntegralCurve::UpdateOrAddKey( float InTime, int32 Value, float KeyTimeTolerance )
 {
 	for( int32 KeyIndex = 0; KeyIndex < Keys.Num(); ++KeyIndex )
 	{
 		float KeyTime = Keys[KeyIndex].Time;
 
-		if( KeyTime > Time )
+		if (FMath::IsNearlyEqual(KeyTime, InTime, KeyTimeTolerance))
 		{
-			// All the rest of the keys exist after the key we want to add
-			// so there is no point in searching
-			break;
+			Keys[KeyIndex].Value = Value;
+			return GetKeyHandle(KeyIndex);
 		}
 
-		if( KeyTime == Time )
+		if( KeyTime > InTime )
 		{
-			Keys[KeyIndex].Value = Value;
-			return GetKeyHandle(KeyIndex);
+			// All the rest of the keys exist after the key we want to add
+			// so there is no point in searching
+			break;
 		}
 	}
 
 	// A key wasnt found, add it now
-	return AddKey( Time, Value );
+	return AddKey( InTime, Value );
 }
 
 
@@ -1847,7 +1847,7 @@ FIntegralKey FIntegralCurve::GetKey(FKeyHandle KeyHandle) const
 }
 
 
-FKeyHandle FIntegralCurve::FindKey(float KeyTime) const
+FKeyHandle FIntegralCurve::FindKey(float KeyTime, float KeyTimeTolerance) const
 {
 	int32 Start = 0;
 	int32 End = Keys.Num() - 1;
@@ -1858,7 +1858,7 @@ FKeyHandle FIntegralCurve::FindKey(float KeyTime) const
 		int32 TestPos = Start + (End - Start) / 2;
 		float TestKeyTime = Keys[TestPos].Time;
 
-		if (TestKeyTime == KeyTime)
+		if (FMath::IsNearlyEqual(TestKeyTime, KeyTime, KeyTimeTolerance))
 		{
 			return GetKeyHandle(TestPos);
 		}
diff --git a/Engine/Source/Runtime/Engine/Private/DataChannel.cpp b/Engine/Source/Runtime/Engine/Private/DataChannel.cpp
--- a/Engine/Source/Runtime/Engine/Private/DataChannel.cpp
+++ b/Engine/Source/Runtime/Engine/Private/DataChannel.cpp
@@ -25,6 +25,7 @@ DEFINE_LOG_CATEGORY_STATIC(LogNetPartialBunch, Warning, All);
 
 DECLARE_CYCLE_STAT(TEXT("ActorChan_ReceivedBunch"), Stat_ActorChanReceivedBunch, STATGROUP_Net);
 DECLARE_CYCLE_STAT(TEXT("ActorChan_CleanUp"), Stat_ActorChanCleanUp, STATGROUP_Net);
+DECLARE_CYCLE_STAT(TEXT("ActorChan_PostNetInit"), Stat_PostNetInit, STATGROUP_Net);
 
 extern FAutoConsoleVariable CVarDoReplicationContextString;
 
@@ -2085,6 +2086,7 @@ void UActorChannel::ProcessBunch( FInBunch & Bunch )
 	// After all properties have been initialized, call PostNetInit. This should call BeginPlay() so initialization can be done with proper starting values.
 	if (Actor && bSpawnedNewActor)
 	{
+		SCOPE_CYCLE_COUNTER(Stat_PostNetInit);
 		Actor->PostNetInit();
 	}
 }
diff --git a/Engine/Source/Runtime/Engine/Private/DemoNetDriver.cpp b/Engine/Source/Runtime/Engine/Private/DemoNetDriver.cpp
--- a/Engine/Source/Runtime/Engine/Private/DemoNetDriver.cpp
+++ b/Engine/Source/Runtime/Engine/Private/DemoNetDriver.cpp
@@ -1016,17 +1016,8 @@ static void DemoReplicateActor( AActor* Actor, UNetConnection* Connection, APlay
 		AActor* Actor;
 	};
 
-	// All actors marked for replication are assumed to be relevant for demo recording.
-	/*
-	if
-		(	Actor
-		&&	((IsNetClient && Actor->bTearOff) || Actor->RemoteRole != ROLE_None || (IsNetClient && Actor->Role != ROLE_None && Actor->Role != ROLE_Authority) || Actor->bForceDemoRelevant)
-		&&  (!Actor->bNetTemporary || !Connection->SentTemporaries.Contains(Actor))
-		&& (Actor == Connection->PlayerController || Cast<APlayerController>(Actor) == NULL)
-		)
-	*/
-	bool bReplicated = false;
-
+	const int32 OriginalOutBunches = Connection->Driver->OutBunches;
+	
 	if ( Actor != NULL )
 	{
 		// We need to swap roles if:
@@ -1066,12 +1057,8 @@ static void DemoReplicateActor( AActor* Actor, UNetConnection* Connection, APlay
 			if (Channel != NULL && !Channel->Closing)
 			{
 				// Send it out!
-				if (Channel->IsNetReady(0))
-				{
-					Channel->ReplicateActor();
-					bReplicated = true;
-				}
-			
+				Channel->ReplicateActor();
+
 				// Close the channel if this actor shouldn't have one
 				if (!bShouldHaveChannel)
 				{
@@ -1081,9 +1068,9 @@ static void DemoReplicateActor( AActor* Actor, UNetConnection* Connection, APlay
 		}
 	}
 
-	if ( bMustReplicate && !bReplicated )
+	if ( bMustReplicate && Connection->Driver->OutBunches == OriginalOutBunches )
 	{
-		UE_LOG( LogDemo, Warning, TEXT( "DemoReplicateActor: bMustReplicate is true but bReplicated is false" ) );
+		UE_LOG( LogDemo, Error, TEXT( "DemoReplicateActor: bMustReplicate is true but nothing was sent: %s" ), Actor ? *Actor->GetName() : TEXT( "NULL" ) );
 	}
 }
 
@@ -1316,6 +1303,9 @@ void UDemoNetDriver::TickDemoRecord( float DeltaSeconds )
 	*FileAr << DeltaTimeChecksum;
 #endif
 
+	// flush out any pending network traffic
+	ClientConnections[0]->FlushNet();
+
 	// Make sure we don't have anything in the buffer for this new frame
 	check( ClientConnections[0]->SendBuffer.GetNumBits() == 0 );
 
@@ -1884,6 +1874,8 @@ void UDemoNetDriver::LoadCheckpoint( FArchive* GotoCheckpointArchive, int64 Goto
 
 	PauseChannels( false );
 
+	FNetworkReplayDelegates::OnPreScrub.Broadcast(GetWorld());
+
 #if 1
 	// Destroy all non startup actors. They will get restored with the checkpoint
 	for ( FActorIterator It( GetWorld() ); It; ++It )
diff --git a/Engine/Source/Runtime/Engine/Private/LevelTick.cpp b/Engine/Source/Runtime/Engine/Private/LevelTick.cpp
--- a/Engine/Source/Runtime/Engine/Private/LevelTick.cpp
+++ b/Engine/Source/Runtime/Engine/Private/LevelTick.cpp
@@ -44,6 +44,7 @@ DEFINE_STAT(STAT_AsyncWorkWaitTime);
 DEFINE_STAT(STAT_PhysicsTime);
 
 DEFINE_STAT(STAT_SpawnActorTime);
+DEFINE_STAT(STAT_ActorBeginPlay);
 
 DEFINE_STAT(STAT_GCSweepTime);
 DEFINE_STAT(STAT_GCMarkTime);
diff --git a/Engine/Source/Runtime/Engine/Private/NetConnection.cpp b/Engine/Source/Runtime/Engine/Private/NetConnection.cpp
--- a/Engine/Source/Runtime/Engine/Private/NetConnection.cpp
+++ b/Engine/Source/Runtime/Engine/Private/NetConnection.cpp
@@ -576,7 +576,7 @@ void UNetConnection::FlushNet(bool bIgnoreSimulation)
 #if DO_ENABLE_NET_TEST
 		// if the connection is closing/being destroyed/etc we need to send immediately regardless of settings
 		// because we won't be around to send it delayed
-		if (State == USOCK_Closed || IsGarbageCollecting() || bIgnoreSimulation)
+		if (State == USOCK_Closed || IsGarbageCollecting() || bIgnoreSimulation || InternalAck)
 		{
 			// Checked in FlushNet() so each child class doesn't have to implement this
 			if (Driver->IsNetResourceValid())
@@ -996,11 +996,19 @@ void UNetConnection::ReceivedPacket( FBitReader& Reader )
 				const bool ValidUnreliableOpen = Bunch.bOpen && (Bunch.bClose || Bunch.bPartial);
 				if (!ValidUnreliableOpen)
 				{
-					UE_LOG( LogNetTraffic, Warning, TEXT( "      Received unreliable bunch before open (Channel %d Current Sequence %i)" ), Bunch.ChIndex, InReliable[Bunch.ChIndex] );
+					if ( InternalAck )
+					{
+						// Should be impossible with 100% reliable connections
+						UE_LOG( LogNetTraffic, Error, TEXT( "      Received unreliable bunch before open with reliable connection (Channel %d Current Sequence %i)" ), Bunch.ChIndex, InReliable[Bunch.ChIndex] );
+					}
+					else
+					{
+						UE_LOG( LogNetTraffic, Warning, TEXT( "      Received unreliable bunch before open (Channel %d Current Sequence %i)" ), Bunch.ChIndex, InReliable[Bunch.ChIndex] );
+					}
+
 					// Since we won't be processing this packet, don't ack it
 					// We don't want the sender to think this bunch was processed when it really wasn't
 					bSkipAck = true;
-					check( !InternalAck );		// Should be impossible with 100% reliable connections
 					continue;
 				}
 			}
@@ -1063,8 +1071,6 @@ void UNetConnection::ReceivedPacket( FBitReader& Reader )
 
 	ValidateSendBuffer();
 
-	check( !bSkipAck || !InternalAck );		// 100% reliable connections shouldn't be skipping acks
-
 	// Acknowledge the packet.
 	if ( !bSkipAck )
 	{
diff --git a/Engine/Source/Runtime/Engine/Private/NetworkDriver.cpp b/Engine/Source/Runtime/Engine/Private/NetworkDriver.cpp
--- a/Engine/Source/Runtime/Engine/Private/NetworkDriver.cpp
+++ b/Engine/Source/Runtime/Engine/Private/NetworkDriver.cpp
@@ -1735,6 +1735,22 @@ void UNetDriver::AddReferencedObjects(UObject* InThis, FReferenceCollector& Coll
 
 #if DO_ENABLE_NET_TEST
 
+void UNetDriver::SetPacketSimulationSettings(FPacketSimulationSettings NewSettings)
+{
+	PacketSimulationSettings = NewSettings;
+	if (ServerConnection)
+	{
+		ServerConnection->UpdatePacketSimulationSettings();
+	}
+	for (UNetConnection* ClientConnection : ClientConnections)
+	{
+		if (ClientConnection)
+		{
+			ClientConnection->UpdatePacketSimulationSettings();
+		}
+	}
+}
+
 class FPacketSimulationConsoleCommandVisitor 
 {
 public:
