Copying //UE4/Dev-Blueprints to Dev-Main (//UE4/Dev-Main) @ 2781164

diff --git a/Engine/Source/Editor/UnrealEd/Private/Kismet2/Kismet2.cpp b/Engine/Source/Editor/UnrealEd/Private/Kismet2/Kismet2.cpp
--- a/Engine/Source/Editor/UnrealEd/Private/Kismet2/Kismet2.cpp
+++ b/Engine/Source/Editor/UnrealEd/Private/Kismet2/Kismet2.cpp
@@ -41,6 +41,7 @@
 #include "Engine/SimpleConstructionScript.h"
 #include "Engine/SCS_Node.h"
 #include "GeneralProjectSettings.h"
+#include "Developer/BlueprintProfiler/Public/BlueprintProfilerModule.h"
 
 DECLARE_CYCLE_STAT(TEXT("Compile Blueprint"), EKismetCompilerStats_CompileBlueprint, STATGROUP_KismetCompiler);
 DECLARE_CYCLE_STAT(TEXT("Broadcast Precompile"), EKismetCompilerStats_BroadcastPrecompile, STATGROUP_KismetCompiler);
@@ -737,12 +738,18 @@ void FKismetEditorUtilities::CompileBlueprint(UBlueprint* BlueprintObj, bool bIs
 	}
 	auto ReinstanceHelper = FBlueprintCompileReinstancer::Create(OldClass);
 
-	// Suppress errors/warnings in the log if we're recompiling on load on a build machine
-	Results.bLogInfoOnly = BlueprintObj->bIsRegeneratingOnLoad && GIsBuildMachine;
+	// If enabled, suppress errors/warnings in the log if we're recompiling on load on a build machine
+	static const FBoolConfigValueHelper IgnoreCompileOnLoadErrorsOnBuildMachine(TEXT("Kismet"), TEXT("bIgnoreCompileOnLoadErrorsOnBuildMachine"), GEngineIni);
+	Results.bLogInfoOnly = BlueprintObj->bIsRegeneratingOnLoad && GIsBuildMachine && IgnoreCompileOnLoadErrorsOnBuildMachine;
+
+	// Determine if we want profiling data
+	IBlueprintProfilerInterface* ProfilerInterface = FModuleManager::GetModulePtr<IBlueprintProfilerInterface>("BlueprintProfiler");
+	const bool bProfilerActive = ProfilerInterface && ProfilerInterface->IsProfilerEnabled() && GetDefault<UEditorExperimentalSettings>()->bBlueprintPerformanceAnalysisTools;
 
 	FKismetCompilerOptions CompileOptions;
 	CompileOptions.bSaveIntermediateProducts = bSaveIntermediateProducts;
 	CompileOptions.bRegenerateSkelton = !bSkeletonUpToDate;
+	CompileOptions.bAddInstrumentation = bProfilerActive;
 	Compiler.CompileBlueprint(BlueprintObj, CompileOptions, Results, ReinstanceHelper);
 
 	FBlueprintEditorUtils::UpdateDelegatesInBlueprint(BlueprintObj);
@@ -857,26 +864,6 @@ void FKismetEditorUtilities::CompileBlueprint(UBlueprint* BlueprintObj, bool bIs
 	if (BlueprintPackage != NULL)
 	{
 		BlueprintPackage->SetDirtyFlag(bStartedWithUnsavedChanges);
-
-		UBlueprintEditorSettings* Settings = GetMutableDefault<UBlueprintEditorSettings>();
-		const bool bShouldSaveOnCompile = !bIsRegeneratingOnLoad && (( Settings->SaveOnCompile == SoC_Always ) || ( ( Settings->SaveOnCompile == SoC_SuccessOnly ) && ( BlueprintObj->Status == BS_UpToDate ) ));
-
-		// Only try to save on compile if we aren't running a commandlet (i.e. cooking a blueprint shouldn't try to save it)
-		if ( !IsRunningCommandlet() && bShouldSaveOnCompile && !GIsAutomationTesting )
-		{
-			bool const bIsLevelPackage = (UWorld::FindWorldInPackage(BlueprintPackage) != nullptr);
-			// we don't want to save the entire level (especially if this 
-			// compile was already kicked off as a result of a level save, as it
-			// could cause a recursive save)... let the "SaveOnCompile" setting 
-			// only save blueprint assets
-			if (!bIsLevelPackage)
-			{
-				TArray<UPackage*> PackagesToSave;
-				PackagesToSave.Add(BlueprintPackage);
-
-				FEditorFileUtils::PromptForCheckoutAndSave(PackagesToSave, /*bCheckDirty =*/true, /*bPromptToSave =*/false);
-			}			
-		}
 	}	
 }
 
@@ -977,9 +964,9 @@ static void ConformComponentsUtils::ConformRemovedNativeComponents(UObject* BpCd
 	const AActor* NativeCDO = GetDefault<AActor>(NativeSuperClass);
 
 	TInlineComponentArray<UActorComponent*> OldNativeComponents;
-	ActorCDO->GetComponents(OldNativeComponents);
 	TInlineComponentArray<UActorComponent*> NewNativeComponents;
-	NativeCDO->GetComponents(NewNativeComponents);
+	ActorCDO->GetComponents(OldNativeComponents);
+	USceneComponent* OldNativeRootComponent = ActorCDO->GetRootComponent();
 
 	TSet<UObject*> DestroyedComponents;
 	for (UActorComponent* Component : OldNativeComponents)
@@ -987,6 +974,9 @@ static void ConformComponentsUtils::ConformRemovedNativeComponents(UObject* BpCd
 		UObject* NativeArchetype = FindNativeArchetype(Component);
 		if ((NativeArchetype == nullptr) || !NativeArchetype->HasAnyFlags(RF_ClassDefaultObject))
 		{
+			// Keep track of components inherited from the native super class that are still valid.
+			NewNativeComponents.Add(Component);
+
 			continue;
 		}
 		// else, the component has been removed from our native super class
@@ -994,6 +984,13 @@ static void ConformComponentsUtils::ConformRemovedNativeComponents(UObject* BpCd
 		Component->DestroyComponent(/*bPromoteChildren =*/false);
 		DestroyedComponents.Add(Component);
 
+		// The DestroyComponent() call above will clear the RootComponent value in this case.
+		if(Component == OldNativeRootComponent)
+		{
+			// Restore it here so that it will be reassigned to match the native CDO's value below.
+			ActorCDO->SetRootComponent(OldNativeRootComponent);
+		}
+
 		UClass* ComponentClass = Component->GetClass();
 		for (TFieldIterator<UArrayProperty> ArrayPropIt(NativeSuperClass); ArrayPropIt; ++ArrayPropIt)
 		{
@@ -1026,6 +1023,14 @@ static void ConformComponentsUtils::ConformRemovedNativeComponents(UObject* BpCd
 		// @TODO: have to also remove from map properties now that they're available
 	}
 
+	auto FindComponentTemplateByNameInActorCDO = [&NewNativeComponents](FName ToFind) -> UActorComponent**
+	{
+		return NewNativeComponents.FindByPredicate([ToFind](const UActorComponent* ActorComponent) -> bool
+		{
+			return ActorComponent && ActorComponent->GetFName() == ToFind;
+		});
+	};
+
 	// 
 	for (TFieldIterator<UObjectProperty> ObjPropIt(NativeSuperClass); ObjPropIt; ++ObjPropIt)
 	{
@@ -1034,10 +1039,47 @@ static void ConformComponentsUtils::ConformRemovedNativeComponents(UObject* BpCd
 
 		if (DestroyedComponents.Contains(PropObjValue))
 		{
+			// Get the "new" value that's currently set on the native parent CDO. We need the Blueprint CDO to reflect this update in property value.
 			UObject* SuperObjValue = ObjectProp->GetObjectPropertyValue_InContainer(NativeCDO);
+			if (SuperObjValue && SuperObjValue->IsA<UActorComponent>())
+			{
+				// For components, make sure we use the instance that's owned by the Blueprint CDO and not the native parent CDO's instance.
+				if (UActorComponent** ComponentTemplatePtr = FindComponentTemplateByNameInActorCDO(SuperObjValue->GetFName()))
+				{
+					SuperObjValue = Cast<UObject>(*ComponentTemplatePtr);
+				}
+			}
+			
+			// Update the Blueprint CDO to match the native parent CDO.
 			ObjectProp->SetObjectPropertyValue_InContainer(ActorCDO, SuperObjValue);
 		}
 	}
+
+	// Fix up the attachment hierarchy for inherited scene components that are still valid.
+	for (UActorComponent* Component : NewNativeComponents)
+	{
+		if (USceneComponent* SceneComponent = Cast<USceneComponent>(Component))
+		{
+			// If the component in the Blueprint CDO was attached to a component that's been removed, update the Blueprint's component instance to match the archetype in the native parent CDO.
+			if (DestroyedComponents.Contains(SceneComponent->AttachParent))
+			{
+				if (USceneComponent* NativeArchetype = Cast<USceneComponent>(FindNativeArchetype(SceneComponent)))
+				{
+					USceneComponent* NewAttachParent = NativeArchetype->AttachParent;
+					if (NewAttachParent)
+					{
+						// Make sure we use the instance that's owned by the Blueprint CDO and not the native parent CDO's instance.
+						if (UActorComponent** ComponentTemplatePtr = FindComponentTemplateByNameInActorCDO(NewAttachParent->GetFName()))
+						{
+							NewAttachParent = CastChecked<USceneComponent>(*ComponentTemplatePtr);
+						}
+					}
+
+					SceneComponent->AttachParent = NewAttachParent;
+				}
+			}
+		}
+	}
 }
 
 static UObject* ConformComponentsUtils::FindNativeArchetype(UActorComponent* Component)
@@ -1149,6 +1191,39 @@ void FKismetEditorUtilities::AddComponentsToBlueprint(UBlueprint* Blueprint, con
 	TMap<USceneComponent*, USCS_Node*> SceneComponentsToAdd;
 	TMap<USceneComponent*, USCS_Node*> InstanceComponentToNodeMap;
 
+	struct FAddComponentsToBlueprintImpl
+	{
+		/** 
+		 * Creates a new USCS_Node in the TargetSCS, duplicating the specified 
+		 * component (leaving the new node unattached). If a copy was already 
+		 * made (found  in NewSceneComponents) then that will be returned instead.
+		 */
+		static USCS_Node* MakeComponentCopy(UActorComponent* ActorComponent, USimpleConstructionScript* TargetSCS, TMap<USceneComponent*, USCS_Node*>& NewSceneComponents)
+		{
+			USceneComponent* AsSceneComponent = Cast<USceneComponent>(ActorComponent);
+			if (AsSceneComponent != nullptr)
+			{
+				USCS_Node** ExistingCopy = NewSceneComponents.Find(AsSceneComponent);
+				if (ExistingCopy != nullptr)
+				{
+					return *ExistingCopy;
+				}
+			}
+
+			USCS_Node* NewSCSNode = TargetSCS->CreateNode(ActorComponent->GetClass(), ActorComponent->GetFName());
+			UEditorEngine::CopyPropertiesForUnrelatedObjects(ActorComponent, NewSCSNode->ComponentTemplate);
+
+			// Clear the instance component flag
+			NewSCSNode->ComponentTemplate->CreationMethod = EComponentCreationMethod::Native;
+
+			if (AsSceneComponent != nullptr)
+			{
+				NewSceneComponents.Add(AsSceneComponent, NewSCSNode);
+			}
+			return NewSCSNode;
+		}
+	};
+
 	AActor* Actor = nullptr;
 
 	for (UActorComponent* ActorComponent : Components)
@@ -1165,14 +1240,9 @@ void FKismetEditorUtilities::AddComponentsToBlueprint(UBlueprint* Blueprint, con
 				check(Actor);
 			}
 
-			USCS_Node* SCSNode = SCS->CreateNode(ActorComponent->GetClass(), ActorComponent->GetFName());
-			UEditorEngine::CopyPropertiesForUnrelatedObjects(ActorComponent, SCSNode->ComponentTemplate);
-
-			// Clear the instance component flag
-			SCSNode->ComponentTemplate->CreationMethod = EComponentCreationMethod::Native;
+			USCS_Node* SCSNode = FAddComponentsToBlueprintImpl::MakeComponentCopy(ActorComponent, SCS, InstanceComponentToNodeMap);
 
 			USceneComponent* SceneComponent = Cast<USceneComponent>(ActorComponent);
-
 			// The easy part is non-scene component or the Root simply add it
 			if (SceneComponent == nullptr)
 			{
@@ -1180,8 +1250,6 @@ void FKismetEditorUtilities::AddComponentsToBlueprint(UBlueprint* Blueprint, con
 			}
 			else
 			{
-				InstanceComponentToNodeMap.Add(SceneComponent,SCSNode);
-
 				if (ActorComponent == Actor->GetRootComponent())
 				{
 					if (OptionalNewRootNode != nullptr)
@@ -1197,13 +1265,26 @@ void FKismetEditorUtilities::AddComponentsToBlueprint(UBlueprint* Blueprint, con
 				else if (SceneComponent->AttachParent->IsCreatedByConstructionScript())
 				{
 					USCS_Node* ParentSCSNode = nullptr;
-					for (UBlueprint* ParentBlueprint : ParentBPStack)
+					if (USCS_Node** ParentSCSNodePtr = InstanceComponentToNodeMap.Find(SceneComponent->AttachParent))
+					{
+						ParentSCSNode = *ParentSCSNodePtr;
+					}
+					else if (Components.Contains(SceneComponent->AttachParent))
+					{
+						// since you cannot rely on the order of the supplied  
+						// Components array, we might be looking for a parent 
+						// that hasn't been added yet
+						ParentSCSNode = FAddComponentsToBlueprintImpl::MakeComponentCopy(SceneComponent->AttachParent, SCS, InstanceComponentToNodeMap);
+					}
+					else
 					{
-						USCS_Node** ParentSCSNodePtr = InstanceComponentToNodeMap.Find(SceneComponent->AttachParent);
-						if (ParentSCSNodePtr)
+						for (UBlueprint* ParentBlueprint : ParentBPStack)
 						{
-							ParentSCSNode = *ParentSCSNodePtr;
-							break;
+							ParentSCSNode = ParentBlueprint->SimpleConstructionScript->FindSCSNode(SceneComponent->AttachParent->GetFName());
+							if (ParentSCSNode)
+							{
+								break;
+							}
 						}
 					}
 					check(ParentSCSNode);
diff --git a/Engine/Source/Editor/UnrealEd/Private/Kismet2/KismetReinstanceUtilities.cpp b/Engine/Source/Editor/UnrealEd/Private/Kismet2/KismetReinstanceUtilities.cpp
--- a/Engine/Source/Editor/UnrealEd/Private/Kismet2/KismetReinstanceUtilities.cpp
+++ b/Engine/Source/Editor/UnrealEd/Private/Kismet2/KismetReinstanceUtilities.cpp
@@ -11,6 +11,7 @@
 #include "BlueprintEditor.h"
 #include "Engine/BlueprintGeneratedClass.h"
 #include "Engine/Selection.h"
+#include "BlueprintEditorSettings.h"
 
 DECLARE_CYCLE_STAT(TEXT("Replace Instances"), EKismetReinstancerStats_ReplaceInstancesOfClass, STATGROUP_KismetReinstancer );
 DECLARE_CYCLE_STAT(TEXT("Find Referencers"), EKismetReinstancerStats_FindReferencers, STATGROUP_KismetReinstancer );
@@ -130,6 +131,7 @@ struct FReplaceReferenceHelper
 TSet<TWeakObjectPtr<UBlueprint>> FBlueprintCompileReinstancer::DependentBlueprintsToRefresh = TSet<TWeakObjectPtr<UBlueprint>>();
 TSet<TWeakObjectPtr<UBlueprint>> FBlueprintCompileReinstancer::DependentBlueprintsToRecompile = TSet<TWeakObjectPtr<UBlueprint>>();
 TSet<TWeakObjectPtr<UBlueprint>> FBlueprintCompileReinstancer::DependentBlueprintsToByteRecompile = TSet<TWeakObjectPtr<UBlueprint>>();
+TSet<UBlueprint*> FBlueprintCompileReinstancer::CompiledBlueprintsToSave = TSet<UBlueprint*>();
 
 UClass* FBlueprintCompileReinstancer::HotReloadedOldClass = nullptr;
 UClass* FBlueprintCompileReinstancer::HotReloadedNewClass = nullptr;
@@ -141,6 +143,7 @@ FBlueprintCompileReinstancer::FBlueprintCompileReinstancer(UClass* InClassToRein
 	, bHasReinstanced(false)
 	, bSkipGarbageCollection(bSkipGC)
 	, ClassToReinstanceDefaultValuesCRC(0)
+	, bIsSourceReinstancer(false)
 {
 	if( InClassToReinstance != NULL )
 	{
@@ -156,6 +159,8 @@ FBlueprintCompileReinstancer::FBlueprintCompileReinstancer(UClass* InClassToRein
 		ClassToReinstance->ClassFlags |= CLASS_NewerVersionExists;
 		const FName RenistanceName = MakeUniqueObjectName(GetTransientPackage(), ClassToReinstance->GetClass(), *FString::Printf(TEXT("REINST_%s"), *ClassToReinstance->GetName()));
 		DuplicatedClass = (UClass*)StaticDuplicateObject(ClassToReinstance, GetTransientPackage(), RenistanceName, ~RF_Transactional); 
+		// If you compile a blueprint that is part of the rootset, there's no reason for the REINST version to be part of the rootset:
+		DuplicatedClass->RemoveFromRoot();
 
 		ClassToReinstance->ClassFlags &= ~CLASS_NewerVersionExists;
 		GIsDuplicatingClassForReinstancing = false;
@@ -256,6 +261,16 @@ FBlueprintCompileReinstancer::FBlueprintCompileReinstancer(UClass* InClassToRein
 		{
 			ClassToReinstanceDefaultValuesCRC = GeneratingBP->CrcLastCompiledCDO;
 			FBlueprintEditorUtils::GetDependentBlueprints(GeneratingBP, Dependencies);
+
+			bool const bIsLevelPackage = (UWorld::FindWorldInPackage(GeneratingBP->GetOutermost()) != nullptr);
+			// we don't want to save the entire level (especially if this 
+			// compile was already kicked off as a result of a level save, as it
+			// could cause a recursive save)... let the "SaveOnCompile" setting 
+			// only save blueprint assets
+			if (!bIsLevelPackage)
+			{
+				CompiledBlueprintsToSave.Add(GeneratingBP);
+			}
 		}
 	}
 }
@@ -320,6 +335,29 @@ void FBlueprintCompileReinstancer::OptionallyRefreshNodes(UBlueprint* CurrentBP)
 
 FBlueprintCompileReinstancer::~FBlueprintCompileReinstancer()
 {
+	if (bIsSourceReinstancer)
+	{
+		if (CompiledBlueprintsToSave.Num() > 0)
+		{
+			if ( !IsRunningCommandlet() && !GIsAutomationTesting )
+			{
+				TArray<UPackage*> PackagesToSave;
+				for (UBlueprint* BP : CompiledBlueprintsToSave)
+				{
+					UBlueprintEditorSettings* Settings = GetMutableDefault<UBlueprintEditorSettings>();
+					const bool bShouldSaveOnCompile = (( Settings->SaveOnCompile == SoC_Always ) || ( ( Settings->SaveOnCompile == SoC_SuccessOnly ) && ( BP->Status == BS_UpToDate ) ));
+
+					if (bShouldSaveOnCompile)
+					{
+						PackagesToSave.Add(BP->GetOutermost());
+					}
+				}
+
+				FEditorFileUtils::PromptForCheckoutAndSave(PackagesToSave, /*bCheckDirty =*/true, /*bPromptToSave =*/false);
+			}
+			CompiledBlueprintsToSave.Empty();		
+		}
+	}
 }
 
 class FReinstanceFinalizer : public TSharedFromThis<FReinstanceFinalizer>
@@ -530,7 +568,17 @@ void FBlueprintCompileReinstancer::ReinstanceObjects(bool bForceAlwaysReinstance
 	// Make sure we only reinstance classes once!
 	static TArray<TSharedRef<FBlueprintCompileReinstancer>> QueueToReinstance;
 	TSharedRef<FBlueprintCompileReinstancer> SharedThis = AsShared();
-	const bool bAlreadyQueued = QueueToReinstance.Contains(SharedThis);
+	bool bAlreadyQueued = QueueToReinstance.Contains(SharedThis);
+
+	// We may already be reinstancing this class, this happens when a dependent blueprint has a compile error and we try to reinstance the stub:
+	for (const auto& Entry : QueueToReinstance)
+	{
+		if (Entry->ClassToReinstance == SharedThis->ClassToReinstance)
+		{
+			bAlreadyQueued = true;
+		}
+	}
+
 	if (!bAlreadyQueued && !bHasReinstanced)
 	{
 		QueueToReinstance.Push(SharedThis);
@@ -542,6 +590,9 @@ void FBlueprintCompileReinstancer::ReinstanceObjects(bool bForceAlwaysReinstance
 
 		if (QueueToReinstance.Num() && (QueueToReinstance[0] == SharedThis))
 		{
+			// Mark it as the source reinstancer, no other reinstancer can get here until this Blueprint finishes compiling
+			bIsSourceReinstancer = true;
+
 			TSet<TWeakObjectPtr<UBlueprint>> CompiledBlueprints;
 
 			while (DependentBlueprintsToRecompile.Num())
@@ -551,7 +602,10 @@ void FBlueprintCompileReinstancer::ReinstanceObjects(bool bForceAlwaysReinstance
 				Iter.RemoveCurrent();
 				if (auto BP = BPPtr.Get())
 				{
-					FKismetEditorUtilities::CompileBlueprint(BP, false, bSkipGarbageCollection, false, nullptr, false, true);
+					// it's unsafe to GC in the middle of reinstancing because there may be other reinstancers still alive with references to 
+					// otherwise unreferenced classes:
+					const bool bSkipGC = true;
+					FKismetEditorUtilities::CompileBlueprint(BP, false, bSkipGC, false, nullptr, false, true);
 					CompiledBlueprints.Add(BP);
 				}
 			}
@@ -630,6 +684,9 @@ void FBlueprintCompileReinstancer::UpdateBytecodeReferences()
 		TMap<UObject*, UObject*> FieldMappings;
 		GenerateFieldMappings(FieldMappings);
 
+		// Determine whether or not we will be updating references for an Animation Blueprint class.
+		const bool bIsAnimBlueprintClass = !!Cast<UAnimBlueprint>(ClassToReinstance->ClassGeneratedBy);
+
 		for( auto DependentBP = Dependencies.CreateIterator(); DependentBP; ++DependentBP )
 		{
 			UClass* BPClass = (*DependentBP)->GeneratedClass;
@@ -647,6 +704,13 @@ void FBlueprintCompileReinstancer::UpdateBytecodeReferences()
 
 			BPClass->ClearFunctionMapsCaches();
 
+			// Ensure that Animation Blueprint child class dependencies are always re-linked, as the child may reference properties generated during
+			// compilation of the parent class, which will have shifted to a TRASHCLASS Outer at this point (see UAnimBlueprintGeneratedClass::Link()).
+			if(bIsAnimBlueprintClass && BPClass->IsChildOf(ClassToReinstance))
+			{
+				BPClass->StaticLink(true);
+			}
+
 			bool bBPWasChanged = false;
 			// For each function defined in this blueprint, run through the bytecode, and update any refs from the old properties to the new
 			for( TFieldIterator<UFunction> FuncIter(BPClass, EFieldIteratorFlags::ExcludeSuper); FuncIter; ++FuncIter )
@@ -774,7 +838,11 @@ void FActorReplacementHelper::Finalize(const TMap<UObject*, UObject*>& OldToNewI
 	FEditorScriptExecutionGuard ScriptGuard;
 
 	// run the construction script, which will use the properties we just copied over
-	if (CachedActorData.IsValid())
+	if (NewActor->CurrentTransactionAnnotation.IsValid())
+	{
+		NewActor->RerunConstructionScripts();
+	}
+	else if (CachedActorData.IsValid())
 	{
 		CachedActorData->ComponentInstanceData.FindAndReplaceInstances(OldToNewInstanceMap);
 		NewActor->ExecuteConstruction(TargetWorldTransform, &CachedActorData->ComponentInstanceData);
@@ -1085,6 +1153,31 @@ namespace InstancedPropertyUtils
 
 void FBlueprintCompileReinstancer::ReplaceInstancesOfClass(UClass* OldClass, UClass* NewClass, UObject*	OriginalCDO, TSet<UObject*>* ObjectsThatShouldUseOldStuff, bool bClassObjectReplaced, bool bPreserveRootComponent)
 {
+	TMap<UClass*, UClass*> OldToNewClassMap;
+	OldToNewClassMap.Add(OldClass, NewClass);
+	ReplaceInstancesOfClass_Inner(OldToNewClassMap, OriginalCDO, ObjectsThatShouldUseOldStuff, bClassObjectReplaced, bPreserveRootComponent);
+}
+
+void FBlueprintCompileReinstancer::BatchReplaceInstancesOfClass(TMap<UClass*, UClass*>& InOldToNewClassMap, TSet<UObject*>* ObjectsThatShouldUseOldStuff, bool bClassObjectReplaced, bool bPreserveRootComponent)
+{
+	if (InOldToNewClassMap.Num() == 0)
+	{
+		return;
+	}
+
+	ReplaceInstancesOfClass_Inner(InOldToNewClassMap, nullptr, ObjectsThatShouldUseOldStuff, bClassObjectReplaced, bPreserveRootComponent);
+}
+
+void FBlueprintCompileReinstancer::ReplaceInstancesOfClass_Inner(TMap<UClass*, UClass*>& InOldToNewClassMap, UObject* InOriginalCDO, TSet<UObject*>* ObjectsThatShouldUseOldStuff, bool bClassObjectReplaced, bool bPreserveRootComponent)
+{
+	// If there is an original CDO, we are only reinstancing a single class
+	check((InOriginalCDO != nullptr && InOldToNewClassMap.Num() == 1) || InOriginalCDO == nullptr); // (InOldToNewClassMap.Num() > 1 && InOriginalCDO == nullptr) || (InOldToNewClassMap.Num() == 1 && InOriginalCDO != nullptr));
+
+	if (InOldToNewClassMap.Num() == 0)
+	{
+		return;
+	}
+
 	USelection* SelectedActors;
 	bool bSelectionChanged = false;
 	TArray<UObject*> ObjectsToReplace;
@@ -1092,8 +1185,6 @@ void FBlueprintCompileReinstancer::ReplaceInstancesOfClass(UClass* OldClass, UCl
 
 	// Map of old objects to new objects
 	TMap<UObject*, UObject*> OldToNewInstanceMap;
-	TMap<UClass*, UClass*>   OldToNewClassMap;
-	OldToNewClassMap.Add(OldClass, NewClass);
 
 	TMap<FStringAssetReference, UObject*> ReinstancedObjectsWeakReferenceMap;
 
@@ -1119,220 +1210,234 @@ void FBlueprintCompileReinstancer::ReplaceInstancesOfClass(UClass* OldClass, UCl
 		TMap<UObject*, UObject*> ReplacedObjects;
 	} ObjectRemappingHelper;
 
-	FDelegateHandle OnObjectsReplacedHandle = GEditor->OnObjectsReplaced().AddRaw(&ObjectRemappingHelper,&FObjectRemappingHelper::OnObjectsReplaced);
-
-	{ BP_SCOPED_COMPILER_EVENT_STAT(EKismetReinstancerStats_ReplaceInstancesOfClass);
-
-
-		const bool bIncludeDerivedClasses = false;
-		GetObjectsOfClass(OldClass, ObjectsToReplace, bIncludeDerivedClasses);
-	
+	FDelegateHandle OnObjectsReplacedHandle = GEditor->OnObjectsReplaced().AddRaw(&ObjectRemappingHelper, &FObjectRemappingHelper::OnObjectsReplaced);
+	{
 		SelectedActors = GEditor->GetSelectedActors();
 		SelectedActors->BeginBatchSelectOperation();
 		SelectedActors->Modify();
-		
-		// Then fix 'real' (non archetype) instances of the class
-		for (UObject* OldObject : ObjectsToReplace)
+
+		for (TPair<UClass*, UClass*> OldToNewClass : InOldToNewClassMap)
 		{
-			// Skip non-archetype instances, EXCEPT for component templates
-			const bool bIsComponent = NewClass->IsChildOf(UActorComponent::StaticClass());
-			if ((!bIsComponent && OldObject->IsTemplate()) || OldObject->IsPendingKill())
-			{
-				continue;
-			}
+			UClass* OldClass = OldToNewClass.Key;
+			UClass* NewClass = OldToNewClass.Value;
+			check(OldClass && NewClass);
+			check(OldClass != NewClass);
 
-			UBlueprint* CorrespondingBlueprint  = Cast<UBlueprint>(OldObject->GetClass()->ClassGeneratedBy);
-			UObject*    OldBlueprintDebugObject = nullptr;
-			// If this object is being debugged, cache it off so we can preserve the 'object being debugged' association
-			if ((CorrespondingBlueprint != nullptr) && (CorrespondingBlueprint->GetObjectBeingDebugged() == OldObject))
-			{
-				OldBlueprintDebugObject = OldObject;
-			}
+			{ BP_SCOPED_COMPILER_EVENT_STAT(EKismetReinstancerStats_ReplaceInstancesOfClass);
 
-			AActor*  OldActor = Cast<AActor>(OldObject);
-			UObject* NewUObject = nullptr;
-			// if the object to replace is an actor...
-			if (OldActor != nullptr)
-			{
-				FVector  Location = FVector::ZeroVector;
-				FRotator Rotation = FRotator::ZeroRotator;
-				if (USceneComponent* OldRootComponent = OldActor->GetRootComponent())
-				{
-					// We need to make sure that the ComponentToWorld transform is up to date, but we don't want to run any initialization logic
-					// so we silence the update, cache it off, revert the change (so no events are raised), and then directly update the transform
-					// with the value calculated in ConditionalUpdateComponentToWorld:
-					FScopedMovementUpdate SilenceMovement(OldRootComponent);
-					
-					OldRootComponent->ConditionalUpdateComponentToWorld();
-					FTransform OldComponentToWorld = OldRootComponent->ComponentToWorld;
-					SilenceMovement.RevertMove();
-					
-					OldRootComponent->ComponentToWorld = OldComponentToWorld;
-					Location = OldActor->GetActorLocation();
-					Rotation = OldActor->GetActorRotation();
-				}
+				const bool bIncludeDerivedClasses = false;
+				GetObjectsOfClass(OldClass, ObjectsToReplace, bIncludeDerivedClasses);
 
-				// If this actor was spawned from an Archetype, we spawn the new actor from the new version of that archetype
-				UObject* OldArchetype = OldActor->GetArchetype();
-				UWorld*  World        = OldActor->GetWorld();
-				AActor*  NewArchetype = Cast<AActor>(OldToNewInstanceMap.FindRef(OldArchetype));
-				// Check that either this was an instance of the class directly, or we found a new archetype for it
-				check(OldArchetype == OldClass->GetDefaultObject() || NewArchetype);
-
-				// Spawn the new actor instance, in the same level as the original, but deferring running the construction script until we have transferred modified properties
-				ULevel*  ActorLevel  = OldActor->GetLevel();
-				UClass** MappedClass = OldToNewClassMap.Find(OldActor->GetClass());
-				UClass*  SpawnClass  = MappedClass ? *MappedClass : NewClass;
-
-				FActorSpawnParameters SpawnInfo;
-				SpawnInfo.OverrideLevel      = ActorLevel;
-				SpawnInfo.Template           = NewArchetype;
-				SpawnInfo.SpawnCollisionHandlingOverride = ESpawnActorCollisionHandlingMethod::AlwaysSpawn;
-				SpawnInfo.bDeferConstruction = true;
-
-				// Temporarily remove the deprecated flag so we can respawn the Blueprint in the level
-				const bool bIsClassDeprecated = SpawnClass->HasAnyClassFlags(CLASS_Deprecated);
-				SpawnClass->ClassFlags &= ~CLASS_Deprecated;
-
-				AActor* NewActor = World->SpawnActor(SpawnClass, &Location, &Rotation, SpawnInfo);
-				// Reassign the deprecated flag if it was previously assigned
-				if (bIsClassDeprecated)
+				// Then fix 'real' (non archetype) instances of the class
+				for (UObject* OldObject : ObjectsToReplace)
 				{
-					SpawnClass->ClassFlags |= CLASS_Deprecated;
-				}
+					// Skip non-archetype instances, EXCEPT for component templates
+					const bool bIsComponent = NewClass->IsChildOf(UActorComponent::StaticClass());
+					if ((!bIsComponent && OldObject->IsTemplate()) || OldObject->IsPendingKill())
+					{
+						//OldObject->SetClass(NewClass);
+						continue;
+					}
 
-				check(NewActor != nullptr);
-				NewUObject = NewActor;
-				// store the new actor for the second pass (NOTE: this detaches 
-				// OldActor from all child/parent attachments)
-				//
-				// running the NewActor's construction-script is saved for that 
-				// second pass (because the construction-script may reference 
-				// another instance that hasn't been replaced yet).
-				ReplacementActors.Add(FActorReplacementHelper(NewActor, OldActor));
+					UBlueprint* CorrespondingBlueprint = Cast<UBlueprint>(OldObject->GetClass()->ClassGeneratedBy);
+					UObject*    OldBlueprintDebugObject = nullptr;
+					// If this object is being debugged, cache it off so we can preserve the 'object being debugged' association
+					if ((CorrespondingBlueprint != nullptr) && (CorrespondingBlueprint->GetObjectBeingDebugged() == OldObject))
+					{
+						OldBlueprintDebugObject = OldObject;
+					}
 
-				ReinstancedObjectsWeakReferenceMap.Add(OldObject, NewUObject);
+					AActor*  OldActor = Cast<AActor>(OldObject);
+					UObject* NewUObject = nullptr;
+					// if the object to replace is an actor...
+					if (OldActor != nullptr)
+					{
+						FVector  Location = FVector::ZeroVector;
+						FRotator Rotation = FRotator::ZeroRotator;
+						if (USceneComponent* OldRootComponent = OldActor->GetRootComponent())
+						{
+							// We need to make sure that the ComponentToWorld transform is up to date, but we don't want to run any initialization logic
+							// so we silence the update, cache it off, revert the change (so no events are raised), and then directly update the transform
+							// with the value calculated in ConditionalUpdateComponentToWorld:
+							FScopedMovementUpdate SilenceMovement(OldRootComponent);
+
+							OldRootComponent->ConditionalUpdateComponentToWorld();
+							FTransform OldComponentToWorld = OldRootComponent->ComponentToWorld;
+							SilenceMovement.RevertMove();
+
+							OldRootComponent->ComponentToWorld = OldComponentToWorld;
+							Location = OldActor->GetActorLocation();
+							Rotation = OldActor->GetActorRotation();
+						}
 
-				OldActor->DestroyConstructedComponents(); // don't want to serialize components from the old actor
-				// Unregister native components so we don't copy any sub-components they generate for themselves (like UCameraComponent does)
-				OldActor->UnregisterAllComponents();
+						// If this actor was spawned from an Archetype, we spawn the new actor from the new version of that archetype
+						UObject* OldArchetype = OldActor->GetArchetype();
+						UWorld*  World = OldActor->GetWorld();
+						AActor*  NewArchetype = Cast<AActor>(OldToNewInstanceMap.FindRef(OldArchetype));
+						// Check that either this was an instance of the class directly, or we found a new archetype for it
+						check(OldArchetype == OldClass->GetDefaultObject() || NewArchetype);
+
+						// Spawn the new actor instance, in the same level as the original, but deferring running the construction script until we have transferred modified properties
+						ULevel*  ActorLevel = OldActor->GetLevel();
+						UClass** MappedClass = InOldToNewClassMap.Find(OldActor->GetClass());
+						UClass*  SpawnClass = MappedClass ? *MappedClass : NewClass;
+
+						FActorSpawnParameters SpawnInfo;
+						SpawnInfo.OverrideLevel = ActorLevel;
+						SpawnInfo.Template = NewArchetype;
+						SpawnInfo.SpawnCollisionHandlingOverride = ESpawnActorCollisionHandlingMethod::AlwaysSpawn;
+						SpawnInfo.bDeferConstruction = true;
+
+						// Temporarily remove the deprecated flag so we can respawn the Blueprint in the level
+						const bool bIsClassDeprecated = SpawnClass->HasAnyClassFlags(CLASS_Deprecated);
+						SpawnClass->ClassFlags &= ~CLASS_Deprecated;
+
+						AActor* NewActor = World->SpawnActor(SpawnClass, &Location, &Rotation, SpawnInfo);
+						
+						if (OldActor->CurrentTransactionAnnotation.IsValid())
+						{
+							NewActor->CurrentTransactionAnnotation = OldActor->CurrentTransactionAnnotation;
+						}
 
-				// Unregister any native components, might have cached state based on properties we are going to overwrite
-				NewActor->UnregisterAllComponents(); 
+						// Reassign the deprecated flag if it was previously assigned
+						if (bIsClassDeprecated)
+						{
+							SpawnClass->ClassFlags |= CLASS_Deprecated;
+						}
 
-				UEngine::FCopyPropertiesForUnrelatedObjectsParams Params;
-				Params.bPreserveRootComponent = bPreserveRootComponent;
-				UEngine::CopyPropertiesForUnrelatedObjects(OldActor, NewActor, Params);
+						check(NewActor != nullptr);
+						NewUObject = NewActor;
+						// store the new actor for the second pass (NOTE: this detaches 
+						// OldActor from all child/parent attachments)
+						//
+						// running the NewActor's construction-script is saved for that 
+						// second pass (because the construction-script may reference 
+						// another instance that hasn't been replaced yet).
+						ReplacementActors.Add(FActorReplacementHelper(NewActor, OldActor));
 
-				// reset properties/streams
-				NewActor->ResetPropertiesForConstruction(); 
-				// register native components
-				NewActor->RegisterAllComponents(); 
+						ReinstancedObjectsWeakReferenceMap.Add(OldObject, NewUObject);
 
-				// 
-				// clean up the old actor (unselect it, remove it from the world, etc.)...
+						OldActor->DestroyConstructedComponents(); // don't want to serialize components from the old actor
+																  // Unregister native components so we don't copy any sub-components they generate for themselves (like UCameraComponent does)
+						OldActor->UnregisterAllComponents();
 
-				if (OldActor->IsSelected())
-				{
-					GEditor->SelectActor(OldActor, /*bInSelected =*/false, /*bNotify =*/false);
-					bSelectionChanged = true;
-				}
-				if (GEditor->Layers.IsValid()) // ensure(NULL != GEditor->Layers) ?? While cooking the Layers is NULL.
-				{
-					GEditor->Layers->DisassociateActorFromLayers(OldActor);
-				}
+						// Unregister any native components, might have cached state based on properties we are going to overwrite
+						NewActor->UnregisterAllComponents();
 
- 				World->EditorDestroyActor(OldActor, /*bShouldModifyLevel =*/true);
- 				OldToNewInstanceMap.Add(OldActor, NewActor);
-			}
-			else
-			{
-				auto OldFlags = OldObject->GetFlags();
+						UEngine::FCopyPropertiesForUnrelatedObjectsParams Params;
+						Params.bPreserveRootComponent = bPreserveRootComponent;
+						UEngine::CopyPropertiesForUnrelatedObjects(OldActor, NewActor, Params);
 
-				FName OldName(OldObject->GetFName());
-				OldObject->Rename(NULL, OldObject->GetOuter(), REN_DoNotDirty | REN_DontCreateRedirectors);
-				NewUObject = NewObject<UObject>(OldObject->GetOuter(), NewClass, OldName);
-				check(NewUObject != nullptr);
+						// reset properties/streams
+						NewActor->ResetPropertiesForConstruction();
+						// register native components
+						NewActor->RegisterAllComponents();
 
-				auto FlagMask = RF_Public | RF_ArchetypeObject | RF_Transactional | RF_Transient | RF_TextExportTransient | RF_InheritableComponentTemplate; //TODO: what about RF_RootSet and RF_Standalone ?
-				NewUObject->SetFlags(OldFlags & FlagMask);
+						// 
+						// clean up the old actor (unselect it, remove it from the world, etc.)...
 
-				InstancedPropertyUtils::FInstancedPropertyMap InstancedPropertyMap;
-				InstancedPropertyUtils::FArchiveInstancedSubObjCollector  InstancedSubObjCollector(OldObject, InstancedPropertyMap);
-				UEditorEngine::CopyPropertiesForUnrelatedObjects(OldObject, NewUObject);
-				InstancedPropertyUtils::FArchiveInsertInstancedSubObjects InstancedSubObjSpawner(NewUObject, InstancedPropertyMap);
+						if (OldActor->IsSelected())
+						{
+							GEditor->SelectActor(OldActor, /*bInSelected =*/false, /*bNotify =*/false);
+							bSelectionChanged = true;
+						}
+						if (GEditor->Layers.IsValid()) // ensure(NULL != GEditor->Layers) ?? While cooking the Layers is NULL.
+						{
+							GEditor->Layers->DisassociateActorFromLayers(OldActor);
+						}
 
-				if (UAnimInstance* AnimTree = Cast<UAnimInstance>(NewUObject))
-				{
-					// Initialising the anim instance isn't enough to correctly set up the skeletal mesh again in a
-					// paused world, need to initialise the skeletal mesh component that contains the anim instance.
-					if (USkeletalMeshComponent* SkelComponent = Cast<USkeletalMeshComponent>(AnimTree->GetOuter()))
-					{
-						SkelComponent->InitAnim(true);
+						World->EditorDestroyActor(OldActor, /*bShouldModifyLevel =*/true);
+						OldToNewInstanceMap.Add(OldActor, NewActor);
 					}
-				}
-
-				bool bWasRegistered = false;
-				if (bIsComponent)
-				{
-					UActorComponent* OldComponent = CastChecked<UActorComponent>(OldObject);
-					if (OldComponent->IsRegistered())
+					else
 					{
-						bWasRegistered = true;
-						OldComponent->UnregisterComponent();
-					}
-				}
+						auto OldFlags = OldObject->GetFlags();
 
-				OldObject->RemoveFromRoot();
-				OldObject->MarkPendingKill();
+						FName OldName(OldObject->GetFName());
+						OldObject->Rename(NULL, OldObject->GetOuter(), REN_DoNotDirty | REN_DontCreateRedirectors);
+						NewUObject = NewObject<UObject>(OldObject->GetOuter(), NewClass, OldName);
+						check(NewUObject != nullptr);
 
-				OldToNewInstanceMap.Add(OldObject, NewUObject);
+						auto FlagMask = RF_Public | RF_ArchetypeObject | RF_Transactional | RF_Transient | RF_TextExportTransient | RF_InheritableComponentTemplate; //TODO: what about RF_RootSet and RF_Standalone ?
+						NewUObject->SetFlags(OldFlags & FlagMask);
 
-				if (bIsComponent)
-				{
-					UActorComponent* Component = CastChecked<UActorComponent>(NewUObject);
-					AActor* OwningActor = Component->GetOwner();
-					if (OwningActor)
-					{
-						OwningActor->ResetOwnedComponents();
+						InstancedPropertyUtils::FInstancedPropertyMap InstancedPropertyMap;
+						InstancedPropertyUtils::FArchiveInstancedSubObjCollector  InstancedSubObjCollector(OldObject, InstancedPropertyMap);
+						UEditorEngine::CopyPropertiesForUnrelatedObjects(OldObject, NewUObject);
+						InstancedPropertyUtils::FArchiveInsertInstancedSubObjects InstancedSubObjSpawner(NewUObject, InstancedPropertyMap);
 
-						// Check to see if they have an editor that potentially needs to be refreshed
-						if (OwningActor->GetClass()->ClassGeneratedBy)
+						if (UAnimInstance* AnimTree = Cast<UAnimInstance>(NewUObject))
 						{
-							PotentialEditorsForRefreshing.AddUnique(OwningActor->GetClass()->ClassGeneratedBy);
+							// Initialising the anim instance isn't enough to correctly set up the skeletal mesh again in a
+							// paused world, need to initialise the skeletal mesh component that contains the anim instance.
+							if (USkeletalMeshComponent* SkelComponent = Cast<USkeletalMeshComponent>(AnimTree->GetOuter()))
+							{
+								SkelComponent->InitAnim(true);
+							}
 						}
 
-						// we need to keep track of actor instances that need 
-						// their construction scripts re-ran (since we've just 
-						// replaced a component they own)
-						
-						// Skipping CDOs as CSs are not allowed for them.
-						if (!OwningActor->HasAnyFlags(RF_ClassDefaultObject))
+						bool bWasRegistered = false;
+						if (bIsComponent)
+						{
+							UActorComponent* OldComponent = CastChecked<UActorComponent>(OldObject);
+							if (OldComponent->IsRegistered())
+							{
+								bWasRegistered = true;
+								OldComponent->UnregisterComponent();
+							}
+						}
+
+						OldObject->RemoveFromRoot();
+						OldObject->MarkPendingKill();
+
+						OldToNewInstanceMap.Add(OldObject, NewUObject);
+
+						if (bIsComponent)
 						{
-							OwnersToRerunConstructionScript.Add(OwningActor);
+							UActorComponent* Component = CastChecked<UActorComponent>(NewUObject);
+							AActor* OwningActor = Component->GetOwner();
+							if (OwningActor)
+							{
+								OwningActor->ResetOwnedComponents();
+
+								// Check to see if they have an editor that potentially needs to be refreshed
+								if (OwningActor->GetClass()->ClassGeneratedBy)
+								{
+									PotentialEditorsForRefreshing.AddUnique(OwningActor->GetClass()->ClassGeneratedBy);
+								}
+
+								// we need to keep track of actor instances that need 
+								// their construction scripts re-ran (since we've just 
+								// replaced a component they own)
+
+								// Skipping CDOs as CSs are not allowed for them.
+								if (!OwningActor->HasAnyFlags(RF_ClassDefaultObject))
+								{
+									OwnersToRerunConstructionScript.Add(OwningActor);
+								}
+							}
+
+							if (bWasRegistered)
+							{
+								Component->RegisterComponent();
+							}
 						}
 					}
 
-					if (bWasRegistered)
+					// If this original object came from a blueprint and it was in the selected debug set, change the debugging to the new object.
+					if ((CorrespondingBlueprint) && (OldBlueprintDebugObject) && (NewUObject))
 					{
-						Component->RegisterComponent();
+						CorrespondingBlueprint->SetObjectBeingDebugged(NewUObject);
 					}
-				}
-			}
 
-			// If this original object came from a blueprint and it was in the selected debug set, change the debugging to the new object.
-			if ((CorrespondingBlueprint) && (OldBlueprintDebugObject) && (NewUObject))
-			{
-				CorrespondingBlueprint->SetObjectBeingDebugged(NewUObject);
-			}
-
-			if (bLogConversions)
-			{
-				UE_LOG(LogBlueprint, Log, TEXT("Converted instance '%s' to '%s'"), *OldObject->GetPathName(), *NewUObject->GetPathName());
+					if (bLogConversions)
+					{
+						UE_LOG(LogBlueprint, Log, TEXT("Converted instance '%s' to '%s'"), *OldObject->GetPathName(), *NewUObject->GetPathName());
+					}
+				}
 			}
 		}
 	}
-
 	GEditor->OnObjectsReplaced().Remove(OnObjectsReplacedHandle);
 
 	// Now replace any pointers to the old archetypes/instances with pointers to the new one
@@ -1343,12 +1448,24 @@ void FBlueprintCompileReinstancer::ReplaceInstancesOfClass(UClass* OldClass, UCl
 	OldToNewInstanceMap.GenerateValueArray(DstObjects); // Also look for references in new spawned objects.
 
 	SourceObjects.Append(DstObjects);
-
-	FReplaceReferenceHelper::IncludeCDO(OldClass, NewClass, OldToNewInstanceMap, SourceObjects, OriginalCDO);
 	
-	if (bClassObjectReplaced)
+	if (InOriginalCDO)
 	{
-		FReplaceReferenceHelper::IncludeClass(OldClass, NewClass, OldToNewInstanceMap, SourceObjects, ObjectsToReplace);
+		check(InOldToNewClassMap.Num() == 1);
+		for (TPair<UClass*, UClass*> OldToNewClass : InOldToNewClassMap)
+		{
+			UClass* OldClass = OldToNewClass.Key;
+			UClass* NewClass = OldToNewClass.Value;
+			check(OldClass && NewClass);
+			check(OldClass != NewClass);
+
+			FReplaceReferenceHelper::IncludeCDO(OldClass, NewClass, OldToNewInstanceMap, SourceObjects, InOriginalCDO);
+
+			if (bClassObjectReplaced)
+			{
+				FReplaceReferenceHelper::IncludeClass(OldClass, NewClass, OldToNewInstanceMap, SourceObjects, ObjectsToReplace);
+			}
+		}
 	}
 
 	FReplaceReferenceHelper::FindAndReplaceReferences(SourceObjects, ObjectsThatShouldUseOldStuff, ObjectsToReplace, OldToNewInstanceMap, ReinstancedObjectsWeakReferenceMap);
diff --git a/Engine/Source/Editor/UnrealEd/Private/PlayLevel.cpp b/Engine/Source/Editor/UnrealEd/Private/PlayLevel.cpp
--- a/Engine/Source/Editor/UnrealEd/Private/PlayLevel.cpp
+++ b/Engine/Source/Editor/UnrealEd/Private/PlayLevel.cpp
@@ -2317,6 +2317,7 @@ void UEditorEngine::PlayInEditor( UWorld* InWorld, bool bInSimulateInEditor )
 	PlayInSettings->SetPlayNetMode(OrigPlayNetMode);
 
 	// Monitoring when PIE corrupts references between the World and the PIE generated World for UE-20486
+	if (EditorWorld)
 	{
 		TArray<ULevel*> Levels = EditorWorld->GetLevels();
 
diff --git a/Engine/Source/Editor/UnrealEd/Private/Settings/SettingsClasses.cpp b/Engine/Source/Editor/UnrealEd/Private/Settings/SettingsClasses.cpp
--- a/Engine/Source/Editor/UnrealEd/Private/Settings/SettingsClasses.cpp
+++ b/Engine/Source/Editor/UnrealEd/Private/Settings/SettingsClasses.cpp
@@ -66,7 +66,8 @@ UEditorExperimentalSettings::UEditorExperimentalSettings( const FObjectInitializ
 	: Super(ObjectInitializer)
 	, bUnifiedBlueprintEditor(true)
 	, bBlueprintableComponents(true)
-
+	, bBlueprintPerformanceAnalysisTools(false)
+	, BlueprintProfilerAverageSampleCount(20)
 {
 }
 
diff --git a/Engine/Source/Editor/UnrealEd/Public/Kismet2/BlueprintEditorUtils.h b/Engine/Source/Editor/UnrealEd/Public/Kismet2/BlueprintEditorUtils.h
--- a/Engine/Source/Editor/UnrealEd/Public/Kismet2/BlueprintEditorUtils.h
+++ b/Engine/Source/Editor/UnrealEd/Public/Kismet2/BlueprintEditorUtils.h
@@ -1359,3 +1359,22 @@ public:
 	 */
 	static void HandleDisableEditableWhenInherited(UObject* ModifiedObject, TArray<UObject*>& ArchetypeInstances);
 };
+
+struct UNREALED_API FBlueprintDuplicationScopeFlags
+{
+	enum EFlags : uint32
+	{
+		NoFlags = 0,
+		NoExtraCompilation = 1 << 0,
+		TheSameTimelineGuid = 1 << 1,
+	};
+
+	static uint32 bStaticFlags;
+	static bool HasAnyFlag(uint32 InFlags)
+	{
+		return 0 != (bStaticFlags & InFlags);
+	}
+
+	TGuardValue<uint32> Guard;
+	FBlueprintDuplicationScopeFlags(uint32 InFlags) : Guard(bStaticFlags, InFlags) {}
+};
\ No newline at end of file
diff --git a/Engine/Source/Editor/UnrealEd/Public/Kismet2/KismetReinstanceUtilities.h b/Engine/Source/Editor/UnrealEd/Public/Kismet2/KismetReinstanceUtilities.h
--- a/Engine/Source/Editor/UnrealEd/Public/Kismet2/KismetReinstanceUtilities.h
+++ b/Engine/Source/Editor/UnrealEd/Public/Kismet2/KismetReinstanceUtilities.h
@@ -40,6 +40,7 @@ protected:
 	static TSet<TWeakObjectPtr<UBlueprint>> DependentBlueprintsToRefresh;
 	static TSet<TWeakObjectPtr<UBlueprint>> DependentBlueprintsToRecompile;
 	static TSet<TWeakObjectPtr<UBlueprint>> DependentBlueprintsToByteRecompile;
+	static TSet<UBlueprint*> CompiledBlueprintsToSave;
 
 	static UClass* HotReloadedOldClass;
 	static UClass* HotReloadedNewClass;
@@ -77,6 +78,9 @@ protected:
 	/** Objects that should keep reference to old class */
 	TSet<UObject*> ObjectsThatShouldUseOldStuff;
 
+	/** TRUE if this is the source reinstancer that all other active reinstancing is spawned from */
+	bool bIsSourceReinstancer;
+
 public:
 	// FSerializableObject interface
 	virtual void AddReferencedObjects(FReferenceCollector& Collector) override;
@@ -110,6 +114,9 @@ public:
 	/** Worker function to replace all instances of OldClass with a new instance of NewClass */
 	static void ReplaceInstancesOfClass(UClass* OldClass, UClass* NewClass, UObject* OriginalCDO = NULL, TSet<UObject*>* ObjectsThatShouldUseOldStuff = NULL, bool bClassObjectReplaced = false, bool bPreserveRootComponent = true);
 
+	/** Batch replaces a mapping of one or more classes to their new class by leveraging ReplaceInstancesOfClass */
+	static void BatchReplaceInstancesOfClass(TMap<UClass*, UClass*>& InOldToNewClassMap, TSet<UObject*>* ObjectsThatShouldUseOldStuff = NULL, bool bClassObjectReplaced = false, bool bPreserveRootComponent = true);
+
 	/**
 	 * When re-instancing a component, we have to make sure all instance owners' 
 	 * construction scripts are re-ran (in-case modifying the component alters 
@@ -162,6 +169,9 @@ private:
 	 * @param Name The name that CDO has to be renamed with (or created with).
 	 */
 	static UObject* GetClassCDODuplicate(UObject* CDO, FName Name);
+
+	/** Handles the work of ReplaceInstancesOfClass, handling both normal replacement of instances and batch */
+	static void ReplaceInstancesOfClass_Inner(TMap<UClass*, UClass*>& InOldToNewClassMap, UObject* InOriginalCDO, TSet<UObject*>* ObjectsThatShouldUseOldStuff = NULL, bool bClassObjectReplaced = false, bool bPreserveRootComponent = true);
 };
 
 
diff --git a/Engine/Source/Editor/UnrealEd/UnrealEd.Build.cs b/Engine/Source/Editor/UnrealEd/UnrealEd.Build.cs
--- a/Engine/Source/Editor/UnrealEd/UnrealEd.Build.cs
+++ b/Engine/Source/Editor/UnrealEd/UnrealEd.Build.cs
@@ -143,7 +143,8 @@ public class UnrealEd : ModuleRules
                 "LocalizationService",
 				"AddContentDialog",
 				"GameProjectGeneration",
-                "HierarchicalLODUtilities"
+                "HierarchicalLODUtilities",
+                "BlueprintNativeCodeGen",               
 			}
 		);
 
@@ -213,7 +214,7 @@ public class UnrealEd : ModuleRules
 				"SizeMap",
                 "PortalProxies",
                 "PortalServices",
-                "GeometryCacheEd"
+                "GeometryCacheEd",
 			}
 		);
 
diff --git a/Engine/Source/Programs/AutomationTool/AutomationUtils/CommandletUtils.cs b/Engine/Source/Programs/AutomationTool/AutomationUtils/CommandletUtils.cs
--- a/Engine/Source/Programs/AutomationTool/AutomationUtils/CommandletUtils.cs
+++ b/Engine/Source/Programs/AutomationTool/AutomationUtils/CommandletUtils.cs
@@ -179,20 +179,6 @@ namespace AutomationTool
             RunCommandlet(ProjectName, UE4Exe, "UpdateGameProject", Parameters);
         }
 
-        /// <summary>
-        /// Runs the GenerateNativePluginFromBlueprint commandlet.
-        /// </summary>
-        /// <param name="ManifestPath"></param>
-        /// <param name="UE4Exe"></param>
-        /// <param name="Parameters"></param>
-        public static void GenerateNativePluginFromBlueprintCommandlet(FileReference ProjectName, string ManifestPath, string UE4Exe = "UE4Editor-Cmd.exe", string Parameters = "")
-        {
-            Parameters = String.Format("{0} -manifest=\"{1}\"", Parameters, ManifestPath);
-            Parameters.Trim();
-
-            RunCommandlet(ProjectName, UE4Exe, "GenerateNativePluginFromBlueprint", Parameters);
-        }
-
 		/// <summary>
 		/// Runs a commandlet using Engine/Binaries/Win64/UE4Editor-Cmd.exe.
 		/// </summary>
diff --git a/Engine/Source/Programs/AutomationTool/AutomationUtils/ProjectParams.cs b/Engine/Source/Programs/AutomationTool/AutomationUtils/ProjectParams.cs
--- a/Engine/Source/Programs/AutomationTool/AutomationUtils/ProjectParams.cs
+++ b/Engine/Source/Programs/AutomationTool/AutomationUtils/ProjectParams.cs
@@ -298,6 +298,7 @@ namespace AutomationTool
 			this.bIsCodeBasedProject = InParams.bIsCodeBasedProject;
 			this.bCodeSign = InParams.bCodeSign;
 			this.UploadSymbols = InParams.UploadSymbols;
+            this.RunAssetNativization = InParams.RunAssetNativization;
 		}
 
 		/// <summary>
@@ -404,6 +405,7 @@ namespace AutomationTool
             bool? IterativeDeploy = null,
 			bool? FastCook = null,
 			bool? IgnoreCookErrors = null,
+            bool? RunAssetNativization = null,
 			bool? CodeSign = null,
 			bool? UploadSymbols = null,
 			string Provision = null,
@@ -567,6 +569,7 @@ namespace AutomationTool
 			this.IterativeDeploy = GetParamValueIfNotSpecified(Command, IterativeDeploy, this.IterativeDeploy, new string[] {"iterativedeploy", "iterate" } );
 			this.FastCook = GetParamValueIfNotSpecified(Command, FastCook, this.FastCook, "FastCook");
 			this.IgnoreCookErrors = GetParamValueIfNotSpecified(Command, IgnoreCookErrors, this.IgnoreCookErrors, "IgnoreCookErrors");
+            this.RunAssetNativization = GetParamValueIfNotSpecified(Command, RunAssetNativization, this.RunAssetNativization, "nativizeAssets");
 			this.UploadSymbols = GetParamValueIfNotSpecified(Command, UploadSymbols, this.UploadSymbols, "uploadsymbols");
 			this.Device = ParseParamValueIfNotSpecified(Command, Device, "device", String.Empty).Trim(new char[] { '\"' });
 
@@ -953,6 +956,12 @@ namespace AutomationTool
 		[Help("createappbundle", "When archiving for Mac, set this to true to package it in a .app bundle instead of normal loose files")]
 		public bool CreateAppBundle;
 
+        /// <summary>
+        /// Determines if Blueprint assets should be substituted with auto-generated code.
+        /// </summary>
+        [Help("nativizeAssets", "Runs a \"nativization\" pass on Blueprint assets, converting then into C++ (replacing the assets with the generated source).")]
+        public bool RunAssetNativization;
+
 		/// <summary>
 		/// Shared: Ref to an auto-generated plugin file that should be incorporated into the project's build
 		/// </summary>
@@ -1505,7 +1514,7 @@ namespace AutomationTool
 				ProjectGameExePath = null;
 			}
 
-			var Properties = ProjectUtils.GetProjectProperties(RawProjectPath, ClientTargetPlatforms);
+            var Properties = ProjectUtils.GetProjectProperties(RawProjectPath, ClientTargetPlatforms, RunAssetNativization);
 
 			bUsesSteam = Properties.bUsesSteam;
 			bUsesCEF3 = Properties.bUsesCEF3;
@@ -2255,6 +2264,7 @@ namespace AutomationTool
 				CommandUtils.LogLog("bUsesSlate={0}", bUsesSlate);
                 CommandUtils.LogLog("bDebugBuildsActuallyUseDebugCRT={0}", bDebugBuildsActuallyUseDebugCRT);
 				CommandUtils.LogLog("UploadSymbols={0}", UploadSymbols);
+                CommandUtils.LogLog("NativizeAssets={0}", RunAssetNativization);
 				CommandUtils.LogLog("Project Params **************");
 			}
 			bLogged = true;
diff --git a/Engine/Source/Programs/AutomationTool/AutomationUtils/ProjectUtils.cs b/Engine/Source/Programs/AutomationTool/AutomationUtils/ProjectUtils.cs
--- a/Engine/Source/Programs/AutomationTool/AutomationUtils/ProjectUtils.cs
+++ b/Engine/Source/Programs/AutomationTool/AutomationUtils/ProjectUtils.cs
@@ -117,7 +117,7 @@ namespace AutomationTool
 		/// </summary>
 		/// <param name="RawProjectPath">Full project path.</param>
 		/// <returns>Properties of the project.</returns>
-		public static ProjectProperties GetProjectProperties(FileReference RawProjectPath, List<UnrealTargetPlatform> ClientTargetPlatforms = null)
+		public static ProjectProperties GetProjectProperties(FileReference RawProjectPath, List<UnrealTargetPlatform> ClientTargetPlatforms = null, bool AssetNativizationRequested = false)
 		{
 			string ProjectKey = "UE4";
 			if (RawProjectPath != null)
@@ -127,7 +127,7 @@ namespace AutomationTool
 			ProjectProperties Properties;
 			if (PropertiesCache.TryGetValue(ProjectKey, out Properties) == false)
 			{
-				Properties = DetectProjectProperties(RawProjectPath, ClientTargetPlatforms);
+                Properties = DetectProjectProperties(RawProjectPath, ClientTargetPlatforms, AssetNativizationRequested);
 				PropertiesCache.Add(ProjectKey, Properties);
 			}
 			return Properties;
@@ -158,7 +158,7 @@ namespace AutomationTool
 			return ProjectClientBinariesPath;
 		}
 
-		private static bool RequiresTempTarget(FileReference RawProjectPath, List<UnrealTargetPlatform> ClientTargetPlatforms)
+		private static bool RequiresTempTarget(FileReference RawProjectPath, List<UnrealTargetPlatform> ClientTargetPlatforms, bool AssetNativizationRequested)
 		{
 			// check to see if we already have a Target.cs file
 			if (File.Exists (Path.Combine (Path.GetDirectoryName (RawProjectPath.FullName), "Source", RawProjectPath.GetFileNameWithoutExtension() + ".Target.cs")))
@@ -176,6 +176,18 @@ namespace AutomationTool
 				}
 			}
 
+            //
+            // once we reach this point, we can surmise that this is an asset-
+            // only (code free) project
+
+            if (AssetNativizationRequested)
+            {
+                // we're going to be converting some of the project's assets 
+                // into native code, so we require a distinct target (executable) 
+                // be generated for this project
+                return true;
+            }
+
 			// no Target file, now check to see if build settings have changed
 			List<UnrealTargetPlatform> TargetPlatforms = ClientTargetPlatforms;
 			if (ClientTargetPlatforms == null || ClientTargetPlatforms.Count < 1)
@@ -262,7 +274,7 @@ namespace AutomationTool
 		/// </summary>
 		/// <param name="RawProjectPath">Full project path.</param>
 		/// <returns>Project properties.</returns>
-		private static ProjectProperties DetectProjectProperties(FileReference RawProjectPath, List<UnrealTargetPlatform> ClientTargetPlatforms)
+        private static ProjectProperties DetectProjectProperties(FileReference RawProjectPath, List<UnrealTargetPlatform> ClientTargetPlatforms, bool AssetNativizationRequested)
 		{
 			var Properties = new ProjectProperties();
 			Properties.RawProjectPath = RawProjectPath;
@@ -271,12 +283,20 @@ namespace AutomationTool
 			List<string> ExtraSearchPaths = null;
 			if (RawProjectPath != null)
 			{
-				if (RequiresTempTarget(RawProjectPath, ClientTargetPlatforms))
+                if (RequiresTempTarget(RawProjectPath, ClientTargetPlatforms, AssetNativizationRequested))
 				{
 					GenerateTempTarget(RawProjectPath);
 					Properties.bWasGenerated = true;
 					ExtraSearchPaths = new List<string>();
-					ExtraSearchPaths.Add(CommandUtils.CombinePaths(Path.GetDirectoryName(RawProjectPath.FullName), "Intermediate", "Source"));
+
+                    string TempTargetDir = CommandUtils.CombinePaths(Path.GetDirectoryName(RawProjectPath.FullName), "Intermediate", "Source");
+                    ExtraSearchPaths.Add(TempTargetDir);
+
+                    // in case the RulesCompiler (what we use to find all the 
+                    // Target.cs files) has already cached the contents of this 
+                    // directory, then we need to invalidate that cache (so 
+                    // it'll find/use the new Target.cs file)
+                    RulesCompiler.InvalidateRulesFileCache(TempTargetDir);
 				}
 				else if (File.Exists(Path.Combine(Path.GetDirectoryName(RawProjectPath.FullName), "Intermediate", "Source", Path.GetFileNameWithoutExtension(RawProjectPath.FullName) + ".Target.cs")))
 				{
diff --git a/Engine/Source/Programs/AutomationTool/Scripts/AutomationScripts.Automation.csproj b/Engine/Source/Programs/AutomationTool/Scripts/AutomationScripts.Automation.csproj
--- a/Engine/Source/Programs/AutomationTool/Scripts/AutomationScripts.Automation.csproj
+++ b/Engine/Source/Programs/AutomationTool/Scripts/AutomationScripts.Automation.csproj
@@ -74,7 +74,6 @@
     <Compile Include="FixupRedirects.Automation.cs" />
     <Compile Include="GenerateDSYM.Automation.cs" />
     <Compile Include="RebuildLightMapsCommand.Automation.cs" />
-    <Compile Include="NativizeScriptAssets.Automation.cs" />
     <Compile Include="RocketBuild.Automation.cs" />
     <Compile Include="IPhonePackager.Automation.cs" />
     <Compile Include="LauncherLocalization.Automation.cs" />
diff --git a/Engine/Source/Programs/AutomationTool/Scripts/BuildCookRun.Automation.cs b/Engine/Source/Programs/AutomationTool/Scripts/BuildCookRun.Automation.cs
--- a/Engine/Source/Programs/AutomationTool/Scripts/BuildCookRun.Automation.cs
+++ b/Engine/Source/Programs/AutomationTool/Scripts/BuildCookRun.Automation.cs
@@ -199,15 +199,20 @@ public class BuildCookRun : BuildCommand
 
 	protected void DoBuildCookRun(ProjectParams Params)
 	{
+		const ProjectBuildTargets ClientTargets = ProjectBuildTargets.ClientCooked | ProjectBuildTargets.ServerCooked;
+        bool bGenerateNativeScripts = Params.RunAssetNativization;
 		int WorkingCL = -1;
 		if (P4Enabled && AllowSubmit)
 		{
 			WorkingCL = P4.CreateChange(P4Env.Client, String.Format("{0} build from changelist {1}", Params.ShortProjectName, P4Env.Changelist));
 		}
 
-        Project.NativizeScriptAssets(this, Params);
-        Project.Build(this, Params, WorkingCL);
+        Project.Build(this, Params, WorkingCL, bGenerateNativeScripts ? (ProjectBuildTargets.All & ~ClientTargets) : ProjectBuildTargets.All);
 		Project.Cook(Params);
+		if (bGenerateNativeScripts)
+		{
+            Project.Build(this, Params, WorkingCL, ClientTargets);
+		}
 		Project.CopyBuildToStagingDirectory(Params);
 		Project.Package(Params, WorkingCL);
 		Project.Archive(Params);
diff --git a/Engine/Source/Programs/AutomationTool/Scripts/BuildProjectCommand.Automation.cs b/Engine/Source/Programs/AutomationTool/Scripts/BuildProjectCommand.Automation.cs
--- a/Engine/Source/Programs/AutomationTool/Scripts/BuildProjectCommand.Automation.cs
+++ b/Engine/Source/Programs/AutomationTool/Scripts/BuildProjectCommand.Automation.cs
@@ -8,6 +8,21 @@ using System.Linq;
 using AutomationTool;
 using UnrealBuildTool;
 
+[Flags]
+public enum ProjectBuildTargets
+{
+	None = 0,
+	Editor = 1 << 0,
+	ClientCooked = 1 << 1,
+	ServerCooked = 1 << 2,
+	Bootstrap = 1 << 3,
+	CrashReporter = 1 << 4,
+	Programs = 1 << 5,
+
+	// All targets
+	All = Editor | ClientCooked | ServerCooked | Bootstrap | CrashReporter | Programs,
+}
+
 /// <summary>
 /// Helper command used for compiling.
 /// </summary>
@@ -22,9 +37,8 @@ public partial class Project : CommandUtils
 {
 	#region Build Command
 
-	public static void Build(BuildCommand Command, ProjectParams Params, int WorkingCL = -1)
+    public static void Build(BuildCommand Command, ProjectParams Params, int WorkingCL = -1, ProjectBuildTargets TargetMask = ProjectBuildTargets.All)
 	{
-
 		Params.ValidateAndLog();
 
 		if (!Params.Build)
@@ -38,23 +52,15 @@ public partial class Project : CommandUtils
 		var Agenda = new UE4Build.BuildAgenda();
 		var CrashReportPlatforms = new HashSet<UnrealTargetPlatform>();
 
-        string ScriptPluginArgs = "";
-        // if we're utilizing an auto-generated code plugin/module (a product of 
-        // the BuildCookeRun process), make sure to compile it along with the targets here
-        if (Params.UseNativizedScriptPlugin())
-        {
-            ScriptPluginArgs = "-PLUGIN \"" + Params.NativizedScriptPlugin + "\"";
-        }
-
 		// Setup editor targets
-		if (Params.HasEditorTargets && !Params.Rocket)
+		if (Params.HasEditorTargets && !Params.Rocket && (TargetMask & ProjectBuildTargets.Editor) == ProjectBuildTargets.Editor)
 		{
 			// @todo Mac: proper platform detection
 			UnrealTargetPlatform EditorPlatform = HostPlatform.Current.HostEditorPlatform;
 			const UnrealTargetConfiguration EditorConfiguration = UnrealTargetConfiguration.Development;
 
 			CrashReportPlatforms.Add(EditorPlatform);
-            Agenda.AddTargets(Params.EditorTargets.ToArray(), EditorPlatform, EditorConfiguration, Params.CodeBasedUprojectPath, ScriptPluginArgs);
+            Agenda.AddTargets(Params.EditorTargets.ToArray(), EditorPlatform, EditorConfiguration, Params.CodeBasedUprojectPath);
 			if (Params.EditorTargets.Contains("UnrealHeaderTool") == false)
 			{
 				Agenda.AddTargets(new string[] { "UnrealHeaderTool" }, EditorPlatform, EditorConfiguration);
@@ -73,8 +79,16 @@ public partial class Project : CommandUtils
 			}
 		}
 
+        string ScriptPluginArgs = "";
+        // if we're utilizing an auto-generated code plugin/module (a product of 
+        // the cook process), make sure to compile it along with the targets here
+        if (Params.UseNativizedScriptPlugin())
+        {
+            ScriptPluginArgs = "-PLUGIN \"" + Params.NativizedScriptPlugin + "\"";
+        }
+
 		// Setup cooked targets
-		if (Params.HasClientCookedTargets)
+		if (Params.HasClientCookedTargets && (TargetMask & ProjectBuildTargets.ClientCooked) == ProjectBuildTargets.ClientCooked)
 		{
 			foreach (var BuildConfig in Params.ClientConfigsToBuild)
 			{
@@ -85,7 +99,7 @@ public partial class Project : CommandUtils
 				}
 			}
 		}
-		if (Params.HasServerCookedTargets)
+		if (Params.HasServerCookedTargets && (TargetMask & ProjectBuildTargets.ServerCooked) == ProjectBuildTargets.ServerCooked)
 		{
 			foreach (var BuildConfig in Params.ServerConfigsToBuild)
 			{
@@ -96,7 +110,7 @@ public partial class Project : CommandUtils
 				}
 			}
 		}
-		if (!Params.NoBootstrapExe && !Params.Rocket)
+		if (!Params.NoBootstrapExe && !Params.Rocket && (TargetMask & ProjectBuildTargets.Bootstrap) == ProjectBuildTargets.Bootstrap)
 		{
 			UnrealBuildTool.UnrealTargetPlatform[] BootstrapPackagedGamePlatforms = { UnrealBuildTool.UnrealTargetPlatform.Win32, UnrealBuildTool.UnrealTargetPlatform.Win64 };
 			foreach(UnrealBuildTool.UnrealTargetPlatform BootstrapPackagedGamePlatform in BootstrapPackagedGamePlatforms)
@@ -107,7 +121,7 @@ public partial class Project : CommandUtils
 				}
 			}
 		}
-		if (Params.CrashReporter && !Params.Rocket)
+		if (Params.CrashReporter && !Params.Rocket && (TargetMask & ProjectBuildTargets.CrashReporter) == ProjectBuildTargets.CrashReporter)
 		{
 			foreach (var CrashReportPlatform in CrashReportPlatforms)
 			{
@@ -117,7 +131,7 @@ public partial class Project : CommandUtils
 				}
 			}
 		}
-		if (Params.HasProgramTargets && !Params.Rocket)
+		if (Params.HasProgramTargets && !Params.Rocket && (TargetMask & ProjectBuildTargets.Programs) == ProjectBuildTargets.Programs)
 		{
 			foreach (var BuildConfig in Params.ClientConfigsToBuild)
 			{
diff --git a/Engine/Source/Programs/AutomationTool/Scripts/CookCommand.Automation.cs b/Engine/Source/Programs/AutomationTool/Scripts/CookCommand.Automation.cs
--- a/Engine/Source/Programs/AutomationTool/Scripts/CookCommand.Automation.cs
+++ b/Engine/Source/Programs/AutomationTool/Scripts/CookCommand.Automation.cs
@@ -216,11 +216,20 @@ public partial class Project : CommandUtils
                     CommandletParams += " -compressed";
                 }
                 // we provide the option for users to run a conversion on certain (script) assets, translating them 
-                // into native source code... to use those classes over the assets, we need to incorporate the 
-                // generated plugin into the cooker
-                if (Params.UseNativizedScriptPlugin())
+                // into native source code... the cooker needs to 
+                if (Params.RunAssetNativization)
                 {
-                    CommandletParams += " -PLUGIN=\"" + Params.NativizedScriptPlugin + "\"";
+                    string GeneratedPluginName = "NativizedScript";
+                    string ProjectDir = Params.RawProjectPath.Directory.ToString();
+                    string GeneratedPluginDirectory = Path.GetFullPath(CommandUtils.CombinePaths(ProjectDir, "Intermediate", "Plugins", GeneratedPluginName));
+                    string PluginDescFilename  = GeneratedPluginName + ".uplugin";
+                    string GeneratedPluginPath = Path.GetFullPath(CommandUtils.CombinePaths(GeneratedPluginDirectory, PluginDescFilename));
+
+                    CommandletParams += " -NativizeAssets -NativizedAssetPlugin=\"" + GeneratedPluginPath + "\"";
+
+                    // fill out this, so as to coordinate with build automation
+                    // (the builder now needs to compile in the generated assets)
+                    Params.NativizedScriptPlugin = new FileReference(GeneratedPluginPath);
                 }
                 if (Params.HasAdditionalCookerOptions)
                 {
diff --git a/Engine/Source/Programs/UnrealBuildTool/Configuration/UEBuildModule.cs b/Engine/Source/Programs/UnrealBuildTool/Configuration/UEBuildModule.cs
--- a/Engine/Source/Programs/UnrealBuildTool/Configuration/UEBuildModule.cs
+++ b/Engine/Source/Programs/UnrealBuildTool/Configuration/UEBuildModule.cs
@@ -31,6 +31,26 @@ namespace UnrealBuildTool
 		{
 			return ModuleType != UEBuildModuleType.Game;
 		}
+		public static UEBuildModuleType FromHostType(ModuleHostType ModuleType)
+		{
+			switch (ModuleType)
+			{
+				case ModuleHostType.Runtime:
+				case ModuleHostType.RuntimeNoCommandlet:
+                case ModuleHostType.ServerOnly:
+					return UEBuildModuleType.Runtime;
+				case ModuleHostType.Developer:
+					return UEBuildModuleType.Developer;
+				case ModuleHostType.Editor:
+				case ModuleHostType.EditorNoCommandlet:
+					return UEBuildModuleType.Editor;
+				case ModuleHostType.Program:
+					return UEBuildModuleType.Program;
+				default:
+					return UEBuildModuleType.Unknown;
+			}
+		}
+
 	}
 	/// <summary>
 	/// Distribution level of module.
diff --git a/Engine/Source/Programs/UnrealBuildTool/Configuration/UEBuildTarget.cs b/Engine/Source/Programs/UnrealBuildTool/Configuration/UEBuildTarget.cs
--- a/Engine/Source/Programs/UnrealBuildTool/Configuration/UEBuildTarget.cs
+++ b/Engine/Source/Programs/UnrealBuildTool/Configuration/UEBuildTarget.cs
@@ -3577,7 +3577,26 @@ namespace UnrealBuildTool
 						// not a plugin, see if it is a game module 
 						if (!ModuleFileName.IsUnderDirectory(UnrealBuildTool.EngineSourceDirectory))
 						{
-							ModuleType = UEBuildModuleType.Game;
+							if (ProjectDescriptor != null && ProjectDescriptor.Modules != null)
+							{
+								foreach (ModuleDescriptor ProjectModule in ProjectDescriptor.Modules)
+								{
+									if (ProjectModule.Name == ModuleName)
+									{
+										ModuleType = UEBuildModuleTypeExtensions.FromHostType(ProjectModule.Type);
+										if (ModuleType == UEBuildModuleType.Runtime)
+										{
+											// We expect game runtime modules to be marked as Game module type.
+											ModuleType = UEBuildModuleType.Game;
+										}
+									}
+								}
+							}
+							if (ModuleType == UEBuildModuleType.Unknown)
+							{
+								// No descriptor file or module was not on the list
+								ModuleType = UEBuildModuleType.Game;
+							}						
 						}
 						else
 						{
diff --git a/Engine/Source/Programs/UnrealBuildTool/System/DirectoryLookupCache.cs b/Engine/Source/Programs/UnrealBuildTool/System/DirectoryLookupCache.cs
--- a/Engine/Source/Programs/UnrealBuildTool/System/DirectoryLookupCache.cs
+++ b/Engine/Source/Programs/UnrealBuildTool/System/DirectoryLookupCache.cs
@@ -112,6 +112,11 @@ namespace UnrealBuildTool
 			}
 		}
 
+        static public void InvalidateCachedDirectory(DirectoryReference Directory)
+        {
+            Directories.Remove(Directory);
+        }
+
 		static Dictionary<DirectoryReference, DirectoryCache> Directories = new Dictionary<DirectoryReference, DirectoryCache>();
 	}
 }
diff --git a/Engine/Source/Programs/UnrealBuildTool/System/ModuleDescriptor.cs b/Engine/Source/Programs/UnrealBuildTool/System/ModuleDescriptor.cs
--- a/Engine/Source/Programs/UnrealBuildTool/System/ModuleDescriptor.cs
+++ b/Engine/Source/Programs/UnrealBuildTool/System/ModuleDescriptor.cs
@@ -16,6 +16,7 @@ namespace UnrealBuildTool
 		Editor,
 		EditorNoCommandlet,
 		Program,
+        ServerOnly,
 	}
 
 	public enum ModuleLoadingPhase
@@ -189,6 +190,8 @@ namespace UnrealBuildTool
 					return TargetType == TargetRules.TargetType.Editor || bBuildEditor;
 				case ModuleHostType.Program:
 					return TargetType == TargetRules.TargetType.Program;
+                case ModuleHostType.ServerOnly:
+                    return TargetType != TargetRules.TargetType.Client;
 			}
 
 			return false;
diff --git a/Engine/Source/Programs/UnrealBuildTool/System/RulesCompiler.cs b/Engine/Source/Programs/UnrealBuildTool/System/RulesCompiler.cs
--- a/Engine/Source/Programs/UnrealBuildTool/System/RulesCompiler.cs
+++ b/Engine/Source/Programs/UnrealBuildTool/System/RulesCompiler.cs
@@ -1688,6 +1688,13 @@ namespace UnrealBuildTool
 			return SourceFiles;
 		}
 
+        public static void InvalidateRulesFileCache(string DirectoryPath)
+        {
+            DirectoryReference Directory = new DirectoryReference(DirectoryPath);
+            RootFolderToRulesFileCache.Remove(Directory);
+            DirectoryLookupCache.InvalidateCachedDirectory(Directory);
+        }
+
 		private static IReadOnlyList<FileReference> FindAllRulesFiles(DirectoryReference Directory, RulesFileType Type)
 		{
 			// Check to see if we've already cached source files for this folder
diff --git a/Engine/Source/Programs/UnrealHeaderTool/Private/CodeGenerator.cpp b/Engine/Source/Programs/UnrealHeaderTool/Private/CodeGenerator.cpp
--- a/Engine/Source/Programs/UnrealHeaderTool/Private/CodeGenerator.cpp
+++ b/Engine/Source/Programs/UnrealHeaderTool/Private/CodeGenerator.cpp
@@ -36,6 +36,13 @@ FCompilerMetadataManager GScriptHelper;
 /** C++ name lookup helper */
 FNameLookupCPP NameLookupCPP;
 
+namespace
+{
+	static FString AsTEXT(FString InStr)
+	{
+		return FString::Printf(TEXT("TEXT(\"%s\")"), *InStr);
+	}
+}
 /**
  * Finds exact match of Identifier in string. Returns nullptr if none is found.
  *
@@ -808,7 +815,7 @@ private:
 				return FString::Printf(TEXT("%s::StaticClass()"), NameLookupCPP.GetNameCPP(ItemClass));
 			}
 
-			if (!bRequiresValidObject && ItemClass->HasAllClassFlags(CLASS_Native))
+			if (!bRequiresValidObject)
 			{
 				Suffix = TEXT("_NoRegister");
 			}
@@ -891,6 +898,21 @@ FString FNativeClassHeaderGenerator::GetOverriddenName(const UField* Item)
 	return Item->GetName();
 }
 
+FName FNativeClassHeaderGenerator::GetOverriddenFName(const UField* Item)
+{
+	FString OverriddenName = Item->GetMetaData(TEXT("OverrideNativeName"));
+	if (!OverriddenName.IsEmpty())
+	{
+		return FName(*OverriddenName);
+	}
+	return Item->GetFName();
+}
+
+FString FNativeClassHeaderGenerator::GetOverriddenPathName(const UField* Item)
+{
+	return FString::Printf(TEXT("%s.%s"), *FClass::GetTypePackageName(Item), *GetOverriddenName(Item));
+}
+
 FString FNativeClassHeaderGenerator::GetOverriddenNameForLiteral(const UField* Item)
 {
 	FString OverriddenName = Item->GetMetaData(TEXT("OverrideNativeName"));
@@ -924,6 +946,10 @@ FString FNativeClassHeaderGenerator::PropertyNew(FString& Meta, UProperty* Prop,
 			TargetClass = SubclassOfProperty->MetaClass;
 		}
 		ExtraArgs = FString::Printf(TEXT(", %s"), *GetSingletonName(TargetClass, false)); 
+		if (UClassProperty* ClassProperty = Cast<UClassProperty>(Prop))
+		{
+			ExtraArgs += FString::Printf(TEXT(", %s"), *GetSingletonName(ClassProperty->PropertyClass, false));
+		}
 	}
 	else if (UInterfaceProperty* InterfaceProperty = Cast<UInterfaceProperty>(Prop))
 	{
@@ -967,9 +993,11 @@ FString FNativeClassHeaderGenerator::PropertyNew(FString& Meta, UProperty* Prop,
 		ExtraArgs = FString::Printf(TEXT(", %s"), *GetSingletonName(TargetFunction));
 	}
 
-	FString Constructor = FString::Printf(TEXT("new(EC_InternalUseOnlyConstructor, %s, TEXT(\"%s\"), RF_Public|RF_Transient|RF_MarkAsNative) U%s(%s, 0x%016llx%s);"),
+	const TCHAR* UPropertyObjectFlags = FClass::IsOwnedByDynamicType(Prop) ? TEXT("RF_Public|RF_Transient") : TEXT("RF_Public|RF_Transient|RF_MarkAsNative");
+	FString Constructor = FString::Printf(TEXT("new(EC_InternalUseOnlyConstructor, %s, TEXT(\"%s\"), %s) U%s(%s, 0x%016llx%s);"),
 		*OuterString,
 		*FNativeClassHeaderGenerator::GetOverriddenName(Prop),
+		UPropertyObjectFlags,
 		*Prop->GetClass()->GetName(), 
 		*PropMacro,
 		Prop->PropertyFlags & ~CPF_ComputedFlags, 
@@ -1211,7 +1239,7 @@ void FNativeClassHeaderGenerator::ExportGeneratedPackageInitCode(const UPackage*
 		GeneratedFunctionText.Logf(TEXT("#endif\r\n"));
 	}
 
-	GeneratedFunctionText.Logf(TEXT("\t\t\tReturnPackage->SetPackageFlags(PKG_CompiledIn | 0x%08X);\r\n"), InPackage->GetPackageFlags() & (PKG_ClientOptional | PKG_ServerSideOnly | PKG_EditorOnly));
+	GeneratedFunctionText.Logf(TEXT("\t\t\tReturnPackage->SetPackageFlags(PKG_CompiledIn | 0x%08X);\r\n"), InPackage->GetPackageFlags() & (PKG_ClientOptional | PKG_ServerSideOnly | PKG_EditorOnly | PKG_Developer));
 	TheFlagAudit.Add(InPackage, TEXT("PackageFlags"), InPackage->GetPackageFlags());
 	{
 		FGuid Guid;
@@ -1342,7 +1370,10 @@ void FNativeClassHeaderGenerator::ExportNativeGeneratedInitCode(FClass* Class, F
 		{
 			GeneratedClassRegisterFunctionText.Logf(TEXT("\t\t\t%s;\r\n"), *GetSingletonName(Class->GetSuperClass()));
 		}
-		GeneratedClassRegisterFunctionText.Logf(TEXT("\t\t\t%s;\r\n"), *GetPackageSingletonName(CastChecked<UPackage>(Class->GetOutermost())));
+		if (!bIsDynamic)
+		{
+			GeneratedClassRegisterFunctionText.Logf(TEXT("\t\t\t%s;\r\n"), *GetPackageSingletonName(CastChecked<UPackage>(Class->GetOutermost())));
+		}
 		GeneratedClassRegisterFunctionText.Logf(TEXT("\t\t\tOuterClass = %s::StaticClass();\r\n"), NameLookupCPP.GetNameCPP(Class));
 		GeneratedClassRegisterFunctionText.Logf(TEXT("\t\t\tif (!(OuterClass->ClassFlags & CLASS_Constructed))\r\n"), NameLookupCPP.GetNameCPP(Class));
 		GeneratedClassRegisterFunctionText.Logf(TEXT("\t\t\t{\r\n"), NameLookupCPP.GetNameCPP(Class));
@@ -1355,6 +1386,19 @@ void FNativeClassHeaderGenerator::ExportNativeGeneratedInitCode(FClass* Class, F
 		GeneratedClassRegisterFunctionText.Logf(TEXT("\r\n"));
 
 		FString OuterString = FString(TEXT("OuterClass"));
+
+		{
+			TMap<FName, FString>* MetaDataMap = UMetaData::GetMapForObject(Class);
+			FClassMetaData* ClassMetaData = GScriptHelper.FindClassData(Class);
+			if (MetaDataMap && ClassMetaData && Class)
+			{
+				if (!ClassMetaData->bObjectInitializerConstructorDeclared && ClassMetaData->bDefaultConstructorDeclared)
+				{
+					MetaDataMap->Add(FName(TEXT("OnlyDefaultConstructorDeclared")), FString());
+				}
+			}
+		}
+
 		FString Meta = GetMetaDataCodeForObject(Class, *OuterString, TEXT("\t\t\t\t"));
 		// properties
 		{
@@ -1426,7 +1470,13 @@ void FNativeClassHeaderGenerator::ExportNativeGeneratedInitCode(FClass* Class, F
 			GeneratedClassRegisterFunctionText.Logf(TEXT("#endif\r\n"));
 		}
 
+		if (bIsDynamic)
+		{
+			GeneratedClassRegisterFunctionText.Logf(TEXT("\t\t\t\tOuterClass->GetDefaultObject();\r\n"));
+		}
+
 		GeneratedClassRegisterFunctionText.Logf(TEXT("\t\t\t}\r\n"));
+		
 		GeneratedClassRegisterFunctionText.Logf(TEXT("\t\t}\r\n"));
 		GeneratedClassRegisterFunctionText.Logf(TEXT("\t\tcheck(OuterClass->GetClass());\r\n"));
 		GeneratedClassRegisterFunctionText.Logf(TEXT("\t\treturn OuterClass;\r\n"));
@@ -1449,9 +1499,11 @@ void FNativeClassHeaderGenerator::ExportNativeGeneratedInitCode(FClass* Class, F
 		FString OverriddenClassName = *FNativeClassHeaderGenerator::GetOverriddenName(Class);
 		const TCHAR* ClassNameCPP = NameLookupCPP.GetNameCPP(Class);
 
-		GeneratedFunctionText.Logf(TEXT("\tstatic FCompiledInDefer Z_CompiledInDefer_UClass_%s(%s, &%s::StaticClass, TEXT(\"%s\"), %s);\r\n"),
+		GeneratedFunctionText.Logf(TEXT("\tstatic FCompiledInDefer Z_CompiledInDefer_UClass_%s(%s, &%s::StaticClass, TEXT(\"%s\"), %s, %s, %s);\r\n"),
 			ClassNameCPP, *SingletonName, ClassNameCPP, bIsDynamic ? *OverriddenClassName : ClassNameCPP,
-			bIsDynamic ? TEXT("true") : TEXT("false"));
+			bIsDynamic ? TEXT("true") : TEXT("false"),
+			bIsDynamic ? *AsTEXT(FClass::GetTypePackageName(Class)) : TEXT("nullptr"),
+			bIsDynamic ? *AsTEXT(FNativeClassHeaderGenerator::GetOverriddenPathName(Class)) : TEXT("nullptr"));
 		
 		// Append base class' CRC at the end of the generated code, this will force update derived classes
 		// when base class changes during hot-reload.
@@ -1519,7 +1571,19 @@ void FNativeClassHeaderGenerator::ExportFunction(UFunction* Function, FScope* Sc
 		CurrentFunctionText.Logf(TEXT("\t\tUObject* Outer=nullptr;\r\n"));
 	}
 
-	CurrentFunctionText.Logf(TEXT("\t\tstatic UFunction* ReturnFunction = NULL;\r\n"));
+	UField* FieldOuter = Cast<UField>(Function->GetOuter());
+	const bool bIsDynamic = (FieldOuter && FClass::IsDynamic(FieldOuter));
+
+	if (!bIsDynamic)
+	{
+		CurrentFunctionText.Logf(TEXT("\t\tstatic UFunction* ReturnFunction = NULL;\r\n"));
+	}
+	else
+	{
+		CurrentFunctionText.Logf(TEXT("\t\tUFunction* ReturnFunction = static_cast<UFunction*>(StaticFindObjectFast( UFunction::StaticClass(), Outer, %s ));\r\n")
+			, *FNativeClassHeaderGenerator::GetOverriddenNameForLiteral(Function));
+	}
+
 	CurrentFunctionText.Logf(TEXT("\t\tif (!ReturnFunction)\r\n"));
 	CurrentFunctionText.Logf(TEXT("\t\t{\r\n"));
 	FString SuperFunctionString(TEXT("NULL"));
@@ -1549,11 +1613,12 @@ void FNativeClassHeaderGenerator::ExportFunction(UFunction* Function, FScope* Sc
 		StructureSize = FString::Printf(TEXT(", sizeof(%s)"), *GetEventStructParamsName(TempFunction->GetOuter(), *FunctionName));
 	}
 
-	FString UFunctionType = bIsDelegate ? TEXT("UDelegateFunction") : TEXT("UFunction");
-
-	CurrentFunctionText.Logf(TEXT("\t\t\tReturnFunction = new(EC_InternalUseOnlyConstructor, Outer, TEXT(\"%s\"), RF_Public|RF_Transient|RF_MarkAsNative) %s(FObjectInitializer(), %s, 0x%08X, %d%s);\r\n"),
+	const TCHAR* UFunctionType = bIsDelegate ? TEXT("UDelegateFunction") : TEXT("UFunction");
+	const TCHAR* UFunctionObjectFlags = FClass::IsOwnedByDynamicType(Function) ? TEXT("RF_Public|RF_Transient") : TEXT("RF_Public|RF_Transient|RF_MarkAsNative");
+	CurrentFunctionText.Logf(TEXT("\t\t\tReturnFunction = new(EC_InternalUseOnlyConstructor, Outer, TEXT(\"%s\"), %s) %s(FObjectInitializer(), %s, 0x%08X, %d%s);\r\n"),
 		*FNativeClassHeaderGenerator::GetOverriddenName(Function),
-		*UFunctionType,
+		UFunctionObjectFlags,
+		UFunctionType,
 		*SuperFunctionString,
 		Function->FunctionFlags,
 		(uint32)Function->RepOffset,
@@ -2778,9 +2843,11 @@ void FNativeClassHeaderGenerator::ExportGeneratedStructBodyMacros(const TArray<U
 			GeneratedPackageCPP.Logf(TEXT("\treturn Singleton;\r\n"));
 			GeneratedPackageCPP.Logf(TEXT("}\r\n"));
 
-			GeneratedPackageCPP.Logf(TEXT("static FCompiledInDeferStruct Z_CompiledInDeferStruct_UScriptStruct_%s(%s::StaticStruct, TEXT(\"%s\"), TEXT(\"%s\"), %s);\r\n"),
+			GeneratedPackageCPP.Logf(TEXT("static FCompiledInDeferStruct Z_CompiledInDeferStruct_UScriptStruct_%s(%s::StaticStruct, TEXT(\"%s\"), TEXT(\"%s\"), %s, %s, %s);\r\n"),
 				StructNameCPP, StructNameCPP, *Struct->GetOutermost()->GetName(), *Struct->GetName(),
-				bIsDynamic ? TEXT("true") : TEXT("false"));
+				bIsDynamic ? TEXT("true") : TEXT("false"),
+				bIsDynamic ? *AsTEXT(FClass::GetTypePackageName(Struct)) : TEXT("nullptr"),
+				bIsDynamic ? *AsTEXT(FNativeClassHeaderGenerator::GetOverriddenPathName(Struct)) : TEXT("nullptr"));
 
 			// Generate StaticRegisterNatives equivalent for structs without classes.
 			if (!Struct->GetOuter()->IsA(UStruct::StaticClass()))
@@ -2836,7 +2903,14 @@ void FNativeClassHeaderGenerator::ExportGeneratedStructBodyMacros(const TArray<U
 			}
 
 			GeneratedStructRegisterFunctionText.Logf(TEXT("\t\textern uint32 %s();\r\n"), *CRCFuncName);
-			GeneratedStructRegisterFunctionText.Logf(TEXT("\t\tstatic UScriptStruct* ReturnStruct = FindExistingStructIfHotReloadOrDynamic(Outer, TEXT(\"%s\"), sizeof(%s), %s(), %s);\r\n"), *ScriptStruct->GetName(), NameLookupCPP.GetNameCPP(Struct), *CRCFuncName, bIsDynamic ? TEXT("true") : TEXT("false"));
+			if (!bIsDynamic)
+			{
+				GeneratedStructRegisterFunctionText.Logf(TEXT("\t\tstatic UScriptStruct* ReturnStruct = FindExistingStructIfHotReloadOrDynamic(Outer, TEXT(\"%s\"), sizeof(%s), %s(), false);\r\n"), *ScriptStruct->GetName(), NameLookupCPP.GetNameCPP(Struct), *CRCFuncName);
+			}
+			else
+			{
+				GeneratedStructRegisterFunctionText.Logf(TEXT("\t\tUScriptStruct* ReturnStruct = FindExistingStructIfHotReloadOrDynamic(Outer, TEXT(\"%s\"), sizeof(%s), %s(), true);\r\n"), *ScriptStruct->GetName(), NameLookupCPP.GetNameCPP(Struct), *CRCFuncName);
+			}
 			GeneratedStructRegisterFunctionText.Logf(TEXT("\t\tif (!ReturnStruct)\r\n"));
 			GeneratedStructRegisterFunctionText.Logf(TEXT("\t\t{\r\n"));
 			FString BaseStructString(TEXT("NULL"));
@@ -2858,8 +2932,10 @@ void FNativeClassHeaderGenerator::ExportGeneratedStructBodyMacros(const TArray<U
 				ExplicitSizeString = FString::Printf(TEXT(", sizeof(%s), ALIGNOF(%s)"), NameLookupCPP.GetNameCPP(ScriptStruct), NameLookupCPP.GetNameCPP(ScriptStruct));
 			}
 
-			GeneratedStructRegisterFunctionText.Logf(TEXT("\t\t\tReturnStruct = new(EC_InternalUseOnlyConstructor, Outer, TEXT(\"%s\"), RF_Public|RF_Transient|RF_MarkAsNative) UScriptStruct(FObjectInitializer(), %s, %s, EStructFlags(0x%08X)%s);\r\n"),
+			const TCHAR* UStructObjectFlags = bIsDynamic ? TEXT("RF_Public|RF_Transient") : TEXT("RF_Public|RF_Transient|RF_MarkAsNative");
+			GeneratedStructRegisterFunctionText.Logf(TEXT("\t\t\tReturnStruct = new(EC_InternalUseOnlyConstructor, Outer, TEXT(\"%s\"), %s) UScriptStruct(FObjectInitializer(), %s, %s, EStructFlags(0x%08X)%s);\r\n"),
 				*ScriptStruct->GetName(),
+				UStructObjectFlags,
 				*BaseStructString,
 				*CppStructOpsString,
 				(uint32)(ScriptStruct->StructFlags & ~STRUCT_ComputedFlags),
@@ -2959,10 +3035,11 @@ void FNativeClassHeaderGenerator::ExportGeneratedEnumsInitCode(const TArray<UEnu
 		GeneratedPackageCPP.Logf(TEXT("\treturn Singleton;\r\n"));
 		GeneratedPackageCPP.Logf(TEXT("}\r\n"));
 
-		GeneratedPackageCPP.Logf(TEXT("static FCompiledInDeferEnum Z_CompiledInDeferEnum_UEnum_%s(%s_StaticEnum, TEXT(\"%s\"), TEXT(\"%s\"), %s);\r\n"), 
+		GeneratedPackageCPP.Logf(TEXT("static FCompiledInDeferEnum Z_CompiledInDeferEnum_UEnum_%s(%s_StaticEnum, TEXT(\"%s\"), TEXT(\"%s\"), %s, %s, %s);\r\n"), 
 			*Enum->GetName(), *Enum->GetName(), *Enum->GetOutermost()->GetName(), *Enum->GetName(),
-			bIsDynamic ? TEXT("true") : TEXT("false"));
-
+			bIsDynamic ? TEXT("true") : TEXT("false"),
+			bIsDynamic ? *AsTEXT(FClass::GetTypePackageName(Enum)) : TEXT("nullptr"),
+			bIsDynamic ? *AsTEXT(FNativeClassHeaderGenerator::GetOverriddenPathName(Enum)) : TEXT("nullptr"));
 		{
 			const FString EnumSingletonName = GetSingletonName(Enum);
 			GeneratedFunctionDeclarations.Log(FTypeSingletonCache::Get(Enum).GetExternDecl());
@@ -2987,11 +3064,19 @@ void FNativeClassHeaderGenerator::ExportGeneratedEnumsInitCode(const TArray<UEnu
 				GeneratedEnumRegisterFunctionText.Logf(TEXT("\t\tUPackage* Outer = FindOrConstructDynamicTypePackage(TEXT(\"%s\"));\r\n"), *PackageSingletonName);
 			}
 			GeneratedEnumRegisterFunctionText.Logf(TEXT("\t\textern uint32 %s();\r\n"), *CRCFuncName);
-			GeneratedEnumRegisterFunctionText.Logf(TEXT("\t\tstatic UEnum* ReturnEnum = FindExistingEnumIfHotReload(Outer, TEXT(\"%s\"), 0, %s());\r\n"), *Enum->GetName(), *CRCFuncName);
+			if (!bIsDynamic)
+			{
+				GeneratedEnumRegisterFunctionText.Logf(TEXT("\t\tstatic UEnum* ReturnEnum = FindExistingEnumIfHotReloadOrDynamic(Outer, TEXT(\"%s\"), 0, %s(), false);\r\n"), *Enum->GetName(), *CRCFuncName);
+			}
+			else
+			{
+				GeneratedEnumRegisterFunctionText.Logf(TEXT("\t\tUEnum* ReturnEnum = FindExistingEnumIfHotReloadOrDynamic(Outer, TEXT(\"%s\"), 0, %s(), true);\r\n"), *Enum->GetName(), *CRCFuncName);
+			}
 			GeneratedEnumRegisterFunctionText.Logf(TEXT("\t\tif (!ReturnEnum)\r\n"));
 			GeneratedEnumRegisterFunctionText.Logf(TEXT("\t\t{\r\n"));
 
-			GeneratedEnumRegisterFunctionText.Logf(TEXT("\t\t\tReturnEnum = new(EC_InternalUseOnlyConstructor, Outer, TEXT(\"%s\"), RF_Public|RF_Transient|RF_MarkAsNative) UEnum(FObjectInitializer());\r\n"), *Enum->GetName());
+			const TCHAR* UEnumObjectFlags = bIsDynamic ? TEXT("RF_Public|RF_Transient") : TEXT("RF_Public|RF_Transient|RF_MarkAsNative");
+			GeneratedEnumRegisterFunctionText.Logf(TEXT("\t\t\tReturnEnum = new(EC_InternalUseOnlyConstructor, Outer, TEXT(\"%s\"), %s) UEnum(FObjectInitializer());\r\n"), *Enum->GetName(), UEnumObjectFlags);
 			GeneratedEnumRegisterFunctionText.Logf(TEXT("\t\t\tTArray<TPair<FName, uint8>> EnumNames;\r\n"));
 			for (int32 Index = 0; Index < Enum->NumEnums(); Index++)
 			{
@@ -5457,6 +5542,11 @@ ECompilationResult::Type PreparseModules(FUHTMakefile& UHTMakefile, const FStrin
 			Package->SetPackageFlags(PKG_EditorOnly);
 		}
 
+		if (Module.ModuleType == EBuildModuleType::Developer)
+		{
+			Package->SetPackageFlags(Package->GetPackageFlags() | PKG_Developer);
+		}
+
 		// Add new module or overwrite whatever we had loaded, that data is obsolete.
 		UHTMakefile.AddPackage(Package);
 		GPackageToManifestModuleMap.Add(Package, &Module);
diff --git a/Engine/Source/Programs/UnrealHeaderTool/Private/HeaderParser.cpp b/Engine/Source/Programs/UnrealHeaderTool/Private/HeaderParser.cpp
--- a/Engine/Source/Programs/UnrealHeaderTool/Private/HeaderParser.cpp
+++ b/Engine/Source/Programs/UnrealHeaderTool/Private/HeaderParser.cpp
@@ -391,11 +391,6 @@ namespace
 			}
 		}
 
-		if ( ( FuncInfo.FunctionFlags & FUNC_NetServer ) && !( FuncInfo.FunctionFlags & FUNC_NetValidate ) )
-		{
-			FError::Throwf( TEXT( "Server RPC missing 'WithValidation' keyword in the UPROPERTY() declaration statement.  Required for security purposes." ) );
-		}
-
 		if (FuncInfo.FunctionFlags & FUNC_Net)
 		{
 			// Network replicated functions are always events
@@ -2482,7 +2477,7 @@ void FHeaderParser::VerifyRepNotifyCallbacks( UClass* TargetClass )
 				for( UField* TestField = SearchClass->Children; TestField; TestField = TestField->Next )
 				{
 					UFunction* TestFunc = Cast<UFunction>(TestField);
-					if( TestFunc && TestFunc->GetFName() == Prop->RepNotifyFunc )
+					if (TestFunc && FNativeClassHeaderGenerator::GetOverriddenFName(TestFunc) == Prop->RepNotifyFunc)
 					{
 						TargetFunc = TestFunc;
 						break;
@@ -5565,6 +5560,12 @@ void FHeaderParser::CompileFunctionDeclaration(FClasses& AllClasses)
 
 	ProcessFunctionSpecifiers(FuncInfo, SpecifiersFound);
 
+	const bool bClassGeneratedFromBP = FClass::IsDynamic(GetCurrentClass());
+	if ((FuncInfo.FunctionFlags & FUNC_NetServer) && !(FuncInfo.FunctionFlags & FUNC_NetValidate) && !bClassGeneratedFromBP)
+	{
+		FError::Throwf(TEXT("Server RPC missing 'WithValidation' keyword in the UPROPERTY() declaration statement.  Required for security purposes."));
+	}
+
 	if ((0 != (FuncInfo.FunctionExportFlags & FUNCEXPORT_CustomThunk)) && !MetaData.Contains("CustomThunk"))
 	{
 		MetaData.Add(TEXT("CustomThunk"), TEXT("true"));
diff --git a/Engine/Source/Programs/UnrealHeaderTool/Private/NativeClassExporter.h b/Engine/Source/Programs/UnrealHeaderTool/Private/NativeClassExporter.h
--- a/Engine/Source/Programs/UnrealHeaderTool/Private/NativeClassExporter.h
+++ b/Engine/Source/Programs/UnrealHeaderTool/Private/NativeClassExporter.h
@@ -191,12 +191,6 @@ private:
 	FString GetSingletonName(UField* Item, bool bRequiresValidObject=true);
 	/** Return the name of the singleton function that returns the UObject for Item */
 	FString GetSingletonName(FClass* Item, bool bRequiresValidObject=true);
-
-	/** 
-	 * Properties in source files generated from blueprint assets have a symbol name that differs from the source asset.
-	 * This function returns the original name of the field (rather than the native, symbol name).
-	 */
-	FString GetOverriddenName(const UField* Item);
 	
 	/** 
 	 * Returns the name (overridden if marked up) with TEXT("") or "" wrappers for use in a string literal.
@@ -539,6 +533,14 @@ private:
 
 public:
 
+	/** 
+	 * Properties in source files generated from blueprint assets have a symbol name that differs from the source asset.
+	 * This function returns the original name of the field (rather than the native, symbol name).
+	 */
+	static FString GetOverriddenName(const UField* Item);
+	static FName GetOverriddenFName(const UField* Item);
+	static FString GetOverriddenPathName(const UField* Item);
+
 	// Constructor
 	FNativeClassHeaderGenerator(
 		const UPackage* InPackage,
diff --git a/Engine/Source/Programs/UnrealHeaderTool/Private/ParserClass.cpp b/Engine/Source/Programs/UnrealHeaderTool/Private/ParserClass.cpp
--- a/Engine/Source/Programs/UnrealHeaderTool/Private/ParserClass.cpp
+++ b/Engine/Source/Programs/UnrealHeaderTool/Private/ParserClass.cpp
@@ -87,13 +87,25 @@ void FClass::GetShowCategories(TArray<FString>& OutShowCategories) const
 	}
 }
 
-bool FClass::IsDynamic(UField* Field)
+bool FClass::IsDynamic(const UField* Field)
 {
 	static const FName NAME_ReplaceConverted(TEXT("ReplaceConverted"));
 	return Field->HasMetaData(NAME_ReplaceConverted);
 }
 
-FString FClass::GetTypePackageName(UField* Field)
+bool FClass::IsOwnedByDynamicType(const UField* Field)
+{
+	for (const UField* OuterField = Cast<const UField>(Field->GetOuter()); OuterField; OuterField = Cast<const UField>(OuterField->GetOuter()))
+	{
+		if (IsDynamic(OuterField))
+		{
+			return true;
+		}
+	}
+	return false;
+}
+
+FString FClass::GetTypePackageName(const UField* Field)
 {
 	static const FName NAME_ReplaceConverted(TEXT("ReplaceConverted"));
 	FString PackageName = Field->GetMetaData(NAME_ReplaceConverted);
@@ -112,4 +124,3 @@ FString FClass::GetTypePackageName(UField* Field)
 	}
 	return PackageName;
 }
-
diff --git a/Engine/Source/Programs/UnrealHeaderTool/Private/ParserClass.h b/Engine/Source/Programs/UnrealHeaderTool/Private/ParserClass.h
--- a/Engine/Source/Programs/UnrealHeaderTool/Private/ParserClass.h
+++ b/Engine/Source/Programs/UnrealHeaderTool/Private/ParserClass.h
@@ -46,8 +46,11 @@ public:
 	void GetShowCategories(TArray<FString>& OutShowCategories) const;
 
 	/** Helper function that checks if the field is a dynamic type (can be constructed post-startup) */
-	static bool IsDynamic(UField* Field);
+	static bool IsDynamic(const UField* Field);
+
+	/** Helper function that checks if the field is belongs to a dynamic type */
+	static bool IsOwnedByDynamicType(const UField* Field);
 
 	/** Helper function to get the source replaced package name */
-	static FString GetTypePackageName(UField* Field);
+	static FString GetTypePackageName(const UField* Field);
 };
diff --git a/Engine/Source/Runtime/AIModule/Classes/AIController.h b/Engine/Source/Runtime/AIModule/Classes/AIController.h
--- a/Engine/Source/Runtime/AIModule/Classes/AIController.h
+++ b/Engine/Source/Runtime/AIModule/Classes/AIController.h
@@ -269,7 +269,7 @@ public:
 	virtual FVector GetFocalPointOnActor(const AActor *Actor) const;
 
 	/** Set the position that controller should be looking at. */
-	UFUNCTION(BlueprintCallable, Category="AI", meta=(DisplayName="SetFocalPoint"))
+	UFUNCTION(BlueprintCallable, Category="AI", meta=(DisplayName="SetFocalPoint", Keywords="focus"))
 	void K2_SetFocalPoint(FVector FP);
 
 	/** Set Focus for actor, will set FocalPoint as a result. */
diff --git a/Engine/Source/Runtime/AIModule/Private/AIController.cpp b/Engine/Source/Runtime/AIModule/Private/AIController.cpp
--- a/Engine/Source/Runtime/AIModule/Private/AIController.cpp
+++ b/Engine/Source/Runtime/AIModule/Private/AIController.cpp
@@ -476,7 +476,8 @@ void AAIController::Possess(APawn* InPawn)
 	}
 	CachedGameplayTasksComponent = GTComp;
 
-	if (GTComp)
+	// Prevents re-entrant issues.
+	if (GTComp && !GTComp->OnClaimedResourcesChange.Contains(this, GET_FUNCTION_NAME_CHECKED(AAIController, OnGameplayTaskResourcesClaimed)))
 	{
 		GTComp->OnClaimedResourcesChange.AddDynamic(this, &AAIController::OnGameplayTaskResourcesClaimed);
 
diff --git a/Engine/Source/Runtime/Core/Private/Math/UnrealMath.cpp b/Engine/Source/Runtime/Core/Private/Math/UnrealMath.cpp
--- a/Engine/Source/Runtime/Core/Private/Math/UnrealMath.cpp
+++ b/Engine/Source/Runtime/Core/Private/Math/UnrealMath.cpp
@@ -5,6 +5,7 @@
 =============================================================================*/
 
 #include "CorePrivatePCH.h"
+#include "PropertyPortFlags.h"
 
 DEFINE_LOG_CATEGORY(LogUnrealMath);
 
@@ -3168,4 +3169,14 @@ void FMath::PolarToCartesian(const FVector2D InPolar, FVector2D& OutCart)
 {
 	OutCart.X = InPolar.X * Cos(InPolar.Y);
 	OutCart.Y = InPolar.X * Sin(InPolar.Y);
+}
+
+bool FRandomStream::ExportTextItem(FString& ValueStr, FRandomStream const& DefaultValue, class UObject* Parent, int32 PortFlags, class UObject* ExportRootScope) const
+{
+	if (0 != (PortFlags & EPropertyPortFlags::PPF_ExportCpp))
+	{
+		ValueStr += FString::Printf(TEXT("FRandomStream(%i)"), DefaultValue.GetInitialSeed());
+		return true;
+	}
+	return false;
 }
\ No newline at end of file
diff --git a/Engine/Source/Runtime/Core/Public/Math/RandomStream.h b/Engine/Source/Runtime/Core/Public/Math/RandomStream.h
--- a/Engine/Source/Runtime/Core/Public/Math/RandomStream.h
+++ b/Engine/Source/Runtime/Core/Public/Math/RandomStream.h
@@ -282,6 +282,19 @@ public:
 		}
 	}
 
+	/**
+	 * Exports the RandomStreams value to a string.
+	 *
+	 * @param ValueStr Will hold the string value.
+	 * @param DefaultValue The default value.
+	 * @param Parent Not used.
+	 * @param PortFlags Not used.
+	 * @param ExportRootScope Not used.
+	 * @return true on success, false otherwise.
+	 * @see ImportTextItem
+	 */
+	CORE_API bool ExportTextItem(FString& ValueStr, FRandomStream const& DefaultValue, class UObject* Parent, int32 PortFlags, class UObject* ExportRootScope) const;
+
 protected:
 
 	/**
diff --git a/Engine/Source/Runtime/Core/Public/UObject/ObjectVersion.h b/Engine/Source/Runtime/Core/Public/UObject/ObjectVersion.h
--- a/Engine/Source/Runtime/Core/Public/UObject/ObjectVersion.h
+++ b/Engine/Source/Runtime/Core/Public/UObject/ObjectVersion.h
@@ -580,6 +580,8 @@ enum EUnrealEngineObjectUE4Version
 	VER_UE4_BLUEPRINT_CUSTOM_EVENT_CONST_INPUT,
 	// Renamed HighFrequencyGain to LowPassFilterFrequency
 	VER_UE4_USE_LOW_PASS_FILTER_FREQ,
+	// UAnimBlueprintGeneratedClass can be replaced by a dynamic class. Use TSubclassOf<UAnimInstance> instead.
+	VER_UE4_NO_ANIM_BP_CLASS_IN_GAMEPLAY_CODE,
 	// -----<new versions can be added before this line>-------------------------------------------------
 	// - this needs to be the last line (see note below)
 	VER_UE4_AUTOMATIC_VERSION_PLUS_ONE,
