Integrate MTBF analytics to 4.10 => Main
Changes are edits/reimplementations of changes from 4.10 branch - not integrations as far as p4 is concerned to avoid issues in future integrations.
Changes are from the following changes in 4.10
  -CL 2717436 - Added AbnormalShutdown engine analytics event when editor starts after a crash.
  -CL 2724685 - Added AbnormalShutdown engine analytics event when editor starts after a crash.
  -CL 2725534 - Added AbnormalShutdown engine analytics event when editor starts after a crash.
  -CL 2736468 - Improved Engine analytics handling for Editor and games [UE-21892]
  -CL 2736480 - Updating wording in privacy settings text.
  -CL 2738137 - Fix build errors in standalone tools such as BuildPatchTool, CrashReporterClient, etc.
  -CL 2740043 - Improved abnormal termination tracking for analytics to handle multiple instances of the editor.
  -CL 2750764 - Added check for debugger present when reporting abnormal termination to analytics.
  -CL 2753661 - Added FSystemWideCriticalSection for desktop platforms. Used by analytics to lock access to editor instances list in the OS.
  -CL 2753709 - Updating wording in privacy settings text.
  -CL 2755381 - Mac and Linux CIS fix
  -CL 2755831 - Linux CIS fix
  -CL 2773486 - Improvements to previous changes from CL 2753661 based on code review.

diff --git a/Engine/Source/Editor/DetailCustomizations/Private/DetailCustomizations.cpp b/Engine/Source/Editor/DetailCustomizations/Private/DetailCustomizations.cpp
--- a/Engine/Source/Editor/DetailCustomizations/Private/DetailCustomizations.cpp
+++ b/Engine/Source/Editor/DetailCustomizations/Private/DetailCustomizations.cpp
@@ -104,7 +104,7 @@
 #include "TextCustomization.h"
 #include "AnimTrailNodeDetails.h"
 #include "MaterialProxySettingsCustomizations.h"
-
+#include "ImportantToggleSettingCustomization.h"
 
 IMPLEMENT_MODULE( FDetailCustomizationsModule, DetailCustomizations );
 
@@ -299,6 +299,9 @@ void FDetailCustomizationsModule::RegisterObjectCustomizations()
 	RegisterCustomClassLayout("FbxImportUI", FOnGetDetailCustomizationInstance::CreateStatic(&FFbxImportUIDetails::MakeInstance));
 
 	RegisterCustomClassLayout("ConfigHierarchyPropertyView", FOnGetDetailCustomizationInstance::CreateStatic(&FConfigPropertyHelperDetails::MakeInstance));
+
+	RegisterCustomClassLayout("AnalyticsPrivacySettings", FOnGetDetailCustomizationInstance::CreateStatic(&FImportantToggleSettingCustomization::MakeInstance));
+	RegisterCustomClassLayout("EndUserSettings", FOnGetDetailCustomizationInstance::CreateStatic(&FImportantToggleSettingCustomization::MakeInstance));
 }
 
 
@@ -323,4 +326,4 @@ void FDetailCustomizationsModule::RegisterCustomPropertyTypeLayout(FName Propert
 	static FName PropertyEditor("PropertyEditor");
 	FPropertyEditorModule& PropertyModule = FModuleManager::GetModuleChecked<FPropertyEditorModule>(PropertyEditor);
 	PropertyModule.RegisterCustomPropertyTypeLayout(PropertyTypeName, PropertyTypeLayoutDelegate);
-}
+}
\ No newline at end of file
diff --git a/Engine/Source/Editor/DetailCustomizations/Private/ImportantToggleSettingCustomization.cpp b/Engine/Source/Editor/DetailCustomizations/Private/ImportantToggleSettingCustomization.cpp
new file mode 100644
--- /dev/null
+++ b/Engine/Source/Editor/DetailCustomizations/Private/ImportantToggleSettingCustomization.cpp
@@ -0,0 +1,198 @@
+// Copyright 1998-2015 Epic Games, Inc. All Rights Reserved.
+
+#include "DetailCustomizationsPrivatePCH.h"
+#include "ImportantToggleSettingCustomization.h"
+#include "Engine/ImportantToggleSettingInterface.h"
+#include "SHyperlink.h"
+
+#define LOCTEXT_NAMESPACE "ImportantToggleSettingCustomization"
+
+class SImportantToggleButton : public SCompoundWidget
+{
+public:
+
+	SLATE_BEGIN_ARGS(SImportantToggleButton)
+		: _Text()
+	{}
+
+		SLATE_STYLE_ARGUMENT(FCheckBoxStyle, CheckBoxStyle)
+		SLATE_ARGUMENT(FText, Text)
+		SLATE_ARGUMENT(FText, ToolTipText)
+		SLATE_ATTRIBUTE(bool, IsSet)
+		SLATE_EVENT(FSimpleDelegate, OnToggled)
+
+	SLATE_END_ARGS()
+
+	void Construct(const FArguments& InArgs)
+	{
+		OnToggled = InArgs._OnToggled;
+		IsSetAttribute = InArgs._IsSet;
+
+		FSlateFontInfo LargeDetailsFont = IDetailLayoutBuilder::GetDetailFontBold();
+		LargeDetailsFont.Size += 4;
+
+		ChildSlot
+		[
+			SNew(SCheckBox)
+			.Style(InArgs._CheckBoxStyle)
+			.IsChecked(this, &SImportantToggleButton::GetCheckedState)
+			.OnCheckStateChanged(this, &SImportantToggleButton::OnClick)
+			.ToolTipText(InArgs._ToolTipText)
+			.Padding(FMargin(16.0f, 12.0f))
+			.ForegroundColor(FSlateColor::UseForeground())
+			.IsFocusable(true)
+			[
+				SNew(STextBlock)
+				.Text(InArgs._Text)
+				.Font(LargeDetailsFont)
+			]
+		];
+	}
+
+private:
+	void OnClick(ECheckBoxState State)
+	{
+		OnToggled.ExecuteIfBound();
+	}
+
+	ECheckBoxState GetCheckedState() const
+	{
+		return IsSetAttribute.Get() ? ECheckBoxState::Checked : ECheckBoxState::Unchecked;
+	}
+
+private:
+	TAttribute<bool> IsSetAttribute;
+	FSimpleDelegate OnToggled;
+};
+
+//====================================================================================
+// FImportantToggleSettingCustomization
+//====================================================================================
+
+TSharedRef<IDetailCustomization> FImportantToggleSettingCustomization::MakeInstance()
+{
+	return MakeShareable(new FImportantToggleSettingCustomization);
+}
+
+void FImportantToggleSettingCustomization::CustomizeDetails(IDetailLayoutBuilder& DetailBuilder)
+{
+	TArray<TWeakObjectPtr<UObject>> Objects;
+	DetailBuilder.GetObjectsBeingCustomized(Objects);
+
+	if (Objects.Num() == 1)
+	{
+		ToggleSettingObject = Objects[0];
+		IImportantToggleSettingInterface* ToogleSettingInterface = Cast<IImportantToggleSettingInterface>(ToggleSettingObject.Get());
+
+		if (ToogleSettingInterface != nullptr)
+		{
+			FName CategoryName;
+			FName PropertyName;
+			ToogleSettingInterface->GetToogleCategoryAndPropertyNames(CategoryName, PropertyName);
+			IDetailCategoryBuilder& Category = DetailBuilder.EditCategory(CategoryName);
+			TogglePropertyHandle = DetailBuilder.GetProperty(PropertyName);
+
+			FSlateFontInfo StateDescriptionFont = IDetailLayoutBuilder::GetDetailFont();
+			StateDescriptionFont.Size += 4;
+
+			// Customize collision section
+			Category.InitiallyCollapsed(false)
+			.AddProperty(TogglePropertyHandle)
+			.ShouldAutoExpand(true)
+			.CustomWidget()
+			[
+				SNew(SVerticalBox)
+				+ SVerticalBox::Slot()
+				.AutoHeight()
+				.Padding(0.0f, 12.0f, 0.0f, 0.0f)
+				[
+					SNew(SHorizontalBox)
+					+ SHorizontalBox::Slot()
+					.AutoWidth()
+					[
+						SNew(SImportantToggleButton)
+						.CheckBoxStyle(FEditorStyle::Get(), "Property.ToggleButton.Start")
+						.Text(ToogleSettingInterface->GetFalseStateLabel())
+						.ToolTipText(ToogleSettingInterface->GetFalseStateTooltip())
+						.IsSet(this, &FImportantToggleSettingCustomization::IsToggleValue, false)
+						.OnToggled(this, &FImportantToggleSettingCustomization::OnToggledTo, false)
+					]
+					+ SHorizontalBox::Slot()
+					.AutoWidth()
+					[
+						SNew(SImportantToggleButton)
+						.CheckBoxStyle(FEditorStyle::Get(), "Property.ToggleButton.End")
+						.Text(ToogleSettingInterface->GetTrueStateLabel())
+						.ToolTipText(ToogleSettingInterface->GetTrueStateTooltip())
+						.IsSet(this, &FImportantToggleSettingCustomization::IsToggleValue, true)
+						.OnToggled(this, &FImportantToggleSettingCustomization::OnToggledTo, true)
+					]
+					+ SHorizontalBox::Slot()
+					.HAlign(HAlign_Right)
+					.Padding(0.0f, 12.0f)
+					[
+						SNew(SVerticalBox)
+						+ SVerticalBox::Slot()
+						.VAlign(VAlign_Center)
+						[
+							SNew(SHyperlink)
+							.Text(ToogleSettingInterface->GetAdditionalInfoUrlLabel())
+							.OnNavigate(this, &FImportantToggleSettingCustomization::OnNavigateHyperlink, ToogleSettingInterface->GetAdditionalInfoUrl())
+						]
+					]
+				]
+				+ SVerticalBox::Slot()
+				.AutoHeight()
+				.Padding(0.0f, 12.0f)
+				[
+					SNew(STextBlock)
+					.AutoWrapText(true)
+					.Text(this, &FImportantToggleSettingCustomization::GetDescriptionText)
+					.Font(StateDescriptionFont)
+				]
+			];			
+		}
+	}
+
+
+}
+
+bool FImportantToggleSettingCustomization::IsToggleValue(bool bValue) const
+{
+	bool bPropertyValue = false;
+	TogglePropertyHandle->GetValue(bPropertyValue);
+	return bPropertyValue == bValue;
+}
+
+void FImportantToggleSettingCustomization::OnToggledTo(bool bSetTo)
+{
+	TogglePropertyHandle->SetValue(bSetTo);
+}
+
+void FImportantToggleSettingCustomization::OnNavigateHyperlink(FString Url)
+{
+	FPlatformProcess::LaunchURL(*Url, nullptr, nullptr);
+}
+
+FText FImportantToggleSettingCustomization::GetDescriptionText() const
+{
+	IImportantToggleSettingInterface* ToogleSettingInterface = Cast<IImportantToggleSettingInterface>(ToggleSettingObject.Get());
+
+	if (ToogleSettingInterface != nullptr)
+	{
+		bool bPropertyValue = false;
+		TogglePropertyHandle->GetValue(bPropertyValue);
+
+		if (bPropertyValue)
+		{
+			return ToogleSettingInterface->GetTrueStateDescription();
+		}
+		else
+		{
+			return ToogleSettingInterface->GetFalseStateDescription();
+		}
+	}
+	return FText::GetEmpty();
+}
+
+#undef LOCTEXT_NAMESPACE
diff --git a/Engine/Source/Editor/DetailCustomizations/Private/ImportantToggleSettingCustomization.h b/Engine/Source/Editor/DetailCustomizations/Private/ImportantToggleSettingCustomization.h
new file mode 100644
--- /dev/null
+++ b/Engine/Source/Editor/DetailCustomizations/Private/ImportantToggleSettingCustomization.h
@@ -0,0 +1,40 @@
+// Copyright 1998-2015 Epic Games, Inc. All Rights Reserved.
+
+#pragma once
+
+//class FImportantToggleSettingCustomization : public IPropertyTypeCustomization
+//{
+//public:
+//	static TSharedRef<IPropertyTypeCustomization> MakeInstance();
+//
+//	/** IPropertyTypeCustomization instance */
+//	virtual void CustomizeHeader(TSharedRef<IPropertyHandle> StructPropertyHandle, FDetailWidgetRow& HeaderRow, IPropertyTypeCustomizationUtils& StructCustomizationUtils) override;
+//	virtual void CustomizeChildren(TSharedRef<IPropertyHandle> StructPropertyHandle, class IDetailChildrenBuilder& ChildBuilder, IPropertyTypeCustomizationUtils& StructCustomizationUtils) override;
+//
+//protected:
+//	TSharedPtr<IPropertyHandle> ToggleProperty;
+//};
+
+class FImportantToggleSettingCustomization : public IDetailCustomization
+{
+public:
+	/** Makes a new instance of this detail layout class for a specific detail view requesting it */
+	static TSharedRef<IDetailCustomization> MakeInstance();
+
+	// IDetailCustomization interface
+	virtual void CustomizeDetails(IDetailLayoutBuilder& DetailBuilder) override;
+	// End of IDetailCustomization interface
+
+private:
+	//ECheckBoxState IsBoolPropertyTrue(TSharedRef<IPropertyHandle> PropertyHandle) const;
+	//void OnCheckStateChanged(ECheckBoxState CheckBoxState, TSharedRef<IPropertyHandle> PropertyHandle);
+
+	bool IsToggleValue(bool bValue) const;
+	void OnToggledTo(bool bSetTo);
+	void OnNavigateHyperlink(FString Url);
+	FText GetDescriptionText() const;
+
+private:
+	TSharedPtr<IPropertyHandle> TogglePropertyHandle;
+	TWeakObjectPtr<UObject> ToggleSettingObject;
+};
\ No newline at end of file
diff --git a/Engine/Source/Editor/EditorSettingsViewer/EditorSettingsViewer.Build.cs b/Engine/Source/Editor/EditorSettingsViewer/EditorSettingsViewer.Build.cs
--- a/Engine/Source/Editor/EditorSettingsViewer/EditorSettingsViewer.Build.cs
+++ b/Engine/Source/Editor/EditorSettingsViewer/EditorSettingsViewer.Build.cs
@@ -26,6 +26,7 @@ namespace UnrealBuildTool.Rules
                     "InternationalizationSettings",
 					"BlueprintGraph",
                     "EditorStyle",
+                    "Analytics"
 				}
 			);
 
diff --git a/Engine/Source/Editor/EditorSettingsViewer/Private/EditorSettingsViewerModule.cpp b/Engine/Source/Editor/EditorSettingsViewer/Private/EditorSettingsViewerModule.cpp
--- a/Engine/Source/Editor/EditorSettingsViewer/Private/EditorSettingsViewerModule.cpp
+++ b/Engine/Source/Editor/EditorSettingsViewer/Private/EditorSettingsViewerModule.cpp
@@ -14,7 +14,7 @@
 #include "Tests/AutomationTestSettings.h"
 #include "BlueprintEditorSettings.h"
 #include "CrashReporterSettings.h"
-
+#include "Analytics/AnalyticsPrivacySettings.h"
 
 #define LOCTEXT_NAMESPACE "FEditorSettingsViewerModule"
 
@@ -61,6 +61,7 @@ public:
 			RegisterGeneralSettings(*SettingsModule);
 			RegisterLevelEditorSettings(*SettingsModule);
 			RegisterContentEditorsSettings(*SettingsModule);
+			RegisterPrivacySettings(*SettingsModule);
 
 			SettingsModule->RegisterViewer("Editor", *this);
 		}
@@ -218,6 +219,16 @@ protected:
 		);
 	}
 
+	void RegisterPrivacySettings(ISettingsModule& SettingsModule)
+	{
+		// Analytics
+		SettingsModule.RegisterSettings("Editor", "Privacy", "Analytics",
+			LOCTEXT("PrivacyAnalyticsSettingsName", "Usage Data"),
+			LOCTEXT("PrivacyAnalyticsSettingsDescription", "Configure the way your Editor usage information is handled."),
+			GetMutableDefault<UAnalyticsPrivacySettings>()
+			);
+	}
+
 	/** Unregisters all settings. */
 	void UnregisterSettings()
 	{
diff --git a/Engine/Source/Editor/ProjectSettingsViewer/Private/ProjectSettingsViewerModule.cpp b/Engine/Source/Editor/ProjectSettingsViewer/Private/ProjectSettingsViewerModule.cpp
--- a/Engine/Source/Editor/ProjectSettingsViewer/Private/ProjectSettingsViewerModule.cpp
+++ b/Engine/Source/Editor/ProjectSettingsViewer/Private/ProjectSettingsViewerModule.cpp
@@ -27,6 +27,7 @@
 #include "AI/Navigation/RecastNavMesh.h"
 #include "Engine/NetworkSettings.h"
 #include "PhysicsEngine/PhysicsSettings.h"
+#include "Engine/EndUserSettings.h"
 
 #define LOCTEXT_NAMESPACE "FProjectSettingsViewerModule"
 
@@ -161,6 +162,13 @@ protected:
 			LOCTEXT("CrowdManagerSettingsDescription", "Settings for the AI Crowd Manager."),
 			GetMutableDefault<UCrowdManager>()
 			);
+
+		// End-user settings
+		SettingsModule.RegisterSettings("Project", "Engine", "EndUser",
+			LOCTEXT("EndUserSettingsName", "End-User Settings"),
+			LOCTEXT("EndUserSettingsDescription", "Settings you may wish to expose to end-users of your game."),
+			GetMutableDefault<UEndUserSettings>()
+			);
 	}
 
 	/**
diff --git a/Engine/Source/Editor/UnrealEd/Classes/Analytics/AnalyticsPrivacySettings.h b/Engine/Source/Editor/UnrealEd/Classes/Analytics/AnalyticsPrivacySettings.h
new file mode 100644
--- /dev/null
+++ b/Engine/Source/Editor/UnrealEd/Classes/Analytics/AnalyticsPrivacySettings.h
@@ -0,0 +1,44 @@
+// Copyright 1998-2015 Epic Games, Inc. All Rights Reserved.
+
+//=============================================================================
+// AnalyticsPrivacySettings
+//
+// A configuration class that holds information for the user's privacy settings.
+// Supplied so that the editor 'remembers' the last setup the user had.
+//=============================================================================
+
+#pragma once
+#include "Engine/ImportantToggleSettingInterface.h"
+#include "AnalyticsPrivacySettings.generated.h"
+
+UCLASS(MinimalAPI, hidecategories=Object, config=EditorSettings)
+class UAnalyticsPrivacySettings : public UObject, public IImportantToggleSettingInterface
+{
+	GENERATED_UCLASS_BODY()
+
+	/** Determines whether the editor sends usage information to Epic Games in order to improve Unreal Engine. Your information will never be shared with 3rd parties. */
+	UPROPERTY(EditAnywhere, config, Category=Options)
+	bool bSendUsageData;
+
+public:
+
+	// BEGIN IImportantToggleSettingInterface
+	virtual void GetToogleCategoryAndPropertyNames(FName& OutCategory, FName& OutProperty) const override;
+	virtual FText GetFalseStateLabel() const override;
+	virtual FText GetFalseStateTooltip() const override;
+	virtual FText GetFalseStateDescription() const override;
+	virtual FText GetTrueStateLabel() const override;
+	virtual FText GetTrueStateTooltip() const override;
+	virtual FText GetTrueStateDescription() const override;
+	virtual FString GetAdditionalInfoUrl() const override;
+	virtual FText GetAdditionalInfoUrlLabel() const override;
+	// END IImportantToggleSettingInterface
+
+#if WITH_EDITOR
+	//~ Begin UObject Interface
+	virtual void PostEditChangeProperty(struct FPropertyChangedEvent& PropertyChangedEvent) override;
+	//~ End UObject Interface
+
+	void OnSendFullUsageDataChanged();
+#endif
+};
\ No newline at end of file
diff --git a/Engine/Source/Editor/UnrealEd/Private/Analytics/AnalyticsPrivacySettings.cpp b/Engine/Source/Editor/UnrealEd/Private/Analytics/AnalyticsPrivacySettings.cpp
new file mode 100644
--- /dev/null
+++ b/Engine/Source/Editor/UnrealEd/Private/Analytics/AnalyticsPrivacySettings.cpp
@@ -0,0 +1,101 @@
+// Copyright 1998-2015 Epic Games, Inc. All Rights Reserved.
+
+#include "UnrealEd.h"
+#include "Analytics/AnalyticsPrivacySettings.h"
+#include "IAnalyticsProvider.h"
+#include "EngineAnalytics.h"
+
+#define LOCTEXT_NAMESPACE "AnalyticsPrivacySettings"
+
+UAnalyticsPrivacySettings::UAnalyticsPrivacySettings(const FObjectInitializer& ObjectInitializer)
+	: Super(ObjectInitializer)
+	, bSendUsageData(true)
+{
+}
+
+void UAnalyticsPrivacySettings::GetToogleCategoryAndPropertyNames(FName& OutCategory, FName& OutProperty) const
+{
+	OutCategory = FName("Options");
+	OutProperty = FName("bSendUsageData");
+};
+
+FText UAnalyticsPrivacySettings::GetFalseStateLabel() const
+{
+	return LOCTEXT("FalseStateLabel", "Don't Send");
+};
+
+FText UAnalyticsPrivacySettings::GetFalseStateTooltip() const
+{
+	return LOCTEXT("FalseStateTooltip", "Don't send Editor usage data to Epic Games.");
+};
+
+FText UAnalyticsPrivacySettings::GetFalseStateDescription() const
+{
+	return LOCTEXT("FalseStateDescription", "By opting out you have chosen to not send Editor usage data to Epic Games. Please consider opting in to help improve Unreal Engine. Epic Games will never sell or trade individual usage data to / with third party organizations. If you enable this feature, we will collect information about how you use the editor, when you use the editor, the type of projects you are creating, how you interact with the various editor components and we would perform occasional checks on the type of hardware/OS you are using.");
+};
+
+FText UAnalyticsPrivacySettings::GetTrueStateLabel() const
+{
+	return LOCTEXT("TrueStateLabel", "Send Usage Data");
+};
+
+FText UAnalyticsPrivacySettings::GetTrueStateTooltip() const
+{
+	return LOCTEXT("TrueStateTooltip", "Send your Editor usage data to Epic Games.");
+};
+
+FText UAnalyticsPrivacySettings::GetTrueStateDescription() const
+{
+	return LOCTEXT("TrueStateDescription", "By opting in you have chosen to send Editor usage data to Epic Games. Thank you for helping to improve Unreal Engine. Epic Games will never sell or trade individual usage data to / with third party organizations. We will collect information about how you use the editor, when you use the editor, the type of projects you are creating, how you interact with the various editor components and we perform occasional checks on the type of hardware/OS you are using.");
+};
+
+FString UAnalyticsPrivacySettings::GetAdditionalInfoUrl() const
+{
+	return FString(TEXT("http://epicgames.com/privacynotice"));
+};
+
+FText UAnalyticsPrivacySettings::GetAdditionalInfoUrlLabel() const
+{
+	return LOCTEXT("HyperlinkLabel", "Epic Games Privacy Notice");
+};
+
+#if WITH_EDITOR
+void UAnalyticsPrivacySettings::PostEditChangeProperty(struct FPropertyChangedEvent& PropertyChangedEvent)
+{
+	Super::PostEditChangeProperty(PropertyChangedEvent);
+
+	const FName PropertyName = (PropertyChangedEvent.Property != NULL) ? PropertyChangedEvent.Property->GetFName() : NAME_None;
+	if (PropertyName == GET_MEMBER_NAME_CHECKED(UAnalyticsPrivacySettings, bSendUsageData))
+	{
+		OnSendFullUsageDataChanged();
+	}
+}
+
+void UAnalyticsPrivacySettings::OnSendFullUsageDataChanged()
+{
+	if (bSendUsageData)
+	{
+		// Attempt to initialize analytics and send opt-in event
+		if (!FEngineAnalytics::IsAvailable())
+		{
+			FEngineAnalytics::Initialize();
+			if (FEngineAnalytics::IsAvailable())
+			{
+				FEngineAnalytics::GetProvider().RecordEvent(FString("Editor.Privacy.EndUserOptIn"));
+			}
+		}
+	}
+	else
+	{
+		// Send opt-out event and shutdown analytics
+		if (FEngineAnalytics::IsAvailable())
+		{
+			FEngineAnalytics::GetProvider().RecordEvent(FString("Editor.Privacy.EndUserOptOut"));
+			FEngineAnalytics::Shutdown();
+		}
+	}
+}
+
+#endif
+
+#undef LOCTEXT_NAMESPACE
diff --git a/Engine/Source/Editor/UnrealEd/Private/Analytics/EditorAnalytics.cpp b/Engine/Source/Editor/UnrealEd/Private/Analytics/EditorAnalytics.cpp
--- a/Engine/Source/Editor/UnrealEd/Private/Analytics/EditorAnalytics.cpp
+++ b/Engine/Source/Editor/UnrealEd/Private/Analytics/EditorAnalytics.cpp
@@ -8,9 +8,381 @@
 
 #include "EditorAnalytics.h"
 #include "GeneralProjectSettings.h"
+#include "GameProjectGenerationModule.h"
 
 #define LOCTEXT_NAMESPACE "EditorAnalytics"
 
+namespace EditorAnalyticsDefs
+{
+	static const FTimespan SessionRecordExpiration(30, 0, 0, 0);	// 30 days
+	static const FTimespan SessionRecordTimeout(0, 30, 0);			// 30 minutes
+	static const FTimespan GlobalLockWaitTimeout(0, 0, 0, 0, 500);	// 1/2 second
+	static const FString StoreId(TEXT("Epic Games"));
+	static const FString CrashSessionToken(TEXT("Crashed"));
+	static const FString DebuggerSessionToken(TEXT("Debugger"));
+	static const FString AbnormalSessionToken(TEXT("AbnormalShutdown"));
+	static const FString SessionRecordListSection(TEXT("List"));
+	static const FString SessionRecordSectionPrefix(TEXT("Unreal Engine/Editor Sessions/"));
+	static const FString EditorSessionsVersionString(TEXT("1_1"));
+	static const FString CrashStoreKey(TEXT("IsCrash"));
+	static const FString EngineVersionStoreKey(TEXT("EngineVersion"));
+	static const FString TimestampStoreKey(TEXT("Timestamp"));
+	static const FString DebuggerStoreKey(TEXT("IsDebugger"));
+	static const FString GlobalLockName(TEXT("UE4_EditorAnalytics_Lock"));
+	static const FString FalseValueString(TEXT("0"));
+	static const FString TrueValueString(TEXT("1"));
+}
+
+/* Handles writing session records to platform's storage to track crashed and timed-out editor sessions */
+class FSessionManager
+{
+public:
+	FSessionManager()
+		: bInitialized(false)
+		, bEditorCrashed(false)
+	{}
+
+	void Initialize(bool bWaitForLock = true)
+	{
+		if (!FEngineAnalytics::IsAvailable())
+		{
+			return;
+		}
+
+		TArray<FSessionRecord> SessionRecordsToReport;
+
+		{
+			// Scoped lock
+			FSystemWideCriticalSection StoredValuesLock(EditorAnalyticsDefs::GlobalLockName, bWaitForLock ? EditorAnalyticsDefs::GlobalLockWaitTimeout : FTimespan::Zero());
+
+			// Get list of sessions in storage
+			if (StoredValuesLock.IsValid() && BeginReadWriteRecords())
+			{
+				TArray<FSessionRecord> SessionRecordsToDelete;
+
+				// Attempt check each stored session
+				for (FSessionRecord& Record : SessionRecords)
+				{
+					FTimespan RecordAge = FDateTime::UtcNow() - Record.Timestamp;
+
+					if (Record.bCrashed)
+					{
+						// Crashed sessions
+						SessionRecordsToReport.Add(Record);
+						SessionRecordsToDelete.Add(Record);
+					}
+					else if (RecordAge > EditorAnalyticsDefs::SessionRecordExpiration)
+					{
+						// Delete expired session records
+						SessionRecordsToDelete.Add(Record);
+					}
+					else if (RecordAge > EditorAnalyticsDefs::SessionRecordTimeout)
+					{
+						// Timed out sessions
+						SessionRecordsToReport.Add(Record);
+						SessionRecordsToDelete.Add(Record);
+					}
+				}
+
+				for (FSessionRecord& DeletingRecord : SessionRecordsToDelete)
+				{
+					DeleteStoredRecord(DeletingRecord);
+				}
+
+				// Create a session record for this session
+				CreateAndWriteRecordForSession();
+
+				// Update and release list of sessions in storage
+				EndReadWriteRecords();
+
+				// Register for crash callbacks
+				FCoreDelegates::OnHandleSystemError.AddRaw(this, &FSessionManager::OnCrashing);
+
+				bInitialized = true;
+			}
+		}
+
+		for (FSessionRecord& ReportingSession : SessionRecordsToReport)
+		{
+			// Send error report for session that timed out or crashed
+			SendAbnormalShutdownReport(ReportingSession);
+		}
+	}
+
+	void Heartbeat()
+	{
+		if (!FEngineAnalytics::IsAvailable())
+		{
+			return;
+		}
+
+		if (!bInitialized)
+		{
+			// Try late initialization
+			const bool bWaitForLock = false;
+			Initialize(bWaitForLock);
+		}
+
+		// Update timestamp in the session record for this session 
+		if (bInitialized)
+		{
+			CurrentSession.Timestamp = FDateTime::UtcNow();
+
+			FPlatformMisc::SetStoredValue(EditorAnalyticsDefs::StoreId, CurrentSessionSectionName, EditorAnalyticsDefs::TimestampStoreKey, TimestampToString(CurrentSession.Timestamp));
+		}
+	}
+
+	void Shutdown()
+	{
+		// Clear the session record for this session
+		if (bInitialized)
+		{
+			FCoreDelegates::OnHandleSystemError.RemoveAll(this);
+
+			if (!bEditorCrashed)
+			{
+				FPlatformMisc::DeleteStoredValue(EditorAnalyticsDefs::StoreId, CurrentSessionSectionName, EditorAnalyticsDefs::CrashStoreKey);
+				FPlatformMisc::DeleteStoredValue(EditorAnalyticsDefs::StoreId, CurrentSessionSectionName, EditorAnalyticsDefs::EngineVersionStoreKey);
+				FPlatformMisc::DeleteStoredValue(EditorAnalyticsDefs::StoreId, CurrentSessionSectionName, EditorAnalyticsDefs::TimestampStoreKey);
+				FPlatformMisc::DeleteStoredValue(EditorAnalyticsDefs::StoreId, CurrentSessionSectionName, EditorAnalyticsDefs::DebuggerStoreKey);
+			}
+
+			bInitialized = false;
+		}
+	}
+
+private:
+	struct FSessionRecord
+	{
+		FString SessionId;
+		FString EngineVersion;
+		FDateTime Timestamp;
+		bool bCrashed;
+		bool bIsDebugger;
+	};
+
+private:
+	bool BeginReadWriteRecords()
+	{
+		SessionRecords.Empty();
+
+		// Lock and read the list of sessions in storage
+		FString ListSectionName = GetStoreSectionString(EditorAnalyticsDefs::SessionRecordListSection);
+
+		// Write list to SessionRecords member
+		FString SessionListString;
+		FPlatformMisc::GetStoredValue(EditorAnalyticsDefs::StoreId, ListSectionName, TEXT("SessionList"), SessionListString);
+
+		// Parse SessionListString for session ids
+		TArray<FString> SessionIds;
+		SessionListString.ParseIntoArray(SessionIds, TEXT(","));
+
+		// Retrieve all the sessions in the list from storage
+		for (FString& SessionId : SessionIds)
+		{
+			FString SectionName = GetStoreSectionString(SessionId);
+
+			FString IsCrashString;
+			FString EngineVersionString;
+			FString TimestampString;
+			FString IsDebuggerString;
+
+			if (FPlatformMisc::GetStoredValue(EditorAnalyticsDefs::StoreId, SectionName, EditorAnalyticsDefs::CrashStoreKey, IsCrashString) &&
+				FPlatformMisc::GetStoredValue(EditorAnalyticsDefs::StoreId, SectionName, EditorAnalyticsDefs::EngineVersionStoreKey, EngineVersionString) &&
+				FPlatformMisc::GetStoredValue(EditorAnalyticsDefs::StoreId, SectionName, EditorAnalyticsDefs::TimestampStoreKey, TimestampString) &&
+				FPlatformMisc::GetStoredValue(EditorAnalyticsDefs::StoreId, SectionName, EditorAnalyticsDefs::DebuggerStoreKey, IsDebuggerString))
+			{
+				FSessionRecord NewRecord;
+				NewRecord.SessionId = SessionId;
+				NewRecord.EngineVersion = EngineVersionString;
+				NewRecord.Timestamp = StringToTimestamp(TimestampString);
+				NewRecord.bCrashed = IsCrashString == EditorAnalyticsDefs::TrueValueString;
+				NewRecord.bIsDebugger = IsDebuggerString == EditorAnalyticsDefs::TrueValueString;
+
+				SessionRecords.Add(NewRecord);
+			}
+			else
+			{
+				// Clean up orphaned values, if there are any
+				FPlatformMisc::DeleteStoredValue(EditorAnalyticsDefs::StoreId, SectionName, EditorAnalyticsDefs::CrashStoreKey);
+				FPlatformMisc::DeleteStoredValue(EditorAnalyticsDefs::StoreId, SectionName, EditorAnalyticsDefs::EngineVersionStoreKey);
+				FPlatformMisc::DeleteStoredValue(EditorAnalyticsDefs::StoreId, SectionName, EditorAnalyticsDefs::TimestampStoreKey);
+				FPlatformMisc::DeleteStoredValue(EditorAnalyticsDefs::StoreId, SectionName, EditorAnalyticsDefs::DebuggerStoreKey);
+			}
+		}
+
+		return true;
+	}
+
+	void EndReadWriteRecords()
+	{
+		// Update the list of sessions in storage to match SessionRecords
+		FString SessionListString;
+		if (SessionRecords.Num() > 0)
+		{
+			for (FSessionRecord& Session : SessionRecords)
+			{
+				SessionListString.Append(Session.SessionId);
+				SessionListString.Append(TEXT(","));
+			}
+			SessionListString.RemoveAt(SessionListString.Len() - 1);
+		}
+
+		FString ListSectionName = GetStoreSectionString(EditorAnalyticsDefs::SessionRecordListSection);
+		FPlatformMisc::SetStoredValue(EditorAnalyticsDefs::StoreId, ListSectionName, TEXT("SessionList"), SessionListString);
+
+		// Clear SessionRecords member
+		SessionRecords.Empty();
+	}
+
+	void DeleteStoredRecord(const FSessionRecord& Record)
+	{
+		// Delete the session record in storage
+		FString SessionId = Record.SessionId;
+		FString SectionName = GetStoreSectionString(SessionId);
+
+		FPlatformMisc::DeleteStoredValue(EditorAnalyticsDefs::StoreId, SectionName, EditorAnalyticsDefs::CrashStoreKey);
+		FPlatformMisc::DeleteStoredValue(EditorAnalyticsDefs::StoreId, SectionName, EditorAnalyticsDefs::EngineVersionStoreKey);
+		FPlatformMisc::DeleteStoredValue(EditorAnalyticsDefs::StoreId, SectionName, EditorAnalyticsDefs::TimestampStoreKey);
+		FPlatformMisc::DeleteStoredValue(EditorAnalyticsDefs::StoreId, SectionName, EditorAnalyticsDefs::DebuggerStoreKey);
+
+		// Remove the session record from SessionRecords list
+		SessionRecords.RemoveAll([&SessionId](const FSessionRecord& X){ return X.SessionId == SessionId; });
+	}
+
+	void SendAbnormalShutdownReport(const FSessionRecord& Record)
+	{
+		FGameProjectGenerationModule& GameProjectModule = FModuleManager::LoadModuleChecked<FGameProjectGenerationModule>(TEXT("GameProjectGeneration"));
+		bool bHasCode = GameProjectModule.Get().ProjectHasCodeFiles();
+
+#if PLATFORM_WINDOWS
+		const FString PlatformName(TEXT("Windows"));
+#elif PLATFORM_MAC
+		const FString PlatformName(TEXT("Mac"));
+#elif PLATFORM_LINUX
+		const FString PlatformName(TEXT("Linux"));
+#else
+		const FString PlatformName(TEXT("Unknown"));
+#endif
+
+		FGuid SessionId;
+		FString SessionIdString = Record.SessionId;
+		if (FGuid::Parse(SessionIdString, SessionId))
+		{
+			// convert session guid to one with braces for sending to analytics
+			SessionIdString = SessionId.ToString(EGuidFormats::DigitsWithHyphensInBraces);
+		}
+
+		FString ShutdownTypeString = Record.bCrashed ? EditorAnalyticsDefs::CrashSessionToken :
+			(Record.bIsDebugger ? EditorAnalyticsDefs::DebuggerSessionToken : EditorAnalyticsDefs::AbnormalSessionToken);
+
+		TArray< FAnalyticsEventAttribute > AbnormalShutdownAttributes;
+		AbnormalShutdownAttributes.Add(FAnalyticsEventAttribute(FString("SessionId"), SessionIdString));
+		AbnormalShutdownAttributes.Add(FAnalyticsEventAttribute(FString("EngineVersion"), Record.EngineVersion));
+		AbnormalShutdownAttributes.Add(FAnalyticsEventAttribute(FString("ShutdownType"), ShutdownTypeString));
+		AbnormalShutdownAttributes.Add(FAnalyticsEventAttribute(FString("Timestamp"), Record.Timestamp.ToIso8601()));
+
+		FEditorAnalytics::ReportEvent(TEXT("Editor.AbnormalShutdown"), PlatformName, bHasCode, AbnormalShutdownAttributes);
+	}
+
+	void CreateAndWriteRecordForSession()
+	{
+		FGuid SessionId;
+		if (FGuid::Parse(FEngineAnalytics::GetProvider().GetSessionID(), SessionId))
+		{
+			// convert session guid to one without braces or other chars that might not be suitable for storage
+			CurrentSession.SessionId = SessionId.ToString(EGuidFormats::DigitsWithHyphens);
+		}
+		else
+		{
+			CurrentSession.SessionId = FEngineAnalytics::GetProvider().GetSessionID();
+		}
+
+		CurrentSession.EngineVersion = FEngineVersion::Current().ToString(EVersionComponent::Changelist);
+		CurrentSession.Timestamp = FDateTime::UtcNow();
+		CurrentSession.bCrashed = false;
+		CurrentSession.bIsDebugger = FPlatformMisc::IsDebuggerPresent();
+		CurrentSessionSectionName = GetStoreSectionString(CurrentSession.SessionId);
+
+		FString IsDebuggerString = CurrentSession.bIsDebugger ? EditorAnalyticsDefs::TrueValueString : EditorAnalyticsDefs::FalseValueString;
+
+		FPlatformMisc::SetStoredValue(EditorAnalyticsDefs::StoreId, CurrentSessionSectionName, EditorAnalyticsDefs::CrashStoreKey, EditorAnalyticsDefs::FalseValueString);
+		FPlatformMisc::SetStoredValue(EditorAnalyticsDefs::StoreId, CurrentSessionSectionName, EditorAnalyticsDefs::EngineVersionStoreKey, CurrentSession.EngineVersion);
+		FPlatformMisc::SetStoredValue(EditorAnalyticsDefs::StoreId, CurrentSessionSectionName, EditorAnalyticsDefs::TimestampStoreKey, TimestampToString(CurrentSession.Timestamp));
+		FPlatformMisc::SetStoredValue(EditorAnalyticsDefs::StoreId, CurrentSessionSectionName, EditorAnalyticsDefs::DebuggerStoreKey, IsDebuggerString);
+
+		SessionRecords.Add(CurrentSession);
+	}
+
+	void OnCrashing()
+	{
+		if (!bEditorCrashed && bInitialized)
+		{
+			bEditorCrashed = true;
+
+			FPlatformMisc::SetStoredValue(EditorAnalyticsDefs::StoreId, CurrentSessionSectionName, EditorAnalyticsDefs::CrashStoreKey, EditorAnalyticsDefs::TrueValueString);
+		}
+	}
+
+	static FString TimestampToString(FDateTime InTimestamp)
+	{
+		return LexicalConversion::ToString(InTimestamp.ToUnixTimestamp());
+	}
+
+	static FDateTime StringToTimestamp(FString InString)
+	{
+		int64 TimestampUnix;
+		if (LexicalConversion::TryParseString(TimestampUnix, *InString))
+		{
+			return FDateTime::FromUnixTimestamp(TimestampUnix);
+		}
+		return FDateTime::MinValue();
+	}
+
+	static FString GetStoreSectionString(FString InSuffix)
+	{
+		return FString::Printf(TEXT("%s%s/%s"), *EditorAnalyticsDefs::SessionRecordSectionPrefix, *EditorAnalyticsDefs::EditorSessionsVersionString, *InSuffix);
+	}
+
+private:
+	bool bInitialized;
+	bool bEditorCrashed;
+	FSessionRecord CurrentSession;
+	FString CurrentSessionSectionName;
+	TArray<FSessionRecord> SessionRecords;
+};
+
+
+TSharedPtr<FSessionManager> FEditorAnalytics::SessionManager;
+
+void FEditorAnalytics::InitializeSessionManager()
+{
+	// Create the session manager singleton if we are a normal editor session
+	if (GIsEditor && !IsRunningCommandlet() && !SessionManager.IsValid() )
+	{
+		SessionManager = MakeShareable(new FSessionManager());
+		SessionManager->Initialize();
+	}
+}
+
+void FEditorAnalytics::ShutdownSessionManager()
+{
+	// Destroy the session manager singleton if it exists
+	if (SessionManager.IsValid())
+	{
+		SessionManager->Shutdown();
+		SessionManager.Reset();
+	}
+}
+
+void FEditorAnalytics::SessionHeartbeat()
+{
+	if (SessionManager.IsValid())
+	{
+		SessionManager->Heartbeat();
+	}
+}
+
 void FEditorAnalytics::ReportBuildRequirementsFailure(FString EventName, FString PlatformName, bool bHasCode, int32 Requirements)
 {
 	TArray<FAnalyticsEventAttribute> ParamArray;
diff --git a/Engine/Source/Editor/UnrealEd/Private/EditorServer.cpp b/Engine/Source/Editor/UnrealEd/Private/EditorServer.cpp
--- a/Engine/Source/Editor/UnrealEd/Private/EditorServer.cpp
+++ b/Engine/Source/Editor/UnrealEd/Private/EditorServer.cpp
@@ -68,13 +68,10 @@
 #include "Components/DrawFrustumComponent.h"
 #include "UnrealEngine.h"
 #include "AI/Navigation/NavLinkRenderingComponent.h"
-
 #include "PhysicsPublic.h"
-
 #include "AnimationRecorder.h"
-
 #include "Settings/EditorSettings.h"
-
+#include "Analytics/AnalyticsPrivacySettings.h"
 #include "KismetReinstanceUtilities.h"
 
 DEFINE_LOG_CATEGORY_STATIC(LogEditorServer, Log, All);
@@ -1107,7 +1104,7 @@ void UEditorEngine::HandleTransactorUndo( FUndoSessionContext SessionContext, bo
 
 bool UEditorEngine::AreEditorAnalyticsEnabled() const 
 {
-	return GetDefault<UEditorSettings>()->bEditorAnalyticsEnabled;
+	return GetDefault<UAnalyticsPrivacySettings>()->bSendUsageData;
 }
 
 void UEditorEngine::CreateStartupAnalyticsAttributes( TArray<FAnalyticsEventAttribute>& StartSessionAttributes ) const
diff --git a/Engine/Source/Editor/UnrealEd/Private/Settings/EditorSettings.cpp b/Engine/Source/Editor/UnrealEd/Private/Settings/EditorSettings.cpp
--- a/Engine/Source/Editor/UnrealEd/Private/Settings/EditorSettings.cpp
+++ b/Engine/Source/Editor/UnrealEd/Private/Settings/EditorSettings.cpp
@@ -8,7 +8,7 @@ UEditorSettings::UEditorSettings(const FObjectInitializer& ObjectInitializer)
 	: Super(ObjectInitializer)
 {
 	bCopyStarterContentPreference = false;
-	bEditorAnalyticsEnabled = true;
+	bEditorAnalyticsEnabled_DEPRECATED = true;
 	AutoScalabilityWorkScaleAmount = 1;
 }
 
diff --git a/Engine/Source/Editor/UnrealEd/Private/UnrealEdMisc.cpp b/Engine/Source/Editor/UnrealEd/Private/UnrealEdMisc.cpp
--- a/Engine/Source/Editor/UnrealEd/Private/UnrealEdMisc.cpp
+++ b/Engine/Source/Editor/UnrealEd/Private/UnrealEdMisc.cpp
@@ -442,6 +442,8 @@ void FUnrealEdMisc::OnInit()
 
 	SlowTask.EnterProgressFrame(10);
 
+	FEditorAnalytics::InitializeSessionManager();
+
 	// Send Project Analytics
 	InitEngineAnalytics();
 	
@@ -571,6 +573,8 @@ void FUnrealEdMisc::EditorAnalyticsHeartbeat()
 	FEngineAnalytics::GetProvider().RecordEvent(TEXT("Editor.Usage.Heartbeat"), Attributes);
 	
 	LastHeartbeatTime = FPlatformTime::Seconds();
+
+	FEditorAnalytics::SessionHeartbeat();
 }
 
 void FUnrealEdMisc::TickAssetAnalytics()
@@ -764,6 +768,8 @@ void FUnrealEdMisc::OnExit()
 		FEditorViewportStats::SendUsageData();
 	}
 
+	FEditorAnalytics::ShutdownSessionManager();
+
 	FInputBindingManager::Get().UnregisterUserDefinedChordChanged(OnUserDefinedChordChangedDelegateHandle);
 	FMessageLog::OnMessageSelectionChanged().Unbind();
 	FUObjectToken::DefaultOnMessageTokenActivated().Unbind();
diff --git a/Engine/Source/Editor/UnrealEd/Public/EditorAnalytics.h b/Engine/Source/Editor/UnrealEd/Public/EditorAnalytics.h
--- a/Engine/Source/Editor/UnrealEd/Public/EditorAnalytics.h
+++ b/Engine/Source/Editor/UnrealEd/Public/EditorAnalytics.h
@@ -67,6 +67,21 @@ namespace EAnalyticsErrorCodes
 class FEditorAnalytics
 {
 public:
+	/**
+	 * Starts the session manager if needed
+	 */
+	UNREALED_API static void InitializeSessionManager();
+
+	/**
+	 * Stops the session manager if needed
+	 */
+	UNREALED_API static void ShutdownSessionManager();
+
+	/**
+	 * Update the editor session's heartbeat time
+	 */
+	UNREALED_API static void SessionHeartbeat();
+
 	/**
 	* Reports an event to the analytics system if it is enabled
 	*
@@ -110,4 +125,7 @@ public:
 
 	/* Determine whether the error code should cause the reporter to display the message in a dialog. */
 	UNREALED_API static bool ShouldElevateMessageThroughDialog(const int32 ErrorCode);
+
+private:
+	static TSharedPtr<class FSessionManager> SessionManager;
 };
diff --git a/Engine/Source/Editor/UnrealEd/Public/Settings/EditorSettings.h b/Engine/Source/Editor/UnrealEd/Public/Settings/EditorSettings.h
--- a/Engine/Source/Editor/UnrealEd/Public/Settings/EditorSettings.h
+++ b/Engine/Source/Editor/UnrealEd/Public/Settings/EditorSettings.h
@@ -14,8 +14,8 @@ class UNREALED_API UEditorSettings : public UObject
 	bool bLoadTheMostRecentlyLoadedProjectAtStartup; // Note that this property is NOT config since it is not necessary to save the value to ini. It is determined at startup in UEditorEngine::InitEditor().
 
 	/** Can the editor report usage analytics (types of assets being spawned, etc...) back to Epic in order for us to improve the editor user experience?  Note: The editor must be restarted for changes to take effect. */
-	UPROPERTY(EditAnywhere, config, Category=General, AdvancedDisplay)
-	bool bEditorAnalyticsEnabled;
+	UPROPERTY()
+	bool bEditorAnalyticsEnabled_DEPRECATED;
 
 	// =====================================================================
 	// The following options are NOT exposed in the preferences Editor
diff --git a/Engine/Source/Editor/UnrealEd/UnrealEd.Build.cs b/Engine/Source/Editor/UnrealEd/UnrealEd.Build.cs
--- a/Engine/Source/Editor/UnrealEd/UnrealEd.Build.cs
+++ b/Engine/Source/Editor/UnrealEd/UnrealEd.Build.cs
@@ -144,7 +144,8 @@ public class UnrealEd : ModuleRules
 				"AddContentDialog",
 				"GameProjectGeneration",
                 "HierarchicalLODUtilities",
-                "BlueprintNativeCodeGen",               
+                "BlueprintNativeCodeGen",
+                "Analytics"
 			}
 		);
 
diff --git a/Engine/Source/Runtime/Core/Private/GenericPlatform/GenericPlatformMisc.cpp b/Engine/Source/Runtime/Core/Private/GenericPlatform/GenericPlatformMisc.cpp
--- a/Engine/Source/Runtime/Core/Private/GenericPlatform/GenericPlatformMisc.cpp
+++ b/Engine/Source/Runtime/Core/Private/GenericPlatform/GenericPlatformMisc.cpp
@@ -334,6 +334,30 @@ bool FGenericPlatformMisc::GetStoredValue(const FString& InStoreId, const FStrin
 	return false;
 }
 
+bool FGenericPlatformMisc::DeleteStoredValue(const FString& InStoreId, const FString& InSectionName, const FString& InKeyName)
+{	
+	check(!InStoreId.IsEmpty());
+	check(!InSectionName.IsEmpty());
+	check(!InKeyName.IsEmpty());
+
+	// This assumes that FPlatformProcess::ApplicationSettingsDir() returns a user-specific directory; it doesn't on Windows, but Windows overrides this behavior to use the registry
+	const FString ConfigPath = FString(FPlatformProcess::ApplicationSettingsDir()) / InStoreId / FString(TEXT("KeyValueStore.ini"));
+
+	FConfigFile ConfigFile;
+	ConfigFile.Read(ConfigPath);
+
+	FConfigSection* const Section = ConfigFile.Find(InSectionName);
+	if (Section)
+	{
+		int32 RemovedNum = Section->Remove(*InKeyName);
+
+		ConfigFile.Dirty = true;
+		return ConfigFile.Write(ConfigPath) && RemovedNum == 1;
+	}
+
+	return false;
+}
+
 void FGenericPlatformMisc::LowLevelOutputDebugString( const TCHAR *Message )
 {
 	FPlatformMisc::LocalPrint( Message );
diff --git a/Engine/Source/Runtime/Core/Private/GenericPlatform/GenericPlatformProcess.cpp b/Engine/Source/Runtime/Core/Private/GenericPlatform/GenericPlatformProcess.cpp
--- a/Engine/Source/Runtime/Core/Private/GenericPlatform/GenericPlatformProcess.cpp
+++ b/Engine/Source/Runtime/Core/Private/GenericPlatform/GenericPlatformProcess.cpp
@@ -530,3 +530,8 @@ bool FGenericPlatformProcess::IsFirstInstance()
 	return true;
 #endif
 }
+
+FSystemWideCriticalSectionNotImplemented::FSystemWideCriticalSectionNotImplemented(const FString& Name, FTimespan Timeout)
+{
+	UE_LOG(LogHAL, Fatal, TEXT("FSystemWideCriticalSection not implemented on this platform"));
+}
diff --git a/Engine/Source/Runtime/Core/Private/IOS/IOSPlatformMisc.cpp b/Engine/Source/Runtime/Core/Private/IOS/IOSPlatformMisc.cpp
--- a/Engine/Source/Runtime/Core/Private/IOS/IOSPlatformMisc.cpp
+++ b/Engine/Source/Runtime/Core/Private/IOS/IOSPlatformMisc.cpp
@@ -120,6 +120,12 @@ bool FIOSPlatformMisc::GetStoredValue(const FString& InStoreId, const FString& I
 	return false;
 }
 
+bool FIOSPlatformMisc::DeleteStoredValue(const FString& InStoreId, const FString& InSectionName, const FString& InKeyName)
+{
+	// No Implementation (currently only used by editor code so not needed on iOS)
+	return false;
+}
+
 //void FIOSPlatformMisc::LowLevelOutputDebugStringf(const TCHAR *Fmt, ... )
 //{
 //
diff --git a/Engine/Source/Runtime/Core/Private/Linux/LinuxPlatformProcess.cpp b/Engine/Source/Runtime/Core/Private/Linux/LinuxPlatformProcess.cpp
--- a/Engine/Source/Runtime/Core/Private/Linux/LinuxPlatformProcess.cpp
+++ b/Engine/Source/Runtime/Core/Private/Linux/LinuxPlatformProcess.cpp
@@ -7,6 +7,7 @@
 #include <sys/wait.h>
 #include <sys/resource.h>
 #include <sys/ioctl.h> // ioctl
+#include <sys/file.h>
 #include <asm/ioctls.h> // FIONREAD
 #include "LinuxApplication.h" // FLinuxApplication::IsForeground()
 
@@ -1313,3 +1314,54 @@ void FLinuxPlatformProcess::CeaseBeingFirstInstance()
 	}
 #endif
 }
+
+FLinuxSystemWideCriticalSection::FLinuxSystemWideCriticalSection(const FString& InName, FTimespan InTimeout)
+{
+	check(InName.Len() > 0)
+	check(InTimeout >= FTimespan::Zero())
+	check(InTimeout.GetTotalSeconds() < (double)FLT_MAX)
+
+	const FString LockPath = FString(FLinuxPlatformProcess::ApplicationSettingsDir()) / InName;
+	FString NormalizedFilepath(LockPath);
+	NormalizedFilepath.ReplaceInline(TEXT("\\"), TEXT("/"));
+
+	// Attempt to open a file and then lock with flock (NOTE: not an atomic operation, but best we can do)
+	FileHandle = open(TCHAR_TO_UTF8(*NormalizedFilepath), O_CREAT | O_WRONLY | O_NONBLOCK, S_IRUSR | S_IWUSR | S_IRGRP | S_IWGRP | S_IROTH | S_IWOTH);
+
+	if (FileHandle == -1 && InTimeout != FTimespan::Zero())
+	{
+		FDateTime ExpireTime = FDateTime::UtcNow() + InTimeout;
+		const float RetrySeconds = FMath::Min((float)InTimeout.GetTotalSeconds(), 0.25f);
+
+		do
+		{
+			// retry until timeout
+			FLinuxPlatformProcess::Sleep(RetrySeconds);
+			FileHandle = open(TCHAR_TO_UTF8(*NormalizedFilepath), O_CREAT | O_WRONLY | O_NONBLOCK, S_IRUSR | S_IWUSR | S_IRGRP | S_IWGRP | S_IROTH | S_IWOTH);
+		} while (FileHandle == -1 && FDateTime::UtcNow() < ExpireTime);
+	}
+	if (FileHandle != -1)
+	{
+		flock(FileHandle, LOCK_EX);
+	}
+}
+
+FLinuxSystemWideCriticalSection::~FLinuxSystemWideCriticalSection()
+{
+	Release();
+}
+
+bool FLinuxSystemWideCriticalSection::IsValid() const
+{
+	return FileHandle != -1;
+}
+
+void FLinuxSystemWideCriticalSection::Release()
+{
+	if (IsValid())
+	{
+		flock(FileHandle, LOCK_UN);
+		close(FileHandle);
+		FileHandle = -1;
+	}
+}
diff --git a/Engine/Source/Runtime/Core/Private/Mac/MacPlatformProcess.cpp b/Engine/Source/Runtime/Core/Private/Mac/MacPlatformProcess.cpp
--- a/Engine/Source/Runtime/Core/Private/Mac/MacPlatformProcess.cpp
+++ b/Engine/Source/Runtime/Core/Private/Mac/MacPlatformProcess.cpp
@@ -1106,3 +1106,55 @@ FRunnableThread* FMacPlatformProcess::CreateRunnableThread()
 {
 	return new FRunnableThreadMac();
 }
+
+
+FMacSystemWideCriticalSection::FMacSystemWideCriticalSection(const FString& InName, FTimespan InTimeout)
+{
+	check(InName.Len() > 0)
+	check(InTimeout >= FTimespan::Zero())
+	check(InTimeout.GetTotalSeconds() < (double)FLT_MAX)	// we'll need this to fit in a single-precision float later so check it here
+
+	FDateTime ExpireTime = FDateTime::UtcNow() + InTimeout;
+
+	// This lock implementation is using files so correct any backslashes in the name
+	const FString LockPath = FString(FMacPlatformProcess::ApplicationSettingsDir()) / InName;
+	FString NormalizedFilepath(LockPath);
+	NormalizedFilepath.ReplaceInline(TEXT("\\"), TEXT("/"));
+
+	// Attempt to open a file with O_EXLOCK (equivalent of atomic open() + flock())
+	FileHandle = open(TCHAR_TO_UTF8(*NormalizedFilepath), O_CREAT | O_WRONLY | O_EXLOCK | O_NONBLOCK, S_IRUSR | S_IWUSR | S_IRGRP | S_IWGRP | S_IROTH | S_IWOTH);
+
+	if (FileHandle == -1 && InTimeout != FTimespan::Zero())
+	{
+		// Failed to get a valid file handle so the file is probably open and locked by another owner - retry until we timeout
+
+		for (FTimespan TimeoutRemaining = ExpireTime - FDateTime::UtcNow();
+			 FileHandle == -1 && TimeoutRemaining > FTimespan::Zero();
+			 TimeoutRemaining = ExpireTime - FDateTime::UtcNow())
+		{
+			// retry until timeout
+			float RetrySeconds = FMath::Min((float)TimeoutRemaining.GetTotalSeconds(), 0.25f);
+			FMacPlatformProcess::Sleep(RetrySeconds);
+			FileHandle = open(TCHAR_TO_UTF8(*NormalizedFilepath), O_CREAT | O_WRONLY | O_EXLOCK | O_NONBLOCK, S_IRUSR | S_IWUSR | S_IRGRP | S_IWGRP | S_IROTH | S_IWOTH);
+		}
+	}
+}
+
+FMacSystemWideCriticalSection::~FMacSystemWideCriticalSection()
+{
+	Release();
+}
+
+bool FMacSystemWideCriticalSection::IsValid() const
+{
+	return FileHandle != -1;
+}
+
+void FMacSystemWideCriticalSection::Release()
+{
+	if (IsValid())
+	{
+		close(FileHandle);
+		FileHandle = -1;
+	}
+}
diff --git a/Engine/Source/Runtime/Core/Private/Windows/WindowsCriticalSection.cpp b/Engine/Source/Runtime/Core/Private/Windows/WindowsCriticalSection.cpp
new file mode 100644
--- /dev/null
+++ b/Engine/Source/Runtime/Core/Private/Windows/WindowsCriticalSection.cpp
@@ -0,0 +1,75 @@
+// Copyright 1998-2015 Epic Games, Inc. All Rights Reserved.
+
+#include "CorePrivatePCH.h"
+#include "WindowsCriticalSection.h"
+#include "AllowWindowsPlatformTypes.h"
+
+FWindowsSystemWideCriticalSection::FWindowsSystemWideCriticalSection(const FString& InName, FTimespan InTimeout)
+{
+	// Check for valid name length not exceeding the Window mutex name length limit
+	check(InName.Len() > 0)
+	check(InName.Len() < MAX_PATH)
+	// Check for valid positive timeouts
+	check(InTimeout >= FTimespan::Zero())
+	check(InTimeout.GetTotalMilliseconds() < (double)0x7FFFFFFE)	// limit timespan to a number of millisec that will fit in a signed int32
+
+	// Disallow backslashes as they aren't allowed in Windows mutex names
+	FString NormalizedMutexName(InName);
+	NormalizedMutexName.ReplaceInline(TEXT("\\"), TEXT("/"));
+
+	TCHAR MutexName[MAX_PATH] = TEXT("");
+	FCString::Strcpy(MutexName, *NormalizedMutexName);
+
+	// Attempt to create and take ownership of a named mutex
+	Mutex = CreateMutex(NULL, true, MutexName);
+
+	if (Mutex != NULL && GetLastError() == ERROR_ALREADY_EXISTS)
+	{
+		// CreateMutex returned a valid handle but we didn't get ownership because another process/thread has already created it
+		bool bMutexOwned = false;
+
+		if (InTimeout != FTimespan::Zero())
+		{
+			// We have a handle already so try waiting for it to be released by the current owner
+			DWORD WaitResult = WaitForSingleObject(Mutex, FMath::TruncToInt((float)InTimeout.GetTotalMilliseconds()));
+
+			// WAIT_OBJECT_0 = we got ownership when the previous owner released it
+			// WAIT_ABANDONED = we got ownership when the previous owner exited WITHOUT releasing the mutex gracefully (we own it now but the state of any shared resource could be corrupted!)
+			if (WaitResult == WAIT_ABANDONED || WaitResult == WAIT_OBJECT_0)
+			{
+				bMutexOwned = true;
+			}
+		}
+
+		if (!bMutexOwned)
+		{
+			// We failed to gain ownership by waiting so close the handle to avoid leaking it.
+			CloseHandle(Mutex);
+			Mutex = NULL;
+		}
+	}
+}
+
+FWindowsSystemWideCriticalSection::~FWindowsSystemWideCriticalSection()
+{
+	Release();
+}
+
+bool FWindowsSystemWideCriticalSection::IsValid() const
+{
+	return Mutex != NULL;
+}
+
+void FWindowsSystemWideCriticalSection::Release()
+{
+	if (IsValid())
+	{
+		// Release ownership
+		ReleaseMutex(Mutex);
+		// Also release the handle so it isn't leaked
+		CloseHandle(Mutex);
+		Mutex = NULL;
+	}
+}
+
+#include "HideWindowsPlatformTypes.h"
diff --git a/Engine/Source/Runtime/Core/Private/Windows/WindowsPlatformMisc.cpp b/Engine/Source/Runtime/Core/Private/Windows/WindowsPlatformMisc.cpp
--- a/Engine/Source/Runtime/Core/Private/Windows/WindowsPlatformMisc.cpp
+++ b/Engine/Source/Runtime/Core/Private/Windows/WindowsPlatformMisc.cpp
@@ -2031,6 +2031,47 @@ bool FWindowsPlatformMisc::GetStoredValue(const FString& InStoreId, const FStrin
 	return QueryRegKey(HKEY_CURRENT_USER, *FullRegistryKey, *InKeyName, OutValue);
 }
 
+bool FWindowsPlatformMisc::DeleteStoredValue(const FString& InStoreId, const FString& InSectionName, const FString& InKeyName)
+{
+	// Deletes values in reg keys and also deletes the owning key if it becomes empty
+
+	check(!InStoreId.IsEmpty());
+	check(!InSectionName.IsEmpty());
+	check(!InKeyName.IsEmpty());
+
+	FString FullRegistryKey = FString(TEXT("Software")) / InStoreId / InSectionName;
+	FullRegistryKey = FullRegistryKey.Replace(TEXT("/"), TEXT("\\")); // we use forward slashes, but the registry needs back slashes
+
+	HKEY hKey;
+	HRESULT Result = ::RegOpenKeyEx(HKEY_CURRENT_USER, *FullRegistryKey, 0, KEY_WRITE | KEY_READ, &hKey);
+	if (Result == ERROR_SUCCESS)
+	{
+		Result = ::RegDeleteValue(hKey, *InKeyName);
+
+		// Query for sub-keys in the open key
+		TCHAR CheckKeyName[256];
+		::DWORD CheckKeyNameLength = sizeof(CheckKeyName) / sizeof(CheckKeyName[0]);
+		HRESULT EnumResult = RegEnumKeyEx(hKey, 0, CheckKeyName, &CheckKeyNameLength, NULL, NULL, NULL, NULL);
+		bool bZeroSubKeys = EnumResult != ERROR_SUCCESS;
+
+		// Query for a remaining value in the open key
+		wchar_t CheckValueName[256];
+		::DWORD CheckValueNameLength = sizeof(CheckValueName) / sizeof(CheckValueName[0]);
+		EnumResult = RegEnumValue(hKey, 0, CheckValueName, &CheckValueNameLength, NULL, NULL, NULL, NULL);
+		bool bZeroValues = EnumResult != ERROR_SUCCESS;
+
+		::RegCloseKey(hKey);
+
+		if (bZeroSubKeys && bZeroValues)
+		{
+			// No more values - delete the section
+			::RegDeleteKey(HKEY_CURRENT_USER, *FullRegistryKey);
+		}
+	}
+
+	return Result == ERROR_SUCCESS;
+}
+
 uint32 FWindowsPlatformMisc::GetLastError()
 {
 	return (uint32)::GetLastError();
diff --git a/Engine/Source/Runtime/Core/Public/Android/AndroidProcess.h b/Engine/Source/Runtime/Core/Public/Android/AndroidProcess.h
--- a/Engine/Source/Runtime/Core/Public/Android/AndroidProcess.h
+++ b/Engine/Source/Runtime/Core/Public/Android/AndroidProcess.h
@@ -37,3 +37,5 @@ struct CORE_API FAndroidPlatformProcess : public FGenericPlatformProcess
 };
 
 typedef FAndroidPlatformProcess FPlatformProcess;
+
+typedef FSystemWideCriticalSectionNotImplemented FSystemWideCriticalSection;
diff --git a/Engine/Source/Runtime/Core/Public/GenericPlatform/GenericPlatformMisc.h b/Engine/Source/Runtime/Core/Public/GenericPlatform/GenericPlatformMisc.h
--- a/Engine/Source/Runtime/Core/Public/GenericPlatform/GenericPlatformMisc.h
+++ b/Engine/Source/Runtime/Core/Public/GenericPlatform/GenericPlatformMisc.h
@@ -438,6 +438,17 @@ public:
 	 */
 	static bool GetStoredValue(const FString& InStoreId, const FString& InSectionName, const FString& InKeyName, FString& OutValue);
 
+	/**
+  	 *	Deletes value for the given section and key in the platform specific key->value store
+	 *  Note: The key->value store is user-specific, but may be used to share data between different applications for the same user
+	 *
+	 *  @param	InStoreId			The name used to identify the store you want to use (eg, MyGame)
+	 *	@param	InSectionName		The section that this key->value pair is placed within (can contain / separators, eg UserDetails/AccountInfo)
+	 *	@param	InKeyName			The name of the key to set the value for
+	 *	@return	bool				true if the value was deleted correctly, false if not found or couldn't delete
+	 */
+	static bool DeleteStoredValue(const FString& InStoreId, const FString& InSectionName, const FString& InKeyName);
+
 	/** Sends a message to a remote tool, and debugger consoles */
 	static void LowLevelOutputDebugString(const TCHAR *Message);
 	static void VARARGS LowLevelOutputDebugStringf(const TCHAR *Format, ... );
diff --git a/Engine/Source/Runtime/Core/Public/GenericPlatform/GenericPlatformProcess.h b/Engine/Source/Runtime/Core/Public/GenericPlatform/GenericPlatformProcess.h
--- a/Engine/Source/Runtime/Core/Public/GenericPlatform/GenericPlatformProcess.h
+++ b/Engine/Source/Runtime/Core/Public/GenericPlatform/GenericPlatformProcess.h
@@ -566,3 +566,28 @@ struct CORE_API FGenericPlatformProcess
 	#include "PThreadCriticalSection.h"
 	typedef FPThreadsCriticalSection FCriticalSection;
 #endif
+
+/** Platforms that don't need a working FSystemWideCriticalSection can just typedef this one */
+class FSystemWideCriticalSectionNotImplemented
+{
+public:
+	/** Construct a named, system-wide critical section and attempt to get access/ownership of it */
+	explicit FSystemWideCriticalSectionNotImplemented(const FString& Name, FTimespan Timeout = FTimespan::Zero());
+
+	/** Destructor releases system-wide critical section if it is currently owned */
+	~FSystemWideCriticalSectionNotImplemented() {}
+
+	/**
+	 * Does the calling thread have ownership of the system-wide critical section?
+	 *
+	 * @return True if the system-wide lock is obtained.
+	 */
+	bool IsValid() const { return false; }
+
+	/** Releases system-wide critical section if it is currently owned */
+	void Release() {}
+
+private:
+	FSystemWideCriticalSectionNotImplemented(const FSystemWideCriticalSectionNotImplemented&);
+	FSystemWideCriticalSectionNotImplemented& operator=(const FSystemWideCriticalSectionNotImplemented&);
+};
diff --git a/Engine/Source/Runtime/Core/Public/HTML5/HTML5PlatformProcess.h b/Engine/Source/Runtime/Core/Public/HTML5/HTML5PlatformProcess.h
--- a/Engine/Source/Runtime/Core/Public/HTML5/HTML5PlatformProcess.h
+++ b/Engine/Source/Runtime/Core/Public/HTML5/HTML5PlatformProcess.h
@@ -65,3 +65,5 @@ public:
 };
 
 typedef FHTML5CriticalSection FCriticalSection;
+
+typedef FSystemWideCriticalSectionNotImplemented FSystemWideCriticalSection;
diff --git a/Engine/Source/Runtime/Core/Public/IOS/IOSPlatformMisc.h b/Engine/Source/Runtime/Core/Public/IOS/IOSPlatformMisc.h
--- a/Engine/Source/Runtime/Core/Public/IOS/IOSPlatformMisc.h
+++ b/Engine/Source/Runtime/Core/Public/IOS/IOSPlatformMisc.h
@@ -93,6 +93,7 @@ struct CORE_API FIOSPlatformMisc : public FGenericPlatformMisc
 	static FString GetDefaultLocale();
 	static bool SetStoredValue(const FString& InStoreId, const FString& InSectionName, const FString& InKeyName, const FString& InValue);
 	static bool GetStoredValue(const FString& InStoreId, const FString& InSectionName, const FString& InKeyName, FString& OutValue);
+	static bool DeleteStoredValue(const FString& InStoreId, const FString& InSectionName, const FString& InKeyName);
 	static TArray<uint8> GetSystemFontBytes();
 	static TArray<FString> GetPreferredLanguages();
 	static FString GetLocalCurrencyCode();
diff --git a/Engine/Source/Runtime/Core/Public/IOS/IOSPlatformProcess.h b/Engine/Source/Runtime/Core/Public/IOS/IOSPlatformProcess.h
--- a/Engine/Source/Runtime/Core/Public/IOS/IOSPlatformProcess.h
+++ b/Engine/Source/Runtime/Core/Public/IOS/IOSPlatformProcess.h
@@ -54,3 +54,4 @@ public:
 
 typedef FIOSPlatformAffinity FPlatformAffinity;
 
+typedef FSystemWideCriticalSectionNotImplemented FSystemWideCriticalSection;
diff --git a/Engine/Source/Runtime/Core/Public/Linux/LinuxPlatformProcess.h b/Engine/Source/Runtime/Core/Public/Linux/LinuxPlatformProcess.h
--- a/Engine/Source/Runtime/Core/Public/Linux/LinuxPlatformProcess.h
+++ b/Engine/Source/Runtime/Core/Public/Linux/LinuxPlatformProcess.h
@@ -301,3 +301,35 @@ inline bool FProcHandle::Close()
 	}
 	return false;
 }
+
+/**
+  * Linux implementation of the FSystemWideCriticalSection. Uses exclusive file locking.
+  **/
+class FLinuxSystemWideCriticalSection
+{
+public:
+	/** Construct a named, system-wide critical section and attempt to get access/ownership of it */
+	explicit FLinuxSystemWideCriticalSection(const FString& InName, FTimespan InTimeout = FTimespan::Zero());
+
+	/** Destructor releases system-wide critical section if it is currently owned */
+	~FLinuxSystemWideCriticalSection();
+
+	/**
+	 * Does the calling thread have ownership of the system-wide critical section?
+	 *
+	 * @return True if the system-wide lock is obtained. WARNING: Returns true for abandoned locks so shared resources can be in undetermined states.
+	 */
+	bool IsValid() const;
+
+	/** Releases system-wide critical section if it is currently owned */
+	void Release();
+
+private:
+	FLinuxSystemWideCriticalSection(const FLinuxSystemWideCriticalSection&);
+	FLinuxSystemWideCriticalSection& operator=(const FLinuxSystemWideCriticalSection&);
+
+private:
+	int32 FileHandle;
+};
+
+typedef FLinuxSystemWideCriticalSection FSystemWideCriticalSection;
diff --git a/Engine/Source/Runtime/Core/Public/Mac/MacPlatformProcess.h b/Engine/Source/Runtime/Core/Public/Mac/MacPlatformProcess.h
--- a/Engine/Source/Runtime/Core/Public/Mac/MacPlatformProcess.h
+++ b/Engine/Source/Runtime/Core/Public/Mac/MacPlatformProcess.h
@@ -24,8 +24,8 @@ public:
 };
 
 /**
-* Mac implementation of the Process OS functions
-**/
+ * Mac implementation of the Process OS functions
+ **/
 struct CORE_API FMacPlatformProcess : public FGenericPlatformProcess
 {
 	struct FProcEnumInfo;
@@ -138,5 +138,35 @@ struct CORE_API FMacPlatformProcess : public FGenericPlatformProcess
 	static const TCHAR* UserLogsDir();
 };
 
-typedef FMacPlatformProcess FPlatformProcess;
+/**
+ * Mac implementation of the FSystemWideCriticalSection. Uses exclusive file locking.
+ **/
+class FMacSystemWideCriticalSection
+{
+public:
+	/** Construct a named, system-wide critical section and attempt to get access/ownership of it */
+	explicit FMacSystemWideCriticalSection(const FString& InName, FTimespan InTimeout = FTimespan::Zero());
 
+	/** Destructor releases system-wide critical section if it is currently owned */
+	~FMacSystemWideCriticalSection();
+
+	/**
+	 * Does the calling thread have ownership of the system-wide critical section?
+	 *
+	 * @return True if the system-wide lock is obtained. WARNING: Returns true for abandoned locks so shared resources can be in undetermined states.
+	 */
+	bool IsValid() const;
+
+	/** Releases system-wide critical section if it is currently owned */
+	void Release();
+
+private:
+	FMacSystemWideCriticalSection(const FMacSystemWideCriticalSection&);
+	FMacSystemWideCriticalSection& operator=(const FMacSystemWideCriticalSection&);
+
+private:
+	int32 FileHandle;
+};
+
+typedef FMacPlatformProcess FPlatformProcess;
+typedef FMacSystemWideCriticalSection FSystemWideCriticalSection;
diff --git a/Engine/Source/Runtime/Core/Public/WinRT/WinRTProcess.h b/Engine/Source/Runtime/Core/Public/WinRT/WinRTProcess.h
--- a/Engine/Source/Runtime/Core/Public/WinRT/WinRTProcess.h
+++ b/Engine/Source/Runtime/Core/Public/WinRT/WinRTProcess.h
@@ -39,3 +39,4 @@ typedef FWinRTProcess FPlatformProcess;
 #include "../WinRT/WinRTCriticalSection.h"
 typedef FWinRTCriticalSection FCriticalSection;
 
+typedef FSystemWideCriticalSectionNotImplemented FSystemWideCriticalSection;
diff --git a/Engine/Source/Runtime/Core/Public/Windows/WindowsCriticalSection.h b/Engine/Source/Runtime/Core/Public/Windows/WindowsCriticalSection.h
--- a/Engine/Source/Runtime/Core/Public/Windows/WindowsCriticalSection.h
+++ b/Engine/Source/Runtime/Core/Public/Windows/WindowsCriticalSection.h
@@ -68,3 +68,31 @@ public:
 		LeaveCriticalSection(&CriticalSection);
 	}
 };
+
+/** System-Wide Critical Section for windows using mutex */
+class CORE_API FWindowsSystemWideCriticalSection
+{
+public:
+	/** Construct a named, system-wide critical section and attempt to get access/ownership of it */
+	explicit FWindowsSystemWideCriticalSection(const FString& InName, FTimespan InTimeout = FTimespan::Zero());
+
+	/** Destructor releases system-wide critical section if it is currently owned */
+	~FWindowsSystemWideCriticalSection();
+
+	/**
+	 * Does the calling thread have ownership of the system-wide critical section?
+	 *
+	 * @return True if obtained. WARNING: Returns true for an owned but previously abandoned locks so shared resources can be in undetermined states. You must handle shared data robustly.
+	 */
+	bool IsValid() const;
+
+	/** Releases system-wide critical section if it is currently owned */
+	void Release();
+
+private:
+	FWindowsSystemWideCriticalSection(const FWindowsSystemWideCriticalSection&);
+	FWindowsSystemWideCriticalSection& operator=(const FWindowsSystemWideCriticalSection&);
+
+private:
+	HANDLE Mutex;
+};
diff --git a/Engine/Source/Runtime/Core/Public/Windows/WindowsPlatformMisc.h b/Engine/Source/Runtime/Core/Public/Windows/WindowsPlatformMisc.h
--- a/Engine/Source/Runtime/Core/Public/Windows/WindowsPlatformMisc.h
+++ b/Engine/Source/Runtime/Core/Public/Windows/WindowsPlatformMisc.h
@@ -105,6 +105,7 @@ struct CORE_API FWindowsPlatformMisc
 	static void RaiseException( uint32 ExceptionCode );
 	static bool SetStoredValue(const FString& InStoreId, const FString& InSectionName, const FString& InKeyName, const FString& InValue);
 	static bool GetStoredValue(const FString& InStoreId, const FString& InSectionName, const FString& InKeyName, FString& OutValue);
+	static bool DeleteStoredValue(const FString& InStoreId, const FString& InSectionName, const FString& InKeyName);
 
 	static bool CoInitialize();
 	static void CoUninitialize();
diff --git a/Engine/Source/Runtime/Core/Public/Windows/WindowsPlatformProcess.h b/Engine/Source/Runtime/Core/Public/Windows/WindowsPlatformProcess.h
--- a/Engine/Source/Runtime/Core/Public/Windows/WindowsPlatformProcess.h
+++ b/Engine/Source/Runtime/Core/Public/Windows/WindowsPlatformProcess.h
@@ -217,3 +217,4 @@ inline bool FProcHandle::Close()
 #include "WindowsCriticalSection.h"
 
 typedef FWindowsCriticalSection FCriticalSection;
+typedef FWindowsSystemWideCriticalSection FSystemWideCriticalSection;
diff --git a/Engine/Source/Runtime/Engine/Classes/Engine/EndUserSettings.h b/Engine/Source/Runtime/Engine/Classes/Engine/EndUserSettings.h
new file mode 100644
--- /dev/null
+++ b/Engine/Source/Runtime/Engine/Classes/Engine/EndUserSettings.h
@@ -0,0 +1,42 @@
+// Copyright 1998-2015 Epic Games, Inc. All Rights Reserved.
+
+//=============================================================================
+// EndUserSettings.h
+//
+// A configuration class that holds information for the end-user of a game.
+// You may want to expose these values in your game's UI.
+// Supplied so that the engine 'remembers' the last setup the user had.
+//=============================================================================
+
+#pragma once
+#include "ImportantToggleSettingInterface.h"
+#include "EndUserSettings.generated.h"
+
+UCLASS(MinimalAPI, hidecategories=Object, config=Engine, defaultconfig)
+class UEndUserSettings : public UObject, public IImportantToggleSettingInterface
+{
+	GENERATED_UCLASS_BODY()
+
+	/** Determines whether the engine sends anonymous usage information about game sessions to Epic Games in order to improve Unreal Engine. Information will never be shared with 3rd parties. */
+	UPROPERTY(EditAnywhere, config, Category=Privacy)
+	bool bSendAnonymousUsageDataToEpic;
+
+public:
+	/** Sets bSendAnonymousUsageDataToEpic. It MUST be set this way when the end user changes the value to trigger the correct state change in the analytics system. */
+	void SetSendAnonymousUsageDataToEpic(bool bEnable);
+
+	// BEGIN IImportantToggleSettingInterface
+	virtual void GetToogleCategoryAndPropertyNames(FName& OutCategory, FName& OutProperty) const override;
+	virtual FText GetFalseStateLabel() const override;
+	virtual FText GetFalseStateTooltip() const override;
+	virtual FText GetFalseStateDescription() const override;
+	virtual FText GetTrueStateLabel() const override;
+	virtual FText GetTrueStateTooltip() const override;
+	virtual FText GetTrueStateDescription() const override;
+	virtual FString GetAdditionalInfoUrl() const override;
+	virtual FText GetAdditionalInfoUrlLabel() const override;
+	// END IImportantToggleSettingInterface
+
+private:
+	void OnSendAnonymousUsageDataToEpicChanged();
+};
diff --git a/Engine/Source/Runtime/Engine/Classes/Engine/Engine.h b/Engine/Source/Runtime/Engine/Classes/Engine/Engine.h
--- a/Engine/Source/Runtime/Engine/Classes/Engine/Engine.h
+++ b/Engine/Source/Runtime/Engine/Classes/Engine/Engine.h
@@ -1097,9 +1097,9 @@ public:
 	UPROPERTY(EditAnywhere, config, Category=Settings)
 	float NearClipPlane;
 
-	/** Can a runtime game/application report anonymous hardware survey statistics (such as display resolution and GPU model) back to Epic? */
-	UPROPERTY(EditAnywhere, config, Category=Settings, AdvancedDisplay)
-	uint32 bHardwareSurveyEnabled:1;
+	/** DEPRECATED - Can a runtime game/application report anonymous hardware survey statistics (such as display resolution and GPU model) back to Epic? */
+	UPROPERTY()
+	uint32 bHardwareSurveyEnabled_DEPRECATED:1;
 
 	/** Flag for completely disabling subtitles for localized sounds. */
 	UPROPERTY(EditAnywhere, config, Category=Subtitles)
@@ -2620,6 +2620,8 @@ public:
 
 	/** @return true if editor analytics are enabled */
 	virtual bool AreEditorAnalyticsEnabled() const { return false; }
+	/** @return true if in-game analytics are enabled */
+	bool AreGameAnalyticsEnabled() const;
 	virtual void CreateStartupAnalyticsAttributes( TArray<struct FAnalyticsEventAttribute>& StartSessionAttributes ) const {}
 	
 	/** @return true if the engine is autosaving a package */
diff --git a/Engine/Source/Runtime/Engine/Classes/Engine/ImportantToggleSettingInterface.h b/Engine/Source/Runtime/Engine/Classes/Engine/ImportantToggleSettingInterface.h
new file mode 100644
--- /dev/null
+++ b/Engine/Source/Runtime/Engine/Classes/Engine/ImportantToggleSettingInterface.h
@@ -0,0 +1,35 @@
+// Copyright 1998-2015 Epic Games, Inc. All Rights Reserved.
+
+//=============================================================================
+// ImportantToogleSettingInterface.h
+//
+// Interface for settings classes containing a single, emphasized boolean value.
+// Implement this interface to make a settings category containing one boolean
+// property lots of extra info for the user.
+//=============================================================================
+
+#pragma once
+#include "Engine/EngineTypes.h"
+#include "ImportantToggleSettingInterface.generated.h"
+
+UINTERFACE(MinimalAPI, meta=(CannotImplementInterfaceInBlueprint))
+class UImportantToggleSettingInterface : public UInterface
+{
+	GENERATED_UINTERFACE_BODY()
+};
+
+class IImportantToggleSettingInterface
+{
+	GENERATED_IINTERFACE_BODY()
+
+	virtual void GetToogleCategoryAndPropertyNames(FName& OutCategory, FName& OutProperty) const { OutCategory = NAME_None; OutProperty = NAME_None; }
+
+	virtual FText GetFalseStateLabel() const { return FText::GetEmpty(); }
+	virtual FText GetFalseStateTooltip() const { return FText::GetEmpty(); }
+	virtual FText GetFalseStateDescription() const { return FText::GetEmpty(); }
+	virtual FText GetTrueStateLabel() const { return FText::GetEmpty(); }
+	virtual FText GetTrueStateTooltip() const { return FText::GetEmpty(); }
+	virtual FText GetTrueStateDescription() const { return FText::GetEmpty(); }
+	virtual FString GetAdditionalInfoUrl() const { return FString(); }
+	virtual FText GetAdditionalInfoUrlLabel() const { return FText::GetEmpty(); }
+};
diff --git a/Engine/Source/Runtime/Engine/Engine.Build.cs b/Engine/Source/Runtime/Engine/Engine.Build.cs
--- a/Engine/Source/Runtime/Engine/Engine.Build.cs
+++ b/Engine/Source/Runtime/Engine/Engine.Build.cs
@@ -34,6 +34,8 @@ public class Engine : ModuleRules
 				"AutomationWorker",
                 "Analytics",
 				"MovieSceneCapture",
+				"DesktopPlatform",
+				"Analytics"
 			}
 		);
 
diff --git a/Engine/Source/Runtime/Engine/Private/EndUserSettings.cpp b/Engine/Source/Runtime/Engine/Private/EndUserSettings.cpp
new file mode 100644
--- /dev/null
+++ b/Engine/Source/Runtime/Engine/Private/EndUserSettings.cpp
@@ -0,0 +1,96 @@
+// Copyright 1998-2015 Epic Games, Inc. All Rights Reserved.
+
+#include "EnginePrivate.h"
+#include "Engine/EndUserSettings.h"
+#include "Public/Interfaces/IAnalyticsProvider.h"
+#include "EngineAnalytics.h"
+
+#define LOCTEXT_NAMESPACE "EndUserSettings"
+
+UEndUserSettings::UEndUserSettings(const FObjectInitializer& ObjectInitializer)
+	: Super(ObjectInitializer)
+	, bSendAnonymousUsageDataToEpic(true)
+{
+}
+
+void UEndUserSettings::GetToogleCategoryAndPropertyNames(FName& OutCategory, FName& OutProperty) const
+{
+	OutCategory = FName("Privacy");
+	OutProperty = FName("bSendAnonymousUsageDataToEpic");
+};
+
+FText UEndUserSettings::GetFalseStateLabel() const
+{
+	return LOCTEXT("FalseStateLabel", "End-users don't send");
+};
+
+FText UEndUserSettings::GetFalseStateTooltip() const
+{
+	return LOCTEXT("FalseStateTooltip", "By default, your end-users don't send anonymous usage data to Epic Games.");
+};
+
+FText UEndUserSettings::GetFalseStateDescription() const
+{
+	return LOCTEXT("FalseStateDescription", "You have defaulted your users' settings to not send anonymous usage data to Epic Games. You can allow users to opt-in by adding this setting to your game UI and calling UEndUserSettings::SetSendAnonymousUsageDataToEpic() with their choice. Please consider defaulting this to true or allowing users to switch it on to help improve Unreal Engine. Epic Games will never sell or trade individual usage data to / with third party organizations. When permitted, we collect information about your users' game sessions, the game they're playing and for how long they play. Their information would be encrypted and sent anonymously, and only when they run your product.");
+};
+
+FText UEndUserSettings::GetTrueStateLabel() const
+{
+	return LOCTEXT("TrueStateLabel", "End-users send to Epic");
+};
+
+FText UEndUserSettings::GetTrueStateTooltip() const
+{
+	return LOCTEXT("TrueStateTooltip", "By default, your end-users send anonymous usage data to Epic Games.");
+};
+
+FText UEndUserSettings::GetTrueStateDescription() const
+{
+	return LOCTEXT("TrueStateDescription", "You have defaulted your users' settings to send anonymous usage data to Epic Games. You can allow users to opt-out by adding this setting to your game UI and calling UEndUserSettings::SetSendAnonymousUsageDataToEpic() with their choice. Thank you for helping to improve Unreal Engine. Epic Games will never sell or trade individual usage data to / with third party organizations. We will collect information about your users' game sessions, the game they're playing and for how long they play. Their information is sent anonymously. Their information is only sent when they run your product, it is encrypted and sent to our servers.");
+};
+
+FString UEndUserSettings::GetAdditionalInfoUrl() const
+{
+	return FString(TEXT("http://epicgames.com/privacynotice"));
+};
+
+FText UEndUserSettings::GetAdditionalInfoUrlLabel() const
+{
+	return LOCTEXT("HyperlinkLabel", "Epic Games Privacy Notice");
+};
+
+void UEndUserSettings::SetSendAnonymousUsageDataToEpic(bool bEnable)
+{
+	if (bSendAnonymousUsageDataToEpic == bEnable)
+	{
+		bSendAnonymousUsageDataToEpic = bEnable;
+		OnSendAnonymousUsageDataToEpicChanged();
+	}
+}
+
+void UEndUserSettings::OnSendAnonymousUsageDataToEpicChanged()
+{
+	if (bSendAnonymousUsageDataToEpic)
+	{
+		// Attempt to initialize analytics and send opt-in event
+		if (!FEngineAnalytics::IsAvailable())
+		{
+			FEngineAnalytics::Initialize();
+			if (FEngineAnalytics::IsAvailable())
+			{
+				FEngineAnalytics::GetProvider().RecordEvent(FString("Engine.Privacy.EndUserOptIn"));
+			}
+		}
+	}
+	else
+	{
+		// Send opt-out event and shutdown analytics
+		if (FEngineAnalytics::IsAvailable())
+		{
+			FEngineAnalytics::GetProvider().RecordEvent(FString("Engine.Privacy.EndUserOptOut"));
+			FEngineAnalytics::Shutdown();
+		}
+	}
+}
+
+#undef LOCTEXT_NAMESPACE
diff --git a/Engine/Source/Runtime/Engine/Private/EngineAnalytics.cpp b/Engine/Source/Runtime/Engine/Private/EngineAnalytics.cpp
--- a/Engine/Source/Runtime/Engine/Private/EngineAnalytics.cpp
+++ b/Engine/Source/Runtime/Engine/Private/EngineAnalytics.cpp
@@ -8,8 +8,9 @@
 #include "GeneralProjectSettings.h"
 
 bool FEngineAnalytics::bIsInitialized;
+bool FEngineAnalytics::bIsEditorRun;
+bool FEngineAnalytics::bIsGameRun;
 TSharedPtr<IAnalyticsProvider> FEngineAnalytics::Analytics;
-bool FEngineAnalytics::bShouldSendUsageEvents;
 
 /**
  * Engine analytics config log to initialize the engine analytics provider.
@@ -41,22 +42,16 @@ void FEngineAnalytics::Initialize()
 
 	// this will only be true for builds that have editor support (currently PC, Mac, Linux)
 	// The idea here is to only send editor events for actual editor runs, not for things like -game runs of the editor.
-	const bool bIsEditorRun = WITH_EDITOR && GIsEditor && !IsRunningCommandlet();
+	bIsEditorRun = WITH_EDITOR && GIsEditor && !IsRunningCommandlet();
 
 	// We also want to identify a real run of a game, which is NOT necessarily the opposite of an editor run.
 	// Ideally we'd be able to tell explicitly, but with content-only games, it becomes difficult.
 	// So we ensure we are not an editor run, we don't have EDITOR stuff compiled in, we are not running a commandlet,
 	// we are not a generic, utility program, and we require cooked data.
-	const bool bIsGameRun = !WITH_EDITOR && !IsRunningCommandlet() && !FPlatformProperties::IsProgram() && FPlatformProperties::RequiresCookedData();
-
-	const bool bShouldInitAnalytics = bIsEditorRun || bIsGameRun;
+	bIsGameRun = !WITH_EDITOR && !IsRunningCommandlet() && !FPlatformProperties::IsProgram() && FPlatformProperties::RequiresCookedData();
 
 	// Outside of the editor, the only engine analytics usage is the hardware survey
-	bShouldSendUsageEvents = bIsEditorRun 
-		? GEngine->AreEditorAnalyticsEnabled() 
-		: bIsGameRun 
-			? GEngine->bHardwareSurveyEnabled 
-			: false;
+	const bool bShouldInitAnalytics = (bIsEditorRun && GEngine->AreEditorAnalyticsEnabled()) || (bIsGameRun && GEngine->AreGameAnalyticsEnabled());
 
 	if (bShouldInitAnalytics)
 	{
@@ -109,9 +104,25 @@ void FEngineAnalytics::Initialize()
 			Analytics = FAnalytics::Get().CreateAnalyticsProvider(
 				FName(*DefaultEngineAnalyticsConfig.Execute(TEXT("ProviderModuleName"), true)), 
 				DefaultEngineAnalyticsConfig);
+
 			if (Analytics.IsValid())
 			{
-				Analytics->SetUserID(FString::Printf(TEXT("%s|%s|%s"), *FPlatformMisc::GetMachineId().ToString(EGuidFormats::Digits).ToLower(), *FPlatformMisc::GetEpicAccountId(), *FPlatformMisc::GetOperatingSystemId()));
+				// Use an anonymous user id in-game
+				if (bIsGameRun)
+				{
+					FString AnonymousId;
+					if (!FPlatformMisc::GetStoredValue(TEXT("Epic Games"), TEXT("Unreal Engine/Privacy"), TEXT("AnonymousID"), AnonymousId) || AnonymousId.IsEmpty())
+					{
+						AnonymousId = FGuid::NewGuid().ToString(EGuidFormats::DigitsWithHyphensInBraces);
+						FPlatformMisc::SetStoredValue(TEXT("Epic Games"), TEXT("Unreal Engine/Privacy"), TEXT("AnonymousID"), AnonymousId);
+					}
+
+					Analytics->SetUserID(FString::Printf(TEXT("ANON-%s"), *AnonymousId));
+				}
+				else
+				{
+					Analytics->SetUserID(FString::Printf(TEXT("%s|%s|%s"), *FPlatformMisc::GetMachineId().ToString(EGuidFormats::Digits).ToLower(), *FPlatformMisc::GetEpicAccountId(), *FPlatformMisc::GetOperatingSystemId()));
+				}
 
 				TArray<FAnalyticsEventAttribute> StartSessionAttributes;
 				GEngine->CreateStartupAnalyticsAttributes( StartSessionAttributes );
@@ -123,10 +134,11 @@ void FEngineAnalytics::Initialize()
 				StartSessionAttributes.Emplace(TEXT("ProjectVersion"), ProjectSettings.ProjectVersion);
 
 				Analytics->StartSession( StartSessionAttributes );
+
+				bIsInitialized = true;
 			}
 		}
 	}
-	bIsInitialized = true;
 }
 
 
diff --git a/Engine/Source/Runtime/Engine/Private/GameEngine.cpp b/Engine/Source/Runtime/Engine/Private/GameEngine.cpp
--- a/Engine/Source/Runtime/Engine/Private/GameEngine.cpp
+++ b/Engine/Source/Runtime/Engine/Private/GameEngine.cpp
@@ -203,8 +203,8 @@ void UGameEngine::ConditionallyOverrideSettings(int32& ResolutionX, int32& Resol
 	FParse::Value(FCommandLine::Get(), TEXT("ResX="), ResolutionX);
 	FParse::Value(FCommandLine::Get(), TEXT("ResY="), ResolutionY);
 
-	// consume available desktop area
-	FDisplayMetrics DisplayMetrics;
+		// consume available desktop area
+		FDisplayMetrics DisplayMetrics;
 	FDisplayMetrics::GetDisplayMetrics(DisplayMetrics);
 		
 	int32 DesktopResolutionX = DisplayMetrics.PrimaryDisplayWidth;
@@ -421,7 +421,7 @@ UEngine::UEngine(const FObjectInitializer& ObjectInitializer)
 
 	bUseSound = true;
 
-	bHardwareSurveyEnabled = true;
+	bHardwareSurveyEnabled_DEPRECATED = true;
 	bPendingHardwareSurveyResults = false;
 	bIsInitialized = false;
 
@@ -556,7 +556,7 @@ void UGameEngine::PreExit()
 
 			if (World->GetGameInstance() != nullptr)
 			{
-				World->GetGameInstance()->Shutdown();
+			World->GetGameInstance()->Shutdown();
 			}
 
 			World->FlushLevelStreaming(EFlushLevelStreamingType::Visibility);
@@ -602,7 +602,7 @@ bool UGameEngine::NetworkRemapPath(UWorld* InWorld, FString& Str, bool bReading
 		{
 			const FString StreamingLevelName = StreamingLevel->GetWorldAsset().GetLongPackageName();
 			if (StreamingLevelName == PrefixedName)
-			{
+	{
 				Str = PrefixedName;
 				return true;
 			}
@@ -1020,45 +1020,45 @@ void UGameEngine::Tick( float DeltaSeconds, bool bIdleMode )
 	}
 
 	// Tick the viewport
-	if ( GameViewport != NULL && !bIdleMode )
-	{
-		SCOPE_CYCLE_COUNTER(STAT_GameViewportTick);
-		GameViewport->Tick(DeltaSeconds);
-	}
+		if ( GameViewport != NULL && !bIdleMode )
+		{
+			SCOPE_CYCLE_COUNTER(STAT_GameViewportTick);
+			GameViewport->Tick(DeltaSeconds);
+		}
 
-	if (FPlatformProperties::SupportsWindowedMode())
-	{
-		// Hide the splashscreen and show the game window
-		static bool bFirstTime = true;
-		if ( bFirstTime )
+		if (FPlatformProperties::SupportsWindowedMode())
 		{
-			bFirstTime = false;
-			FPlatformSplash::Hide();
-			if ( GameViewportWindow.IsValid() )
+			// Hide the splashscreen and show the game window
+			static bool bFirstTime = true;
+			if ( bFirstTime )
 			{
-				GameViewportWindow.Pin()->ShowWindow();
-				FSlateApplication::Get().RegisterGameViewport( GameViewportWidget.ToSharedRef() );
+				bFirstTime = false;
+				FPlatformSplash::Hide();
+				if ( GameViewportWindow.IsValid() )
+				{
+					GameViewportWindow.Pin()->ShowWindow();
+					FSlateApplication::Get().RegisterGameViewport( GameViewportWidget.ToSharedRef() );
+				}
 			}
 		}
-	}
 
-	if (!bIdleMode && !IsRunningDedicatedServer() && !IsRunningCommandlet())
-	{
-		// Render everything.
-		RedrawViewports();
-	}
+		if (!bIdleMode && !IsRunningDedicatedServer() && !IsRunningCommandlet())
+		{
+			// Render everything.
+			RedrawViewports();
+		}
 
-	if( GIsClient )
-	{
-		// Update resource streaming after viewports have had a chance to update view information. Normal update.
-		QUICK_SCOPE_CYCLE_COUNTER(STAT_UGameEngine_Tick_IStreamingManager);
-		IStreamingManager::Get().Tick( DeltaSeconds );
-	}
+		if( GIsClient )
+		{
+			// Update resource streaming after viewports have had a chance to update view information. Normal update.
+				QUICK_SCOPE_CYCLE_COUNTER(STAT_UGameEngine_Tick_IStreamingManager);
+				IStreamingManager::Get().Tick( DeltaSeconds );
+			}
 
 	// Update Audio. This needs to occur after rendering as the rendering code updates the listener position.
 	FAudioDeviceManager* AudioDeviceManager = GEngine->GetAudioDeviceManager();
 	if (AudioDeviceManager)
-	{
+			{
 		AudioDeviceManager->UpdateActiveAudioDevices(bIsAnyNonPreviewWorldUnpaused);
 	}
 
diff --git a/Engine/Source/Runtime/Engine/Private/ImportantToggleSettingInterface.cpp b/Engine/Source/Runtime/Engine/Private/ImportantToggleSettingInterface.cpp
new file mode 100644
--- /dev/null
+++ b/Engine/Source/Runtime/Engine/Private/ImportantToggleSettingInterface.cpp
@@ -0,0 +1,9 @@
+// Copyright 1998-2015 Epic Games, Inc. All Rights Reserved.
+
+#include "EnginePrivate.h"
+#include "Engine/ImportantToggleSettingInterface.h"
+
+UImportantToggleSettingInterface::UImportantToggleSettingInterface(const FObjectInitializer& ObjectInitializer)
+	: Super(ObjectInitializer)
+{
+}
\ No newline at end of file
diff --git a/Engine/Source/Runtime/Engine/Private/UnrealEngine.cpp b/Engine/Source/Runtime/Engine/Private/UnrealEngine.cpp
--- a/Engine/Source/Runtime/Engine/Private/UnrealEngine.cpp
+++ b/Engine/Source/Runtime/Engine/Private/UnrealEngine.cpp
@@ -126,6 +126,7 @@
 #endif
 
 #include "InstancedReferenceSubobjectHelper.h"
+#include "Engine/EndUserSettings.h"
 #include "ABTesting.h"
 #include "Performance/EnginePerformanceTargets.h"
 
@@ -6125,7 +6126,7 @@ void UEngine::InitHardwareSurvey()
 {
 	bool bEnabled = true;
 #if !WITH_EDITORONLY_DATA
-	bEnabled = bHardwareSurveyEnabled;
+	bEnabled = AreGameAnalyticsEnabled();
 #endif
 
 	// The hardware survey costs time and we don't want to slow down debug builds.
@@ -10081,6 +10082,11 @@ UGameViewportClient* UEngine::GameViewportForWorld(const UWorld *InWorld) const
 	return (Context ? Context->GameViewport : NULL);
 }
 
+bool UEngine::AreGameAnalyticsEnabled() const
+{ 
+	return GetDefault<UEndUserSettings>()->bSendAnonymousUsageDataToEpic;
+}
+
 FWorldContext* UEngine::GetWorldContextFromGameViewport(const UGameViewportClient *InViewport)
 {
 	for (FWorldContext& WorldContext : WorldList)
diff --git a/Engine/Source/Runtime/Engine/Public/EngineAnalytics.h b/Engine/Source/Runtime/Engine/Public/EngineAnalytics.h
--- a/Engine/Source/Runtime/Engine/Public/EngineAnalytics.h
+++ b/Engine/Source/Runtime/Engine/Public/EngineAnalytics.h
@@ -24,16 +24,16 @@ public:
 	 */
 	static ENGINE_API IAnalyticsProvider& GetProvider();
 	/** Helper function to determine if the provider is valid. */
-	static ENGINE_API bool IsAvailable() { return bShouldSendUsageEvents && Analytics.IsValid(); }
+	static ENGINE_API bool IsAvailable() { return Analytics.IsValid(); }
 	/** Called to initialize the singleton. */
-	static void Initialize();
+	static ENGINE_API void Initialize();
 	/** Called to shut down the singleton */
-	static void Shutdown();
+	static ENGINE_API void Shutdown();
 
 private:
 	static bool bIsInitialized;
-	/** This allows us to "opt-out" of engine analytics. We still create a session, but we don't send any usage events. */
-	static bool bShouldSendUsageEvents;
+	static bool bIsEditorRun;
+	static bool bIsGameRun;
 	static TSharedPtr<IAnalyticsProvider> Analytics;
 };
 
