Fixing 100% crash when performing class reinstancing during  hot reload after changing class constructor only (no structural changes to classes).

Allowed Kismet Reinstancer to have the same source and target class objects when reinstancing during hot reload.

diff --git a/Engine/Source/Editor/UnrealEd/Private/Kismet2/KismetReinstanceUtilities.cpp b/Engine/Source/Editor/UnrealEd/Private/Kismet2/KismetReinstanceUtilities.cpp
--- a/Engine/Source/Editor/UnrealEd/Private/Kismet2/KismetReinstanceUtilities.cpp
+++ b/Engine/Source/Editor/UnrealEd/Private/Kismet2/KismetReinstanceUtilities.cpp
@@ -1221,7 +1221,7 @@ void FBlueprintCompileReinstancer::ReplaceInstancesOfClass_Inner(TMap<UClass*, U
 			UClass* OldClass = OldToNewClass.Key;
 			UClass* NewClass = OldToNewClass.Value;
 			check(OldClass && NewClass);
-			check(OldClass != NewClass);
+			check(OldClass != NewClass || GIsHotReload);
 
 			{ BP_SCOPED_COMPILER_EVENT_STAT(EKismetReinstancerStats_ReplaceInstancesOfClass);
 
@@ -1457,7 +1457,7 @@ void FBlueprintCompileReinstancer::ReplaceInstancesOfClass_Inner(TMap<UClass*, U
 			UClass* OldClass = OldToNewClass.Key;
 			UClass* NewClass = OldToNewClass.Value;
 			check(OldClass && NewClass);
-			check(OldClass != NewClass);
+			check(OldClass != NewClass || GIsHotReload);
 
 			FReplaceReferenceHelper::IncludeCDO(OldClass, NewClass, OldToNewInstanceMap, SourceObjects, InOriginalCDO);
 
