Downgraded the mismatched manifest warning in the loc gather

This always happens due to OneSky exporting our PO files without the correct name, which causes constant warnings on our CIS machines.

diff --git a/Engine/Source/Editor/UnrealEd/Private/Commandlets/InternationalizationExportCommandlet.cpp b/Engine/Source/Editor/UnrealEd/Private/Commandlets/InternationalizationExportCommandlet.cpp
--- a/Engine/Source/Editor/UnrealEd/Private/Commandlets/InternationalizationExportCommandlet.cpp
+++ b/Engine/Source/Editor/UnrealEd/Private/Commandlets/InternationalizationExportCommandlet.cpp
@@ -605,7 +605,7 @@ bool UInternationalizationExportCommandlet::DoImport(const FString& SourcePath,
 
 		if (PortableObject.GetProjectName() != ManifestName.Replace(TEXT(".manifest"), TEXT("")))
 		{
-			UE_LOG(LogInternationalizationExportCommandlet, Warning, TEXT("The project name (%s) in the file (%s) did not match the target manifest project (%s)."), *POFilePath, *PortableObject.GetProjectName(), *ManifestName.Replace(TEXT(".manifest"), TEXT("")));
+			UE_LOG(LogInternationalizationExportCommandlet, Log, TEXT("The project name (%s) in the file (%s) did not match the target manifest project (%s)."), *POFilePath, *PortableObject.GetProjectName(), *ManifestName.Replace(TEXT(".manifest"), TEXT("")));
 		}
 
 		const FString ManifestFileName = DestinationPath / ManifestName;
