Removed invalid check.

diff --git a/Engine/Source/Editor/UnrealEd/Private/CookOnTheFlyServer.cpp b/Engine/Source/Editor/UnrealEd/Private/CookOnTheFlyServer.cpp
--- a/Engine/Source/Editor/UnrealEd/Private/CookOnTheFlyServer.cpp
+++ b/Engine/Source/Editor/UnrealEd/Private/CookOnTheFlyServer.cpp
@@ -4399,8 +4399,6 @@ void UCookOnTheFlyServer::ValidateCookByTheBookSettings() const
 		check(CookByTheBookOptions->bGenerateDependenciesForMaps == false);
 		check(CookByTheBookOptions->bGenerateStreamingInstallManifests == false);
 	}
-
-	check(!IsCookFlagSet(ECookInitializationFlags::Iterative) || !(CookByTheBookOptions->CreateReleaseVersion.IsEmpty()))
 }
 
 void UCookOnTheFlyServer::StartCookByTheBook( const FCookByTheBookStartupOptions& CookByTheBookStartupOptions )
