Copying //UE4/Orion-Staging to //UE4/Main (Origin: //Orion/Dev-General [REDACTED_HANDLE])

diff --git a/Engine/Source/Editor/UnrealEd/Private/Analytics/EditorAnalytics.cpp b/Engine/Source/Editor/UnrealEd/Private/Analytics/EditorAnalytics.cpp
--- a/Engine/Source/Editor/UnrealEd/Private/Analytics/EditorAnalytics.cpp
+++ b/Engine/Source/Editor/UnrealEd/Private/Analytics/EditorAnalytics.cpp
@@ -12,405 +12,6 @@
 
 #define LOCTEXT_NAMESPACE "EditorAnalytics"
 
-namespace EditorAnalyticsDefs
-{
-	static const FTimespan SessionRecordExpiration(30, 0, 0, 0);	// 30 days
-	static const FTimespan SessionRecordTimeout(0, 30, 0);			// 30 minutes
-	static const FTimespan GlobalLockWaitTimeout(0, 0, 0, 0, 500);	// 1/2 second
-	static const FString StoreId(TEXT("Epic Games"));
-	static const FString CrashSessionToken(TEXT("Crashed"));
-	static const FString DebuggerSessionToken(TEXT("Debugger"));
-	static const FString AbnormalSessionToken(TEXT("AbnormalShutdown"));
-	static const FString SessionRecordListSection(TEXT("List"));
-	static const FString SessionRecordSectionPrefix(TEXT("Unreal Engine/Editor Sessions/"));
-	static const FString EditorSessionsVersionString(TEXT("1_1"));
-	static const FString CrashStoreKey(TEXT("IsCrash"));
-	static const FString EngineVersionStoreKey(TEXT("EngineVersion"));
-	static const FString TimestampStoreKey(TEXT("Timestamp"));
-	static const FString DebuggerStoreKey(TEXT("IsDebugger"));
-	static const FString GlobalLockName(TEXT("UE4_EditorAnalytics_Lock"));
-	static const FString FalseValueString(TEXT("0"));
-	static const FString TrueValueString(TEXT("1"));
-}
-
-/* Handles writing session records to platform's storage to track crashed and timed-out editor sessions */
-class FSessionManager
-{
-public:
-	FSessionManager()
-		: bInitialized(false)
-		, bEditorCrashed(false)
-	{}
-
-	void Initialize(bool bWaitForLock = true)
-	{
-		if (!FEngineAnalytics::IsAvailable())
-		{
-			return;
-		}
-
-		TArray<FSessionRecord> SessionRecordsToReport;
-
-		{
-			// Scoped lock
-			FSystemWideCriticalSection StoredValuesLock(EditorAnalyticsDefs::GlobalLockName, bWaitForLock ? EditorAnalyticsDefs::GlobalLockWaitTimeout : FTimespan::Zero());
-
-		// Get list of sessions in storage
-			if (StoredValuesLock.IsValid() && BeginReadWriteRecords())
-		{
-			TArray<FSessionRecord> SessionRecordsToDelete;
-
-			// Attempt check each stored session
-			for (FSessionRecord& Record : SessionRecords)
-			{
-				FTimespan RecordAge = FDateTime::UtcNow() - Record.Timestamp;
-
-				if (Record.bCrashed)
-				{
-					// Crashed sessions
-					SessionRecordsToReport.Add(Record);
-					SessionRecordsToDelete.Add(Record);
-				}
-					else if (RecordAge > EditorAnalyticsDefs::SessionRecordExpiration)
-				{
-					// Delete expired session records
-					SessionRecordsToDelete.Add(Record);
-				}
-				else if (RecordAge > EditorAnalyticsDefs::SessionRecordTimeout)
-				{
-					// Timed out sessions
-					SessionRecordsToReport.Add(Record);
-					SessionRecordsToDelete.Add(Record);
-				}
-			}
-
-			for (FSessionRecord& DeletingRecord : SessionRecordsToDelete)
-			{
-				DeleteStoredRecord(DeletingRecord);
-			}
-
-			// Create a session record for this session
-			CreateAndWriteRecordForSession();
-
-			// Update and release list of sessions in storage
-			EndReadWriteRecords();
-
-			// Register for crash callbacks
-			FCoreDelegates::OnHandleSystemError.AddRaw(this, &FSessionManager::OnCrashing);
-
-			bInitialized = true;
-		}
-		}
-
-		for (FSessionRecord& ReportingSession : SessionRecordsToReport)
-		{
-			// Send error report for session that timed out or crashed
-			SendAbnormalShutdownReport(ReportingSession);
-		}
-	}
-
-	void Heartbeat()
-	{
-		if (!FEngineAnalytics::IsAvailable())
-		{
-			return;
-		}
-
-		if (!bInitialized)
-		{
-			// Try late initialization
-			const bool bWaitForLock = false;
-			Initialize(bWaitForLock);
-		}
-
-		// Update timestamp in the session record for this session 
-		if (bInitialized)
-		{
-			CurrentSession.Timestamp = FDateTime::UtcNow();
-
-			FPlatformMisc::SetStoredValue(EditorAnalyticsDefs::StoreId, CurrentSessionSectionName, EditorAnalyticsDefs::TimestampStoreKey, TimestampToString(CurrentSession.Timestamp));
-		}
-	}
-
-	void Shutdown()
-	{
-		// Clear the session record for this session
-		if (bInitialized)
-		{
-			FCoreDelegates::OnHandleSystemError.RemoveAll(this);
-
-			if (!bEditorCrashed)
-			{
-				FPlatformMisc::DeleteStoredValue(EditorAnalyticsDefs::StoreId, CurrentSessionSectionName, EditorAnalyticsDefs::CrashStoreKey);
-				FPlatformMisc::DeleteStoredValue(EditorAnalyticsDefs::StoreId, CurrentSessionSectionName, EditorAnalyticsDefs::EngineVersionStoreKey);
-				FPlatformMisc::DeleteStoredValue(EditorAnalyticsDefs::StoreId, CurrentSessionSectionName, EditorAnalyticsDefs::TimestampStoreKey);
-				FPlatformMisc::DeleteStoredValue(EditorAnalyticsDefs::StoreId, CurrentSessionSectionName, EditorAnalyticsDefs::DebuggerStoreKey);
-			}
-
-			bInitialized = false;
-		}
-	}
-
-private:
-	struct FSessionRecord
-	{
-		FString SessionId;
-		FString EngineVersion;
-		FDateTime Timestamp;
-		bool bCrashed;
-		bool bIsDebugger;
-	};
-
-private:
-	bool BeginReadWriteRecords()
-	{
-		SessionRecords.Empty();
-
-		// Lock and read the list of sessions in storage
-		FString ListSectionName = GetStoreSectionString(EditorAnalyticsDefs::SessionRecordListSection);
-
-		// Write list to SessionRecords member
-		FString SessionListString;
-		FPlatformMisc::GetStoredValue(EditorAnalyticsDefs::StoreId, ListSectionName, TEXT("SessionList"), SessionListString);
-
-		// Parse SessionListString for session ids
-		TArray<FString> SessionIds;
-		SessionListString.ParseIntoArray(SessionIds, TEXT(","));
-
-		// Retrieve all the sessions in the list from storage
-		for (FString& SessionId : SessionIds)
-		{
-			FString SectionName = GetStoreSectionString(SessionId);
-
-			FString IsCrashString;
-			FString EngineVersionString;
-			FString TimestampString;
-			FString IsDebuggerString;
-
-			if (FPlatformMisc::GetStoredValue(EditorAnalyticsDefs::StoreId, SectionName, EditorAnalyticsDefs::CrashStoreKey, IsCrashString) &&
-				FPlatformMisc::GetStoredValue(EditorAnalyticsDefs::StoreId, SectionName, EditorAnalyticsDefs::EngineVersionStoreKey, EngineVersionString) &&
-				FPlatformMisc::GetStoredValue(EditorAnalyticsDefs::StoreId, SectionName, EditorAnalyticsDefs::TimestampStoreKey, TimestampString) &&
-				FPlatformMisc::GetStoredValue(EditorAnalyticsDefs::StoreId, SectionName, EditorAnalyticsDefs::DebuggerStoreKey, IsDebuggerString))
-			{
-				FSessionRecord NewRecord;
-				NewRecord.SessionId = SessionId;
-				NewRecord.EngineVersion = EngineVersionString;
-				NewRecord.Timestamp = StringToTimestamp(TimestampString);
-				NewRecord.bCrashed = IsCrashString == EditorAnalyticsDefs::TrueValueString;
-				NewRecord.bIsDebugger = IsDebuggerString == EditorAnalyticsDefs::TrueValueString;
-
-				SessionRecords.Add(NewRecord);
-			}
-			else
-			{
-				// Clean up orphaned values, if there are any
-				FPlatformMisc::DeleteStoredValue(EditorAnalyticsDefs::StoreId, SectionName, EditorAnalyticsDefs::CrashStoreKey);
-				FPlatformMisc::DeleteStoredValue(EditorAnalyticsDefs::StoreId, SectionName, EditorAnalyticsDefs::EngineVersionStoreKey);
-				FPlatformMisc::DeleteStoredValue(EditorAnalyticsDefs::StoreId, SectionName, EditorAnalyticsDefs::TimestampStoreKey);
-				FPlatformMisc::DeleteStoredValue(EditorAnalyticsDefs::StoreId, SectionName, EditorAnalyticsDefs::DebuggerStoreKey);
-			}
-		}
-
-		return true;
-	}
-
-	void EndReadWriteRecords()
-	{
-		// Update the list of sessions in storage to match SessionRecords
-		FString SessionListString;
-		if (SessionRecords.Num() > 0)
-		{
-			for (FSessionRecord& Session : SessionRecords)
-			{
-				SessionListString.Append(Session.SessionId);
-				SessionListString.Append(TEXT(","));
-			}
-			SessionListString.RemoveAt(SessionListString.Len() - 1);
-		}
-
-		FString ListSectionName = GetStoreSectionString(EditorAnalyticsDefs::SessionRecordListSection);
-		FPlatformMisc::SetStoredValue(EditorAnalyticsDefs::StoreId, ListSectionName, TEXT("SessionList"), SessionListString);
-
-		// Clear SessionRecords member
-		SessionRecords.Empty();
-	}
-
-	void DeleteStoredRecord(const FSessionRecord& Record)
-	{
-		// Delete the session record in storage
-		FString SessionId = Record.SessionId;
-		FString SectionName = GetStoreSectionString(SessionId);
-
-		FPlatformMisc::DeleteStoredValue(EditorAnalyticsDefs::StoreId, SectionName, EditorAnalyticsDefs::CrashStoreKey);
-		FPlatformMisc::DeleteStoredValue(EditorAnalyticsDefs::StoreId, SectionName, EditorAnalyticsDefs::EngineVersionStoreKey);
-		FPlatformMisc::DeleteStoredValue(EditorAnalyticsDefs::StoreId, SectionName, EditorAnalyticsDefs::TimestampStoreKey);
-		FPlatformMisc::DeleteStoredValue(EditorAnalyticsDefs::StoreId, SectionName, EditorAnalyticsDefs::DebuggerStoreKey);
-
-		// Remove the session record from SessionRecords list
-		SessionRecords.RemoveAll([&SessionId](const FSessionRecord& X){ return X.SessionId == SessionId; });
-	}
-
-	/**
-  	 * @EventName Editor.AbnormalShutdown
-	 *
-	 * @Trigger Fired only by the editor at startup, once for each "abnormal shutdown" detected that has not already been sent.
-	 *
-	 * @Type Static
-	 *
-	 * @EventParam SessionId Analytics SessionID of the session that abnormally terminated.
-	 * @EventParam EngineVersion EngineVersion of the session that abnormally terminated.
-	 * @EventParam ShutdownType - one of Crashed, Debugger, or AbormalShutdown
-	 *               * Crashed - we definitely detected a crash (whether or not a debugger was attached)
-	 *               * Debugger - the session crashed or shutdown abnormally, but we had a debugger attached at startup, so abnormal termination is much more likely because the user was debugging.
-	 *               * AbnormalShutdown - this happens when we didn't detect a normal shutdown, but none of the above cases is the cause.
-	 * @EventParam Timestamp - the UTC time of the last known time the crashed session was running, within 5 minutes.
-	 *
-	 * @TODO: Debugger should be a completely separate flag, since it's orthogonal to whether we detect a crash or shutdown.
-	 *
-	 * @Comments The editor will only try to check for abnormal terminations if it determines it is a real editor run (not a commandlet or PIE, or editor -game run), and the user has not disabled sending Engine Usage data to Epic via the Editor global preferences.
-	 *
-	 * SessionId can be used to find the actual session associated with this crash in the data.
-	 *
-	 * Use THIS instead of the current event's analytics AppVersion field, since the crash could be from an old version. Could technically be found by searching the data for the SessionId and checking the version there, but sending it here simplifies the backend processing and ensure we have the version here even if we never got the original events.
-	 *
-	 * If multiple versions of the editor or launched, this code will properly track each one and its shutdown status. So during startup, an editor instance may need to fire off several events.
-	 *
-	 * When attributing abnormal terminations to engine versions, be sure to use the EngineVersion associated with this event, and not the AppVersion. AppVersion is for the session that is currently sending the event, not for the session that crashed. That is why EngineVersion is sent separately.
-	 *
-	 * The editor updates this timestamp every 5 minutes, so we should know the time of the crash within 5 minutes. It should technically correlate with the last heartbeat we receive in the data for that session.
-	 */
-	void SendAbnormalShutdownReport(const FSessionRecord& Record)
-	{
-		FGameProjectGenerationModule& GameProjectModule = FModuleManager::LoadModuleChecked<FGameProjectGenerationModule>(TEXT("GameProjectGeneration"));
-		bool bHasCode = GameProjectModule.Get().ProjectHasCodeFiles();
-
-#if PLATFORM_WINDOWS
-		const FString PlatformName(TEXT("Windows"));
-#elif PLATFORM_MAC
-		const FString PlatformName(TEXT("Mac"));
-#elif PLATFORM_LINUX
-		const FString PlatformName(TEXT("Linux"));
-#else
-		const FString PlatformName(TEXT("Unknown"));
-#endif
-
-		FGuid SessionId;
-		FString SessionIdString = Record.SessionId;
-		if (FGuid::Parse(SessionIdString, SessionId))
-		{
-			// convert session guid to one with braces for sending to analytics
-			SessionIdString = SessionId.ToString(EGuidFormats::DigitsWithHyphensInBraces);
-		}
-
-		FString ShutdownTypeString = Record.bCrashed ? EditorAnalyticsDefs::CrashSessionToken :
-			(Record.bIsDebugger ? EditorAnalyticsDefs::DebuggerSessionToken : EditorAnalyticsDefs::AbnormalSessionToken);
-
-		TArray< FAnalyticsEventAttribute > AbnormalShutdownAttributes;
-		AbnormalShutdownAttributes.Add(FAnalyticsEventAttribute(FString("SessionId"), SessionIdString));
-		AbnormalShutdownAttributes.Add(FAnalyticsEventAttribute(FString("EngineVersion"), Record.EngineVersion));
-		AbnormalShutdownAttributes.Add(FAnalyticsEventAttribute(FString("ShutdownType"), ShutdownTypeString));
-		AbnormalShutdownAttributes.Add(FAnalyticsEventAttribute(FString("Timestamp"), Record.Timestamp.ToIso8601()));
-
-		FEditorAnalytics::ReportEvent(TEXT("Editor.AbnormalShutdown"), PlatformName, bHasCode, AbnormalShutdownAttributes);
-	}
-
-	void CreateAndWriteRecordForSession()
-	{
-		FGuid SessionId;
-		if (FGuid::Parse(FEngineAnalytics::GetProvider().GetSessionID(), SessionId))
-		{
-			// convert session guid to one without braces or other chars that might not be suitable for storage
-			CurrentSession.SessionId = SessionId.ToString(EGuidFormats::DigitsWithHyphens);
-		}
-		else
-		{
-			CurrentSession.SessionId = FEngineAnalytics::GetProvider().GetSessionID();
-		}
-
-		CurrentSession.EngineVersion = FEngineVersion::Current().ToString(EVersionComponent::Changelist);
-		CurrentSession.Timestamp = FDateTime::UtcNow();
-		CurrentSession.bCrashed = false;
-		CurrentSession.bIsDebugger = FPlatformMisc::IsDebuggerPresent();
-		CurrentSessionSectionName = GetStoreSectionString(CurrentSession.SessionId);
-
-		FString IsDebuggerString = CurrentSession.bIsDebugger ? EditorAnalyticsDefs::TrueValueString : EditorAnalyticsDefs::FalseValueString;
-
-		FPlatformMisc::SetStoredValue(EditorAnalyticsDefs::StoreId, CurrentSessionSectionName, EditorAnalyticsDefs::CrashStoreKey, EditorAnalyticsDefs::FalseValueString);
-		FPlatformMisc::SetStoredValue(EditorAnalyticsDefs::StoreId, CurrentSessionSectionName, EditorAnalyticsDefs::EngineVersionStoreKey, CurrentSession.EngineVersion);
-		FPlatformMisc::SetStoredValue(EditorAnalyticsDefs::StoreId, CurrentSessionSectionName, EditorAnalyticsDefs::TimestampStoreKey, TimestampToString(CurrentSession.Timestamp));
-		FPlatformMisc::SetStoredValue(EditorAnalyticsDefs::StoreId, CurrentSessionSectionName, EditorAnalyticsDefs::DebuggerStoreKey, IsDebuggerString);
-
-		SessionRecords.Add(CurrentSession);
-	}
-
-	void OnCrashing()
-	{
-		if (!bEditorCrashed && bInitialized)
-		{
-			bEditorCrashed = true;
-
-			FPlatformMisc::SetStoredValue(EditorAnalyticsDefs::StoreId, CurrentSessionSectionName, EditorAnalyticsDefs::CrashStoreKey, EditorAnalyticsDefs::TrueValueString);
-		}
-	}
-
-	static FString TimestampToString(FDateTime InTimestamp)
-	{
-		return LexicalConversion::ToString(InTimestamp.ToUnixTimestamp());
-	}
-
-	static FDateTime StringToTimestamp(FString InString)
-	{
-		int64 TimestampUnix;
-		if (LexicalConversion::TryParseString(TimestampUnix, *InString))
-		{
-			return FDateTime::FromUnixTimestamp(TimestampUnix);
-		}
-		return FDateTime::MinValue();
-	}
-
-	static FString GetStoreSectionString(FString InSuffix)
-	{
-		return FString::Printf(TEXT("%s%s/%s"), *EditorAnalyticsDefs::SessionRecordSectionPrefix, *EditorAnalyticsDefs::EditorSessionsVersionString, *InSuffix);
-	}
-
-private:
-	bool bInitialized;
-	bool bEditorCrashed;
-	FSessionRecord CurrentSession;
-	FString CurrentSessionSectionName;
-	TArray<FSessionRecord> SessionRecords;
-};
-
-
-TSharedPtr<FSessionManager> FEditorAnalytics::SessionManager;
-
-void FEditorAnalytics::InitializeSessionManager()
-{
-	// Create the session manager singleton if we are a normal editor session
-	if (GIsEditor && !IsRunningCommandlet() && !SessionManager.IsValid() )
-	{
-		SessionManager = MakeShareable(new FSessionManager());
-		SessionManager->Initialize();
-	}
-}
-
-void FEditorAnalytics::ShutdownSessionManager()
-{
-	// Destroy the session manager singleton if it exists
-	if (SessionManager.IsValid())
-	{
-		SessionManager->Shutdown();
-		SessionManager.Reset();
-	}
-}
-
-void FEditorAnalytics::SessionHeartbeat()
-{
-	if (SessionManager.IsValid())
-	{
-		SessionManager->Heartbeat();
-	}
-}
 
 void FEditorAnalytics::ReportBuildRequirementsFailure(FString EventName, FString PlatformName, bool bHasCode, int32 Requirements)
 {
diff --git a/Engine/Source/Editor/UnrealEd/Private/Commandlets/ContentCommandlets.cpp b/Engine/Source/Editor/UnrealEd/Private/Commandlets/ContentCommandlets.cpp
--- a/Engine/Source/Editor/UnrealEd/Private/Commandlets/ContentCommandlets.cpp
+++ b/Engine/Source/Editor/UnrealEd/Private/Commandlets/ContentCommandlets.cpp
@@ -745,12 +745,58 @@ bool UResavePackagesCommandlet::PerformPreloadOperations( FLinkerLoad* PackageLi
 
 	return bResult;
 }
+
+
+bool UResavePackagesCommandlet::CheckoutFile(const FString& Filename)
+{
+	if (!bAutoCheckOut)
+	{
+		return true;
+	}
+
+	bool bIsReadOnly = IFileManager::Get().IsReadOnly(*Filename);
+	if (!bIsReadOnly)
+	{
+		return true;
+	}
+
+	//FString FullPath = FPaths::ConvertRelativePathToFull(StreamingLevelPackageFilename);
+	ISourceControlProvider& SourceControlProvider = ISourceControlModule::Get().GetProvider();
+	FSourceControlStatePtr SourceControlState = SourceControlProvider.GetState(*Filename, EStateCacheUsage::ForceUpdate);
+	if (SourceControlState.IsValid())
+	{
+		if (SourceControlState->IsCheckedOutOther())
+		{
+			UE_LOG(LogContentCommandlet, Error, TEXT("[REPORT] %s level is already checked out by someone else, can not submit!"), *Filename);
+		}
+		else if (!SourceControlState->IsCurrent())
+		{
+			UE_LOG(LogContentCommandlet, Error, TEXT("[REPORT] %s is not synced to head, can not submit"), *Filename);
+		}
+		else
+		{
+			if (SourceControlProvider.Execute(ISourceControlOperation::Create<FCheckOut>(), *Filename) == ECommandResult::Succeeded)
+			{
+				UE_LOG(LogContentCommandlet, Display, TEXT("[REPORT] %s Checked out successfully"), *Filename);
+				return true;
+			}
+			else
+			{
+				UE_LOG(LogContentCommandlet, Error, TEXT("[REPORT] %s could not be checked out!"), *Filename);
+			}
+		}
+	}
+	return false;
+}
+
 void UResavePackagesCommandlet::PerformAdditionalOperations(class UWorld* World, bool& bSavePackage)
 {
 	check(World);
 
 	if (bShouldBuildLighting)
 	{
+		bool bShouldProceedWithLightmapRebuild = true;
+
 		static bool bHasLoadedStartupPackages = false;
 		if (bHasLoadedStartupPackages == false)
 		{
@@ -780,32 +826,119 @@ void UResavePackagesCommandlet::PerformAdditionalOperations(class UWorld* World,
 		WorldContext.SetCurrentWorld(World);
 		GWorld = World;
 
-		if (World->StreamingLevels.Num())
+		TArray<FString> SublevelFilenames;
+
+		// if we can't check out the main map or it's not up to date then we can't do the lighting rebuild at all!
+		FString WorldPackageName;
+		if (FPackageName::DoesPackageExist(World->GetOutermost()->GetName(), NULL, &WorldPackageName))
+		{
+			if (CheckoutFile(WorldPackageName))
+			{
+				SublevelFilenames.Add(WorldPackageName);
+			}
+			else
+			{
+				bShouldProceedWithLightmapRebuild = false;
+			}
+		}
+		else
+		{
+			bShouldProceedWithLightmapRebuild = false;
+		}
+		
+
+
+		if (bShouldProceedWithLightmapRebuild)
 		{
 			World->LoadSecondaryLevels(true, NULL);
 
 			for (ULevelStreaming* NextStreamingLevel : World->StreamingLevels)
 			{
+				FString StreamingLevelPackageFilename;
+				const FString StreamingLevelWorldAssetPackageName = NextStreamingLevel->GetWorldAssetPackageName();
+				if (FPackageName::DoesPackageExist(StreamingLevelWorldAssetPackageName, NULL, &StreamingLevelPackageFilename))
+				{
+					// check to see if we need to check this package out
+					if (CheckoutFile(StreamingLevelPackageFilename))
+					{
+						SublevelFilenames.Add(StreamingLevelPackageFilename);
+					}
+					else
+					{
+						bShouldProceedWithLightmapRebuild = false;
+						break;
+					}
+				}
+
 				NextStreamingLevel->bShouldBeVisible = true;
 				NextStreamingLevel->bShouldBeLoaded = true;
 			}
+		}
+
 
+		// If nothing came up that stops us from continuing, then start building lightmass
+		if (bShouldProceedWithLightmapRebuild)
+		{
 			World->FlushLevelStreaming(EFlushLevelStreamingType::Full);
-		}
-		// We need any deferred commands added when loading to be executed before we start building lighting.
-		GEngine->TickDeferredCommands();
 
-		GRedirectCollector.ResolveStringAssetReference();
+			// We need any deferred commands added when loading to be executed before we start building lighting.
+			GEngine->TickDeferredCommands();
 
+			GRedirectCollector.ResolveStringAssetReference();
 
-		FLightingBuildOptions LightingOptions;
-		// Always build on production
-		LightingOptions.QualityLevel = Quality_Production;
+			FLightingBuildOptions LightingOptions;
+			// Always build on production
+			LightingOptions.QualityLevel = Quality_Production;
+
+			GEditor->BuildLighting(LightingOptions);
+			while (GEditor->IsLightingBuildCurrentlyRunning())
+			{
+				GEditor->UpdateBuildLighting();
+			}
+
+			for (ULevelStreaming* NextStreamingLevel : World->StreamingLevels)
+			{
+				FString StreamingLevelPackageFilename;
+				const FString StreamingLevelWorldAssetPackageName = NextStreamingLevel->GetWorldAssetPackageName();
+				if (FPackageName::DoesPackageExist(StreamingLevelWorldAssetPackageName, NULL, &StreamingLevelPackageFilename))
+				{
+					UPackage* SubLevelPackage = NextStreamingLevel->GetLoadedLevel()->GetOutermost();
+					if (!SavePackageHelper(SubLevelPackage, StreamingLevelPackageFilename))
+					{
+						UE_LOG(LogContentCommandlet, Error, TEXT("[REPORT] Failed to save sub level: "), *StreamingLevelPackageFilename);
+					}
+					/*if (GEditor->SavePackage(SubLevelPackage, NULL, , *StreamingLevelPackageFilename, GWarn))
+					{
+							
+					}*/
+				}
+			}
+		}
+		else
+		{
+			UE_LOG(LogContentCommandlet, Error, TEXT("[REPORT] Failed to complete steps necessary to start a lightmass build of "), *World->GetName());
+		}
 
-		GEditor->BuildLighting(LightingOptions);
-		while (GEditor->IsLightingBuildCurrentlyRunning())
+		if ((bShouldProceedWithLightmapRebuild == false)||(bSavePackage == false))
+		{
+			// don't save our parent package
+			bSavePackage = false;
+			
+			//FString FullPath = FPaths::ConvertRelativePathToFull(StreamingLevelPackageFilename);
+			ISourceControlProvider& SourceControlProvider = ISourceControlModule::Get().GetProvider();
+			
+			// revert all our packages
+			for (const auto& SublevelFilename : SublevelFilenames)
+			{
+				SourceControlProvider.Execute(ISourceControlOperation::Create<FRevert>(), *SublevelFilename);
+			}
+		}
+		else
 		{
-			GEditor->UpdateBuildLighting();
+			for(const auto& SublevelFilename : SublevelFilenames)
+			{
+				FilesToSubmit.Add(SublevelFilename);
+			}
 		}
 
 		World->RemoveFromRoot();
diff --git a/Engine/Source/Editor/UnrealEd/Private/Commandlets/GatherTextFromSourceCommandlet.cpp b/Engine/Source/Editor/UnrealEd/Private/Commandlets/GatherTextFromSourceCommandlet.cpp
--- a/Engine/Source/Editor/UnrealEd/Private/Commandlets/GatherTextFromSourceCommandlet.cpp
+++ b/Engine/Source/Editor/UnrealEd/Private/Commandlets/GatherTextFromSourceCommandlet.cpp
@@ -16,6 +16,12 @@ UGatherTextFromSourceCommandlet::UGatherTextFromSourceCommandlet(const FObjectIn
 
 const FString UGatherTextFromSourceCommandlet::FPreProcessorDescriptor::DefineString(TEXT("#define "));
 const FString UGatherTextFromSourceCommandlet::FPreProcessorDescriptor::UndefString(TEXT("#undef "));
+const FString UGatherTextFromSourceCommandlet::FPreProcessorDescriptor::IfString(TEXT("#if "));
+const FString UGatherTextFromSourceCommandlet::FPreProcessorDescriptor::IfDefString(TEXT("#ifdef "));
+const FString UGatherTextFromSourceCommandlet::FPreProcessorDescriptor::ElIfString(TEXT("#elif "));
+const FString UGatherTextFromSourceCommandlet::FPreProcessorDescriptor::ElseString(TEXT("#else"));
+const FString UGatherTextFromSourceCommandlet::FPreProcessorDescriptor::EndIfString(TEXT("#endif"));
+const FString UGatherTextFromSourceCommandlet::FPreProcessorDescriptor::DefinedString(TEXT("defined "));
 const FString UGatherTextFromSourceCommandlet::FPreProcessorDescriptor::LocNamespaceString(TEXT("LOCTEXT_NAMESPACE"));
 const FString UGatherTextFromSourceCommandlet::FPreProcessorDescriptor::LocDefRegionString(TEXT("LOC_DEFINE_REGION"));
 const FString UGatherTextFromSourceCommandlet::FPreProcessorDescriptor::IniNamespaceString(TEXT("["));
@@ -212,6 +218,16 @@ int32 UGatherTextFromSourceCommandlet::Main( const FString& Params )
 
 	Parsables.Add(new FUndefDescriptor());
 
+	Parsables.Add(new FIfDescriptor());
+
+	Parsables.Add(new FIfDefDescriptor());
+
+	Parsables.Add(new FElIfDescriptor());
+
+	Parsables.Add(new FElseDescriptor());
+
+	Parsables.Add(new FEndIfDescriptor());
+
 	Parsables.Add(new FCommandMacroDescriptor());
 
 	// New Localization System with Namespace as literal argument.
@@ -231,6 +247,12 @@ int32 UGatherTextFromSourceCommandlet::Main( const FString& Params )
 	FSourceFileParseContext ParseCtxt;
 	ParseCtxt.ManifestInfo = ManifestInfo;
 
+	// Get whether we should gather editor-only data. Typically only useful for the localization of UE4 itself.
+	if (!GetBoolFromConfig(*SectionName, TEXT("ShouldGatherFromEditorOnlyData"), ParseCtxt.ShouldGatherFromEditorOnlyData, GatherTextConfigPath))
+	{
+		ParseCtxt.ShouldGatherFromEditorOnlyData = false;
+	}
+
 	// Parse all source files for macros and add entries to SourceParsedEntries
 	for ( FString& SourceFile : FilesToProcess)
 	{
@@ -255,6 +277,7 @@ int32 UGatherTextFromSourceCommandlet::Main( const FString& Params )
 		ParseCtxt.WithinStringLiteral = false;
 		ParseCtxt.WithinNamespaceDefine = false;
 		ParseCtxt.WithinStartingLine.Empty();
+		ParseCtxt.FlushMacroStack();
 
 		FString SourceFileText;
 		if (!FFileHelper::LoadFileToString(SourceFileText, *SourceFile))
@@ -549,13 +572,63 @@ bool UGatherTextFromSourceCommandlet::ParseSourceText(const FString& Text, const
 
 bool UGatherTextFromSourceCommandlet::FSourceFileParseContext::AddManifestText( const FString& Token, const FString& InNamespace, const FString& SourceText, const FContext& Context )
 {
-	FString EntryDescription = FString::Printf( TEXT("In %s macro at %s - line %d:%s"),
-		*Token,
-		*Filename, 
-		LineNumber, 
-		*LineText);
-	FLocItem Source( SourceText.ReplaceEscapedCharWithChar() );
-	return ManifestInfo->AddEntry(EntryDescription, InNamespace, Source, Context);
+	const bool bIsEditorOnly = EvaluateMacroStack() == EMacroBlockState::EditorOnly;
+
+	if (!bIsEditorOnly || ShouldGatherFromEditorOnlyData)
+	{
+		FString EntryDescription = FString::Printf( TEXT("In %s macro at %s - line %d:%s"),
+			*Token,
+			*Filename, 
+			LineNumber, 
+			*LineText);
+		FLocItem Source( SourceText.ReplaceEscapedCharWithChar() );
+		return ManifestInfo->AddEntry(EntryDescription, InNamespace, Source, Context);
+	}
+
+	return false;
+}
+
+void UGatherTextFromSourceCommandlet::FSourceFileParseContext::PushMacroBlock( const FString& InBlockCtx )
+{
+	MacroBlockStack.Push(InBlockCtx);
+	CachedMacroBlockState.Reset();
+}
+
+void UGatherTextFromSourceCommandlet::FSourceFileParseContext::PopMacroBlock()
+{
+	if (MacroBlockStack.Num() > 0)
+	{
+		MacroBlockStack.Pop(/*bAllowShrinking*/false);
+		CachedMacroBlockState.Reset();
+	}
+}
+
+void UGatherTextFromSourceCommandlet::FSourceFileParseContext::FlushMacroStack()
+{
+	MacroBlockStack.Reset();
+}
+
+UGatherTextFromSourceCommandlet::EMacroBlockState UGatherTextFromSourceCommandlet::FSourceFileParseContext::EvaluateMacroStack() const
+{
+	if (CachedMacroBlockState.IsSet())
+	{
+		return CachedMacroBlockState.GetValue();
+	}
+
+	static const FString WithEditorString = TEXT("WITH_EDITOR");
+	static const FString WithEditorOnlyDataString = TEXT("WITH_EDITORONLY_DATA");
+
+	CachedMacroBlockState = EMacroBlockState::Normal;
+	for (const FString& BlockCtx : MacroBlockStack)
+	{
+		if (BlockCtx.Equals(WithEditorString, ESearchCase::CaseSensitive) || BlockCtx.Equals(WithEditorOnlyDataString, ESearchCase::CaseSensitive))
+		{
+			CachedMacroBlockState = EMacroBlockState::EditorOnly;
+			break;
+		}
+	}
+
+	return CachedMacroBlockState.GetValue();
 }
 
 void UGatherTextFromSourceCommandlet::FDefineDescriptor::TryParse(const FString& Text, FSourceFileParseContext& Context) const
@@ -640,6 +713,81 @@ void UGatherTextFromSourceCommandlet::FUndefDescriptor::TryParse(const FString&
 	}
 }
 
+void UGatherTextFromSourceCommandlet::FIfDescriptor::TryParse(const FString& Text, FSourceFileParseContext& Context) const
+{
+	// Attempt to parse something of the format
+	// #if <defname>
+
+	FString RemainingText = Text.RightChop(GetToken().Len()).Trim();
+
+	if (!Context.WithinBlockComment && !Context.WithinLineComment && !Context.WithinStringLiteral)
+	{
+		// Handle "#if defined <defname>"
+		if (RemainingText.StartsWith(DefinedString, ESearchCase::CaseSensitive))
+		{
+			RemainingText = Text.RightChop(DefinedString.Len()).Trim();
+		}
+
+		Context.PushMacroBlock(RemainingText);
+	}
+}
+
+void UGatherTextFromSourceCommandlet::FIfDefDescriptor::TryParse(const FString& Text, FSourceFileParseContext& Context) const
+{
+	// Attempt to parse something of the format
+	// #ifdef <defname>
+
+	FString RemainingText = Text.RightChop(GetToken().Len()).Trim();
+
+	if (!Context.WithinBlockComment && !Context.WithinLineComment && !Context.WithinStringLiteral)
+	{
+		Context.PushMacroBlock(RemainingText);
+	}
+}
+
+void UGatherTextFromSourceCommandlet::FElIfDescriptor::TryParse(const FString& Text, FSourceFileParseContext& Context) const
+{
+	// Attempt to parse something of the format
+	// #elif <defname>
+
+	FString RemainingText = Text.RightChop(GetToken().Len()).Trim();
+
+	if (!Context.WithinBlockComment && !Context.WithinLineComment && !Context.WithinStringLiteral)
+	{
+		// Handle "#elif defined <defname>"
+		if (RemainingText.StartsWith(DefinedString, ESearchCase::CaseSensitive))
+		{
+			RemainingText = Text.RightChop(DefinedString.Len()).Trim();
+		}
+
+		Context.PopMacroBlock(); // Pop the current #if or #ifdef state
+		Context.PushMacroBlock(RemainingText);
+	}
+}
+
+void UGatherTextFromSourceCommandlet::FElseDescriptor::TryParse(const FString& Text, FSourceFileParseContext& Context) const
+{
+	// Attempt to parse something of the format
+	// #else
+
+	if (!Context.WithinBlockComment && !Context.WithinLineComment && !Context.WithinStringLiteral)
+	{
+		Context.PopMacroBlock(); // Pop the current #if or #ifdef state
+		Context.PushMacroBlock(FString());
+	}
+}
+
+void UGatherTextFromSourceCommandlet::FEndIfDescriptor::TryParse(const FString& Text, FSourceFileParseContext& Context) const
+{
+	// Attempt to parse something of the format
+	// #endif
+
+	if (!Context.WithinBlockComment && !Context.WithinLineComment && !Context.WithinStringLiteral)
+	{
+		Context.PopMacroBlock(); // Pop the current #if or #ifdef state
+	}
+}
+
 bool UGatherTextFromSourceCommandlet::FMacroDescriptor::ParseArgsFromMacro(const FString& Text, TArray<FString>& Args, FSourceFileParseContext& Context) const
 {
 	// Attempt to parse something of the format
diff --git a/Engine/Source/Editor/UnrealEd/Private/EditorActor.cpp b/Engine/Source/Editor/UnrealEd/Private/EditorActor.cpp
--- a/Engine/Source/Editor/UnrealEd/Private/EditorActor.cpp
+++ b/Engine/Source/Editor/UnrealEd/Private/EditorActor.cpp
@@ -29,6 +29,7 @@
 #include "Engine/StaticMeshActor.h"
 #include "ComponentEditorUtils.h"
 #include "LevelEditor.h"
+#include "Engine/LODActor.h"
 
 #define LOCTEXT_NAMESPACE "UnrealEd.EditorActor"
 
@@ -704,6 +705,8 @@ bool UUnrealEdEngine::edactDeleteSelected( UWorld* InWorld, bool bVerifyDeletion
 
 	bool	bRequestedDeleteAllByLevel = false;
 	bool	bRequestedDeleteAllByActor = false;
+	//  @todo remove me : this is to be removed once HLOD remove actor delegate gets moved to safer place
+	bool	bRequestedDeleteAllByLODActor = false;
 	EAppMsgType::Type MessageType = ActorsToDelete.Num() > 1 ? EAppMsgType::YesNoYesAllNoAll : EAppMsgType::YesNo;
 	int32		DeleteCount = 0;
 
@@ -722,22 +725,41 @@ bool UUnrealEdEngine::edactDeleteSelected( UWorld* InWorld, bool bVerifyDeletion
 
 		bool bReferencedByLevelScript = (NULL != LSB && FBlueprintEditorUtils::FindNumReferencesToActorFromLevelScript(LSB, Actor) > 0);
 		bool bReferencedByActor = false;
+		bool bReferencedByLODActor = false;
 		for (AActor* ReferencingActor : ReferencingActors)
 		{
-			if (ReferencingActor->ParentComponentActor.Get() != Actor)
+			if (ReferencingActor->IsA(ALODActor::StaticClass()))
+			{
+				bReferencedByLODActor = true;
+				break;
+			}
+			// If the referencing actor is a child actor that is referencing us, do not treat it
+			// as referencing for the purposes of warning about deletion
+			UChildActorComponent* ParentComponent = ReferencingActor->GetParentComponent();
+			if (ParentComponent == nullptr || ParentComponent->GetOwner() != Actor)
 			{
 				bReferencedByActor = true;
 				break;
 			}
 		}
 
-		if ( bReferencedByLevelScript || bReferencedByActor )
+		if (bReferencedByLevelScript || bReferencedByActor || bReferencedByLODActor)
 		{
 			if (( bReferencedByLevelScript && !bRequestedDeleteAllByLevel ) ||
-				( bReferencedByActor && !bRequestedDeleteAllByActor ))
+				( bReferencedByActor && !bRequestedDeleteAllByActor ) ||
+				(bReferencedByLODActor && !bRequestedDeleteAllByLODActor) )
 			{
 				FText ConfirmDelete;
-				if ( bReferencedByLevelScript && bReferencedByActor )
+
+				// check LODActor outside of normal check
+				// you might like to know other actor is referencing it
+				if (bReferencedByLODActor)
+				{
+					ConfirmDelete = FText::Format(LOCTEXT("ConfirmDeleteActorReferencedByHLOD",
+						"Actor {0} is referenced by LODActor, do you really want to delete it?"),
+						FText::FromString(Actor->GetName()));
+				}
+				else if (bReferencedByLevelScript && bReferencedByActor)
 				{
 					ConfirmDelete = FText::Format(LOCTEXT( "ConfirmDeleteActorReferenceByScriptAndActor",
 														   "Actor {0} is referenced by the level blueprint and another Actor, do you really want to delete it?"),
@@ -761,6 +783,7 @@ bool UUnrealEdEngine::edactDeleteSelected( UWorld* InWorld, bool bVerifyDeletion
 				{
 					bRequestedDeleteAllByLevel = bReferencedByLevelScript;
 					bRequestedDeleteAllByActor = bReferencedByActor;
+					bRequestedDeleteAllByLODActor = bReferencedByLODActor;
 				}
 				else if ( Result == EAppReturnType::NoAll )
 				{
@@ -783,6 +806,18 @@ bool UUnrealEdEngine::edactDeleteSelected( UWorld* InWorld, bool bVerifyDeletion
 					ReferencingActors[ReferencingActorIndex]->Modify();
 				}
 			}
+			if (bReferencedByLODActor)
+			{
+				for (int32 ReferencingActorIndex = 0; ReferencingActorIndex < ReferencingActors.Num(); ReferencingActorIndex++)
+				{
+					ALODActor* LODActor = Cast<ALODActor>(ReferencingActors[ReferencingActorIndex]);
+					// it's possible other actor is referencing this
+					if (LODActor)
+					{
+						LODActor->RemoveSubActor(Actor);
+					}
+				}
+			}
 		}
 	
 		bool bRebuildNavigation = false;
diff --git a/Engine/Source/Editor/UnrealEd/Private/EditorEngine.cpp b/Engine/Source/Editor/UnrealEd/Private/EditorEngine.cpp
--- a/Engine/Source/Editor/UnrealEd/Private/EditorEngine.cpp
+++ b/Engine/Source/Editor/UnrealEd/Private/EditorEngine.cpp
@@ -992,6 +992,8 @@ void UEditorEngine::Tick( float DeltaSeconds, bool bIdleMode )
 		StaticTick(DeltaSeconds, !!GAsyncLoadingUseFullTimeLimit, GAsyncLoadingTimeLimit / 1000.f);
 	}
 
+	FEngineAnalytics::Tick(DeltaSeconds);
+
 	// Look for realtime flags.
 	bool IsRealtime = false;
 
diff --git a/Engine/Source/Editor/UnrealEd/Private/EditorExporters.cpp b/Engine/Source/Editor/UnrealEd/Private/EditorExporters.cpp
--- a/Engine/Source/Editor/UnrealEd/Private/EditorExporters.cpp
+++ b/Engine/Source/Editor/UnrealEd/Private/EditorExporters.cpp
@@ -1943,7 +1943,7 @@ namespace MaterialExportUtils
 
 	UMaterial* CreateMaterial(const FFlattenMaterial& InFlattenMaterial, UPackage* Outer, const FString& BaseName, EObjectFlags Flags, TArray<UObject*>& OutGeneratedAssets)
 	{
-		return FMaterialUtilities::CreateMaterial(InFlattenMaterial, Outer, BaseName, Flags, OutGeneratedAssets);
+		return FMaterialUtilities::CreateMaterial(InFlattenMaterial, Outer, BaseName, Flags, FMaterialProxySettings(), OutGeneratedAssets);
 	}
 
 	bool ExportBaseColor(ULandscapeComponent* LandscapeComponent, int32 TextureSize, TArray<FColor>& OutSamples)
diff --git a/Engine/Source/Editor/UnrealEd/Private/EditorSelectUtils.cpp b/Engine/Source/Editor/UnrealEd/Private/EditorSelectUtils.cpp
--- a/Engine/Source/Editor/UnrealEd/Private/EditorSelectUtils.cpp
+++ b/Engine/Source/Editor/UnrealEd/Private/EditorSelectUtils.cpp
@@ -533,9 +533,9 @@ void UUnrealEdEngine::SelectActor(AActor* Actor, bool bInSelected, bool bNotify,
 		if(bInSelected)
 		{
 			// If trying to select an Actor spawned by a ChildACtorComponent, instead select Actor that spawned us
-			if(Actor->ParentComponentActor.IsValid())
+			if (UChildActorComponent* ParentComponent = Actor->GetParentComponent())
 			{
-				Actor = Actor->ParentComponentActor.Get();
+				Actor = ParentComponent->GetOwner();
 			}
 		}
 
diff --git a/Engine/Source/Editor/UnrealEd/Private/Kismet2/DebuggerCommands.cpp b/Engine/Source/Editor/UnrealEd/Private/Kismet2/DebuggerCommands.cpp
--- a/Engine/Source/Editor/UnrealEd/Private/Kismet2/DebuggerCommands.cpp
+++ b/Engine/Source/Editor/UnrealEd/Private/Kismet2/DebuggerCommands.cpp
@@ -231,7 +231,7 @@ void FPlayWorldCommands::RegisterCommands()
 
 	// PIE
 	UI_COMMAND( RepeatLastPlay, "Play", "Launches a game preview session in the same mode as the last game preview session launched from the Game Preview Modes dropdown next to the Play button on the level editor toolbar", EUserInterfaceActionType::Button, FInputChord( EKeys::P, EModifierKey::Alt ) )
-	UI_COMMAND( PlayInViewport, "Selected Viewport (PIE)", "Play this level in the active level editor viewport", EUserInterfaceActionType::Check, FInputChord() );
+	UI_COMMAND( PlayInViewport, "Selected Viewport", "Play this level in the active level editor viewport", EUserInterfaceActionType::Check, FInputChord() );
 	UI_COMMAND( PlayInEditorFloating, "New Editor Window (PIE)", "Play this level in a new window", EUserInterfaceActionType::Check, FInputChord() );
 	UI_COMMAND( PlayInVR, "VR Preview", "Play this level in VR", EUserInterfaceActionType::Check, FInputChord() );
 	UI_COMMAND( PlayInMobilePreview, "Mobile Preview (PIE)", "Play this level as a mobile device preview (runs in its own process)", EUserInterfaceActionType::Check, FInputChord() );
diff --git a/Engine/Source/Editor/UnrealEd/Private/Kismet2/KismetReinstanceUtilities.cpp b/Engine/Source/Editor/UnrealEd/Private/Kismet2/KismetReinstanceUtilities.cpp
--- a/Engine/Source/Editor/UnrealEd/Private/Kismet2/KismetReinstanceUtilities.cpp
+++ b/Engine/Source/Editor/UnrealEd/Private/Kismet2/KismetReinstanceUtilities.cpp
@@ -410,14 +410,35 @@ public:
 		}
 
 		const bool bIsAnimInstance = ClassToReinstance->IsChildOf<UAnimInstance>();
+		//UAnimBlueprintGeneratedClass* AnimClass = Cast<UAnimBlueprintGeneratedClass>(ClassToReinstance);
 		if (bIsAnimInstance)
 		{
-			for (auto Obj : ObjectsToFinalize)
+			for(auto Obj : ObjectsToFinalize)
 			{
-				// Initialising the anim instance isn't enough to correctly set up the skeletal mesh again in a
-				// paused world, need to initialise the skeletal mesh component that contains the anim instance.
-				if (USkeletalMeshComponent* SkelComponent = Cast<USkeletalMeshComponent>(Obj->GetOuter()))
+				if(USkeletalMeshComponent* SkelComponent = Cast<USkeletalMeshComponent>(Obj->GetOuter()))
 				{
+					// This snippet catches all of the exposed value handlers that will have invalid UFunctions
+					// and clears the init flag so they will be reinitialized on the next call to InitAnim.
+					// Unknown whether there are other unreachable properties so currently clearing the anim
+					// instance below
+					// #TODO investigate reinstancing anim blueprints to correctly catch all deep references
+
+					//UAnimInstance* ActiveInstance = SkelComponent->GetAnimInstance();
+					//if(AnimClass && ActiveInstance)
+					//{
+					//	for(UStructProperty* NodeProp : AnimClass->AnimNodeProperties)
+					//	{
+					//		// Guaranteed to have only FAnimNode_Base pointers added during compilation
+					//		FAnimNode_Base* AnimNode = NodeProp->ContainerPtrToValuePtr<FAnimNode_Base>(ActiveInstance);
+					//
+					//		AnimNode->EvaluateGraphExposedInputs.bInitialized = false;
+					//	}
+					//}
+
+					// Clear out the script instance on the component to force a rebuild during initialization.
+					// This is necessary to correctly reinitialize certain properties that still reference the 
+					// old class as they are unreachable during reinstancing.
+					SkelComponent->AnimScriptInstance = nullptr;
 					SkelComponent->InitAnim(true);
 				}
 			}
diff --git a/Engine/Source/Editor/UnrealEd/Private/LODCluster.cpp b/Engine/Source/Editor/UnrealEd/Private/LODCluster.cpp
--- a/Engine/Source/Editor/UnrealEd/Private/LODCluster.cpp
+++ b/Engine/Source/Editor/UnrealEd/Private/LODCluster.cpp
@@ -315,8 +315,11 @@ ALODActor* FLODCluster::BuildActor(ULevel* InLevel, const int32 LODIdx, const bo
 
 			// Mark dirty according to whether or not this is a preview build
 			NewActor->SetIsDirty(!bCreateMeshes);
-
-			FHierarchicalLODUtilities::BuildStaticMeshForLODActor(NewActor, AssetsOuter, LODSetup, LODIdx);
+			
+			if (bCreateMeshes)
+			{
+				FHierarchicalLODUtilities::BuildStaticMeshForLODActor(NewActor, AssetsOuter, LODSetup, LODIdx);
+			}
 
 			return NewActor;
 		}
diff --git a/Engine/Source/Editor/UnrealEd/Private/StaticLightingSystem/StaticLightingSystem.cpp b/Engine/Source/Editor/UnrealEd/Private/StaticLightingSystem/StaticLightingSystem.cpp
--- a/Engine/Source/Editor/UnrealEd/Private/StaticLightingSystem/StaticLightingSystem.cpp
+++ b/Engine/Source/Editor/UnrealEd/Private/StaticLightingSystem/StaticLightingSystem.cpp
@@ -42,6 +42,7 @@ DEFINE_LOG_CATEGORY(LogStaticLightingSystem);
 #include "Engine/Selection.h"
 #include "Components/SkyLightComponent.h"
 #include "Components/LightmassPortalComponent.h"
+#include "Runtime/CoreUObject/Public/Misc/UObjectToken.h"
 
 #define LOCTEXT_NAMESPACE "StaticLightingSystem"
 
@@ -776,13 +777,21 @@ void UpdateStaticLightingHLODTreeIndices(TMultiMap<AActor*, FStaticLightingMesh*
 
 			for (FStaticLightingMesh* SubActorMesh : SubActorMeshes)
 			{
-				checkf(SubActorMesh->HLODTreeIndex == 0, TEXT("ERROR: HLODTreeIndex != 0 for '%s' while processing LOD actor '%s'"), *SubActorMesh->Component->GetPathName(), *LODActor->GetName());
+				if (SubActorMesh->HLODTreeIndex == 0)
 				{
 					SubActorMesh->HLODTreeIndex = HLODTreeIndex;
 					SubActorMesh->HLODChildStartIndex = HLODLeafIndex;
 					SubActorMesh->HLODChildEndIndex = HLODLeafIndex;
 					++HLODLeafIndex;
 				}
+				else
+				{
+					// Output error to message log containing tokens to the problematic objects
+					FMessageLog("LightingResults").Warning()
+						->AddToken(FUObjectToken::Create(SubActorMesh->Component->GetOwner()))
+						->AddToken(FTextToken::Create(LOCTEXT("LightmassError_InvalidHLODTreeIndex", "will not be correctly lit since it is part of another Hierarchical LOD cluster besides ")))
+						->AddToken(FUObjectToken::Create(LODActor));
+				}
 			}
 		}
 	}
diff --git a/Engine/Source/Editor/UnrealEd/Private/ThumbnailHelpers.cpp b/Engine/Source/Editor/UnrealEd/Private/ThumbnailHelpers.cpp
--- a/Engine/Source/Editor/UnrealEd/Private/ThumbnailHelpers.cpp
+++ b/Engine/Source/Editor/UnrealEd/Private/ThumbnailHelpers.cpp
@@ -595,6 +595,11 @@ void FStaticMeshThumbnailScene::GetViewMatrixParameters(const float InFOVDegrees
 FAnimationThumbnailScene
 ***************************************************************
 */
+AAnimationThumbnailSkeletalMeshActor::AAnimationThumbnailSkeletalMeshActor(const FObjectInitializer& ObjectInitializer)
+	: Super(ObjectInitializer.SetDefaultSubobjectClass<UDebugSkelMeshComponent>(TEXT("SkeletalMeshComponent0")))
+{
+
+}
 
 FAnimationSequenceThumbnailScene::FAnimationSequenceThumbnailScene()
 	: FThumbnailPreviewScene()
@@ -606,7 +611,7 @@ FAnimationSequenceThumbnailScene::FAnimationSequenceThumbnailScene()
 	SpawnInfo.SpawnCollisionHandlingOverride = ESpawnActorCollisionHandlingMethod::AlwaysSpawn;
 	SpawnInfo.bNoFail = true;
 	SpawnInfo.ObjectFlags = RF_Transient;
-	PreviewActor = GetWorld()->SpawnActor<ASkeletalMeshActor>(SpawnInfo);
+	PreviewActor = GetWorld()->SpawnActor<AAnimationThumbnailSkeletalMeshActor>(SpawnInfo);
 
 	PreviewActor->SetActorEnableCollision(false);
 }
@@ -634,9 +639,12 @@ bool FAnimationSequenceThumbnailScene::SetAnimation(UAnimSequenceBase* InAnimati
 				// Handle posing the mesh at the middle of the animation
 				const float AnimPosition = InAnimation->SequenceLength / 2.f;
 
-				PreviewActor->GetSkeletalMeshComponent()->PlayAnimation(InAnimation,false);
-				PreviewActor->GetSkeletalMeshComponent()->Stop();
-				PreviewActor->GetSkeletalMeshComponent()->SetPosition(AnimPosition, false);
+				UDebugSkelMeshComponent* MeshComponent = CastChecked<UDebugSkelMeshComponent>(PreviewActor->GetSkeletalMeshComponent());
+
+				MeshComponent->EnablePreview(true, InAnimation, nullptr);
+				MeshComponent->Play(false);
+				MeshComponent->Stop();
+				MeshComponent->SetPosition(AnimPosition, false);
 
 				UAnimSingleNodeInstance* SingleNodeInstance = PreviewActor->GetSkeletalMeshComponent()->GetSingleNodeInstance();
 				if (SingleNodeInstance)
diff --git a/Engine/Source/Editor/UnrealEd/Private/UnrealEdMisc.cpp b/Engine/Source/Editor/UnrealEd/Private/UnrealEdMisc.cpp
--- a/Engine/Source/Editor/UnrealEd/Private/UnrealEdMisc.cpp
+++ b/Engine/Source/Editor/UnrealEd/Private/UnrealEdMisc.cpp
@@ -443,8 +443,6 @@ void FUnrealEdMisc::OnInit()
 
 	SlowTask.EnterProgressFrame(10);
 
-	FEditorAnalytics::InitializeSessionManager();
-
 	// Send Project Analytics
 	InitEngineAnalytics();
 	
@@ -574,8 +572,6 @@ void FUnrealEdMisc::EditorAnalyticsHeartbeat()
 	FEngineAnalytics::GetProvider().RecordEvent(TEXT("Editor.Usage.Heartbeat"), Attributes);
 	
 	LastHeartbeatTime = FPlatformTime::Seconds();
-
-	FEditorAnalytics::SessionHeartbeat();
 }
 
 void FUnrealEdMisc::TickAssetAnalytics()
@@ -769,8 +765,6 @@ void FUnrealEdMisc::OnExit()
 		FEditorViewportStats::SendUsageData();
 	}
 
-	FEditorAnalytics::ShutdownSessionManager();
-
 	FInputBindingManager::Get().UnregisterUserDefinedChordChanged(OnUserDefinedChordChangedDelegateHandle);
 	FMessageLog::OnMessageSelectionChanged().Unbind();
 	FUObjectToken::DefaultOnMessageTokenActivated().Unbind();
diff --git a/Engine/Source/Editor/UnrealEd/Public/EditorAnalytics.h b/Engine/Source/Editor/UnrealEd/Public/EditorAnalytics.h
--- a/Engine/Source/Editor/UnrealEd/Public/EditorAnalytics.h
+++ b/Engine/Source/Editor/UnrealEd/Public/EditorAnalytics.h
@@ -67,21 +67,6 @@ namespace EAnalyticsErrorCodes
 class FEditorAnalytics
 {
 public:
-	/**
-	 * Starts the session manager if needed
-	 */
-	UNREALED_API static void InitializeSessionManager();
-
-	/**
-	 * Stops the session manager if needed
-	 */
-	UNREALED_API static void ShutdownSessionManager();
-
-	/**
-	 * Update the editor session's heartbeat time
-	 */
-	UNREALED_API static void SessionHeartbeat();
-
 	/**
 	* Reports an event to the analytics system if it is enabled
 	*
@@ -127,5 +112,5 @@ public:
 	UNREALED_API static bool ShouldElevateMessageThroughDialog(const int32 ErrorCode);
 
 private:
-	static TSharedPtr<class FSessionManager> SessionManager;
+	static TSharedPtr<class FEngineSessionManager> SessionManager;
 };
diff --git a/Engine/Source/Editor/UnrealEd/Public/ThumbnailHelpers.h b/Engine/Source/Editor/UnrealEd/Public/ThumbnailHelpers.h
--- a/Engine/Source/Editor/UnrealEd/Public/ThumbnailHelpers.h
+++ b/Engine/Source/Editor/UnrealEd/Public/ThumbnailHelpers.h
@@ -4,6 +4,8 @@
 #pragma once
 
 #include "PreviewScene.h"
+#include "Animation/SkeletalMeshActor.h"
+#include "ThumbnailHelpers.generated.h"
 
 class FSceneViewFamily;
 class USCS_Node;
@@ -131,6 +133,12 @@ private:
 	AStaticMeshActor* PreviewActor;
 };
 
+UCLASS(ClassGroup = ISkeletalMeshes, ComponentWrapperClass, ConversionRoot, meta = (ChildCanTick))
+class AAnimationThumbnailSkeletalMeshActor : public ASkeletalMeshActor
+{
+	GENERATED_UCLASS_BODY()
+};
+
 class UNREALED_API FAnimationSequenceThumbnailScene : public FThumbnailPreviewScene
 {
 public:
@@ -149,7 +157,7 @@ protected:
 
 private:
 	/** The skeletal mesh actor used to display all animation thumbnails */
-	class ASkeletalMeshActor* PreviewActor;
+	AAnimationThumbnailSkeletalMeshActor* PreviewActor;
 
 	/** Animation we are generating the thumbnail for */
 	class UAnimSequenceBase* PreviewAnimation;
diff --git a/Engine/Source/Editor/WorldBrowser/Private/Tiles/WorldTileCollectionModel.cpp b/Engine/Source/Editor/WorldBrowser/Private/Tiles/WorldTileCollectionModel.cpp
--- a/Engine/Source/Editor/WorldBrowser/Private/Tiles/WorldTileCollectionModel.cpp
+++ b/Engine/Source/Editor/WorldBrowser/Private/Tiles/WorldTileCollectionModel.cpp
@@ -2043,29 +2043,7 @@ bool FWorldTileCollectionModel::GenerateLODLevels(FLevelModelList InLevelList, i
 			LandscapeFlattenMaterial.SpecularSize = SimplificationDetails.LandscapeMaterialSettings.bSpecularMap ? SimplificationDetails.LandscapeMaterialSettings.TextureSize : FIntPoint::ZeroValue;
 			
 			FMaterialUtilities::ExportLandscapeMaterial(Landscape, PrimitivesToHide, LandscapeFlattenMaterial);
-
-			// Fill landscape material constants
-			{
-				if (LandscapeFlattenMaterial.MetallicSamples.Num() == 0)
-				{
-					LandscapeFlattenMaterial.MetallicSize = FIntPoint(1, 1);
-					LandscapeFlattenMaterial.MetallicSamples.SetNum(1);
-					LandscapeFlattenMaterial.MetallicSamples[0].DWColor() = *(uint32*)(&SimplificationDetails.LandscapeMaterialSettings.MetallicConstant);
-				}
-				if (LandscapeFlattenMaterial.RoughnessSamples.Num() == 0)
-				{
-					LandscapeFlattenMaterial.RoughnessSize = FIntPoint(1, 1);
-					LandscapeFlattenMaterial.RoughnessSamples.SetNum(1);
-					LandscapeFlattenMaterial.RoughnessSamples[0].DWColor() = *(uint32*)(&SimplificationDetails.LandscapeMaterialSettings.RoughnessConstant);
-				}
-				if (LandscapeFlattenMaterial.SpecularSamples.Num() == 0)
-				{
-					LandscapeFlattenMaterial.SpecularSize = FIntPoint(1, 1);
-					LandscapeFlattenMaterial.SpecularSamples.SetNum(1);
-					LandscapeFlattenMaterial.SpecularSamples[0].DWColor() = *(uint32*)(&SimplificationDetails.LandscapeMaterialSettings.SpecularConstant);
-				}
-			}
-		
+					
 			if (SimplificationDetails.bBakeGrassToLandscape)
 			{
 				Landscape->FlushGrassComponents(); // wipe this and let it fix itself later
@@ -2073,7 +2051,7 @@ bool FWorldTileCollectionModel::GenerateLODLevels(FLevelModelList InLevelList, i
 			FString LandscapeBaseAssetName = FString::Printf(TEXT("%s_LOD%d"), *Landscape->GetName(), TargetLODIndex + 1);
 			// Construct landscape material
 			UMaterial* StaticLandscapeMaterial = FMaterialUtilities::CreateMaterial(
-				LandscapeFlattenMaterial, AssetsOuter, *(AssetsPath + LandscapeBaseAssetName), RF_Public|RF_Standalone, GeneratedAssets);
+				LandscapeFlattenMaterial, AssetsOuter, *(AssetsPath + LandscapeBaseAssetName), RF_Public | RF_Standalone, SimplificationDetails.LandscapeMaterialSettings, GeneratedAssets);
 			// Currently landscape exports world space normal map
 			StaticLandscapeMaterial->bTangentSpaceNormal = false;
 			StaticLandscapeMaterial->PostEditChange();
diff --git a/Engine/Source/Programs/AutomationTool/Android/AndroidPlatform.Automation.cs b/Engine/Source/Programs/AutomationTool/Android/AndroidPlatform.Automation.cs
--- a/Engine/Source/Programs/AutomationTool/Android/AndroidPlatform.Automation.cs
+++ b/Engine/Source/Programs/AutomationTool/Android/AndroidPlatform.Automation.cs
@@ -493,16 +493,16 @@ public class AndroidPlatform : Platform
 		string DeviceName = Params.Device.Replace(":", "_");
 
 		// Try retrieving the UFS files manifest files from the device
-		string UFSManifestFileName = CombinePaths(SC.StageDirectory, DeploymentContext.UFSDeployedManifestFileName + "_" + DeviceName);
-		ProcessResult UFSResult = RunAdbCommand(Params, " pull " + RemoteDir + "/" + DeploymentContext.UFSDeployedManifestFileName + " \"" + UFSManifestFileName + "\"", null, ERunOptions.AppMustExist);
+		string UFSManifestFileName = CombinePaths(SC.StageDirectory, SC.UFSDeployedManifestFileName + "_" + DeviceName);
+		ProcessResult UFSResult = RunAdbCommand(Params, " pull " + RemoteDir + "/" + SC.UFSDeployedManifestFileName + " \"" + UFSManifestFileName + "\"", null, ERunOptions.AppMustExist);
 		if (!UFSResult.Output.Contains("bytes"))
 		{
 			return false;
 		}
 
 		// Try retrieving the non UFS files manifest files from the device
-		string NonUFSManifestFileName = CombinePaths(SC.StageDirectory, DeploymentContext.NonUFSDeployedManifestFileName + "_" + DeviceName);
-		ProcessResult NonUFSResult = RunAdbCommand(Params, " pull " + RemoteDir + "/" + DeploymentContext.NonUFSDeployedManifestFileName + " \"" + NonUFSManifestFileName + "\"", null, ERunOptions.AppMustExist);
+		string NonUFSManifestFileName = CombinePaths(SC.StageDirectory, SC.NonUFSDeployedManifestFileName + "_" + DeviceName);
+		ProcessResult NonUFSResult = RunAdbCommand(Params, " pull " + RemoteDir + "/" + SC.NonUFSDeployedManifestFileName + " \"" + NonUFSManifestFileName + "\"", null, ERunOptions.AppMustExist);
 		if (!NonUFSResult.Output.Contains("bytes"))
 		{
 			// Did not retrieve both so delete one we did retrieve
diff --git a/Engine/Source/Programs/AutomationTool/AutomationUtils/CommandletUtils.cs b/Engine/Source/Programs/AutomationTool/AutomationUtils/CommandletUtils.cs
--- a/Engine/Source/Programs/AutomationTool/AutomationUtils/CommandletUtils.cs
+++ b/Engine/Source/Programs/AutomationTool/AutomationUtils/CommandletUtils.cs
@@ -115,7 +115,7 @@ namespace AutomationTool
 				MapsToRebuildLighting.Trim();
 			}
 
-			RunCommandlet(ProjectName, UE4Exe, "ResavePackages", String.Format("-buildlighting -MapsOnly -ProjectOnly -AllowCommandletRendering -AutoCheckoutPackages {0} {1}", MapsToRebuildLighting, Parameters));
+			RunCommandlet(ProjectName, UE4Exe, "ResavePackages", String.Format("-buildlighting -MapsOnly -ProjectOnly -AllowCommandletRendering {0} {1}", MapsToRebuildLighting, Parameters));
 		}
 
         /// <summary>
diff --git a/Engine/Source/Programs/AutomationTool/AutomationUtils/DeploymentContext.cs b/Engine/Source/Programs/AutomationTool/AutomationUtils/DeploymentContext.cs
--- a/Engine/Source/Programs/AutomationTool/AutomationUtils/DeploymentContext.cs
+++ b/Engine/Source/Programs/AutomationTool/AutomationUtils/DeploymentContext.cs
@@ -190,11 +190,11 @@ public class DeploymentContext //: ProjectParams
 	/// <summary>
 	/// Filename for the manifest of files currently deployed on a device.
 	/// </summary>
-	static public readonly string UFSDeployedManifestFileName		= "Manifest_UFSFiles.txt";
-	static public readonly string NonUFSDeployedManifestFileName	= "Manifest_NonUFSFiles.txt";
+	static public readonly string sUFSDeployedManifestFileName		= "Manifest_UFSFiles.txt";
+	static public readonly string sNonUFSDeployedManifestFileName	= "Manifest_NonUFSFiles.txt";
 
 	
-
+	
 
 	/// <summary>
 	/// The client connects to dedicated server to get data
@@ -755,4 +755,20 @@ public class DeploymentContext //: ProjectParams
 	{
 		return Path.Combine(StageDirectory, NonUFSDeployObsoleteFileName);
 	}
+
+	public string UFSDeployedManifestFileName
+	{
+		get
+		{
+			return "Manifest_UFSFiles_" + StageTargetPlatform.PlatformType.ToString() + ".txt";
+		}
+	}
+
+	public string NonUFSDeployedManifestFileName
+	{
+		get
+		{
+			return "Manifest_NonUFSFiles_" + StageTargetPlatform.PlatformType.ToString() + ".txt";
+		}
+	}
 }
diff --git a/Engine/Source/Programs/AutomationTool/GUBP/LegacyNodes.cs b/Engine/Source/Programs/AutomationTool/GUBP/LegacyNodes.cs
--- a/Engine/Source/Programs/AutomationTool/GUBP/LegacyNodes.cs
+++ b/Engine/Source/Programs/AutomationTool/GUBP/LegacyNodes.cs
@@ -2648,9 +2648,9 @@ partial class GUBP
 		{
 			var UEProjectDirectory = GetUEProjectDirectory();
 			var UEProjectName = GetUEProjectName();
-			var OneSkyConfigName = GetOneSkyConfigName();
-			var OneSkyProjectGroupName = GetOneSkyProjectGroupName();
-			var OneSkyProjectNames = GetOneSkyProjectNames();
+			var LocalizationProjectNames = GetLocalizationProjectNames();
+			var LocalizationProviderName = GetLocalizationProviderName();
+			var LocalizationProviderArgs = GetLocalizationProviderArgs();
 
 			// Build the correct command line arguments.
 			var CommandLineArguments = "";
@@ -2665,24 +2665,24 @@ partial class GUBP
 				CommandLineArguments += " -UEProjectName=\"" + UEProjectName + "\"";
 			}
 
-			if (!String.IsNullOrEmpty(OneSkyConfigName))
+			if (!String.IsNullOrEmpty(LocalizationProjectNames))
 			{
-				CommandLineArguments += " -OneSkyConfigName=\"" + OneSkyConfigName + "\"";
+				CommandLineArguments += " -LocalizationProjectNames=\"" + LocalizationProjectNames + "\"";
 			}
 
-			if (!String.IsNullOrEmpty(OneSkyProjectGroupName))
+			if (!String.IsNullOrEmpty(LocalizationBranchSuffix))
 			{
-				CommandLineArguments += " -OneSkyProjectGroupName=\"" + OneSkyProjectGroupName + "\"";
+				CommandLineArguments += " -LocalizationBranch=\"" + LocalizationBranchSuffix + "\"";
 			}
 
-			if (!String.IsNullOrEmpty(OneSkyProjectNames))
+			if (!String.IsNullOrEmpty(LocalizationProviderName))
 			{
-				CommandLineArguments += " -OneSkyProjectNames=\"" + OneSkyProjectNames + "\"";
+				CommandLineArguments += " -LocalizationProvider=\"" + LocalizationProviderName + "\"";
 			}
 
-			if (!String.IsNullOrEmpty(LocalizationBranchSuffix))
+			if (!String.IsNullOrEmpty(LocalizationProviderArgs))
 			{
-				CommandLineArguments += " -OneSkyBranchSuffix=\"" + LocalizationBranchSuffix + "\"";
+				CommandLineArguments += " " + LocalizationProviderArgs;
 			}
 
 			// Run the localise script.
@@ -2708,19 +2708,19 @@ partial class GUBP
 			throw new AutomationException("Unimplemented GetUEProjectName.");
 		}
 
-		protected virtual string GetOneSkyConfigName()
+		protected virtual string GetLocalizationProjectNames()
 		{
-			throw new AutomationException("Unimplemented GetOneSkyConfigName.");
+			throw new AutomationException("Unimplemented GetLocalizationProjectNames.");
 		}
 
-		protected virtual string GetOneSkyProjectGroupName()
+		protected virtual string GetLocalizationProviderName()
 		{
-			throw new AutomationException("Unimplemented GetOneSkyProjectGroupName.");
+			throw new AutomationException("Unimplemented GetLocalizationProviderName.");
 		}
 
-		protected virtual string GetOneSkyProjectNames()
+		protected virtual string GetLocalizationProviderArgs()
 		{
-			throw new AutomationException("Unimplemented GetOneSkyProjectNames.");
+			return "";
 		}
 
 		protected string LocalizationBranchSuffix;
@@ -2754,19 +2754,19 @@ partial class GUBP
 			return "";
 		}
 
-		protected override string GetOneSkyConfigName()
+		protected override string GetLocalizationProjectNames()
 		{
-			return "OneSkyConfig_EpicGames";
+			return "Engine,Editor,EditorTutorials,PropertyNames,ToolTips,Category,Keywords";
 		}
 
-		protected override string GetOneSkyProjectGroupName()
+		protected override string GetLocalizationProviderName()
 		{
-			return "Unreal Engine";
+			return "OneSky";
 		}
 
-		protected override string GetOneSkyProjectNames()
+		protected override string GetLocalizationProviderArgs()
 		{
-			return "Engine,Editor,EditorTutorials,PropertyNames,ToolTips,Category,Keywords";
+			return "-OneSkyConfigName=\"OneSkyConfig_EpicGames\" -OneSkyProjectGroupName=\"Unreal Engine\"";
 		}
 	}
 
diff --git a/Engine/Source/Programs/AutomationTool/IOS/IOSPlatform.Automation.cs b/Engine/Source/Programs/AutomationTool/IOS/IOSPlatform.Automation.cs
--- a/Engine/Source/Programs/AutomationTool/IOS/IOSPlatform.Automation.cs
+++ b/Engine/Source/Programs/AutomationTool/IOS/IOSPlatform.Automation.cs
@@ -721,9 +721,9 @@ public class IOSPlatform : Platform
 				int EndPos = Contents.IndexOf("</string>", Pos);
 				BundleIdentifier = Contents.Substring(Pos, EndPos - Pos);
 			}
-			RunAndLog(CmdEnv, DeployServer, "Backup -file \"" + CombinePaths(Params.BaseStageDirectory, PlatformName, DeploymentContext.UFSDeployedManifestFileName) + "\" -file \"" + CombinePaths(Params.BaseStageDirectory, PlatformName, DeploymentContext.NonUFSDeployedManifestFileName) + "\"" + (String.IsNullOrEmpty(Params.Device) ? "" : " -device " + Params.Device.Substring(4)) + " -bundle " + BundleIdentifier);
+			RunAndLog(CmdEnv, DeployServer, "Backup -file \"" + CombinePaths(Params.BaseStageDirectory, "IOS", SC.UFSDeployedManifestFileName) + "\" -file \"" + CombinePaths(Params.BaseStageDirectory, "IOS", SC.NonUFSDeployedManifestFileName) + "\"" + (String.IsNullOrEmpty(Params.Device) ? "" : " -device " + Params.Device.Substring(4)) + " -bundle " + BundleIdentifier);
 
-			string[] ManifestFiles = Directory.GetFiles(CombinePaths(Params.BaseStageDirectory, PlatformName), "*_Manifest_UFS*.txt");
+			string[] ManifestFiles = Directory.GetFiles(CombinePaths(Params.BaseStageDirectory, "IOS"), "*_Manifest_UFS*.txt");
 			UFSManifests.AddRange(ManifestFiles);
 
 			ManifestFiles = Directory.GetFiles(CombinePaths(Params.BaseStageDirectory, PlatformName), "*_Manifest_NonUFS*.txt");
diff --git a/Engine/Source/Programs/AutomationTool/Localization/Localization.Automation.csproj b/Engine/Source/Programs/AutomationTool/Localization/Localization.Automation.csproj
new file mode 100644
--- /dev/null
+++ b/Engine/Source/Programs/AutomationTool/Localization/Localization.Automation.csproj
@@ -0,0 +1,67 @@
+﻿<?xml version="1.0" encoding="utf-8"?>
+<Project ToolsVersion="4.0" DefaultTargets="Build" xmlns="http://schemas.microsoft.com/developer/msbuild/2003">
+  <Import Project="$(MSBuildExtensionsPath)\$(MSBuildToolsVersion)\Microsoft.Common.props" Condition="Exists('$(MSBuildExtensionsPath)\$(MSBuildToolsVersion)\Microsoft.Common.props')" />
+  <PropertyGroup>
+    <Configuration Condition=" '$(Configuration)' == '' ">Debug</Configuration>
+    <Platform Condition=" '$(Platform)' == '' ">AnyCPU</Platform>
+    <ProjectGuid>{C2F959B1-1727-4598-9C32-3D9ACCE72750}</ProjectGuid>
+    <OutputType>Library</OutputType>
+    <AppDesignerFolder>Properties</AppDesignerFolder>
+    <RootNamespace>AutomationTool</RootNamespace>
+    <AssemblyName>Localization.Automation</AssemblyName>
+    <TargetFrameworkVersion>v4.5</TargetFrameworkVersion>
+    <FileAlignment>512</FileAlignment>
+  </PropertyGroup>
+  <PropertyGroup Condition=" '$(Configuration)|$(Platform)' == 'Debug|AnyCPU' ">
+    <DebugSymbols>true</DebugSymbols>
+    <DebugType>full</DebugType>
+    <Optimize>false</Optimize>
+    <OutputPath>..\..\..\..\Binaries\DotNET\AutomationScripts\</OutputPath>
+    <DefineConstants>DEBUG;TRACE</DefineConstants>
+    <ErrorReport>prompt</ErrorReport>
+    <WarningLevel>4</WarningLevel>
+    <PlatformTarget>AnyCPU</PlatformTarget>
+    <TreatWarningsAsErrors>true</TreatWarningsAsErrors>
+  </PropertyGroup>
+  <PropertyGroup Condition=" '$(Configuration)|$(Platform)' == 'Development|AnyCPU' ">
+    <DebugType>pdbonly</DebugType>
+    <Optimize>true</Optimize>
+    <OutputPath>..\..\..\..\Binaries\DotNET\AutomationScripts\</OutputPath>
+    <DefineConstants>TRACE</DefineConstants>
+    <ErrorReport>prompt</ErrorReport>
+    <WarningLevel>4</WarningLevel>
+    <PlatformTarget>AnyCPU</PlatformTarget>
+    <TreatWarningsAsErrors>true</TreatWarningsAsErrors>
+  </PropertyGroup>
+  <ItemGroup>
+    <Reference Include="System" />
+    <Reference Include="System.Core" />
+    <Reference Include="System.IO.Compression" />
+    <Reference Include="System.IO.Compression.FileSystem" />
+    <Reference Include="System.Runtime.Serialization" />
+    <Reference Include="System.Web" />
+    <Reference Include="System.Web.Extensions" />
+    <Reference Include="System.Xml.Linq" />
+    <Reference Include="System.Data.DataSetExtensions" />
+    <Reference Include="Microsoft.CSharp" />
+    <Reference Include="System.Data" />
+    <Reference Include="System.Xml" />
+  </ItemGroup>
+  <ItemGroup>
+    <Compile Include="LocalizationProvider.cs" />
+  </ItemGroup>
+  <ItemGroup>
+    <ProjectReference Include="..\..\UnrealBuildTool\UnrealBuildTool.csproj">
+      <Project>{fd7c5e1a-cfe4-4fd5-a525-1eb1599a39ac}</Project>
+      <Name>UnrealBuildTool</Name>
+      <Private>False</Private>
+    </ProjectReference>
+    <ProjectReference Include="..\AutomationUtils\AutomationUtils.Automation.csproj">
+      <Project>{2c96a7f2-b1a3-4258-8e0a-e588ff41a53e}</Project>
+      <Name>AutomationUtils.Automation</Name>
+      <Private>False</Private>
+    </ProjectReference>
+  </ItemGroup>
+  <ItemGroup />
+  <Import Project="$(MSBuildToolsPath)\Microsoft.CSharp.targets" />
+</Project>
\ No newline at end of file
diff --git a/Engine/Source/Programs/AutomationTool/Localization/LocalizationProvider.cs b/Engine/Source/Programs/AutomationTool/Localization/LocalizationProvider.cs
new file mode 100644
--- /dev/null
+++ b/Engine/Source/Programs/AutomationTool/Localization/LocalizationProvider.cs
@@ -0,0 +1,148 @@
+// Copyright 1998-2015 Epic Games, Inc. All Rights Reserved.
+using System;
+using System.Collections.Generic;
+using System.IO;
+using AutomationTool;
+using UnrealBuildTool;
+
+namespace EpicGames.Localization
+{
+
+public struct ProjectImportExportInfo
+{
+	/** The destination path for the files containing the gathered data for this project (extracted from its config file - relative to the root working directory for the commandlet) */
+	public string DestinationPath;
+
+	/** The name to use for this projects manifest file (extracted from its config file) */
+	public string ManifestName;
+
+	/** The name to use for this projects archive file (extracted from its config file) */
+	public string ArchiveName;
+
+	/** The name to use for this projects portable object file (extracted from its config file) */
+	public string PortableObjectName;
+
+	/** The native culture for this project (extracted from its config file) */
+	public string NativeCulture;
+
+	/** The cultures to generate for this project (extracted from its config file) */
+	public List<string> CulturesToGenerate;
+
+	/** True if we should use a per-culture directly when importing/exporting */
+	public bool bUseCultureDirectory;
+};
+
+public struct ProjectInfo
+{
+	/** The name of this project */
+	public string ProjectName;
+
+	/** Path to this projects localization config files (relative to the root working directory for the commandlet) - ordered so that iterating them runs in the correct order */
+	public List<string> LocalizationConfigFiles;
+
+	/** Config data used by the PO file import process */
+	public ProjectImportExportInfo ImportInfo;
+
+	/** Config data used by the PO file export process */
+	public ProjectImportExportInfo ExportInfo;
+};
+
+public abstract class LocalizationProvider
+{
+	public struct LocalizationProviderArgs
+	{
+		public string RootWorkingDirectory;
+		public CommandUtils CommandUtils;
+	};
+
+	public LocalizationProvider(LocalizationProviderArgs InArgs)
+	{
+		RootWorkingDirectory = InArgs.RootWorkingDirectory;
+		CommandUtils = InArgs.CommandUtils;
+
+		LocalizationBranchName = CommandUtils.ParseParamValue("LocalizationBranch");
+		bUploadAllCultures = CommandUtils.ParseParam("UploadAllCultures");
+	}
+
+	public virtual string GetLocalizationProviderId()
+	{
+		throw new AutomationException("Unimplemented GetLocalizationProviderId.");
+	}
+
+	public virtual void DownloadProjectFromLocalizationProvider(string ProjectName, ProjectImportExportInfo ProjectImportInfo)
+	{
+		throw new AutomationException("Unimplemented DownloadProjectFromLocalizationProvider.");
+	}
+
+	public virtual void UploadProjectToLocalizationProvider(string ProjectName, ProjectImportExportInfo ProjectExportInfo)
+	{
+		throw new AutomationException("Unimplemented UploadProjectToLocalizationProvider.");
+	}
+
+	public static LocalizationProvider GetLocalizationProvider(string InLocalizationProviderId, LocalizationProvider.LocalizationProviderArgs InLocalizationProviderArgs)
+	{
+		if (CachedLocalizationProviderTypes == null)
+		{
+			// Find all types that derive from LocalizationProvider in any of our DLLs
+			CachedLocalizationProviderTypes = new Dictionary<string, Type>();
+			var LoadedAssemblies = AppDomain.CurrentDomain.GetAssemblies();
+			foreach (var Dll in LoadedAssemblies)
+			{
+				var AllTypes = Dll.GetTypes();
+				foreach (var PotentialLocalizationNodeType in AllTypes)
+				{
+					if (PotentialLocalizationNodeType != typeof(LocalizationProvider) && !PotentialLocalizationNodeType.IsAbstract && !PotentialLocalizationNodeType.IsInterface && typeof(LocalizationProvider).IsAssignableFrom(PotentialLocalizationNodeType))
+					{
+						// Types should implement a static StaticGetLocalizationProviderId method
+						var Method = PotentialLocalizationNodeType.GetMethod("StaticGetLocalizationProviderId");
+						if (Method != null)
+						{
+							try
+							{
+								var LocalizationProviderId = Method.Invoke(null, null) as string;
+								CachedLocalizationProviderTypes.Add(LocalizationProviderId, PotentialLocalizationNodeType);
+							}
+							catch
+							{
+								BuildCommand.LogWarning("Type '{0}' threw when calling its StaticGetLocalizationProviderId method.", PotentialLocalizationNodeType.FullName);
+							}
+						}
+						else
+						{
+							BuildCommand.LogWarning("Type '{0}' derives from LocalizationProvider but is missing its StaticGetLocalizationProviderId method.", PotentialLocalizationNodeType.FullName);
+						}
+					}
+				}
+			}
+		}
+
+		Type LocalizationNodeType;
+		CachedLocalizationProviderTypes.TryGetValue(InLocalizationProviderId, out LocalizationNodeType);
+		if (LocalizationNodeType != null)
+		{
+			try
+			{
+				return Activator.CreateInstance(LocalizationNodeType, new object[] { InLocalizationProviderArgs }) as LocalizationProvider;
+			}
+			catch
+			{
+				BuildCommand.LogWarning("Unable to create an instance of the type '{0}'", LocalizationNodeType.FullName);
+			}
+		}
+		else
+		{
+			throw new AutomationException("Could not find a localization provider for: '" + InLocalizationProviderId + "'");
+		}
+
+		return null;
+	}
+
+	protected string RootWorkingDirectory;
+	protected string LocalizationBranchName;
+	protected bool bUploadAllCultures;
+	protected CommandUtils CommandUtils;
+
+	private static Dictionary<string, Type> CachedLocalizationProviderTypes;
+};
+
+}
diff --git a/Engine/Source/Programs/AutomationTool/OneSkyLocalization/OneSkyLocalization.Automation.csproj b/Engine/Source/Programs/AutomationTool/OneSkyLocalization/OneSkyLocalization.Automation.csproj
new file mode 100644
--- /dev/null
+++ b/Engine/Source/Programs/AutomationTool/OneSkyLocalization/OneSkyLocalization.Automation.csproj
@@ -0,0 +1,77 @@
+﻿<?xml version="1.0" encoding="utf-8"?>
+<Project ToolsVersion="4.0" DefaultTargets="Build" xmlns="http://schemas.microsoft.com/developer/msbuild/2003">
+  <Import Project="$(MSBuildExtensionsPath)\$(MSBuildToolsVersion)\Microsoft.Common.props" Condition="Exists('$(MSBuildExtensionsPath)\$(MSBuildToolsVersion)\Microsoft.Common.props')" />
+  <PropertyGroup>
+    <Configuration Condition=" '$(Configuration)' == '' ">Debug</Configuration>
+    <Platform Condition=" '$(Platform)' == '' ">AnyCPU</Platform>
+    <ProjectGuid>{39369368-5603-46DA-BAAE-766BB3CB2B51}</ProjectGuid>
+    <OutputType>Library</OutputType>
+    <AppDesignerFolder>Properties</AppDesignerFolder>
+    <RootNamespace>AutomationTool</RootNamespace>
+    <AssemblyName>OneSkyLocalization.Automation</AssemblyName>
+    <TargetFrameworkVersion>v4.5</TargetFrameworkVersion>
+    <FileAlignment>512</FileAlignment>
+  </PropertyGroup>
+  <PropertyGroup Condition=" '$(Configuration)|$(Platform)' == 'Debug|AnyCPU' ">
+    <DebugSymbols>true</DebugSymbols>
+    <DebugType>full</DebugType>
+    <Optimize>false</Optimize>
+    <OutputPath>..\..\..\..\Binaries\DotNET\AutomationScripts\</OutputPath>
+    <DefineConstants>DEBUG;TRACE</DefineConstants>
+    <ErrorReport>prompt</ErrorReport>
+    <WarningLevel>4</WarningLevel>
+    <PlatformTarget>AnyCPU</PlatformTarget>
+    <TreatWarningsAsErrors>true</TreatWarningsAsErrors>
+  </PropertyGroup>
+  <PropertyGroup Condition=" '$(Configuration)|$(Platform)' == 'Development|AnyCPU' ">
+    <DebugType>pdbonly</DebugType>
+    <Optimize>true</Optimize>
+    <OutputPath>..\..\..\..\Binaries\DotNET\AutomationScripts\</OutputPath>
+    <DefineConstants>TRACE</DefineConstants>
+    <ErrorReport>prompt</ErrorReport>
+    <WarningLevel>4</WarningLevel>
+    <PlatformTarget>AnyCPU</PlatformTarget>
+    <TreatWarningsAsErrors>true</TreatWarningsAsErrors>
+  </PropertyGroup>
+  <ItemGroup>
+    <Reference Include="OneSky, Version=1.0.0.0, Culture=neutral, processorArchitecture=MSIL">
+      <HintPath>..\..\..\..\Binaries\DotNET\OneSky.dll</HintPath>
+      <SpecificVersion>False</SpecificVersion>
+      <Private>False</Private>
+    </Reference>
+    <Reference Include="System" />
+    <Reference Include="System.Core" />
+    <Reference Include="System.IO.Compression" />
+    <Reference Include="System.IO.Compression.FileSystem" />
+    <Reference Include="System.Runtime.Serialization" />
+    <Reference Include="System.Web" />
+    <Reference Include="System.Web.Extensions" />
+    <Reference Include="System.Xml.Linq" />
+    <Reference Include="System.Data.DataSetExtensions" />
+    <Reference Include="Microsoft.CSharp" />
+    <Reference Include="System.Data" />
+    <Reference Include="System.Xml" />
+  </ItemGroup>
+  <ItemGroup>
+    <Compile Include="OneSkyConfigHelper.cs" />
+    <Compile Include="OneSkyLocalizationProvider.cs" />
+  </ItemGroup>
+  <ItemGroup>
+    <ProjectReference Include="..\..\UnrealBuildTool\UnrealBuildTool.csproj">
+      <Project>{fd7c5e1a-cfe4-4fd5-a525-1eb1599a39ac}</Project>
+      <Name>UnrealBuildTool</Name>
+      <Private>False</Private>
+    </ProjectReference>
+    <ProjectReference Include="..\AutomationUtils\AutomationUtils.Automation.csproj">
+      <Project>{2c96a7f2-b1a3-4258-8e0a-e588ff41a53e}</Project>
+      <Name>AutomationUtils.Automation</Name>
+      <Private>False</Private>
+    </ProjectReference>
+    <ProjectReference Include="..\Localization\Localization.Automation.csproj">
+      <Project>{c2f959b1-1727-4598-9c32-3d9acce72750}</Project>
+      <Name>Localization.Automation</Name>
+    </ProjectReference>
+  </ItemGroup>
+  <ItemGroup />
+  <Import Project="$(MSBuildToolsPath)\Microsoft.CSharp.targets" />
+</Project>
\ No newline at end of file
diff --git a/Engine/Source/Programs/AutomationTool/OneSkyLocalization/OneSkyLocalizationProvider.cs b/Engine/Source/Programs/AutomationTool/OneSkyLocalization/OneSkyLocalizationProvider.cs
new file mode 100644
--- /dev/null
+++ b/Engine/Source/Programs/AutomationTool/OneSkyLocalization/OneSkyLocalizationProvider.cs
@@ -0,0 +1,272 @@
+// Copyright 1998-2015 Epic Games, Inc. All Rights Reserved.
+using System;
+using System.Collections.Generic;
+using System.IO;
+using System.Linq;
+using AutomationTool;
+using UnrealBuildTool;
+using OneSky;
+using EpicGames.Localization;
+using EpicGames.OneSkyLocalization.Config;
+
+namespace EpicGames.OneSkyLocalization
+{
+
+[Help("OneSkyConfigName", "Name of the config data to use (see OneSkyConfigHelper).")]
+[Help("OneSkyProjectGroupName", "Name of the project group in OneSky.")]
+public class OneSkyLocalizationProvider : LocalizationProvider
+{
+	public OneSkyLocalizationProvider(LocalizationProviderArgs InArgs)
+		: base(InArgs)
+	{
+		var OneSkyConfigName = CommandUtils.ParseParamValue("OneSkyConfigName");
+		if (OneSkyConfigName == null)
+		{
+			throw new AutomationException("Missing required command line argument: 'OneSkyConfigName'");
+		}
+
+		var OneSkyProjectGroupName = CommandUtils.ParseParamValue("OneSkyProjectGroupName");
+		if (OneSkyProjectGroupName == null)
+		{
+			throw new AutomationException("Missing required command line argument: 'OneSkyProjectGroupName'");
+		}
+
+		OneSkyConfigData OneSkyConfig = OneSkyConfigHelper.Find(OneSkyConfigName);
+		OneSkyService = new OneSkyService(OneSkyConfig.ApiKey, OneSkyConfig.ApiSecret);
+		OneSkyProjectGroup = GetOneSkyProjectGroup(OneSkyProjectGroupName);
+	}
+
+	public static string StaticGetLocalizationProviderId()
+	{
+		return "OneSky";
+	}
+
+	public override string GetLocalizationProviderId()
+	{
+		return StaticGetLocalizationProviderId();
+	}
+
+	public override void DownloadProjectFromLocalizationProvider(string ProjectName, ProjectImportExportInfo ProjectImportInfo)
+	{
+		var OneSkyFileName = GetOneSkyFilename(ProjectImportInfo.PortableObjectName);
+		var OneSkyProject = GetOneSkyProject(ProjectName);
+		var OneSkyFile = OneSkyProject.UploadedFiles.FirstOrDefault(f => f.Filename == OneSkyFileName);
+
+		// Create changelist for backed up POs from OneSky.
+		int OneSkyDownloadedPOChangeList = 0;
+		if (CommandUtils.P4Enabled)
+		{
+			OneSkyDownloadedPOChangeList = CommandUtils.P4.CreateChange(CommandUtils.P4Env.Client, "OneSky downloaded PO backup.");
+		}
+
+		//Export
+		if (OneSkyFile != null)
+		{
+			var CulturesToExport = new List<string>();
+			foreach (var OneSkyCulture in OneSkyProject.EnabledCultures)
+			{
+				// Only export the OneSky cultures that we care about for this project
+				if (ProjectImportInfo.CulturesToGenerate.Contains(OneSkyCulture))
+				{
+					CulturesToExport.Add(OneSkyCulture);
+				}
+			}
+
+			ExportOneSkyFileToDirectory(OneSkyFile, new DirectoryInfo(CommandUtils.CombinePaths(RootWorkingDirectory, ProjectImportInfo.DestinationPath)), ProjectImportInfo.PortableObjectName, CulturesToExport, ProjectImportInfo.bUseCultureDirectory, OneSkyDownloadedPOChangeList);
+		}
+
+		// Submit changelist for backed up POs from OneSky.
+		if (CommandUtils.P4Enabled)
+		{
+			int SubmittedChangeList;
+			CommandUtils.P4.Submit(OneSkyDownloadedPOChangeList, out SubmittedChangeList);
+		}
+	}
+
+	private void ExportOneSkyFileToDirectory(UploadedFile OneSkyFile, DirectoryInfo DestinationDirectory, string DestinationFilename, IEnumerable<string> Cultures, bool bUseCultureDirectory, int OneSkyDownloadedPOChangeList)
+	{
+		foreach (var Culture in Cultures)
+		{
+			var CultureDirectory = (bUseCultureDirectory) ? new DirectoryInfo(Path.Combine(DestinationDirectory.FullName, Culture)) : DestinationDirectory;
+			if (!CultureDirectory.Exists)
+			{
+				CultureDirectory.Create();
+			}
+
+			using (var MemoryStream = new MemoryStream())
+			{
+				var ExportTranslationState = OneSkyFile.ExportTranslation(Culture, MemoryStream).Result;
+				if (ExportTranslationState == UploadedFile.ExportTranslationState.Success)
+				{
+					var ExportFile = new FileInfo(Path.Combine(CultureDirectory.FullName, DestinationFilename));
+
+					// Write out the updated PO file so that the gather commandlet will import the new data from it
+					{
+						var ExportFileWasReadOnly = false;
+						if (ExportFile.Exists)
+						{
+							// We're going to clobber the existing PO file, so make sure it's writable (it may be read-only if in Perforce)
+							ExportFileWasReadOnly = ExportFile.IsReadOnly;
+							ExportFile.IsReadOnly = false;
+						}
+
+						MemoryStream.Position = 0;
+						using (Stream FileStream = ExportFile.OpenWrite())
+						{
+							MemoryStream.CopyTo(FileStream);
+							Console.WriteLine("[SUCCESS] Exporting: '{0}' as '{1}' ({2})", OneSkyFile.Filename, ExportFile.FullName, Culture);
+						}
+
+						if (ExportFileWasReadOnly)
+						{
+							ExportFile.IsReadOnly = true;
+						}
+					}
+
+					// Also update the back-up copy so we can diff against what we got from OneSky, and what the gather commandlet produced
+					{
+						var ExportFileCopy = new FileInfo(Path.Combine(ExportFile.DirectoryName, String.Format("{0}_FromOneSky{1}", Path.GetFileNameWithoutExtension(ExportFile.Name), ExportFile.Extension)));
+
+						var ExportFileCopyWasReadOnly = false;
+						if (ExportFileCopy.Exists)
+						{
+							// We're going to clobber the existing PO file, so make sure it's writable (it may be read-only if in Perforce)
+							ExportFileCopyWasReadOnly = ExportFileCopy.IsReadOnly;
+							ExportFileCopy.IsReadOnly = false;
+						}
+
+						ExportFile.CopyTo(ExportFileCopy.FullName, true);
+
+						if (ExportFileCopyWasReadOnly)
+						{
+							ExportFileCopy.IsReadOnly = true;
+						}
+
+						// Add/check out backed up POs from OneSky.
+						if (CommandUtils.P4Enabled)
+						{
+							UE4Build.AddBuildProductsToChangelist(OneSkyDownloadedPOChangeList, new List<string>() { ExportFileCopy.FullName });
+						}
+					}
+				}
+				else if (ExportTranslationState == UploadedFile.ExportTranslationState.NoContent)
+				{
+					Console.WriteLine("[WARNING] Exporting: '{0}' ({1}) has no translations!", OneSkyFile.Filename, Culture);
+				}
+				else
+				{
+					Console.WriteLine("[FAILED] Exporting: '{0}' ({1})", OneSkyFile.Filename, Culture);
+				}
+			}
+		}
+	}
+
+	public override void UploadProjectToLocalizationProvider(string ProjectName, ProjectImportExportInfo ProjectExportInfo)
+	{
+		var OneSkyProject = GetOneSkyProject(ProjectName);
+
+		Func<string, FileInfo> GetPathForCulture = (string Culture) =>
+		{
+			if (ProjectExportInfo.bUseCultureDirectory)
+			{
+				return new FileInfo(Path.Combine(RootWorkingDirectory, ProjectExportInfo.DestinationPath, Culture, ProjectExportInfo.PortableObjectName));
+			}
+			else
+			{
+				return new FileInfo(Path.Combine(RootWorkingDirectory, ProjectExportInfo.DestinationPath, ProjectExportInfo.PortableObjectName));
+			}
+		};
+
+		// Upload the .po file for the native culture first
+		UploadFileToOneSky(OneSkyProject, GetPathForCulture(ProjectExportInfo.NativeCulture), ProjectExportInfo.NativeCulture);
+
+		if (bUploadAllCultures)
+		{
+			// Upload the remaining .po files for the other cultures
+			foreach (var Culture in ProjectExportInfo.CulturesToGenerate)
+			{
+				// Skip native culture as we uploaded it above
+				if (Culture != ProjectExportInfo.NativeCulture)
+				{
+					UploadFileToOneSky(OneSkyProject, GetPathForCulture(Culture), Culture);
+				}
+			}
+		}
+	}
+
+	private void UploadFileToOneSky(OneSky.Project OneSkyProject, FileInfo FileToUpload, string Culture)
+	{
+		using (var FileStream = FileToUpload.OpenRead())
+		{
+			// Read the BOM
+			var UTF8BOM = new byte[3];
+			FileStream.Read(UTF8BOM, 0, 3);
+
+			// We want to ignore the utf8 BOM
+			if (UTF8BOM[0] != 0xef || UTF8BOM[1] != 0xbb || UTF8BOM[2] != 0xbf)
+			{
+				FileStream.Position = 0;
+			}
+
+			var OneSkyFileName = GetOneSkyFilename(Path.GetFileName(FileToUpload.FullName));
+
+			Console.WriteLine("Uploading: '{0}' as '{1}' ({2})", FileToUpload.FullName, OneSkyFileName, Culture);
+
+			var UploadedFile = OneSkyProject.Upload(OneSkyFileName, FileStream, Culture).Result;
+
+			if (UploadedFile == null)
+			{
+				Console.WriteLine("[FAILED] Uploading: '{0}' ({1})", FileToUpload.FullName, Culture);
+			}
+			else
+			{
+				Console.WriteLine("[SUCCESS] Uploading: '{0}' ({1})", FileToUpload.FullName, Culture);
+			}
+		}
+	}
+
+	private ProjectGroup GetOneSkyProjectGroup(string ProjectGroupName)
+	{
+		var OneSkyProjectGroup = OneSkyService.ProjectGroups.FirstOrDefault(g => g.Name == ProjectGroupName);
+
+		if (OneSkyProjectGroup == null)
+		{
+			OneSkyProjectGroup = new ProjectGroup(ProjectGroupName, "en");
+			OneSkyService.ProjectGroups.Add(OneSkyProjectGroup);
+		}
+
+		return OneSkyProjectGroup;
+	}
+
+	private OneSky.Project GetOneSkyProject(string ProjectName, string ProjectDescription = "")
+	{
+		OneSky.Project OneSkyProject = OneSkyProjectGroup.Projects.FirstOrDefault(p => p.Name == ProjectName);
+
+		if (OneSkyProject == null)
+		{
+			ProjectType projectType = OneSkyService.ProjectTypes.First(pt => pt.Code == "website");
+
+			OneSkyProject = new OneSky.Project(ProjectName, ProjectDescription, projectType);
+			OneSkyProjectGroup.Projects.Add(OneSkyProject);
+		}
+
+		return OneSkyProject;
+	}
+
+	private string GetOneSkyFilename(string BaseFilename)
+	{
+		var OneSkyFileName = BaseFilename;
+		if (!String.IsNullOrEmpty(LocalizationBranchName))
+		{
+			// Apply the branch suffix. OneSky will take care of merging the files from different branches together.
+			var OneSkyFileNameWithSuffix = Path.GetFileNameWithoutExtension(OneSkyFileName) + "_" + LocalizationBranchName + Path.GetExtension(OneSkyFileName);
+			OneSkyFileName = OneSkyFileNameWithSuffix;
+		}
+		return OneSkyFileName;
+	}
+
+	private OneSkyService OneSkyService;
+	private ProjectGroup OneSkyProjectGroup;
+}
+
+}
diff --git a/Engine/Source/Programs/AutomationTool/Scripts/AutomationScripts.Automation.csproj b/Engine/Source/Programs/AutomationTool/Scripts/AutomationScripts.Automation.csproj
--- a/Engine/Source/Programs/AutomationTool/Scripts/AutomationScripts.Automation.csproj
+++ b/Engine/Source/Programs/AutomationTool/Scripts/AutomationScripts.Automation.csproj
@@ -107,6 +107,14 @@
       <Name>GUBP.Automation</Name>
       <Private>False</Private>
     </ProjectReference>
+    <ProjectReference Include="..\Localization\Localization.Automation.csproj">
+      <Project>{c2f959b1-1727-4598-9c32-3d9acce72750}</Project>
+      <Name>Localization.Automation</Name>
+    </ProjectReference>
+    <ProjectReference Include="..\OneSkyLocalization\OneSkyLocalization.Automation.csproj">
+      <Project>{39369368-5603-46da-baae-766bb3cb2b51}</Project>
+      <Name>OneSkyLocalization.Automation</Name>
+    </ProjectReference>
   </ItemGroup>
   <Import Project="$(MSBuildToolsPath)\Microsoft.CSharp.targets" />
 </Project>
\ No newline at end of file
diff --git a/Engine/Source/Programs/AutomationTool/Scripts/CopyBuildToStagingDirectory.Automation.cs b/Engine/Source/Programs/AutomationTool/Scripts/CopyBuildToStagingDirectory.Automation.cs
--- a/Engine/Source/Programs/AutomationTool/Scripts/CopyBuildToStagingDirectory.Automation.cs
+++ b/Engine/Source/Programs/AutomationTool/Scripts/CopyBuildToStagingDirectory.Automation.cs
@@ -602,14 +602,14 @@ public partial class Project : CommandUtils
 		}
 	}
 
-	public static void CopyManifestFilesToStageDir(Dictionary<string, string> Mapping, string StageDir, string ManifestName, List<string> CRCFiles)
+	public static void CopyManifestFilesToStageDir(Dictionary<string, string> Mapping, string StageDir, string ManifestName, List<string> CRCFiles, string PlatformName)
 	{
 		Log("Copying {0} to staging directory: {1}", ManifestName, StageDir);
 		string ManifestPath = "";
 		string ManifestFile = "";
 		if (!String.IsNullOrEmpty(ManifestName))
 		{
-			ManifestFile = "Manifest_" + ManifestName + ".txt";
+			ManifestFile = "Manifest_" + ManifestName + "_" + PlatformName + ".txt";
 			ManifestPath = CombinePaths(StageDir, ManifestFile);
 			DeleteFile(ManifestPath);
 		}
@@ -633,14 +633,14 @@ public partial class Project : CommandUtils
 		}
 	}
 
-	public static void CopyManifestFilesToStageDir(Dictionary<string, List<string>> Mapping, string StageDir, string ManifestName, List<string> CRCFiles)
+	public static void CopyManifestFilesToStageDir(Dictionary<string, List<string>> Mapping, string StageDir, string ManifestName, List<string> CRCFiles, string PlatformName)
 	{
 		Log("Copying {0} to staging directory: {1}", ManifestName, StageDir);
 		string ManifestPath = "";
 		string ManifestFile = "";
 		if (!String.IsNullOrEmpty(ManifestName))
 		{
-			ManifestFile = "Manifest_" + ManifestName + ".txt";
+			ManifestFile = "Manifest_" + ManifestName + "_" + PlatformName + ".txt";
 			ManifestPath = CombinePaths(StageDir, ManifestFile);
 			DeleteFile(ManifestPath);
 		}
@@ -712,17 +712,17 @@ public partial class Project : CommandUtils
 
 	public static void CopyUsingStagingManifest(ProjectParams Params, DeploymentContext SC)
 	{
-		CopyManifestFilesToStageDir(SC.NonUFSStagingFiles, SC.StageDirectory, "NonUFSFiles", SC.StageTargetPlatform.GetFilesForCRCCheck());
+		CopyManifestFilesToStageDir(SC.NonUFSStagingFiles, SC.StageDirectory, "NonUFSFiles", SC.StageTargetPlatform.GetFilesForCRCCheck(), SC.StageTargetPlatform.PlatformType.ToString());
 
 		if (!Params.NoDebugInfo)
 		{
-			CopyManifestFilesToStageDir(SC.NonUFSStagingFilesDebug, SC.StageDirectory, "DebugFiles", SC.StageTargetPlatform.GetFilesForCRCCheck());
+			CopyManifestFilesToStageDir(SC.NonUFSStagingFilesDebug, SC.StageDirectory, "DebugFiles", SC.StageTargetPlatform.GetFilesForCRCCheck(), SC.StageTargetPlatform.PlatformType.ToString());
 		}
 
 		bool bStageUnrealFileSystemFiles = !Params.CookOnTheFly && !Params.UsePak(SC.StageTargetPlatform) && !Params.FileServer;
 		if (bStageUnrealFileSystemFiles)
 		{
-			CopyManifestFilesToStageDir(SC.UFSStagingFiles, SC.StageDirectory, "UFSFiles", SC.StageTargetPlatform.GetFilesForCRCCheck());
+			CopyManifestFilesToStageDir(SC.UFSStagingFiles, SC.StageDirectory, "UFSFiles", SC.StageTargetPlatform.GetFilesForCRCCheck(), SC.StageTargetPlatform.PlatformType.ToString());
 		}
 	}
 
@@ -1577,8 +1577,8 @@ public partial class Project : CommandUtils
 		// add the manifest
 		if (!DeltaManifest.Contains("NonUFS"))
 		{
-			DeltaFiles.Add(DeploymentContext.NonUFSDeployedManifestFileName);
-			DeltaFiles.Add(DeploymentContext.UFSDeployedManifestFileName);
+			DeltaFiles.Add(SC.NonUFSDeployedManifestFileName);
+			DeltaFiles.Add(SC.UFSDeployedManifestFileName);
 		}
 
 		// TODO: determine files which need to be removed
@@ -1772,8 +1772,8 @@ public partial class Project : CommandUtils
 						}
 
 						// get the staged file data
-						Dictionary<string, string> StagedUFSFiles = ReadStagedManifest(Params, SC, DeploymentContext.UFSDeployedManifestFileName);
-						Dictionary<string, string> StagedNonUFSFiles = ReadStagedManifest(Params, SC, DeploymentContext.NonUFSDeployedManifestFileName);
+						Dictionary<string, string> StagedUFSFiles = ReadStagedManifest(Params, SC, SC.UFSDeployedManifestFileName);
+						Dictionary<string, string> StagedNonUFSFiles = ReadStagedManifest(Params, SC, SC.NonUFSDeployedManifestFileName);
 
 						WriteObsoleteManifest(Params, SC, DeployedUFSFiles, StagedUFSFiles, DeploymentContext.UFSDeployObsoleteFileName);
 						WriteObsoleteManifest(Params, SC, DeployedNonUFSFiles, StagedNonUFSFiles, DeploymentContext.NonUFSDeployObsoleteFileName);
diff --git a/Engine/Source/Programs/AutomationTool/Scripts/LauncherLocalization.Automation.cs b/Engine/Source/Programs/AutomationTool/Scripts/LauncherLocalization.Automation.cs
--- a/Engine/Source/Programs/AutomationTool/Scripts/LauncherLocalization.Automation.cs
+++ b/Engine/Source/Programs/AutomationTool/Scripts/LauncherLocalization.Automation.cs
@@ -4,7 +4,6 @@ using System.Collections.Generic;
 using System.Globalization;
 using System.Linq;
 using System.IO;
-using System.Reflection;
 using AutomationTool;
 using UnrealBuildTool;
 using OneSky;
@@ -199,80 +198,3 @@ class LauncherLocalization : BuildCommand
 		}
 	}
 }
-
-
-namespace EpicGames.OneSkyLocalization.Config
-{
-	/// <summary>
-	/// Class for retrieving OneSky localization configuration data
-	/// </summary>
-	public class OneSkyConfigHelper
-	{
-		// List of configs is cached off for fetching from multiple times
-		private static Dictionary<string, OneSkyConfigData> Configs;
-
-		public static OneSkyConfigData Find(string ConfigName)
-		{
-			if (Configs == null)
-			{
-				// Load all secret configs by trying to instantiate all classes derived from OneSkyConfigData from all loaded DLLs.
-				// Note that we're using the default constructor on the secret config types.
-				Configs = new Dictionary<string, OneSkyConfigData>();
-				Assembly[] LoadedAssemblies = AppDomain.CurrentDomain.GetAssemblies();
-				foreach (var Dll in LoadedAssemblies)
-				{
-					Type[] AllTypes = Dll.GetTypes();
-					foreach (var PotentialConfigType in AllTypes)
-					{
-						if (PotentialConfigType != typeof(OneSkyConfigData) && typeof(OneSkyConfigData).IsAssignableFrom(PotentialConfigType))
-						{
-							try
-							{
-								OneSkyConfigData Config = Activator.CreateInstance(PotentialConfigType) as OneSkyConfigData;
-								if (Config != null)
-								{
-									Configs.Add(Config.Name, Config);
-								}
-							}
-							catch
-							{
-								BuildCommand.LogWarning("Unable to create OneSky config data: {0}", PotentialConfigType.Name);
-							}
-						}
-					}
-				}
-			}
-			OneSkyConfigData LoadedConfig;
-			Configs.TryGetValue(ConfigName, out LoadedConfig);
-			if (LoadedConfig == null)
-			{
-				throw new AutomationException("Unable to find requested OneSky config data: {0}", ConfigName);
-			}
-			return LoadedConfig;
-		}
-	}
-
-	/// <summary>
-	/// Class for storing OneSky Localization configuration data
-	/// </summary>
-	public class OneSkyConfigData
-	{
-		public OneSkyConfigData(string InName, string InApiKey, string InApiSecret)
-		{
-			Name = InName;
-			ApiKey = InApiKey;
-			ApiSecret = InApiSecret;
-		}
-
-		public string Name;
-		public string ApiKey;
-		public string ApiSecret;
-
-		public void SpewValues()
-		{
-			CommandUtils.Log("Name : {0}", Name);
-			CommandUtils.Log("ApiKey : {0}", ApiKey);
-			//CommandUtils.LogConsole("ApiSecret : {0}", ApiSecret);  // This should probably never be revealed.
-		}
-	}
-}
