Copying //UE4/Dev-Editor to Dev-Main (//UE4/Dev-Main)

diff --git a/Engine/Source/Editor/PropertyEditor/Private/UserInterface/PropertyEditor/SPropertyEditorAsset.cpp b/Engine/Source/Editor/PropertyEditor/Private/UserInterface/PropertyEditor/SPropertyEditorAsset.cpp
--- a/Engine/Source/Editor/PropertyEditor/Private/UserInterface/PropertyEditor/SPropertyEditorAsset.cpp
+++ b/Engine/Source/Editor/PropertyEditor/Private/UserInterface/PropertyEditor/SPropertyEditorAsset.cpp
@@ -513,6 +513,44 @@ FText SPropertyEditorAsset::OnGetToolTip() const
 	return ToolTip;
 }
 
+bool SPropertyEditorAsset::HasEngineOuterObject() const
+{
+	if (PropertyEditor.IsValid())
+	{
+		// Determine whether any outer object is an Engine asset
+		TArray<UObject*> OuterObjects;
+		PropertyEditor->GetPropertyHandle()->GetOuterObjects(OuterObjects);
+		for (UObject* OuterObject : OuterObjects)
+		{
+			const FAssetData AssetData(OuterObject);
+			const FString AssetPath = AssetData.ObjectPath.ToString();
+			if (FPackageName::IsEnginePackageName(AssetPath))
+			{
+				return true;
+			}
+		}
+	}
+
+	return false;
+}
+
+bool SPropertyEditorAsset::CanAssetBeAssigned(const FAssetData& AssetData) const
+{
+	// Prevent Engine assets being assigned non-Engine asset properties
+	const bool bHasEngineOuterObject = HasEngineOuterObject();
+	const FString AssetPath = AssetData.ObjectPath.ToString();
+	if (bHasEngineOuterObject && !FPackageName::IsEnginePackageName(AssetPath))
+	{
+		FMessageDialog::Open(EAppMsgType::Ok, FText::Format(
+			LOCTEXT("ObjectAssignmentToEngineFailed", "Cannot assign a Project object {0} to an Engine property."),
+			FText::FromString(AssetPath)));
+
+		return false;
+	}
+
+	return true;
+}
+
 void SPropertyEditorAsset::SetValue( const FAssetData& AssetData )
 {
 	AssetComboButton->SetIsOpen(false);
@@ -523,7 +561,10 @@ void SPropertyEditorAsset::SetValue( const FAssetData& AssetData )
 	{
 		if(PropertyEditor.IsValid())
 		{
-			PropertyEditor->GetPropertyHandle()->SetValue(AssetData);
+			if (CanAssetBeAssigned(AssetData))
+			{
+				PropertyEditor->GetPropertyHandle()->SetValue(AssetData);
+			}
 		}
 
 		OnSetObject.ExecuteIfBound(AssetData);
@@ -708,27 +749,30 @@ FText SPropertyEditorAsset::GetOnBrowseToolTip() const
 
 void SPropertyEditorAsset::OnUse()
 {
-	if(PropertyEditor.IsValid())
+	// Load selected assets
+	FEditorDelegates::LoadSelectedAssetsIfNeeded.Broadcast();
+
+	// try to get a selected object of our class
+	UObject* Selection = NULL;
+	if (ObjectClass && ObjectClass->IsChildOf(AActor::StaticClass()))
 	{
-		PropertyEditor->GetPropertyHandle()->SetObjectValueFromSelection();
+		Selection = GEditor->GetSelectedActors()->GetTop(ObjectClass);
 	}
-	else
+	else if (ObjectClass)
 	{
-		// Load selected assets
-		FEditorDelegates::LoadSelectedAssetsIfNeeded.Broadcast();
+		// Get the first material selected
+		Selection = GEditor->GetSelectedObjects()->GetTop(ObjectClass);
+	}
 
-		// try to get a selected object of our class
-		UObject* Selection = NULL;
-		if( ObjectClass && ObjectClass->IsChildOf( AActor::StaticClass() ) )
-		{
-			Selection = GEditor->GetSelectedActors()->GetTop( ObjectClass );
-		}
-		else if( ObjectClass )
+	if (PropertyEditor.IsValid())
+	{
+		if (Selection && CanAssetBeAssigned(Selection))
 		{
-			// Get the first material selected
-			Selection = GEditor->GetSelectedObjects()->GetTop( ObjectClass );
+			PropertyEditor->GetPropertyHandle()->SetObjectValueFromSelection();
 		}
-
+	}
+	else
+	{
 		// Check against custom asset filter
 		if (Selection != NULL
 			&& OnShouldFilterAsset.IsBound()
diff --git a/Engine/Source/Editor/PropertyEditor/Private/UserInterface/PropertyEditor/SPropertyEditorAsset.h b/Engine/Source/Editor/PropertyEditor/Private/UserInterface/PropertyEditor/SPropertyEditorAsset.h
--- a/Engine/Source/Editor/PropertyEditor/Private/UserInterface/PropertyEditor/SPropertyEditorAsset.h
+++ b/Engine/Source/Editor/PropertyEditor/Private/UserInterface/PropertyEditor/SPropertyEditorAsset.h
@@ -136,6 +136,21 @@ private:
 	 */
 	FText OnGetToolTip() const;
 
+	/**
+	 * Determine whether any outer object is an Engine asset
+	 */
+	bool HasEngineOuterObject() const;
+
+	/**
+	 * Determine whether the given asset is allowed to be assigned to the property.
+	 * This is currently used to prevent non-engine assets being assigned to properties of engine assets.
+	 *
+	 * @param	AssetData	The asset to be assigned
+	 * @return	false		if the asset cannot be assigned,
+	 *			true		if it can be assigned or if this is not known (e.g. because a delegate handles the reference assignment)
+	 */
+	bool CanAssetBeAssigned(const FAssetData& AssetData) const;
+
 	/** 
 	 * Set the value of the asset referenced by this property editor.
 	 * Will set the underlying property handle if there is one.
diff --git a/Engine/Source/Editor/PropertyEditor/Private/UserInterface/PropertyEditor/SPropertyEditorClass.cpp b/Engine/Source/Editor/PropertyEditor/Private/UserInterface/PropertyEditor/SPropertyEditorClass.cpp
--- a/Engine/Source/Editor/PropertyEditor/Private/UserInterface/PropertyEditor/SPropertyEditorClass.cpp
+++ b/Engine/Source/Editor/PropertyEditor/Private/UserInterface/PropertyEditor/SPropertyEditorClass.cpp
@@ -240,8 +240,33 @@ void SPropertyEditorClass::SendToObjects(const FString& NewValue)
 {
 	if(PropertyEditor.IsValid())
 	{
-		const TSharedRef<IPropertyHandle> PropertyHandle = PropertyEditor->GetPropertyHandle();
-		PropertyHandle->SetValueFromFormattedString(NewValue);
+		// Check whether any of the outer objects are engine objects
+		bool bHasEngineOuterObject = false;
+		TArray<UObject*> OuterObjects;
+		PropertyEditor->GetPropertyHandle()->GetOuterObjects(OuterObjects);
+		for (UObject* OuterObject : OuterObjects)
+		{
+			const FAssetData AssetData(OuterObject);
+			if (FPackageName::IsEnginePackageName(AssetData.ObjectPath.ToString()))
+			{
+				bHasEngineOuterObject = true;
+				break;
+			}
+		}
+
+		if (bHasEngineOuterObject && !FPackageName::IsEnginePackageName(NewValue))
+		{
+			// If it's an engine outer object, check whether the value is an Engine class, and if not deem the value invalid.
+			FMessageDialog::Open(EAppMsgType::Ok, FText::Format(
+				LOCTEXT("ClassAssignmentToEngineFailed", "Cannot assign a Project class {0} to an Engine property."),
+				FText::FromString(NewValue)));
+		}
+		else
+		{
+			// Otherwise assign the property
+			const TSharedRef<IPropertyHandle> PropertyHandle = PropertyEditor->GetPropertyHandle();
+			PropertyHandle->SetValueFromFormattedString(NewValue);
+		}
 	}
 	else
 	{
diff --git a/Engine/Source/Editor/SourceControlWindows/Private/SSourceControlSubmit.cpp b/Engine/Source/Editor/SourceControlWindows/Private/SSourceControlSubmit.cpp
--- a/Engine/Source/Editor/SourceControlWindows/Private/SSourceControlSubmit.cpp
+++ b/Engine/Source/Editor/SourceControlWindows/Private/SSourceControlSubmit.cpp
@@ -533,16 +533,7 @@ void FSourceControlWindows::ChoosePackagesToCheckIn()
 		if (ISourceControlModule::Get().GetProvider().IsAvailable())
 		{
 			// make sure we update the SCC status of all packages (this could take a long time, so we will run it as a background task)
-			TArray<FString> Packages;
-			FEditorFileUtils::FindAllPackageFiles(Packages);
-
-			// Get list of filenames corresponding to packages
-			TArray<FString> Filenames = SourceControlHelpers::PackageFilenames(Packages);
-
-			// Add game config files to the list
-			FEditorFileUtils::FindAllConfigFiles(Filenames);
-
-			// Add pending delete paths to the list
+			TArray<FString> Filenames;
 			Filenames.Add(FPaths::ConvertRelativePathToFull(FPaths::EngineContentDir()));
 			Filenames.Add(FPaths::ConvertRelativePathToFull(FPaths::GameContentDir()));
 			Filenames.Add(FPaths::ConvertRelativePathToFull(FPaths::GameConfigDir()));
diff --git a/Engine/Source/Editor/StaticMeshEditor/Private/StaticMeshEditorTools.cpp b/Engine/Source/Editor/StaticMeshEditor/Private/StaticMeshEditorTools.cpp
--- a/Engine/Source/Editor/StaticMeshEditor/Private/StaticMeshEditorTools.cpp
+++ b/Engine/Source/Editor/StaticMeshEditor/Private/StaticMeshEditorTools.cpp
@@ -1274,6 +1274,14 @@ void FMeshSectionSettingsLayout::OnMaterialChanged(UMaterialInterface* NewMateri
 {
 	UStaticMesh& StaticMesh = GetStaticMesh();
 
+	if (FPackageName::IsEnginePackageName(StaticMesh.GetPathName()) && !FPackageName::IsEnginePackageName(NewMaterial->GetPathName()))
+	{
+		FMessageDialog::Open(EAppMsgType::Ok, FText::Format(
+			LOCTEXT("ObjectAssignmentToEngineFailed", "Cannot assign a Project object {0} to an Engine property."),
+			FText::FromString(NewMaterial->GetPathName())));
+		return;
+	}
+
 	// flag the property (Materials) we're modifying so that not all of the object is rebuilt.
 	UProperty* ChangedProperty = NULL;
 	ChangedProperty = FindField<UProperty>( UStaticMesh::StaticClass(), "Materials" );
diff --git a/Engine/Source/Editor/UnrealEd/Classes/Editor/EditorPerProjectUserSettings.h b/Engine/Source/Editor/UnrealEd/Classes/Editor/EditorPerProjectUserSettings.h
--- a/Engine/Source/Editor/UnrealEd/Classes/Editor/EditorPerProjectUserSettings.h
+++ b/Engine/Source/Editor/UnrealEd/Classes/Editor/EditorPerProjectUserSettings.h
@@ -4,6 +4,17 @@
 
 #include "EditorPerProjectUserSettings.generated.h"
 
+// Fbx export compatibility
+UENUM()
+enum class EFbxExportCompatibility : uint8
+{
+	FBX_2010,
+	FBX_2011,
+	FBX_2012,
+	FBX_2013,
+	FBX_2014,
+	FBX_2016
+};
 
 UCLASS(minimalapi, autoexpandcategories=(ViewportControls, ViewportLookAndFeel, LevelEditing, SourceControl, Content, Startup),
 	   hidecategories=(Object, Options, Grid, RotationGrid),
@@ -67,6 +78,10 @@ class UEditorPerProjectUserSettings : public UObject
 	UPROPERTY(EditAnywhere, config, Category=Export)
 	uint32 bKeepAttachHierarchy:1;
 
+	/** This will set the fbx sdk compatibility when exporting to fbx file. The default value is 2013 */
+	UPROPERTY(EditAnywhere, config, Category = Export)
+	EFbxExportCompatibility FbxExportCompatibility;
+
 	/** Select to make Distributions use the curves, not the baked lookup tables. */
 	UPROPERTY(config)
 	uint32 bUseCurvesForDistributions:1; //(GDistributionType == 0)
diff --git a/Engine/Source/Editor/UnrealEd/Classes/Factories/FbxMeshImportData.h b/Engine/Source/Editor/UnrealEd/Classes/Factories/FbxMeshImportData.h
--- a/Engine/Source/Editor/UnrealEd/Classes/Factories/FbxMeshImportData.h
+++ b/Engine/Source/Editor/UnrealEd/Classes/Factories/FbxMeshImportData.h
@@ -33,6 +33,14 @@ class UFbxMeshImportData : public UFbxAssetImportData
 {
 	GENERATED_UCLASS_BODY()
 
+	/** If this option is true the node absolute transform (transform, offset and pivot) will be apply to the mesh vertices. */
+	UPROPERTY(EditAnywhere, AdvancedDisplay, config, Category = ImportSettings, meta = (ImportType = "Mesh"))
+	bool bTransformVertexToAbsolute;
+
+	/** - Experimental - If this option is true the inverse node pivot will be apply to the mesh vertices. The pivot from the DCC will then be the origin of the mesh.*/
+	UPROPERTY(EditAnywhere, AdvancedDisplay, config, Category = ImportSettings, meta = (ImportType = "Mesh"))
+	bool bBakePivotInVertex;
+
 	/** Enables importing of mesh LODs from FBX LOD groups, if present in the FBX file */
 	UPROPERTY(EditAnywhere, AdvancedDisplay, config, Category=ImportSettings, meta=(OBJRestrict="true", ImportType="Mesh", ToolTip="If enabled, creates LOD models for Unreal meshes from LODs in the import file; If not enabled, only the base mesh from the LOD group is imported"))
 	uint32 bImportMeshLODs:1;
diff --git a/Engine/Source/Editor/UnrealEd/Classes/Factories/FbxSceneImportFactory.h b/Engine/Source/Editor/UnrealEd/Classes/Factories/FbxSceneImportFactory.h
--- a/Engine/Source/Editor/UnrealEd/Classes/Factories/FbxSceneImportFactory.h
+++ b/Engine/Source/Editor/UnrealEd/Classes/Factories/FbxSceneImportFactory.h
@@ -17,6 +17,8 @@ public:
 	FString OriginalImportPath;
 	FString OriginalFullImportName;
 
+	bool bOriginalTypeChanged;
+
 	bool bOverridePath;
 	FString OverrideImportPath;
 	FString OverrideFullImportName;
@@ -24,6 +26,13 @@ public:
 	//The name of the Options so reimport can show the options
 	FString OptionName;
 
+	//The node pivot user want to bake by default its the first node that reference the mesh
+	uint64 PivotNodeUid;
+	FString PivotNodeName; //Transient variable use for the ui to not display a uid
+
+	//If there is more then one Pivots we have to present a list to the user so he can choose which node pivot he want to bake
+	TMap<FVector, TArray<uint64>> NodeReferencePivots;
+
 	FString GetImportPath()
 	{
 		if (bOverridePath)
@@ -65,10 +74,13 @@ public:
 		, bImportAttribute(true)
 		, OriginalImportPath(TEXT(""))
 		, OriginalFullImportName(TEXT(""))
+		, bOriginalTypeChanged(false)
 		, bOverridePath(false)
 		, OverrideImportPath(TEXT(""))
 		, OverrideFullImportName(TEXT(""))
 		, OptionName(TEXT(""))
+		, PivotNodeUid(INVALID_UNIQUE_ID)
+		, PivotNodeName(TEXT("-"))
 		, IsContentObjectUpToDate(false)
 		, ContentPackage(nullptr)
 		, ContentObject(nullptr)
@@ -100,6 +112,8 @@ public:
 	FString LODGroup;
 	int32 LODLevel;
 	int32 MorphNum;
+	bool IsLod;
+	bool IsCollision;
 
 	FFbxMeshInfo()
 		: FaceNum(0)
@@ -112,6 +126,8 @@ public:
 		, LODGroup(TEXT(""))
 		, LODLevel(0)
 		, MorphNum(0)
+		, IsLod(false)
+		, IsCollision(false)
 	{}
 
 	virtual ~FFbxMeshInfo() {}
@@ -147,6 +163,84 @@ public:
 	virtual UClass *GetType();
 };
 
+class FFbxCameraInfo : public TSharedFromThis<FFbxCameraInfo>
+{
+public:
+	FString Name;
+	uint64 UniqueId;
+
+	bool ProjectionPerspective;
+	float AspectWidth;
+	float AspectHeight;
+	float NearPlane;
+	float FarPlane;
+	float OrthoZoom;
+	float FieldOfView;
+	
+
+	FFbxCameraInfo()
+		: Name(TEXT(""))
+		, UniqueId(INVALID_UNIQUE_ID)
+		, ProjectionPerspective(true)
+		, AspectWidth(0.0f)
+		, AspectHeight(0.0f)
+		, NearPlane(0.0f)
+		, FarPlane(0.0f)
+		, OrthoZoom(0.0f)
+		, FieldOfView(0.0f)
+	{}
+};
+
+class FFbxLightInfo : public TSharedFromThis<FFbxLightInfo>
+{
+public:
+	FString Name;
+	uint64 UniqueId;
+	int32 Type; //ePoint=0, eDirectional=1, eSpot=2, eArea=3, eVolume=4
+	FColor Color; //RGB color no alpha
+	float Intensity; //fbx default is 100
+	int32 Decay; //eNone=0, eLinear=1, eQuadratic=2, eCubic=3
+	bool CastLight;
+	bool CastShadow;
+	FColor ShadowColor; //RGB color no alpha
+
+	float InnerAngle;
+	float OuterAngle;
+	float Fog;
+	float DecayStart;
+	bool EnableNearAttenuation;
+	float NearAttenuationStart;
+	float NearAttenuationEnd;
+	bool EnableFarAttenuation;
+	float FarAttenuationStart;
+	float FarAttenuationEnd;
+
+	//Notes:
+	//Fbx use positive X to point light direction, we have to turn the component 90 degree in z object space
+
+	FFbxLightInfo()
+		: Name(TEXT(""))
+		, UniqueId(INVALID_UNIQUE_ID)
+		, Type(0)
+		, Color(FColor::White)
+		, Intensity(1.0f)
+		, Decay(0)
+		, CastLight(false)
+		, CastShadow(false)
+		, ShadowColor(FColor::Black)
+		, InnerAngle(0.0f)
+		, OuterAngle(0.0f)
+		, Fog(0.0f)
+		, DecayStart(0.0f)
+		, EnableNearAttenuation(false)
+		, NearAttenuationStart(0.0f)
+		, NearAttenuationEnd(0.0f)
+		, EnableFarAttenuation(false)
+		, FarAttenuationStart(0.0f)
+		, FarAttenuationEnd(0.0f)
+	{}
+};
+
 //Node use to store the scene hierarchy transform will be relative to the parent
 class FFbxNodeInfo : public TSharedFromThis<FFbxNodeInfo>
 {
@@ -158,9 +252,12 @@ public:
 	TSharedPtr<FFbxNodeInfo> ParentNodeInfo;
 	
 	TSharedPtr<FFbxAttributeInfo> AttributeInfo;
+	uint64 AttributeUniqueId;
 	FString AttributeType;
 
 	FTransform Transform;
+	FVector PivotRotation;
+	FVector PivotScaling;
 	bool bImportNode;
 
 	TArray<TSharedPtr<FFbxNodeInfo>> Childrens;
@@ -172,6 +269,7 @@ public:
 		, NodeHierarchyPath(TEXT(""))
 		, ParentNodeInfo(NULL)
 		, AttributeInfo(NULL)
+		, AttributeUniqueId(INVALID_UNIQUE_ID)
 		, AttributeType(TEXT(""))
 		, Transform(FTransform::Identity)
 		, bImportNode(true)
@@ -191,6 +289,11 @@ public:
 	int32 TotalTextureNum;
 	TArray<TSharedPtr<FFbxMeshInfo>> MeshInfo;
 	TArray<TSharedPtr<FFbxNodeInfo>> HierarchyInfo;
+	
+	//Component attributes
+	TMap<uint64, TSharedPtr<FFbxLightInfo>> LightInfo;
+	TMap<uint64, TSharedPtr<FFbxCameraInfo>> CameraInfo;
+	
 	/* true if it has animation */
 	bool bHasAnimation;
 	double FrameRate;
@@ -206,6 +309,17 @@ public:
 		, FrameRate(0.0)
 		, TotalTime(0.0)
 	{}
+
+	//Function helper to find the LOD parent of a nodeinfo with a mesh attribute
+	static TSharedPtr<FFbxNodeInfo> RecursiveFindLODParentNode(TSharedPtr<FFbxNodeInfo> NodeInfo)
+	{
+		if (!NodeInfo.IsValid())
+			return nullptr;
+		if (NodeInfo->ParentNodeInfo.IsValid() && NodeInfo->ParentNodeInfo->AttributeType.Compare(TEXT("eLODGroup")) == 0)
+			return NodeInfo->ParentNodeInfo;
+
+		return RecursiveFindLODParentNode(NodeInfo->ParentNodeInfo);
+	}
 };
 
 namespace UnFbx
@@ -238,15 +352,6 @@ class UNREALED_API UFbxSceneImportFactory : public UFactory
 	UPROPERTY(Transient)
 	class UFbxSceneImportOptionsSkeletalMesh* SceneImportOptionsSkeletalMesh;
 	
-	/** Import options UI detail when importing fbx scene animation*/
-	UPROPERTY(Transient)
-	class UFbxSceneImportOptionsAnimation* SceneImportOptionsAnimation;
-	
-	/** Import options UI detail when importing fbx scene material*/
-	UPROPERTY(Transient)
-	class UFbxSceneImportOptionsMaterial* SceneImportOptionsMaterial;
-
-
 	/** Import data used when importing static meshes */
 	UPROPERTY(Transient)
 	class UFbxStaticMeshImportData* StaticMeshImportData;
@@ -284,13 +389,15 @@ protected:
 	* LocalTransform: Must be a valid pointer
 	* PreMultiplyTransform: Can be nullptr
 	*/
-	void ApplyTransformToComponent(USceneComponent *SceneComponent, FTransform *LocalTransform, FTransform *PreMultiplyTransform);
+	void ApplyTransformToComponent(USceneComponent *SceneComponent, FTransform *LocalTransform, FTransform *PreMultiplyTransform, FVector &PivotLocation);
 
 	/** This will add the scene transform options to the root node */
 	void ApplySceneTransformOptionsToRootNode(TSharedPtr<FFbxSceneInfo> SceneInfoPtr);
 
 	/** Import all skeletal mesh from the fbx scene */
 	void ImportAllSkeletalMesh(void* VoidRootNodeToImport, void* VoidFbxImporter, EObjectFlags Flags, int32& NodeIndex, int32& InterestingNodeCount , TSharedPtr<FFbxSceneInfo> SceneInfo);
+	
+	UObject* ImportOneSkeletalMesh(void* VoidRootNodeToImport, void* VoidFbxImporter, TSharedPtr<FFbxSceneInfo> SceneInfo, EObjectFlags Flags, TArray<void*> &VoidNodeArray, int32 &TotalNumNodes);
 
 	/** Import all static mesh from the fbx scene */
 	void ImportAllStaticMesh(void* VoidRootNodeToImport, void* VoidFbxImporter, EObjectFlags Flags, int32& NodeIndex, int32& InterestingNodeCount, TSharedPtr<FFbxSceneInfo> SceneInfo);
@@ -307,7 +414,7 @@ protected:
 	/** Create a package for the specified node. Package will be the concatenation of UFbxSceneImportFactory::Path and Node->GetName(). */
 	UPackage *CreatePackageForNode(FString PackageName, FString &StaticMeshName);
 
-	static TSharedPtr<FFbxSceneInfo> ConvertSceneInfo(void* VoidFbxSceneInfo);
+	static TSharedPtr<FFbxSceneInfo> ConvertSceneInfo(void* VoidFbxImporter, void* VoidFbxSceneInfo);
 	static void ExtractMaterialInfo(void* FbxImporterVoid, TSharedPtr<FFbxSceneInfo> SceneInfoPtr);
 	bool SetStaticMeshComponentOverrideMaterial(class UStaticMeshComponent* StaticMeshComponent, TSharedPtr<FFbxNodeInfo> NodeInfo);
 
diff --git a/Engine/Source/Editor/UnrealEd/Classes/Factories/FbxSceneImportOptions.h b/Engine/Source/Editor/UnrealEd/Classes/Factories/FbxSceneImportOptions.h
--- a/Engine/Source/Editor/UnrealEd/Classes/Factories/FbxSceneImportOptions.h
+++ b/Engine/Source/Editor/UnrealEd/Classes/Factories/FbxSceneImportOptions.h
@@ -51,6 +51,14 @@ class UFbxSceneImportOptions : public UObject
 	UPROPERTY(EditAnywhere, config, Category = Transform)
 	float ImportUniformScale;
 
+	/** If this option is true the node absolute transform (transform, offset and pivot) will be apply to the mesh vertices. */
+	UPROPERTY()
+	bool bTransformVertexToAbsolute;
+
+	/** - Experimental - If this option is true the inverse node pivot will be apply to the mesh vertices. The pivot from the DCC will then be the origin of the mesh. */
+	UPROPERTY(EditAnywhere, config, category = Meshes)
+	bool bBakePivotInVertex;
+
 	/** If enabled, creates LOD models for Unreal static meshes from LODs in the import file; If not enabled, only the base static mesh from the LOD group is imported. */
 	UPROPERTY(EditAnywhere, config, category = Meshes)
 	uint32 bImportStaticMeshLODs : 1;
@@ -59,13 +67,9 @@ class UFbxSceneImportOptions : public UObject
 	UPROPERTY(EditAnywhere, config, category = Meshes)
 	uint32 bImportSkeletalMeshLODs : 1;
 
-	/** TODO - Name of the parent actor or the blueprint. This is ignore if the user choose to create level actors or an existing blueprint */
-	//UPROPERTY()
-	//FString ParentName;
-
-	/** TODO - If user choose existing blueprint, this is the path of the blueprint asset */
-	//UPROPERTY()
-	//FString ExistingBlueprintPath
+	/** If either importing of textures (or materials) is enabled, this option will cause normal map values to be inverted */
+	UPROPERTY(EditAnywhere, config, Category = Textures)
+	uint32 bInvertNormalMaps : 1;
 };
 
 
diff --git a/Engine/Source/Editor/UnrealEd/Classes/Factories/FbxSceneImportOptionsSkeletalMesh.h b/Engine/Source/Editor/UnrealEd/Classes/Factories/FbxSceneImportOptionsSkeletalMesh.h
--- a/Engine/Source/Editor/UnrealEd/Classes/Factories/FbxSceneImportOptionsSkeletalMesh.h
+++ b/Engine/Source/Editor/UnrealEd/Classes/Factories/FbxSceneImportOptionsSkeletalMesh.h
@@ -40,10 +40,38 @@ class UFbxSceneImportOptionsSkeletalMesh : public UObject
 	UPROPERTY(EditAnywhere, config, Category = SkeletalMesh)
 	uint32 bKeepOverlappingVertices : 1;
 
+	//////////////////////////////////////////////////////////////////////////
+	// Animation section
+
+	/** True to import animations from the FBX File */
+	UPROPERTY(EditAnywhere, config, Category = Animation)
+	uint32 bImportAnimations : 1;
+
+	/** Type of asset to import from the FBX file */
+	UPROPERTY(EditAnywhere, Category = Animation, config, meta = (DisplayName = "Animation Length"))
+	TEnumAsByte<enum EFBXAnimationLengthImportType> AnimationLength;
+
+	/** Frame range used when Set Range is used in Animation Length */
+	UPROPERTY(EditAnywhere, AdvancedDisplay, Category = Animation, meta = (UIMin = 0, ClampMin = 0))
+	FInt32Interval FrameImportRange;
+
+	/** Enable this option to use default sample rate for the imported animation at 30 frames per second */
+	UPROPERTY(EditAnywhere, AdvancedDisplay, config, Category = Animation, meta = (ToolTip = "If enabled, samples all animation curves to 30 FPS"))
+	bool bUseDefaultSampleRate;
+
+	/** Import if custom attribute as a curve within the animation **/
+	UPROPERTY(EditAnywhere, AdvancedDisplay, config, Category = Animation)
+	bool bImportCustomAttribute;
+
+	/** Type of asset to import from the FBX file */
+	UPROPERTY(EditAnywhere, AdvancedDisplay, config, Category = Animation)
+	bool bPreserveLocalTransform;
+
+	/** Type of asset to import from the FBX file */
+	UPROPERTY(EditAnywhere, AdvancedDisplay, config, Category = Animation)
+	bool bDeleteExistingMorphTargetCurves;
 
-	/** Enables experimental Mikk tangent generation for skeletal meshes */
-	/*UPROPERTY(EditAnywhere, Category = SkeletalMesh)
-	uint32 bUseExperimentalTangentGeneration : 1;*/
+	void FillSkeletalMeshInmportData(class UFbxSkeletalMeshImportData* SkeletalMeshImportData, class UFbxAnimSequenceImportData* AnimSequenceImportData);
 };
 
 
diff --git a/Engine/Source/Editor/UnrealEd/Classes/Factories/FbxSceneImportOptionsStaticMesh.h b/Engine/Source/Editor/UnrealEd/Classes/Factories/FbxSceneImportOptionsStaticMesh.h
--- a/Engine/Source/Editor/UnrealEd/Classes/Factories/FbxSceneImportOptionsStaticMesh.h
+++ b/Engine/Source/Editor/UnrealEd/Classes/Factories/FbxSceneImportOptionsStaticMesh.h
@@ -53,10 +53,6 @@ class UFbxSceneImportOptionsStaticMesh : public UObject
 	//////////////////////////////////////////////////////////////////////////
 	// Static Mesh section
 
-	/** If this option is true the staticmesh absolute transform will be apply to the staticmesh vertices. */
-	UPROPERTY()
-	bool bTransformVertexToAbsolute;
-
 	/** For static meshes, enabling this option will combine all meshes in the FBX into a single monolithic mesh in Unreal */
 	UPROPERTY()
 	FName StaticMeshLODGroup;
diff --git a/Engine/Source/Editor/UnrealEd/Classes/Factories/FbxStaticMeshImportData.h b/Engine/Source/Editor/UnrealEd/Classes/Factories/FbxStaticMeshImportData.h
--- a/Engine/Source/Editor/UnrealEd/Classes/Factories/FbxStaticMeshImportData.h
+++ b/Engine/Source/Editor/UnrealEd/Classes/Factories/FbxStaticMeshImportData.h
@@ -26,10 +26,6 @@ class UFbxStaticMeshImportData : public UFbxMeshImportData
 {
 	GENERATED_UCLASS_BODY()
 
-	/** If this option is true the staticmesh absolute transform will be apply to the staticmesh vertices. */
-	UPROPERTY(EditAnywhere, AdvancedDisplay, Category = ImportSettings, meta = (ImportType = "StaticMesh"))
-	bool bTransformVertexToAbsolute;
-
 	/** For static meshes, enabling this option will combine all meshes in the FBX into a single monolithic mesh in Unreal */
 	UPROPERTY(EditAnywhere, AdvancedDisplay, Category=ImportSettings, meta=(ImportType="StaticMesh"))
 	FName StaticMeshLODGroup;
diff --git a/Engine/Source/Editor/UnrealEd/Classes/Factories/ReimportFbxSceneFactory.h b/Engine/Source/Editor/UnrealEd/Classes/Factories/ReimportFbxSceneFactory.h
--- a/Engine/Source/Editor/UnrealEd/Classes/Factories/ReimportFbxSceneFactory.h
+++ b/Engine/Source/Editor/UnrealEd/Classes/Factories/ReimportFbxSceneFactory.h
@@ -24,7 +24,9 @@ class UReimportFbxSceneFactory : public UFbxSceneImportFactory, public FReimport
 	//~ End UFactory Interface
 private:
 	EReimportResult::Type ReimportStaticMesh(void* VoidFbxImporter, TSharedPtr<FFbxMeshInfo> MeshInfo);
+	EReimportResult::Type ReimportSkeletalMesh(void* VoidFbxImporter, TSharedPtr<FFbxMeshInfo> MeshInfo);
 	EReimportResult::Type ImportStaticMesh(void* VoidFbxImporter, TSharedPtr<FFbxMeshInfo> MeshInfo, TSharedPtr<FFbxSceneInfo> SceneInfoPtr);
+	EReimportResult::Type ImportSkeletalMesh(void* VoidRootNodeToImport, void* VoidFbxImporter, TSharedPtr<FFbxMeshInfo> MeshInfo, TSharedPtr<FFbxSceneInfo> SceneInfoPtr);
 	UBlueprint *UpdateOriginalBluePrint(FString &BluePrintFullName, void* VoidNodeStatusMapPtr, TSharedPtr<FFbxSceneInfo> SceneInfoPtr, TSharedPtr<FFbxSceneInfo> SceneInfoOriginalPtr, TArray<FAssetData> &AssetDataToDelete);
 
 	FString FbxImportFileName;
diff --git a/Engine/Source/Editor/UnrealEd/Private/EditorEngine.cpp b/Engine/Source/Editor/UnrealEd/Private/EditorEngine.cpp
--- a/Engine/Source/Editor/UnrealEd/Private/EditorEngine.cpp
+++ b/Engine/Source/Editor/UnrealEd/Private/EditorEngine.cpp
@@ -5503,7 +5503,10 @@ void UEditorEngine::NotifyToolsOfObjectReplacement(const TMap<UObject*, UObject*
 		PropertyEditorModule.ReplaceViewedObjects( OldToNewInstanceMap );
 	}
 
-	// Check to see if any selected components were reinstanced
+	// Allow any other observers to act upon the object replacement
+	BroadcastObjectsReplaced(OldToNewInstanceMap);
+
+	// Check to see if any selected components were reinstanced, as a final step.
 	USelection* ComponentSelection = GetSelectedComponents();
 	if (ComponentSelection)
 	{
@@ -5527,8 +5530,6 @@ void UEditorEngine::NotifyToolsOfObjectReplacement(const TMap<UObject*, UObject*
 		}
 		ComponentSelection->EndBatchSelectOperation();
 	}
-	
-	BroadcastObjectsReplaced(OldToNewInstanceMap);
 }
 
 void UEditorEngine::DisableRealtimeViewports()
diff --git a/Engine/Source/Editor/UnrealEd/Private/EditorServer.cpp b/Engine/Source/Editor/UnrealEd/Private/EditorServer.cpp
--- a/Engine/Source/Editor/UnrealEd/Private/EditorServer.cpp
+++ b/Engine/Source/Editor/UnrealEd/Private/EditorServer.cpp
@@ -3403,12 +3403,30 @@ void UEditorEngine::PasteSelectedActorsFromClipboard( UWorld* InWorld, const FTe
 				// List of group actors in the selection
 				TArray<AGroupActor*> GroupActors;
 
-				// Move the actors.
+				struct FAttachData
+				{
+					FAttachData(AActor* InParentActor, FName InSocketName)
+						: ParentActor(InParentActor)
+						, SocketName(InSocketName)
+					{}
+
+					AActor* ParentActor;
+					FName SocketName;
+				};
+
+				TArray<FAttachData, TInlineAllocator<8>> AttachData;
+				AttachData.Reserve(NumActorsToMove);
+
+				// Break any parent attachments and move the actors.
 				AActor* SingleActor = NULL;
 				for ( FSelectionIterator It( GEditor->GetSelectedActorIterator() ) ; It ; ++It )
 				{
 					AActor* Actor = static_cast<AActor*>( *It );
-					checkSlow( Actor->IsA(AActor::StaticClass()) );
+
+					AActor* ParentActor = Actor->GetAttachParentActor();
+					FName SocketName = Actor->GetAttachParentSocketName();
+					Actor->DetachRootComponentFromParent(true);
+					AttachData.Emplace(ParentActor, SocketName);
 
 					// If this actor is in a group, add it to the list
 					if (GEditor->bGroupingActive)
@@ -3422,7 +3440,16 @@ void UEditorEngine::PasteSelectedActorsFromClipboard( UWorld* InWorld, const FTe
 
 					SingleActor = Actor;
 					Actor->SetActorLocation(Actor->GetActorLocation() + Adjust, false);
+				}
+
+				// Restore attachments
+				int Index = 0;
+				for (FSelectionIterator It(GEditor->GetSelectedActorIterator()); It; ++It)
+				{
+					AActor* Actor = static_cast<AActor*>(*It);
+					Actor->AttachRootComponentToActor(AttachData[Index].ParentActor, AttachData[Index].SocketName, EAttachLocation::KeepWorldPosition);
 					Actor->PostEditMove(true);
+					Index++;
 				}
 
 				// Update the pivot location.
diff --git a/Engine/Source/Editor/UnrealEd/Private/Factories/EditorFactories.cpp b/Engine/Source/Editor/UnrealEd/Private/Factories/EditorFactories.cpp
--- a/Engine/Source/Editor/UnrealEd/Private/Factories/EditorFactories.cpp
+++ b/Engine/Source/Editor/UnrealEd/Private/Factories/EditorFactories.cpp
@@ -5740,11 +5740,24 @@ bool UReimportFbxSkeletalMeshFactory::FactoryCanImport(const FString& Filename)
 }
 
 bool UReimportFbxSkeletalMeshFactory::CanReimport( UObject* Obj, TArray<FString>& OutFilenames )
-{	
+{
 	USkeletalMesh* SkeletalMesh = Cast<USkeletalMesh>(Obj);
-	if(SkeletalMesh && !Obj->IsA( UDestructibleMesh::StaticClass() ))
+	if (SkeletalMesh && !Obj->IsA(UDestructibleMesh::StaticClass()))
 	{
-		SkeletalMesh->AssetImportData->ExtractFilenames(OutFilenames);
+		if (SkeletalMesh->AssetImportData)
+		{
+			UFbxAssetImportData *FbxAssetImportData = Cast<UFbxAssetImportData>(SkeletalMesh->AssetImportData);
+			if (FbxAssetImportData != nullptr && FbxAssetImportData->bImportAsScene)
+			{
+				//This skeletal mesh was import with a scene import, we cannot reimport it here
+				return false;
+			}
+			SkeletalMesh->AssetImportData->ExtractFilenames(OutFilenames);
+		}
+		else
+		{
+			OutFilenames.Add(TEXT(""));
+		}
 		return true;
 	}
 	return false;
diff --git a/Engine/Source/Editor/UnrealEd/Private/Fbx/FbxFactory.cpp b/Engine/Source/Editor/UnrealEd/Private/Fbx/FbxFactory.cpp
--- a/Engine/Source/Editor/UnrealEd/Private/Fbx/FbxFactory.cpp
+++ b/Engine/Source/Editor/UnrealEd/Private/Fbx/FbxFactory.cpp
@@ -369,11 +369,15 @@ UObject* UFbxFactory::FactoryCreateBinary
 								{
 									if (Node->GetChildCount() > LODIndex)
 									{
-										SkelMeshNodeArray.Add(Node->GetChild(LODIndex));
+										FbxNode *MeshNode = FbxImporter->FindLODGroupNode(Node, LODIndex);
+										if(MeshNode != nullptr)
+											SkelMeshNodeArray.Add(MeshNode);
 									}
 									else // in less some LODGroups have less level, use the last level
 									{
-										SkelMeshNodeArray.Add(Node->GetChild(Node->GetChildCount() - 1));
+										FbxNode *MeshNode = FbxImporter->FindLODGroupNode(Node, Node->GetChildCount() - 1);
+										if(MeshNode != nullptr)
+											SkelMeshNodeArray.Add(MeshNode);
 									}
 								}
 								else
@@ -505,12 +509,14 @@ UObject* UFbxFactory::FactoryCreateBinary
 UObject* UFbxFactory::RecursiveImportNode(void* VoidFbxImporter, void* VoidNode, UObject* InParent, FName InName, EObjectFlags Flags, int32& NodeIndex, int32 Total, TArray<UObject*>& OutNewAssets)
 {
 	UObject* NewObject = NULL;
-
+	UnFbx::FFbxImporter *FbxImporter = (UnFbx::FFbxImporter *)VoidFbxImporter;
 	FbxNode* Node = (FbxNode*)VoidNode;
 	if (Node->GetNodeAttribute() && Node->GetNodeAttribute()->GetAttributeType() == FbxNodeAttribute::eLODGroup && Node->GetChildCount() > 0 )
 	{
 		// import base mesh
-		NewObject = ImportANode(VoidFbxImporter, Node->GetChild(0), InParent, InName, Flags, NodeIndex, Total);
+		FbxNode *MeshNode = FbxImporter->FindLODGroupNode(Node, 0);
+		if(MeshNode != nullptr)
+			NewObject = ImportANode(VoidFbxImporter, MeshNode, InParent, InName, Flags, NodeIndex, Total);
 
 		if ( NewObject )
 		{
@@ -536,8 +542,9 @@ UObject* UFbxFactory::RecursiveImportNode(void* VoidFbxImporter, void* VoidNode,
 			// import LOD meshes
 			for (int32 LODIndex = 1; LODIndex < Node->GetChildCount(); LODIndex++)
 			{
-				FbxNode* ChildNode = Node->GetChild(LODIndex);
-				ImportANode(VoidFbxImporter, ChildNode, InParent, InName, Flags, NodeIndex, Total, NewObject, LODIndex);
+				FbxNode* ChildNode = FbxImporter->FindLODGroupNode(Node, LODIndex);
+				if(ChildNode != nullptr)
+					ImportANode(VoidFbxImporter, ChildNode, InParent, InName, Flags, NodeIndex, Total, NewObject, LODIndex);
 			}
 		}
 	}
diff --git a/Engine/Source/Editor/UnrealEd/Private/Fbx/FbxMainExport.cpp b/Engine/Source/Editor/UnrealEd/Private/Fbx/FbxMainExport.cpp
--- a/Engine/Source/Editor/UnrealEd/Private/Fbx/FbxMainExport.cpp
+++ b/Engine/Source/Editor/UnrealEd/Private/Fbx/FbxMainExport.cpp
@@ -39,6 +39,8 @@
 #include "Materials/MaterialExpressionConstant2Vector.h"
 #include "Materials/MaterialExpressionConstant3Vector.h"
 #include "Materials/MaterialExpressionConstant4Vector.h"
+#include "Materials/MaterialExpressionTextureBase.h"
+#include "Materials/MaterialExpressionTextureSample.h"
 
 #include "Matinee/InterpData.h"
 #include "Matinee/InterpTrackMove.h"
@@ -185,9 +187,34 @@ void FFbxExporter::WriteToFile(const TCHAR* Filename)
 	IOS_REF.SetBoolProp(EXP_FBX_ANIMATION,       true);
 	IOS_REF.SetBoolProp(EXP_FBX_GLOBAL_SETTINGS, true);
 
+	//Get the compatibility from the editor settings
+	const char* CompatibilitySetting = FBX_2013_00_COMPATIBLE;
+	const EFbxExportCompatibility FbxExportCompatibility = GetDefault<UEditorPerProjectUserSettings>()->FbxExportCompatibility;
+	switch (FbxExportCompatibility)
+	{
+		case EFbxExportCompatibility::FBX_2010:
+			CompatibilitySetting = FBX_2010_00_COMPATIBLE;
+			break;
+		case EFbxExportCompatibility::FBX_2011:
+			CompatibilitySetting = FBX_2011_00_COMPATIBLE;
+			break;
+		case EFbxExportCompatibility::FBX_2012:
+			CompatibilitySetting = FBX_2012_00_COMPATIBLE;
+			break;
+		case EFbxExportCompatibility::FBX_2013:
+			CompatibilitySetting = FBX_2013_00_COMPATIBLE;
+			break;
+		case EFbxExportCompatibility::FBX_2014:
+			CompatibilitySetting = FBX_2014_00_COMPATIBLE;
+			break;
+		case EFbxExportCompatibility::FBX_2016:
+			CompatibilitySetting = FBX_2016_00_COMPATIBLE;
+			break;
+	}
+	
 	// We export using FBX 2013 format because many users are still on that version and FBX 2014 files has compatibility issues with
 	// normals when importing to an earlier version of the plugin
-	if (!Exporter->SetFileExportVersion(FBX_2013_00_COMPATIBLE, FbxSceneRenamer::eNone))
+	if (!Exporter->SetFileExportVersion(CompatibilitySetting, FbxSceneRenamer::eNone))
 	{
 		UE_LOG(LogFbx, Warning, TEXT("Call to KFbxExporter::SetFileExportVersion(FBX_2013_00_COMPATIBLE) to export 2013 fbx file format failed.\n"));
 	}
@@ -351,7 +378,7 @@ void FFbxExporter::ExportLevelMesh( ULevel* InLevel, AMatineeActor* InMatineeAct
 		}
 		else if(Actor->IsA(ACameraActor::StaticClass()))
 		{
-			ExportCamera(CastChecked<ACameraActor>(Actor), InMatineeActor, true); // Just export the placement of the particle emitter.
+			ExportCamera(CastChecked<ACameraActor>(Actor), InMatineeActor, false); // Just export the placement of the particle emitter.
 		}
 		else if( Actor != NULL )
 		{
@@ -361,77 +388,101 @@ void FFbxExporter::ExportLevelMesh( ULevel* InLevel, AMatineeActor* InMatineeAct
 	}
 }
 
-/**
- * Exports the light-specific information for a light actor.
- */
-void FFbxExporter::ExportLight( ALight* Actor, AMatineeActor* InMatineeActor )
+void FFbxExporter::FillFbxLightAttribute(FbxLight* Light, FbxNode* FbxParentNode, ULightComponent* BaseLight)
 {
-	if (Scene == NULL || Actor == NULL || !Actor->GetLightComponent()) return;
-
-	// Export the basic actor information.
-	FbxNode* FbxActor = ExportActor( Actor, InMatineeActor ); // this is the pivot node
-	// The real fbx light node
-	FbxNode* FbxLightNode = FbxActor->GetParent();
-
-	ULightComponent* BaseLight = Actor->GetLightComponent();
-
-	FString FbxNodeName = GetActorNodeName(Actor, InMatineeActor);
-
-	// Export the basic light information
-	FbxLight* Light = FbxLight::Create(Scene, TCHAR_TO_UTF8(*FbxNodeName));
 	Light->Intensity.Set(BaseLight->Intensity);
 	Light->Color.Set(Converter.ConvertToFbxColor(BaseLight->LightColor));
-	
+
 	// Add one user property for recording the Brightness animation
-	CreateAnimatableUserProperty(FbxLightNode, BaseLight->Intensity, "UE_Intensity", "UE_Matinee_Light_Intensity");
-	
+	CreateAnimatableUserProperty(FbxParentNode, BaseLight->Intensity, "UE_Intensity", "UE_Matinee_Light_Intensity");
+
 	// Look for the higher-level light types and determine the lighting method
 	if (BaseLight->IsA(UPointLightComponent::StaticClass()))
 	{
-		UPointLightComponent* PointLight = (UPointLightComponent*) BaseLight;
+		UPointLightComponent* PointLight = (UPointLightComponent*)BaseLight;
 		if (BaseLight->IsA(USpotLightComponent::StaticClass()))
 		{
-			USpotLightComponent* SpotLight = (USpotLightComponent*) BaseLight;
+			USpotLightComponent* SpotLight = (USpotLightComponent*)BaseLight;
 			Light->LightType.Set(FbxLight::eSpot);
 
 			// Export the spot light parameters.
-			if (!FMath::IsNearlyZero(SpotLight->InnerConeAngle))
+			if (!FMath::IsNearlyZero(SpotLight->InnerConeAngle*2.0f))
 			{
-				Light->InnerAngle.Set(SpotLight->InnerConeAngle);
+				Light->InnerAngle.Set(SpotLight->InnerConeAngle*2.0f);
 			}
 			else // Maya requires a non-zero inner cone angle
 			{
 				Light->InnerAngle.Set(0.01f);
 			}
-			Light->OuterAngle.Set(SpotLight->OuterConeAngle);
+			Light->OuterAngle.Set(SpotLight->OuterConeAngle*2.0f);
 		}
 		else
 		{
 			Light->LightType.Set(FbxLight::ePoint);
 		}
-		
+
 		// Export the point light parameters.
 		Light->EnableFarAttenuation.Set(true);
 		Light->FarAttenuationEnd.Set(PointLight->AttenuationRadius);
 		// Add one user property for recording the FalloffExponent animation
-		CreateAnimatableUserProperty(FbxLightNode, PointLight->AttenuationRadius, "UE_Radius", "UE_Matinee_Light_Radius");
-		
+		CreateAnimatableUserProperty(FbxParentNode, PointLight->AttenuationRadius, "UE_Radius", "UE_Matinee_Light_Radius");
+
 		// Add one user property for recording the FalloffExponent animation
-		CreateAnimatableUserProperty(FbxLightNode, PointLight->LightFalloffExponent, "UE_FalloffExponent", "UE_Matinee_Light_FalloffExponent");
+		CreateAnimatableUserProperty(FbxParentNode, PointLight->LightFalloffExponent, "UE_FalloffExponent", "UE_Matinee_Light_FalloffExponent");
 	}
 	else if (BaseLight->IsA(UDirectionalLightComponent::StaticClass()))
 	{
 		// The directional light has no interesting properties.
 		Light->LightType.Set(FbxLight::eDirectional);
+		Light->Intensity.Set(BaseLight->Intensity*100.0f);
 	}
-	
+}
+
+/**
+ * Exports the light-specific information for a light actor.
+ */
+void FFbxExporter::ExportLight( ALight* Actor, AMatineeActor* InMatineeActor )
+{
+	if (Scene == NULL || Actor == NULL || !Actor->GetLightComponent()) return;
+
+	// Export the basic actor information.
+	FbxNode* FbxActor = ExportActor( Actor, InMatineeActor ); // this is the pivot node
+	// The real fbx light node
+	FbxNode* FbxLightNode = FbxActor->GetParent();
+
+	ULightComponent* BaseLight = Actor->GetLightComponent();
+
+	FString FbxNodeName = GetActorNodeName(Actor, InMatineeActor);
+
+	// Export the basic light information
+	FbxLight* Light = FbxLight::Create(Scene, TCHAR_TO_UTF8(*FbxNodeName));
+	FillFbxLightAttribute(Light, FbxLightNode, BaseLight);	
 	FbxActor->SetNodeAttribute(Light);
 }
 
+void FFbxExporter::FillFbxCameraAttribute(FbxNode* ParentNode, FbxCamera* Camera, UCameraComponent *CameraComponent)
+{
+	// Export the view area information
+	Camera->ProjectionType.Set(CameraComponent->ProjectionMode == ECameraProjectionMode::Type::Perspective ? FbxCamera::ePerspective : FbxCamera::eOrthogonal);
+	Camera->SetAspect(FbxCamera::eFixedRatio, CameraComponent->AspectRatio, 1.0f);
+	Camera->FilmAspectRatio.Set(CameraComponent->AspectRatio);
+	Camera->SetApertureWidth(CameraComponent->AspectRatio * 0.612f); // 0.612f is a magic number from Maya that represents the ApertureHeight
+	Camera->SetApertureMode(FbxCamera::eFocalLength);
+	Camera->FocalLength.Set(Camera->ComputeFocalLength(CameraComponent->FieldOfView));
+
+	// Add one user property for recording the AspectRatio animation
+	CreateAnimatableUserProperty(ParentNode, CameraComponent->AspectRatio, "UE_AspectRatio", "UE_Matinee_Camera_AspectRatio");
+
+	// Push the near/far clip planes away, as the engine uses larger values than the default.
+	Camera->SetNearPlane(10.0f);
+	Camera->SetFarPlane(100000.0f);
+}
+
 void FFbxExporter::ExportCamera( ACameraActor* Actor, AMatineeActor* InMatineeActor, bool bExportComponents )
 {
 	if (Scene == NULL || Actor == NULL) return;
 
+	UCameraComponent *CameraComponent = Actor->GetCameraComponent();
 	// Export the basic actor information.
 	FbxNode* FbxActor = ExportActor( Actor, InMatineeActor, bExportComponents ); // this is the pivot node
 	// The real fbx camera node
@@ -441,21 +492,7 @@ void FFbxExporter::ExportCamera( ACameraActor* Actor, AMatineeActor* InMatineeAc
 
 	// Create a properly-named FBX camera structure and instantiate it in the FBX scene graph
 	FbxCamera* Camera = FbxCamera::Create(Scene, TCHAR_TO_UTF8(*FbxNodeName));
-
-	// Export the view area information
-	Camera->ProjectionType.Set(FbxCamera::ePerspective);
-	Camera->SetAspect(FbxCamera::eFixedRatio, Actor->GetCameraComponent()->AspectRatio, 1.0f);
-	Camera->FilmAspectRatio.Set(Actor->GetCameraComponent()->AspectRatio);
-	Camera->SetApertureWidth(Actor->GetCameraComponent()->AspectRatio * 0.612f); // 0.612f is a magic number from Maya that represents the ApertureHeight
-	Camera->SetApertureMode(FbxCamera::eFocalLength);
-	Camera->FocalLength.Set(Camera->ComputeFocalLength(Actor->GetCameraComponent()->FieldOfView));
-	
-	// Add one user property for recording the AspectRatio animation
-	CreateAnimatableUserProperty(FbxCameraNode, Actor->GetCameraComponent()->AspectRatio, "UE_AspectRatio", "UE_Matinee_Camera_AspectRatio");
-
-	// Push the near/far clip planes away, as the engine uses larger values than the default.
-	Camera->SetNearPlane(10.0f);
-	Camera->SetFarPlane(100000.0f);
+	FillFbxCameraAttribute(FbxCameraNode, Camera, CameraComponent);
 
 	FbxActor->SetNodeAttribute(Camera);
 
@@ -914,59 +951,105 @@ void FFbxExporter::ExportLandscape(ALandscapeProxy* Actor, bool bSelectedOnly)
 	ExportLandscapeToFbx(Actor, *FbxNodeName, FbxActor, bSelectedOnly);
 }
 
-FbxDouble3 SetMaterialComponent(FColorMaterialInput& MatInput)
+FbxDouble3 SetMaterialComponent(FColorMaterialInput& MatInput, bool ToLinear)
 {
-	FColor FinalColor;
+	FColor RGBColor;
+	FLinearColor LinearColor;
+	bool LinearSet = false;
 	
 	if (MatInput.Expression)
 	{
 		if (Cast<UMaterialExpressionConstant>(MatInput.Expression))
 		{
-			UMaterialExpressionConstant* Expr = Cast<UMaterialExpressionConstant>(MatInput.Expression);
-			FinalColor = FColor(Expr->R);
+			UMaterialExpressionConstant* Expr = Cast<UMaterialExpressionConstant>(MatInput.Expression);			
+			RGBColor = FColor(Expr->R);
 		}
 		else if (Cast<UMaterialExpressionVectorParameter>(MatInput.Expression))
 		{
 			UMaterialExpressionVectorParameter* Expr = Cast<UMaterialExpressionVectorParameter>(MatInput.Expression);
-			FinalColor = Expr->DefaultValue.ToFColor(true);
+			LinearColor = Expr->DefaultValue;
+			LinearSet = true;
+			//Linear to sRGB color space conversion
+			RGBColor = Expr->DefaultValue.ToFColor(true);
 		}
 		else if (Cast<UMaterialExpressionConstant3Vector>(MatInput.Expression))
 		{
 			UMaterialExpressionConstant3Vector* Expr = Cast<UMaterialExpressionConstant3Vector>(MatInput.Expression);
-			FinalColor.R = Expr->Constant.R;
-			FinalColor.G = Expr->Constant.G;
-			FinalColor.B = Expr->Constant.B;
+			RGBColor.R = Expr->Constant.R;
+			RGBColor.G = Expr->Constant.G;
+			RGBColor.B = Expr->Constant.B;
 		}
 		else if (Cast<UMaterialExpressionConstant4Vector>(MatInput.Expression))
 		{
 			UMaterialExpressionConstant4Vector* Expr = Cast<UMaterialExpressionConstant4Vector>(MatInput.Expression);
-			FinalColor.R = Expr->Constant.R;
-			FinalColor.G = Expr->Constant.G;
-			FinalColor.B = Expr->Constant.B;
-			//FinalColor.A = Expr->A;
+			RGBColor.R = Expr->Constant.R;
+			RGBColor.G = Expr->Constant.G;
+			RGBColor.B = Expr->Constant.B;
+			//RGBColor.A = Expr->A;
 		}
 		else if (Cast<UMaterialExpressionConstant2Vector>(MatInput.Expression))
 		{
 			UMaterialExpressionConstant2Vector* Expr = Cast<UMaterialExpressionConstant2Vector>(MatInput.Expression);
-			FinalColor.R = Expr->R;
-			FinalColor.G = Expr->G;
-			FinalColor.B = 0;
+			RGBColor.R = Expr->R;
+			RGBColor.G = Expr->G;
+			RGBColor.B = 0;
 		}
 		else
 		{
-			FinalColor.R = MatInput.Constant.R / 128.0;
-			FinalColor.G = MatInput.Constant.G / 128.0;
-			FinalColor.B = MatInput.Constant.B / 128.0;
+			RGBColor.R = MatInput.Constant.R;
+			RGBColor.G = MatInput.Constant.G;
+			RGBColor.B = MatInput.Constant.B;
 		}
 	}
 	else
 	{
-		FinalColor.R = MatInput.Constant.R / 128.0;
-		FinalColor.G = MatInput.Constant.G / 128.0;
-		FinalColor.B = MatInput.Constant.B / 128.0;
+		RGBColor.R = MatInput.Constant.R;
+		RGBColor.G = MatInput.Constant.G;
+		RGBColor.B = MatInput.Constant.B;
 	}
 	
-	return FbxDouble3(FinalColor.R, FinalColor.G, FinalColor.B);
+	if (ToLinear)
+	{
+		if (!LinearSet)
+		{
+			//sRGB to linear color space conversion
+			LinearColor = FLinearColor(RGBColor);
+		}
+		return FbxDouble3(LinearColor.R, LinearColor.G, LinearColor.B);
+	}
+	return FbxDouble3(RGBColor.R, RGBColor.G, RGBColor.B);
+}
+
+bool FFbxExporter::FillFbxTextureProperty(const char *PropertyName, const FExpressionInput& MaterialInput, FbxSurfaceMaterial* FbxMaterial)
+{
+	if (Scene == NULL)
+	{
+		return false;
+	}
+
+	FbxProperty FbxColorProperty = FbxMaterial->FindProperty(PropertyName);
+	if (FbxColorProperty.IsValid())
+	{
+		if (MaterialInput.IsConnected() && MaterialInput.Expression)
+		{
+			if (MaterialInput.Expression->IsA(UMaterialExpressionTextureSample::StaticClass()))
+			{
+				UMaterialExpressionTextureSample *TextureSample = Cast<UMaterialExpressionTextureSample>(MaterialInput.Expression);
+				if (TextureSample && TextureSample->Texture && TextureSample->Texture->AssetImportData)
+				{
+					FString TextureSourceFullPath = TextureSample->Texture->AssetImportData->GetFirstFilename();
+					//Create a fbx property
+					FbxFileTexture* lTexture = FbxFileTexture::Create(Scene, "EnvSamplerTex");
+					lTexture->SetFileName(TCHAR_TO_UTF8(*TextureSourceFullPath));
+					lTexture->SetTextureUse(FbxTexture::eStandard);
+					lTexture->SetMappingType(FbxTexture::eUV);
+					lTexture->ConnectDstProperty(FbxColorProperty);
+					return true;
+				}
+			}
+		}
+	}
+	return false;
 }
 
 /**
@@ -989,7 +1072,7 @@ FbxSurfaceMaterial* FFbxExporter::ExportMaterial(UMaterial* Material)
 	if (Material->GetShadingModel() == MSM_DefaultLit)
 	{
 		FbxMaterial = FbxSurfacePhong::Create(Scene, TCHAR_TO_UTF8(*Material->GetName()));
-		//((FbxSurfacePhong*)FbxMaterial)->Specular.Set(SetMaterialComponent(Material->Specular));
+		//((FbxSurfacePhong*)FbxMaterial)->Specular.Set(Material->Specular));
 		//((FbxSurfacePhong*)FbxMaterial)->Shininess.Set(Material->SpecularPower.Constant);
 	}
 	else // if (Material->ShadingModel == MSM_Unlit)
@@ -997,13 +1080,39 @@ FbxSurfaceMaterial* FFbxExporter::ExportMaterial(UMaterial* Material)
 		FbxMaterial = FbxSurfaceLambert::Create(Scene, TCHAR_TO_UTF8(*Material->GetName()));
 	}
 	
-	((FbxSurfaceLambert*)FbxMaterial)->Emissive.Set(SetMaterialComponent(Material->EmissiveColor));
-	((FbxSurfaceLambert*)FbxMaterial)->Diffuse.Set(SetMaterialComponent(Material->BaseColor));
+	
 	((FbxSurfaceLambert*)FbxMaterial)->TransparencyFactor.Set(Material->Opacity.Constant);
 
 	// Fill in the profile_COMMON effect with the material information.
-	// TODO: Look for textures/constants in the Material expressions...
+	//Fill the texture or constant
+	if (!FillFbxTextureProperty(FbxSurfaceMaterial::sDiffuse, Material->BaseColor, FbxMaterial))
+	{
+		((FbxSurfaceLambert*)FbxMaterial)->Diffuse.Set(SetMaterialComponent(Material->BaseColor, true));
+	}
+	if (!FillFbxTextureProperty(FbxSurfaceMaterial::sEmissive, Material->EmissiveColor, FbxMaterial))
+	{
+		((FbxSurfaceLambert*)FbxMaterial)->Emissive.Set(SetMaterialComponent(Material->EmissiveColor, true));
+	}
 	
+	//Always set the ambient to zero since we dont have ambient in unreal we want to avoid default value in DCCs
+	((FbxSurfaceLambert*)FbxMaterial)->Ambient.Set(FbxDouble3(0.0, 0.0, 0.0));
+
+	//Set the Normal map only if there is a texture sampler
+	FillFbxTextureProperty(FbxSurfaceMaterial::sNormalMap, Material->Normal, FbxMaterial);
+
+	if (Material->BlendMode == BLEND_Translucent)
+	{
+		if (!FillFbxTextureProperty(FbxSurfaceMaterial::sTransparentColor, Material->Opacity, FbxMaterial))
+		{
+			FbxDouble3 OpacityValue((FbxDouble)(Material->Opacity.Constant), (FbxDouble)(Material->Opacity.Constant), (FbxDouble)(Material->Opacity.Constant));
+			((FbxSurfaceLambert*)FbxMaterial)->TransparentColor.Set(OpacityValue);
+		}
+		if (!FillFbxTextureProperty(FbxSurfaceMaterial::sTransparencyFactor, Material->OpacityMask, FbxMaterial))
+		{
+			((FbxSurfaceLambert*)FbxMaterial)->TransparencyFactor.Set(Material->OpacityMask.Constant);
+		}
+	}
+
 	FbxMaterials.Add(Material, FbxMaterial);
 	
 	return FbxMaterial;
@@ -1111,13 +1220,29 @@ FbxNode* FFbxExporter::ExportActor(AActor* Actor, AMatineeActor* InMatineeActor,
 		// this doesn't work with skeletalmeshcomponent
 		FbxNode* ParentNode = FindActor(ParentActor);
 		FVector ActorLocation, ActorRotation, ActorScale;
-		
+
+		// For cameras and lights: always add a Y-pivot rotation to get the correct coordinate system.
+		FTransform RotationDirectionConvert = FTransform::Identity;
+		if (Actor->IsA(ACameraActor::StaticClass()) || Actor->IsA(ALight::StaticClass()))
+		{
+			if (Actor->IsA(ACameraActor::StaticClass()))
+			{
+				FRotator Rotator(0.0f, 0.0f, 90.0f);
+				RotationDirectionConvert = FTransform(Rotator);
+			}
+			else if (Actor->IsA(ALight::StaticClass()))
+			{
+				FRotator Rotator(0.0f, -90.0f, 0.0f);
+				RotationDirectionConvert = FTransform(Rotator);
+			}
+		}
+
 		//If the parent is the root or is not export use the root node as the parent
 		if (bKeepHierarchy && ParentNode)
 		{
 			// Set the default position of the actor on the transforms
 			// The transformation is different from FBX's Z-up: invert the Y-axis for translations and the Y/Z angle values in rotations.
-			const FTransform RelativeTransform = Actor->GetTransform().GetRelativeTransform(ParentActor->GetTransform());
+			const FTransform RelativeTransform = RotationDirectionConvert * Actor->GetTransform().GetRelativeTransform(ParentActor->GetTransform());
 			ActorLocation = RelativeTransform.GetTranslation();
 			ActorRotation = RelativeTransform.GetRotation().Euler();
 			ActorScale = RelativeTransform.GetScale3D();
@@ -1152,32 +1277,6 @@ FbxNode* FFbxExporter::ExportActor(AActor* Actor, AMatineeActor* InMatineeActor,
 		ActorNode->LclRotation.Set(Converter.ConvertToFbxRot(ActorRotation));
 		ActorNode->LclScaling.Set(Converter.ConvertToFbxScale(ActorScale));
 	
-		// For cameras and lights: always add a Y-pivot rotation to get the correct coordinate system.
-		if (Actor->IsA(ACameraActor::StaticClass()) || Actor->IsA(ALight::StaticClass()))
-		{
-			FString FbxPivotNodeName = GetActorNodeName(Actor, NULL);
-
-			if (FbxPivotNodeName == FbxNodeName)
-			{
-				FbxPivotNodeName += TEXT("_pivot");
-			}
-
-			FbxNode* PivotNode = FbxNode::Create(Scene, TCHAR_TO_UTF8(*FbxPivotNodeName));
-			PivotNode->LclRotation.Set(FbxVector4(90, 0, -90));
-
-			if (Actor->IsA(ACameraActor::StaticClass()))
-			{
-				PivotNode->SetPostRotation(FbxNode::eSourcePivot, FbxVector4(0, -90, 0));
-			}
-			else if (Actor->IsA(ALight::StaticClass()))
-			{
-				PivotNode->SetPostRotation(FbxNode::eSourcePivot, FbxVector4(-90, 0, 0));
-			}
-			ActorNode->AddChild(PivotNode);
-
-			ActorNode = PivotNode;
-		}
-
 		if( bExportComponents )
 		{
 			TInlineComponentArray<UActorComponent*> ActorComponents;
@@ -1187,12 +1286,19 @@ FbxNode* FFbxExporter::ExportActor(AActor* Actor, AMatineeActor* InMatineeActor,
 			for( int32 ComponentIndex = 0; ComponentIndex < ActorComponents.Num(); ++ComponentIndex )
 			{
 				UActorComponent* Component = ActorComponents[ComponentIndex];
+				USceneComponent* SceneComp = CastChecked<USceneComponent>(Component);
+
+				if (SceneComp && SceneComp->bHiddenInGame)
+				{
+					//Skip hidden component like camera mesh or other editor helper
+					continue;
+				}
 
 				UStaticMeshComponent* StaticMeshComp = Cast<UStaticMeshComponent>( Component );
 				USkeletalMeshComponent* SkelMeshComp = Cast<USkeletalMeshComponent>( Component );
 				UChildActorComponent* ChildActorComp = Cast<UChildActorComponent>( Component );
 
-				if( StaticMeshComp && StaticMeshComp->StaticMesh )
+				if( StaticMeshComp && StaticMeshComp->StaticMesh)
 				{
 					ComponentsToExport.Add( StaticMeshComp );
 				}
@@ -1200,13 +1306,20 @@ FbxNode* FFbxExporter::ExportActor(AActor* Actor, AMatineeActor* InMatineeActor,
 				{
 					ComponentsToExport.Add( SkelMeshComp );
 				}
+				else if (Component->IsA(UCameraComponent::StaticClass()))
+				{
+					ComponentsToExport.Add(Component);
+				}
+				else if (Component->IsA(ULightComponent::StaticClass()))
+				{
+					ComponentsToExport.Add(Component);
+				}
 				else if (ChildActorComp && ChildActorComp->ChildActor)
 				{
 					ComponentsToExport.Add(ChildActorComp);
 				}
 			}
 
-
 			for( int32 CompIndex = 0; CompIndex < ComponentsToExport.Num(); ++CompIndex )
 			{
 				UActorComponent* Component = ComponentsToExport[CompIndex];
@@ -1219,19 +1332,34 @@ FbxNode* FFbxExporter::ExportActor(AActor* Actor, AMatineeActor* InMatineeActor,
 					// This actor has multiple components
 					// create a child node under the actor for each component
 					FbxNode* CompNode = FbxNode::Create(Scene, TCHAR_TO_UTF8(*Component->GetName()));
-
+					
 					if( SceneComp != Actor->GetRootComponent() )
 					{
 						// Transform is relative to the root component
-						const FTransform RelativeTransform = SceneComp->GetComponentToWorld().GetRelativeTransform(Actor->GetTransform());
+						
+						RotationDirectionConvert = FTransform::Identity;
+						// For cameras and lights: always add a Y-pivot rotation to get the correct coordinate system.
+						if (Component->IsA(UCameraComponent::StaticClass()) || Component->IsA(ULightComponent::StaticClass()))
+						{
+							if (Component->IsA(UCameraComponent::StaticClass()))
+							{
+								FRotator Rotator(0.0f, 0.0f, 90.0f);
+								RotationDirectionConvert = FTransform(Rotator);
+							}
+							else if (Component->IsA(ULightComponent::StaticClass()))
+							{
+								FRotator Rotator(0.0f, -90.0f, 0.0f);
+								RotationDirectionConvert = FTransform(Rotator);
+							}
+						}
+						const FTransform RelativeTransform = RotationDirectionConvert * SceneComp->GetComponentToWorld().GetRelativeTransform(Actor->GetTransform());
 						CompNode->LclTranslation.Set(Converter.ConvertToFbxPos(RelativeTransform.GetTranslation()));
 						CompNode->LclRotation.Set(Converter.ConvertToFbxRot(RelativeTransform.GetRotation().Euler()));
 						CompNode->LclScaling.Set(Converter.ConvertToFbxScale(RelativeTransform.GetScale3D()));
 					}
 
-					ActorNode->AddChild(CompNode);
-
 					ExportNode = CompNode;
+					ActorNode->AddChild(CompNode);
 				}
 
 				UStaticMeshComponent* StaticMeshComp = Cast<UStaticMeshComponent>( Component );
@@ -1258,6 +1386,18 @@ FbxNode* FFbxExporter::ExportActor(AActor* Actor, AMatineeActor* InMatineeActor,
 				{
 					ExportSkeletalMeshComponent(SkelMeshComp, *SkelMeshComp->GetName(), ExportNode);
 				}
+				else if (Component->IsA(UCameraComponent::StaticClass()))
+				{
+					FbxCamera* Camera = FbxCamera::Create(Scene, TCHAR_TO_UTF8(*Component->GetName()));
+					FillFbxCameraAttribute(ActorNode, Camera, Cast<UCameraComponent>(Component));
+					ExportNode->SetNodeAttribute(Camera);
+				}
+				else if (Component->IsA(ULightComponent::StaticClass()))
+				{
+					FbxLight* Light = FbxLight::Create(Scene, TCHAR_TO_UTF8(*Component->GetName()));
+					FillFbxLightAttribute(Light, ActorNode, Cast<ULightComponent>(Component));
+					ExportNode->SetNodeAttribute(Light);
+				}
 				else if (ChildActorComp && ChildActorComp->ChildActor)
 				{
 					FbxNode* ChildActorNode = ExportActor(ChildActorComp->ChildActor, InMatineeActor, true);
diff --git a/Engine/Source/Editor/UnrealEd/Private/Fbx/FbxMainImport.cpp b/Engine/Source/Editor/UnrealEd/Private/Fbx/FbxMainImport.cpp
--- a/Engine/Source/Editor/UnrealEd/Private/Fbx/FbxMainImport.cpp
+++ b/Engine/Source/Editor/UnrealEd/Private/Fbx/FbxMainImport.cpp
@@ -229,6 +229,7 @@ void ApplyImportUIToImportOptions(UFbxImportUI* ImportUI, FBXImportOptions& InOu
 	InOutImportOptions.bKeepOverlappingVertices = ImportUI->SkeletalMeshImportData->bKeepOverlappingVertices;
 	InOutImportOptions.bCombineToSingle = ImportUI->bCombineMeshes;
 	InOutImportOptions.bTransformVertexToAbsolute = ImportUI->StaticMeshImportData->bTransformVertexToAbsolute;
+	InOutImportOptions.bBakePivotInVertex = ImportUI->StaticMeshImportData->bBakePivotInVertex;
 	InOutImportOptions.VertexColorImportOption = ImportUI->StaticMeshImportData->VertexColorImportOption;
 	InOutImportOptions.VertexOverrideColor = ImportUI->StaticMeshImportData->VertexOverrideColor;
 	InOutImportOptions.bRemoveDegenerates = ImportUI->StaticMeshImportData->bRemoveDegenerates;
@@ -299,7 +300,7 @@ FFbxImporter::FFbxImporter()
 	
 	ImportOptions = new FBXImportOptions();
 	FMemory::Memzero(*ImportOptions);
-	ImportOptions->MaterialPrefixName = NAME_None;
+	ImportOptions->MaterialBasePath = NAME_None;
 	
 	CurPhase = NOTSTARTED;
 }
@@ -444,7 +445,7 @@ int32 FFbxImporter::GetImportType(const FString& InFilename)
 	return Result; 
 }
 
-bool FFbxImporter::GetSceneInfo(FString Filename, FbxSceneInfo& SceneInfo)
+bool FFbxImporter::GetSceneInfo(FString Filename, FbxSceneInfo& SceneInfo, bool bPreventMaterialNameClash /*= false*/)
 {
 	bool Result = true;
 	GWarn->BeginSlowTask( NSLOCTEXT("FbxImporter", "BeginGetSceneInfoTask", "Parse FBX file to get scene info"), true );
@@ -460,7 +461,7 @@ bool FFbxImporter::GetSceneInfo(FString Filename, FbxSceneInfo& SceneInfo)
 		}
 		GWarn->UpdateProgress( 40, 100 );
 	case FILEOPENED:
-		if (!ImportFile(Filename))
+		if (!ImportFile(Filename, bPreventMaterialNameClash))
 		{
 			Result = false;
 			break;
@@ -505,14 +506,15 @@ bool FFbxImporter::GetSceneInfo(FString Filename, FbxSceneInfo& SceneInfo)
 				MeshInfo.LODGroup = NULL;
 				if (GeoNode)
 				{
-					FbxNode* ParentNode = GeoNode->GetParent();
-					if (ParentNode->GetNodeAttribute() && ParentNode->GetNodeAttribute()->GetAttributeType() == FbxNodeAttribute::eLODGroup)
+					FbxNode* ParentNode = RecursiveFindParentLodGroup(GeoNode->GetParent());
+					if (ParentNode != nullptr && ParentNode->GetNodeAttribute() && ParentNode->GetNodeAttribute()->GetAttributeType() == FbxNodeAttribute::eLODGroup)
 					{
 						FbxNodeAttribute* LODGroup = ParentNode->GetNodeAttribute();
 						MeshInfo.LODGroup = MakeString(ParentNode->GetName());
 						for (int32 LODIndex = 0; LODIndex < ParentNode->GetChildCount(); LODIndex++)
 						{
-							if (GeoNode == ParentNode->GetChild(LODIndex))
+							FbxNode *MeshNode = FindLODGroupNode(ParentNode, LODIndex);
+							if (GeoNode == MeshNode)
 							{
 								MeshInfo.LODLevel = LODIndex;
 								break;
@@ -626,6 +628,8 @@ void FFbxImporter::TraverseHierarchyNodeRecursively(FbxSceneInfo& SceneInfo, Fbx
 		ChildInfo.UniqueId = ChildNode->GetUniqueID();
 		ChildInfo.ParentName = ParentInfo.ObjectName;
 		ChildInfo.ParentUniqueId = ParentInfo.UniqueId;
+		ChildInfo.RotationPivot = ChildNode->RotationPivot.Get();
+		ChildInfo.ScalePivot = ChildNode->ScalingPivot.Get();
 		ChildInfo.Transform = ChildNode->EvaluateLocalTransform();
 		if (ChildNode->GetNodeAttribute())
 		{
@@ -800,12 +804,44 @@ bool FFbxImporter::OpenFile(FString Filename, bool bParseStatistics, bool bForSc
 	return Result;
 }
 
+void FFbxImporter::FixMaterialClashName()
+{
+	FbxArray<FbxSurfaceMaterial*> MaterialArray;
+	Scene->FillMaterialArray(MaterialArray);
+	TSet<FString> AllMaterialName;
+	for (int32 MaterialIndex = 0; MaterialIndex < MaterialArray.Size(); ++MaterialIndex)
+	{
+		FbxSurfaceMaterial *Material = MaterialArray[MaterialIndex];
+		FString MaterialName = UTF8_TO_TCHAR(Material->GetName());
+		if (AllMaterialName.Contains(MaterialName))
+		{
+			FString OriginalMaterialName = MaterialName;
+			//Use the fbx nameclash 1 convention: NAMECLASH1_KEY
+			//This will add _ncl1_
+			FString MaterialBaseName = MaterialName + TEXT(NAMECLASH1_KEY);
+			int32 NameIndex = 1;
+			MaterialName = MaterialBaseName + FString::FromInt(NameIndex++);
+			while (AllMaterialName.Contains(MaterialName))
+			{
+				MaterialName = MaterialBaseName + FString::FromInt(NameIndex++);
+			}
+			//Rename the Material
+			Material->SetName(TCHAR_TO_UTF8(*MaterialName));
+			AddTokenizedErrorMessage(
+				FTokenizedMessage::Create(EMessageSeverity::Warning,
+										  FText::Format(LOCTEXT("FbxImport_MaterialNameClash", "FBX Scene Loading: Found material name clash, name clash can be wrongly reassign at reimport , material '{0}' was rename '{1}'"), FText::FromString(OriginalMaterialName), FText::FromString(MaterialName))),
+										  FFbxErrors::Generic_LoadingSceneFailed);
+		}
+		AllMaterialName.Add(MaterialName);
+	}
+}
+
 #ifdef IOS_REF
 #undef  IOS_REF
 #define IOS_REF (*(SdkManager->GetIOSettings()))
 #endif
 
-bool FFbxImporter::ImportFile(FString Filename)
+bool FFbxImporter::ImportFile(FString Filename, bool bPreventMaterialNameClash /*=false*/)
 {
 	bool Result = true;
 	
@@ -833,6 +869,12 @@ bool FFbxImporter::ImportFile(FString Filename)
 	// Import the scene.
 	bStatus = Importer->Import(Scene);
 
+	//Make sure we don't have name clash for materials
+	if (bPreventMaterialNameClash)
+	{
+		FixMaterialClashName();
+	}
+
 	// Get the version number of the FBX file format.
 	Importer->GetFileVersion(FileMajor, FileMinor, FileRevision);
 
@@ -861,7 +903,7 @@ bool FFbxImporter::ImportFile(FString Filename)
 //-------------------------------------------------------------------------
 //
 //-------------------------------------------------------------------------
-bool FFbxImporter::ImportFromFile(const FString& Filename, const FString& Type)
+bool FFbxImporter::ImportFromFile(const FString& Filename, const FString& Type, bool bPreventMaterialNameClash /*= false*/)
 {
 	bool Result = true;
 
@@ -875,7 +917,7 @@ bool FFbxImporter::ImportFromFile(const FString& Filename, const FString& Type)
 			break;
 		}
 	case FILEOPENED:
-		if (!ImportFile(FString(Filename)))
+		if (!ImportFile(FString(Filename), bPreventMaterialNameClash))
 		{
 			Result = false;
 			CurPhase = NOTSTARTED;
@@ -947,6 +989,9 @@ bool FFbxImporter::ImportFromFile(const FString& Filename, const FString& Type)
 			//{
 			//	KFbxSystemUnit::cm.ConvertScene( FbxScene );
 			//}
+
+			//Warn the user if there is some geometry that cannot be imported because they are not reference by any scene node attribute
+			ValidateAllMeshesAreReferenceByNodeAttribute();
 		}
 		
 	default:
@@ -1074,8 +1119,17 @@ FbxAMatrix FFbxImporter::ComputeTotalMatrix(FbxNode* Node)
 	//For Single Matrix situation, obtain transfrom matrix from eDESTINATION_SET, which include pivot offsets and pre/post rotations.
 	FbxAMatrix& GlobalTransform = Scene->GetAnimationEvaluator()->GetNodeGlobalTransform(Node);
 	
-	FbxAMatrix TotalMatrix;
-	TotalMatrix = GlobalTransform * Geometry;
+	FbxAMatrix PivotGeometry;
+	if (ImportOptions->bBakePivotInVertex)
+	{
+		FbxVector4 RotationPivot = Node->GetRotationPivot(FbxNode::eSourcePivot);
+		FbxVector4 FullPivot;
+		FullPivot[0] = -RotationPivot[0];
+		FullPivot[1] = -RotationPivot[1];
+		FullPivot[2] = -RotationPivot[2];
+		PivotGeometry.SetT(FullPivot);
+	}
+	FbxAMatrix TotalMatrix = ImportOptions->bTransformVertexToAbsolute ? GlobalTransform * Geometry * PivotGeometry : PivotGeometry;
 
 	return TotalMatrix;
 }
@@ -1200,6 +1254,62 @@ void FillFbxSkelMeshArray(FbxNode* Node, TArray<FbxNode*>& outSkelMeshArray)
 	}
 }
 
+void FFbxImporter::ValidateAllMeshesAreReferenceByNodeAttribute()
+{
+	for (int GeoIndex = 0; GeoIndex < Scene->GetGeometryCount(); ++GeoIndex)
+	{
+		bool FoundOneGeometryLinkToANode = false;
+		FbxGeometry *Geometry = Scene->GetGeometry(GeoIndex);
+		for (int NodeIndex = 0; NodeIndex < Scene->GetNodeCount(); ++NodeIndex)
+		{
+			FbxNode *SceneNode = Scene->GetNode(NodeIndex);
+			FbxGeometry *NodeGeometry = static_cast<FbxGeometry*>(SceneNode->GetMesh());
+			if (NodeGeometry && NodeGeometry->GetUniqueID() == Geometry->GetUniqueID())
+			{
+				FoundOneGeometryLinkToANode = true;
+				break;
+			}
+		}
+		if (!FoundOneGeometryLinkToANode)
+		{
+			FString GeometryName = (Geometry->GetName() && Geometry->GetName()[0] != '\0') ? UTF8_TO_TCHAR(Geometry->GetName()) : TEXT("[Geometry have no name]");
+			AddTokenizedErrorMessage(FTokenizedMessage::Create(EMessageSeverity::Warning,
+				FText::Format(LOCTEXT("FailedToImport_NoObjectLinkToNode", "Mesh {0} in the fbx file is not reference by any hierarchy node."), FText::FromString(GeometryName))),
+				FFbxErrors::Generic_ImportingNewObjectFailed);
+		}
+	}
+}
+
+FbxNode *FFbxImporter::RecursiveGetFirstMeshNode(FbxNode* Node)
+{
+	if (Node->GetMesh() != nullptr)
+		return Node;
+	for (int32 ChildIndex = 0; ChildIndex < Node->GetChildCount(); ++ChildIndex)
+	{
+		FbxNode *MeshNode = RecursiveGetFirstMeshNode(Node->GetChild(ChildIndex));
+		if (MeshNode != nullptr)
+			return MeshNode;
+	}
+	return nullptr;
+}
+
+FbxNode* FFbxImporter::FindLODGroupNode(FbxNode* NodeLodGroup, int32 LodIndex)
+{
+	check(NodeLodGroup->GetChildCount() >= LodIndex);
+	FbxNode *ChildNode = NodeLodGroup->GetChild(LodIndex);
+
+	return RecursiveGetFirstMeshNode(ChildNode);
+}
+
+FbxNode *FFbxImporter::RecursiveFindParentLodGroup(FbxNode *ParentNode)
+{
+	if (ParentNode == nullptr)
+		return nullptr;
+	if (ParentNode->GetNodeAttribute() && ParentNode->GetNodeAttribute()->GetAttributeType() == FbxNodeAttribute::eLODGroup)
+		return ParentNode;
+	return RecursiveFindParentLodGroup(ParentNode->GetParent());
+}
+
 void FFbxImporter::RecursiveFixSkeleton(FbxNode* Node, TArray<FbxNode*> &SkelMeshes, bool bImportNestedMeshes )
 {
 	for (int32 i = 0; i < Node->GetChildCount(); i++)
@@ -1426,7 +1536,7 @@ void FFbxImporter::BuildFbxMatrixForImportTransform(FbxAMatrix& OutMatrix, UFbxA
 */
 void FFbxImporter::RecursiveFindFbxSkelMesh(FbxNode* Node, TArray< TArray<FbxNode*>* >& outSkelMeshArray, TArray<FbxNode*>& SkeletonArray, bool ExpandLOD)
 {
-	FbxNode* SkelMeshNode = NULL;
+	FbxNode* SkelMeshNode = nullptr;
 	FbxNode* NodeToAdd = Node;
 
 	DumpFBXNode(Node);
@@ -1438,11 +1548,11 @@ void FFbxImporter::RecursiveFindFbxSkelMesh(FbxNode* Node, TArray< TArray<FbxNod
 	else if (Node->GetNodeAttribute() && Node->GetNodeAttribute()->GetAttributeType() == FbxNodeAttribute::eLODGroup)
 	{
 		// for LODgroup, add the LODgroup to OutSkelMeshArray according to the skeleton that the first child bind to
-		SkelMeshNode = Node->GetChild(0);
+		SkelMeshNode = FindLODGroupNode(Node, 0);
 		// check if the first child is skeletal mesh
-		if (!(SkelMeshNode->GetMesh() && SkelMeshNode->GetMesh()->GetDeformerCount(FbxDeformer::eSkin) > 0))
+		if (SkelMeshNode != nullptr && !(SkelMeshNode->GetMesh() && SkelMeshNode->GetMesh()->GetDeformerCount(FbxDeformer::eSkin) > 0))
 		{
-			SkelMeshNode = NULL;
+			SkelMeshNode = nullptr;
 		}
 		else if (ExpandLOD)
 		{
@@ -1456,7 +1566,7 @@ void FFbxImporter::RecursiveFindFbxSkelMesh(FbxNode* Node, TArray< TArray<FbxNod
 	{
 		// find root skeleton
 
-		check(SkelMeshNode->GetMesh() != NULL);
+		check(SkelMeshNode->GetMesh() != nullptr);
 		const int32 fbxDeformerCount = SkelMeshNode->GetMesh()->GetDeformerCount();
 		FbxSkin* Deformer = static_cast<FbxSkin*>( SkelMeshNode->GetMesh()->GetDeformer(0, FbxDeformer::eSkin) );
 		
@@ -1540,9 +1650,9 @@ void FFbxImporter::RecursiveFindRigidMesh(FbxNode* Node, TArray< TArray<FbxNode*
 	else if (Node->GetNodeAttribute() && Node->GetNodeAttribute()->GetAttributeType() == FbxNodeAttribute::eLODGroup)
 	{
 		// for LODgroup, add the LODgroup to OutSkelMeshArray according to the skeleton that the first child bind to
-		FbxNode* FirstLOD = Node->GetChild(0);
+		FbxNode* FirstLOD = FindLODGroupNode(Node, 0);
 		// check if the first child is skeletal mesh
-		if (FirstLOD->GetMesh())
+		if (FirstLOD != nullptr && FirstLOD->GetMesh())
 		{
 			if (FirstLOD->GetMesh()->GetDeformerCount(FbxDeformer::eSkin) == 0 )
 			{
@@ -1598,7 +1708,9 @@ void FFbxImporter::RecursiveFindRigidMesh(FbxNode* Node, TArray< TArray<FbxNode*
 		int32 ChildIndex;
 		for (ChildIndex=0; ChildIndex<Node->GetChildCount(); ++ChildIndex)
 		{
-			RecursiveFindRigidMesh(Node->GetChild(ChildIndex), outSkelMeshArray, SkeletonArray, ExpandLOD);
+			FbxNode* MeshNode = FindLODGroupNode(Node, ChildIndex);
+			if(MeshNode != nullptr)
+				RecursiveFindRigidMesh(MeshNode, outSkelMeshArray, SkeletonArray, ExpandLOD);
 		}
 	}
 }
@@ -1723,7 +1835,7 @@ FbxNode* FFbxImporter::FindFBXMeshesByBone(const FName& RootBoneName, bool bExpa
 			FbxNode* Node = (*SkelMeshArray[SkelMeshIndex])[0];
 			if (Node->GetNodeAttribute() && Node->GetNodeAttribute()->GetAttributeType() == FbxNodeAttribute::eLODGroup)
 			{
-				MeshNode = Node->GetChild(0);
+				MeshNode = FindLODGroupNode(Node, 0);
 			}
 			else
 			{
@@ -1755,7 +1867,9 @@ FbxNode* FFbxImporter::FindFBXMeshesByBone(const FName& RootBoneName, bool bExpa
 						FbxNode* Node = SkelMeshes[NodeIndex];
 						if (Node->GetNodeAttribute() && Node->GetNodeAttribute()->GetAttributeType() == FbxNodeAttribute::eLODGroup)
 						{
-							OutFBXMeshNodeArray.Add(Node->GetChild(0));
+							FbxNode *InnerMeshNode = FindLODGroupNode(Node, 0);
+							if(InnerMeshNode != nullptr)
+								OutFBXMeshNodeArray.Add(InnerMeshNode);
 						}
 						else
 						{
diff --git a/Engine/Source/Editor/UnrealEd/Private/Fbx/FbxMaterialImport.cpp b/Engine/Source/Editor/UnrealEd/Private/Fbx/FbxMaterialImport.cpp
--- a/Engine/Source/Editor/UnrealEd/Private/Fbx/FbxMaterialImport.cpp
+++ b/Engine/Source/Editor/UnrealEd/Private/Fbx/FbxMaterialImport.cpp
@@ -145,6 +145,9 @@ UTexture* UnFbx::FFbxImporter::ImportTexture( FbxFileTexture* FbxTexture, bool b
 
 		if ( UnrealTexture != NULL )
 		{
+			//Make sure the AssetImportData point on the texture file and not on the fbx files since the factory point on the fbx file
+			UnrealTexture->AssetImportData->Update(IFileManager::Get().ConvertToAbsolutePathForExternalAppForRead(*Filename));
+
 			// Notify the asset registry
 			FAssetRegistryModule::AssetCreated(UnrealTexture);
 
@@ -437,9 +440,9 @@ void UnFbx::FFbxImporter::CreateUnrealMaterial(FbxSurfaceMaterial& FbxMaterial,
 	}
 	FString MaterialFullName = GetMaterialFullName(FbxMaterial);
 	FString BasePackageName = FPackageName::GetLongPackagePath(Parent->GetOutermost()->GetName());
-	if (ImportOptions->MaterialPrefixName != NAME_None)
+	if (ImportOptions->MaterialBasePath != NAME_None)
 	{
-		BasePackageName += ImportOptions->MaterialPrefixName.ToString();
+		BasePackageName = ImportOptions->MaterialBasePath.ToString();
 	}
 	else
 	{
@@ -507,9 +510,9 @@ void UnFbx::FFbxImporter::CreateUnrealMaterial(FbxSurfaceMaterial& FbxMaterial,
 		UE_LOG(LogFbxMaterialImport, Display, TEXT("-------------------------------"));
 	}
 #endif
-	CreateAndLinkExpressionForMaterialProperty( FbxMaterial, UnrealMaterial, FbxSurfaceMaterial::sDiffuse, UnrealMaterial->BaseColor, false, UVSets, FVector2D(240,-320) );
+	CreateAndLinkExpressionForMaterialProperty(FbxMaterial, UnrealMaterial, FbxSurfaceMaterial::sDiffuse, UnrealMaterial->BaseColor, false, UVSets, FVector2D(240, -320));
 	CreateAndLinkExpressionForMaterialProperty( FbxMaterial, UnrealMaterial, FbxSurfaceMaterial::sEmissive, UnrealMaterial->EmissiveColor, false, UVSets, FVector2D(240,-64) );
-	CreateAndLinkExpressionForMaterialProperty( FbxMaterial, UnrealMaterial, FbxSurfaceMaterial::sSpecular, UnrealMaterial->Specular, false, UVSets, FVector2D(240, -128) );
+	//CreateAndLinkExpressionForMaterialProperty( FbxMaterial, UnrealMaterial, FbxSurfaceMaterial::sSpecular, UnrealMaterial->Specular, false, UVSets, FVector2D(240, -128) );
 	if (!CreateAndLinkExpressionForMaterialProperty( FbxMaterial, UnrealMaterial, FbxSurfaceMaterial::sNormalMap, UnrealMaterial->Normal, true, UVSets, FVector2D(240,256) ) )
 	{
 		CreateAndLinkExpressionForMaterialProperty( FbxMaterial, UnrealMaterial, FbxSurfaceMaterial::sBump, UnrealMaterial->Normal, true, UVSets, FVector2D(240,256) ); // no bump in unreal, use as normal map
diff --git a/Engine/Source/Editor/UnrealEd/Private/Fbx/FbxMeshImportData.cpp b/Engine/Source/Editor/UnrealEd/Private/Fbx/FbxMeshImportData.cpp
--- a/Engine/Source/Editor/UnrealEd/Private/Fbx/FbxMeshImportData.cpp
+++ b/Engine/Source/Editor/UnrealEd/Private/Fbx/FbxMeshImportData.cpp
@@ -8,6 +8,7 @@ UFbxMeshImportData::UFbxMeshImportData(const FObjectInitializer& ObjectInitializ
 {
 	NormalImportMethod = FBXNIM_ComputeNormals;
 	NormalGenerationMethod = EFBXNormalGenerationMethod::MikkTSpace;
+	bBakePivotInVertex = false;
 }
 
 bool UFbxMeshImportData::CanEditChange(const UProperty* InProperty) const
diff --git a/Engine/Source/Editor/UnrealEd/Private/Fbx/FbxSceneImportData.cpp b/Engine/Source/Editor/UnrealEd/Private/Fbx/FbxSceneImportData.cpp
--- a/Engine/Source/Editor/UnrealEd/Private/Fbx/FbxSceneImportData.cpp
+++ b/Engine/Source/Editor/UnrealEd/Private/Fbx/FbxSceneImportData.cpp
@@ -74,6 +74,7 @@ UnFbx::FBXImportOptions *JSONToFbxOption(TSharedPtr<FJsonValue> OptionJsonValue,
 		Option->NormalGenerationMethod = (EFBXNormalGenerationMethod::Type)NormalGenerationMethod;
 	}
 	OptionObj->TryGetBoolField("bTransformVertexToAbsolute", Option->bTransformVertexToAbsolute);
+	OptionObj->TryGetBoolField("bBakePivotInVertex", Option->bBakePivotInVertex);
 	OptionObj->TryGetBoolField("bCombineToSingle", Option->bCombineToSingle);
 	int32 VertexColorImportOption;
 	if (OptionObj->TryGetNumberField("VertexColorImportOption", VertexColorImportOption))
@@ -104,10 +105,19 @@ UnFbx::FBXImportOptions *JSONToFbxOption(TSharedPtr<FJsonValue> OptionJsonValue,
 		Option->StaticMeshLODGroup = FName(*LODGroup);
 	}
 	OptionObj->TryGetBoolField("bImportStaticMeshLODs", Option->bImportStaticMeshLODs);
-	FString MaterialPrefixName;
-	if (OptionObj->TryGetStringField("MaterialPrefixName", MaterialPrefixName))
+
+	//Skeletal mesh options
+	OptionObj->TryGetBoolField("bUpdateSkeletonReferencePose", Option->bUpdateSkeletonReferencePose);
+	OptionObj->TryGetBoolField("bUseT0AsRefPose", Option->bUseT0AsRefPose);
+	OptionObj->TryGetBoolField("bPreserveSmoothingGroups", Option->bPreserveSmoothingGroups);
+	OptionObj->TryGetBoolField("bImportMeshesInBoneHierarchy", Option->bImportMeshesInBoneHierarchy);
+	OptionObj->TryGetBoolField("bImportMorphTargets", Option->bImportMorph);
+	OptionObj->TryGetBoolField("bKeepOverlappingVertices", Option->bKeepOverlappingVertices);
+
+	FString MaterialBasePath;
+	if (OptionObj->TryGetStringField("MaterialBasePath", MaterialBasePath))
 	{
-		Option->MaterialPrefixName = FName(*MaterialPrefixName);
+		Option->MaterialBasePath = FName(*MaterialBasePath);
 	}
 	return Option;
 }
@@ -139,10 +149,11 @@ FString FbxOptionToJSON(FString OptionName, UnFbx::FBXImportOptions *Option)
 		Option->ImportUniformScale
 		);
 
-	JsonString += FString::Printf(TEXT("\"NormalImportMethod\" : \"%d\", \"NormalGenerationMethod\" : \"%d\", \"bTransformVertexToAbsolute\" : \"%d\", \"bCombineToSingle\" : \"%d\", \"VertexColorImportOption\" : \"%d\", \"VertexOverrideColor\" : {\"R\" : \"%d\", \"G\" : \"%d\", \"B\" : \"%d\", \"A\" : \"%d\" }, "),
+	JsonString += FString::Printf(TEXT("\"NormalImportMethod\" : \"%d\", \"NormalGenerationMethod\" : \"%d\", \"bTransformVertexToAbsolute\" : \"%d\", \"bBakePivotInVertex\" : \"%d\", \"bCombineToSingle\" : \"%d\", \"VertexColorImportOption\" : \"%d\", \"VertexOverrideColor\" : {\"R\" : \"%d\", \"G\" : \"%d\", \"B\" : \"%d\", \"A\" : \"%d\" }, "),
 		(int32)Option->NormalImportMethod,
 		(int32)Option->NormalGenerationMethod,
 		Option->bTransformVertexToAbsolute ? 1 : 0,
+		Option->bBakePivotInVertex ? 1 : 0,
 		Option->bCombineToSingle ? 1 : 0,
 		(int32)Option->VertexColorImportOption,
 		Option->VertexOverrideColor.R,
@@ -162,8 +173,17 @@ FString FbxOptionToJSON(FString OptionName, UnFbx::FBXImportOptions *Option)
 		Option->bImportStaticMeshLODs ? 1 : 0
 		);
 
-	JsonString += FString::Printf(TEXT("\"MaterialPrefixName\" : \"%s\"}"), *(Option->MaterialPrefixName.ToString()));
-	//TODO add skeletalmesh and animation
+	//Skeletal mesh options
+	JsonString += FString::Printf(TEXT("\"bUpdateSkeletonReferencePose\" : \"%d\", \"bUseT0AsRefPose\" : \"%d\", \"bPreserveSmoothingGroups\" : \"%d\", \"bImportMeshesInBoneHierarchy\" : \"%d\", \"bImportMorphTargets\" : \"%d\", \"bKeepOverlappingVertices\" : \"%d\", "),
+		Option->bUpdateSkeletonReferencePose ? 1 : 0,
+		Option->bUseT0AsRefPose ? 1 : 0,
+		Option->bPreserveSmoothingGroups ? 1 : 0,
+		Option->bImportMeshesInBoneHierarchy ? 1 : 0,
+		Option->bImportMorph ? 1 : 0,
+		Option->bKeepOverlappingVertices ? 1 : 0
+		);
+
+	JsonString += FString::Printf(TEXT("\"MaterialBasePath\" : \"%s\"}"), *(Option->MaterialBasePath.ToString()));
 	return JsonString;
 }
 
@@ -171,7 +191,7 @@ FString FbxNodeInfoToJSON(const TSharedPtr<FFbxNodeInfo> NodeInfo)
 {
 	check(NodeInfo.IsValid());
 
-	FString JsonString = FString::Printf(TEXT("{ \"NodeName\" : \"%s\", \"UniqueId\" : \"%llu\", \"NodeHierarchyPath\" : \"%s\", \"bImportNode\" : \"%d\", \"ParentUniqueId\" : \"%llu\", \"AttributeType\" : \"%s\", \"AttributeInfo\" : \"%llu\", \"Materials\" : ["),
+	FString JsonString = FString::Printf(TEXT("{ \"NodeName\" : \"%s\", \"UniqueId\" : \"%llu\", \"NodeHierarchyPath\" : \"%s\", \"bImportNode\" : \"%d\", \"ParentUniqueId\" : \"%llu\", \"AttributeType\" : \"%s\", \"AttributeUniqueId\" : \"%llu\", \"Materials\" : ["),
 		*(NodeInfo->NodeName),
 		NodeInfo->UniqueId,
 		*(NodeInfo->NodeHierarchyPath),
@@ -209,7 +229,7 @@ FString FbxMeshInfoToJSON(const TSharedPtr<FFbxMeshInfo> MeshInfo)
 {
 	check(MeshInfo.IsValid());
 
-	FString JsonString = FString::Printf(TEXT("{ \"Name\" : \"%s\", \"UniqueId\" : \"%llu\", \"bImportAttribute\" : \"%d\", \"OptionName\" : \"%s\", \"bIsSkelMesh\" : \"%d\", \"OriginalImportPath\" : \"%s\", \"OriginalFullImportName\" : \"%s\", \"bOverridePath\" : \"%d\", \"OverrideImportPath\" : \"%s\", \"OverrideFullImportName\" : \"%s\" }"),
+	FString JsonString = FString::Printf(TEXT("{ \"Name\" : \"%s\", \"UniqueId\" : \"%llu\", \"bImportAttribute\" : \"%d\", \"OptionName\" : \"%s\", \"bIsSkelMesh\" : \"%d\", \"OriginalImportPath\" : \"%s\", \"OriginalFullImportName\" : \"%s\", \"bOverridePath\" : \"%d\", \"OverrideImportPath\" : \"%s\", \"OverrideFullImportName\" : \"%s\", \"PivotNodeUid\" : \"%llu\", \"LODGroup\" : \"%s\", \"LODLevel\" : \"%d\", \"IsLod\" : \"%d\" }"),
 		*(MeshInfo->Name),
 		MeshInfo->UniqueId,
 		MeshInfo->bImportAttribute ? 1 : 0,
@@ -219,7 +239,11 @@ FString FbxMeshInfoToJSON(const TSharedPtr<FFbxMeshInfo> MeshInfo)
 		*(MeshInfo->OriginalFullImportName),
 		MeshInfo->bOverridePath ? 1 : 0,
 		*(MeshInfo->OverrideImportPath),
-		*(MeshInfo->OverrideFullImportName)
+		*(MeshInfo->OverrideFullImportName),
+		MeshInfo->PivotNodeUid,
+		*(MeshInfo->LODGroup),
+		MeshInfo->LODLevel,
+		MeshInfo->IsLod ? 1 : 0
 		);
 	return JsonString;
 }
@@ -331,10 +355,10 @@ void UFbxSceneImportData::FromJson(FString InJsonString)
 			continue;
 		}
 		//Read General data
-		if (!SceneInfoObj->TryGetBoolField("bImportScene", bImportScene)) continue;
-		if (!SceneInfoObj->TryGetBoolField("bCreateFolderHierarchy", bCreateFolderHierarchy)) continue;
-		if (!SceneInfoObj->TryGetNumberField("HierarchyType", HierarchyType)) continue;
-		if (!SceneInfoObj->TryGetStringField("BluePrintFullName", BluePrintFullName)) continue;
+		SceneInfoObj->TryGetBoolField("bImportScene", bImportScene);
+		SceneInfoObj->TryGetBoolField("bCreateFolderHierarchy", bCreateFolderHierarchy);
+		SceneInfoObj->TryGetNumberField("HierarchyType", HierarchyType);
+		SceneInfoObj->TryGetStringField("BluePrintFullName", BluePrintFullName);
 
 		//Read Options
 		const TArray<TSharedPtr<FJsonValue>>* JSONOptions;
@@ -367,21 +391,32 @@ void UFbxSceneImportData::FromJson(FString InJsonString)
 			if (!MeshInfoObj->TryGetStringField("UniqueId", UniqueIDStr)) continue;
 			MeshInfo->UniqueId = FCString::Atoi64(*UniqueIDStr);
 
-			if (!MeshInfoObj->TryGetBoolField("bImportAttribute", MeshInfo->bImportAttribute)) continue;
+			MeshInfoObj->TryGetBoolField("bImportAttribute", MeshInfo->bImportAttribute);
+
+			MeshInfoObj->TryGetStringField("OptionName", MeshInfo->OptionName);
 
-			if (!MeshInfoObj->TryGetStringField("OptionName", MeshInfo->OptionName)) continue;
+			MeshInfoObj->TryGetBoolField("bIsSkelMesh", MeshInfo->bIsSkelMesh);
 
-			if (!MeshInfoObj->TryGetBoolField("bIsSkelMesh", MeshInfo->bIsSkelMesh)) continue;
+			MeshInfoObj->TryGetStringField("OriginalImportPath", MeshInfo->OriginalImportPath);
 
-			if (!MeshInfoObj->TryGetStringField("OriginalImportPath", MeshInfo->OriginalImportPath)) continue;
+			MeshInfoObj->TryGetStringField("OriginalFullImportName", MeshInfo->OriginalFullImportName);
 
-			if (!MeshInfoObj->TryGetStringField("OriginalFullImportName", MeshInfo->OriginalFullImportName)) continue;
+			MeshInfoObj->TryGetBoolField("bOverridePath", MeshInfo->bOverridePath);
+
+			MeshInfoObj->TryGetStringField("OverrideImportPath", MeshInfo->OverrideImportPath);
+
+			MeshInfoObj->TryGetStringField("OverrideFullImportName", MeshInfo->OverrideFullImportName);
+
+			if (MeshInfoObj->TryGetStringField("PivotNodeUid", UniqueIDStr))
+			{
+				MeshInfo->PivotNodeUid = FCString::Atoi64(*UniqueIDStr);
+			}
 
-			if (!MeshInfoObj->TryGetBoolField("bOverridePath", MeshInfo->bOverridePath)) continue;
+			MeshInfoObj->TryGetStringField("LODGroup", MeshInfo->LODGroup);
 
-			if (!MeshInfoObj->TryGetStringField("OverrideImportPath", MeshInfo->OverrideImportPath)) continue;
+			MeshInfoObj->TryGetNumberField("LODLevel", MeshInfo->LODLevel);
 
-			if (!MeshInfoObj->TryGetStringField("OverrideFullImportName", MeshInfo->OverrideFullImportName)) continue;
+			MeshInfoObj->TryGetBoolField("IsLod", MeshInfo->IsLod);
 				
 			SceneInfoSourceData->MeshInfo.Add(MeshInfo);
 		}
@@ -400,44 +435,48 @@ void UFbxSceneImportData::FromJson(FString InJsonString)
 
 			if (!NodeInfoObj->TryGetStringField("NodeName", NodeInfo->NodeName)) continue;
 
-			if (!NodeInfoObj->TryGetStringField("NodeHierarchyPath", NodeInfo->NodeHierarchyPath)) continue;
+			NodeInfoObj->TryGetStringField("NodeHierarchyPath", NodeInfo->NodeHierarchyPath);
 
 			FString UniqueIDStr;
 			if (!NodeInfoObj->TryGetStringField("UniqueId", UniqueIDStr)) continue;
 			NodeInfo->UniqueId = FCString::Atoi64(*UniqueIDStr);
 
-			if (!NodeInfoObj->TryGetBoolField("bImportNode", NodeInfo->bImportNode)) continue;
+			NodeInfoObj->TryGetBoolField("bImportNode", NodeInfo->bImportNode);
 
 			//Find the parent
-			if (!NodeInfoObj->TryGetStringField("ParentUniqueId", UniqueIDStr)) continue;
-			uint64 ParentUniqueId = FCString::Atoi64(*UniqueIDStr);
-			if (ParentUniqueId != 0)
+			if (NodeInfoObj->TryGetStringField("ParentUniqueId", UniqueIDStr))
 			{
-				//Find the parent
-				for (TSharedPtr<FFbxNodeInfo> SearchNodeInfo : SceneInfoSourceData->HierarchyInfo)
+				uint64 ParentUniqueId = FCString::Atoi64(*UniqueIDStr);
+				if (ParentUniqueId != 0)
 				{
-					if (SearchNodeInfo->UniqueId == ParentUniqueId)
+					//Find the parent
+					for (TSharedPtr<FFbxNodeInfo> SearchNodeInfo : SceneInfoSourceData->HierarchyInfo)
 					{
-						NodeInfo->ParentNodeInfo = SearchNodeInfo;
-						SearchNodeInfo->Childrens.Add(NodeInfo);
-						break;
+						if (SearchNodeInfo->UniqueId == ParentUniqueId)
+						{
+							NodeInfo->ParentNodeInfo = SearchNodeInfo;
+							SearchNodeInfo->Childrens.Add(NodeInfo);
+							break;
+						}
 					}
 				}
 			}
 
-			if (!NodeInfoObj->TryGetStringField("AttributeType", NodeInfo->AttributeType)) continue;
+			NodeInfoObj->TryGetStringField("AttributeType", NodeInfo->AttributeType);
 
-			if (!NodeInfoObj->TryGetStringField("AttributeInfo", UniqueIDStr)) continue;
-			uint64 attributeUniqueID = FCString::Atoi64(*UniqueIDStr);
-			if (attributeUniqueID != 0)
+			if (NodeInfoObj->TryGetStringField("AttributeUniqueId", UniqueIDStr))
 			{
-				//Find the attribute info
-				for (TSharedPtr<FFbxMeshInfo> MeshInfo : SceneInfoSourceData->MeshInfo)
+				NodeInfo->AttributeUniqueId = FCString::Atoi64(*UniqueIDStr);
+				if (NodeInfo->AttributeUniqueId != 0)
 				{
-					if (MeshInfo->UniqueId == attributeUniqueID)
+					//Find the attribute info
+					for (TSharedPtr<FFbxMeshInfo> MeshInfo : SceneInfoSourceData->MeshInfo)
 					{
-						NodeInfo->AttributeInfo = MeshInfo;
-						break;
+						if (MeshInfo->UniqueId == NodeInfo->AttributeUniqueId)
+						{
+							NodeInfo->AttributeInfo = MeshInfo;
+							break;
+						}
 					}
 				}
 			}
@@ -457,22 +496,22 @@ void UFbxSceneImportData::FromJson(FString InJsonString)
 
 				if (!MaterialObj->TryGetStringField("Name", MaterialInfo->Name)) continue;
 
-				if (!MaterialObj->TryGetStringField("HierarchyPath", MaterialInfo->HierarchyPath)) continue;
+				MaterialObj->TryGetStringField("HierarchyPath", MaterialInfo->HierarchyPath);
 
 				if (!MaterialObj->TryGetStringField("UniqueId", UniqueIDStr)) continue;
 				MaterialInfo->UniqueId = FCString::Atoi64(*UniqueIDStr);
 
-				if (!MaterialObj->TryGetBoolField("bImportAttribute", MaterialInfo->bImportAttribute)) continue;
+				MaterialObj->TryGetBoolField("bImportAttribute", MaterialInfo->bImportAttribute);
 
-				if (!MaterialObj->TryGetStringField("OriginalImportPath", MaterialInfo->OriginalImportPath)) continue;
+				MaterialObj->TryGetStringField("OriginalImportPath", MaterialInfo->OriginalImportPath);
 
-				if (!MaterialObj->TryGetStringField("OriginalFullImportName", MaterialInfo->OriginalFullImportName)) continue;
+				MaterialObj->TryGetStringField("OriginalFullImportName", MaterialInfo->OriginalFullImportName);
 
-				if (!MaterialObj->TryGetBoolField("bOverridePath", MaterialInfo->bOverridePath)) continue;
+				MaterialObj->TryGetBoolField("bOverridePath", MaterialInfo->bOverridePath);
 
-				if (!MaterialObj->TryGetStringField("OverrideImportPath", MaterialInfo->OverrideImportPath)) continue;
+				MaterialObj->TryGetStringField("OverrideImportPath", MaterialInfo->OverrideImportPath);
 
-				if (!MaterialObj->TryGetStringField("OverrideFullImportName", MaterialInfo->OverrideFullImportName)) continue;
+				MaterialObj->TryGetStringField("OverrideFullImportName", MaterialInfo->OverrideFullImportName);
 
 				NodeInfo->Materials.Add(MaterialInfo);
 			}
