Copying //UE4/Dev-Editor to Dev-Main (//UE4/Dev-Main)

diff --git a/Engine/Source/Editor/UnrealEd/Private/Fbx/FbxSceneImportFactory.cpp b/Engine/Source/Editor/UnrealEd/Private/Fbx/FbxSceneImportFactory.cpp
--- a/Engine/Source/Editor/UnrealEd/Private/Fbx/FbxSceneImportFactory.cpp
+++ b/Engine/Source/Editor/UnrealEd/Private/Fbx/FbxSceneImportFactory.cpp
@@ -6,6 +6,10 @@
 #include "SSkeletonWidget.h"
 #include "Dialogs/DlgPickPath.h"
 
+#include "Components/PointLightComponent.h"
+#include "Components/DirectionalLightComponent.h"
+#include "Components/SpotLightComponent.h"
+
 #include "AssetSelection.h"
 
 #include "FbxImporter.h"
@@ -24,8 +28,6 @@
 
 #include "OutputDevice.h"
 
-#include "FbxImporter.h"
-
 #define LOCTEXT_NAMESPACE "FBXSceneImportFactory"
 
 using namespace UnFbx;
@@ -42,6 +44,27 @@ using namespace UnFbx;
 //Initialize static default option name
 FString UFbxSceneImportFactory::DefaultOptionName = FString(TEXT("Default"));
 
+FbxNode* FindFbxNodeById(UnFbx::FFbxImporter* FbxImporter, FbxNode *CurrentNode, uint64 UniqueId)
+{
+	if (CurrentNode == nullptr)
+	{
+		CurrentNode = FbxImporter->Scene->GetRootNode();
+	}
+	if (CurrentNode->GetUniqueID() == UniqueId)
+	{
+		return CurrentNode;
+	}
+	for (int32 ChildIndex = 0; ChildIndex < CurrentNode->GetChildCount(); ++ChildIndex)
+	{
+		FbxNode *FoundNode = FindFbxNodeById(FbxImporter, CurrentNode->GetChild(ChildIndex), UniqueId);
+		if (FoundNode != nullptr)
+		{
+			return FoundNode;
+		}
+	}
+	return nullptr;
+}
+
 bool GetFbxSceneImportOptions(UnFbx::FFbxImporter* FbxImporter
 	, TSharedPtr<FFbxSceneInfo> SceneInfoPtr
 	, UnFbx::FBXImportOptions *GlobalImportSettings
@@ -49,8 +72,6 @@ bool GetFbxSceneImportOptions(UnFbx::FFbxImporter* FbxImporter
 	, UFbxSceneImportOptionsStaticMesh *StaticMeshImportData
 	, ImportOptionsNameMap &NameOptionsMap
 	, UFbxSceneImportOptionsSkeletalMesh *SkeletalMeshImportData
-	, UFbxSceneImportOptionsAnimation *AnimSequenceImportData
-	, UFbxSceneImportOptionsMaterial *TextureImportData
 	, FString Path)
 {
 	//Make sure we don't put the global transform into the vertex position of the mesh
@@ -66,9 +87,12 @@ bool GetFbxSceneImportOptions(UnFbx::FFbxImporter* FbxImporter
 
 	GlobalImportSettings->bConvertScene = true;
 
+	GlobalImportSettings->bBakePivotInVertex = SceneImportOptions->bBakePivotInVertex;
+	GlobalImportSettings->bInvertNormalMap = SceneImportOptions->bInvertNormalMaps;
+
 	//TODO this options will be set by the fbxscene UI in the material options tab, it also should be save/load from config file
 	//Prefix materials package name to put all material under Material folder (this avoid name clash with meshes)
-	GlobalImportSettings->MaterialPrefixName = FName("/Materials/");
+	GlobalImportSettings->MaterialBasePath = NAME_None;
 
 	TSharedPtr<SWindow> ParentWindow;
 	if (FModuleManager::Get().IsModuleLoaded("MainFrame"))
@@ -77,7 +101,7 @@ bool GetFbxSceneImportOptions(UnFbx::FFbxImporter* FbxImporter
 		ParentWindow = MainFrame.GetParentWindow();
 	}
 	TSharedRef<SWindow> Window = SNew(SWindow)
-		.ClientSize(FVector2D(800.f, 650.f))
+		.ClientSize(FVector2D(820.f, 650.f))
 		.Title(NSLOCTEXT("UnrealEd", "FBXSceneImportOpionsTitle", "FBX Scene Import Options"));
 	TSharedPtr<SFbxSceneOptionWindow> FbxSceneOptionWindow;
 
@@ -90,8 +114,6 @@ bool GetFbxSceneImportOptions(UnFbx::FFbxImporter* FbxImporter
 			.SceneImportOptionsStaticMeshDisplay(StaticMeshImportData)
 			.OverrideNameOptionsMap(&NameOptionsMap)
 			.SceneImportOptionsSkeletalMeshDisplay(SkeletalMeshImportData)
-			.SceneImportOptionsAnimationDisplay(AnimSequenceImportData)
-			.SceneImportOptionsMaterialDisplay(TextureImportData)
 			.OwnerWindow(Window)
 			.FullPath(Path)
 			);
@@ -104,8 +126,10 @@ bool GetFbxSceneImportOptions(UnFbx::FFbxImporter* FbxImporter
 	}
 
 	//setup all options
+	GlobalImportSettings->bBakePivotInVertex = SceneImportOptions->bBakePivotInVertex;
 	GlobalImportSettings->bImportStaticMeshLODs = SceneImportOptions->bImportStaticMeshLODs;
 	GlobalImportSettings->bImportSkeletalMeshLODs = SceneImportOptions->bImportSkeletalMeshLODs;
+	GlobalImportSettings->bInvertNormalMap = SceneImportOptions->bInvertNormalMaps;
 
 	//Set the override material into the options
 	for (TSharedPtr<FFbxNodeInfo> NodeInfo : SceneInfoPtr->HierarchyInfo)
@@ -251,10 +275,146 @@ void UFbxSceneImportFactory::ExtractMaterialInfo(void* FbxImporterVoid, TSharedP
 	ExtractMaterialInfoFromNode(FbxImporter, RootNode, SceneInfoPtr, ExtractedMaterials, ExtractedTextures, CurrentHierarchyPath);
 }
 
+bool IsPartOfSkeletonHierarchy(const TMap<uint64, const UnFbx::FbxNodeInfo*> &NodeInfoMap, const UnFbx::FbxNodeInfo& NodeInfo)
+{
+	FString AttributeType = NodeInfo.AttributeType;
+	if (AttributeType.Compare(TEXT("eSkeleton")) == 0)
+	{
+		return true;
+	}
+	if (NodeInfoMap.Contains(NodeInfo.ParentUniqueId))
+	{
+		const UnFbx::FbxNodeInfo *ParentNodeInfo = *NodeInfoMap.Find(NodeInfo.ParentUniqueId);
+		return IsPartOfSkeletonHierarchy(NodeInfoMap, *ParentNodeInfo);
+	}
+	return false;
+}
+
+void FetchFbxCameraInScene(UnFbx::FFbxImporter *FbxImporter, FbxNode* ParentNode, TSharedPtr<FFbxSceneInfo> SceneInfoPtr)
+{
+	if (ParentNode->GetNodeAttribute() && ParentNode->GetNodeAttribute()->GetAttributeType() == FbxNodeAttribute::eCamera)
+	{
+		FbxCamera* CameraAttribute = (FbxCamera*)ParentNode->GetNodeAttribute();
+		if (CameraAttribute != nullptr && CameraAttribute->GetNode() != nullptr && !SceneInfoPtr->CameraInfo.Contains(CameraAttribute->GetUniqueID()))
+		{
+			FbxNode* CameraNode = CameraAttribute->GetNode();
+			TSharedPtr<FFbxCameraInfo> CameraInfo = MakeShareable(new FFbxCameraInfo());
+
+			if (CameraAttribute->GetName()[0] != '\0')
+			{
+				CameraInfo->Name = FbxImporter->MakeName(CameraAttribute->GetName());
+			}
+			else
+			{
+				CameraInfo->Name = FbxImporter->MakeString(CameraNode ? CameraNode->GetName() : "None");
+			}
+			CameraInfo->UniqueId = CameraAttribute->GetUniqueID();
+
+			CameraInfo->AspectWidth = CameraAttribute->AspectWidth.Get();
+			CameraInfo->AspectHeight = CameraAttribute->AspectHeight.Get();
+			CameraInfo->NearPlane = CameraAttribute->NearPlane.Get();
+			CameraInfo->FarPlane = CameraAttribute->FarPlane.Get();
+			CameraInfo->ProjectionPerspective = CameraAttribute->ProjectionType.Get() == FbxCamera::ePerspective;
+			CameraInfo->OrthoZoom = CameraAttribute->OrthoZoom.Get();
+			CameraInfo->FieldOfView = CameraAttribute->FieldOfView.Get();
+			SceneInfoPtr->CameraInfo.Add(CameraInfo->UniqueId, CameraInfo);
+		}
+	}
+	for (int i = 0; i < ParentNode->GetChildCount(); ++i)
+	{
+		FbxNode* Child = ParentNode->GetChild(i);
+		FetchFbxCameraInScene(FbxImporter, Child, SceneInfoPtr);
+	}
+}
+
+
+void FetchFbxLightInScene(UnFbx::FFbxImporter *FbxImporter, FbxNode* ParentNode, TSharedPtr<FFbxSceneInfo> SceneInfoPtr)
+{
+	if (ParentNode->GetNodeAttribute() && ParentNode->GetNodeAttribute()->GetAttributeType() == FbxNodeAttribute::eLight)
+	{
+		FbxLight* LightAttribute = (FbxLight*)ParentNode->GetNodeAttribute();
+		if (LightAttribute != nullptr && LightAttribute->GetNode() != nullptr && !SceneInfoPtr->LightInfo.Contains(LightAttribute->GetUniqueID()))
+		{
+			FbxNode* LightNode = LightAttribute->GetNode();
+			TSharedPtr<FFbxLightInfo> LightInfo = MakeShareable(new FFbxLightInfo());
+
+			if (LightAttribute->GetName()[0] != '\0')
+			{
+				LightInfo->Name = FbxImporter->MakeName(LightAttribute->GetName());
+			}
+			else
+			{
+				LightInfo->Name = FbxImporter->MakeString(LightNode ? LightNode->GetName() : "None");
+			}
+			LightInfo->UniqueId = LightAttribute->GetUniqueID();
+			switch (LightAttribute->LightType.Get())
+			{
+			case FbxLight::ePoint:
+				LightInfo->Type = 0;
+				break;
+			case FbxLight::eDirectional:
+				LightInfo->Type = 1;
+				break;
+			case FbxLight::eSpot:
+				LightInfo->Type = 2;
+				break;
+			case FbxLight::eArea:
+				LightInfo->Type = 3;
+				break;
+			case FbxLight::eVolume:
+				LightInfo->Type = 4;
+				break;
+			}
+			LightInfo->Color.R = (uint8)(LightAttribute->Color.Get()[0] * 255.0);
+			LightInfo->Color.G = (uint8)(LightAttribute->Color.Get()[1] * 255.0);
+			LightInfo->Color.B = (uint8)(LightAttribute->Color.Get()[2] * 255.0);
+			LightInfo->Intensity = (float)(LightAttribute->Intensity.Get());
+			switch (LightAttribute->DecayType.Get())
+			{
+			case FbxLight::EDecayType::eNone:
+				LightInfo->Decay = 0;
+				break;
+			case FbxLight::EDecayType::eLinear:
+				LightInfo->Decay = 1;
+				break;
+			case FbxLight::EDecayType::eQuadratic:
+				LightInfo->Decay = 2;
+				break;
+			case FbxLight::EDecayType::eCubic:
+				LightInfo->Decay = 3;
+				break;
+			}
+			LightInfo->CastLight = LightAttribute->CastLight.Get();
+			LightInfo->CastShadow = LightAttribute->CastShadows.Get();
+			LightInfo->ShadowColor.R = (uint8)(LightAttribute->ShadowColor.Get()[0] * 255.0);
+			LightInfo->ShadowColor.G = (uint8)(LightAttribute->ShadowColor.Get()[1] * 255.0);
+			LightInfo->ShadowColor.B = (uint8)(LightAttribute->ShadowColor.Get()[2] * 255.0);
+
+			LightInfo->InnerAngle = (float)(LightAttribute->InnerAngle.Get());
+			LightInfo->OuterAngle = (float)(LightAttribute->OuterAngle.Get());
+			LightInfo->Fog = (float)(LightAttribute->Fog.Get());
+			LightInfo->DecayStart = (float)(LightAttribute->DecayStart.Get());
+			LightInfo->EnableNearAttenuation = LightAttribute->EnableNearAttenuation.Get();
+			LightInfo->NearAttenuationStart = (float)(LightAttribute->NearAttenuationStart.Get());
+			LightInfo->NearAttenuationEnd = (float)(LightAttribute->NearAttenuationEnd.Get());
+			LightInfo->EnableFarAttenuation = LightAttribute->EnableFarAttenuation.Get();
+			LightInfo->FarAttenuationStart = (float)(LightAttribute->FarAttenuationStart.Get());
+			LightInfo->FarAttenuationEnd = (float)(LightAttribute->FarAttenuationEnd.Get());
+			SceneInfoPtr->LightInfo.Add(LightInfo->UniqueId, LightInfo);
+		}
+	}
+	for (int i = 0; i < ParentNode->GetChildCount(); ++i)
+	{
+		FbxNode* Child = ParentNode->GetChild(i);
+		FetchFbxLightInScene(FbxImporter, Child, SceneInfoPtr);
+	}
+}
+
 // TODO we should replace the old UnFbx:: data by the new data that use shared pointer.
 // For now we convert the old structure to the new one
-TSharedPtr<FFbxSceneInfo> UFbxSceneImportFactory::ConvertSceneInfo(void* VoidFbxSceneInfo)
+TSharedPtr<FFbxSceneInfo> UFbxSceneImportFactory::ConvertSceneInfo(void* VoidFbxImporter, void* VoidFbxSceneInfo)
 {
+	UnFbx::FFbxImporter *FbxImporter = (UnFbx::FFbxImporter*)VoidFbxImporter;
 	UnFbx::FbxSceneInfo &SceneInfo = *((UnFbx::FbxSceneInfo*)VoidFbxSceneInfo);
 	TSharedPtr<FFbxSceneInfo> SceneInfoPtr = MakeShareable(new FFbxSceneInfo());
 	SceneInfoPtr->NonSkinnedMeshNum = SceneInfo.NonSkinnedMeshNum;
@@ -265,8 +425,37 @@ TSharedPtr<FFbxSceneInfo> UFbxSceneImportFactory::ConvertSceneInfo(void* VoidFbx
 	SceneInfoPtr->bHasAnimation = SceneInfo.bHasAnimation;
 	SceneInfoPtr->FrameRate = SceneInfo.FrameRate;
 	SceneInfoPtr->TotalTime = SceneInfo.TotalTime;
+
+	//Get the valid skeletal mesh from the fbx file and store it in the map
+	TMap<uint64, FbxMesh*> ValidSkeletalMesh;
+	FbxNode* RootNodeToImport = FbxImporter->Scene->GetRootNode();
+	TArray< TArray<FbxNode*>* > SkelMeshArray;
+	UnFbx::FBXImportOptions* FbxImportOptionsPtr = FbxImporter->GetImportOptions();
+	bool OldValueImportMeshesInBoneHierarchy = FbxImportOptionsPtr->bImportMeshesInBoneHierarchy;
+	FbxImportOptionsPtr->bImportMeshesInBoneHierarchy = true;
+	FbxImporter->FillFbxSkelMeshArrayInScene(RootNodeToImport, SkelMeshArray, false);
+	FbxImportOptionsPtr->bImportMeshesInBoneHierarchy = OldValueImportMeshesInBoneHierarchy;
+
+	for (int32 i = 0; i < SkelMeshArray.Num(); i++)
+	{
+		TArray<FbxNode*> NodeArray = *SkelMeshArray[i];
+		if (NodeArray.Num() < 1)
+			continue;
+		FbxNode* RootNodeArrayNode = NodeArray[0];
+		FbxMesh* Mesh = RootNodeArrayNode->GetMesh();
+		if (Mesh != nullptr)
+		{
+			ValidSkeletalMesh.Add(Mesh->GetUniqueID(), Mesh);
+		}
+	}
+
 	for (const UnFbx::FbxMeshInfo MeshInfo : SceneInfo.MeshInfo)
 	{
+		//Add the skeletal mesh if its a valid one
+		if (MeshInfo.bIsSkelMesh && !ValidSkeletalMesh.Contains(MeshInfo.UniqueId))
+		{
+			continue;
+		}
 		TSharedPtr<FFbxMeshInfo> MeshInfoPtr = MakeShareable(new FFbxMeshInfo());
 		MeshInfoPtr->FaceNum = MeshInfo.FaceNum;
 		MeshInfoPtr->VertexNum = MeshInfo.VertexNum;
@@ -280,15 +469,33 @@ TSharedPtr<FFbxSceneInfo> UFbxSceneImportFactory::ConvertSceneInfo(void* VoidFbx
 		MeshInfoPtr->MorphNum = MeshInfo.MorphNum;
 		MeshInfoPtr->Name = MeshInfo.Name;
 		MeshInfoPtr->UniqueId = MeshInfo.UniqueId;
+		MeshInfoPtr->OptionName = DefaultOptionName;
+
+		MeshInfoPtr->IsLod = MeshInfoPtr->LODLevel > 0;
+		MeshInfoPtr->IsCollision = MeshInfoPtr->Name.Contains(TEXT("UCX")) || MeshInfoPtr->Name.Contains(TEXT("UBX")) || MeshInfoPtr->Name.Contains(TEXT("MCDCX")) || MeshInfoPtr->Name.Contains(TEXT("USP"));
+
 		SceneInfoPtr->MeshInfo.Add(MeshInfoPtr);
 	}
+	
+	//Find all light and camera in the scene
+	FetchFbxCameraInScene(FbxImporter, RootNodeToImport, SceneInfoPtr);
+	FetchFbxLightInScene(FbxImporter, RootNodeToImport, SceneInfoPtr);
+
+	TMap<uint64, const UnFbx::FbxNodeInfo*> NodeInfoMap;
 	for (const UnFbx::FbxNodeInfo& NodeInfo : SceneInfo.HierarchyInfo)
 	{
+		NodeInfoMap.Add(NodeInfo.UniqueId, &NodeInfo);
+		if (IsPartOfSkeletonHierarchy(NodeInfoMap, NodeInfo))
+		{
+			continue;
+		}
+
 		TSharedPtr<FFbxNodeInfo> NodeInfoPtr = MakeShareable(new FFbxNodeInfo());
 		NodeInfoPtr->NodeName = NodeInfo.ObjectName;
 		NodeInfoPtr->UniqueId = NodeInfo.UniqueId;
 		NodeInfoPtr->AttributeType = NodeInfo.AttributeType;
-		
+		NodeInfoPtr->AttributeUniqueId = NodeInfo.AttributeUniqueId;
+
 		//Find the parent
 		NodeInfoPtr->ParentNodeInfo = nullptr;
 		for (TSharedPtr<FFbxNodeInfo> ParentPtr : SceneInfoPtr->HierarchyInfo)
@@ -320,6 +527,48 @@ TSharedPtr<FFbxSceneInfo> UFbxSceneImportFactory::ConvertSceneInfo(void* VoidFbx
 		NodeInfoPtr->Transform.SetTranslation(UnFbx::FFbxDataConverter::ConvertPos(NewLocalT));
 		NodeInfoPtr->Transform.SetScale3D(UnFbx::FFbxDataConverter::ConvertScale(NewLocalS));
 		NodeInfoPtr->Transform.SetRotation(UnFbx::FFbxDataConverter::ConvertRotToQuat(NewLocalQ));
+		NodeInfoPtr->PivotRotation = UnFbx::FFbxDataConverter::ConvertPos(NodeInfo.RotationPivot);
+		NodeInfoPtr->PivotScaling = UnFbx::FFbxDataConverter::ConvertPos(NodeInfo.ScalePivot);
+		
+		//Set the attribute pivot dictionary
+		if (NodeInfoPtr->AttributeInfo.IsValid())
+		{
+			if (NodeInfoPtr->AttributeInfo->NodeReferencePivots.Contains(NodeInfoPtr->PivotRotation))
+			{
+				TArray<uint64>*NodeUidArray = NodeInfoPtr->AttributeInfo->NodeReferencePivots.Find(NodeInfoPtr->PivotRotation);
+				check(NodeUidArray);
+				NodeUidArray->Add(NodeInfoPtr->UniqueId);
+			}
+			else
+			{
+				TArray<uint64> NodeUidArray;
+				NodeUidArray.Add(NodeInfoPtr->UniqueId);
+				NodeInfoPtr->AttributeInfo->NodeReferencePivots.Add(NodeInfoPtr->PivotRotation, NodeUidArray);
+
+			}
+			if (NodeInfoPtr->AttributeInfo->PivotNodeUid == INVALID_UNIQUE_ID)
+			{
+				NodeInfoPtr->AttributeInfo->PivotNodeUid = NodeInfoPtr->UniqueId;
+				NodeInfoPtr->AttributeInfo->PivotNodeName = NodeInfoPtr->NodeName;
+			}
+		}
+
+		if (SceneInfoPtr->LightInfo.Contains(NodeInfoPtr->AttributeUniqueId))
+		{
+			//Add the z rotation of 90 degree locally for every light. Light direction differ from fbx to unreal 
+			FRotator LightRotator(0.0f, 90.0f, 0.0f);
+			FTransform LightTransform = FTransform(LightRotator);
+			NodeInfoPtr->Transform = LightTransform * NodeInfoPtr->Transform;
+		}
+		else if (SceneInfoPtr->CameraInfo.Contains(NodeInfoPtr->AttributeUniqueId))
+		{
+			//Add a roll of -90 degree locally for every cameras. Camera up vector differ from fbx to unreal
+			FRotator LightRotator(0.0f, 0.0f, -90.0f);
+			FTransform LightTransform = FTransform(LightRotator);
+			//Remove the scale of the node holding a camera (the mesh is provide by the engine and can be different in size)
+			NodeInfoPtr->Transform.SetScale3D(FVector(1.0f));
+			NodeInfoPtr->Transform = LightTransform * NodeInfoPtr->Transform;
+		}
 
 		//by default we import all node
 		NodeInfoPtr->bImportNode = true;
@@ -327,6 +576,21 @@ TSharedPtr<FFbxSceneInfo> UFbxSceneImportFactory::ConvertSceneInfo(void* VoidFbx
 		//Add the node to the hierarchy
 		SceneInfoPtr->HierarchyInfo.Add(NodeInfoPtr);
 	}
+
+	for (TSharedPtr<FFbxNodeInfo> NodeInfo : SceneInfoPtr->HierarchyInfo)
+	{
+		if (NodeInfo->AttributeType.Compare(TEXT("eLODGroup")) == 0)
+		{
+			for (TSharedPtr<FFbxNodeInfo> ChildNodeInfo : NodeInfo->Childrens)
+			{
+				if (ChildNodeInfo->AttributeType.Compare(TEXT("eMesh")) != 0)
+				{
+					//We don't import under LOD group other stuff then the mesh
+					ChildNodeInfo->bImportNode = false;
+				}
+			}
+		}
+	}
 	return SceneInfoPtr;
 }
 
@@ -403,8 +667,6 @@ UFbxSceneImportFactory::UFbxSceneImportFactory(const FObjectInitializer& ObjectI
 	SceneImportOptions = CreateDefaultSubobject<UFbxSceneImportOptions>(TEXT("SceneImportOptions"), true);
 	SceneImportOptionsStaticMesh = CreateDefaultSubobject<UFbxSceneImportOptionsStaticMesh>(TEXT("SceneImportOptionsStaticMesh"), true);
 	SceneImportOptionsSkeletalMesh = CreateDefaultSubobject<UFbxSceneImportOptionsSkeletalMesh>(TEXT("SceneImportOptionsSkeletalMesh"), true);
-	SceneImportOptionsAnimation = CreateDefaultSubobject<UFbxSceneImportOptionsAnimation>(TEXT("SceneImportOptionsAnimation"), true);
-	SceneImportOptionsMaterial = CreateDefaultSubobject<UFbxSceneImportOptionsMaterial>(TEXT("SceneImportOptionsMaterial"), true);
 
 	StaticMeshImportData = CreateDefaultSubobject<UFbxStaticMeshImportData>(TEXT("StaticMeshImportData"), true);
 	SkeletalMeshImportData = CreateDefaultSubobject<UFbxSkeletalMeshImportData>(TEXT("SkeletalMeshImportData"), true);
@@ -474,7 +736,7 @@ UFbxSceneImportData* CreateReImportAsset(const FString &PackagePath, const FStri
 	{
 		NewUniqueName = MakeUniqueObjectName(ReImportAsset, UFbxSceneImportData::StaticClass(), FName(*AssetName)).ToString();
 	}
-	ReImportAsset->Rename(*NewUniqueName);
+	ReImportAsset->Rename(*NewUniqueName, nullptr, REN_DontCreateRedirectors);
 	ReImportAsset->SceneInfoSourceData = SceneInfo;
 	//Copy the options map
 	for (auto kvp : NameOptionsMap)
@@ -555,7 +817,7 @@ FFeedbackContext*	Warn
 	GlobalImportSettings->bConvertScene = true;
 
 	//Read the fbx and store the hierarchy's information so we can reuse it after importing all the model in the fbx file
-	if (!FbxImporter->ImportFromFile(*FbxImportFileName, Type))
+	if (!FbxImporter->ImportFromFile(*FbxImportFileName, Type, true))
 	{
 		// Log the error message and fail the import.
 		Warn->Log(ELogVerbosity::Error, FbxImporter->GetErrorMessage());
@@ -569,19 +831,23 @@ FFeedbackContext*	Warn
 		return nullptr;
 	}
 
+	//Make sure the Skeleton is null and not garbage, as we are importing the skeletalmesh for the first time we do not need any skeleton
+	GlobalImportSettings->SkeletonForAnimation = nullptr;
+	GlobalImportSettings->PhysicsAsset = nullptr;
+
 	FString PackageName = "";
 	InParent->GetName(PackageName);
 	Path = FPaths::GetPath(PackageName);
 
 	UnFbx::FbxSceneInfo SceneInfo;
 	//Read the scene and found all instance with their scene information.
-	FbxImporter->GetSceneInfo(FbxImportFileName, SceneInfo);
+	FbxImporter->GetSceneInfo(FbxImportFileName, SceneInfo, true);
 	
 	GlobalImportSettingsReference = new UnFbx::FBXImportOptions();
 	SFbxSceneOptionWindow::CopyFbxOptionsToFbxOptions(GlobalImportSettings, GlobalImportSettingsReference);
 
 	//Convert old structure to the new scene export structure
-	TSharedPtr<FFbxSceneInfo> SceneInfoPtr = ConvertSceneInfo(&SceneInfo);
+	TSharedPtr<FFbxSceneInfo> SceneInfoPtr = ConvertSceneInfo(FbxImporter, &SceneInfo);
 	
 	//Get import material info
 	ExtractMaterialInfo(FbxImporter, SceneInfoPtr);
@@ -593,8 +859,6 @@ FFeedbackContext*	Warn
 		, SceneImportOptionsStaticMesh
 		, NameOptionsMap
 		, SceneImportOptionsSkeletalMesh
-		, SceneImportOptionsAnimation
-		, SceneImportOptionsMaterial
 		, Path))
 	{
 		//User cancel the scene import
@@ -636,10 +900,13 @@ FFeedbackContext*	Warn
 		}
 	}
 
-	//We are a scene import set the flag for the reimport factory
+	//We are a scene import set the flag for the reimport factory for both static mesh and skeletal mesh
 	StaticMeshImportData->bImportAsScene = true;
 	StaticMeshImportData->FbxSceneImportDataReference = ReimportData;
 
+	SkeletalMeshImportData->bImportAsScene = true;
+	SkeletalMeshImportData->FbxSceneImportDataReference = ReimportData;
+
 	//Get the scene root node
 	FbxNode* RootNodeToImport = nullptr;
 	RootNodeToImport = FbxImporter->Scene->GetRootNode();
@@ -706,14 +973,23 @@ FFeedbackContext*	Warn
 				UBlueprint* SceneBlueprint = FKismetEditorUtilities::CreateBlueprintFromActor(FullnameBP, HierarchyActor, true);
 				if (SceneBlueprint != nullptr && ReimportData != nullptr)
 				{
+					//let the scene blueprint be the return object for this import
+					ReturnObject = SceneBlueprint;
 					//Set the blueprint path name in the re import scene data asset, this will allow re import to find the original import blueprint
 					ReimportData->BluePrintFullName = SceneBlueprint->GetPathName();
 				}
+				GEngine->BroadcastLevelActorListChanged();
 			}
 		}
 		break;
 	}
-
+	
+	//If there is no content asset create return the fbx scene import data
+	//This can happen if we only import actor in the scene like lights and camera
+	if (ReturnObject == nullptr)
+	{
+		ReturnObject = ReimportData;
+	}
 	//Release the FbxImporter 
 	FbxImporter->ReleaseScene();
 	FbxImporter = nullptr;
@@ -779,6 +1055,60 @@ bool UFbxSceneImportFactory::SetStaticMeshComponentOverrideMaterial(UStaticMeshC
 	return bOverrideMaterial;
 }
 
+USceneComponent *CreateCameraComponent(AActor *ParentActor, TSharedPtr<FFbxCameraInfo> CameraInfo)
+{
+	UCameraComponent *CameraComponent = NewObject<UCameraComponent>(ParentActor, *(CameraInfo->Name));
+	CameraComponent->SetProjectionMode(CameraInfo->ProjectionPerspective ? ECameraProjectionMode::Perspective : ECameraProjectionMode::Orthographic);
+	CameraComponent->SetAspectRatio(CameraInfo->AspectWidth / CameraInfo->AspectHeight);
+	CameraComponent->SetOrthoNearClipPlane(CameraInfo->NearPlane);
+	CameraComponent->SetOrthoFarClipPlane(CameraInfo->FarPlane);
+	CameraComponent->SetOrthoWidth(CameraInfo->AspectWidth);
+	return CameraComponent;
+}
+
+USceneComponent *CreateLightComponent(AActor *ParentActor, TSharedPtr<FFbxLightInfo> LightInfo)
+{
+	ULightComponent *LightComponent = nullptr;
+	switch (LightInfo->Type)
+	{
+	case 0:
+	{
+		//Point light
+		UPointLightComponent *PointLightComponent = NewObject<UPointLightComponent>(ParentActor, *(LightInfo->Name));
+		PointLightComponent->SetAttenuationRadius(LightInfo->EnableFarAttenuation ? LightInfo->FarAttenuationEnd : 16384.0f);
+		LightComponent = static_cast<ULightComponent*>(PointLightComponent);
+		LightComponent->SetIntensity(LightInfo->Intensity);
+	}
+		break;
+	case 1:
+	{
+		//Directional light
+		UDirectionalLightComponent *DirectionalLightComponent = NewObject<UDirectionalLightComponent>(ParentActor, *(LightInfo->Name));
+		LightComponent = static_cast<ULightComponent*>(DirectionalLightComponent);
+		LightComponent->SetIntensity(LightInfo->Intensity/100); //In unreal this is a ratio not a lumen value
+	}
+		break;
+	case 2:
+	{
+		//Spot light
+		USpotLightComponent *SpotLightComponent = NewObject<USpotLightComponent>(ParentActor, *(LightInfo->Name));
+		SpotLightComponent->SetInnerConeAngle(LightInfo->InnerAngle/2.0f);
+		SpotLightComponent->SetOuterConeAngle(LightInfo->OuterAngle/2.0f);
+		SpotLightComponent->SetAttenuationRadius(LightInfo->EnableFarAttenuation ? LightInfo->FarAttenuationEnd : 16384.0f);
+		LightComponent = static_cast<ULightComponent*>(SpotLightComponent);
+		LightComponent->SetIntensity(LightInfo->Intensity);
+	}
+		break;
+	case 3:
+	case 4:
+		return nullptr;
+		break;
+	}
+	LightComponent->SetLightColor(LightInfo->Color);
+	LightComponent->SetCastShadows(LightInfo->CastShadow);
+	return LightComponent;
+}
+
 void UFbxSceneImportFactory::CreateLevelActorHierarchy(TSharedPtr<FFbxSceneInfo> SceneInfoPtr)
 {
 	EComponentMobility::Type MobilityType = SceneImportOptions->bImportAsDynamic ? EComponentMobility::Movable : EComponentMobility::Static;
@@ -800,9 +1130,16 @@ void UFbxSceneImportFactory::CreateLevelActorHierarchy(TSharedPtr<FFbxSceneInfo>
 		{
 			continue;
 		}
+		
+		TSharedPtr<FFbxNodeInfo> LODParentNodeInfo = nullptr;
+		if (NodeInfo->AttributeType.Compare(TEXT("eMesh")) == 0)
+		{
+			LODParentNodeInfo = FFbxSceneInfo::RecursiveFindLODParentNode(NodeInfo);
+		}
+
 		//Find the asset that link with this node attribute
 		UObject *AssetToPlace = (NodeInfo->AttributeInfo.IsValid() && AllNewAssets.Contains(NodeInfo->AttributeInfo)) ? AllNewAssets[NodeInfo->AttributeInfo] : nullptr;
-
+		
 		//create actor
 		AActor *PlacedActor = nullptr;
 		if (AssetToPlace != nullptr)
@@ -818,7 +1155,7 @@ void UFbxSceneImportFactory::CreateLevelActorHierarchy(TSharedPtr<FFbxSceneInfo>
 				SetStaticMeshComponentOverrideMaterial(StaticMeshComponent, NodeInfo);
 			}
 		}
-		else if (IsEmptyAttribute(NodeInfo->AttributeType) || NodeInfo->AttributeType.Compare("eMesh") == 0)
+		else if (IsEmptyAttribute(NodeInfo->AttributeType) || NodeInfo->AttributeType.Compare("eMesh") == 0 || NodeInfo->AttributeUniqueId != INVALID_UNIQUE_ID)
 		{
 			//Create an empty actor if the node is an empty attribute or the attribute is a mesh(static mesh or skeletal mesh) that was not export
 			UActorFactory* Factory = GEditor->FindActorFactoryByClass(UActorFactoryEmptyActor::StaticClass());
@@ -827,7 +1164,27 @@ void UFbxSceneImportFactory::CreateLevelActorHierarchy(TSharedPtr<FFbxSceneInfo>
 			UObject* EmptyActorAsset = EmptyActorAssetData.GetAsset();
 			//Place an empty actor
 			PlacedActor = FActorFactoryAssetProxy::AddActorForAsset(EmptyActorAsset, bSelectActor);
-			USceneComponent* RootComponent = NewObject<USceneComponent>(PlacedActor, USceneComponent::GetDefaultSceneRootVariableName());
+			USceneComponent* RootComponent = nullptr;
+			if(NodeInfo->AttributeType.Compare("eLight") == 0)
+			{
+				TSharedPtr<FFbxLightInfo> LightInfo = *SceneInfoPtr->LightInfo.Find(NodeInfo->AttributeUniqueId);
+				RootComponent = CreateLightComponent(PlacedActor, LightInfo);
+			}
+			else if (NodeInfo->AttributeType.Compare("eCamera") == 0)
+			{
+				TSharedPtr<FFbxCameraInfo> CameraInfo = *SceneInfoPtr->CameraInfo.Find(NodeInfo->AttributeUniqueId);
+				RootComponent = CreateCameraComponent(PlacedActor, CameraInfo);
+			}
+			
+			if(RootComponent == nullptr)
+			{
+				if (LODParentNodeInfo.IsValid())
+				{
+					//This is not LOD index 0, don't export the transform. Lod 0 should have an asset to place
+					continue;
+				}
+				RootComponent = NewObject<USceneComponent>(PlacedActor, USceneComponent::GetDefaultSceneRootVariableName());
+			}
 			RootComponent->Mobility = MobilityType;
 			RootComponent->bVisualizeComponent = true;
 			PlacedActor->SetRootComponent(RootComponent);
@@ -853,6 +1210,10 @@ void UFbxSceneImportFactory::CreateLevelActorHierarchy(TSharedPtr<FFbxSceneInfo>
 				//Map the new actor name with the old name in case the name is changing
 				NewActorNameMap.Add(NodeInfo->UniqueId, PlacedActor);
 				uint64 ParentUniqueId = NodeInfo->ParentNodeInfo.IsValid() ? NodeInfo->ParentNodeInfo->UniqueId : 0;
+				if (LODParentNodeInfo.IsValid())
+				{
+					ParentUniqueId = LODParentNodeInfo->UniqueId;
+				}
 				AActor *ParentActor = nullptr;
 				//If there is a parent we must set the parent actor
 				if (NewActorNameMap.Contains(ParentUniqueId))
@@ -870,8 +1231,22 @@ void UFbxSceneImportFactory::CreateLevelActorHierarchy(TSharedPtr<FFbxSceneInfo>
 						}
 					}
 				}
+				//Find the pivot location
+				FVector PivotLocation(0.0f);
+				if (GlobalImportSettings->bBakePivotInVertex && NodeInfo->AttributeInfo.IsValid() && NodeInfo->AttributeInfo->PivotNodeUid != INVALID_UNIQUE_ID)
+				{
+					for (TSharedPtr<FFbxNodeInfo> NodeInfoIter : SceneInfoPtr->HierarchyInfo)
+					{
+						if (NodeInfoIter->UniqueId == NodeInfo->AttributeInfo->PivotNodeUid)
+						{
+							PivotLocation = NodeInfoIter->PivotRotation;
+							break;
+						}
+					}
+				}
+
 				//Apply the hierarchy local transform to the root component
-				ApplyTransformToComponent(RootComponent, &(NodeInfo->Transform), ParentActor == nullptr ? &RootTransform : nullptr);
+				ApplyTransformToComponent(RootComponent, &(NodeInfo->Transform), ParentActor == nullptr ? &RootTransform : nullptr, PivotLocation);
 			}
 		}
 		//We select only the first actor
@@ -898,7 +1273,7 @@ AActor *UFbxSceneImportFactory::CreateActorComponentsHierarchy(TSharedPtr<FFbxSc
 	//This is a group create an empty actor that just have a transform
 	UObject* EmptyActorAsset = EmptyActorAssetData.GetAsset();
 	//Place an empty actor
-	RootActorContainer = FActorFactoryAssetProxy::AddActorForAsset(EmptyActorAsset, true);
+	RootActorContainer = FActorFactoryAssetProxy::AddActorForAsset(EmptyActorAsset, false);
 	check(RootActorContainer != nullptr);
 	ActorRootComponent = NewObject<USceneComponent>(RootActorContainer, USceneComponent::GetDefaultSceneRootVariableName());
 	check(ActorRootComponent != nullptr);
@@ -926,6 +1301,12 @@ AActor *UFbxSceneImportFactory::CreateActorComponentsHierarchy(TSharedPtr<FFbxSc
 		{
 			continue;
 		}
+
+		TSharedPtr<FFbxNodeInfo> LODParentNodeInfo = nullptr;
+		if (NodeInfo->AttributeType.Compare(TEXT("eMesh")) == 0)
+		{
+			LODParentNodeInfo = FFbxSceneInfo::RecursiveFindLODParentNode(NodeInfo);
+		}
 		//Find the asset that link with this node attribute
 		UObject *AssetToPlace = (NodeInfo->AttributeInfo.IsValid() && AllNewAssets.Contains(NodeInfo->AttributeInfo)) ? AllNewAssets[NodeInfo->AttributeInfo] : nullptr;
 
@@ -954,15 +1335,33 @@ AActor *UFbxSceneImportFactory::CreateActorComponentsHierarchy(TSharedPtr<FFbxSc
 				SceneComponent->Mobility = MobilityType;
 			}
 		}
-		else if (IsEmptyAttribute(NodeInfo->AttributeType) || NodeInfo->AttributeType.Compare("eMesh") == 0)
+		else if (IsEmptyAttribute(NodeInfo->AttributeType) || NodeInfo->AttributeType.Compare("eMesh") == 0 || NodeInfo->AttributeUniqueId != INVALID_UNIQUE_ID)
 		{
+			if (NodeInfo->AttributeType.Compare("eLight") == 0 && SceneInfoPtr->LightInfo.Contains(NodeInfo->AttributeUniqueId))
+			{
+				TSharedPtr<FFbxLightInfo> LightInfo = *SceneInfoPtr->LightInfo.Find(NodeInfo->AttributeUniqueId);
+				SceneComponent = CreateLightComponent(RootActorContainer, LightInfo);
+			}
+			else if (NodeInfo->AttributeType.Compare("eCamera") == 0 && SceneInfoPtr->CameraInfo.Contains(NodeInfo->AttributeUniqueId))
+			{
+				TSharedPtr<FFbxCameraInfo> CameraInfo = *SceneInfoPtr->CameraInfo.Find(NodeInfo->AttributeUniqueId);
+				SceneComponent = CreateCameraComponent(RootActorContainer, CameraInfo);
+			}
+			
+			if(SceneComponent == nullptr)
+			{
+				if (LODParentNodeInfo.IsValid())
+				{
+					//This is not LOD index 0, don't export the transform. Lod 0 should have an asset to place
+					continue;
+				}
+				SceneComponent = NewObject<USceneComponent>(RootActorContainer, NAME_None);
+			}
 			//Component will be rename later
-			SceneComponent = NewObject<USceneComponent>(RootActorContainer, NAME_None);
 			SceneComponent->Mobility = MobilityType;
 		}
 		else
 		{
-			//TODO log which fbx attribute we cannot create an actor from
 			continue;
 		}
 
@@ -973,7 +1372,7 @@ AActor *UFbxSceneImportFactory::CreateActorComponentsHierarchy(TSharedPtr<FFbxSc
 		{
 			NewUniqueName = MakeUniqueObjectName(RootActorContainer, USceneComponent::StaticClass(), FName(*NodeInfo->NodeName)).ToString();
 		}
-		SceneComponent->Rename(*NewUniqueName);
+		SceneComponent->Rename(*NewUniqueName, nullptr, REN_DontCreateRedirectors);
 
 		//Add the component to the owner actor and register it
 		RootActorContainer->AddInstanceComponent(SceneComponent);
@@ -985,6 +1384,10 @@ AActor *UFbxSceneImportFactory::CreateActorComponentsHierarchy(TSharedPtr<FFbxSc
 		//Find the parent component by unique ID and attach(as child) the newly created scenecomponent
 		//Attach the component to the rootcomponent if we dont find any parent component
 		uint64 ParentUniqueId = NodeInfo->ParentNodeInfo.IsValid() ? NodeInfo->ParentNodeInfo->UniqueId : 0;
+		if (LODParentNodeInfo.IsValid())
+		{
+			ParentUniqueId = LODParentNodeInfo->UniqueId;
+		}
 		USceneComponent *ParentRootComponent = nullptr;
 		if (NewSceneComponentNameMap.Contains(ParentUniqueId))
 		{
@@ -999,8 +1402,23 @@ AActor *UFbxSceneImportFactory::CreateActorComponentsHierarchy(TSharedPtr<FFbxSc
 			SceneComponent->AttachTo(ActorRootComponent, NAME_None, EAttachLocation::KeepWorldPosition);
 		}
 
+		//Find the pivot location
+		FVector PivotLocation(0.0f);
+		if (GlobalImportSettings->bBakePivotInVertex && NodeInfo->AttributeInfo.IsValid() && NodeInfo->AttributeInfo->PivotNodeUid != INVALID_UNIQUE_ID)
+		{
+			for (TSharedPtr<FFbxNodeInfo> NodeInfoIter : SceneInfoPtr->HierarchyInfo)
+			{
+				if (NodeInfoIter->UniqueId == NodeInfo->AttributeInfo->PivotNodeUid)
+				{
+					PivotLocation = NodeInfoIter->PivotRotation;
+					break;
+				}
+			}
+		}
+
 		//Apply the local transform to the scene component
-		ApplyTransformToComponent(SceneComponent, &(NodeInfo->Transform), ParentRootComponent != nullptr ? nullptr : &RootTransform);
+		ApplyTransformToComponent(SceneComponent, &(NodeInfo->Transform), ParentRootComponent != nullptr ? nullptr : &RootTransform, PivotLocation);
+		
 	}
 	// End of iteration of the hierarchy
 	//////////////////////////////////////////////////////////////////////////
@@ -1008,107 +1426,173 @@ AActor *UFbxSceneImportFactory::CreateActorComponentsHierarchy(TSharedPtr<FFbxSc
 	return RootActorContainer;
 }
 
-void UFbxSceneImportFactory::ApplyTransformToComponent(USceneComponent *SceneComponent, FTransform *LocalTransform, FTransform *PreMultiplyTransform)
+void UFbxSceneImportFactory::ApplyTransformToComponent(USceneComponent *SceneComponent, FTransform *LocalTransform, FTransform *PreMultiplyTransform, FVector &PivotLocation)
 {
+	check(SceneComponent);
+	check(LocalTransform);
+	FTransform LocalTransformAdjusted = (*LocalTransform);
+	if (GlobalImportSettings->bBakePivotInVertex && !PivotLocation.IsNearlyZero())
+	{
+		FTransform PivotTransform;
+		PivotTransform.SetLocation(PivotLocation);
+		LocalTransformAdjusted = PivotTransform * LocalTransformAdjusted;
+	}
 	//In case there is no parent we must multiply the root transform
 	if (PreMultiplyTransform != nullptr)
 	{
 		FTransform OutTransform = FTransform::Identity;
-		FTransform::Multiply(&OutTransform, LocalTransform, PreMultiplyTransform);
+		FTransform::Multiply(&OutTransform, &LocalTransformAdjusted, PreMultiplyTransform);
 		SceneComponent->SetRelativeTransform(OutTransform);
 	}
 	else
 	{
-		SceneComponent->SetRelativeTransform(*LocalTransform);
+		SceneComponent->SetRelativeTransform(LocalTransformAdjusted);
 	}
 }
-void UFbxSceneImportFactory::ImportAllSkeletalMesh(void* VoidRootNodeToImport, void* VoidFbxImporter, EObjectFlags Flags, int32& NodeIndex, int32& InterestingNodeCount, TSharedPtr<FFbxSceneInfo> SceneInfo)
+
+
+UObject* UFbxSceneImportFactory::ImportOneSkeletalMesh(void* VoidRootNodeToImport, void* VoidFbxImporter, TSharedPtr<FFbxSceneInfo> SceneInfo, EObjectFlags Flags, TArray<void*> &VoidNodeArray, int32 &TotalNumNodes)
 {
-	UnFbx::FFbxImporter* FbxImporter = (UnFbx::FFbxImporter*)VoidFbxImporter;
 	FbxNode *RootNodeToImport = (FbxNode *)VoidRootNodeToImport;
-	InterestingNodeCount = 1;
-	TArray< TArray<FbxNode*>* > SkelMeshArray;
-	FbxImporter->FillFbxSkelMeshArrayInScene(RootNodeToImport, SkelMeshArray, false);
-	InterestingNodeCount = SkelMeshArray.Num();
+	UnFbx::FFbxImporter* FbxImporter = (UnFbx::FFbxImporter*)VoidFbxImporter;
+	TArray<FbxNode*> NodeArray;
+	for (void *VoidNode : VoidNodeArray)
+	{
+		FbxNode* Node = (FbxNode*)VoidNode;
+		NodeArray.Add(Node);
+	}
+	UObject* NewObject = nullptr;
+	UPackage* Pkg = nullptr;
+	TotalNumNodes += NodeArray.Num();
+	TSharedPtr<FFbxNodeInfo> RootNodeInfo;
+	if (TotalNumNodes > 0)
+	{
+		FbxNode* RootNodeArrayNode = NodeArray[0];
+		if (!FindSceneNodeInfo(SceneInfo, RootNodeArrayNode->GetUniqueID(), RootNodeInfo))
+		{
+			return nullptr;
+		}
+		if (!RootNodeInfo->AttributeInfo.IsValid() || RootNodeInfo->AttributeInfo->GetType() != USkeletalMesh::StaticClass() || !RootNodeInfo->AttributeInfo->bImportAttribute)
+		{
+			return nullptr;
+		}
+	}
+	if (!RootNodeInfo.IsValid())
+	{
+		return nullptr;
+	}
+	//Set the options
+	//Apply the correct fbx options
+	TSharedPtr<FFbxMeshInfo> MeshInfo = StaticCastSharedPtr<FFbxMeshInfo>(RootNodeInfo->AttributeInfo);
 
-	int32 TotalNumNodes = 0;
+	UnFbx::FBXImportOptions* OverrideImportSettings = GetOptionsFromName(MeshInfo->OptionName);
+	if (OverrideImportSettings != nullptr)
+	{
+		SFbxSceneOptionWindow::CopyFbxOptionsToFbxOptions(OverrideImportSettings, GlobalImportSettings);
+		SFbxSceneOptionWindow::CopyFbxOptionsToSkeletalMeshOptions(OverrideImportSettings, SceneImportOptionsSkeletalMesh);
+		SceneImportOptionsSkeletalMesh->FillSkeletalMeshInmportData(SkeletalMeshImportData, AnimSequenceImportData);
+	}
+	else
+	{
+		//Use the default options if we dont found options
+		SFbxSceneOptionWindow::CopyFbxOptionsToFbxOptions(GlobalImportSettingsReference, GlobalImportSettings);
+		SFbxSceneOptionWindow::CopyFbxOptionsToSkeletalMeshOptions(GlobalImportSettingsReference, SceneImportOptionsSkeletalMesh);
+		SceneImportOptionsSkeletalMesh->FillSkeletalMeshInmportData(SkeletalMeshImportData, AnimSequenceImportData);
+	}
 
-	for (int32 i = 0; i < SkelMeshArray.Num(); i++)
+	//This is possible for the user to not bake pivot for a particular root node
+	if (GlobalImportSettings->bBakePivotInVertex && RootNodeInfo->AttributeInfo->PivotNodeUid == INVALID_UNIQUE_ID)
 	{
-		UObject* NewObject = nullptr;
-		TArray<FbxNode*> NodeArray = *SkelMeshArray[i];
-		UPackage* Pkg = nullptr;
-		TotalNumNodes += NodeArray.Num();
-		TSharedPtr<FFbxNodeInfo> RootNodeInfo;
-		if (TotalNumNodes > 0)
+		GlobalImportSettings->bBakePivotInVertex = false;
+	}
+
+	// check if there is LODGroup for this skeletal mesh
+	int32 MaxLODLevel = 1;
+	for (int32 j = 0; j < NodeArray.Num(); j++)
+	{
+		FbxNode* Node = NodeArray[j];
+		if (Node->GetNodeAttribute() && Node->GetNodeAttribute()->GetAttributeType() == FbxNodeAttribute::eLODGroup)
 		{
-			FbxNode* RootNodeArrayNode = NodeArray[0];
-			if (!FindSceneNodeInfo(SceneInfo, RootNodeArrayNode->GetUniqueID(), RootNodeInfo))
+			// get max LODgroup level
+			if (MaxLODLevel < Node->GetChildCount())
 			{
-				continue;
-			}
-			if (!RootNodeInfo->AttributeInfo.IsValid() || RootNodeInfo->AttributeInfo->GetType() != USkeletalMesh::StaticClass())
-			{
-				continue;
+				MaxLODLevel = Node->GetChildCount();
 			}
 		}
-		// check if there is LODGroup for this skeletal mesh
-		int32 MaxLODLevel = 1;
+	}
+
+	int32 LODIndex;
+	for (LODIndex = 0; LODIndex < MaxLODLevel; LODIndex++)
+	{
+		TArray<FbxNode*> SkelMeshNodeArray;
 		for (int32 j = 0; j < NodeArray.Num(); j++)
 		{
 			FbxNode* Node = NodeArray[j];
 			if (Node->GetNodeAttribute() && Node->GetNodeAttribute()->GetAttributeType() == FbxNodeAttribute::eLODGroup)
 			{
-				// get max LODgroup level
-				if (MaxLODLevel < Node->GetChildCount())
+				if (Node->GetChildCount() > LODIndex)
 				{
-					MaxLODLevel = Node->GetChildCount();
+					FbxNode *MeshNode = FbxImporter->FindLODGroupNode(Node, LODIndex);
+					if(MeshNode != nullptr)
+						SkelMeshNodeArray.Add(MeshNode);
 				}
+				else // in less some LODGroups have less level, use the last level
+				{
+					FbxNode *MeshNode = FbxImporter->FindLODGroupNode(Node, Node->GetChildCount() - 1);
+					if (MeshNode != nullptr)
+						SkelMeshNodeArray.Add(MeshNode);
+				}
+			}
+			else
+			{
+				SkelMeshNodeArray.Add(Node);
 			}
 		}
-		int32 LODIndex;
-		for (LODIndex = 0; LODIndex < MaxLODLevel; LODIndex++)
+
+		//Make sure to bake the pivot the user choose to bake
+		TArray<FbxNode*> SkelMeshNodePivotArray;
+		bool bUseSkelMeshNodePivotArray = false;
+		if (GlobalImportSettings->bBakePivotInVertex)
 		{
-			TArray<FbxNode*> SkelMeshNodeArray;
-			for (int32 j = 0; j < NodeArray.Num(); j++)
+			for (FbxNode *SkelMeshNode : SkelMeshNodeArray)
 			{
-				FbxNode* Node = NodeArray[j];
-				if (Node->GetNodeAttribute() && Node->GetNodeAttribute()->GetAttributeType() == FbxNodeAttribute::eLODGroup)
+				TSharedPtr<FFbxNodeInfo> ExportNodeInfo;
+				if (FindSceneNodeInfo(SceneInfo, SkelMeshNode->GetUniqueID(), ExportNodeInfo))
 				{
-					if (Node->GetChildCount() > LODIndex)
-					{
-						SkelMeshNodeArray.Add(Node->GetChild(LODIndex));
-					}
-					else // in less some LODGroups have less level, use the last level
+					FbxNode *NodePivot = FindFbxNodeById(FbxImporter, nullptr, ExportNodeInfo->AttributeInfo->PivotNodeUid);
+					if (NodePivot != nullptr)
 					{
-						SkelMeshNodeArray.Add(Node->GetChild(Node->GetChildCount() - 1));
+						SkelMeshNodePivotArray.Add(NodePivot);
+						bUseSkelMeshNodePivotArray = true;
+						continue;
 					}
 				}
-				else
-				{
-					SkelMeshNodeArray.Add(Node);
-				}
+				SkelMeshNodePivotArray.Add(SkelMeshNode);
 			}
+		}
 
-			if (LODIndex == 0 && SkelMeshNodeArray.Num() != 0)
+		if (LODIndex == 0 && SkelMeshNodeArray.Num() != 0)
+		{
+			FName OutputName = FbxImporter->MakeNameForMesh(SkelMeshNodeArray[0]->GetName(), SkelMeshNodeArray[0]);
+			FString PackageName = Path + TEXT("/") + OutputName.ToString();
+			FString SkeletalMeshName;
+			Pkg = CreatePackageForNode(PackageName, SkeletalMeshName);
+			if (Pkg == nullptr)
+				break;
+			RootNodeInfo->AttributeInfo->SetOriginalImportPath(PackageName);
+			FName SkeletalMeshFName = FName(*SkeletalMeshName);
+			//Import the skeletal mesh
+			USkeletalMesh* NewMesh = FbxImporter->ImportSkeletalMesh(Pkg, bUseSkelMeshNodePivotArray ? SkelMeshNodePivotArray : SkelMeshNodeArray, SkeletalMeshFName, Flags, SkeletalMeshImportData);
+			NewObject = NewMesh;
+			if (NewMesh)
 			{
-				FName OutputName = FbxImporter->MakeNameForMesh(SkelMeshNodeArray[0]->GetName(), SkelMeshNodeArray[0]);
-				FString PackageName = Path + TEXT("/") + OutputName.ToString();
-				FString SkeletalMeshName;
-				Pkg = CreatePackageForNode(PackageName, SkeletalMeshName);
-				if (Pkg == nullptr)
-					break;
-				RootNodeInfo->AttributeInfo->SetOriginalImportPath(PackageName);
-				FName SkeletalMeshFName = FName(*SkeletalMeshName);
-				USkeletalMesh* NewMesh = FbxImporter->ImportSkeletalMesh(Pkg, SkelMeshNodeArray, SkeletalMeshFName, Flags, SkeletalMeshImportData);
-				NewObject = NewMesh;
-				if (NewMesh)
+				TSharedPtr<FFbxNodeInfo> SkelMeshNodeInfo;
+				if (FindSceneNodeInfo(SceneInfo, SkelMeshNodeArray[0]->GetUniqueID(), SkelMeshNodeInfo) && SkelMeshNodeInfo.IsValid() && SkelMeshNodeInfo->AttributeInfo.IsValid())
+				{
+					AllNewAssets.Add(SkelMeshNodeInfo->AttributeInfo, NewObject);
+				}
+				if (GlobalImportSettings->bImportAnimations)
 				{
-					TSharedPtr<FFbxNodeInfo> SkelMeshNodeInfo;
-					if (FindSceneNodeInfo(SceneInfo, SkelMeshNodeArray[0]->GetUniqueID(), SkelMeshNodeInfo) && SkelMeshNodeInfo.IsValid() && SkelMeshNodeInfo->AttributeInfo.IsValid())
-					{
-						AllNewAssets.Add(SkelMeshNodeInfo->AttributeInfo, NewObject);
-					}
 					// We need to remove all scaling from the root node before we set up animation data.
 					// Othewise some of the global transform calculations will be incorrect.
 					FbxImporter->RemoveTransformSettingsFromFbxNode(RootNodeToImport, SkeletalMeshImportData);
@@ -1117,34 +1601,63 @@ void UFbxSceneImportFactory::ImportAllSkeletalMesh(void* VoidRootNodeToImport, v
 					// Reapply the transforms for the rest of the import
 					FbxImporter->ApplyTransformSettingsToFbxNode(RootNodeToImport, SkeletalMeshImportData);
 				}
+
+				//Set the data in the node info
+				RootNodeInfo->AttributeInfo->SetOriginalImportPath(PackageName);
+				RootNodeInfo->AttributeInfo->SetOriginalFullImportName(NewObject->GetPathName());
 			}
-			else if (NewObject) // the base skeletal mesh is imported successfully
+		}
+		else if (NewObject) // the base skeletal mesh is imported successfully
+		{
+			USkeletalMesh* BaseSkeletalMesh = Cast<USkeletalMesh>(NewObject);
+			FName LODObjectName = NAME_None;
+			//Import skeletal mesh LOD
+			USkeletalMesh *LODObject = FbxImporter->ImportSkeletalMesh(GetTransientPackage(), bUseSkelMeshNodePivotArray ? SkelMeshNodePivotArray : SkelMeshNodeArray, LODObjectName, RF_NoFlags, SkeletalMeshImportData);
+			bool bImportSucceeded = FbxImporter->ImportSkeletalMeshLOD(LODObject, BaseSkeletalMesh, LODIndex, false);
+			if (bImportSucceeded)
 			{
-				USkeletalMesh* BaseSkeletalMesh = Cast<USkeletalMesh>(NewObject);
-				FName LODObjectName = NAME_None;
-				USkeletalMesh *LODObject = FbxImporter->ImportSkeletalMesh(GetTransientPackage(), SkelMeshNodeArray, LODObjectName, RF_NoFlags, SkeletalMeshImportData);
-				bool bImportSucceeded = FbxImporter->ImportSkeletalMeshLOD(LODObject, BaseSkeletalMesh, LODIndex, false);
-
-				if (bImportSucceeded)
-				{
-					BaseSkeletalMesh->LODInfo[LODIndex].ScreenSize = 1.0f / (MaxLODLevel * LODIndex);
-				}
-				else
-				{
-					FbxImporter->AddTokenizedErrorMessage(FTokenizedMessage::Create(EMessageSeverity::Error, LOCTEXT("FailedToImport_SkeletalMeshLOD", "Failed to import Skeletal mesh LOD.")), FFbxErrors::SkeletalMesh_LOD_FailedToImport);
-				}
+				BaseSkeletalMesh->LODInfo[LODIndex].ScreenSize = 1.0f / (MaxLODLevel * LODIndex);
 			}
-
-			// import morph target
-			if (NewObject && SkeletalMeshImportData->bImportMorphTargets)
+			else
 			{
-				if (Pkg == nullptr)
-					continue;
-				// TODO: Disable material importing when importing morph targets
-				FbxImporter->ImportFbxMorphTarget(SkelMeshNodeArray, Cast<USkeletalMesh>(NewObject), Pkg, LODIndex);
+				FbxImporter->AddTokenizedErrorMessage(FTokenizedMessage::Create(EMessageSeverity::Error, LOCTEXT("FailedToImport_SkeletalMeshLOD", "Failed to import Skeletal mesh LOD.")), FFbxErrors::SkeletalMesh_LOD_FailedToImport);
 			}
 		}
 
+		// import morph target
+		if (NewObject && SkeletalMeshImportData->bImportMorphTargets)
+		{
+			if (Pkg == nullptr)
+				continue;
+			// TODO: Disable material importing when importing morph targets
+			FbxImporter->ImportFbxMorphTarget(SkelMeshNodeArray, Cast<USkeletalMesh>(NewObject), Pkg, LODIndex);
+		}
+	}
+
+	return NewObject;
+}
+
+void UFbxSceneImportFactory::ImportAllSkeletalMesh(void* VoidRootNodeToImport, void* VoidFbxImporter, EObjectFlags Flags, int32& NodeIndex, int32& InterestingNodeCount, TSharedPtr<FFbxSceneInfo> SceneInfo)
+{
+	UnFbx::FFbxImporter* FbxImporter = (UnFbx::FFbxImporter*)VoidFbxImporter;
+	FbxNode *RootNodeToImport = (FbxNode *)VoidRootNodeToImport;
+	InterestingNodeCount = 1;
+	TArray< TArray<FbxNode*>* > SkelMeshArray;
+	FbxImporter->FillFbxSkelMeshArrayInScene(RootNodeToImport, SkelMeshArray, false);
+	InterestingNodeCount = SkelMeshArray.Num();
+
+	int32 TotalNumNodes = 0;
+
+	for (int32 i = 0; i < SkelMeshArray.Num(); i++)
+	{
+		TArray<FbxNode*> NodeArray = *SkelMeshArray[i];
+		TArray<void*> VoidNodeArray;
+		for (FbxNode *Node : NodeArray)
+		{
+			void* VoidNode = (void*)Node;
+			VoidNodeArray.Add(VoidNode);
+		}
+		UObject* NewObject = ImportOneSkeletalMesh(VoidRootNodeToImport, VoidFbxImporter, SceneInfo, Flags, VoidNodeArray, TotalNumNodes);
 		if (NewObject)
 		{
 			NodeIndex++;
@@ -1157,7 +1670,7 @@ void UFbxSceneImportFactory::ImportAllSkeletalMesh(void* VoidRootNodeToImport, v
 	}
 
 	// if total nodes we found is 0, we didn't find anything. 
-	if (TotalNumNodes == 0)
+	if (SkelMeshArray.Num() > 0 && TotalNumNodes == 0)
 	{
 		FbxImporter->AddTokenizedErrorMessage(FTokenizedMessage::Create(EMessageSeverity::Error, LOCTEXT("FailedToImport_NoMeshFoundOnRoot", "Could not find any valid mesh on the root hierarchy. If you have mesh in the sub hierarchy, please enable option of [Import Meshes In Bone Hierarchy] when import.")),
 			FFbxErrors::SkeletalMesh_NoMeshFoundOnRoot);
@@ -1214,8 +1727,13 @@ UObject* UFbxSceneImportFactory::RecursiveImportNode(void* VoidFbxImporter, void
 
 	if (Node->GetNodeAttribute() && Node->GetNodeAttribute()->GetAttributeType() == FbxNodeAttribute::eLODGroup && Node->GetChildCount() > 0)
 	{
-		// import base mesh
-		NewObject = ImportANode(VoidFbxImporter, Node->GetChild(0), Flags, NodeIndex, SceneInfo, OutNodeInfo, PackagePath, Total);
+		//Find the deepest mesh child for the first LOD
+		FbxNode* MeshChildNode = FFbxImporter->FindLODGroupNode(Node, 0);
+		if (MeshChildNode != nullptr)
+		{
+			// import base mesh
+			NewObject = ImportANode(VoidFbxImporter, MeshChildNode, Flags, NodeIndex, SceneInfo, OutNodeInfo, PackagePath, Total);
+		}
 
 		if (NewObject)
 		{
@@ -1228,8 +1746,9 @@ UObject* UFbxSceneImportFactory::RecursiveImportNode(void* VoidFbxImporter, void
 				// import LOD meshes
 				for (int32 LODIndex = 1; LODIndex < Node->GetChildCount(); LODIndex++)
 				{
-					FbxNode* ChildNode = Node->GetChild(LODIndex);
-					ImportANode(VoidFbxImporter, ChildNode, Flags, NodeIndex, SceneInfo, OutNodeInfo, PackagePath, Total, NewObject, LODIndex);
+					FbxNode* ChildNode = FFbxImporter->FindLODGroupNode(Node, LODIndex);
+					if(ChildNode != nullptr)
+						ImportANode(VoidFbxImporter, ChildNode, Flags, NodeIndex, SceneInfo, OutNodeInfo, PackagePath, Total, NewObject, LODIndex);
 				}
 			}
 		}
@@ -1338,8 +1857,23 @@ UObject* UFbxSceneImportFactory::ImportANode(void* VoidFbxImporter, void* VoidNo
 		SFbxSceneOptionWindow::CopyFbxOptionsToStaticMeshOptions(GlobalImportSettingsReference, SceneImportOptionsStaticMesh);
 		SceneImportOptionsStaticMesh->FillStaticMeshInmportData(StaticMeshImportData, SceneImportOptions);
 	}
+	if (GlobalImportSettings->bBakePivotInVertex && OutNodeInfo->AttributeInfo->PivotNodeUid == INVALID_UNIQUE_ID)
+	{
+		GlobalImportSettings->bBakePivotInVertex = false;
+	}
 	FName StaticMeshFName = FName(*(OutNodeInfo->AttributeInfo->Name));
+	//Make sure to bake the pivot the user choose to bake
+	if (GlobalImportSettings->bBakePivotInVertex)
+	{
+		FbxNode *NodePivot = FindFbxNodeById(FFbxImporter, nullptr, OutNodeInfo->AttributeInfo->PivotNodeUid);
+		if (NodePivot != nullptr)
+		{
+			Node = NodePivot;
+		}
+	}
+	
 	NewObject = FFbxImporter->ImportStaticMesh(Pkg, Node, StaticMeshFName, Flags, StaticMeshImportData, Cast<UStaticMesh>(InMesh), LODIndex);
+
 	OutNodeInfo->AttributeInfo->SetOriginalImportPath(PackageName);
 	OutNodeInfo->AttributeInfo->SetOriginalFullImportName(NewObject->GetPathName());
 
diff --git a/Engine/Source/Editor/UnrealEd/Private/Fbx/FbxSceneImportOptions.cpp b/Engine/Source/Editor/UnrealEd/Private/Fbx/FbxSceneImportOptions.cpp
--- a/Engine/Source/Editor/UnrealEd/Private/Fbx/FbxSceneImportOptions.cpp
+++ b/Engine/Source/Editor/UnrealEd/Private/Fbx/FbxSceneImportOptions.cpp
@@ -6,11 +6,14 @@
 UFbxSceneImportOptions::UFbxSceneImportOptions(const FObjectInitializer& ObjectInitializer)
 	: Super(ObjectInitializer)
 {
+	bTransformVertexToAbsolute = false;
+	bBakePivotInVertex = false;
 	bCreateContentFolderHierarchy = false;
 	bImportAsDynamic = false;
-	HierarchyType = FBXSOCHT_CreateActorComponents;
+	HierarchyType = FBXSOCHT_CreateBlueprint;
 	bImportStaticMeshLODs = false;
 	bImportSkeletalMeshLODs = false;
+	bInvertNormalMaps = false;
 	ImportTranslation = FVector(0);
 	ImportRotation = FRotator(0);
 	ImportUniformScale = 1.0f;
diff --git a/Engine/Source/Editor/UnrealEd/Private/Fbx/FbxSceneImportOptionsSkeletalMesh.cpp b/Engine/Source/Editor/UnrealEd/Private/Fbx/FbxSceneImportOptionsSkeletalMesh.cpp
--- a/Engine/Source/Editor/UnrealEd/Private/Fbx/FbxSceneImportOptionsSkeletalMesh.cpp
+++ b/Engine/Source/Editor/UnrealEd/Private/Fbx/FbxSceneImportOptionsSkeletalMesh.cpp
@@ -12,6 +12,32 @@ UFbxSceneImportOptionsSkeletalMesh::UFbxSceneImportOptionsSkeletalMesh(const FOb
 	, bImportMeshesInBoneHierarchy(true)
 	, bImportMorphTargets(false)
 	, bKeepOverlappingVertices(false)
+	, bImportAnimations(true)
+	, AnimationLength(EFbxSceneVertexColorImportOption::Replace)
+	, FrameImportRange(0, 0)
+	, bUseDefaultSampleRate(false)
+	, bImportCustomAttribute(true)
+	, bPreserveLocalTransform(false)
+	, bDeleteExistingMorphTargetCurves(false)
 {
 }
 
+void UFbxSceneImportOptionsSkeletalMesh::FillSkeletalMeshInmportData(UFbxSkeletalMeshImportData* SkeletalMeshImportData, UFbxAnimSequenceImportData* AnimSequenceImportData)
+{
+	check(SkeletalMeshImportData != nullptr);
+	SkeletalMeshImportData->bImportMeshesInBoneHierarchy = bImportMeshesInBoneHierarchy;
+	SkeletalMeshImportData->bImportMorphTargets = bImportMorphTargets;
+	SkeletalMeshImportData->bKeepOverlappingVertices = bKeepOverlappingVertices;
+	SkeletalMeshImportData->bPreserveSmoothingGroups = bPreserveSmoothingGroups;
+	SkeletalMeshImportData->bUpdateSkeletonReferencePose = bUpdateSkeletonReferencePose;
+	SkeletalMeshImportData->bUseT0AsRefPose = bUseT0AsRefPose;
+	
+	SkeletalMeshImportData->bImportAsScene = true;
+
+	AnimSequenceImportData->AnimationLength = AnimationLength;
+	AnimSequenceImportData->bDeleteExistingMorphTargetCurves = bDeleteExistingMorphTargetCurves;
+	AnimSequenceImportData->bImportCustomAttribute = bImportCustomAttribute;
+	AnimSequenceImportData->bPreserveLocalTransform = bPreserveLocalTransform;
+	AnimSequenceImportData->bUseDefaultSampleRate = bUseDefaultSampleRate;
+	AnimSequenceImportData->FrameImportRange = FrameImportRange;
+}
\ No newline at end of file
diff --git a/Engine/Source/Editor/UnrealEd/Private/Fbx/FbxSceneImportOptionsStaticMesh.cpp b/Engine/Source/Editor/UnrealEd/Private/Fbx/FbxSceneImportOptionsStaticMesh.cpp
--- a/Engine/Source/Editor/UnrealEd/Private/Fbx/FbxSceneImportOptionsStaticMesh.cpp
+++ b/Engine/Source/Editor/UnrealEd/Private/Fbx/FbxSceneImportOptionsStaticMesh.cpp
@@ -5,7 +5,6 @@
 
 UFbxSceneImportOptionsStaticMesh::UFbxSceneImportOptionsStaticMesh(const FObjectInitializer& ObjectInitializer)
 	: Super(ObjectInitializer)
-	, bTransformVertexToAbsolute(false)
 	, StaticMeshLODGroup(NAME_None)
 	, bAutoGenerateCollision(true)
 	, bRemoveDegenerates(true)
@@ -29,7 +28,6 @@ void UFbxSceneImportOptionsStaticMesh::FillStaticMeshInmportData(UFbxStaticMeshI
 	StaticMeshImportData->bGenerateLightmapUVs = bGenerateLightmapUVs;
 	StaticMeshImportData->bOneConvexHullPerUCX = bOneConvexHullPerUCX;
 	StaticMeshImportData->bRemoveDegenerates = bRemoveDegenerates;
-	StaticMeshImportData->bTransformVertexToAbsolute = bTransformVertexToAbsolute;
 	StaticMeshImportData->StaticMeshLODGroup = StaticMeshLODGroup;
 	switch (VertexColorImportOption)
 	{
@@ -50,4 +48,8 @@ void UFbxSceneImportOptionsStaticMesh::FillStaticMeshInmportData(UFbxStaticMeshI
 	StaticMeshImportData->ImportTranslation = SceneImportOptions->ImportTranslation;
 	StaticMeshImportData->ImportRotation = SceneImportOptions->ImportRotation;
 	StaticMeshImportData->ImportUniformScale = SceneImportOptions->ImportUniformScale;
+	StaticMeshImportData->bTransformVertexToAbsolute = SceneImportOptions->bTransformVertexToAbsolute;
+	StaticMeshImportData->bBakePivotInVertex = SceneImportOptions->bBakePivotInVertex;
+
+	StaticMeshImportData->bImportAsScene = true;
 }
\ No newline at end of file
diff --git a/Engine/Source/Editor/UnrealEd/Private/Fbx/FbxSkeletalMeshImport.cpp b/Engine/Source/Editor/UnrealEd/Private/Fbx/FbxSkeletalMeshImport.cpp
--- a/Engine/Source/Editor/UnrealEd/Private/Fbx/FbxSkeletalMeshImport.cpp
+++ b/Engine/Source/Editor/UnrealEd/Private/Fbx/FbxSkeletalMeshImport.cpp
@@ -1536,7 +1536,7 @@ void UnFbx::FFbxImporter::SetupAnimationDataFromMesh(USkeletalMesh* SkeletalMesh
 	}
 }
 
-USkeletalMesh* UnFbx::FFbxImporter::ReimportSkeletalMesh(USkeletalMesh* Mesh, UFbxSkeletalMeshImportData* TemplateImportData)
+USkeletalMesh* UnFbx::FFbxImporter::ReimportSkeletalMesh(USkeletalMesh* Mesh, UFbxSkeletalMeshImportData* TemplateImportData, uint64 SkeletalMeshFbxUID, TArray<FbxNode*> *OutSkeletalMeshArray)
 {
 	if ( !ensure(Mesh) )
 	{
@@ -1555,6 +1555,11 @@ USkeletalMesh* UnFbx::FFbxImporter::ReimportSkeletalMesh(USkeletalMesh* Mesh, UF
 	TArray<FbxNode*>* FbxNodes = NULL;
 	USkeletalMesh* NewMesh = NULL;
 
+	bool Old_ImportRigidMesh = ImportOptions->bImportRigidMesh;
+	bool Old_ImportMaterials = ImportOptions->bImportMaterials;
+	bool Old_ImportTextures = ImportOptions->bImportTextures;
+	bool Old_ImportAnimations = ImportOptions->bImportAnimations;
+
 	// support to update rigid animation mesh
 	ImportOptions->bImportRigidMesh = true;
 
@@ -1563,16 +1568,52 @@ USkeletalMesh* UnFbx::FFbxImporter::ReimportSkeletalMesh(USkeletalMesh* Mesh, UF
 	TArray< TArray<FbxNode*>* > FbxSkelMeshArray;
 	FillFbxSkelMeshArrayInScene(Scene->GetRootNode(), FbxSkelMeshArray, false);
 
-	// if there is only one mesh, use it without name checking 
-	// (because the "Used As Full Name" option enables users name the Unreal mesh by themselves
-	if (FbxSkelMeshArray.Num() > 0 )
+	if(SkeletalMeshFbxUID != 0xFFFFFFFFFFFFFFFF)
 	{
-		FbxNodes = FbxSkelMeshArray[0];
+		//Scene reimport know which skeletal mesh we want to reimport
+		for (TArray<FbxNode*>*SkeletalMeshNodes : FbxSkelMeshArray)
+		{
+			if (SkeletalMeshNodes->Num() > 0)
+			{
+				FbxNode *Node = (*SkeletalMeshNodes)[0];
+				FbxNode *SkeletalMeshNode = Node;
+				if (Node->GetNodeAttribute() && Node->GetNodeAttribute()->GetAttributeType() == FbxNodeAttribute::eLODGroup)
+				{
+					SkeletalMeshNode = FindLODGroupNode(Node, 0);
+				}
+
+				if (SkeletalMeshNode != nullptr && SkeletalMeshNode->GetNodeAttribute() && SkeletalMeshNode->GetNodeAttribute()->GetUniqueID() == SkeletalMeshFbxUID)
+				{
+					FbxNodes = SkeletalMeshNodes;
+					if (OutSkeletalMeshArray != nullptr)
+					{
+						for (FbxNode *NodeReimport : (*SkeletalMeshNodes))
+						{
+							OutSkeletalMeshArray->Add(NodeReimport);
+						}
+					}
+					break;
+				}
+			}
+			if (FbxNodes != nullptr)
+				break;
+		}
+		if (FbxNodes == nullptr)
+			return nullptr;
 	}
 	else
 	{
-		// @todo - FBX Importing - We need proper support for reimport if the file contains more than one skeletal mesh. 
+		// if there is only one mesh, use it without name checking 
+		// (because the "Used As Full Name" option enables users name the Unreal mesh by themselves
+		if (FbxSkelMeshArray.Num() > 0)
+		{
+			FbxNodes = FbxSkelMeshArray[0];
+		}
+		else
+		{
+			// @todo - FBX Importing - We need proper support for reimport if the file contains more than one skeletal mesh. 
 
+		}
 	}
 
 	if (FbxNodes)
@@ -1609,11 +1650,15 @@ USkeletalMesh* UnFbx::FFbxImporter::ReimportSkeletalMesh(USkeletalMesh* Mesh, UF
 				{
 					if (Node->GetChildCount() > LODIndex)
 					{
-						SkelMeshNodeArray.Add(Node->GetChild(LODIndex));
+						FbxNode *MeshNode = FindLODGroupNode(Node, LODIndex);
+						if(MeshNode != nullptr)
+							SkelMeshNodeArray.Add(MeshNode);
 					}
 					else // in less some LODGroups have less level, use the last level
 					{
-						SkelMeshNodeArray.Add(Node->GetChild(Node->GetChildCount() - 1));
+						FbxNode *MeshNode = FindLODGroupNode(Node, Node->GetChildCount() - 1);
+						if (MeshNode != nullptr)
+							SkelMeshNodeArray.Add(MeshNode);
 					}
 				}
 				else
@@ -1655,6 +1700,11 @@ USkeletalMesh* UnFbx::FFbxImporter::ReimportSkeletalMesh(USkeletalMesh* Mesh, UF
 		AddTokenizedErrorMessage(FTokenizedMessage::Create(EMessageSeverity::Warning, FText::Format(LOCTEXT("FbxSkeletaLMeshimport_NoFBXMeshMatch", "No FBX mesh matches the Unreal mesh '{0}'."), FText::FromString(Mesh->GetName()))), FFbxErrors::Generic_Mesh_MeshNotFound);
 	}
 
+	ImportOptions->bImportRigidMesh = Old_ImportRigidMesh;
+	ImportOptions->bImportMaterials = Old_ImportMaterials;
+	ImportOptions->bImportTextures = Old_ImportTextures;
+	ImportOptions->bImportAnimations = Old_ImportAnimations;
+
 	return NewMesh;
 }
 
@@ -1805,6 +1855,16 @@ bool UnFbx::FFbxImporter::FillSkelMeshImporterFromFbx( FSkeletalMeshImportData&
 	if (ImportOptions->bImportMaterials)
 	{
 		CreateNodeMaterials(Node, Materials, UVSets);
+		//Make sure all material use by skeletal mesh are flag to be use by skeletal mesh
+		for (UMaterialInterface *MaterialInterface : Materials)
+		{
+			if (MaterialInterface == nullptr || MaterialInterface->GetMaterial() == nullptr)
+			{
+				continue;
+			}
+			bool bNeedRecompile = false;
+			MaterialInterface->GetMaterial()->SetMaterialUsage(bNeedRecompile, EMaterialUsage::MATUSAGE_SkeletalMesh);
+		}
 	}
 	else if (ImportOptions->bImportTextures)
 	{
diff --git a/Engine/Source/Editor/UnrealEd/Private/Fbx/FbxStaticMeshImport.cpp b/Engine/Source/Editor/UnrealEd/Private/Fbx/FbxStaticMeshImport.cpp
--- a/Engine/Source/Editor/UnrealEd/Private/Fbx/FbxStaticMeshImport.cpp
+++ b/Engine/Source/Editor/UnrealEd/Private/Fbx/FbxStaticMeshImport.cpp
@@ -436,7 +436,7 @@ bool UnFbx::FFbxImporter::BuildStaticMeshFromGeometry(FbxNode* Node, UStaticMesh
 
 	int32 VertexCount = Mesh->GetControlPointsCount();
 	int32 WedgeCount = TriangleCount * 3;
-	bool OddNegativeScale = ImportOptions->bTransformVertexToAbsolute ? false : IsOddNegativeScale(TotalMatrix);
+	bool OddNegativeScale = IsOddNegativeScale(TotalMatrix);
 
 	int32 VertexOffset = RawMesh.VertexPositions.Num();
 	int32 WedgeOffset = RawMesh.WedgeIndices.Num();
@@ -518,7 +518,7 @@ bool UnFbx::FFbxImporter::BuildStaticMeshFromGeometry(FbxNode* Node, UStaticMesh
 			else
 			{
 				FbxVector4 FbxPosition = Mesh->GetControlPoints()[ControlPointIndex];
-				FbxVector4 FinalPosition = ImportOptions->bTransformVertexToAbsolute ? TotalMatrix.MultT(FbxPosition) : FbxPosition;
+				FbxVector4 FinalPosition = TotalMatrix.MultT(FbxPosition);
 				int32 VertexIndex = RawMesh.VertexPositions.Add(Converter.ConvertPos(FinalPosition));
 				RawMesh.WedgeIndices[WedgeIndex] = VertexIndex;
 				IndexMap.Add(ControlPointIndex,VertexIndex);
@@ -545,10 +545,7 @@ bool UnFbx::FFbxImporter::BuildStaticMeshFromGeometry(FbxNode* Node, UStaticMesh
 						TangentMapIndex : LayerElementTangent->GetIndexArray().GetAt(TangentMapIndex);
 
 					FbxVector4 TempValue = LayerElementTangent->GetDirectArray().GetAt(TangentValueIndex);
-					if (ImportOptions->bTransformVertexToAbsolute)
-					{
-						TempValue = TotalMatrixForNormal.MultT(TempValue);
-					}
+					TempValue = TotalMatrixForNormal.MultT(TempValue);
 					FVector TangentX = Converter.ConvertDir(TempValue);
 					RawMesh.WedgeTangentX[WedgeIndex] = TangentX.GetSafeNormal();
 
@@ -558,19 +555,13 @@ bool UnFbx::FFbxImporter::BuildStaticMeshFromGeometry(FbxNode* Node, UStaticMesh
 						BinormalMapIndex : LayerElementBinormal->GetIndexArray().GetAt(BinormalMapIndex);
 
 					TempValue = LayerElementBinormal->GetDirectArray().GetAt(BinormalValueIndex);
-					if (ImportOptions->bTransformVertexToAbsolute)
-					{
-						TempValue = TotalMatrixForNormal.MultT(TempValue);
-					}
+					TempValue = TotalMatrixForNormal.MultT(TempValue);
 					FVector TangentY = -Converter.ConvertDir(TempValue);
 					RawMesh.WedgeTangentY[WedgeIndex] = TangentY.GetSafeNormal();
 				}
 
 				FbxVector4 TempValue = LayerElementNormal->GetDirectArray().GetAt(NormalValueIndex);
-				if (ImportOptions->bTransformVertexToAbsolute)
-				{
-					TempValue = TotalMatrixForNormal.MultT(TempValue);
-				}
+				TempValue = TotalMatrixForNormal.MultT(TempValue);
 				FVector TangentZ = Converter.ConvertDir(TempValue);
 				RawMesh.WedgeTangentZ[WedgeIndex] = TangentZ.GetSafeNormal();
 			}
@@ -706,7 +697,7 @@ bool UnFbx::FFbxImporter::BuildStaticMeshFromGeometry(FbxNode* Node, UStaticMesh
 	return true;
 }
 
-UStaticMesh* UnFbx::FFbxImporter::ReimportSceneStaticMesh(uint64 FbxUniqueId, UStaticMesh* Mesh, UFbxStaticMeshImportData* TemplateImportData)
+UStaticMesh* UnFbx::FFbxImporter::ReimportSceneStaticMesh(uint64 FbxNodeUniqueId, uint64 FbxUniqueId, UStaticMesh* Mesh, UFbxStaticMeshImportData* TemplateImportData)
 {
 	TArray<FbxNode*> FbxMeshArray;
 	UStaticMesh* NewMesh = NULL;
@@ -726,10 +717,21 @@ UStaticMesh* UnFbx::FFbxImporter::ReimportSceneStaticMesh(uint64 FbxUniqueId, US
 		//Find the first node using the mesh attribute with the unique ID
 		for (FbxNode *MeshNode : FbxMeshArray)
 		{
-			if(FbxUniqueId == MeshNode->GetMesh()->GetUniqueID())
+			if (FbxNodeUniqueId == INVALID_UNIQUE_ID || ImportOptions->bBakePivotInVertex == false)
 			{
-				Node = MeshNode;
-				break;
+				if (FbxUniqueId == MeshNode->GetMesh()->GetUniqueID())
+				{
+					Node = MeshNode;
+					break;
+				}
+			}
+			else
+			{
+				if (FbxNodeUniqueId == MeshNode->GetUniqueID() && FbxUniqueId == MeshNode->GetMesh()->GetUniqueID())
+				{
+					Node = MeshNode;
+					break;
+				}
 			}
 		}
 	}
@@ -745,7 +747,7 @@ UStaticMesh* UnFbx::FFbxImporter::ReimportSceneStaticMesh(uint64 FbxUniqueId, US
 
 	if (Node)
 	{
-		FbxNode* NodeParent = Node->GetParent();
+		FbxNode* NodeParent = RecursiveFindParentLodGroup(Node->GetParent());
 		// Don't import materials and textures during a reimport
 		ImportOptions->bImportMaterials = true;
 		ImportOptions->bImportTextures = true;
@@ -753,13 +755,17 @@ UStaticMesh* UnFbx::FFbxImporter::ReimportSceneStaticMesh(uint64 FbxUniqueId, US
 		// if the Fbx mesh is a part of LODGroup, update LOD
 		if (NodeParent && NodeParent->GetNodeAttribute() && NodeParent->GetNodeAttribute()->GetAttributeType() == FbxNodeAttribute::eLODGroup)
 		{
-			NewMesh = ImportStaticMesh(Mesh->GetOutermost(), NodeParent->GetChild(0), *Mesh->GetName(), RF_Public | RF_Standalone, TemplateImportData, Mesh, 0);
+			FbxNode *MeshNode = FindLODGroupNode(NodeParent, 0);
+			if(MeshNode != nullptr)
+				NewMesh = ImportStaticMesh(Mesh->GetOutermost(), MeshNode, *Mesh->GetName(), RF_Public | RF_Standalone, TemplateImportData, Mesh, 0);
 			if (NewMesh)
 			{
 				// import LOD meshes
 				for (int32 LODIndex = 1; LODIndex < NodeParent->GetChildCount(); LODIndex++)
 				{
-					ImportStaticMesh(Mesh->GetOutermost(), NodeParent->GetChild(LODIndex), *Mesh->GetName(), RF_Public | RF_Standalone, TemplateImportData, Mesh, LODIndex);
+					MeshNode = FindLODGroupNode(NodeParent, LODIndex);
+					if(MeshNode != nullptr)
+						ImportStaticMesh(Mesh->GetOutermost(), MeshNode, *Mesh->GetName(), RF_Public | RF_Standalone, TemplateImportData, Mesh, LODIndex);
 				}
 			}
 		}
@@ -850,7 +856,8 @@ UStaticMesh* UnFbx::FFbxImporter::ReimportStaticMesh(UStaticMesh* Mesh, UFbxStat
 	if (!Node && FbxMeshArray.IsValidIndex(0))
 	{
 		FbxNode* BaseLODNode = FbxMeshArray[0];
-		FbxNode* NodeParent = BaseLODNode ? BaseLODNode->GetParent() : NULL;
+		
+		FbxNode* NodeParent = BaseLODNode ? RecursiveFindParentLodGroup(BaseLODNode->GetParent()) : nullptr;
 		if (NodeParent && NodeParent->GetNodeAttribute() && NodeParent->GetNodeAttribute()->GetAttributeType() == FbxNodeAttribute::eLODGroup)
 		{
 			// Reimport the entire LOD chain.
@@ -862,7 +869,7 @@ UStaticMesh* UnFbx::FFbxImporter::ReimportStaticMesh(UStaticMesh* Mesh, UFbxStat
 
 	if (Node)
 	{
-		FbxNode* NodeParent = Node->GetParent();
+		FbxNode* NodeParent = RecursiveFindParentLodGroup(Node->GetParent());
 		// Don't import materials and textures during a reimport
 		ImportOptions->bImportMaterials = false;
 		ImportOptions->bImportTextures = false;
@@ -870,13 +877,17 @@ UStaticMesh* UnFbx::FFbxImporter::ReimportStaticMesh(UStaticMesh* Mesh, UFbxStat
 		// if the Fbx mesh is a part of LODGroup, update LOD
 		if (NodeParent && NodeParent->GetNodeAttribute() && NodeParent->GetNodeAttribute()->GetAttributeType() == FbxNodeAttribute::eLODGroup)
 		{
-			NewMesh = ImportStaticMesh(Mesh->GetOuter(), NodeParent->GetChild(0), *Mesh->GetName(), RF_Public|RF_Standalone, TemplateImportData, Mesh, 0);
+			FbxNode *MeshNode = FindLODGroupNode(NodeParent, 0);
+			if(MeshNode != nullptr)
+				NewMesh = ImportStaticMesh(Mesh->GetOuter(), MeshNode, *Mesh->GetName(), RF_Public|RF_Standalone, TemplateImportData, Mesh, 0);
 			if (NewMesh)
 			{
 				// import LOD meshes
 				for (int32 LODIndex = 1; LODIndex < NodeParent->GetChildCount(); LODIndex++)
 				{
-					ImportStaticMesh(Mesh->GetOuter(), NodeParent->GetChild(LODIndex), *Mesh->GetName(), RF_Public|RF_Standalone, TemplateImportData, Mesh, LODIndex);
+					MeshNode = FindLODGroupNode(NodeParent, LODIndex);
+					if (MeshNode != nullptr)
+						ImportStaticMesh(Mesh->GetOuter(), MeshNode, *Mesh->GetName(), RF_Public|RF_Standalone, TemplateImportData, Mesh, LODIndex);
 				}
 			}
 		}
@@ -1263,8 +1274,7 @@ UStaticMesh* UnFbx::FFbxImporter::ImportStaticMeshAsSingle(UObject* InParent, TA
 			StaticMesh->Materials.Empty();
 		}
 		
-		// Build a new map of sections with the unique material set
-		FMeshSectionInfoMap NewMap;
+		// Replace map of sections with the unique material set
 		int32 NumMaterials = FMath::Min(SortedMaterials.Num(),MaxMaterialIndex+1);
 		for (int32 MaterialIndex = 0; MaterialIndex < NumMaterials; ++MaterialIndex)
 		{
@@ -1273,13 +1283,10 @@ UStaticMesh* UnFbx::FFbxImporter::ImportStaticMeshAsSingle(UObject* InParent, TA
 			int32 Index = StaticMesh->Materials.Add(SortedMaterials[MaterialIndex].Material);
 
 			Info.MaterialIndex = Index;
-			NewMap.Set( LODIndex, MaterialIndex, Info);
+			StaticMesh->SectionInfoMap.Remove(LODIndex, MaterialIndex);
+			StaticMesh->SectionInfoMap.Set(LODIndex, MaterialIndex, Info);
 		}
 
-		// Copy the final section map into the static mesh
-		StaticMesh->SectionInfoMap.Clear();
-		StaticMesh->SectionInfoMap.CopyFrom(NewMap);
-
 		FRawMesh LocalRawMesh;
 		SrcModel.RawMeshBulkData->LoadRawMesh(LocalRawMesh);
 
@@ -1615,7 +1622,7 @@ bool UnFbx::FFbxImporter::ImportCollisionModels(UStaticMesh* StaticMesh, const F
 
 		for ( ControlPointsIndex = 0; ControlPointsIndex < ControlPointsCount; ControlPointsIndex++ )
 		{
-			new(CollisionVertices)FVector(Converter.ConvertPos(ImportOptions->bTransformVertexToAbsolute ? Matrix.MultT(ControlPoints[ControlPointsIndex]) : ControlPoints[ControlPointsIndex]));
+			new(CollisionVertices)FVector(Converter.ConvertPos(Matrix.MultT(ControlPoints[ControlPointsIndex])));
 		}
 
 		int32 TriangleCount = FbxMesh->GetPolygonCount();
diff --git a/Engine/Source/Editor/UnrealEd/Private/Fbx/FbxStaticMeshImportData.cpp b/Engine/Source/Editor/UnrealEd/Private/Fbx/FbxStaticMeshImportData.cpp
--- a/Engine/Source/Editor/UnrealEd/Private/Fbx/FbxStaticMeshImportData.cpp
+++ b/Engine/Source/Editor/UnrealEd/Private/Fbx/FbxStaticMeshImportData.cpp
@@ -14,6 +14,7 @@ UFbxStaticMeshImportData::UFbxStaticMeshImportData(const FObjectInitializer& Obj
 	bOneConvexHullPerUCX = true;
 	bAutoGenerateCollision = true;
 	bTransformVertexToAbsolute = true;
+	bBakePivotInVertex = false;
 	VertexOverrideColor = FColor(255, 255, 255, 255);
 }
 
diff --git a/Engine/Source/Editor/UnrealEd/Private/Fbx/ReimportFbxSceneFactory.cpp b/Engine/Source/Editor/UnrealEd/Private/Fbx/ReimportFbxSceneFactory.cpp
--- a/Engine/Source/Editor/UnrealEd/Private/Fbx/ReimportFbxSceneFactory.cpp
+++ b/Engine/Source/Editor/UnrealEd/Private/Fbx/ReimportFbxSceneFactory.cpp
@@ -19,6 +19,7 @@
 
 #include "Editor/KismetCompiler/Public/KismetCompilerModule.h"
 #include "Kismet2/KismetEditorUtilities.h"
+#include "Kismet2/BlueprintEditorUtils.h"
 #include "Toolkits/AssetEditorManager.h"
 #include "ContentBrowserModule.h"
 
@@ -50,6 +51,16 @@ UFbxSceneImportData *GetFbxSceneImportData(UObject *Obj)
 				SceneImportData = ImportData->bImportAsScene ? ImportData->FbxSceneImportDataReference : nullptr;
 			}
 		}
+		if (Obj->IsA(USkeletalMesh::StaticClass()))
+		{
+			//Reimport from one of the static mesh
+			USkeletalMesh* SkeletalMesh = Cast<USkeletalMesh>(Obj);
+			if (SkeletalMesh != nullptr && SkeletalMesh->AssetImportData != nullptr)
+			{
+				ImportData = Cast<UFbxAssetImportData>(SkeletalMesh->AssetImportData);
+				SceneImportData = ImportData->bImportAsScene ? ImportData->FbxSceneImportDataReference : nullptr;
+			}
+		}
 		//TODO: add all type the fbx scene import can create: material, texture, skeletal mesh, animation, ... 
 	}
 	return SceneImportData;
@@ -115,12 +126,55 @@ void RecursivelyCreateOriginalPath(UnFbx::FFbxImporter* FbxImporter, TSharedPtr<
 	}
 }
 
+void SetNodeInfoTypeChanged(TSharedPtr<FFbxNodeInfo> NodeInfoA, TSharedPtr<FFbxNodeInfo> NodeInfoB)
+{
+	if (NodeInfoA->AttributeInfo.IsValid())
+	{
+		//We found a match verify the type
+		NodeInfoA->AttributeInfo->bOriginalTypeChanged = NodeInfoA->AttributeInfo.IsValid() != NodeInfoB->AttributeInfo.IsValid();
+		if (!NodeInfoA->AttributeInfo->bOriginalTypeChanged && NodeInfoA->AttributeInfo.IsValid() && NodeInfoB->AttributeInfo.IsValid())
+		{
+			NodeInfoA->AttributeInfo->bOriginalTypeChanged = NodeInfoA->AttributeInfo->GetType() != NodeInfoB->AttributeInfo->GetType();
+		}
+		if (!NodeInfoA->AttributeInfo->bOriginalTypeChanged)
+		{
+			UObject *ContentObjectA = NodeInfoA->AttributeInfo->GetContentObject();
+			if (ContentObjectA != nullptr)
+			{
+				if (!ContentObjectA->IsA(NodeInfoA->AttributeInfo->GetType()))
+				{
+					NodeInfoA->AttributeInfo->bOriginalTypeChanged = true;
+				}
+			}
+		}
+	}
+}
+
+void MarkAssetTypeChanged(TSharedPtr<FFbxSceneInfo> SceneInfoPtr, TSharedPtr<FFbxSceneInfo> SceneInfoOriginalPtr)
+{
+	//Set the node info
+	for (TSharedPtr<FFbxNodeInfo> NodeInfo : SceneInfoPtr->HierarchyInfo)
+	{
+		for (TSharedPtr<FFbxNodeInfo> NodeInfoOriginal : SceneInfoOriginalPtr->HierarchyInfo)
+		{
+			if (NodeInfo->NodeHierarchyPath == NodeInfoOriginal->NodeHierarchyPath)
+			{
+				//Set the current
+				SetNodeInfoTypeChanged(NodeInfo, NodeInfoOriginal);
+				//Set the original
+				SetNodeInfoTypeChanged(NodeInfoOriginal, NodeInfo);
+			}
+		}
+	}
+}
+
 bool GetFbxSceneReImportOptions(UnFbx::FFbxImporter* FbxImporter
 	, TSharedPtr<FFbxSceneInfo> SceneInfoPtr
 	, TSharedPtr<FFbxSceneInfo> SceneInfoOriginalPtr
 	, UnFbx::FBXImportOptions *GlobalImportSettings
 	, UFbxSceneImportOptions *SceneImportOptions
 	, UFbxSceneImportOptionsStaticMesh *StaticMeshImportData
+	, UFbxSceneImportOptionsSkeletalMesh *SkeletalMeshImportData
 	, ImportOptionsNameMap &NameOptionsMap
 	, FbxSceneReimportStatusMap &MeshStatusMap
 	, FbxSceneReimportStatusMap &NodeStatusMap
@@ -140,9 +194,6 @@ bool GetFbxSceneReImportOptions(UnFbx::FFbxImporter* FbxImporter
 
 	GlobalImportSettings->bConvertScene = true;
 
-	//TODO this options will be set by the fbxscene UI in the material options tab, it also should be save/load from config file
-	//Prefix materials package name to put all material under Material folder (this avoid name clash with meshes)
-	GlobalImportSettings->MaterialPrefixName = FName("/Materials/");
 
 	TSharedPtr<SWindow> ParentWindow;
 	if (FModuleManager::Get().IsModuleLoaded("MainFrame"))
@@ -157,6 +208,7 @@ bool GetFbxSceneReImportOptions(UnFbx::FFbxImporter* FbxImporter
 
 	//Make sure the display option show the save default options
 	SFbxSceneOptionWindow::CopyFbxOptionsToStaticMeshOptions(GlobalImportSettings, StaticMeshImportData);
+	SFbxSceneOptionWindow::CopyFbxOptionsToSkeletalMeshOptions(GlobalImportSettings, SkeletalMeshImportData);
 
 	Window->SetContent
 		(
@@ -165,6 +217,7 @@ bool GetFbxSceneReImportOptions(UnFbx::FFbxImporter* FbxImporter
 			.SceneInfoOriginal(SceneInfoOriginalPtr)
 			.SceneImportOptionsDisplay(SceneImportOptions)
 			.SceneImportOptionsStaticMeshDisplay(StaticMeshImportData)
+			.SceneImportOptionsSkeletalMeshDisplay(SkeletalMeshImportData)
 			.OverrideNameOptionsMap(&NameOptionsMap)
 			.MeshStatusMap(&MeshStatusMap)
 			.CanReimportHierarchy(bCanReimportHierarchy)
@@ -181,9 +234,12 @@ bool GetFbxSceneReImportOptions(UnFbx::FFbxImporter* FbxImporter
 		return false;
 	}
 
+	//Set the bakepivot option in the SceneImportOptions
+	SceneImportOptions->bBakePivotInVertex = GlobalImportSettings->bBakePivotInVertex;
 	//setup all options
 	GlobalImportSettings->bImportStaticMeshLODs = SceneImportOptions->bImportStaticMeshLODs;
 	GlobalImportSettings->bImportSkeletalMeshLODs = SceneImportOptions->bImportSkeletalMeshLODs;
+	SceneImportOptions->bInvertNormalMaps = GlobalImportSettings->bInvertNormalMap;
 
 	//Set the override material into the options
 	for (TSharedPtr<FFbxNodeInfo> NodeInfo : SceneInfoPtr->HierarchyInfo)
@@ -251,9 +307,15 @@ EReimportResult::Type UReimportFbxSceneFactory::Reimport(UObject* Obj)
 	//Always convert the scene
 	GlobalImportSettings->bConvertScene = true;
 	GlobalImportSettings->bImportScene = ReimportData->bImportScene;
+	if (ReimportData->NameOptionsMap.Contains(DefaultOptionName))
+	{
+		UnFbx::FBXImportOptions *DefaultOption = *(ReimportData->NameOptionsMap.Find(DefaultOptionName));
+		GlobalImportSettings->bBakePivotInVertex = DefaultOption->bBakePivotInVertex;
+		GlobalImportSettings->bInvertNormalMap = DefaultOption->bInvertNormalMap;
+	}
 
 	//Read the fbx and store the hierarchy's information so we can reuse it after importing all the model in the fbx file
-	if (!FbxImporter->ImportFromFile(*FbxImportFileName, FPaths::GetExtension(FbxImportFileName)))
+	if (!FbxImporter->ImportFromFile(*FbxImportFileName, FPaths::GetExtension(FbxImportFileName), true))
 	{
 		// Log the error message and fail the import.
 		GWarn->Log(ELogVerbosity::Error, FbxImporter->GetErrorMessage());
@@ -263,16 +325,20 @@ EReimportResult::Type UReimportFbxSceneFactory::Reimport(UObject* Obj)
 		return EReimportResult::Failed;
 	}
 
+	//Make sure the Skeleton is null and not garbage, as we are importing the skeletalmesh for the first time we do not need any skeleton
+	GlobalImportSettings->SkeletonForAnimation = nullptr;
+	GlobalImportSettings->PhysicsAsset = nullptr;
+
 	FString PackageName = "";
 	Obj->GetOutermost()->GetName(PackageName);
 	Path = FPaths::GetPath(PackageName);
 
 	UnFbx::FbxSceneInfo SceneInfo;
 	//Read the scene and found all instance with their scene information.
-	FbxImporter->GetSceneInfo(FbxImportFileName, SceneInfo);
+	FbxImporter->GetSceneInfo(FbxImportFileName, SceneInfo, true);
 
 	//Convert old structure to the new scene export structure
-	TSharedPtr<FFbxSceneInfo> SceneInfoPtr = ConvertSceneInfo(&SceneInfo);
+	TSharedPtr<FFbxSceneInfo> SceneInfoPtr = ConvertSceneInfo(FbxImporter, &SceneInfo);
 	//Get import material info
 	ExtractMaterialInfo(FbxImporter, SceneInfoPtr);
 
@@ -300,6 +366,8 @@ EReimportResult::Type UReimportFbxSceneFactory::Reimport(UObject* Obj)
 
 	FillSceneHierarchyPath(SceneInfoPtr);
 
+	MarkAssetTypeChanged(SceneInfoPtr, ReimportData->SceneInfoSourceData);
+
 	FbxSceneReimportStatusMap MeshStatusMap;
 	FbxSceneReimportStatusMap NodeStatusMap;
 	bool bCanReimportHierarchy = ReimportData->HierarchyType == (int32)EFBXSceneOptionsCreateHierarchyType::FBXSOCHT_CreateBlueprint && !ReimportData->BluePrintFullName.IsEmpty();
@@ -310,6 +378,7 @@ EReimportResult::Type UReimportFbxSceneFactory::Reimport(UObject* Obj)
 		, GlobalImportSettings
 		, SceneImportOptions
 		, SceneImportOptionsStaticMesh
+		, SceneImportOptionsSkeletalMesh
 		, NameOptionsMap
 		, MeshStatusMap
 		, NodeStatusMap
@@ -331,6 +400,7 @@ EReimportResult::Type UReimportFbxSceneFactory::Reimport(UObject* Obj)
 	ReimportData->SceneInfoSourceData = SceneInfoPtr;
 	ReimportData->SourceFbxFile = FPaths::ConvertRelativePathToFull(FbxImportFileName);
 	ReimportData->bImportScene = GlobalImportSettingsReference->bImportScene;
+
 	//Copy the options map
 	ReimportData->NameOptionsMap.Reset();
 	for (auto kvp : NameOptionsMap)
@@ -338,12 +408,17 @@ EReimportResult::Type UReimportFbxSceneFactory::Reimport(UObject* Obj)
 		ReimportData->NameOptionsMap.Add(kvp.Key, kvp.Value);
 	}
 
+	StaticMeshImportData->bImportAsScene = true;
+	StaticMeshImportData->FbxSceneImportDataReference = ReimportData;
+	SkeletalMeshImportData->bImportAsScene = true;
+	SkeletalMeshImportData->FbxSceneImportDataReference = ReimportData;
+
 	FAssetRegistryModule& AssetRegistryModule = FModuleManager::LoadModuleChecked<FAssetRegistryModule>("AssetRegistry");
 	TArray<FAssetData> AssetDataToDelete;
 	for (TSharedPtr<FFbxMeshInfo> MeshInfo : SceneInfoPtr->MeshInfo)
 	{
 		//Delete all the delete asset
-		if (!MeshStatusMap.Contains(MeshInfo->OriginalImportPath))
+		if (!MeshStatusMap.Contains(MeshInfo->OriginalImportPath) || MeshInfo->bOriginalTypeChanged)
 		{
 			continue;
 		}
@@ -353,16 +428,18 @@ EReimportResult::Type UReimportFbxSceneFactory::Reimport(UObject* Obj)
 			continue;
 		}
 		//Make sure we load all package that will be deleted
-		UPackage* PkgExist = LoadPackage(nullptr, *(MeshInfo->GetImportPath()), LOAD_Verify | LOAD_NoWarn);
+		UPackage* PkgExist = MeshInfo->GetContentPackage();
 		if (PkgExist == nullptr)
 		{
 			continue;
 		}
-		PkgExist->FullyLoad();
 		//Find the asset
 		AssetDataToDelete.Add(AssetRegistryModule.Get().GetAssetByObjectPath(FName(*(MeshInfo->GetFullImportName()))));
 	}
 
+	FbxNode* RootNodeToImport = nullptr;
+	RootNodeToImport = FbxImporter->Scene->GetRootNode();
+
 	AllNewAssets.Empty();
 	AssetToSyncContentBrowser.Empty();
 	EReimportResult::Type ReimportResult = EReimportResult::Succeeded;
@@ -394,14 +471,19 @@ EReimportResult::Type UReimportFbxSceneFactory::Reimport(UObject* Obj)
 			}
 			else
 			{
-				//TODO reimport skeletal mesh
+				ReimportResult = ReimportSkeletalMesh(FbxImporter, MeshInfo);
 			}
 		}
 		else if ((MeshStatus & EFbxSceneReimportStatusFlags::Added) != EFbxSceneReimportStatusFlags::None || (MeshStatus & EFbxSceneReimportStatusFlags::Same) != EFbxSceneReimportStatusFlags::None)
 		{
-			//Create a package for this node
-			//Get Parent hierarchy name to create new package path
-			ReimportResult = ImportStaticMesh(FbxImporter, MeshInfo, SceneInfoPtr);
+			if (!MeshInfo->bIsSkelMesh)
+			{
+				ReimportResult = ImportStaticMesh(FbxImporter, MeshInfo, SceneInfoPtr);
+			}
+			else
+			{
+				ReimportResult = ImportSkeletalMesh(RootNodeToImport, FbxImporter, MeshInfo, SceneInfoPtr);
+			}
 		}
 	}
 
@@ -411,6 +493,10 @@ EReimportResult::Type UReimportFbxSceneFactory::Reimport(UObject* Obj)
 	SceneImportOptionsStaticMesh->FillStaticMeshInmportData(StaticMeshImportData, SceneImportOptions);
 	StaticMeshImportData->SaveConfig();
 
+	SFbxSceneOptionWindow::CopyFbxOptionsToSkeletalMeshOptions(GlobalImportSettingsReference, SceneImportOptionsSkeletalMesh);
+	SceneImportOptionsSkeletalMesh->FillSkeletalMeshInmportData(SkeletalMeshImportData, AnimSequenceImportData);
+	SkeletalMeshImportData->SaveConfig();
+
 	//Update the blueprint
 	UBlueprint *ReimportBlueprint = nullptr;
 	if (bCanReimportHierarchy && GlobalImportSettingsReference->bImportScene)
@@ -422,10 +508,13 @@ EReimportResult::Type UReimportFbxSceneFactory::Reimport(UObject* Obj)
 	TArray<TSharedPtr<FFbxMeshInfo>> ToRemoveHierarchyNode;
 	for (TSharedPtr<FFbxMeshInfo> MeshInfo : ReimportData->SceneInfoSourceData->MeshInfo)
 	{
-		EFbxSceneReimportStatusFlags MeshStatus = *(MeshStatusMap.Find(MeshInfo->OriginalImportPath));
-		if ((MeshStatus & EFbxSceneReimportStatusFlags::Removed) != EFbxSceneReimportStatusFlags::None)
+		if (MeshStatusMap.Contains(MeshInfo->OriginalImportPath))
 		{
-			ToRemoveHierarchyNode.Add(MeshInfo);
+			EFbxSceneReimportStatusFlags MeshStatus = *(MeshStatusMap.Find(MeshInfo->OriginalImportPath));
+			if ((MeshStatus & EFbxSceneReimportStatusFlags::Removed) != EFbxSceneReimportStatusFlags::None)
+			{
+				ToRemoveHierarchyNode.Add(MeshInfo);
+			}
 		}
 	}
 	for (TSharedPtr<FFbxMeshInfo> MeshInfo : ToRemoveHierarchyNode)
@@ -578,9 +667,14 @@ UBlueprint *UReimportFbxSceneFactory::UpdateOriginalBluePrint(FString &BluePrint
 		}
 		//Create the new nodes from the hierarchy actor
 		FKismetEditorUtilities::AddComponentsToBlueprint(BluePrint, HierarchyActor->GetInstanceComponents());
-		//Cleanup the temporary actor
-		HierarchyActor->Destroy();
+		
+		UWorld* World = HierarchyActor->GetWorld();
+		World->DestroyActor(HierarchyActor);
 
+		GEngine->BroadcastLevelActorListChanged();
+
+		FBlueprintEditorUtils::MarkBlueprintAsModified(BluePrint);
+		FKismetEditorUtilities::CompileBlueprint(BluePrint);
 		BluePrint->Modify();
 		BluePrint->PostEditChange();
 		AssetToSyncContentBrowser.Add(BluePrint);
@@ -589,6 +683,105 @@ UBlueprint *UReimportFbxSceneFactory::UpdateOriginalBluePrint(FString &BluePrint
 	return nullptr;
 }
 
+EReimportResult::Type UReimportFbxSceneFactory::ImportSkeletalMesh(void* VoidRootNodeToImport, void* VoidFbxImporter, TSharedPtr<FFbxMeshInfo> MeshInfo, TSharedPtr<FFbxSceneInfo> SceneInfoPtr)
+{
+	UnFbx::FFbxImporter* FbxImporter = (UnFbx::FFbxImporter*)VoidFbxImporter;
+	FbxNode *RootNodeToImport = (FbxNode *)VoidRootNodeToImport;
+	//FEditorDelegates::OnAssetPreImport.Broadcast(this, UStaticMesh::StaticClass(), GWorld, FName(""), TEXT("fbx"));
+	FbxNode *GeometryParentNode = nullptr;
+	//Get the first parent geometry node
+	for (int idx = 0; idx < FbxImporter->Scene->GetGeometryCount(); ++idx)
+	{
+		FbxGeometry *Geometry = FbxImporter->Scene->GetGeometry(idx);
+		if (Geometry->GetUniqueID() == MeshInfo->UniqueId)
+		{
+			GeometryParentNode = Geometry->GetNode();
+			break;
+		}
+	}
+	if (GeometryParentNode == nullptr)
+	{
+		FbxImporter->AddTokenizedErrorMessage(FTokenizedMessage::Create(EMessageSeverity::Error, FText::Format(FText::FromString("Reimport Mesh {0} fail, the mesh dont have any parent node inside the fbx."), FText::FromString(MeshInfo->GetImportPath()))), FName(TEXT("Reimport Fbx Scene")));
+		return EReimportResult::Failed;
+	}
+
+	FString PackageName = MeshInfo->GetImportPath();
+	FString StaticMeshName;
+	UPackage* Pkg = CreatePackageForNode(PackageName, StaticMeshName);
+	if (Pkg == nullptr)
+	{
+		return EReimportResult::Failed;
+	}
+
+	//Copy default options to StaticMeshImportData
+	SFbxSceneOptionWindow::CopyFbxOptionsToSkeletalMeshOptions(GlobalImportSettingsReference, SceneImportOptionsSkeletalMesh);
+	SceneImportOptionsSkeletalMesh->FillSkeletalMeshInmportData(SkeletalMeshImportData, AnimSequenceImportData);
+
+	UnFbx::FBXImportOptions* OverrideImportSettings = GetOptionsFromName(MeshInfo->OptionName);
+	if (OverrideImportSettings != nullptr)
+	{
+		SFbxSceneOptionWindow::CopyFbxOptionsToFbxOptions(OverrideImportSettings, GlobalImportSettings);
+		SFbxSceneOptionWindow::CopyFbxOptionsToSkeletalMeshOptions(OverrideImportSettings, SceneImportOptionsSkeletalMesh);
+	}
+	else
+	{
+		SFbxSceneOptionWindow::CopyFbxOptionsToFbxOptions(GlobalImportSettingsReference, GlobalImportSettings);
+		SFbxSceneOptionWindow::CopyFbxOptionsToSkeletalMeshOptions(GlobalImportSettingsReference, SceneImportOptionsSkeletalMesh);
+	}
+	SceneImportOptionsSkeletalMesh->FillSkeletalMeshInmportData(SkeletalMeshImportData, AnimSequenceImportData);
+	//Override the pivot bake option
+	if (GlobalImportSettings->bBakePivotInVertex && MeshInfo->PivotNodeUid == INVALID_UNIQUE_ID)
+	{
+		GlobalImportSettings->bBakePivotInVertex = false;
+	}
+
+	TArray< TArray<FbxNode*>* > SkelMeshArray;
+	FbxImporter->FillFbxSkelMeshArrayInScene(RootNodeToImport, SkelMeshArray, false);
+	UObject* NewObject = nullptr;
+	for (int32 i = 0; i < SkelMeshArray.Num(); i++)
+	{
+		TArray<FbxNode*> NodeArray = *SkelMeshArray[i];
+		FbxNode* RootNodeArrayNode = NodeArray[0];
+		TSharedPtr<FFbxNodeInfo> RootNodeInfo;
+		if (!FindSceneNodeInfo(SceneInfoPtr, RootNodeArrayNode->GetUniqueID(), RootNodeInfo))
+		{
+			continue;
+		}
+		if (!RootNodeInfo->AttributeInfo.IsValid() || RootNodeInfo->AttributeInfo->GetType() != USkeletalMesh::StaticClass())
+		{
+			continue;
+		}
+		TSharedPtr<FFbxMeshInfo> RootMeshInfo = StaticCastSharedPtr<FFbxMeshInfo>(RootNodeInfo->AttributeInfo);
+		if (!RootMeshInfo.IsValid() || RootMeshInfo->UniqueId != MeshInfo->UniqueId)
+		{
+			continue;
+		}
+
+		TArray<void*> VoidNodeArray;
+		for (FbxNode *Node : NodeArray)
+		{
+			void* VoidNode = (void*)Node;
+			VoidNodeArray.Add(VoidNode);
+		}
+		int32 TotalNumNodes = 0;
+		
+		NewObject = ImportOneSkeletalMesh(VoidRootNodeToImport, VoidFbxImporter, SceneInfoPtr, RF_Public | RF_Standalone, VoidNodeArray, TotalNumNodes);
+	}
+
+	for (int32 i = 0; i < SkelMeshArray.Num(); i++)
+	{
+		delete SkelMeshArray[i];
+	}
+
+	if (NewObject == nullptr)
+	{
+		return EReimportResult::Failed;
+	}
+	AllNewAssets.Add(MeshInfo, NewObject);
+	AssetToSyncContentBrowser.Add(NewObject);
+	return EReimportResult::Succeeded;
+}
+
 EReimportResult::Type UReimportFbxSceneFactory::ImportStaticMesh(void* VoidFbxImporter, TSharedPtr<FFbxMeshInfo> MeshInfo, TSharedPtr<FFbxSceneInfo> SceneInfoPtr)
 {
 	UnFbx::FFbxImporter* FbxImporter = (UnFbx::FFbxImporter*)VoidFbxImporter;
@@ -634,9 +827,13 @@ EReimportResult::Type UReimportFbxSceneFactory::ImportStaticMesh(void* VoidFbxIm
 		SFbxSceneOptionWindow::CopyFbxOptionsToStaticMeshOptions(GlobalImportSettingsReference, SceneImportOptionsStaticMesh);
 	}
 	SceneImportOptionsStaticMesh->FillStaticMeshInmportData(StaticMeshImportData, SceneImportOptions);
-
+	//Override the pivot bake option
+	if (GlobalImportSettings->bBakePivotInVertex && MeshInfo->PivotNodeUid == INVALID_UNIQUE_ID)
+	{
+		GlobalImportSettings->bBakePivotInVertex = false;
+	}
 	FName StaticMeshFName = FName(*(MeshInfo->Name));
-	UStaticMesh *NewObject = FbxImporter->ImportStaticMesh(Pkg, GeometryParentNode, StaticMeshFName, EObjectFlags::RF_Standalone, StaticMeshImportData );
+	UStaticMesh *NewObject = FbxImporter->ImportStaticMesh(Pkg, GeometryParentNode, StaticMeshFName, RF_Public | RF_Standalone, StaticMeshImportData );
 	if (NewObject == nullptr)
 	{
 		return EReimportResult::Failed;
@@ -646,16 +843,105 @@ EReimportResult::Type UReimportFbxSceneFactory::ImportStaticMesh(void* VoidFbxIm
 	return EReimportResult::Succeeded;
 }
 
-EReimportResult::Type UReimportFbxSceneFactory::ReimportStaticMesh(void* VoidFbxImporter, TSharedPtr<FFbxMeshInfo> MeshInfo)
+EReimportResult::Type UReimportFbxSceneFactory::ReimportSkeletalMesh(void* VoidFbxImporter, TSharedPtr<FFbxMeshInfo> MeshInfo)
 {
 	UnFbx::FFbxImporter* FbxImporter = (UnFbx::FFbxImporter*)VoidFbxImporter;
 	//Find the UObject associate with this MeshInfo
-	UPackage* PkgExist = LoadPackage(nullptr, *(MeshInfo->GetImportPath()), LOAD_Verify | LOAD_NoWarn);
-	if (PkgExist != nullptr)
+	UPackage* PkgExist = MeshInfo->GetContentPackage();
+
+	FString AssetName = MeshInfo->GetFullImportName();
+	USkeletalMesh* Mesh = FindObjectSafe<USkeletalMesh>(ANY_PACKAGE, *AssetName);
+	if (Mesh == nullptr)
 	{
-		PkgExist->FullyLoad();
+		//We reimport only skeletal mesh here
+		FbxImporter->AddTokenizedErrorMessage(FTokenizedMessage::Create(EMessageSeverity::Error, FText::Format(FText::FromString("Reimport Mesh {0} fail, the original skeletalmesh in the content browser cannot be load."), FText::FromString(MeshInfo->GetImportPath()))), FName(TEXT("Reimport Fbx Scene")));
+		return EReimportResult::Failed;
 	}
 
+	//Copy default options to SkeletalMeshImportData
+	SFbxSceneOptionWindow::CopyFbxOptionsToSkeletalMeshOptions(GlobalImportSettingsReference, SceneImportOptionsSkeletalMesh);
+	SceneImportOptionsSkeletalMesh->FillSkeletalMeshInmportData(SkeletalMeshImportData, AnimSequenceImportData);
+
+	UnFbx::FBXImportOptions* OverrideImportSettings = GetOptionsFromName(MeshInfo->OptionName);
+	if (OverrideImportSettings != nullptr)
+	{
+		SFbxSceneOptionWindow::CopyFbxOptionsToFbxOptions(OverrideImportSettings, GlobalImportSettings);
+		SFbxSceneOptionWindow::CopyFbxOptionsToSkeletalMeshOptions(OverrideImportSettings, SceneImportOptionsSkeletalMesh);
+	}
+	else
+	{
+		SFbxSceneOptionWindow::CopyFbxOptionsToFbxOptions(GlobalImportSettingsReference, GlobalImportSettings);
+		SFbxSceneOptionWindow::CopyFbxOptionsToSkeletalMeshOptions(GlobalImportSettingsReference, SceneImportOptionsSkeletalMesh);
+	}
+	SceneImportOptionsSkeletalMesh->FillSkeletalMeshInmportData(SkeletalMeshImportData, AnimSequenceImportData);
+
+	//Override the pivot bake option
+	if (GlobalImportSettings->bBakePivotInVertex && MeshInfo->PivotNodeUid == INVALID_UNIQUE_ID)
+	{
+		GlobalImportSettings->bBakePivotInVertex = false;
+	}
+	TArray<FbxNode*> OutSkeletalMeshArray;
+	if (FbxImporter->ReimportSkeletalMesh(Mesh, SkeletalMeshImportData, MeshInfo->UniqueId, &OutSkeletalMeshArray))
+	{
+		Mesh->AssetImportData->Update(FbxImportFileName);
+
+		// Try to find the outer package so we can dirty it up
+		if (Mesh->GetOuter())
+		{
+			Mesh->GetOuter()->MarkPackageDirty();
+		}
+		else
+		{
+			Mesh->MarkPackageDirty();
+		}
+		AllNewAssets.Add(MeshInfo, Mesh);
+		AssetToSyncContentBrowser.Add(Mesh);
+		
+		//TODO reimport all animation
+		//1. Store all anim sequence reference that was originally import, for every skeletal mesh
+		//2. On reimport match the existing one
+		//3. Reimport matching animation
+		if (GlobalImportSettings->bImportAnimations)
+		{
+			/*
+			USkeleton* Skeleton = Mesh->Skeleton;//  AnimSequence->GetSkeleton();
+			if (Skeleton)
+			{
+				//For all anim sequence to reimport
+				{
+					if (UEditorEngine::ReimportFbxAnimation(Skeleton, AnimSequence, ImportData, *Filename))
+					{
+						// update the data in case the file source has changed
+						ImportData->Update(UFactory::CurrentFilename);
+
+						// Try to find the outer package so we can dirty it up
+						if (AnimSequence->GetOuter())
+						{
+							AnimSequence->GetOuter()->MarkPackageDirty();
+						}
+						else
+						{
+							AnimSequence->MarkPackageDirty();
+						}
+					}
+				}
+			}
+			*/
+		}
+	}
+	else
+	{
+		return EReimportResult::Failed;
+	}
+	return EReimportResult::Succeeded;
+}
+
+EReimportResult::Type UReimportFbxSceneFactory::ReimportStaticMesh(void* VoidFbxImporter, TSharedPtr<FFbxMeshInfo> MeshInfo)
+{
+	UnFbx::FFbxImporter* FbxImporter = (UnFbx::FFbxImporter*)VoidFbxImporter;
+	//Load the UObject associate with this MeshInfo
+	MeshInfo->GetContentPackage();
+
 	FString AssetName = MeshInfo->GetFullImportName();
 	UStaticMesh* Mesh = FindObjectSafe<UStaticMesh>(ANY_PACKAGE, *AssetName);
 	if (Mesh == nullptr)
@@ -681,6 +967,11 @@ EReimportResult::Type UReimportFbxSceneFactory::ReimportStaticMesh(void* VoidFbx
 		SFbxSceneOptionWindow::CopyFbxOptionsToStaticMeshOptions(GlobalImportSettingsReference, SceneImportOptionsStaticMesh);
 	}
 	SceneImportOptionsStaticMesh->FillStaticMeshInmportData(StaticMeshImportData, SceneImportOptions);
+	//Override the pivot bake option
+	if (GlobalImportSettings->bBakePivotInVertex && MeshInfo->PivotNodeUid == INVALID_UNIQUE_ID)
+	{
+		GlobalImportSettings->bBakePivotInVertex = false;
+	}
 
 	FbxImporter->ApplyTransformSettingsToFbxNode(FbxImporter->Scene->GetRootNode(), StaticMeshImportData);
 	const TArray<UAssetUserData*>* UserData = Mesh->GetAssetUserDataArray();
@@ -701,7 +992,13 @@ EReimportResult::Type UReimportFbxSceneFactory::ReimportStaticMesh(void* VoidFbx
 	// preserve extended bound settings
 	const FVector PositiveBoundsExtension = Mesh->PositiveBoundsExtension;
 	const FVector NegativeBoundsExtension = Mesh->NegativeBoundsExtension;
-	Mesh = FbxImporter->ReimportSceneStaticMesh(MeshInfo->UniqueId, Mesh, StaticMeshImportData);
+	uint64 NodeInfoUid = INVALID_UNIQUE_ID;
+	if (GlobalImportSettings->bBakePivotInVertex && MeshInfo->PivotNodeUid != INVALID_UNIQUE_ID)
+	{
+		NodeInfoUid = MeshInfo->PivotNodeUid;
+	}
+
+	Mesh = FbxImporter->ReimportSceneStaticMesh(NodeInfoUid, MeshInfo->UniqueId, Mesh, StaticMeshImportData);
 	if (Mesh != nullptr)
 	{
 		//Put back the new mesh data since the reimport is putting back the original import data
