Copying //UE4/Dev-Editor to Dev-Main (//UE4/Dev-Main)

diff --git a/Engine/Source/Editor/UnrealEd/Private/Fbx/SFbxSceneOptionWindow.cpp b/Engine/Source/Editor/UnrealEd/Private/Fbx/SFbxSceneOptionWindow.cpp
--- a/Engine/Source/Editor/UnrealEd/Private/Fbx/SFbxSceneOptionWindow.cpp
+++ b/Engine/Source/Editor/UnrealEd/Private/Fbx/SFbxSceneOptionWindow.cpp
@@ -12,6 +12,7 @@
 #include "STextComboBox.h"
 #include "SEditableTextBox.h"
 #include "../FbxImporter.h"
+#include "Dialogs/DlgPickPath.h"
 
 #define LOCTEXT_NAMESPACE "FBXOption"
 
@@ -25,22 +26,23 @@ SFbxSceneOptionWindow::SFbxSceneOptionWindow()
 	SceneImportOptionsStaticMeshDisplay = nullptr;
 	OverrideNameOptionsMap = nullptr;
 	SceneImportOptionsSkeletalMeshDisplay = nullptr;
-	SceneImportOptionsAnimationDisplay = nullptr;
-	SceneImportOptionsMaterialDisplay = nullptr;
 	OwnerWindow = nullptr;
 	FbxSceneImportTabManager = nullptr;
 	Layout = nullptr;
 	bShouldImport = false;
 	SceneTabTreeview = nullptr;
 	SceneTabDetailsView = nullptr;
+	OverrideNameOptions.Empty();
 	StaticMeshTabListView = nullptr;
 	StaticMeshTabDetailsView = nullptr;
+	SkeletalMeshTabListView = nullptr;
+	SkeletalMeshTabDetailsView = nullptr;
 	SceneReimportTreeview = nullptr;
 	StaticMeshReimportListView = nullptr;
 	StaticMeshReimportDetailsView = nullptr;
 	MaterialsTabListView = nullptr;
 	TexturesArray.Reset();
-	MaterialPrefixName.Empty();
+	MaterialBasePath.Empty();
 }
 
 SFbxSceneOptionWindow::~SFbxSceneOptionWindow()
@@ -57,8 +59,6 @@ SFbxSceneOptionWindow::~SFbxSceneOptionWindow()
 	SceneImportOptionsStaticMeshDisplay = nullptr;
 	OverrideNameOptionsMap = nullptr;
 	SceneImportOptionsSkeletalMeshDisplay = nullptr;
-	SceneImportOptionsAnimationDisplay = nullptr;
-	SceneImportOptionsMaterialDisplay = nullptr;
 	OwnerWindow = nullptr;
 	FbxSceneImportTabManager = nullptr;
 	Layout = nullptr;
@@ -67,12 +67,14 @@ SFbxSceneOptionWindow::~SFbxSceneOptionWindow()
 	SceneTabDetailsView = nullptr;
 	StaticMeshTabListView = nullptr;
 	StaticMeshTabDetailsView = nullptr;
+	SkeletalMeshTabListView = nullptr;
+	SkeletalMeshTabDetailsView = nullptr;
 	SceneReimportTreeview = nullptr;
 	StaticMeshReimportListView = nullptr;
 	StaticMeshReimportDetailsView = nullptr;
 	MaterialsTabListView = nullptr;
 	TexturesArray.Reset();
-	MaterialPrefixName.Empty();
+	MaterialBasePath.Empty();
 }
 
 TSharedRef<SDockTab> SFbxSceneOptionWindow::SpawnSceneTab(const FSpawnTabArgs& Args)
@@ -83,7 +85,7 @@ TSharedRef<SDockTab> SFbxSceneOptionWindow::SpawnSceneTab(const FSpawnTabArgs& A
 
 	TSharedPtr<SBox> InspectorBox;
 	TSharedRef<SDockTab> DockTab = SNew(SDockTab)
-		.TabRole(ETabRole::NomadTab)
+		.TabRole(ETabRole::PanelTab)
 		.Label(LOCTEXT("WidgetFbxSceneActorTab", "Scene"))
 		.ToolTipText(LOCTEXT("WidgetFbxSceneActorTabTextToolTip", "Switch to the scene tab."))
 		[
@@ -148,6 +150,9 @@ TSharedRef<SDockTab> SFbxSceneOptionWindow::SpawnSceneTab(const FSpawnTabArgs& A
 				SAssignNew(InspectorBox, SBox)
 			]
 		];
+	//Prevent user to close the tab
+	DockTab->SetCanCloseTab(SDockTab::FCanCloseTab::CreateRaw(this, &SFbxSceneOptionWindow::CanCloseTab));
+
 	FPropertyEditorModule& PropertyEditorModule = FModuleManager::GetModuleChecked<FPropertyEditorModule>("PropertyEditor");
 	FDetailsViewArgs DetailsViewArgs;
 	DetailsViewArgs.bAllowSearch = false;
@@ -159,14 +164,22 @@ TSharedRef<SDockTab> SFbxSceneOptionWindow::SpawnSceneTab(const FSpawnTabArgs& A
 	return DockTab;
 }
 
+bool SFbxSceneOptionWindow::CanCloseTab()
+{
+	return false;
+}
+
 void SFbxSceneOptionWindow::OnFinishedChangingPropertiesSceneTabDetailView(const FPropertyChangedEvent& PropertyChangedEvent)
 {
 	if (!SceneInfoOriginal.IsValid())
 	{
 		MaterialsTabListView->SetCreateContentFolderHierarchy(SceneImportOptionsDisplay->bCreateContentFolderHierarchy);
 		//Update the MaterialList
-		MaterialsTabListView->UpdateMaterialPrefixName();
+		MaterialsTabListView->UpdateMaterialBasePath();
 	}
+	//Set the Global Import setting
+	GlobalImportSettings->bBakePivotInVertex = SceneImportOptionsDisplay->bBakePivotInVertex;
+	GlobalImportSettings->bInvertNormalMap = SceneImportOptionsDisplay->bInvertNormalMaps;
 }
 
 TSharedRef<SDockTab> SFbxSceneOptionWindow::SpawnStaticMeshTab(const FSpawnTabArgs& Args)
@@ -175,122 +188,440 @@ TSharedRef<SDockTab> SFbxSceneOptionWindow::SpawnStaticMeshTab(const FSpawnTabAr
 	StaticMeshTabListView = SNew(SFbxSceneStaticMeshListView)
 		.SceneInfo(SceneInfo)
 		.GlobalImportSettings(GlobalImportSettings)
+		.OverrideNameOptions(&OverrideNameOptions)
 		.OverrideNameOptionsMap(OverrideNameOptionsMap)
 		.SceneImportOptionsStaticMeshDisplay(SceneImportOptionsStaticMeshDisplay);
 
 	TSharedPtr<SBox> InspectorBox;
 	TSharedRef<SDockTab> DockTab = SNew(SDockTab)
-		.TabRole(ETabRole::NomadTab)
+		.TabRole(ETabRole::PanelTab)
 		.Label(LOCTEXT("WidgetFbxSceneStaticMeshTab", "Static Meshes"))
 		.ToolTipText(LOCTEXT("WidgetFbxSceneStaticMeshTabTextToolTip", "Switch to the static meshes tab."))
 		[
-			SNew(SVerticalBox)
-			+ SVerticalBox::Slot()
-			.FillHeight(1.0f)
+			SNew(SSplitter)
+			.Orientation(Orient_Horizontal)
+			+ SSplitter::Slot()
+			.Value(0.4f)
 			[
-				SNew(SSplitter)
-				.Orientation(Orient_Horizontal)
-				+ SSplitter::Slot()
-				.Value(0.4f)
+				SNew(SBox)
 				[
-					SNew(SBox)
+					StaticMeshTabListView.ToSharedRef()
+				]
+			]
+			+ SSplitter::Slot()
+			.Value(0.6f)
+			[
+				SNew(SVerticalBox)
+				+ SVerticalBox::Slot()
+				.AutoHeight()
+				[
+					SNew(SHorizontalBox)
+					+ SHorizontalBox::Slot()
+					.AutoWidth()
+					[
+						StaticMeshTabListView->CreateOverrideOptionComboBox().ToSharedRef()
+					]
+					+ SHorizontalBox::Slot()
+					.AutoWidth()
+					[
+						SNew(SButton)
+						.HAlign(HAlign_Center)
+						.Text(LOCTEXT("FbxOptionWindow_SM_Select_asset_using", "Select Asset Using"))
+						.OnClicked(StaticMeshTabListView.Get(), &SFbxSceneStaticMeshListView::OnSelectAssetUsing)
+					]
+					+ SHorizontalBox::Slot()
+					.AutoWidth()
+					[
+						SNew(SButton)
+						.Text(LOCTEXT("FbxOptionWindow_SM_CreateOverride", "Create Override"))
+						.ToolTipText(LOCTEXT("FbxOptionWindow_SM_CreateOverrideTooltip", "Create Override to specify custom import options for some static meshes.\nTo assign options use context menu on static meshes."))
+						.OnClicked(StaticMeshTabListView.Get(), &SFbxSceneStaticMeshListView::OnCreateOverrideOptions)
+					]
+					+ SHorizontalBox::Slot()
+					.AutoWidth()
 					[
-						StaticMeshTabListView.ToSharedRef()
+						SNew(SButton)
+						.HAlign(HAlign_Center)
+						.Text(LOCTEXT("FbxOptionWindow_SM_Delete", "Delete"))
+						.IsEnabled(StaticMeshTabListView.Get(), &SFbxSceneStaticMeshListView::CanDeleteOverride)
+						.OnClicked(StaticMeshTabListView.Get(), &SFbxSceneStaticMeshListView::OnDeleteOverride)
 					]
 				]
-				+ SSplitter::Slot()
-				.Value(0.6f)
+				+ SVerticalBox::Slot()
+				.FillHeight(1.0f)
+				[
+					SAssignNew(InspectorBox, SBox)
+				]
+			]
+		];
+
+	//Prevent user to close the tab
+	DockTab->SetCanCloseTab(SDockTab::FCanCloseTab::CreateRaw(this, &SFbxSceneOptionWindow::CanCloseTab));
+
+	FPropertyEditorModule& PropertyEditorModule = FModuleManager::GetModuleChecked<FPropertyEditorModule>("PropertyEditor");
+	FDetailsViewArgs DetailsViewArgs;
+	DetailsViewArgs.bAllowSearch = false;
+	DetailsViewArgs.NameAreaSettings = FDetailsViewArgs::HideNameArea;
+	StaticMeshTabDetailsView = PropertyEditorModule.CreateDetailView(DetailsViewArgs);
+	InspectorBox->SetContent(StaticMeshTabDetailsView->AsShared());
+	StaticMeshTabDetailsView->SetObject(SceneImportOptionsStaticMeshDisplay);
+	StaticMeshTabDetailsView->OnFinishedChangingProperties().AddSP(StaticMeshTabListView.Get(), &SFbxSceneStaticMeshListView::OnFinishedChangingProperties);
+	return DockTab;
+}
+
+TSharedRef<SDockTab> SFbxSceneOptionWindow::SpawnSkeletalMeshReimportTab(const FSpawnTabArgs& Args)
+{
+
+	//Create the Skeletal mesh listview
+	SkeletalMeshReimportListView = SNew(SFbxSceneSkeletalMeshReimportListView)
+		.SceneInfo(SceneInfo)
+		.SceneInfoOriginal(SceneInfoOriginal)
+		.GlobalImportSettings(GlobalImportSettings)
+		.OverrideNameOptions(&OverrideNameOptions)
+		.OverrideNameOptionsMap(OverrideNameOptionsMap)
+		.SceneImportOptionsSkeletalMeshDisplay(SceneImportOptionsSkeletalMeshDisplay)
+		.MeshStatusMap(MeshStatusMap);
+
+	TSharedPtr<SBox> InspectorBox;
+	TSharedRef<SDockTab> DockTab = SNew(SDockTab)
+		.TabRole(ETabRole::PanelTab)
+		.Label(LOCTEXT("WidgetFbxSceneReimportSkeletalMeshTab", "Skeletal Meshes"))
+		.ToolTipText(LOCTEXT("WidgetFbxSceneReimportSkeletalMeshTabTextToolTip", "Switch to the reimport Skeletal meshes tab."))
+		[
+			SNew(SVerticalBox)
+			+ SVerticalBox::Slot()
+			.FillHeight(1.0f)
+			[
+				SNew(SHorizontalBox)
+				+ SHorizontalBox::Slot()
+				.FillWidth(1.0f)
 				[
 					SNew(SVerticalBox)
 					+ SVerticalBox::Slot()
+					.HAlign(HAlign_Left)
 					.AutoHeight()
 					[
-						SNew(SHorizontalBox)
-						+ SHorizontalBox::Slot()
-						.AutoWidth()
+						SNew(SUniformGridPanel)
+						.SlotPadding(2)
+						+ SUniformGridPanel::Slot(0, 0)
+						.VAlign(VAlign_Center)
+						.HAlign(HAlign_Left)
 						[
-							StaticMeshTabListView->CreateOverrideOptionComboBox().ToSharedRef()
+							SNew(STextBlock)
+							.Text(LOCTEXT("FbxOptionWindow_Scene_Filters_Label", "Filters:"))
 						]
-						+ SHorizontalBox::Slot()
-						.AutoWidth()
+						+ SUniformGridPanel::Slot(1, 0)
 						[
-							SNew(SButton)
-							.Text(LOCTEXT("FbxOptionWindow_SM_CreateOverride", "Create Override"))
-							.OnClicked(StaticMeshTabListView.Get(), &SFbxSceneStaticMeshListView::OnCreateOverrideOptions)
+							SNew(SBorder)
+							.BorderImage(FEditorStyle::GetBrush("ToolPanel.GroupBorder"))
+							[
+								SNew(SHorizontalBox)
+								+ SHorizontalBox::Slot()
+								.AutoWidth()
+								[
+									SNew(SCheckBox)
+									.HAlign(HAlign_Center)
+									.OnCheckStateChanged(SkeletalMeshReimportListView.Get(), &SFbxSceneSkeletalMeshReimportListView::OnToggleFilterAddContent)
+									.IsChecked(SkeletalMeshReimportListView.Get(), &SFbxSceneSkeletalMeshReimportListView::IsFilterAddContentChecked)
+								]
+								+ SHorizontalBox::Slot()
+								.FillWidth(1.0f)
+								.Padding(0.0f, 3.0f, 6.0f, 3.0f)
+								.VAlign(VAlign_Center)
+								[
+									SNew(STextBlock)
+									.Text(LOCTEXT("FbxOptionWindow_Scene_Reimport_Filter_Add_Content", "Add"))
+								]
+							]
+						]
+						+ SUniformGridPanel::Slot(2, 0)
+						[
+							SNew(SBorder)
+							.BorderImage(FEditorStyle::GetBrush("ToolPanel.GroupBorder"))
+							[
+								SNew(SHorizontalBox)
+								+ SHorizontalBox::Slot()
+								.AutoWidth()
+								[
+									SNew(SCheckBox)
+									.HAlign(HAlign_Center)
+									.OnCheckStateChanged(SkeletalMeshReimportListView.Get(), &SFbxSceneSkeletalMeshReimportListView::OnToggleFilterDeleteContent)
+									.IsChecked(SkeletalMeshReimportListView.Get(), &SFbxSceneSkeletalMeshReimportListView::IsFilterDeleteContentChecked)
+								]
+								+ SHorizontalBox::Slot()
+								.FillWidth(1.0f)
+								.Padding(0.0f, 3.0f, 6.0f, 3.0f)
+								.VAlign(VAlign_Center)
+								[
+									SNew(STextBlock)
+									.Text(LOCTEXT("FbxOptionWindow_Scene_Reimport_Filter_Delete_Content", "Delete"))
+								]
+							]
+						]
+						+ SUniformGridPanel::Slot(3, 0)
+						[
+							SNew(SBorder)
+							.BorderImage(FEditorStyle::GetBrush("ToolPanel.GroupBorder"))
+							[
+								SNew(SHorizontalBox)
+								+ SHorizontalBox::Slot()
+								.AutoWidth()
+								[
+									SNew(SCheckBox)
+									.HAlign(HAlign_Center)
+									.OnCheckStateChanged(SkeletalMeshReimportListView.Get(), &SFbxSceneSkeletalMeshReimportListView::OnToggleFilterOverwriteContent)
+									.IsChecked(SkeletalMeshReimportListView.Get(), &SFbxSceneSkeletalMeshReimportListView::IsFilterOverwriteContentChecked)
+								]
+								+ SHorizontalBox::Slot()
+								.FillWidth(1.0f)
+								.Padding(0.0f, 3.0f, 6.0f, 3.0f)
+								.VAlign(VAlign_Center)
+								[
+									SNew(STextBlock)
+									.Text(LOCTEXT("FbxOptionWindow_Scene_Reimport_Filter_Overwrite_Content", "Overwrite"))
+								]
+							]
+						]
+						+ SUniformGridPanel::Slot(4, 0)
+						[
+							SNew(SBorder)
+							.BorderImage(FEditorStyle::GetBrush("ToolPanel.GroupBorder"))
+							[
+								SNew(SHorizontalBox)
+								+ SHorizontalBox::Slot()
+								.AutoWidth()
+								[
+									SNew(SCheckBox)
+									.HAlign(HAlign_Center)
+									.OnCheckStateChanged(SkeletalMeshReimportListView.Get(), &SFbxSceneSkeletalMeshReimportListView::OnToggleFilterDiff)
+									.IsChecked(SkeletalMeshReimportListView.Get(), &SFbxSceneSkeletalMeshReimportListView::IsFilterDiffChecked)
+								]
+								+ SHorizontalBox::Slot()
+								.FillWidth(1.0f)
+								.Padding(0.0f, 3.0f, 6.0f, 3.0f)
+								.VAlign(VAlign_Center)
+								[
+									SNew(STextBlock)
+									.Text(LOCTEXT("FbxOptionWindow_Scene_Reimport_Filter_Diff", "Diff"))
+									.ToolTipText(LOCTEXT("FbxOptionWindow_Scene_Reimport_Filter_Diff_Tooltip", "Show every reimport item that dont match between the original fbx and the new one."))
+								]
+							]
 						]
 					]
 					+ SVerticalBox::Slot()
 					.FillHeight(1.0f)
 					[
-						SAssignNew(InspectorBox, SBox)
+						SNew(SSplitter)
+						.Orientation(Orient_Vertical)
+						+ SSplitter::Slot()
+						.Value(0.4f)
+						[
+							SNew(SBox)
+							[
+								SkeletalMeshReimportListView.ToSharedRef()
+							]
+						]
+						+ SSplitter::Slot()
+						.Value(0.6f)
+						[
+							SNew(SVerticalBox)
+							+ SVerticalBox::Slot()
+							.AutoHeight()
+							[
+								SNew(SHorizontalBox)
+								+ SHorizontalBox::Slot()
+								.AutoWidth()
+								[
+									SkeletalMeshReimportListView->CreateOverrideOptionComboBox().ToSharedRef()
+								]
+								+ SHorizontalBox::Slot()
+								.AutoWidth()
+								[
+									SNew(SButton)
+									.HAlign(HAlign_Center)
+									.Text(LOCTEXT("FbxOptionWindow_SM_Select_asset_using", "Select Asset Using"))
+									.OnClicked(SkeletalMeshReimportListView.Get(), &SFbxSceneSkeletalMeshReimportListView::OnSelectAssetUsing)
+								]
+								+ SHorizontalBox::Slot()
+								.AutoWidth()
+								[
+									SNew(SButton)
+									.Text(LOCTEXT("FbxOptionWindow_SM_CreateOverride", "Create Override"))
+									.ToolTipText(LOCTEXT("FbxOptionWindow_SM_CreateOverrideTooltip", "Create Override to specify custom import options for some Skeletal meshes.\nTo assign options use context menu on Skeletal meshes."))
+									.OnClicked(SkeletalMeshReimportListView.Get(), &SFbxSceneSkeletalMeshReimportListView::OnCreateOverrideOptions)
+								]
+								+ SHorizontalBox::Slot()
+								.AutoWidth()
+								[
+									SNew(SButton)
+									.HAlign(HAlign_Center)
+									.Text(LOCTEXT("FbxOptionWindow_SM_Delete", "Delete"))
+									.IsEnabled(SkeletalMeshReimportListView.Get(), &SFbxSceneSkeletalMeshReimportListView::CanDeleteOverride)
+									.OnClicked(SkeletalMeshReimportListView.Get(), &SFbxSceneSkeletalMeshReimportListView::OnDeleteOverride)
+								]
+							]
+							+ SVerticalBox::Slot()
+							.FillHeight(1.0f)
+							[
+								SAssignNew(InspectorBox, SBox)
+							]
+						]
 					]
 				]
 			]
-			+ SVerticalBox::Slot()
-			.AutoHeight()
-			.HAlign(HAlign_Right)
-			.Padding(2)
-			[
-				SNew(SUniformGridPanel)
-				.SlotPadding(2)
-				+ SUniformGridPanel::Slot(0, 0)
-				[
-					SNew(SButton)
-					.HAlign(HAlign_Center)
-					.Text(LOCTEXT("FbxOptionWindow_SM_Select_asset_using", "Select Asset Using"))
-					.OnClicked(StaticMeshTabListView.Get(), &SFbxSceneStaticMeshListView::OnSelectAssetUsing)
-				]
-				+ SUniformGridPanel::Slot(1, 0)
-				[
-					SNew(SButton)
-					.HAlign(HAlign_Center)
-					.Text(LOCTEXT("FbxOptionWindow_SM_Delete", "Delete"))
-					.IsEnabled(StaticMeshTabListView.Get(), &SFbxSceneStaticMeshListView::CanDeleteOverride)
-					.OnClicked(StaticMeshTabListView.Get(), &SFbxSceneStaticMeshListView::OnDeleteOverride)
-				]
-			]
 		];
+	
+	//Prevent user to close the tab
+	DockTab->SetCanCloseTab(SDockTab::FCanCloseTab::CreateRaw(this, &SFbxSceneOptionWindow::CanCloseTab));
+
 	FPropertyEditorModule& PropertyEditorModule = FModuleManager::GetModuleChecked<FPropertyEditorModule>("PropertyEditor");
 	FDetailsViewArgs DetailsViewArgs;
 	DetailsViewArgs.bAllowSearch = false;
 	DetailsViewArgs.NameAreaSettings = FDetailsViewArgs::HideNameArea;
-	StaticMeshTabDetailsView = PropertyEditorModule.CreateDetailView(DetailsViewArgs);
-	InspectorBox->SetContent(StaticMeshTabDetailsView->AsShared());
-	StaticMeshTabDetailsView->SetObject(SceneImportOptionsStaticMeshDisplay);
-	StaticMeshTabDetailsView->OnFinishedChangingProperties().AddSP(StaticMeshTabListView.Get(), &SFbxSceneStaticMeshListView::OnFinishedChangingProperties);
+	SkeletalMeshReimportDetailsView = PropertyEditorModule.CreateDetailView(DetailsViewArgs);
+	InspectorBox->SetContent(SkeletalMeshReimportDetailsView->AsShared());
+	SkeletalMeshReimportDetailsView->SetObject(SceneImportOptionsSkeletalMeshDisplay);
+	SkeletalMeshReimportDetailsView->OnFinishedChangingProperties().AddSP(SkeletalMeshReimportListView.Get(), &SFbxSceneSkeletalMeshReimportListView::OnFinishedChangingProperties);
 	return DockTab;
 }
 
 TSharedRef<SDockTab> SFbxSceneOptionWindow::SpawnSkeletalMeshTab(const FSpawnTabArgs& Args)
 {
-	return SNew(SDockTab)
-		.TabRole(ETabRole::NomadTab)
+	//Create the static mesh listview
+	SkeletalMeshTabListView = SNew(SFbxSceneSkeletalMeshListView)
+		.SceneInfo(SceneInfo)
+		.GlobalImportSettings(GlobalImportSettings)
+		.OverrideNameOptions(&OverrideNameOptions)
+		.OverrideNameOptionsMap(OverrideNameOptionsMap)
+		.SceneImportOptionsSkeletalMeshDisplay(SceneImportOptionsSkeletalMeshDisplay);
+
+	TSharedPtr<SBox> InspectorBox;
+	TSharedRef<SDockTab> DockTab = SNew(SDockTab)
+		.TabRole(ETabRole::PanelTab)
 		.Label(LOCTEXT("WidgetFbxSceneSkeletalMeshTab", "Skeletal Meshes"))
 		.ToolTipText(LOCTEXT("WidgetFbxSceneSkeletalMeshTabTextToolTip", "Switch to the skeletal meshes tab."))
 		[
-			SNew(STextBlock).Text(LOCTEXT("SkeletalMeshTextBoxPlaceHolder", "Skeletal Mesh Import options placeholder"))
+			SNew(SSplitter)
+			.Orientation(Orient_Horizontal)
+			+ SSplitter::Slot()
+			.Value(0.4f)
+			[
+				SNew(SBox)
+				[
+					SkeletalMeshTabListView.ToSharedRef()
+				]
+			]
+			+ SSplitter::Slot()
+			.Value(0.6f)
+			[
+				SNew(SVerticalBox)
+				+ SVerticalBox::Slot()
+				.AutoHeight()
+				[
+					SNew(SHorizontalBox)
+					+ SHorizontalBox::Slot()
+					.AutoWidth()
+					[
+						SkeletalMeshTabListView->CreateOverrideOptionComboBox().ToSharedRef()
+					]
+					+ SHorizontalBox::Slot()
+					.AutoWidth()
+					[
+						SNew(SButton)
+						.HAlign(HAlign_Center)
+						.Text(LOCTEXT("FbxOptionWindow_SM_Select_asset_using", "Select Asset Using"))
+						.OnClicked(SkeletalMeshTabListView.Get(), &SFbxSceneSkeletalMeshListView::OnSelectAssetUsing)
+					]
+					+ SHorizontalBox::Slot()
+					.AutoWidth()
+					[
+						SNew(SButton)
+						.Text(LOCTEXT("FbxOptionWindow_SM_CreateOverride", "Create Override"))
+						.ToolTipText(LOCTEXT("FbxOptionWindow_SM_CreateOverrideTooltip", "Create Override to specify custom import options for some static meshes.\nTo assign options use context menu on static meshes."))
+						.OnClicked(SkeletalMeshTabListView.Get(), &SFbxSceneSkeletalMeshListView::OnCreateOverrideOptions)
+					]
+					+ SHorizontalBox::Slot()
+					.AutoWidth()
+					[
+						SNew(SButton)
+						.HAlign(HAlign_Center)
+						.Text(LOCTEXT("FbxOptionWindow_SM_Delete", "Delete"))
+						.IsEnabled(SkeletalMeshTabListView.Get(), &SFbxSceneSkeletalMeshListView::CanDeleteOverride)
+						.OnClicked(SkeletalMeshTabListView.Get(), &SFbxSceneSkeletalMeshListView::OnDeleteOverride)
+					]
+				]
+				+ SVerticalBox::Slot()
+				.FillHeight(1.0f)
+				[
+					SAssignNew(InspectorBox, SBox)
+				]
+			]
 		];
+	
+	//Prevent user to close the tab
+	DockTab->SetCanCloseTab(SDockTab::FCanCloseTab::CreateRaw(this, &SFbxSceneOptionWindow::CanCloseTab));
+
+	FPropertyEditorModule& PropertyEditorModule = FModuleManager::GetModuleChecked<FPropertyEditorModule>("PropertyEditor");
+	FDetailsViewArgs DetailsViewArgs;
+	DetailsViewArgs.bAllowSearch = false;
+	DetailsViewArgs.NameAreaSettings = FDetailsViewArgs::HideNameArea;
+	SkeletalMeshTabDetailsView = PropertyEditorModule.CreateDetailView(DetailsViewArgs);
+	InspectorBox->SetContent(SkeletalMeshTabDetailsView->AsShared());
+	SkeletalMeshTabDetailsView->SetObject(SceneImportOptionsSkeletalMeshDisplay);
+	SkeletalMeshTabDetailsView->OnFinishedChangingProperties().AddSP(SkeletalMeshTabListView.Get(), &SFbxSceneSkeletalMeshListView::OnFinishedChangingProperties);
+	return DockTab;
 }
 
-FText SFbxSceneOptionWindow::GetMaterialPrefixName() const
+FText SFbxSceneOptionWindow::GetMaterialBasePath() const
 {
-	return FText::FromString(MaterialPrefixName);
+	return FText::FromString(MaterialBasePath);
 }
 
-void SFbxSceneOptionWindow::OnMaterialPrefixCommited(const FText& InText, ETextCommit::Type InCommitType)
+void SFbxSceneOptionWindow::OnMaterialBasePathCommited(const FText& InText, ETextCommit::Type InCommitType)
 {
-	MaterialPrefixName = InText.ToString();
+	MaterialBasePath = InText.ToString();
 	//Commit only if all the rules are respected
-	if (MaterialPrefixName.StartsWith(TEXT("/")) && MaterialPrefixName.EndsWith(TEXT("/")))
+	if (MaterialBasePath.StartsWith(TEXT("/")) && MaterialBasePath.EndsWith(TEXT("/")))
+	{
+		GlobalImportSettings->MaterialBasePath = FName(*MaterialBasePath);
+		MaterialsTabListView->UpdateMaterialBasePath();
+	}
+}
+
+FReply SFbxSceneOptionWindow::OnMaterialBasePathBrowse()
+{
+	TSharedRef<SDlgPickPath> PickContentPathDlg =
+		SNew(SDlgPickPath)
+		.Title(LOCTEXT("FbxChooseImportOverrideMaterialPath", "Choose Location path for importing all materials"));
+
+	if (PickContentPathDlg->ShowModal() == EAppReturnType::Cancel)
 	{
-		GlobalImportSettings->MaterialPrefixName = FName(*MaterialPrefixName);
-		MaterialsTabListView->UpdateMaterialPrefixName();
+		return FReply::Handled();
 	}
+	MaterialBasePath = PickContentPathDlg->GetPath().ToString();
+	
+	if (MaterialBasePath.IsEmpty())
+	{
+		return FReply::Handled();
+	}
+	//Make sure it start and end with a slash
+	if (!MaterialBasePath.EndsWith(TEXT("/")))
+	{
+		MaterialBasePath += TEXT("/");
+	}
+	if (!MaterialBasePath.StartsWith(TEXT("/")))
+	{
+		MaterialBasePath.InsertAt(0, TEXT("/"));
+	}
+	GlobalImportSettings->MaterialBasePath = FName(*MaterialBasePath);
+	MaterialsTabListView->UpdateMaterialBasePath();
+
+	return FReply::Handled();
 }
 
-FSlateColor SFbxSceneOptionWindow::GetMaterialPrefixTextColor() const
+FSlateColor SFbxSceneOptionWindow::GetMaterialBasePathTextColor() const
 {
-	if (MaterialPrefixName.StartsWith(TEXT("/")) && MaterialPrefixName.EndsWith(TEXT("/")))
+	if (MaterialBasePath.IsEmpty() || (MaterialBasePath.StartsWith(TEXT("/")) && MaterialBasePath.EndsWith(TEXT("/"))))
 	{
 		return FSlateColor::UseForeground();
 	}
@@ -313,7 +644,7 @@ TSharedRef<SDockTab> SFbxSceneOptionWindow::SpawnMaterialTab(const FSpawnTabArgs
 		.CreateContentFolderHierarchy(SceneImportOptionsDisplay->bCreateContentFolderHierarchy != 0);
 
 	TSharedRef<SDockTab> DockTab = SNew(SDockTab)
-		.TabRole(ETabRole::NomadTab)
+		.TabRole(ETabRole::PanelTab)
 		.Label(LOCTEXT("WidgetFbxSceneMaterialsTab", "Materials"))
 		.ToolTipText(LOCTEXT("WidgetFbxSceneMaterialsTabTextToolTip", "Switch to the materials tab."))
 		[
@@ -328,10 +659,10 @@ TSharedRef<SDockTab> SFbxSceneOptionWindow::SpawnMaterialTab(const FSpawnTabArgs
 				.VAlign(VAlign_Center)
 				[
 					SNew(STextBlock)
-					.Text(LOCTEXT("FbxOptionWindow_Scene_Material_Prefix", "Material name prefix path: "))
+					.Text(LOCTEXT("FbxOptionWindow_Scene_Material_Prefix", "Material override base path: "))
 				]
 				+ SHorizontalBox::Slot()
-				.FillWidth(1.0f)
+				.AutoWidth()
 				.Padding(5.0f, 3.0f, 6.0f, 3.0f)
 				.VAlign(VAlign_Center)
 				[
@@ -340,13 +671,22 @@ TSharedRef<SDockTab> SFbxSceneOptionWindow::SpawnMaterialTab(const FSpawnTabArgs
 					[
 						SNew(SEditableText)
 						.SelectAllTextWhenFocused(true)
-						.Text(this, &SFbxSceneOptionWindow::GetMaterialPrefixName)
-						.ToolTipText(LOCTEXT("FbxOptionWindow_Scene_MaterialPrefixName_tooltip", "The prefix must be start and end by '/' use this prefix to add a folder to all material (i.e. /Materials/)"))
-						.OnTextCommitted(this, &SFbxSceneOptionWindow::OnMaterialPrefixCommited)
-						.OnTextChanged(this, &SFbxSceneOptionWindow::OnMaterialPrefixCommited, ETextCommit::Default)
-						.ColorAndOpacity(this, &SFbxSceneOptionWindow::GetMaterialPrefixTextColor)
+						.Text(this, &SFbxSceneOptionWindow::GetMaterialBasePath)
+						.ToolTipText(LOCTEXT("FbxOptionWindow_Scene_MaterialBasePath_tooltip", "The override path must start and end by '/' use this to import all material to a different base path(i.e. /Game/Materials/)"))
+						.OnTextCommitted(this, &SFbxSceneOptionWindow::OnMaterialBasePathCommited)
+						.OnTextChanged(this, &SFbxSceneOptionWindow::OnMaterialBasePathCommited, ETextCommit::Default)
+						.ColorAndOpacity(this, &SFbxSceneOptionWindow::GetMaterialBasePathTextColor)
 					]
 				]
+				+ SHorizontalBox::Slot()
+				.AutoWidth()
+				.VAlign(VAlign_Center)
+				[
+					SNew(SButton)
+					.Text(LOCTEXT("FbxOptionWindow_Scene_Material_Browse", "Browse..."))
+					.ToolTipText(LOCTEXT("FbxOptionWindow_Scene_MaterialBasePath_Browse_tooltip", "Select a path where to save all materials"))
+					.OnClicked(this, &SFbxSceneOptionWindow::OnMaterialBasePathBrowse)
+				]
 			]
 			+ SVerticalBox::Slot()
 			.FillHeight(1.0f)
@@ -359,7 +699,8 @@ TSharedRef<SDockTab> SFbxSceneOptionWindow::SpawnMaterialTab(const FSpawnTabArgs
 					+ SVerticalBox::Slot()
 					.FillHeight(1.0f)
 					[
-						SNew(SSplitter)
+						MaterialsTabListView.ToSharedRef()
+						/*SNew(SSplitter)
 						.Orientation(Orient_Vertical)
 						+ SSplitter::Slot()
 						.Value(0.4f)
@@ -379,11 +720,15 @@ TSharedRef<SDockTab> SFbxSceneOptionWindow::SpawnMaterialTab(const FSpawnTabArgs
 								SNew(STextBlock)
 								.Text(LOCTEXT("FbxOptionWindow_Scene_Texture_Slot", "Put the texture array here"))
 							]
-						]
+						]*/
 					]
 				]
 			]
 		];
+	
+	//Prevent user to close the tab
+	DockTab->SetCanCloseTab(SDockTab::FCanCloseTab::CreateRaw(this, &SFbxSceneOptionWindow::CanCloseTab));
+
 	return DockTab;
 }
 
@@ -396,7 +741,7 @@ TSharedRef<SDockTab> SFbxSceneOptionWindow::SpawnSceneReimportTab(const FSpawnTa
 		.NodeStatusMap(NodeStatusMap);
 
 	TSharedRef<SDockTab> DockTab = SNew(SDockTab)
-		.TabRole(ETabRole::NomadTab)
+		.TabRole(ETabRole::PanelTab)
 		.Label(LOCTEXT("WidgetFbxSceneActorTab", "Scene"))
 		.ToolTipText(LOCTEXT("WidgetFbxSceneActorTabTextToolTip", "Switch to the scene tab."))
 		[
@@ -458,29 +803,91 @@ TSharedRef<SDockTab> SFbxSceneOptionWindow::SpawnSceneReimportTab(const FSpawnTa
 			+ SSplitter::Slot()
 			.Value(0.4f)
 			[
-				SNew(SHorizontalBox)
-				+ SHorizontalBox::Slot()
-				.AutoWidth()
-				.VAlign(VAlign_Top)
+				SNew(SVerticalBox)
+				+ SVerticalBox::Slot()
 				.HAlign(HAlign_Left)
+				.AutoHeight()
 				[
-					SNew(SCheckBox)
-					.HAlign(HAlign_Center)
-					.OnCheckStateChanged(this, &SFbxSceneOptionWindow::OnToggleReimportHierarchy)
-					.IsChecked(this, &SFbxSceneOptionWindow::IsReimportHierarchyChecked)
+					SNew(SHorizontalBox)
+					+ SHorizontalBox::Slot()
+					.AutoWidth()
+					.VAlign(VAlign_Top)
+					.HAlign(HAlign_Left)
+					[
+						SNew(SCheckBox)
+						.HAlign(HAlign_Center)
+						.OnCheckStateChanged(this, &SFbxSceneOptionWindow::OnToggleReimportHierarchy)
+						.IsChecked(this, &SFbxSceneOptionWindow::IsReimportHierarchyChecked)
+					]
+					+ SHorizontalBox::Slot()
+					.FillWidth(1.0f)
+					.Padding(0.0f, 3.0f, 6.0f, 3.0f)
+					.VAlign(VAlign_Top)
+					.HAlign(HAlign_Left)
+					[
+						SNew(STextBlock)
+						.Text(LOCTEXT("FbxOptionWindow_Scene_Reimport_ImportHierarchy", "Reimport Hierarchy"))
+						.ToolTipText(LOCTEXT("FbxOptionWindow_Scene_Reimport_ImportHierarchy_Tooltip", "If Check and the original import was done in a blueprint, the blueprint hierarchy will be revisited to include the fbx changes"))
+					]
 				]
-				+ SHorizontalBox::Slot()
-				.FillWidth(1.0f)
-				.Padding(0.0f, 3.0f, 6.0f, 3.0f)
-				.VAlign(VAlign_Top)
+				+ SVerticalBox::Slot()
 				.HAlign(HAlign_Left)
+				.AutoHeight()
 				[
-					SNew(STextBlock)
-					.Text(LOCTEXT("FbxOptionWindow_Scene_Reimport_ImportHierarchy", "Reimport Hierarchy"))
-					.ToolTipText(LOCTEXT("FbxOptionWindow_Scene_Reimport_ImportHierarchy_Tooltip", "If Check and the original import was done in a blueprint, the blueprint hierarchy will be revisited to include the fbx changes"))
+					SNew(SHorizontalBox)
+					+ SHorizontalBox::Slot()
+					.AutoWidth()
+					.VAlign(VAlign_Top)
+					.HAlign(HAlign_Left)
+					[
+						SNew(SCheckBox)
+						.HAlign(HAlign_Center)
+						.OnCheckStateChanged(this, &SFbxSceneOptionWindow::OnToggleBakePivotInVertex)
+						.IsChecked(this, &SFbxSceneOptionWindow::IsBakePivotInVertexChecked)
+					]
+					+ SHorizontalBox::Slot()
+					.FillWidth(1.0f)
+					.Padding(0.0f, 3.0f, 6.0f, 3.0f)
+					.VAlign(VAlign_Top)
+					.HAlign(HAlign_Left)
+					[
+						SNew(STextBlock)
+						.Text(LOCTEXT("FbxOptionWindow_Scene_Reimport_BakePivotInVertex", "Bake Pivot In Vertex"))
+						.ToolTipText(LOCTEXT("FbxOptionWindow_Scene_Reimport_BakePivotInVertex_Tooltip", "- Experimental - If this option is true the inverse node pivot will be apply to the mesh vertices. The pivot from the DCC will then be the origin of the mesh."))
+					]
+				]
+				+ SVerticalBox::Slot()
+				.HAlign(HAlign_Left)
+				.AutoHeight()
+				[
+					SNew(SHorizontalBox)
+					+ SHorizontalBox::Slot()
+					.AutoWidth()
+					.VAlign(VAlign_Top)
+					.HAlign(HAlign_Left)
+					[
+						SNew(SCheckBox)
+						.HAlign(HAlign_Center)
+						.OnCheckStateChanged(this, &SFbxSceneOptionWindow::OnToggleInvertNormalMap)
+						.IsChecked(this, &SFbxSceneOptionWindow::IsInvertNormalMapChecked)
+					]
+					+ SHorizontalBox::Slot()
+					.FillWidth(1.0f)
+					.Padding(0.0f, 3.0f, 6.0f, 3.0f)
+					.VAlign(VAlign_Top)
+					.HAlign(HAlign_Left)
+					[
+						SNew(STextBlock)
+						.Text(LOCTEXT("FbxOptionWindow_Scene_Reimport_InvertNormalMap", "Invert Normal Map"))
+						.ToolTipText(LOCTEXT("FbxOptionWindow_Scene_Reimport_InvertNormalMap_Tooltip", "If either importing of textures (or materials) is enabled, this option will cause normal map values to be inverted."))
+					]
 				]
 			]
 		];
+
+	//Prevent user to close the tab
+	DockTab->SetCanCloseTab(SDockTab::FCanCloseTab::CreateRaw(this, &SFbxSceneOptionWindow::CanCloseTab));
+
 	return DockTab;
 }
 
@@ -497,6 +904,32 @@ ECheckBoxState SFbxSceneOptionWindow::IsReimportHierarchyChecked() const
 	return  GlobalImportSettings->bImportScene ? ECheckBoxState::Checked : ECheckBoxState::Unchecked;
 }
 
+void SFbxSceneOptionWindow::OnToggleBakePivotInVertex(ECheckBoxState CheckType)
+{
+	if (GlobalImportSettings != nullptr)
+	{
+		GlobalImportSettings->bBakePivotInVertex = CheckType == ECheckBoxState::Checked;
+	}
+}
+
+ECheckBoxState SFbxSceneOptionWindow::IsBakePivotInVertexChecked() const
+{
+	return  GlobalImportSettings->bBakePivotInVertex ? ECheckBoxState::Checked : ECheckBoxState::Unchecked;
+}
+
+void SFbxSceneOptionWindow::OnToggleInvertNormalMap(ECheckBoxState CheckType)
+{
+	if (GlobalImportSettings != nullptr)
+	{
+		GlobalImportSettings->bInvertNormalMap = CheckType == ECheckBoxState::Checked;
+	}
+}
+
+ECheckBoxState SFbxSceneOptionWindow::IsInvertNormalMapChecked() const
+{
+	return  GlobalImportSettings->bInvertNormalMap ? ECheckBoxState::Checked : ECheckBoxState::Unchecked;
+}
+
 TSharedRef<SDockTab> SFbxSceneOptionWindow::SpawnStaticMeshReimportTab(const FSpawnTabArgs& Args)
 {
 	//Create the static mesh listview
@@ -504,14 +937,15 @@ TSharedRef<SDockTab> SFbxSceneOptionWindow::SpawnStaticMeshReimportTab(const FSp
 		.SceneInfo(SceneInfo)
 		.SceneInfoOriginal(SceneInfoOriginal)
 		.GlobalImportSettings(GlobalImportSettings)
+		.OverrideNameOptions(&OverrideNameOptions)
 		.OverrideNameOptionsMap(OverrideNameOptionsMap)
 		.SceneImportOptionsStaticMeshDisplay(SceneImportOptionsStaticMeshDisplay)
 		.MeshStatusMap(MeshStatusMap);
 
 	TSharedPtr<SBox> InspectorBox;
 	TSharedRef<SDockTab> DockTab = SNew(SDockTab)
-		.TabRole(ETabRole::NomadTab)
-		.Label(LOCTEXT("WidgetFbxSceneReimportStaticMeshTab", "Reimport Static Meshes"))
+		.TabRole(ETabRole::PanelTab)
+		.Label(LOCTEXT("WidgetFbxSceneReimportStaticMeshTab", "Static Meshes"))
 		.ToolTipText(LOCTEXT("WidgetFbxSceneReimportStaticMeshTabTextToolTip", "Switch to the reimport static meshes tab."))
 		[
 			SNew(SVerticalBox)
@@ -662,11 +1096,29 @@ TSharedRef<SDockTab> SFbxSceneOptionWindow::SpawnStaticMeshReimportTab(const FSp
 								]
 								+ SHorizontalBox::Slot()
 								.AutoWidth()
+								[
+									SNew(SButton)
+									.HAlign(HAlign_Center)
+									.Text(LOCTEXT("FbxOptionWindow_SM_Select_asset_using", "Select Asset Using"))
+									.OnClicked(StaticMeshReimportListView.Get(), &SFbxSceneStaticMeshReimportListView::OnSelectAssetUsing)
+								]
+								+ SHorizontalBox::Slot()
+								.AutoWidth()
 								[
 									SNew(SButton)
 									.Text(LOCTEXT("FbxOptionWindow_SM_CreateOverride", "Create Override"))
+									.ToolTipText(LOCTEXT("FbxOptionWindow_SM_CreateOverrideTooltip", "Create Override to specify custom import options for some static meshes.\nTo assign options use context menu on static meshes."))
 									.OnClicked(StaticMeshReimportListView.Get(), &SFbxSceneStaticMeshReimportListView::OnCreateOverrideOptions)
 								]
+								+ SHorizontalBox::Slot()
+								.AutoWidth()
+								[
+									SNew(SButton)
+									.HAlign(HAlign_Center)
+									.Text(LOCTEXT("FbxOptionWindow_SM_Delete", "Delete"))
+									.IsEnabled(StaticMeshReimportListView.Get(), &SFbxSceneStaticMeshReimportListView::CanDeleteOverride)
+									.OnClicked(StaticMeshReimportListView.Get(), &SFbxSceneStaticMeshReimportListView::OnDeleteOverride)
+								]
 							]
 							+ SVerticalBox::Slot()
 							.FillHeight(1.0f)
@@ -679,6 +1131,9 @@ TSharedRef<SDockTab> SFbxSceneOptionWindow::SpawnStaticMeshReimportTab(const FSp
 			]
 		];
 	
+	//Prevent user to close the tab
+	DockTab->SetCanCloseTab(SDockTab::FCanCloseTab::CreateRaw(this, &SFbxSceneOptionWindow::CanCloseTab));
+
 	FPropertyEditorModule& PropertyEditorModule = FModuleManager::GetModuleChecked<FPropertyEditorModule>("PropertyEditor");
 	FDetailsViewArgs DetailsViewArgs;
 	DetailsViewArgs.bAllowSearch = false;
@@ -714,14 +1169,14 @@ void SFbxSceneOptionWindow::InitAllTabs()
 					FTabManager::NewStack()
 					->AddTab("Scene", ETabState::OpenedTab)
 					->AddTab("StaticMeshes", ETabState::OpenedTab)
+					->AddTab("SkeletalMeshes", ETabState::OpenedTab)
 					->AddTab("Materials", ETabState::OpenedTab)
 				)
 			);
-		//					->AddTab("SkeletalMeshes", ETabState::OpenedTab)
-							
+
 		FbxSceneImportTabManager->RegisterTabSpawner("Scene", FOnSpawnTab::CreateSP(this, &SFbxSceneOptionWindow::SpawnSceneTab));
 		FbxSceneImportTabManager->RegisterTabSpawner("StaticMeshes", FOnSpawnTab::CreateSP(this, &SFbxSceneOptionWindow::SpawnStaticMeshTab));
-		//FbxSceneImportTabManager->RegisterTabSpawner("SkeletalMeshes", FOnSpawnTab::CreateSP(this, &SFbxSceneOptionWindow::SpawnSkeletalMeshTab));
+		FbxSceneImportTabManager->RegisterTabSpawner("SkeletalMeshes", FOnSpawnTab::CreateSP(this, &SFbxSceneOptionWindow::SpawnSkeletalMeshTab));
 		FbxSceneImportTabManager->RegisterTabSpawner("Materials", FOnSpawnTab::CreateSP(this, &SFbxSceneOptionWindow::SpawnMaterialTab));
 	}
 	else
@@ -737,6 +1192,7 @@ void SFbxSceneOptionWindow::InitAllTabs()
 						FTabManager::NewStack()
 						->AddTab("SceneReImport", ETabState::OpenedTab)
 						->AddTab("StaticMeshesReImport", ETabState::OpenedTab)
+						->AddTab("SkeletalMeshesReImport", ETabState::OpenedTab)
 						->AddTab("Materials", ETabState::OpenedTab)
 					)
 				);
@@ -753,13 +1209,18 @@ void SFbxSceneOptionWindow::InitAllTabs()
 					(
 						FTabManager::NewStack()
 						->AddTab("StaticMeshesReImport", ETabState::OpenedTab)
+						->AddTab("SkeletalMeshesReImport", ETabState::OpenedTab)
 						->AddTab("Materials", ETabState::OpenedTab)
 					)
 				);
 		}
 		FbxSceneImportTabManager->RegisterTabSpawner("StaticMeshesReimport", FOnSpawnTab::CreateSP(this, &SFbxSceneOptionWindow::SpawnStaticMeshReimportTab));
+		FbxSceneImportTabManager->RegisterTabSpawner("SkeletalMeshesReimport", FOnSpawnTab::CreateSP(this, &SFbxSceneOptionWindow::SpawnSkeletalMeshReimportTab));
 		FbxSceneImportTabManager->RegisterTabSpawner("Materials", FOnSpawnTab::CreateSP(this, &SFbxSceneOptionWindow::SpawnMaterialTab));
 	}
+
+	//Prevent Docking the tab outside of the dialog well
+	FbxSceneImportTabManager->SetCanDoDragOperation(false);
 }
 
 void SFbxSceneOptionWindow::Construct(const FArguments& InArgs)
@@ -774,8 +1235,6 @@ void SFbxSceneOptionWindow::Construct(const FArguments& InArgs)
 	SceneImportOptionsStaticMeshDisplay = InArgs._SceneImportOptionsStaticMeshDisplay;
 	OverrideNameOptionsMap = InArgs._OverrideNameOptionsMap;
 	SceneImportOptionsSkeletalMeshDisplay = InArgs._SceneImportOptionsSkeletalMeshDisplay;
-	SceneImportOptionsAnimationDisplay = InArgs._SceneImportOptionsAnimationDisplay;
-	SceneImportOptionsMaterialDisplay = InArgs._SceneImportOptionsMaterialDisplay;
 	OwnerWindow = InArgs._OwnerWindow;
 	FullPath = InArgs._FullPath;
 
@@ -783,15 +1242,10 @@ void SFbxSceneOptionWindow::Construct(const FArguments& InArgs)
 	check(GlobalImportSettings != nullptr);
 	check(SceneImportOptionsDisplay != nullptr);
 	check(SceneImportOptionsStaticMeshDisplay != nullptr);
+	check(SceneImportOptionsSkeletalMeshDisplay != nullptr);
 	check(OverrideNameOptionsMap != nullptr);
 
-	if (!SceneInfoOriginal.IsValid())
-	{
-		check(SceneImportOptionsSkeletalMeshDisplay != nullptr);
-		check(SceneImportOptionsAnimationDisplay != nullptr);
-		check(SceneImportOptionsMaterialDisplay != nullptr);
-	}
-	else
+	if (SceneInfoOriginal.IsValid())
 	{
 		check(MeshStatusMap != nullptr);
 		check(NodeStatusMap);
@@ -799,80 +1253,99 @@ void SFbxSceneOptionWindow::Construct(const FArguments& InArgs)
 
 	check(OwnerWindow.IsValid());
 
-	MaterialPrefixName = GlobalImportSettings->MaterialPrefixName.ToString();
+	MaterialBasePath = GlobalImportSettings->MaterialBasePath == NAME_None ? TEXT("") : GlobalImportSettings->MaterialBasePath.ToString();
 
 	InitAllTabs();
 
-	FText SubmitText = SceneInfoOriginal.IsValid() ? LOCTEXT("FbxOptionWindow_Import", "Reimport") : LOCTEXT("FbxOptionWindow_Import", "Import");
+	FText SubmitText = SceneInfoOriginal.IsValid() ? LOCTEXT("FbxSceneOptionWindow_ReImport", "Reimport") : LOCTEXT("FbxSceneOptionWindow_Import", "Import");
 
 	this->ChildSlot
 	[
-		SNew(SVerticalBox)
-		+SVerticalBox::Slot()
-		.AutoHeight()
-		.Padding(2)
+		SNew(SBorder)
+		.Padding(FMargin(10.0f, 3.0f))
+		.BorderImage(FEditorStyle::GetBrush("ToolPanel.DarkGroupBorder"))
 		[
-			SNew(SBorder)
-			.Padding(FMargin(3))
-			.BorderImage(FEditorStyle::GetBrush("ToolPanel.GroupBorder"))
+			SNew(SVerticalBox)
+			+SVerticalBox::Slot()
+			.AutoHeight()
+			.Padding(2)
 			[
-				SNew(SHorizontalBox)
-				+SHorizontalBox::Slot()
-				.AutoWidth()
-				[
-					SNew(STextBlock)
-					.Font(FEditorStyle::GetFontStyle("CurveEd.LabelFont"))
-					.Text(LOCTEXT("FbxSceneImport_CurrentPath", "Import Asset Path: "))
-				]
-				+SHorizontalBox::Slot()
-				.Padding(5, 0, 0, 0)
-				.AutoWidth()
-				.VAlign(VAlign_Center)
+				SNew(SBorder)
+				.Padding(FMargin(3))
+				.BorderImage(FEditorStyle::GetBrush("ToolPanel.GroupBorder"))
 				[
-					SNew(STextBlock)
-					.Font(FEditorStyle::GetFontStyle("CurveEd.InfoFont"))
-					.Text(FText::FromString(FullPath))
+					SNew(SHorizontalBox)
+					+SHorizontalBox::Slot()
+					.AutoWidth()
+					[
+						SNew(STextBlock)
+						.Font(FEditorStyle::GetFontStyle("CurveEd.LabelFont"))
+						.Text(LOCTEXT("FbxSceneImport_CurrentPath", "Import Asset Path: "))
+					]
+					+SHorizontalBox::Slot()
+					.Padding(5, 0, 0, 0)
+					.AutoWidth()
+					.VAlign(VAlign_Center)
+					[
+						SNew(STextBlock)
+						.Font(FEditorStyle::GetFontStyle("CurveEd.InfoFont"))
+						.Text(FText::FromString(FullPath))
+					]
 				]
 			]
-		]
-		+ SVerticalBox::Slot()
-		.FillHeight(1.0f)
-		.Padding(2)
-		[
-			SpawnDockTab().ToSharedRef()
-		]
-		+ SVerticalBox::Slot()
-		.AutoHeight()
-		.HAlign(HAlign_Right)
-		.Padding(2)
-		[
-			SNew(SUniformGridPanel)
-			.SlotPadding(2)
-			+ SUniformGridPanel::Slot(0, 0)
-			[
-				IDocumentation::Get()->CreateAnchor(FString("Engine/Content/FBX/ImportOptions"))
-			]
-			+ SUniformGridPanel::Slot(1, 0)
+			+ SVerticalBox::Slot()
+			.FillHeight(1.0f)
+			.Padding(2)
 			[
-				SNew(SButton)
-				.HAlign(HAlign_Center)
-				.Text(SubmitText)
-				.IsEnabled(this, &SFbxSceneOptionWindow::CanImport)
-				.OnClicked(this, &SFbxSceneOptionWindow::OnImport)
+				SpawnDockTab().ToSharedRef()
 			]
-			+ SUniformGridPanel::Slot(2, 0)
+			+ SVerticalBox::Slot()
+			.AutoHeight()
+			.HAlign(HAlign_Right)
+			.Padding(2)
 			[
-				SNew(SButton)
-				.HAlign(HAlign_Center)
-				.Text(LOCTEXT("FbxOptionWindow_Cancel", "Cancel"))
-				.ToolTipText(LOCTEXT("FbxOptionWindow_Cancel_ToolTip", "Cancels importing this FBX file"))
-				.OnClicked(this, &SFbxSceneOptionWindow::OnCancel)
+				SNew(SUniformGridPanel)
+				.SlotPadding(2)
+				+ SUniformGridPanel::Slot(0, 0)
+				[
+					IDocumentation::Get()->CreateAnchor(FString("Engine/Content/FBX/ImportOptions"))
+				]
+				+ SUniformGridPanel::Slot(1, 0)
+				[
+					SNew(SButton)
+					.HAlign(HAlign_Center)
+					.Text(SubmitText)
+					.IsEnabled(this, &SFbxSceneOptionWindow::CanImport)
+					.OnClicked(this, &SFbxSceneOptionWindow::OnImport)
+				]
+				+ SUniformGridPanel::Slot(2, 0)
+				[
+					SNew(SButton)
+					.HAlign(HAlign_Center)
+					.Text(LOCTEXT("FbxOptionWindow_Cancel", "Cancel"))
+					.ToolTipText(LOCTEXT("FbxOptionWindow_Cancel_ToolTip", "Cancels importing this FBX file"))
+					.OnClicked(this, &SFbxSceneOptionWindow::OnCancel)
+				]
 			]
 		]
 	];
 
 	//By default we want to see the Scene tab
-	FbxSceneImportTabManager->InvokeTab(FTabId("Scene"));
+	if (!SceneInfoOriginal.IsValid())
+	{
+		FbxSceneImportTabManager->InvokeTab(FTabId("Scene"));
+	}
+	else
+	{
+		if (bCanReimportHierarchy)
+		{
+			FbxSceneImportTabManager->InvokeTab(FTabId("SceneReImport"));
+		}
+		else
+		{
+			FbxSceneImportTabManager->InvokeTab(FTabId("StaticMeshesReimport"));
+		}
+	}
 }
 
 void SFbxSceneOptionWindow::CloseFbxSceneOption()
@@ -894,6 +1367,8 @@ void SFbxSceneOptionWindow::CloseFbxSceneOption()
 	StaticMeshTabListView = nullptr;
 	StaticMeshTabDetailsView = nullptr;
 
+	SkeletalMeshTabListView = nullptr;
+
 	SceneReimportTreeview = nullptr;
 
 	StaticMeshReimportListView = nullptr;
@@ -901,7 +1376,7 @@ void SFbxSceneOptionWindow::CloseFbxSceneOption()
 
 	MaterialsTabListView = nullptr;
 	TexturesArray.Reset();
-	MaterialPrefixName.Empty();
+	MaterialBasePath.Empty();
 
 	SceneInfo = nullptr;
 	SceneInfoOriginal = nullptr;
@@ -911,8 +1386,6 @@ void SFbxSceneOptionWindow::CloseFbxSceneOption()
 	SceneImportOptionsStaticMeshDisplay = nullptr;
 	OverrideNameOptionsMap = nullptr;
 	SceneImportOptionsSkeletalMeshDisplay = nullptr;
-	SceneImportOptionsAnimationDisplay = nullptr;
-	SceneImportOptionsMaterialDisplay = nullptr;
 
 	MeshStatusMap = nullptr;
 	NodeStatusMap = nullptr;
@@ -943,7 +1416,6 @@ void SFbxSceneOptionWindow::CopyStaticMeshOptionsToFbxOptions(UnFbx::FBXImportOp
 	ImportSettings->bGenerateLightmapUVs = StaticMeshOptions->bGenerateLightmapUVs;
 	ImportSettings->bOneConvexHullPerUCX = StaticMeshOptions->bOneConvexHullPerUCX;
 	ImportSettings->bRemoveDegenerates = StaticMeshOptions->bRemoveDegenerates;
-	ImportSettings->bTransformVertexToAbsolute = StaticMeshOptions->bTransformVertexToAbsolute;
 	ImportSettings->StaticMeshLODGroup = StaticMeshOptions->StaticMeshLODGroup;
 	switch (StaticMeshOptions->VertexColorImportOption)
 	{
@@ -991,7 +1463,6 @@ void SFbxSceneOptionWindow::CopyFbxOptionsToStaticMeshOptions(UnFbx::FBXImportOp
 	StaticMeshOptions->bGenerateLightmapUVs = ImportSettings->bGenerateLightmapUVs;
 	StaticMeshOptions->bOneConvexHullPerUCX = ImportSettings->bOneConvexHullPerUCX;
 	StaticMeshOptions->bRemoveDegenerates = ImportSettings->bRemoveDegenerates;
-	StaticMeshOptions->bTransformVertexToAbsolute = ImportSettings->bTransformVertexToAbsolute;
 	StaticMeshOptions->StaticMeshLODGroup = ImportSettings->StaticMeshLODGroup;
 	switch (ImportSettings->VertexColorImportOption)
 	{
@@ -1033,26 +1504,42 @@ void SFbxSceneOptionWindow::CopyFbxOptionsToStaticMeshOptions(UnFbx::FBXImportOp
 
 void SFbxSceneOptionWindow::CopySkeletalMeshOptionsToFbxOptions(UnFbx::FBXImportOptions *ImportSettings, UFbxSceneImportOptionsSkeletalMesh* SkeletalMeshOptions)
 {
-}
-
-void SFbxSceneOptionWindow::CopyFbxOptionsToSkeletalMeshOptions(UnFbx::FBXImportOptions *ImportSettings, class UFbxSceneImportOptionsSkeletalMesh* SkeletalMeshOptions)
-{
-}
-
-void SFbxSceneOptionWindow::CopyAnimationOptionsToFbxOptions(UnFbx::FBXImportOptions *ImportSettings, UFbxSceneImportOptionsAnimation* AnimationOptions)
-{
-}
+	ImportSettings->bCreatePhysicsAsset = SkeletalMeshOptions->bCreatePhysicsAsset;
+	ImportSettings->bImportMeshesInBoneHierarchy = SkeletalMeshOptions->bImportMeshesInBoneHierarchy;
+	ImportSettings->bImportMorph = SkeletalMeshOptions->bImportMorphTargets;
+	ImportSettings->bKeepOverlappingVertices = SkeletalMeshOptions->bKeepOverlappingVertices;
+	ImportSettings->bPreserveSmoothingGroups = SkeletalMeshOptions->bPreserveSmoothingGroups;
+	ImportSettings->bUpdateSkeletonReferencePose = SkeletalMeshOptions->bUpdateSkeletonReferencePose;
+	ImportSettings->bUseT0AsRefPose = SkeletalMeshOptions->bUseT0AsRefPose;
 
-void SFbxSceneOptionWindow::CopyFbxOptionsToAnimationOptions(UnFbx::FBXImportOptions *ImportSettings, class UFbxSceneImportOptionsAnimation* AnimationOptions)
-{
+	ImportSettings->bImportAnimations = SkeletalMeshOptions->bImportAnimations;
+	ImportSettings->AnimationLengthImportType = SkeletalMeshOptions->AnimationLength;
+	ImportSettings->bDeleteExistingMorphTargetCurves = SkeletalMeshOptions->bDeleteExistingMorphTargetCurves;
+	ImportSettings->bImportCustomAttribute = SkeletalMeshOptions->bImportCustomAttribute;
+	ImportSettings->bPreserveLocalTransform = SkeletalMeshOptions->bPreserveLocalTransform;
+	ImportSettings->bResample = SkeletalMeshOptions->bUseDefaultSampleRate;
+	ImportSettings->AnimationRange.X = SkeletalMeshOptions->FrameImportRange.Min;
+	ImportSettings->AnimationRange.Y = SkeletalMeshOptions->FrameImportRange.Max;
 }
 
-void SFbxSceneOptionWindow::CopyMaterialOptionsToFbxOptions(UnFbx::FBXImportOptions *ImportSettings, UFbxSceneImportOptionsMaterial* MaterialOptions)
+void SFbxSceneOptionWindow::CopyFbxOptionsToSkeletalMeshOptions(UnFbx::FBXImportOptions *ImportSettings, class UFbxSceneImportOptionsSkeletalMesh* SkeletalMeshOptions)
 {
-}
+	SkeletalMeshOptions->bCreatePhysicsAsset = ImportSettings->bCreatePhysicsAsset;
+	SkeletalMeshOptions->bImportMeshesInBoneHierarchy = ImportSettings->bImportMeshesInBoneHierarchy;
+	SkeletalMeshOptions->bImportMorphTargets = ImportSettings->bImportMorph;
+	SkeletalMeshOptions->bKeepOverlappingVertices = ImportSettings->bKeepOverlappingVertices;
+	SkeletalMeshOptions->bPreserveSmoothingGroups = ImportSettings->bPreserveSmoothingGroups;
+	SkeletalMeshOptions->bUpdateSkeletonReferencePose = ImportSettings->bUpdateSkeletonReferencePose;
+	SkeletalMeshOptions->bUseT0AsRefPose = ImportSettings->bUseT0AsRefPose;
 
-void SFbxSceneOptionWindow::CopyFbxOptionsToMaterialOptions(UnFbx::FBXImportOptions *ImportSettings, class UFbxSceneImportOptionsMaterial* MaterialOptions)
-{
+	SkeletalMeshOptions->bImportAnimations = ImportSettings->bImportAnimations;
+	SkeletalMeshOptions->AnimationLength = ImportSettings->AnimationLengthImportType;
+	SkeletalMeshOptions->bDeleteExistingMorphTargetCurves = ImportSettings->bDeleteExistingMorphTargetCurves;
+	SkeletalMeshOptions->bImportCustomAttribute = ImportSettings->bImportCustomAttribute;
+	SkeletalMeshOptions->bPreserveLocalTransform = ImportSettings->bPreserveLocalTransform;
+	SkeletalMeshOptions->bUseDefaultSampleRate = ImportSettings->bResample;
+	SkeletalMeshOptions->FrameImportRange.Min = ImportSettings->AnimationRange.X;
+	SkeletalMeshOptions->FrameImportRange.Max = ImportSettings->AnimationRange.Y;
 }
 
 #undef LOCTEXT_NAMESPACE
diff --git a/Engine/Source/Editor/UnrealEd/Private/Fbx/SSceneImportNodeTreeView.cpp b/Engine/Source/Editor/UnrealEd/Private/Fbx/SSceneImportNodeTreeView.cpp
--- a/Engine/Source/Editor/UnrealEd/Private/Fbx/SSceneImportNodeTreeView.cpp
+++ b/Engine/Source/Editor/UnrealEd/Private/Fbx/SSceneImportNodeTreeView.cpp
@@ -45,11 +45,13 @@ class SFbxSceneTreeViewItem : public STableRow< FbxNodeInfoPtr >
 public:
 
 	SLATE_BEGIN_ARGS(SFbxSceneTreeViewItem)
-		: _FbxNodeInfo(NULL)
+		: _FbxNodeInfo(nullptr)
+		, _SceneInfo(nullptr)
 	{}
 
 	/** The item content. */
 	SLATE_ARGUMENT(FbxNodeInfoPtr, FbxNodeInfo)
+	SLATE_ARGUMENT(TSharedPtr<FFbxSceneInfo>, SceneInfo)
 	SLATE_END_ARGS()
 
 	/**
@@ -60,21 +62,56 @@ public:
 	void Construct(const FArguments& InArgs, const TSharedRef<STableViewBase>& InOwnerTableView)
 	{
 		FbxNodeInfo = InArgs._FbxNodeInfo;
+		SceneInfo = InArgs._SceneInfo;
 		
 		//This is suppose to always be valid
 		check(FbxNodeInfo.IsValid());
+		check(SceneInfo.IsValid());
+
 		UClass *IconClass = AActor::StaticClass();
 		if (FbxNodeInfo->AttributeInfo.IsValid())
 		{
 			IconClass = FbxNodeInfo->AttributeInfo->GetType();
 		}
+		else if (FbxNodeInfo->AttributeType.Compare(TEXT("eLight")) == 0)
+		{
+			IconClass = ULightComponent::StaticClass();
+			if (SceneInfo->LightInfo.Contains(FbxNodeInfo->AttributeUniqueId))
+			{
+				TSharedPtr<FFbxLightInfo> LightInfo = *SceneInfo->LightInfo.Find(FbxNodeInfo->AttributeUniqueId);
+				if (LightInfo->Type == 0)
+				{
+					IconClass = UPointLightComponent::StaticClass();
+				}
+				else if (LightInfo->Type == 1)
+				{
+					IconClass = UDirectionalLightComponent::StaticClass();
+				}
+				else if (LightInfo->Type == 2)
+				{
+					IconClass = USpotLightComponent::StaticClass();
+				}
+			}
+		}
+		else if (FbxNodeInfo->AttributeType.Compare(TEXT("eCamera")) == 0)
+		{
+			IconClass = UCameraComponent::StaticClass();
+		}
 		const FSlateBrush* ClassIcon = FClassIconFinder::FindIconForClass(IconClass);
 
+		//Prepare the tooltip
+		FString Tooltip = FbxNodeInfo->NodeName;
+		if (!FbxNodeInfo->AttributeType.IsEmpty() && FbxNodeInfo->AttributeType.Compare(TEXT("eNull")) != 0)
+		{
+			Tooltip += TEXT(" [") + FbxNodeInfo->AttributeType + TEXT("]");
+		}
+
 		this->ChildSlot
 		[
 		SNew(SHorizontalBox)
 		+ SHorizontalBox::Slot()
 			.AutoWidth()
+			.Padding(2.0f, 0.0f, 2.0f, 0.0f)
 			[
 				SNew(SCheckBox)
 				.OnCheckStateChanged(this, &SFbxSceneTreeViewItem::OnItemCheckChanged)
@@ -102,7 +139,7 @@ public:
 			[
 				SNew(STextBlock)
 				.Text(FText::FromString(FbxNodeInfo->NodeName))
-				.ToolTipText(FText::FromString(FbxNodeInfo->NodeName))
+				.ToolTipText(FText::FromString(Tooltip))
 			]
 
 		];
@@ -116,11 +153,43 @@ public:
 
 private:
 
+	void RecursivelySetLODMeshImportState(TSharedPtr<FFbxNodeInfo> NodeInfo, bool bState)
+	{
+		//Set the bImportNode property for all mesh under eLODGroup
+		for (TSharedPtr<FFbxNodeInfo> ChildNodeInfo : NodeInfo->Childrens)
+		{
+			if (!ChildNodeInfo.IsValid())
+				continue;
+			if (ChildNodeInfo->AttributeType.Compare(TEXT("eMesh")) == 0)
+			{
+				ChildNodeInfo->bImportNode = bState;
+			}
+			else
+			{
+				RecursivelySetLODMeshImportState(ChildNodeInfo, bState);
+			}
+		}
+	}
+
 	void OnItemCheckChanged(ECheckBoxState CheckType)
 	{
 		if (!FbxNodeInfo.IsValid())
 			return;
 		FbxNodeInfo->bImportNode = CheckType == ECheckBoxState::Checked;
+		if (FbxNodeInfo->AttributeType.Compare(TEXT("eLODGroup")) == 0)
+		{
+			RecursivelySetLODMeshImportState(FbxNodeInfo, FbxNodeInfo->bImportNode);
+		}
+		if (FbxNodeInfo->AttributeType.Compare(TEXT("eMesh")) == 0)
+		{
+			//Verify if parent is a LOD group
+			TSharedPtr<FFbxNodeInfo> ParentLODNodeInfo = FFbxSceneInfo::RecursiveFindLODParentNode(FbxNodeInfo);
+			if (ParentLODNodeInfo.IsValid())
+			{
+				ParentLODNodeInfo->bImportNode = FbxNodeInfo->bImportNode;
+				RecursivelySetLODMeshImportState(ParentLODNodeInfo, FbxNodeInfo->bImportNode);
+			}
+		}
 	}
 
 	ECheckBoxState IsItemChecked() const
@@ -130,22 +199,41 @@ private:
 
 	/** The node info to build the tree view row from. */
 	FbxNodeInfoPtr FbxNodeInfo;
+	TSharedPtr<FFbxSceneInfo> SceneInfo;
 };
 
 TSharedRef< ITableRow > SFbxSceneTreeView::OnGenerateRowFbxSceneTreeView(FbxNodeInfoPtr Item, const TSharedRef< STableViewBase >& OwnerTable)
 {
 	TSharedRef< SFbxSceneTreeViewItem > ReturnRow = SNew(SFbxSceneTreeViewItem, OwnerTable)
-		.FbxNodeInfo(Item);
+		.FbxNodeInfo(Item)
+		.SceneInfo(SceneInfo);
 	return ReturnRow;
 }
 
 void SFbxSceneTreeView::OnGetChildrenFbxSceneTreeView(FbxNodeInfoPtr InParent, TArray< FbxNodeInfoPtr >& OutChildren)
 {
-	for (auto Child : InParent->Childrens)
+	if (InParent->AttributeType.Compare(TEXT("eLODGroup")) == 0 && InParent->Childrens.Num() > 0)
+	{
+		TSharedPtr<FFbxNodeInfo> Child = InParent;
+		do
+		{
+			Child = Child->Childrens[0];
+			if (Child.IsValid() && Child->AttributeType.Compare(TEXT("eMesh")) == 0)
+			{
+				OutChildren.Add(Child);
+				return;
+			}
+		} while (Child->Childrens.Num() > 0);
+		//We do not found any LOD mesh don't add any child
+	}
+	else
 	{
-		if (Child.IsValid())
+		for (auto Child : InParent->Childrens)
 		{
-			OutChildren.Add(Child);
+			if (Child.IsValid())
+			{
+				OutChildren.Add(Child);
+			}
 		}
 	}
 }
diff --git a/Engine/Source/Editor/UnrealEd/Private/Fbx/SSceneImportStaticMeshListView.cpp b/Engine/Source/Editor/UnrealEd/Private/Fbx/SSceneImportStaticMeshListView.cpp
--- a/Engine/Source/Editor/UnrealEd/Private/Fbx/SSceneImportStaticMeshListView.cpp
+++ b/Engine/Source/Editor/UnrealEd/Private/Fbx/SSceneImportStaticMeshListView.cpp
@@ -8,14 +8,18 @@
 #include "SFbxSceneOptionWindow.h"
 #include "FbxImporter.h"
 #include "STextComboBox.h"
+#include "STextEntryPopup.h"
 
 #define LOCTEXT_NAMESPACE "SFbxSceneStaticMeshListView"
 
-const FName SceneImportCheckBoxSelectionHeaderIdName(TEXT("CheckBoxSelectionHeaderId"));
-const FName SceneImportClassIconHeaderIdName(TEXT("ClassIconHeaderId"));
-const FName SceneImportAssetNameHeaderIdName(TEXT("AssetNameHeaderId"));
-const FName SceneImportContentPathHeaderIdName(TEXT("ContentPathHeaderId"));
-const FName SceneImportOptionsNameHeaderIdName(TEXT("OptionsNameHeaderId"));
+namespace FbxSceneImportStaticMesh
+{
+	const FName SceneImportCheckBoxSelectionHeaderIdName(TEXT("CheckBoxSelectionHeaderId"));
+	const FName SceneImportClassIconHeaderIdName(TEXT("ClassIconHeaderId"));
+	const FName SceneImportAssetNameHeaderIdName(TEXT("AssetNameHeaderId"));
+	const FName SceneImportContentPathHeaderIdName(TEXT("ContentPathHeaderId"));
+	const FName SceneImportOptionsNameHeaderIdName(TEXT("OptionsNameHeaderId"));
+}
 
 /** The item used for visualizing the class in the tree. */
 class SFbxMeshItemTableListViewRow : public SMultiColumnTableRow<FbxMeshInfoPtr>
@@ -24,10 +28,12 @@ public:
 
 	SLATE_BEGIN_ARGS(SFbxMeshItemTableListViewRow)
 		: _FbxMeshInfo(nullptr)
+		, _GlobalImportSettings(nullptr)
 	{}
 
 	/** The item content. */
 	SLATE_ARGUMENT(FbxMeshInfoPtr, FbxMeshInfo)
+	SLATE_ARGUMENT(UnFbx::FBXImportOptions*, GlobalImportSettings)
 	SLATE_END_ARGS()
 
 	/**
@@ -38,9 +44,11 @@ public:
 	void Construct(const FArguments& InArgs, const TSharedRef<STableViewBase>& InOwnerTableView)
 	{
 		FbxMeshInfo = InArgs._FbxMeshInfo;
+		GlobalImportSettings = InArgs._GlobalImportSettings;
 
 		//This is suppose to always be valid
 		check(FbxMeshInfo.IsValid());
+		check(GlobalImportSettings != nullptr);
 
 		SMultiColumnTableRow<FbxMeshInfoPtr>::Construct(
 			FSuperRowType::FArguments()
@@ -52,13 +60,17 @@ public:
 	/** Overridden from SMultiColumnTableRow.  Generates a widget for this column of the list view. */
 	virtual TSharedRef<SWidget> GenerateWidgetForColumn(const FName& ColumnName) override
 	{
-		if (ColumnName == SceneImportCheckBoxSelectionHeaderIdName)
+		if (ColumnName == FbxSceneImportStaticMesh::SceneImportCheckBoxSelectionHeaderIdName)
 		{
-			return SNew(SCheckBox)
-				.OnCheckStateChanged(this, &SFbxMeshItemTableListViewRow::OnItemCheckChanged)
-				.IsChecked(this, &SFbxMeshItemTableListViewRow::IsItemChecked);
+			return SNew(SBox)
+				.HAlign(HAlign_Center)
+				[
+					SNew(SCheckBox)
+					.OnCheckStateChanged(this, &SFbxMeshItemTableListViewRow::OnItemCheckChanged)
+					.IsChecked(this, &SFbxMeshItemTableListViewRow::IsItemChecked)
+				];
 		}
-		else if (ColumnName == SceneImportClassIconHeaderIdName)
+		else if (ColumnName == FbxSceneImportStaticMesh::SceneImportClassIconHeaderIdName)
 		{
 			UClass *IconClass = FbxMeshInfo->GetType();
 			const FSlateBrush* ClassIcon = FClassIconFinder::FindIconForClass(IconClass);
@@ -82,25 +94,31 @@ public:
 				];
 			return IconContent;
 		}
-		else if (ColumnName == SceneImportAssetNameHeaderIdName)
+		else if (ColumnName == FbxSceneImportStaticMesh::SceneImportAssetNameHeaderIdName)
 		{
 			return SNew(STextBlock)
 				.Text(FText::FromString(FbxMeshInfo->Name))
 				.ToolTipText(FText::FromString(FbxMeshInfo->Name));
 		}
-		else if (ColumnName == SceneImportContentPathHeaderIdName)
+		else if (ColumnName == FbxSceneImportStaticMesh::SceneImportContentPathHeaderIdName)
 		{
 			return SNew(STextBlock)
 				.Text(this, &SFbxMeshItemTableListViewRow::GetAssetFullName)
 				.ColorAndOpacity(this, &SFbxMeshItemTableListViewRow::GetContentPathTextColor)
 				.ToolTipText(this, &SFbxMeshItemTableListViewRow::GetAssetFullName);
 		}
-		else if (ColumnName == SceneImportOptionsNameHeaderIdName)
+		else if (ColumnName == FbxSceneImportStaticMesh::SceneImportOptionsNameHeaderIdName)
 		{
 			return SNew(STextBlock)
 				.Text(this, &SFbxMeshItemTableListViewRow::GetAssetOptionName)
 				.ToolTipText(this, &SFbxMeshItemTableListViewRow::GetAssetOptionName);
 		}
+		else if (ColumnName == FbxSceneBaseListViewColumn::PivotColumnId)
+		{
+			return SNew(STextBlock)
+				.Text(this, &SFbxMeshItemTableListViewRow::GetAssetPivotNodeName)
+				.ToolTipText(this, &SFbxMeshItemTableListViewRow::GetAssetPivotNodeName);
+		}
 
 		return SNullWidget::NullWidget;
 	}
@@ -140,11 +158,17 @@ private:
 
 	FText GetAssetOptionName() const
 	{
-			return FText::FromString(FbxMeshInfo->OptionName);
+		return FText::FromString(FbxMeshInfo->OptionName);
+	}
+
+	FText GetAssetPivotNodeName() const
+	{
+		return GlobalImportSettings->bBakePivotInVertex ? FText::FromString(FbxMeshInfo->PivotNodeName) : FText::FromString(TEXT("-"));
 	}
 
 	/** The node info to build the tree view row from. */
 	FbxMeshInfoPtr FbxMeshInfo;
+	UnFbx::FBXImportOptions *GlobalImportSettings;
 };
 
 //////////////////////////////////////////////////////////////////////////
@@ -155,11 +179,11 @@ SFbxSceneStaticMeshListView::~SFbxSceneStaticMeshListView()
 	SceneInfo = nullptr;
 	GlobalImportSettings = nullptr;
 	SceneImportOptionsStaticMeshDisplay = nullptr;
-	CurrentStaticMeshImportOptions = nullptr;
+	CurrentMeshImportOptions = nullptr;
 	FbxMeshesArray.Empty();
 	SceneImportOptionsStaticMeshDisplay = nullptr;
 
-	OverrideNameOptions.Empty();
+	OverrideNameOptions = nullptr;
 	OverrideNameOptionsMap = nullptr;
 	OptionComboBox = nullptr;
 	DefaultOptionNamePtr = nullptr;
@@ -169,11 +193,13 @@ void SFbxSceneStaticMeshListView::Construct(const SFbxSceneStaticMeshListView::F
 {
 	SceneInfo = InArgs._SceneInfo;
 	GlobalImportSettings = InArgs._GlobalImportSettings;
+	OverrideNameOptions = InArgs._OverrideNameOptions;
 	OverrideNameOptionsMap = InArgs._OverrideNameOptionsMap;
 	SceneImportOptionsStaticMeshDisplay = InArgs._SceneImportOptionsStaticMeshDisplay;
 
 	check(SceneInfo.IsValid());
 	check(GlobalImportSettings != nullptr);
+	check(OverrideNameOptions != nullptr);
 	check(OverrideNameOptionsMap != nullptr);
 	check(SceneImportOptionsStaticMeshDisplay != nullptr);
 
@@ -181,17 +207,31 @@ void SFbxSceneStaticMeshListView::Construct(const SFbxSceneStaticMeshListView::F
 	//Set the default options to the current global import settings
 	GlobalImportSettings->bTransformVertexToAbsolute = false;
 	GlobalImportSettings->StaticMeshLODGroup = NAME_None;
-	CurrentStaticMeshImportOptions = GlobalImportSettings;
+	CurrentMeshImportOptions = GlobalImportSettings;
 	
-	DefaultOptionNamePtr = MakeShareable(new FString(UFbxSceneImportFactory::DefaultOptionName));
-	OverrideNameOptions.Add(DefaultOptionNamePtr);
-	OverrideNameOptionsMap->Add(UFbxSceneImportFactory::DefaultOptionName, GlobalImportSettings);
+	bool bNameExist = false;
+	for (TSharedPtr<FString> OverrideName : (*OverrideNameOptions))
+	{
+		if (OverrideName.Get()->Compare(UFbxSceneImportFactory::DefaultOptionName) == 0)
+		{
+			DefaultOptionNamePtr = OverrideName;
+			bNameExist = true;
+			break;
+		}
+	}
+
+	if (!bNameExist)
+	{
+		DefaultOptionNamePtr = MakeShareable(new FString(UFbxSceneImportFactory::DefaultOptionName));
+		OverrideNameOptions->Add(DefaultOptionNamePtr);
+		OverrideNameOptionsMap->Add(UFbxSceneImportFactory::DefaultOptionName, GlobalImportSettings);
+	}
 	
 	for (auto MeshInfoIt = SceneInfo->MeshInfo.CreateIterator(); MeshInfoIt; ++MeshInfoIt)
 	{
 		FbxMeshInfoPtr MeshInfo = (*MeshInfoIt);
 
-		if (!MeshInfo->bIsSkelMesh)
+		if (!MeshInfo->bIsSkelMesh && !MeshInfo->IsLod && !MeshInfo->IsCollision)
 		{
 			MeshInfo->OptionName = UFbxSceneImportFactory::DefaultOptionName;
 			FbxMeshesArray.Add(MeshInfo);
@@ -205,13 +245,13 @@ void SFbxSceneStaticMeshListView::Construct(const SFbxSceneStaticMeshListView::F
 		.SelectionMode(ESelectionMode::Multi)
 		.OnGenerateRow(this, &SFbxSceneStaticMeshListView::OnGenerateRowFbxSceneListView)
 		.OnContextMenuOpening(this, &SFbxSceneStaticMeshListView::OnOpenContextMenu)
-		.OnSelectionChanged(this, &SFbxSceneStaticMeshListView::OnSelectionChanged)
+		.OnSelectionChanged(this, &SFbxSSceneBaseMeshListView::OnSelectionChanged)
 		.HeaderRow
 		(
 			SNew(SHeaderRow)
 				
-			+ SHeaderRow::Column(SceneImportCheckBoxSelectionHeaderIdName)
-			.FixedWidth(20)
+			+ SHeaderRow::Column(FbxSceneImportStaticMesh::SceneImportCheckBoxSelectionHeaderIdName)
+			.FixedWidth(25)
 			.DefaultLabel(FText::GetEmpty())
 			[
 				SNew(SCheckBox)
@@ -219,17 +259,17 @@ void SFbxSceneStaticMeshListView::Construct(const SFbxSceneStaticMeshListView::F
 				.OnCheckStateChanged(this, &SFbxSceneStaticMeshListView::OnToggleSelectAll)
 			]
 				
-			+ SHeaderRow::Column(SceneImportClassIconHeaderIdName)
+			+ SHeaderRow::Column(FbxSceneImportStaticMesh::SceneImportClassIconHeaderIdName)
 			.FixedWidth(20)
 			.DefaultLabel(FText::GetEmpty())
 				
-			+ SHeaderRow::Column(SceneImportAssetNameHeaderIdName)
-			.FillWidth(300)
+			+ SHeaderRow::Column(FbxSceneImportStaticMesh::SceneImportAssetNameHeaderIdName)
+			.FillWidth(250)
 			.HAlignCell(EHorizontalAlignment::HAlign_Left)
 			.DefaultLabel(LOCTEXT("AssetNameHeaderName", "Asset Name"))
 
-			+ SHeaderRow::Column(SceneImportOptionsNameHeaderIdName)
-			.FillWidth(300)
+			+ SHeaderRow::Column(FbxSceneImportStaticMesh::SceneImportOptionsNameHeaderIdName)
+			.FillWidth(200)
 			.HAlignCell(EHorizontalAlignment::HAlign_Left)
 			.DefaultLabel(LOCTEXT("OptionsNameHeaderName", "Options Name"))
 
@@ -244,7 +284,8 @@ void SFbxSceneStaticMeshListView::Construct(const SFbxSceneStaticMeshListView::F
 TSharedRef< ITableRow > SFbxSceneStaticMeshListView::OnGenerateRowFbxSceneListView(FbxMeshInfoPtr Item, const TSharedRef< STableViewBase >& OwnerTable)
 {
 	TSharedRef< SFbxMeshItemTableListViewRow > ReturnRow = SNew(SFbxMeshItemTableListViewRow, OwnerTable)
-		.FbxMeshInfo(Item);
+		.FbxMeshInfo(Item)
+		.GlobalImportSettings(GlobalImportSettings);
 	return ReturnRow;
 }
 
@@ -256,182 +297,31 @@ TSharedPtr<SWidget> SFbxSceneStaticMeshListView::OnOpenContextMenu()
 	const bool bCloseAfterSelection = true;
 	FMenuBuilder MenuBuilder(bCloseAfterSelection, TSharedPtr<FUICommandList>());
 
-	// We always create a section here, even if there is no parent so that clients can still extend the menu
-	MenuBuilder.BeginSection("FbxScene_SM_ImportSection");
-	{
-		const FSlateIcon PlusIcon(FEditorStyle::GetStyleSetName(), "Plus");
-		MenuBuilder.AddMenuEntry(LOCTEXT("CheckForImport", "Add Selection To Import"), FText(), PlusIcon, FUIAction(FExecuteAction::CreateSP(this, &SFbxSceneStaticMeshListView::AddSelectionToImport)));
-		const FSlateIcon MinusIcon(FEditorStyle::GetStyleSetName(), "PropertyWindow.Button_RemoveFromArray");
-		MenuBuilder.AddMenuEntry(LOCTEXT("UncheckForImport", "Remove Selection From Import"), FText(), MinusIcon, FUIAction(FExecuteAction::CreateSP(this, &SFbxSceneStaticMeshListView::RemoveSelectionFromImport)));
-		MenuBuilder.AddMenuSeparator(TEXT("FbxImportScene_SM_Separator"));
-	}
-	MenuBuilder.EndSection();
-	MenuBuilder.BeginSection("FbxScene_SM_OptionsSection", LOCTEXT("FbxScene_SM_Options", "Options:"));
-	{
-		for (auto OptionName : OverrideNameOptions)
-		{
-			MenuBuilder.AddMenuEntry(FText::FromString(*OptionName.Get()), FText(), FSlateIcon(), FUIAction(FExecuteAction::CreateSP(this, &SFbxSceneStaticMeshListView::AssignToOptions, *OptionName.Get())));
-		}
-	}
-	MenuBuilder.EndSection();
-
-	return MenuBuilder.MakeWidget();
-}
-
-TSharedPtr<FString> SFbxSceneStaticMeshListView::FindOptionNameFromName(FString OptionName)
-{
-	for (TSharedPtr<FString> OptionNamePtr : OverrideNameOptions)
-	{
-		if (OptionNamePtr->Compare(OptionName) == 0)
-		{
-			return OptionNamePtr;
-		}
-	}
-	return TSharedPtr<FString>();
-}
-
-void SFbxSceneStaticMeshListView::AssignToOptions(FString OptionName)
-{
-	bool IsDefaultOptions = OptionName.Compare(UFbxSceneImportFactory::DefaultOptionName) == 0;
-	if (!OverrideNameOptionsMap->Contains(OptionName) && !IsDefaultOptions)
-	{
-		return;
-	}
-	TArray<FbxMeshInfoPtr> SelectedItems;
-	GetSelectedItems(SelectedItems);
-	for (FbxMeshInfoPtr ItemPtr : SelectedItems)
-	{
-		ItemPtr->OptionName = OptionName;
-	}
-	OptionComboBox->SetSelectedItem(FindOptionNameFromName(OptionName));
-}
-
-void SFbxSceneStaticMeshListView::AddSelectionToImport()
-{
-	SetSelectionImportState(true);
-}
-
-void SFbxSceneStaticMeshListView::RemoveSelectionFromImport()
-{
-	SetSelectionImportState(false);
-}
-
-void SFbxSceneStaticMeshListView::SetSelectionImportState(bool MarkForImport)
-{
-	TArray<FbxMeshInfoPtr> SelectedItems;
-	GetSelectedItems(SelectedItems);
-	for (auto Item : SelectedItems)
-	{
-		FbxMeshInfoPtr ItemPtr = Item;
-		ItemPtr->bImportAttribute = MarkForImport;
-	}
-}
-
-void SFbxSceneStaticMeshListView::OnFinishedChangingProperties(const FPropertyChangedEvent& PropertyChangedEvent)
-{
-	SFbxSceneOptionWindow::CopyStaticMeshOptionsToFbxOptions(CurrentStaticMeshImportOptions, SceneImportOptionsStaticMeshDisplay);
-}
-
+	const FSlateIcon PlusIcon(FEditorStyle::GetStyleSetName(), "Plus");
+	MenuBuilder.AddMenuEntry(LOCTEXT("CheckForImport", "Add Selection To Import"), FText(), PlusIcon, FUIAction(FExecuteAction::CreateSP(this, &SFbxSceneStaticMeshListView::AddSelectionToImport)));
+	const FSlateIcon MinusIcon(FEditorStyle::GetStyleSetName(), "PropertyWindow.Button_RemoveFromArray");
+	MenuBuilder.AddMenuEntry(LOCTEXT("UncheckForImport", "Remove Selection From Import"), FText(), MinusIcon, FUIAction(FExecuteAction::CreateSP(this, &SFbxSceneStaticMeshListView::RemoveSelectionFromImport)));
 
-void SFbxSceneStaticMeshListView::OnSelectionChanged(FbxMeshInfoPtr Item, ESelectInfo::Type SelectionType)
-{
-	//Change the option selection
-	TArray<FbxMeshInfoPtr> SelectedItems;
-	GetSelectedItems(SelectedItems);
-	for (FbxMeshInfoPtr SelectItem : SelectedItems)
-	{
-		for (auto kvp : *OverrideNameOptionsMap)
-		{
-			if (kvp.Key.Compare(SelectItem->OptionName) == 0)
-			{
-				OptionComboBox->SetSelectedItem(FindOptionNameFromName(kvp.Key));
-				return;
-			}
-		}
-	}
-	//Select Default in case we don't have a valid OptionName
-	OptionComboBox->SetSelectedItem(FindOptionNameFromName(UFbxSceneImportFactory::DefaultOptionName));
-}
-
-
-bool SFbxSceneStaticMeshListView::CanDeleteOverride()  const
-{
-	return CurrentStaticMeshImportOptions != GlobalImportSettings;
-}
-
-FReply SFbxSceneStaticMeshListView::OnDeleteOverride()
-{
-	if (!CanDeleteOverride())
-	{
-		return FReply::Unhandled();
-	}
-
-	FString CurrentOptionName;
+	AddBakePivotMenu(MenuBuilder);
 
-	for (auto kvp : *(OverrideNameOptionsMap))
+	if (OverrideNameOptions->Num() > 0)
 	{
-		if (kvp.Value == CurrentStaticMeshImportOptions)
+		MenuBuilder.BeginSection("FbxScene_SM_OptionsSection", LOCTEXT("FbxScene_SM_Options", "Options:"));
 		{
-			CurrentOptionName = kvp.Key;
-			if (CurrentOptionName.Compare(UFbxSceneImportFactory::DefaultOptionName) == 0)
+			for (auto OptionName : (*OverrideNameOptions))
 			{
-				//Cannot delete the default options
-				return FReply::Handled();
+				MenuBuilder.AddMenuEntry(FText::FromString(*OptionName.Get()), FText(), FSlateIcon(), FUIAction(FExecuteAction::CreateSP(this, &SFbxSceneStaticMeshListView::AssignToOptions, *OptionName.Get())));
 			}
-			break;
-		}
-	}
-	if (CurrentOptionName.IsEmpty())
-	{
-		return FReply::Handled();
-	}
-
-	TArray<TSharedPtr<FFbxMeshInfo>> RemoveKeys;
-	TArray<TSharedPtr<FFbxMeshInfo>> SceneMeshOverrideKeys;
-	for (TSharedPtr<FFbxMeshInfo> MeshInfo : FbxMeshesArray)
-	{
-		if (CurrentOptionName.Compare(MeshInfo->OptionName) == 0)
-		{
-			MeshInfo->OptionName = UFbxSceneImportFactory::DefaultOptionName;
 		}
+		MenuBuilder.EndSection();
 	}
-	OverrideNameOptions.Remove(FindOptionNameFromName(CurrentOptionName));
-	OverrideNameOptionsMap->Remove(CurrentOptionName);
-	UnFbx::FBXImportOptions *OldOverrideOption = CurrentStaticMeshImportOptions;
-	delete OldOverrideOption;
-
-	CurrentStaticMeshImportOptions = GlobalImportSettings;
-	OptionComboBox->SetSelectedItem(OverrideNameOptions[0]);
 
-	return FReply::Handled();
+	return MenuBuilder.MakeWidget();
 }
 
-FReply SFbxSceneStaticMeshListView::OnSelectAssetUsing()
+void SFbxSceneStaticMeshListView::OnFinishedChangingProperties(const FPropertyChangedEvent& PropertyChangedEvent)
 {
-	ClearSelection();
-	FString CurrentOptionName;
-
-	for (auto kvp : *(OverrideNameOptionsMap))
-	{
-		if (kvp.Value == CurrentStaticMeshImportOptions)
-		{
-			CurrentOptionName = kvp.Key;
-			break;
-		}
-	}
-	if (CurrentOptionName.IsEmpty())
-	{
-		return FReply::Handled();
-	}
-
-	for (FbxMeshInfoPtr MeshInfo : FbxMeshesArray)
-	{
-		if (CurrentOptionName.Compare(MeshInfo->OptionName) == 0)
-		{
-			SetItemSelection(MeshInfo, true);
-		}
-	}
-	return FReply::Handled();
+	SFbxSceneOptionWindow::CopyStaticMeshOptionsToFbxOptions(CurrentMeshImportOptions, SceneImportOptionsStaticMeshDisplay);
 }
 
 void SFbxSceneStaticMeshListView::OnChangedOverrideOptions(TSharedPtr<FString> ItemSelected, ESelectInfo::Type SelectInfo)
@@ -439,65 +329,13 @@ void SFbxSceneStaticMeshListView::OnChangedOverrideOptions(TSharedPtr<FString> I
 	check(ItemSelected.IsValid());
 	if (ItemSelected.Get()->Compare(UFbxSceneImportFactory::DefaultOptionName) == 0)
 	{
-		CurrentStaticMeshImportOptions = GlobalImportSettings;
+		CurrentMeshImportOptions = GlobalImportSettings;
 	}
 	else if(OverrideNameOptionsMap->Contains(*(ItemSelected.Get())))
 	{
-		CurrentStaticMeshImportOptions = *OverrideNameOptionsMap->Find(*(ItemSelected.Get()));
-	}
-	SFbxSceneOptionWindow::CopyFbxOptionsToStaticMeshOptions(CurrentStaticMeshImportOptions, SceneImportOptionsStaticMeshDisplay);
-}
-
-FString SFbxSceneStaticMeshListView::FindUniqueOptionName()
-{
-	int32 SuffixeIndex = 1;
-	bool bFoundSimilarName = true;
-	FString UniqueOptionName;
-	while (bFoundSimilarName)
-	{
-		UniqueOptionName = TEXT("Options ") + FString::FromInt(SuffixeIndex++);
-		bFoundSimilarName = false;
-		for (auto kvp : *(OverrideNameOptionsMap))
-		{
-			if (kvp.Key.Compare(UniqueOptionName) == 0)
-			{
-				bFoundSimilarName = true;
-				break;
-			}
-		}
-	}
-	return UniqueOptionName;
-}
-
-FReply SFbxSceneStaticMeshListView::OnCreateOverrideOptions()
-{
-	UnFbx::FBXImportOptions *OverrideOption = new UnFbx::FBXImportOptions();
-	SFbxSceneOptionWindow::CopyFbxOptionsToFbxOptions(GlobalImportSettings, OverrideOption);
-
-	FString OverrideName = FindUniqueOptionName();
-	TSharedPtr<FString> OverrideNamePtr = MakeShareable(new FString(OverrideName));
-	OverrideNameOptions.Add(OverrideNamePtr);
-	OverrideNameOptionsMap->Add(OverrideName, OverrideOption);
-	//Update the selection to the new override
-	OptionComboBox->SetSelectedItem(OverrideNamePtr);
-	return FReply::Handled();
-}
-
-TSharedPtr<STextComboBox> SFbxSceneStaticMeshListView::CreateOverrideOptionComboBox()
-{
-	OptionComboBox = SNew(STextComboBox)
-		.OptionsSource(&OverrideNameOptions)
-		.InitiallySelectedItem(DefaultOptionNamePtr)
-		.OnSelectionChanged(this, &SFbxSceneStaticMeshListView::OnChangedOverrideOptions);
-	return OptionComboBox;
-}
-
-void SFbxSceneStaticMeshListView::OnToggleSelectAll(ECheckBoxState CheckType)
-{
-	for (auto MeshInfo : FbxMeshesArray)
-	{
-		MeshInfo->bImportAttribute = CheckType == ECheckBoxState::Checked;
+		CurrentMeshImportOptions = *OverrideNameOptionsMap->Find(*(ItemSelected.Get()));
 	}
+	SFbxSceneOptionWindow::CopyFbxOptionsToStaticMeshOptions(CurrentMeshImportOptions, SceneImportOptionsStaticMeshDisplay);
 }
 
 #undef LOCTEXT_NAMESPACE
diff --git a/Engine/Source/Editor/UnrealEd/Private/Fbx/SSceneImportStaticMeshListView.h b/Engine/Source/Editor/UnrealEd/Private/Fbx/SSceneImportStaticMeshListView.h
--- a/Engine/Source/Editor/UnrealEd/Private/Fbx/SSceneImportStaticMeshListView.h
+++ b/Engine/Source/Editor/UnrealEd/Private/Fbx/SSceneImportStaticMeshListView.h
@@ -1,15 +1,12 @@
 // Copyright 1998-2016 Epic Games, Inc. All Rights Reserved.
 #pragma once
 
+#include "SSceneBaseMeshListView.h"
 #include "Classes/Factories/FbxSceneImportFactory.h"
 
 class STextComboBox;
 
-typedef TSharedPtr<FFbxMeshInfo> FbxMeshInfoPtr;
-typedef TMap<FString, EFbxSceneReimportStatusFlags> FbxSceneReimportStatusMap;
-typedef FbxSceneReimportStatusMap* FbxSceneReimportStatusMapPtr;
-
-class SFbxSceneStaticMeshListView : public SListView<FbxMeshInfoPtr>
+class SFbxSceneStaticMeshListView : public SFbxSSceneBaseMeshListView
 {
 public:
 	
@@ -23,47 +20,23 @@ public:
 	{}
 		SLATE_ARGUMENT(TSharedPtr<FFbxSceneInfo>, SceneInfo)
 		SLATE_ARGUMENT(UnFbx::FBXImportOptions*, GlobalImportSettings)
+		SLATE_ARGUMENT(FbxOverrideNameOptionsArrayPtr, OverrideNameOptions)
 		SLATE_ARGUMENT(ImportOptionsNameMapPtr, OverrideNameOptionsMap)
 		SLATE_ARGUMENT(class UFbxSceneImportOptionsStaticMesh*, SceneImportOptionsStaticMeshDisplay)
 	SLATE_END_ARGS()
-	
+
 	/** Construct this widget */
 	void Construct(const FArguments& InArgs);
 	TSharedRef< ITableRow > OnGenerateRowFbxSceneListView(FbxMeshInfoPtr Item, const TSharedRef< STableViewBase >& OwnerTable);
-	bool CanDeleteOverride()  const;
-	FReply OnDeleteOverride();
-	FReply OnSelectAssetUsing();
-	FReply OnCreateOverrideOptions();
-	void OnToggleSelectAll(ECheckBoxState CheckType);
 	
 	void OnFinishedChangingProperties(const FPropertyChangedEvent& PropertyChangedEvent);
-	TSharedPtr<STextComboBox> CreateOverrideOptionComboBox();
 
 protected:
-	TSharedPtr<FFbxSceneInfo> SceneInfo;
-	UnFbx::FBXImportOptions* GlobalImportSettings;
-
-	UnFbx::FBXImportOptions *CurrentStaticMeshImportOptions;
-
 	/** the elements we show in the tree view */
-	TArray<FbxMeshInfoPtr> FbxMeshesArray;
 	class UFbxSceneImportOptionsStaticMesh *SceneImportOptionsStaticMeshDisplay;
-
+	
 	/** Open a context menu for the current selection */
 	TSharedPtr<SWidget> OnOpenContextMenu();
-	void AddSelectionToImport();
-	void RemoveSelectionFromImport();
-	void SetSelectionImportState(bool MarkForImport);
-	void OnSelectionChanged(FbxMeshInfoPtr Item, ESelectInfo::Type SelectionType);
-
-	void OnChangedOverrideOptions(TSharedPtr<FString> ItemSelected, ESelectInfo::Type SelectInfo);
-	TSharedPtr<FString> FindOptionNameFromName(FString OptionName);
-	void AssignToOptions(FString OptionName);
-	FString FindUniqueOptionName();
-	
 
-	TArray<TSharedPtr<FString>> OverrideNameOptions;
-	ImportOptionsNameMapPtr OverrideNameOptionsMap;
-	TSharedPtr<STextComboBox> OptionComboBox;
-	TSharedPtr<FString> DefaultOptionNamePtr;
+	virtual void OnChangedOverrideOptions(TSharedPtr<FString> ItemSelected, ESelectInfo::Type SelectInfo);
 };
diff --git a/Engine/Source/Editor/UnrealEd/Private/Fbx/SSceneMaterialsListView.cpp b/Engine/Source/Editor/UnrealEd/Private/Fbx/SSceneMaterialsListView.cpp
--- a/Engine/Source/Editor/UnrealEd/Private/Fbx/SSceneMaterialsListView.cpp
+++ b/Engine/Source/Editor/UnrealEd/Private/Fbx/SSceneMaterialsListView.cpp
@@ -53,9 +53,13 @@ public:
 	{
 		if (ColumnName == MaterialCheckBoxSelectionHeaderIdName)
 		{
-			return SNew(SCheckBox)
-				.OnCheckStateChanged(this, &SFbxMaterialItemTableListViewRow::OnItemCheckChanged)
-				.IsChecked(this, &SFbxMaterialItemTableListViewRow::IsItemChecked);
+			return SNew(SBox)
+				.HAlign(HAlign_Center)
+				[
+					SNew(SCheckBox)
+					.OnCheckStateChanged(this, &SFbxMaterialItemTableListViewRow::OnItemCheckChanged)
+					.IsChecked(this, &SFbxMaterialItemTableListViewRow::IsItemChecked)
+				];
 		}
 		else if (ColumnName == MaterialNameHeaderIdName)
 		{
@@ -177,7 +181,7 @@ void SFbxSceneMaterialsListView::GetMaterialsFromHierarchy(TArray<FbxMaterialInf
 						CurrentNode = CurrentNode->ParentNodeInfo;
 					}
 				}
-				FString AssetName = FullPath + NodeTreePath + GlobalImportSettings->MaterialPrefixName.ToString() + MaterialInfo->Name;
+				FString AssetName = GlobalImportSettings->MaterialBasePath == NAME_None ? FullPath + NodeTreePath + TEXT("/") + MaterialInfo->Name : GlobalImportSettings->MaterialBasePath.ToString() + MaterialInfo->Name;
 				MaterialInfo->SetOriginalImportPath(AssetName);
 				FString OriginalFullImportName = PackageTools::SanitizePackageName(AssetName);
 				OriginalFullImportName = OriginalFullImportName + TEXT(".") + PackageTools::SanitizePackageName(MaterialInfo->Name);
@@ -254,7 +258,7 @@ void SFbxSceneMaterialsListView::Construct(const SFbxSceneMaterialsListView::FAr
 			SNew(SHeaderRow)
 				
 			+ SHeaderRow::Column(MaterialCheckBoxSelectionHeaderIdName)
-			.FixedWidth(20)
+			.FixedWidth(26)
 			.DefaultLabel(FText::GetEmpty())
 			[
 				SNew(SCheckBox)
@@ -287,7 +291,7 @@ TSharedRef< ITableRow > SFbxSceneMaterialsListView::OnGenerateRowFbxSceneListVie
 	return ReturnRow;
 }
 
-void SFbxSceneMaterialsListView::UpdateMaterialPrefixName()
+void SFbxSceneMaterialsListView::UpdateMaterialBasePath()
 {
 	TArray<FbxMaterialInfoPtr> MaterialsUpdated;
 	GetMaterialsFromHierarchy(MaterialsUpdated, SceneInfo, true);
@@ -308,7 +312,6 @@ TSharedPtr<SWidget> SFbxSceneMaterialsListView::OnOpenContextMenu()
 		MenuBuilder.AddMenuEntry(LOCTEXT("CheckForImport", "Add Selection To Import"), FText(), PlusIcon, FUIAction(FExecuteAction::CreateSP(this, &SFbxSceneMaterialsListView::AddSelectionToImport)));
 		const FSlateIcon MinusIcon(FEditorStyle::GetStyleSetName(), "PropertyWindow.Button_RemoveFromArray");
 		MenuBuilder.AddMenuEntry(LOCTEXT("UncheckForImport", "Remove Selection From Import"), FText(), MinusIcon, FUIAction(FExecuteAction::CreateSP(this, &SFbxSceneMaterialsListView::RemoveSelectionFromImport)));
-		MenuBuilder.AddMenuSeparator(TEXT("FbxImportScene_SM_Separator"));
 	}
 	MenuBuilder.EndSection();
 	MenuBuilder.BeginSection("FbxScene_MAT_AssignSection");
@@ -316,8 +319,31 @@ TSharedPtr<SWidget> SFbxSceneMaterialsListView::OnOpenContextMenu()
 		if (SelectCount == 1)
 		{
 			MenuBuilder.AddMenuEntry(LOCTEXT("AssignExistingMaterial", "Assign Existing Material..."), FText(), FSlateIcon(), FUIAction(FExecuteAction::CreateSP(this, &SFbxSceneMaterialsListView::AssignMaterialToExisting)));
+			MenuBuilder.AddMenuEntry(LOCTEXT("ResetAssignExistingMaterial", "Reset Material to Fbx content"), FText(), FSlateIcon(), FUIAction(FExecuteAction::CreateSP(this, &SFbxSceneMaterialsListView::ResetAssignMaterial)));
+			MenuBuilder.AddMenuSeparator();
+			
+			FContentBrowserModule& ContentBrowserModule = FModuleManager::Get().LoadModuleChecked<FContentBrowserModule>(TEXT("ContentBrowser"));
+			// Configure filter for asset picker
+			FAssetPickerConfig Config;
+			Config.Filter.bRecursiveClasses = true;
+			Config.Filter.ClassNames.Add(UMaterialInterface::StaticClass()->GetFName());
+			Config.InitialAssetViewType = EAssetViewType::List;
+			Config.OnAssetSelected = FOnAssetSelected::CreateSP(this, &SFbxSceneMaterialsListView::AssignMaterialAssetData);
+			Config.bAllowNullSelection = false;
+			Config.bFocusSearchBoxWhenOpened = true;
+			Config.bAllowDragging = false;
+			//Thumbnail are not working since we are in a modal dialog
+			Config.bCanShowRealTimeThumbnails = true;
+			// Don't show stuff in Engine
+			Config.Filter.PackagePaths.Add("/Game");
+			Config.Filter.bRecursivePaths = true;
+			TSharedRef<SWidget> Widget =
+				SNew(SBox)
+				[
+					ContentBrowserModule.Get().CreateAssetPicker(Config)
+				];
+			MenuBuilder.AddWidget(Widget, FText::GetEmpty());
 		}
-		MenuBuilder.AddMenuEntry(LOCTEXT("ResetAssignExistingMaterial", "Reset Material to Fbx content"), FText(), FSlateIcon(), FUIAction(FExecuteAction::CreateSP(this, &SFbxSceneMaterialsListView::ResetAssignMaterial)));
 	}
 	MenuBuilder.EndSection();
 	return MenuBuilder.MakeWidget();
@@ -363,6 +389,34 @@ void SFbxSceneMaterialsListView::OnToggleSelectAll(ECheckBoxState CheckType)
 	}
 }
 
+void SFbxSceneMaterialsListView::AssignMaterialAssetData(const FAssetData& AssetData)
+{
+	TArray<FbxMaterialInfoPtr> SelectedItems;
+	int32 SelectCount = GetSelectedItems(SelectedItems);
+	if (SelectCount != 1)
+	{
+		FSlateApplication::Get().DismissAllMenus();
+		return;
+	}
+
+	UPackage *ContentPackage = AssetData.GetPackage();
+	UObject *ContentObject = AssetData.GetAsset();
+	if (ContentObject != nullptr)
+	{
+		if (!ContentObject->HasAnyFlags(RF_Transient) && !ContentObject->IsPendingKill())
+		{
+			for (FbxMaterialInfoPtr ItemPtr : SelectedItems)
+			{
+				//Override the MeshInfo with the new asset path
+				ItemPtr->SetOverridePath(true);
+				ItemPtr->OverrideImportPath = AssetData.PackageName.ToString();
+				ItemPtr->OverrideFullImportName = AssetData.ObjectPath.ToString();
+			}
+		}
+	}
+	FSlateApplication::Get().DismissAllMenus();
+}
+
 void SFbxSceneMaterialsListView::AssignMaterialToExisting()
 {
 	TArray<FbxMaterialInfoPtr> SelectedItems;
diff --git a/Engine/Source/Editor/UnrealEd/Private/Fbx/SSceneMaterialsListView.h b/Engine/Source/Editor/UnrealEd/Private/Fbx/SSceneMaterialsListView.h
--- a/Engine/Source/Editor/UnrealEd/Private/Fbx/SSceneMaterialsListView.h
+++ b/Engine/Source/Editor/UnrealEd/Private/Fbx/SSceneMaterialsListView.h
@@ -37,7 +37,7 @@ public:
 	void OnToggleSelectAll(ECheckBoxState CheckType);
 
 	void SetCreateContentFolderHierarchy(bool CreateFolder) { CreateContentFolderHierarchy = CreateFolder; }
-	void UpdateMaterialPrefixName();
+	void UpdateMaterialBasePath();
 protected:
 	FString FullPath;
 	bool IsReimport;
@@ -59,6 +59,7 @@ protected:
 	void OnSelectionChanged(FbxMaterialInfoPtr Item, ESelectInfo::Type SelectionType);
 
 	void AssignMaterialToExisting();
+	void AssignMaterialAssetData(const FAssetData& AssetData);
 	void ResetAssignMaterial();
 
 	void GetMaterialsFromHierarchy(TArray<FbxMaterialInfoPtr> &OutMaterials, TSharedPtr<FFbxSceneInfo> SceneInfoSource, bool bFillPathInformation);
diff --git a/Engine/Source/Editor/UnrealEd/Private/Fbx/SSceneReimportNodeTreeView.cpp b/Engine/Source/Editor/UnrealEd/Private/Fbx/SSceneReimportNodeTreeView.cpp
--- a/Engine/Source/Editor/UnrealEd/Private/Fbx/SSceneReimportNodeTreeView.cpp
+++ b/Engine/Source/Editor/UnrealEd/Private/Fbx/SSceneReimportNodeTreeView.cpp
@@ -122,11 +122,13 @@ public:
 	SLATE_BEGIN_ARGS(SFbxReimportSceneTreeViewItem)
 		: _FbxNodeInfo(nullptr)
 		, _NodeStatusMap(nullptr)
+		, _SceneInfo(nullptr)
 	{}
 
 	/** The item content. */
 	SLATE_ARGUMENT(FbxNodeInfoPtr, FbxNodeInfo)
 	SLATE_ARGUMENT(FbxSceneReimportStatusMapPtr, NodeStatusMap)
+	SLATE_ARGUMENT(TSharedPtr<FFbxSceneInfo>, SceneInfo)
 	SLATE_END_ARGS()
 
 	/**
@@ -138,15 +140,42 @@ public:
 	{
 		FbxNodeInfo = InArgs._FbxNodeInfo;
 		NodeStatusMap = InArgs._NodeStatusMap;
+		SceneInfo = InArgs._SceneInfo;
 		
 		//This is suppose to always be valid
 		check(FbxNodeInfo.IsValid());
 		check(NodeStatusMap != nullptr);
+		check(SceneInfo.IsValid());
 
 		UClass *IconClass = AActor::StaticClass();
 		if (FbxNodeInfo->AttributeInfo.IsValid())
 		{
-			IconClass = FbxNodeInfo->AttributeInfo->GetType();
+			if(!FbxNodeInfo->AttributeInfo->bOriginalTypeChanged)
+				IconClass = FbxNodeInfo->AttributeInfo->GetType();
+		}
+		else if (FbxNodeInfo->AttributeType.Compare(TEXT("eLight")) == 0)
+		{
+			IconClass = ULightComponent::StaticClass();
+			if (SceneInfo->LightInfo.Contains(FbxNodeInfo->AttributeUniqueId))
+			{
+				TSharedPtr<FFbxLightInfo> LightInfo = *SceneInfo->LightInfo.Find(FbxNodeInfo->AttributeUniqueId);
+				if (LightInfo->Type == 0)
+				{
+					IconClass = UPointLightComponent::StaticClass();
+				}
+				else if (LightInfo->Type == 1)
+				{
+					IconClass = UDirectionalLightComponent::StaticClass();
+				}
+				else if (LightInfo->Type == 2)
+				{
+					IconClass = USpotLightComponent::StaticClass();
+				}
+			}
+		}
+		else if (FbxNodeInfo->AttributeType.Compare(TEXT("eCamera")) == 0)
+		{
+			IconClass = UCameraComponent::StaticClass();
 		}
 		const FSlateBrush* ClassIcon = FClassIconFinder::FindIconForClass(IconClass);
 
@@ -154,6 +183,7 @@ public:
 		[
 		SNew(SHorizontalBox)
 		+ SHorizontalBox::Slot()
+			.Padding(2.0f, 0.0f, 2.0f, 0.0f)
 			.AutoWidth()
 			[
 				SNew(SCheckBox)
@@ -213,8 +243,11 @@ private:
 		{
 			EFbxSceneReimportStatusFlags ReimportFlags = *NodeStatusMap->Find(FbxNodeInfo->NodeHierarchyPath);
 			//The remove only should not be possible this is why there is no remove only case
-
-			if ((ReimportFlags & EFbxSceneReimportStatusFlags::Added) != EFbxSceneReimportStatusFlags::None)
+			if(FbxNodeInfo->AttributeInfo.IsValid() && FbxNodeInfo->AttributeInfo->bOriginalTypeChanged)
+			{
+				TooltipText += LOCTEXT("SFbxReimportSceneTreeViewItem_TypeChangedTooltip", " type has changed, only the transform can be reimport.").ToString();
+			}
+			else if ((ReimportFlags & EFbxSceneReimportStatusFlags::Added) != EFbxSceneReimportStatusFlags::None)
 			{
 				TooltipText += LOCTEXT("SFbxReimportSceneTreeViewItem_AddedTooltip", " Will be add to the blueprint hierarchy.").ToString();
 			}
@@ -253,12 +286,39 @@ private:
 		return SlateBrush;
 	}
 
+	void RecursivelySetLODMeshImportState(TSharedPtr<FFbxNodeInfo> NodeInfo, bool bState)
+	{
+		//Set the bImportNode property for all mesh under eLODGroup
+		for (TSharedPtr<FFbxNodeInfo> ChildNodeInfo : NodeInfo->Childrens)
+		{
+			if (!ChildNodeInfo.IsValid())
+				continue;
+			if (ChildNodeInfo->AttributeType.Compare(TEXT("eMesh")) == 0)
+			{
+				EFbxSceneReimportStatusFlags *StatusFlag = NodeStatusMap->Find(ChildNodeInfo->NodeHierarchyPath);
+				if (bState)
+				{
+					*StatusFlag = *StatusFlag | EFbxSceneReimportStatusFlags::ReimportAsset;
+				}
+				else
+				{
+					*StatusFlag = *StatusFlag & ~EFbxSceneReimportStatusFlags::ReimportAsset;
+				}
+			}
+			else
+			{
+				RecursivelySetLODMeshImportState(ChildNodeInfo, bState);
+			}
+		}
+	}
+
 	void OnItemCheckChanged(ECheckBoxState CheckType)
 	{
 		if (!FbxNodeInfo.IsValid() || !NodeStatusMap->Contains(FbxNodeInfo->NodeHierarchyPath))
 			return;
+		bool bNewState = CheckType == ECheckBoxState::Checked;
 		EFbxSceneReimportStatusFlags *StatusFlag = NodeStatusMap->Find(FbxNodeInfo->NodeHierarchyPath);
-		if (CheckType == ECheckBoxState::Checked)
+		if (bNewState)
 		{
 			*StatusFlag = *StatusFlag | EFbxSceneReimportStatusFlags::ReimportAsset;
 		}
@@ -266,6 +326,29 @@ private:
 		{
 			*StatusFlag = *StatusFlag & ~EFbxSceneReimportStatusFlags::ReimportAsset;
 		}
+
+		if (FbxNodeInfo->AttributeType.Compare(TEXT("eLODGroup")) == 0)
+		{
+			RecursivelySetLODMeshImportState(FbxNodeInfo, bNewState);
+		}
+		if (FbxNodeInfo->AttributeType.Compare(TEXT("eMesh")) == 0)
+		{
+			//Verify if parent is a LOD group
+			TSharedPtr<FFbxNodeInfo> ParentLODNodeInfo = FFbxSceneInfo::RecursiveFindLODParentNode(FbxNodeInfo);
+			if (ParentLODNodeInfo.IsValid())
+			{
+				EFbxSceneReimportStatusFlags *ParentStatusFlag = NodeStatusMap->Find(ParentLODNodeInfo->NodeHierarchyPath);
+				if (bNewState)
+				{
+					*ParentStatusFlag = *ParentStatusFlag | EFbxSceneReimportStatusFlags::ReimportAsset;
+				}
+				else
+				{
+					*ParentStatusFlag = *ParentStatusFlag & ~EFbxSceneReimportStatusFlags::ReimportAsset;
+				}
+				RecursivelySetLODMeshImportState(ParentLODNodeInfo, bNewState);
+			}
+		}
 	}
 
 	ECheckBoxState IsItemChecked() const
@@ -281,13 +364,15 @@ private:
 	FbxNodeInfoPtr FbxNodeInfo;
 
 	FbxSceneReimportStatusMapPtr NodeStatusMap;
+	TSharedPtr<FFbxSceneInfo> SceneInfo;
 };
 
 TSharedRef< ITableRow > SFbxReimportSceneTreeView::OnGenerateRowFbxSceneTreeView(FbxNodeInfoPtr Item, const TSharedRef< STableViewBase >& OwnerTable)
 {
 	TSharedRef< SFbxReimportSceneTreeViewItem > ReturnRow = SNew(SFbxReimportSceneTreeViewItem, OwnerTable)
 		.FbxNodeInfo(Item)
-		.NodeStatusMap(NodeStatusMap);
+		.NodeStatusMap(NodeStatusMap)
+		.SceneInfo(SceneInfo);
 	return ReturnRow;
 }
 
@@ -299,24 +384,63 @@ void SFbxReimportSceneTreeView::OnGetChildrenFbxSceneTreeView(FbxNodeInfoPtr InP
 	//Current node contain the add and same node
 	if (NodeValue->CurrentNode.IsValid())
 	{
-		for (FbxNodeInfoPtr Child : NodeValue->CurrentNode->Childrens)
+		if (NodeValue->CurrentNode->AttributeType.Compare(TEXT("eLODGroup")) == 0 && NodeValue->CurrentNode->Childrens.Num() > 0)
 		{
-			if (Child.IsValid())
+			TSharedPtr<FFbxNodeInfo> Child = NodeValue->CurrentNode;
+			do
 			{
-				ChildNames.Add(Child->NodeName);
-				OutChildren.Add(Child);
+				Child = Child->Childrens[0];
+				if (Child.IsValid() && Child->AttributeType.Compare(TEXT("eMesh")) == 0)
+				{
+					ChildNames.Add(Child->NodeName);
+					OutChildren.Add(Child);
+					return;
+				}
+			} while (Child->Childrens.Num() > 0);
+			//We do not found any LOD mesh don't add any child
+		}
+		else
+		{
+			for (FbxNodeInfoPtr Child : NodeValue->CurrentNode->Childrens)
+			{
+				if (Child.IsValid())
+				{
+					ChildNames.Add(Child->NodeName);
+					OutChildren.Add(Child);
+				}
 			}
 		}
 	}
+
 	//Original node is use for finding the remove nodes
 	if (NodeValue->OriginalNode.IsValid())
 	{
-		for (FbxNodeInfoPtr Child : NodeValue->OriginalNode->Childrens)
+		if (NodeValue->OriginalNode->AttributeType.Compare(TEXT("eLODGroup")) == 0 && NodeValue->OriginalNode->Childrens.Num() > 0)
+		{
+			TSharedPtr<FFbxNodeInfo> Child = NodeValue->OriginalNode;
+			do
+			{
+				Child = Child->Childrens[0];
+				if (Child.IsValid() && Child->AttributeType.Compare(TEXT("eMesh")) == 0)
+				{
+					if (!ChildNames.Contains(Child->NodeName))
+					{
+						OutChildren.Add(Child);
+					}
+					return;
+				}
+			} while (Child->Childrens.Num() > 0);
+			//We do not found any LOD mesh don't add any child
+		}
+		else
 		{
-			//We have a delete node
-			if (Child.IsValid() && !ChildNames.Contains(Child->NodeName))
+			for (FbxNodeInfoPtr Child : NodeValue->OriginalNode->Childrens)
 			{
-				OutChildren.Add(Child);
+				//We have a delete node
+				if (Child.IsValid() && !ChildNames.Contains(Child->NodeName))
+				{
+					OutChildren.Add(Child);
+				}
 			}
 		}
 	}
