Copying //UE4/Dev-Editor to Dev-Main (//UE4/Dev-Main)

diff --git a/Engine/Source/Editor/UnrealEd/Private/Fbx/SSceneReimportStaticMeshListView.cpp b/Engine/Source/Editor/UnrealEd/Private/Fbx/SSceneReimportStaticMeshListView.cpp
--- a/Engine/Source/Editor/UnrealEd/Private/Fbx/SSceneReimportStaticMeshListView.cpp
+++ b/Engine/Source/Editor/UnrealEd/Private/Fbx/SSceneReimportStaticMeshListView.cpp
@@ -9,15 +9,19 @@
 #include "FbxImporter.h"
 #include "STextComboBox.h"
 #include "PackageTools.h"
+#include "STextEntryPopup.h"
 
 #define LOCTEXT_NAMESPACE "SFbxReimportSceneStaticMeshListView"
 
-const FName CheckBoxSelectionHeaderIdName(TEXT("CheckBoxSelectionHeaderId"));
-const FName ClassIconHeaderIdName(TEXT("ClassIconHeaderId"));
-const FName AssetNameHeaderIdName(TEXT("AssetNameHeaderId"));
-const FName AssetStatusHeaderIdName(TEXT("AssetStatusHeaderId"));
-const FName ContentPathHeaderIdName(TEXT("ContentPathHeaderId"));
-const FName OptionNameHeaderIdName(TEXT("OptionNameHeaderId"));
+namespace FbxSceneReimportStaticMesh
+{
+	const FName CheckBoxSelectionHeaderIdName(TEXT("CheckBoxSelectionHeaderId"));
+	const FName ClassIconHeaderIdName(TEXT("ClassIconHeaderId"));
+	const FName AssetNameHeaderIdName(TEXT("AssetNameHeaderId"));
+	const FName AssetStatusHeaderIdName(TEXT("AssetStatusHeaderId"));
+	const FName ContentPathHeaderIdName(TEXT("ContentPathHeaderId"));
+	const FName OptionNameHeaderIdName(TEXT("OptionNameHeaderId"));
+}
 
 class SFbxMeshReimportItemTableListViewRow : public SMultiColumnTableRow<FbxMeshInfoPtr>
 {
@@ -25,9 +29,11 @@ public:
 	SLATE_BEGIN_ARGS(SFbxMeshReimportItemTableListViewRow)
 	: _FbxMeshInfo(nullptr)
 	, _MeshStatusMap(nullptr)
+	, _GlobalImportSettings(nullptr)
 	{}
 		SLATE_ARGUMENT(FbxMeshInfoPtr, FbxMeshInfo)
 		SLATE_ARGUMENT(FbxSceneReimportStatusMapPtr, MeshStatusMap)
+		SLATE_ARGUMENT(UnFbx::FBXImportOptions*, GlobalImportSettings)
 	SLATE_END_ARGS()
 
 
@@ -36,10 +42,12 @@ public:
 	{
 		FbxMeshInfo = InArgs._FbxMeshInfo;
 		MeshStatusMap = InArgs._MeshStatusMap;
+		GlobalImportSettings = InArgs._GlobalImportSettings;
 
 		//This is suppose to always be valid
 		check(FbxMeshInfo.IsValid());
 		check(MeshStatusMap != nullptr);
+		check(GlobalImportSettings != nullptr);
 
 		//Get item data
 		GetItemRowData();
@@ -60,13 +68,18 @@ public:
 	/** Overridden from SMultiColumnTableRow.  Generates a widget for this column of the list view. */
 	virtual TSharedRef<SWidget> GenerateWidgetForColumn(const FName& ColumnName) override
 	{
-		if (ColumnName == CheckBoxSelectionHeaderIdName)
+		if (ColumnName == FbxSceneReimportStaticMesh::CheckBoxSelectionHeaderIdName)
 		{
-			return SNew(SCheckBox)
-				.OnCheckStateChanged(this, &SFbxMeshReimportItemTableListViewRow::OnItemCheckChanged)
-				.IsChecked(this, &SFbxMeshReimportItemTableListViewRow::IsItemChecked);
+			return SNew(SBox)
+				.HAlign(HAlign_Center)
+				[
+					SNew(SCheckBox)
+					.OnCheckStateChanged(this, &SFbxMeshReimportItemTableListViewRow::OnItemCheckChanged)
+					.IsChecked(this, &SFbxMeshReimportItemTableListViewRow::IsItemChecked)
+					.IsEnabled(!FbxMeshInfo->bOriginalTypeChanged)
+				];
 		}
-		else if (ColumnName == ClassIconHeaderIdName && SlateBrush != nullptr)
+		else if (ColumnName == FbxSceneReimportStaticMesh::ClassIconHeaderIdName && SlateBrush != nullptr)
 		{
 			UClass *IconClass = FbxMeshInfo->GetType();
 			const FSlateBrush* ClassIcon = FClassIconFinder::FindIconForClass(IconClass);
@@ -92,32 +105,38 @@ public:
 				];
 				return IconContent;
 		}
-		else if (ColumnName == AssetNameHeaderIdName)
+		else if (ColumnName == FbxSceneReimportStaticMesh::AssetNameHeaderIdName)
 		{
 			return SNew(STextBlock)
 				.Text(FText::FromString(FbxMeshInfo->Name))
 				.ToolTipText(FText::FromString(FbxMeshInfo->Name));
 		}
-		else if (ColumnName == AssetStatusHeaderIdName)
+		else if (ColumnName == FbxSceneReimportStaticMesh::AssetStatusHeaderIdName)
 		{
 			return SNew(STextBlock)
 				.Text(this, &SFbxMeshReimportItemTableListViewRow::GetAssetStatus)
 				.ToolTipText(this, &SFbxMeshReimportItemTableListViewRow::GetAssetStatusTooltip);
 				
 		}
-		else if (ColumnName == ContentPathHeaderIdName)
+		else if (ColumnName == FbxSceneReimportStaticMesh::ContentPathHeaderIdName)
 		{
 			return SNew(STextBlock)
 				.Text(this, &SFbxMeshReimportItemTableListViewRow::GetAssetFullName)
 				.ColorAndOpacity(this, &SFbxMeshReimportItemTableListViewRow::GetContentPathTextColor)
 				.ToolTipText(this, &SFbxMeshReimportItemTableListViewRow::GetAssetFullName);
 		}
-		else if (ColumnName == OptionNameHeaderIdName)
+		else if (ColumnName == FbxSceneReimportStaticMesh::OptionNameHeaderIdName)
 		{
 			return SNew(STextBlock)
 				.Text(this, &SFbxMeshReimportItemTableListViewRow::GetOptionName)
 				.ToolTipText(this, &SFbxMeshReimportItemTableListViewRow::GetOptionName);
 		}
+		else if (ColumnName == FbxSceneBaseListViewColumn::PivotColumnId)
+		{
+			return SNew(STextBlock)
+				.Text(this, &SFbxMeshReimportItemTableListViewRow::GetAssetPivotNodeName)
+				.ToolTipText(this, &SFbxMeshReimportItemTableListViewRow::GetAssetPivotNodeName);
+		}
 
 		return SNullWidget::NullWidget;
 	}
@@ -172,8 +191,12 @@ private:
 		{
 			EFbxSceneReimportStatusFlags ReimportFlags = *MeshStatusMap->Find(FbxMeshInfo->OriginalImportPath);
 			//The remove only should not be possible this is why there is no remove only case
-
-			if ((ReimportFlags & EFbxSceneReimportStatusFlags::FoundContentBrowserAsset) == EFbxSceneReimportStatusFlags::None)
+			if (FbxMeshInfo->bOriginalTypeChanged)
+			{
+				AssetStatus = LOCTEXT("SFbxMeshReimportItemTableListViewRow_AssetTypeChanged", "Type Changed, no reimport").ToString();
+				AssetStatusTooltip = LOCTEXT("SFbxMeshReimportItemTableListViewRow_AssetTypeChangedTooltip", "This item type changed, we cannot reimport an asset of a different type").ToString();
+			}
+			else if ((ReimportFlags & EFbxSceneReimportStatusFlags::FoundContentBrowserAsset) == EFbxSceneReimportStatusFlags::None)
 			{
 				if ((ReimportFlags & EFbxSceneReimportStatusFlags::Added) != EFbxSceneReimportStatusFlags::None)
 				{
@@ -214,7 +237,7 @@ private:
 
 	void OnItemCheckChanged(ECheckBoxState CheckType)
 	{
-		if (!FbxMeshInfo.IsValid())
+		if (!FbxMeshInfo.IsValid() || FbxMeshInfo->bOriginalTypeChanged)
 			return;
 		if (MeshStatusMap->Contains(FbxMeshInfo->OriginalImportPath))
 		{
@@ -232,16 +255,22 @@ private:
 
 	ECheckBoxState IsItemChecked() const
 	{
-		if (MeshStatusMap->Contains(FbxMeshInfo->OriginalImportPath))
+		if (MeshStatusMap->Contains(FbxMeshInfo->OriginalImportPath) && !FbxMeshInfo->bOriginalTypeChanged)
 		{
 			return (*(MeshStatusMap->Find(FbxMeshInfo->OriginalImportPath)) & EFbxSceneReimportStatusFlags::ReimportAsset) != EFbxSceneReimportStatusFlags::None ? ECheckBoxState::Checked : ECheckBoxState::Unchecked;
 		}
 		return ECheckBoxState::Unchecked;
 	}
 
+	FText GetAssetPivotNodeName() const
+	{
+		return GlobalImportSettings->bBakePivotInVertex ? FText::FromString(FbxMeshInfo->PivotNodeName) : FText::FromString(TEXT("-"));
+	}
+
 	/** The node info to build the tree view row from. */
 	FbxMeshInfoPtr FbxMeshInfo;
 	FbxSceneReimportStatusMapPtr MeshStatusMap;
+	UnFbx::FBXImportOptions *GlobalImportSettings;
 	
 	//Get item data
 	FString AssetStatus;
@@ -255,6 +284,7 @@ SFbxSceneStaticMeshReimportListView::~SFbxSceneStaticMeshReimportListView()
 	SceneInfo = nullptr;
 	SceneInfoOriginal = nullptr;
 	GlobalImportSettings = nullptr;
+	OverrideNameOptions = nullptr;
 	OverrideNameOptionsMap = nullptr;
 	SceneImportOptionsStaticMeshDisplay = nullptr;
 	MeshStatusMap = nullptr;
@@ -271,6 +301,7 @@ void SFbxSceneStaticMeshReimportListView::Construct(const SFbxSceneStaticMeshRei
 	SceneInfoOriginal = InArgs._SceneInfoOriginal;
 	MeshStatusMap = InArgs._MeshStatusMap;
 	GlobalImportSettings = InArgs._GlobalImportSettings;
+	OverrideNameOptions = InArgs._OverrideNameOptions;
 	OverrideNameOptionsMap = InArgs._OverrideNameOptionsMap;
 	SceneImportOptionsStaticMeshDisplay = InArgs._SceneImportOptionsStaticMeshDisplay;
 
@@ -278,6 +309,7 @@ void SFbxSceneStaticMeshReimportListView::Construct(const SFbxSceneStaticMeshRei
 	check(SceneInfoOriginal.IsValid());
 	check(MeshStatusMap != nullptr);
 	check(GlobalImportSettings != nullptr);
+	check(OverrideNameOptions != nullptr);
 	check(OverrideNameOptionsMap != nullptr);
 	check(SceneImportOptionsStaticMeshDisplay != nullptr);
 	
@@ -285,20 +317,32 @@ void SFbxSceneStaticMeshReimportListView::Construct(const SFbxSceneStaticMeshRei
 
 	for (auto kvp : *OverrideNameOptionsMap)
 	{
-		if (kvp.Key.Compare(UFbxSceneImportFactory::DefaultOptionName) == 0)
+		bool bNameExist = false;
+		for (TSharedPtr<FString> OverrideName : (*OverrideNameOptions))
 		{
-			OverrideNameOptions.Add(DefaultOptionNamePtr);
-			SFbxSceneOptionWindow::CopyFbxOptionsToFbxOptions(kvp.Value, GlobalImportSettings);
+			if (OverrideName.Get()->Compare(kvp.Key) == 0)
+			{
+				bNameExist = true;
+				break;
+			}
 		}
-		else
+		if (!bNameExist)
 		{
-			OverrideNameOptions.Add(MakeShareable(new FString(kvp.Key)));
+			if (kvp.Key.Compare(UFbxSceneImportFactory::DefaultOptionName) == 0)
+			{
+				OverrideNameOptions->Add(DefaultOptionNamePtr);
+				SFbxSceneOptionWindow::CopyFbxOptionsToFbxOptions(kvp.Value, GlobalImportSettings);
+			}
+			else
+			{
+				OverrideNameOptions->Add(MakeShareable(new FString(kvp.Key)));
+			}
 		}
 	}
 	//Set the default options to the current global import settings
 	GlobalImportSettings->bTransformVertexToAbsolute = false;
 	GlobalImportSettings->StaticMeshLODGroup = NAME_None;
-	CurrentStaticMeshImportOptions = GlobalImportSettings;
+	CurrentMeshImportOptions = GlobalImportSettings;
 
 	FbxMeshesArray.Empty();
 	FilterFbxMeshesArray.Empty();
@@ -309,7 +353,7 @@ void SFbxSceneStaticMeshReimportListView::Construct(const SFbxSceneStaticMeshRei
 
 	for (FbxMeshInfoPtr MeshInfo : SceneInfo->MeshInfo)
 	{
-		if (MeshInfo->bIsSkelMesh)
+		if (MeshInfo->bIsSkelMesh  || MeshInfo->IsLod || MeshInfo->IsCollision)
 		{
 			continue;
 		}
@@ -338,7 +382,7 @@ void SFbxSceneStaticMeshReimportListView::Construct(const SFbxSceneStaticMeshRei
 
 	for (FbxMeshInfoPtr OriginalMeshInfo : SceneInfoOriginal->MeshInfo)
 	{
-		if (OriginalMeshInfo->bIsSkelMesh)
+		if (OriginalMeshInfo->bIsSkelMesh || OriginalMeshInfo->IsLod || OriginalMeshInfo->IsCollision)
 		{
 			continue;
 		}
@@ -356,8 +400,29 @@ void SFbxSceneStaticMeshReimportListView::Construct(const SFbxSceneStaticMeshRei
 				break;
 			}
 		}
-		if (FoundMeshInfo.IsValid())
+		if (FoundMeshInfo.IsValid() && FoundMeshInfo->bOriginalTypeChanged)
+		{
+			//We dont reimport asset that have change their type
+			EFbxSceneReimportStatusFlags StatusFlag = EFbxSceneReimportStatusFlags::None;
+			MeshStatusMap->Add(FoundMeshInfo->OriginalImportPath, StatusFlag);
+		}
+		else if (FoundMeshInfo.IsValid())
 		{
+			//Set the old pivot information if we find one
+			FbxNodeInfoPtr OriginalPivotNode = FindNodeInfoByUid(OriginalMeshInfo->PivotNodeUid, SceneInfoOriginal);
+			if (OriginalPivotNode.IsValid())
+			{
+				for (FbxNodeInfoPtr NodeInfo : SceneInfo->HierarchyInfo)
+				{
+					if (OriginalPivotNode->NodeHierarchyPath.Compare(NodeInfo->NodeHierarchyPath) == 0)
+					{
+						FoundMeshInfo->PivotNodeUid = NodeInfo->UniqueId;
+						FoundMeshInfo->PivotNodeName = NodeInfo->NodeName;
+						break;
+					}
+				}
+			}
+
 			//We have a match
 			EFbxSceneReimportStatusFlags StatusFlag = EFbxSceneReimportStatusFlags::Same;
 			if (OriginalMeshInfo->GetContentObject() != nullptr)
@@ -403,8 +468,8 @@ void SFbxSceneStaticMeshReimportListView::Construct(const SFbxSceneStaticMeshRei
 			(
 				SNew(SHeaderRow)
 				
-				+ SHeaderRow::Column(CheckBoxSelectionHeaderIdName)
-				.FixedWidth(20)
+				+ SHeaderRow::Column(FbxSceneReimportStaticMesh::CheckBoxSelectionHeaderIdName)
+				.FixedWidth(26)
 				.DefaultLabel(FText::GetEmpty())
 				[
 					SNew(SCheckBox)
@@ -412,26 +477,26 @@ void SFbxSceneStaticMeshReimportListView::Construct(const SFbxSceneStaticMeshRei
 					.OnCheckStateChanged(this, &SFbxSceneStaticMeshReimportListView::OnToggleSelectAll)
 				]
 				
-				+ SHeaderRow::Column(ClassIconHeaderIdName)
+				+ SHeaderRow::Column(FbxSceneReimportStaticMesh::ClassIconHeaderIdName)
 				.FixedWidth(20)
 				.DefaultLabel(FText::GetEmpty())
 				
-				+ SHeaderRow::Column(AssetNameHeaderIdName)
+				+ SHeaderRow::Column(FbxSceneReimportStaticMesh::AssetNameHeaderIdName)
 				.FillWidth(250)
 				.HAlignCell(EHorizontalAlignment::HAlign_Left)
 				.DefaultLabel(LOCTEXT("AssetNameHeaderName", "Asset Name"))
 
-				+ SHeaderRow::Column(ContentPathHeaderIdName)
+				+ SHeaderRow::Column(FbxSceneReimportStaticMesh::ContentPathHeaderIdName)
 				.FillWidth(250)
 				.HAlignCell(EHorizontalAlignment::HAlign_Left)
 				.DefaultLabel(LOCTEXT("ContentPathHeaderName", "Content Path"))
 
-				+ SHeaderRow::Column(AssetStatusHeaderIdName)
+				+ SHeaderRow::Column(FbxSceneReimportStaticMesh::AssetStatusHeaderIdName)
 				.FillWidth(160)
 				.HAlignCell(EHorizontalAlignment::HAlign_Left)
 				.DefaultLabel(LOCTEXT("AssetStatusHeaderName", "Asset Status"))
 				
-				+ SHeaderRow::Column(OptionNameHeaderIdName)
+				+ SHeaderRow::Column(FbxSceneReimportStaticMesh::OptionNameHeaderIdName)
 				.FillWidth(100)
 				.HAlignCell(EHorizontalAlignment::HAlign_Left)
 				.DefaultLabel(LOCTEXT("AssetOptionNameHeaderName", "Option Name"))
@@ -443,7 +508,8 @@ TSharedRef< ITableRow > SFbxSceneStaticMeshReimportListView::OnGenerateRowFbxSce
 {
 	TSharedRef< SFbxMeshReimportItemTableListViewRow > ReturnRow = SNew(SFbxMeshReimportItemTableListViewRow, OwnerTable)
 		.FbxMeshInfo(Item)
-		.MeshStatusMap(MeshStatusMap);
+		.MeshStatusMap(MeshStatusMap)
+		.GlobalImportSettings(GlobalImportSettings);
 	return ReturnRow;
 }
 
@@ -452,82 +518,18 @@ void SFbxSceneStaticMeshReimportListView::OnChangedOverrideOptions(TSharedPtr<FS
 	check(ItemSelected.IsValid());
 	if (ItemSelected.Get()->Compare(UFbxSceneImportFactory::DefaultOptionName) == 0)
 	{
-		CurrentStaticMeshImportOptions = GlobalImportSettings;
+		CurrentMeshImportOptions = GlobalImportSettings;
 	}
 	else if (OverrideNameOptionsMap->Contains(*(ItemSelected)))
 	{
-		CurrentStaticMeshImportOptions = *OverrideNameOptionsMap->Find(*(ItemSelected));
-	}
-	SFbxSceneOptionWindow::CopyFbxOptionsToStaticMeshOptions(CurrentStaticMeshImportOptions, SceneImportOptionsStaticMeshDisplay);
-}
-
-FString SFbxSceneStaticMeshReimportListView::FindUniqueOptionName()
-{
-	int32 SuffixeIndex = 1;
-	bool bFoundSimilarName = true;
-	FString UniqueOptionName;
-	while (bFoundSimilarName)
-	{
-		UniqueOptionName = TEXT("Options ") + FString::FromInt(SuffixeIndex++);
-		bFoundSimilarName = false;
-		for (auto kvp : *(OverrideNameOptionsMap))
-		{
-			if (kvp.Key.Compare(UniqueOptionName) == 0)
-			{
-				bFoundSimilarName = true;
-				break;
-			}
-		}
+		CurrentMeshImportOptions = *OverrideNameOptionsMap->Find(*(ItemSelected));
 	}
-	return UniqueOptionName;
-}
-
-FReply SFbxSceneStaticMeshReimportListView::OnCreateOverrideOptions()
-{
-	UnFbx::FBXImportOptions *OverrideOption = new UnFbx::FBXImportOptions();
-	SFbxSceneOptionWindow::CopyFbxOptionsToFbxOptions(GlobalImportSettings, OverrideOption);
-
-	FString OverrideName = FindUniqueOptionName();
-	TSharedPtr<FString> OverrideNamePtr = MakeShareable(new FString(OverrideName));
-	OverrideNameOptions.Add(OverrideNamePtr);
-	OverrideNameOptionsMap->Add(OverrideName, OverrideOption);
-	//Update the selection to the new override
-	OptionComboBox->SetSelectedItem(OverrideNamePtr);
-	return FReply::Handled();
-}
-
-TSharedPtr<STextComboBox> SFbxSceneStaticMeshReimportListView::CreateOverrideOptionComboBox()
-{
-	OptionComboBox = SNew(STextComboBox)
-		.OptionsSource(&OverrideNameOptions)
-		.InitiallySelectedItem(DefaultOptionNamePtr)
-		.OnSelectionChanged(this, &SFbxSceneStaticMeshReimportListView::OnChangedOverrideOptions);
-	return OptionComboBox;
+	SFbxSceneOptionWindow::CopyFbxOptionsToStaticMeshOptions(CurrentMeshImportOptions, SceneImportOptionsStaticMeshDisplay);
 }
 
 void SFbxSceneStaticMeshReimportListView::OnFinishedChangingProperties(const FPropertyChangedEvent& PropertyChangedEvent)
 {
-	SFbxSceneOptionWindow::CopyStaticMeshOptionsToFbxOptions(CurrentStaticMeshImportOptions, SceneImportOptionsStaticMeshDisplay);
-}
-
-void SFbxSceneStaticMeshReimportListView::OnSelectionChanged(FbxMeshInfoPtr Item, ESelectInfo::Type SelectionType)
-{
-	//Change the option selection
-	TArray<FbxMeshInfoPtr> SelectedItems;
-	GetSelectedItems(SelectedItems);
-	for (FbxMeshInfoPtr SelectItem : SelectedItems)
-	{
-		for (auto kvp : *OverrideNameOptionsMap)
-		{
-			if (kvp.Key.Compare(SelectItem->OptionName) == 0)
-			{
-				OptionComboBox->SetSelectedItem(FindOptionNameFromName(kvp.Key));
-				return;
-			}
-		}
-	}
-	//Select Default in case we don't have a valid OptionName
-	OptionComboBox->SetSelectedItem(FindOptionNameFromName(UFbxSceneImportFactory::DefaultOptionName));
+	SFbxSceneOptionWindow::CopyStaticMeshOptionsToFbxOptions(CurrentMeshImportOptions, SceneImportOptionsStaticMeshDisplay);
 }
 
 TSharedPtr<SWidget> SFbxSceneStaticMeshReimportListView::OnOpenContextMenu()
@@ -553,7 +555,6 @@ TSharedPtr<SWidget> SFbxSceneStaticMeshReimportListView::OnOpenContextMenu()
 	{
 		MenuBuilder.BeginSection("FbxImportScene_SM_Assign", LOCTEXT("FbxScene_SM_Assign", "Assign"));
 		{
-			MenuBuilder.AddMenuSeparator(TEXT("FbxImportScene_SM_Separator"));
 			if (SelectCount == 1)
 			{
 				MenuBuilder.AddMenuEntry(LOCTEXT("FbxImportScene_SM_Assign", "Assign existing static mesh..."), FText(), FSlateIcon(), FUIAction(FExecuteAction::CreateSP(this, &SFbxSceneStaticMeshReimportListView::AssignToStaticMesh)));
@@ -566,6 +567,8 @@ TSharedPtr<SWidget> SFbxSceneStaticMeshReimportListView::OnOpenContextMenu()
 		MenuBuilder.EndSection();
 	}
 
+	AddBakePivotMenu(MenuBuilder);
+
 	bool bShowOptionMenu = false;
 	for (FbxMeshInfoPtr MeshInfo : SelectedItems)
 	{
@@ -579,7 +582,7 @@ TSharedPtr<SWidget> SFbxSceneStaticMeshReimportListView::OnOpenContextMenu()
 	{
 		MenuBuilder.BeginSection("FbxScene_SM_OptionsSection", LOCTEXT("FbxScene_SM_Options", "Options:"));
 		{
-			for (auto OptionName : OverrideNameOptions)
+			for (auto OptionName : (*OverrideNameOptions))
 			{
 				MenuBuilder.AddMenuEntry(FText::FromString(*OptionName.Get()), FText(), FSlateIcon(), FUIAction(FExecuteAction::CreateSP(this, &SFbxSceneStaticMeshReimportListView::AssignToOptions, *OptionName.Get())));
 			}
@@ -589,39 +592,6 @@ TSharedPtr<SWidget> SFbxSceneStaticMeshReimportListView::OnOpenContextMenu()
 	return MenuBuilder.MakeWidget();
 }
 
-TSharedPtr<FString> SFbxSceneStaticMeshReimportListView::FindOptionNameFromName(FString OptionName)
-{
-	for (TSharedPtr<FString> OptionNamePtr : OverrideNameOptions)
-	{
-		if (OptionNamePtr->Compare(OptionName) == 0)
-		{
-			return OptionNamePtr;
-		}
-	}
-	return TSharedPtr<FString>();
-}
-
-void SFbxSceneStaticMeshReimportListView::AssignToOptions(FString OptionName)
-{
-	bool IsDefaultOptions = OptionName.Compare(UFbxSceneImportFactory::DefaultOptionName) == 0;
-	if (!OverrideNameOptionsMap->Contains(OptionName) && !IsDefaultOptions)
-	{
-		return;
-	}
-	TArray<FbxMeshInfoPtr> SelectedItems;
-	GetSelectedItems(SelectedItems);
-	for (FbxMeshInfoPtr MeshInfo : SelectedItems)
-	{
-		EFbxSceneReimportStatusFlags ReimportFlags = *MeshStatusMap->Find(MeshInfo->OriginalImportPath);
-		if ((ReimportFlags & EFbxSceneReimportStatusFlags::Removed) != EFbxSceneReimportStatusFlags::None)
-		{
-			continue;
-		}
-		MeshInfo->OptionName = OptionName;
-	}
-	OptionComboBox->SetSelectedItem(FindOptionNameFromName(OptionName));
-}
-
 void SFbxSceneStaticMeshReimportListView::AssignToStaticMesh()
 {
 	TArray<FbxMeshInfoPtr> SelectedItems;
@@ -701,16 +671,6 @@ void SFbxSceneStaticMeshReimportListView::ResetAssignToStaticMesh()
 	}
 }
 
-void SFbxSceneStaticMeshReimportListView::AddSelectionToImport()
-{
-	SetSelectionImportState(true);
-}
-
-void SFbxSceneStaticMeshReimportListView::RemoveSelectionFromImport()
-{
-	SetSelectionImportState(false);
-}
-
 void SFbxSceneStaticMeshReimportListView::SetSelectionImportState(bool MarkForImport)
 {
 	TArray<FbxMeshInfoPtr> SelectedItems;
@@ -802,7 +762,5 @@ void SFbxSceneStaticMeshReimportListView::UpdateFilterList()
 	}
 	this->RequestListRefresh();
 }
-//////////////////////////////////////////////////////////////////////////
-// Skeletal Mesh List
 
 #undef LOCTEXT_NAMESPACE
diff --git a/Engine/Source/Editor/UnrealEd/Private/Fbx/SSceneReimportStaticMeshListView.h b/Engine/Source/Editor/UnrealEd/Private/Fbx/SSceneReimportStaticMeshListView.h
--- a/Engine/Source/Editor/UnrealEd/Private/Fbx/SSceneReimportStaticMeshListView.h
+++ b/Engine/Source/Editor/UnrealEd/Private/Fbx/SSceneReimportStaticMeshListView.h
@@ -2,9 +2,9 @@
 #pragma once
 
 #include "Classes/Factories/FbxSceneImportFactory.h"
-#include "SSceneImportStaticMeshListView.h"
+#include "SSceneBaseMeshListView.h"
 
-class SFbxSceneStaticMeshReimportListView : public SListView<FbxMeshInfoPtr>
+class SFbxSceneStaticMeshReimportListView : public SFbxSSceneBaseMeshListView
 {
 public:
 
@@ -22,19 +22,17 @@ public:
 		SLATE_ARGUMENT(TSharedPtr<FFbxSceneInfo>, SceneInfoOriginal)
 		SLATE_ARGUMENT(FbxSceneReimportStatusMapPtr, MeshStatusMap)
 		SLATE_ARGUMENT( UnFbx::FBXImportOptions*, GlobalImportSettings)
+		SLATE_ARGUMENT(FbxOverrideNameOptionsArrayPtr, OverrideNameOptions)
 		SLATE_ARGUMENT( ImportOptionsNameMapPtr, OverrideNameOptionsMap)
 		SLATE_ARGUMENT( class UFbxSceneImportOptionsStaticMesh*, SceneImportOptionsStaticMeshDisplay)
 		SLATE_END_ARGS()
 
 	/** Construct this widget */
+	
 	void Construct(const FArguments& InArgs);
 	TSharedRef< ITableRow > OnGenerateRowFbxSceneListView(FbxMeshInfoPtr Item, const TSharedRef< STableViewBase >& OwnerTable);
-	FString FindUniqueOptionName();
-	FReply OnCreateOverrideOptions();
-	void OnToggleSelectAll(ECheckBoxState CheckType);
 
 	void OnFinishedChangingProperties(const FPropertyChangedEvent& PropertyChangedEvent);
-	TSharedPtr<STextComboBox> CreateOverrideOptionComboBox();
 
 	/** Filter show every reimport that will add content */
 	void OnToggleFilterAddContent(ECheckBoxState CheckType);
@@ -65,16 +63,10 @@ public:
 	}
 
 protected:
-	TSharedPtr<FFbxSceneInfo> SceneInfo;
 	TSharedPtr<FFbxSceneInfo> SceneInfoOriginal;
-	UnFbx::FBXImportOptions* GlobalImportSettings;
-	ImportOptionsNameMapPtr OverrideNameOptionsMap;
 	class UFbxSceneImportOptionsStaticMesh* SceneImportOptionsStaticMeshDisplay;
 
-	UnFbx::FBXImportOptions *CurrentStaticMeshImportOptions;
-
 	/** the elements we show in the list view */
-	TArray<FbxMeshInfoPtr> FbxMeshesArray;
 	TArray<FbxMeshInfoPtr> FilterFbxMeshesArray;
 	bool FilterAddContent;
 	bool FilterDeleteContent;
@@ -84,22 +76,14 @@ protected:
 	void UpdateFilterList();
 
 	FbxSceneReimportStatusMapPtr MeshStatusMap;
+	
+	virtual void SetSelectionImportState(bool MarkForImport);
+	virtual void OnToggleSelectAll(ECheckBoxState CheckType);
 
 	/** Open a context menu for the current selection */
 	TSharedPtr<SWidget> OnOpenContextMenu();
 	void AssignToStaticMesh();
 	bool ShowResetAssignToStaticMesh();
 	void ResetAssignToStaticMesh();
-	void AddSelectionToImport();
-	void RemoveSelectionFromImport();
-	void SetSelectionImportState(bool MarkForImport);
-
-	void OnSelectionChanged(FbxMeshInfoPtr Item, ESelectInfo::Type SelectionType);
-	void OnChangedOverrideOptions(TSharedPtr<FString> ItemSelected, ESelectInfo::Type SelectInfo);
-	TSharedPtr<FString> FindOptionNameFromName(FString OptionName);
-	void AssignToOptions(FString OptionName);
-
-	TArray<TSharedPtr<FString>> OverrideNameOptions;
-	TSharedPtr<STextComboBox> OptionComboBox;
-	TSharedPtr<FString> DefaultOptionNamePtr;
+	virtual void OnChangedOverrideOptions(TSharedPtr<FString> ItemSelected, ESelectInfo::Type SelectInfo);
 };
\ No newline at end of file
diff --git a/Engine/Source/Editor/UnrealEd/Private/FbxExporter.h b/Engine/Source/Editor/UnrealEd/Private/FbxExporter.h
--- a/Engine/Source/Editor/UnrealEd/Private/FbxExporter.h
+++ b/Engine/Source/Editor/UnrealEd/Private/FbxExporter.h
@@ -234,6 +234,22 @@ private:
 	*/
 	void ExportLandscapeToFbx(ALandscapeProxy* Landscape, const TCHAR* MeshName, FbxNode* FbxActor, bool bSelectedOnly);
 
+	/**
+	* Fill an fbx light with from a unreal light component
+	*@param ParentNode			The parent FbxNode the one over the light node
+	* @param Camera				Fbx light object
+	* @param CameraComponent	Unreal light component
+	*/
+	void FillFbxLightAttribute(FbxLight* Light, FbxNode* FbxParentNode, ULightComponent* BaseLight);
+
+	/**
+	* Fill an fbx camera with from a unreal camera component
+	* @param ParentNode			The parent FbxNode the one over the camera node
+	* @param Camera				Fbx camera object
+	* @param CameraComponent	Unreal camera component
+	*/
+	void FillFbxCameraAttribute(FbxNode* ParentNode, FbxCamera* Camera, UCameraComponent *CameraComponent);
+
 	/**
 	 * Adds FBX skeleton nodes to the FbxScene based on the skeleton in the given USkeletalMesh, and fills
 	 * the given array with the nodes created
@@ -312,6 +328,7 @@ private:
 	/** recursively get skeleton */
 	void GetSkeleton(FbxNode* RootNode, TArray<FbxNode*>& BoneNodes);
 
+	bool FillFbxTextureProperty(const char *PropertyName, const FExpressionInput& MaterialInput, FbxSurfaceMaterial* FbxMaterial);
 	/**
 	 * Exports the profile_COMMON information for a material.
 	 */
diff --git a/Engine/Source/Editor/UnrealEd/Private/FbxImporter.h b/Engine/Source/Editor/UnrealEd/Private/FbxImporter.h
--- a/Engine/Source/Editor/UnrealEd/Private/FbxImporter.h
+++ b/Engine/Source/Editor/UnrealEd/Private/FbxImporter.h
@@ -120,6 +120,7 @@ struct FBXImportOptions
 	EFBXNormalImportMethod NormalImportMethod;
 	EFBXNormalGenerationMethod::Type NormalGenerationMethod;
 	bool bTransformVertexToAbsolute;
+	bool bBakePivotInVertex;
 	// Static Mesh options
 	bool bCombineToSingle;
 	EVertexColorImportOption::Type VertexColorImportOption;
@@ -159,7 +160,7 @@ struct FBXImportOptions
 	*   This prefix can modify the package path for materials (i.e. TEXT("/Materials/")).
 	*   Or both (i.e. TEXT("/Materials/Mat"))
 	*/
-	FName MaterialPrefixName;
+	FName MaterialBasePath;
 
 	//This data allow to override some fbx Material(point by the uint64 id) with existing unreal material asset
 	TMap<uint64, class UMaterialInterface*> OverrideMaterials;
@@ -203,6 +204,8 @@ struct FbxNodeInfo
 	const char* ObjectName;
 	uint64 UniqueId;
 	FbxAMatrix Transform;
+	FbxVector4 RotationPivot;
+	FbxVector4 ScalePivot;
 	
 	const char* AttributeName;
 	uint64 AttributeUniqueId;
@@ -307,7 +310,7 @@ public:
 	 * @param SceneInfo return the scene info
 	 * @return bool true if get scene info successfully
 	 */
-	bool GetSceneInfo(FString Filename, FbxSceneInfo& SceneInfo);
+	bool GetSceneInfo(FString Filename, FbxSceneInfo& SceneInfo, bool bPreventMaterialNameClash = false);
 
 	/**
 	 * Initialize Fbx file for import.
@@ -324,7 +327,7 @@ public:
 	 * @param Filename
 	 * @return bool
 	 */
-	bool ImportFile(FString Filename);
+	bool ImportFile(FString Filename, bool bPreventMaterialNameClash = false);
 	
 	/**
 	 * Attempt to load an FBX scene from a given filename.
@@ -332,7 +335,7 @@ public:
 	 * @param Filename FBX file name to import.
 	 * @returns true on success.
 	 */
-	UNREALED_API bool ImportFromFile(const FString& Filename, const FString& Type);
+	UNREALED_API bool ImportFromFile(const FString& Filename, const FString& Type, bool bPreventMaterialNameClash = false);
 
 	/**
 	 * Retrieve the FBX loader's error message explaining its failure to read a given FBX file.
@@ -354,6 +357,21 @@ public:
 	 */
 	FbxNode* RetrieveObjectFromName(const TCHAR* ObjectName, FbxNode* Root = NULL);
 
+	/**
+	* Find the first node containing a mesh attribute for the specified LOD index.
+	*
+	* @param NodeLodGroup	The LOD group fbx node
+	* @param LodIndex		The index of the LOD we search the mesh node
+	*/
+	FbxNode* FindLODGroupNode(FbxNode* NodeLodGroup, int32 LodIndex);
+
+	/**
+	* Find the first parent node containing a eLODGroup attribute.
+	*
+	* @param ParentNode		The node where to start the search.
+	*/
+	FbxNode *RecursiveFindParentLodGroup(FbxNode *ParentNode);
+
 	/**
 	 * Creates a static mesh with the given name and flags, imported from within the FBX scene.
 	 * @param InParent
@@ -399,7 +417,7 @@ public:
 	* @param Mesh the original Unreal static mesh object
 	* @return UObject* the new Unreal mesh object
 	*/
-	UStaticMesh* ReimportSceneStaticMesh(uint64 FbxUniqueId, UStaticMesh* Mesh, UFbxStaticMeshImportData* TemplateImportData);
+	UStaticMesh* ReimportSceneStaticMesh(uint64 FbxNodeUniqueId, uint64 FbxMeshUniqueId, UStaticMesh* Mesh, UFbxStaticMeshImportData* TemplateImportData);
 	
 	/**
 	* re-import Unreal skeletal mesh from updated Fbx file
@@ -410,7 +428,7 @@ public:
 	* @param Mesh the original Unreal skeletal mesh object
 	* @return UObject* the new Unreal mesh object
 	*/
-	USkeletalMesh* ReimportSkeletalMesh(USkeletalMesh* Mesh, UFbxSkeletalMeshImportData* TemplateImportData);
+	USkeletalMesh* ReimportSkeletalMesh(USkeletalMesh* Mesh, UFbxSkeletalMeshImportData* TemplateImportData, uint64 SkeletalMeshFbxUID = 0xFFFFFFFFFFFFFFFF, TArray<FbxNode*> *OutSkeletalMeshArray = nullptr);
 
 	/**
 	 * Creates a skeletal mesh from Fbx Nodes with the given name and flags, imported from within the FBX scene.
@@ -623,6 +641,19 @@ public:
 	void BuildFbxMatrixForImportTransform(FbxAMatrix& OutMatrix, UFbxAssetImportData* AssetData);
 
 private:
+	/**
+	* Verify that all meshes are also reference by a fbx hierarchy node. If it found some Geometry
+	* not reference it will add a tokenized error.
+	*/
+	void ValidateAllMeshesAreReferenceByNodeAttribute();
+
+	/**
+	* Recursive search for a node having a mesh attribute
+	*
+	* @param Node	The node from which we start the search for the first node containing a mesh attribute
+	*/
+	FbxNode *RecursiveGetFirstMeshNode(FbxNode* Node);
+
 	/**
 	 * ActorX plug-in can export mesh and dummy as skeleton.
 	 * For the mesh and dummy in the skeleton hierarchy, convert them to FBX skeleton.
@@ -677,6 +708,12 @@ private:
 	*/
 	void VerifyGeometry(UStaticMesh* StaticMesh);
 
+	/**
+	* When there is some materials with the same name we add a clash suffixe _ncl1_x.
+	* Example, if we have 3 materials name shader we will get (shader, shader_ncl1_1, shader_ncl1_2).
+	*/
+	void FixMaterialClashName();
+
 public:
 	// current Fbx scene we are importing. Make sure to release it after import
 	FbxScene* Scene;
diff --git a/Engine/Source/Editor/UnrealEd/Private/FbxMeshUtils.cpp b/Engine/Source/Editor/UnrealEd/Private/FbxMeshUtils.cpp
--- a/Engine/Source/Editor/UnrealEd/Private/FbxMeshUtils.cpp
+++ b/Engine/Source/Editor/UnrealEd/Private/FbxMeshUtils.cpp
@@ -34,13 +34,17 @@ namespace FbxMeshUtils
 		{
 			for (int32 ChildIdx = 0; ChildIdx < Node->GetChildCount(); ++ChildIdx)
 			{
-				if ((LODNodeList.Num() - 1) < ChildIdx)
+				FbxNode *MeshNode = FFbxImporter->FindLODGroupNode(Node, ChildIdx);
+				if (MeshNode != nullptr)
 				{
-					TArray<FbxNode*>* NodeList = new TArray<FbxNode*>;
-					LODNodeList.Add(NodeList);
-				}
+					if ((LODNodeList.Num() - 1) < ChildIdx)
+					{
+						TArray<FbxNode*>* NodeList = new TArray<FbxNode*>;
+						LODNodeList.Add(NodeList);
+					}
 
-				LODNodeList[ChildIdx]->Add(Node->GetChild(ChildIdx));
+					LODNodeList[ChildIdx]->Add(MeshNode);
+				}
 			}
 
 			if (MaxLODCount < (Node->GetChildCount() - 1))
@@ -263,13 +267,19 @@ namespace FbxMeshUtils
 							FbxNode* Node = (*MeshObject)[j];
 							if (Node->GetNodeAttribute() && Node->GetNodeAttribute()->GetAttributeType() == FbxNodeAttribute::eLODGroup)
 							{
+								FbxNode *MeshNode = nullptr;
 								if (Node->GetChildCount() > SelectedLOD)
 								{
-									SkelMeshNodeArray.Add(Node->GetChild(SelectedLOD));
+									MeshNode = FFbxImporter->FindLODGroupNode(Node, SelectedLOD);
 								}
 								else // in less some LODGroups have less level, use the last level
 								{
-									SkelMeshNodeArray.Add(Node->GetChild(Node->GetChildCount() - 1));
+									MeshNode = FFbxImporter->FindLODGroupNode(Node, Node->GetChildCount() - 1);
+								}
+
+								if (MeshNode != nullptr)
+								{
+									SkelMeshNodeArray.Add(MeshNode);
 								}
 							}
 							else
diff --git a/Engine/Source/Editor/UnrealEd/Private/FileHelpers.cpp b/Engine/Source/Editor/UnrealEd/Private/FileHelpers.cpp
--- a/Engine/Source/Editor/UnrealEd/Private/FileHelpers.cpp
+++ b/Engine/Source/Editor/UnrealEd/Private/FileHelpers.cpp
@@ -1464,7 +1464,7 @@ ECommandResult::Type FEditorFileUtils::CheckoutPackages(const TArray<UPackage*>&
 	}
 
 	// If any packages failed the check out process, report them to the user so they know
-	if ( CheckOutResult == ECommandResult::Failed )
+	if ( !PkgsWhichFailedCheckout.IsEmpty() )
 	{
 		FFormatNamedArguments Arguments;
 		Arguments.Add(TEXT("Packages"), FText::FromString( PkgsWhichFailedCheckout ));
@@ -3349,16 +3349,28 @@ void FEditorFileUtils::FindAllSubmittablePackageFiles(TMap<FString, FSourceContr
 	}
 }
 
-void FEditorFileUtils::FindAllConfigFiles(TArray<FString>& OutConfigFiles)
+static void FindAllConfigFilesRecursive(TArray<FString>& OutConfigFiles, const FString& ParentDirectory)
 {
 	TArray<FString> IniFilenames;
-	IFileManager::Get().FindFiles(IniFilenames, *(FPaths::GameConfigDir() / TEXT("*.ini")), true, false);
+	IFileManager::Get().FindFiles(IniFilenames, *(FPaths::GameConfigDir() / ParentDirectory / TEXT("*.ini")), true, false);
 	for (const FString& IniFilename : IniFilenames)
 	{
-		OutConfigFiles.Add(FPaths::ConvertRelativePathToFull(FPaths::GameConfigDir() / IniFilename));
+		OutConfigFiles.Add(FPaths::ConvertRelativePathToFull(FPaths::GameConfigDir() / ParentDirectory / IniFilename));
+	}
+
+	TArray<FString> Subdirectories;
+	IFileManager::Get().FindFiles(Subdirectories, *(FPaths::GameConfigDir() / ParentDirectory / TEXT("*")), false, true);
+	for (const FString& Subdirectory : Subdirectories)
+	{
+		FindAllConfigFilesRecursive(OutConfigFiles, ParentDirectory / Subdirectory);
 	}
 }
 
+void FEditorFileUtils::FindAllConfigFiles(TArray<FString>& OutConfigFiles)
+{
+	FindAllConfigFilesRecursive(OutConfigFiles, FString());
+}
+
 void FEditorFileUtils::FindAllSubmittableConfigFiles(TMap<FString, TSharedPtr<class ISourceControlState, ESPMode::ThreadSafe> >& OutConfigFiles)
 {
 	ISourceControlProvider& SourceControlProvider = ISourceControlModule::Get().GetProvider();
diff --git a/Engine/Source/Editor/UnrealEd/Private/Kismet2/BlueprintEditorUtils.cpp b/Engine/Source/Editor/UnrealEd/Private/Kismet2/BlueprintEditorUtils.cpp
--- a/Engine/Source/Editor/UnrealEd/Private/Kismet2/BlueprintEditorUtils.cpp
+++ b/Engine/Source/Editor/UnrealEd/Private/Kismet2/BlueprintEditorUtils.cpp
@@ -7580,7 +7580,9 @@ bool FBlueprintEditorUtils::PropertyValueFromString(const UProperty* Property, c
 		}
 		else if( Property->IsA(UTextProperty::StaticClass()) )
 		{
-			CastChecked<UTextProperty>(Property)->SetPropertyValue_InContainer(DefaultObject, FText::FromString(Value));
+			FStringOutputDevice ImportError;
+			const auto EndOfParsedBuff = Property->ImportText(*Value, Property->ContainerPtrToValuePtr<uint8>(DefaultObject), 0, nullptr, &ImportError);
+			bParseSucceeded = EndOfParsedBuff && ImportError.IsEmpty();
 		}
 		else
 		{
diff --git a/Engine/Source/Editor/UnrealEd/Private/Kismet2/StructureEditorUtils.cpp b/Engine/Source/Editor/UnrealEd/Private/Kismet2/StructureEditorUtils.cpp
--- a/Engine/Source/Editor/UnrealEd/Private/Kismet2/StructureEditorUtils.cpp
+++ b/Engine/Source/Editor/UnrealEd/Private/Kismet2/StructureEditorUtils.cpp
@@ -590,6 +590,18 @@ const TArray<FStructVariableDescription>& FStructureEditorUtils::GetVarDesc(cons
 	return CastChecked<const UUserDefinedStructEditorData>(Struct->EditorData)->VariablesDescriptions;
 }
 
+TArray<FStructVariableDescription>* FStructureEditorUtils::GetVarDescPtr(UUserDefinedStruct* Struct)
+{
+	check(Struct);
+	return Struct->EditorData ? &CastChecked<UUserDefinedStructEditorData>(Struct->EditorData)->VariablesDescriptions : nullptr;
+}
+
+const TArray<FStructVariableDescription>* FStructureEditorUtils::GetVarDescPtr(const UUserDefinedStruct* Struct)
+{
+	check(Struct);
+	return Struct->EditorData ? &CastChecked<const UUserDefinedStructEditorData>(Struct->EditorData)->VariablesDescriptions : nullptr;
+}
+
 FString FStructureEditorUtils::GetTooltip(const UUserDefinedStruct* Struct)
 {
 	const auto StructEditorData = Struct ? Cast<const UUserDefinedStructEditorData>(Struct->EditorData) : NULL;
diff --git a/Engine/Source/Editor/UnrealEd/Private/LevelEditorViewport.cpp b/Engine/Source/Editor/UnrealEd/Private/LevelEditorViewport.cpp
--- a/Engine/Source/Editor/UnrealEd/Private/LevelEditorViewport.cpp
+++ b/Engine/Source/Editor/UnrealEd/Private/LevelEditorViewport.cpp
@@ -3317,7 +3317,7 @@ static float CheckScaleValue( float ScaleDeltaToCheck, float CurrentScaleFactor,
 		AbsoluteScaleValue = FMath::GridSnap( AbsoluteScaleValue, GEditor->GetScaleGridSize() );
 	}
 	// In some situations CurrentExtent can be 0 (eg: when scaling a plane in Z), this causes a divide by 0 that we need to avoid.
-	if(CurrentExtent < KINDA_SMALL_NUMBER) {
+	if(FMath::Abs(CurrentExtent) < KINDA_SMALL_NUMBER) {
 		return AbsoluteScaleValue;
 	}
 	float UnscaledExtent = CurrentExtent / CurrentScaleFactor;
@@ -3337,71 +3337,120 @@ static float CheckScaleValue( float ScaleDeltaToCheck, float CurrentScaleFactor,
 	return AbsoluteScaleValue;
 }
 
-/** 
- * Helper function for ValidateScale().
- * If the setting is enabled, this function will appropriately re-scale the scale delta so that 
- * proportions are preserved when snapping.
- * @param	CurrentScale	The object's current scale
- * @param	bActiveAxes		The axes that are active when scaling interactively.
- * @param	InOutScaleDelta	The scale delta we are potentially transforming.
- * @return true if the axes should be snapped individually, according to the snap setting (i.e. this function had no effect)
- */
-static bool OptionallyPreserveNonUniformScale(const FVector& InCurrentScale, const bool bActiveAxes[3], FVector& InOutScaleDelta)
+/**
+* Helper function for ValidateScale().
+* If the "PreserveNonUniformScale" setting is enabled, this function will appropriately re-scale the scale delta so that
+* proportions are preserved also when snapping.
+* This function will modify the scale delta sign so that scaling is apply in the good direction when using multiple axis in same time.
+* The function will not transform the scale delta in the case the scale delta is not uniform
+* @param    InOriginalPreDragScale		The object's original scale
+* @param	bActiveAxes					The axes that are active when scaling interactively.
+* @param	InOutScaleDelta				The scale delta we are potentially transforming.
+* @return true if the axes should be snapped individually, according to the snap setting (i.e. this function had no effect)
+*/
+static bool ApplyScalingOptions(const FVector& InOriginalPreDragScale, const bool bActiveAxes[3], FVector& InOutScaleDelta)
 {
+	int ActiveAxisCount = 0;
+	bool CurrentValueSameSign = true;
+	bool FirstSignPositive = true;
+	float MaxComponentSum = -1.0f;
+	int32 MaxAxisIndex = -1;
 	const ULevelEditorViewportSettings* ViewportSettings = GetDefault<ULevelEditorViewportSettings>();
+	bool SnapScaleAfter = ViewportSettings->SnapScaleEnabled;
 
-	if(ViewportSettings->SnapScaleEnabled && ViewportSettings->PreserveNonUniformScale)
+	//Found the number of active axis
+	//Found if we have to swap some sign
+	for (int Axis = 0; Axis < 3; ++Axis)
 	{
-		// when using 'auto-precision', we take the max component & snap its scale, then proportionally scale the other components
-		float MaxComponentSum = -1.0f;
-		int32 MaxAxisIndex = -1;
-		for( int Axis = 0; Axis < 3; ++Axis )
+		if (bActiveAxes[Axis])
 		{
-			if( bActiveAxes[Axis] ) 
+			bool CurrentValueIsZero = FMath::IsNearlyZero(InOriginalPreDragScale[Axis], SMALL_NUMBER);
+			//when the current value is zero we assume it is positive
+			bool IsCurrentValueSignPositive = CurrentValueIsZero ? true : InOriginalPreDragScale[Axis] > 0.0f;
+			if (ActiveAxisCount == 0)
 			{
-				const float AbsScale = FMath::Abs(InOutScaleDelta[Axis] + InCurrentScale[Axis]);
-				if(AbsScale > MaxComponentSum)
+				//Set the first value when we find the first active axis
+				FirstSignPositive = IsCurrentValueSignPositive;
+			}
+			else
+			{
+				if (FirstSignPositive != IsCurrentValueSignPositive)
 				{
-					MaxAxisIndex = Axis;
-					MaxComponentSum = AbsScale;
+					CurrentValueSameSign = false;
 				}
 			}
+			ActiveAxisCount++;
 		}
+	}
 
-		check(MaxAxisIndex != -1);
+	//If we scale more then one axis and
+	//we have to swap some sign
+	if (ActiveAxisCount > 1 && !CurrentValueSameSign)
+	{
+		//Change the scale delta to reflect the sign of the value
+		for (int Axis = 0; Axis < 3; ++Axis)
+		{
+			if (bActiveAxes[Axis])
+			{
+				bool CurrentValueIsZero = FMath::IsNearlyZero(InOriginalPreDragScale[Axis], SMALL_NUMBER);
+				//when the current value is zero we assume it is positive
+				bool IsCurrentValueSignPositive = CurrentValueIsZero ? true : InOriginalPreDragScale[Axis] > 0.0f;
+				InOutScaleDelta[Axis] = IsCurrentValueSignPositive ? InOutScaleDelta[Axis] : -(InOutScaleDelta[Axis]);
+			}
+		}
+	}
 
-		float AbsoluteScaleValue = FMath::GridSnap( InCurrentScale[MaxAxisIndex] + InOutScaleDelta[MaxAxisIndex], GEditor->GetScaleGridSize() );
-		float ScaleRatioMax = InCurrentScale[MaxAxisIndex] == 0.0f ? 1.0f : AbsoluteScaleValue / InCurrentScale[MaxAxisIndex];
-		for( int Axis = 0; Axis < 3; ++Axis )
+	if (ViewportSettings->PreserveNonUniformScale)
+	{
+		for (int Axis = 0; Axis < 3; ++Axis)
 		{
-			if( bActiveAxes[Axis] ) 
+			if (bActiveAxes[Axis])
 			{
-				if(Axis == MaxAxisIndex)
-				{
-					InOutScaleDelta[Axis] = AbsoluteScaleValue - InCurrentScale[Axis];
-				}
-				else
+				const float AbsScale = FMath::Abs(InOutScaleDelta[Axis] + InOriginalPreDragScale[Axis]);
+				if (AbsScale > MaxComponentSum)
 				{
-					InOutScaleDelta[Axis] = (InCurrentScale[Axis] * ScaleRatioMax) - InCurrentScale[Axis];
+					MaxAxisIndex = Axis;
+					MaxComponentSum = AbsScale;
 				}
 			}
 		}
 
-		return false;
+		check(MaxAxisIndex != -1);
+
+		float AbsoluteScaleValue = InOriginalPreDragScale[MaxAxisIndex] + InOutScaleDelta[MaxAxisIndex];
+		if (ViewportSettings->SnapScaleEnabled)
+		{
+			AbsoluteScaleValue = FMath::GridSnap(InOriginalPreDragScale[MaxAxisIndex] + InOutScaleDelta[MaxAxisIndex], GEditor->GetScaleGridSize());
+			SnapScaleAfter = false;
+		}
+
+		float ScaleRatioMax = 1.0f;
+		ScaleRatioMax = AbsoluteScaleValue / InOriginalPreDragScale[MaxAxisIndex];
+		for (int Axis = 0; Axis < 3; ++Axis)
+		{
+			if (bActiveAxes[Axis])
+			{
+				InOutScaleDelta[Axis] = (InOriginalPreDragScale[Axis] * ScaleRatioMax) - InOriginalPreDragScale[Axis];
+			}
+		}
 	}
 
-	return ViewportSettings->SnapScaleEnabled;
+	
+	return SnapScaleAfter;
 }
 
+
 /** Helper function for ModifyScale - Check scale criteria to see if this is allowed */
-void FLevelEditorViewportClient::ValidateScale( const FVector& InCurrentScale, const FVector& InBoxExtent, FVector& InOutScaleDelta, bool bInCheckSmallExtent ) const
+void FLevelEditorViewportClient::ValidateScale(const FVector& InOriginalPreDragScale, const FVector& InCurrentScale, const FVector& InBoxExtent, FVector& InOutScaleDelta, bool bInCheckSmallExtent ) const
 {
 	// get the axes that are active in this operation
 	bool bActiveAxes[3];
 	CheckActiveAxes( Widget != NULL ? Widget->GetCurrentAxis() : EAxisList::None, bActiveAxes );
 
-	bool bSnapAxes = OptionallyPreserveNonUniformScale(InCurrentScale, bActiveAxes, InOutScaleDelta);
-	
+	//When scaling with more then one active axis, We must make sure we apply the correct delta sign to each delta scale axis
+	//We also want to support the PreserveNonUniformScale option
+	bool bSnapAxes = ApplyScalingOptions(InOriginalPreDragScale, bActiveAxes, InOutScaleDelta);
+
 	// check each axis
 	for( int Axis = 0; Axis < 3; ++Axis )
 	{
@@ -3425,7 +3474,14 @@ void FLevelEditorViewportClient::ModifyScale( AActor* InActor, FVector& ScaleDel
 
 		const FBox LocalBox = InActor->GetComponentsBoundingBox( true );
 		const FVector ScaledExtents = LocalBox.GetExtent() * CurrentScale;
-		ValidateScale( CurrentScale, ScaledExtents, ScaleDelta, bCheckSmallExtent );
+		const FTransform* PreDragTransform = PreDragActorTransforms.Find(InActor);
+		//In scale mode we need the predrag transform before the first delta calculation
+		if (PreDragTransform == nullptr)
+		{
+			PreDragTransform = &PreDragActorTransforms.Add(InActor, InActor->GetTransform());
+		}
+		check(PreDragTransform);
+		ValidateScale(PreDragTransform->GetScale3D(), CurrentScale, ScaledExtents, ScaleDelta, bCheckSmallExtent);
 
 		if( ScaleDelta.IsNearlyZero() )
 		{
@@ -3436,7 +3492,17 @@ void FLevelEditorViewportClient::ModifyScale( AActor* InActor, FVector& ScaleDel
 
 void FLevelEditorViewportClient::ModifyScale( USceneComponent* InComponent, FVector& ScaleDelta ) const
 {
-	ValidateScale( InComponent->RelativeScale3D, InComponent->Bounds.GetBox().GetExtent(), ScaleDelta );
+	AActor* Actor = InComponent->GetOwner();
+	const FTransform* PreDragTransform = PreDragActorTransforms.Find(Actor);
+	//In scale mode we need the predrag transform before the first delta calculation
+	if (PreDragTransform == nullptr)
+	{
+		PreDragTransform = &PreDragActorTransforms.Add(Actor, Actor->GetTransform());
+	}
+	check(PreDragTransform);
+	const FBox LocalBox = Actor->GetComponentsBoundingBox(true);
+	const FVector ScaledExtents = LocalBox.GetExtent() * InComponent->RelativeScale3D;
+	ValidateScale(PreDragTransform->GetScale3D(), InComponent->RelativeScale3D, ScaledExtents, ScaleDelta);
 
 	if( ScaleDelta.IsNearlyZero() )
 	{
diff --git a/Engine/Source/Editor/UnrealEd/Private/PlayLevel.cpp b/Engine/Source/Editor/UnrealEd/Private/PlayLevel.cpp
--- a/Engine/Source/Editor/UnrealEd/Private/PlayLevel.cpp
+++ b/Engine/Source/Editor/UnrealEd/Private/PlayLevel.cpp
@@ -2901,6 +2901,7 @@ UGameInstance* UEditorEngine::CreatePIEGameInstance(int32 PIEInstance, bool bInS
 					.ClientSize(FVector2D( NewWindowWidth, NewWindowHeight ))
 					.AutoCenter(CenterNewWindow ? EAutoCenter::PreferredWorkArea : EAutoCenter::None)
 					.UseOSWindowBorder(bUseOSWndBorder)
+					.SaneWindowPlacement(!CenterNewWindow)
 					.SizingRule(ESizingRule::UserSized);
 				}
 
@@ -3008,7 +3009,8 @@ UGameInstance* UEditorEngine::CreatePIEGameInstance(int32 PIEInstance, bool bInS
 
 				// Create a new viewport that the viewport widget will use to render the game
 				SlatePlayInEditorSession.SlatePlayInEditorWindowViewport = MakeShareable( new FSceneViewport( ViewportClient, PieViewportWidget ) );
-				PieViewportWidget->SetViewportInterface( SlatePlayInEditorSession.SlatePlayInEditorWindowViewport.ToSharedRef() );
+				SlatePlayInEditorSession.SlatePlayInEditorWindowViewport->SetPlayInEditorGetsMouseControl(GetDefault<ULevelEditorPlaySettings>()->GameGetsMouseControl);
+				PieViewportWidget->SetViewportInterface(SlatePlayInEditorSession.SlatePlayInEditorWindowViewport.ToSharedRef());
 
 				SlatePlayInEditorSession.SlatePlayInEditorWindow = PieWindow;
 
@@ -3092,6 +3094,9 @@ UGameInstance* UEditorEngine::CreatePIEGameInstance(int32 PIEInstance, bool bInS
 	// Clean up any editor actors being referenced 
 	GEngine->BroadcastLevelActorListChanged();
 
+	// Set an undo barrier so that transactions prior to PIE can't be undone
+	GUnrealEd->ResetTransaction(NSLOCTEXT("UnrealEd", "PIEStarted", "PIE Started"));
+
 	return GameInstance;
 }
 
diff --git a/Engine/Source/Editor/UnrealEd/Private/SFbxSceneOptionWindow.h b/Engine/Source/Editor/UnrealEd/Private/SFbxSceneOptionWindow.h
--- a/Engine/Source/Editor/UnrealEd/Private/SFbxSceneOptionWindow.h
+++ b/Engine/Source/Editor/UnrealEd/Private/SFbxSceneOptionWindow.h
@@ -7,6 +7,8 @@
 #include "Fbx/SSceneImportNodeTreeView.h"
 #include "Fbx/SSceneReimportNodeTreeView.h"
 #include "Fbx/SSceneImportStaticMeshListView.h"
+#include "Fbx/SSceneSkeletalMeshListView.h"
+#include "Fbx/SSceneReimportSkeletalMeshListView.h"
 #include "Fbx/SSceneReimportStaticMeshListView.h"
 #include "Fbx/SSceneMaterialsListView.h"
 #include "Factories/FbxSceneImportFactory.h"
@@ -27,8 +29,6 @@ public:
 		, _SceneImportOptionsStaticMeshDisplay(nullptr)
 		, _OverrideNameOptionsMap(nullptr)
 		, _SceneImportOptionsSkeletalMeshDisplay(nullptr)
-		, _SceneImportOptionsAnimationDisplay(nullptr)
-		, _SceneImportOptionsMaterialDisplay(nullptr)
 		, _OwnerWindow()
 		, _FullPath(TEXT(""))
 		{}
@@ -43,8 +43,6 @@ public:
 		SLATE_ARGUMENT( class UFbxSceneImportOptionsStaticMesh*, SceneImportOptionsStaticMeshDisplay)
 		SLATE_ARGUMENT( ImportOptionsNameMapPtr, OverrideNameOptionsMap)
 		SLATE_ARGUMENT( class UFbxSceneImportOptionsSkeletalMesh*, SceneImportOptionsSkeletalMeshDisplay)
-		SLATE_ARGUMENT( class UFbxSceneImportOptionsAnimation*, SceneImportOptionsAnimationDisplay)
-		SLATE_ARGUMENT( class UFbxSceneImportOptionsMaterial*, SceneImportOptionsMaterialDisplay)
 		SLATE_ARGUMENT( TSharedPtr<SWindow>, OwnerWindow)
 		SLATE_ARGUMENT( FString, FullPath )
 	SLATE_END_ARGS()
@@ -90,11 +88,16 @@ public:
 
 	void OnToggleReimportHierarchy(ECheckBoxState CheckType);
 	ECheckBoxState IsReimportHierarchyChecked() const;
+	void OnToggleBakePivotInVertex(ECheckBoxState CheckType);
+	ECheckBoxState IsBakePivotInVertexChecked() const;
+	void OnToggleInvertNormalMap(ECheckBoxState CheckType);
+	ECheckBoxState IsInvertNormalMapChecked() const;
 
 	//Material UI
-	FText GetMaterialPrefixName() const;
-	void OnMaterialPrefixCommited(const FText& InText, ETextCommit::Type InCommitType);
-	FSlateColor GetMaterialPrefixTextColor() const;
+	FText GetMaterialBasePath() const;
+	void OnMaterialBasePathCommited(const FText& InText, ETextCommit::Type InCommitType);
+	FReply OnMaterialBasePathBrowse();
+	FSlateColor GetMaterialBasePathTextColor() const;
 
 	static void CopyFbxOptionsToFbxOptions(UnFbx::FBXImportOptions *SourceOptions, UnFbx::FBXImportOptions *DestinationOptions);
 
@@ -104,16 +107,12 @@ public:
 	static void CopySkeletalMeshOptionsToFbxOptions(UnFbx::FBXImportOptions *ImportSettings, class UFbxSceneImportOptionsSkeletalMesh* SkeletalMeshOptions);
 	static void CopyFbxOptionsToSkeletalMeshOptions(UnFbx::FBXImportOptions *ImportSettings, class UFbxSceneImportOptionsSkeletalMesh* SkeletalMeshOptions);
 
-	static void CopyAnimationOptionsToFbxOptions(UnFbx::FBXImportOptions *ImportSettings, class UFbxSceneImportOptionsAnimation* AnimationOptions);
-	static void CopyFbxOptionsToAnimationOptions(UnFbx::FBXImportOptions *ImportSettings, class UFbxSceneImportOptionsAnimation* AnimationOptions);
-
-	static void CopyMaterialOptionsToFbxOptions(UnFbx::FBXImportOptions *ImportSettings, class UFbxSceneImportOptionsMaterial* MaterialOptions);
-	static void CopyFbxOptionsToMaterialOptions(UnFbx::FBXImportOptions *ImportSettings, class UFbxSceneImportOptionsMaterial* MaterialOptions);
-
 	void OnFinishedChangingPropertiesSceneTabDetailView(const FPropertyChangedEvent& PropertyChangedEvent);
 
 private:
 
+	bool CanCloseTab();
+
 	bool CanImport() const;
 
 	void InitAllTabs();
@@ -125,6 +124,7 @@ private:
 
 	TSharedRef<SDockTab> SpawnSceneReimportTab(const FSpawnTabArgs& Args);
 	TSharedRef<SDockTab> SpawnStaticMeshReimportTab(const FSpawnTabArgs& Args);
+	TSharedRef<SDockTab> SpawnSkeletalMeshReimportTab(const FSpawnTabArgs& Args);
 
 private:
 	//SFbxSceneOptionWindow Arguments
@@ -138,8 +138,6 @@ private:
 	class UFbxSceneImportOptionsStaticMesh* SceneImportOptionsStaticMeshDisplay;
 	ImportOptionsNameMapPtr OverrideNameOptionsMap;
 	class UFbxSceneImportOptionsSkeletalMesh* SceneImportOptionsSkeletalMeshDisplay;
-	class UFbxSceneImportOptionsAnimation* SceneImportOptionsAnimationDisplay;
-	class UFbxSceneImportOptionsMaterial* SceneImportOptionsMaterialDisplay;
 	TSharedPtr< SWindow > OwnerWindow;
 	FString FullPath;
 
@@ -160,16 +158,27 @@ private:
 	//Material tab Variables
 	TSharedPtr<SFbxSceneMaterialsListView> MaterialsTabListView;
 	FbxTextureInfoArray TexturesArray;
-	FString MaterialPrefixName;
+	FString MaterialBasePath;
+
+	//Shared the options name between staticmesh and skeletalmesh
+	FbxOverrideNameOptionsArray OverrideNameOptions;
 
 	//StaticMesh tab Variables
 	TSharedPtr<SFbxSceneStaticMeshListView> StaticMeshTabListView;
 	TSharedPtr<IDetailsView> StaticMeshTabDetailsView;
 
+	//SkeletalMesh tab Variables
+	TSharedPtr<SFbxSceneSkeletalMeshListView> SkeletalMeshTabListView;
+	TSharedPtr<IDetailsView> SkeletalMeshTabDetailsView;
+
 	//Scene Reimport tab variables
 	TSharedPtr<SFbxReimportSceneTreeView> SceneReimportTreeview;
 
 	//StaticMesh Reimport tab variables
 	TSharedPtr<SFbxSceneStaticMeshReimportListView> StaticMeshReimportListView;
 	TSharedPtr<IDetailsView> StaticMeshReimportDetailsView;
+
+	//SkeletalMesh tab Variables
+	TSharedPtr<SFbxSceneSkeletalMeshReimportListView> SkeletalMeshReimportListView;
+	TSharedPtr<IDetailsView> SkeletalMeshReimportDetailsView;
 };
diff --git a/Engine/Source/Editor/UnrealEd/Private/STransformViewportToolbar.cpp b/Engine/Source/Editor/UnrealEd/Private/STransformViewportToolbar.cpp
--- a/Engine/Source/Editor/UnrealEd/Private/STransformViewportToolbar.cpp
+++ b/Engine/Source/Editor/UnrealEd/Private/STransformViewportToolbar.cpp
@@ -327,7 +327,7 @@ TSharedRef< SWidget > STransformViewportToolBar::MakeTransformToolBar( const TSh
 									.Label(this ,&STransformViewportToolBar::GetScaleGridLabel)
 									.OnGetMenuContent(this, &STransformViewportToolBar::FillScaleGridSnapMenu)
 									.ToggleButtonToolTip(Command->GetDescription())
-									.MenuButtonToolTip(LOCTEXT("ScaleGridSnap_ToolTip", "Set the Scale Grid Snap value"))
+									.MenuButtonToolTip(LOCTEXT("ScaleGridSnap_ToolTip", "Set scaling options"))
 									.Icon(Command->GetIcon())
 									.ParentToolBar(SharedThis(this))
 									, ScaleSnapName );
@@ -724,6 +724,8 @@ TSharedRef<SWidget> STransformViewportToolBar::FillScaleGridSnapMenu()
 
 	FMenuBuilder ScaleGridMenuBuilder( bShouldCloseWindowAfterMenuSelection, CommandList );
 
+	ScaleGridMenuBuilder.BeginSection("ScaleSnapOptions", LOCTEXT("ScaleSnapOptions", "Scale Snap"));
+
 	for( int32 CurGridAmountIndex = 0; CurGridAmountIndex < ViewportSettings->ScalingGridSizes.Num(); ++CurGridAmountIndex )
 	{
 		const float CurGridAmount = ViewportSettings->ScalingGridSizes[ CurGridAmountIndex ];
@@ -752,10 +754,11 @@ TSharedRef<SWidget> STransformViewportToolBar::FillScaleGridSnapMenu()
 			NAME_None,
 			EUserInterfaceActionType::RadioButton );
 	}
+	ScaleGridMenuBuilder.EndSection();
 
 	if( !GEditor->UsePercentageBasedScaling() )
 	{
-		ScaleGridMenuBuilder.AddMenuSeparator();
+		ScaleGridMenuBuilder.BeginSection("ScaleGeneralOptions", LOCTEXT("ScaleOptions", "Scaling Options"));
 
 		ScaleGridMenuBuilder.AddMenuEntry(
 			LOCTEXT("ScaleGridPreserveNonUniformScale", "Preserve Non-Uniform Scale"),
@@ -766,6 +769,8 @@ TSharedRef<SWidget> STransformViewportToolBar::FillScaleGridSnapMenu()
 			FIsActionChecked::CreateStatic( &STransformViewportToolBar::IsPreserveNonUniformScaleChecked ) ),
 			NAME_None,
 			EUserInterfaceActionType::Check );
+
+		ScaleGridMenuBuilder.EndSection();
 	}
 
 
diff --git a/Engine/Source/Editor/UnrealEd/Private/StaticMeshEdit.cpp b/Engine/Source/Editor/UnrealEd/Private/StaticMeshEdit.cpp
--- a/Engine/Source/Editor/UnrealEd/Private/StaticMeshEdit.cpp
+++ b/Engine/Source/Editor/UnrealEd/Private/StaticMeshEdit.cpp
@@ -827,6 +827,7 @@ struct ExistingLODMeshData
 
 struct ExistingStaticMeshData
 {
+	TArray<UMaterialInterface*> ExistingMaterials;
 	FMeshSectionInfoMap			ExistingSectionInfoMap;
 	TArray<ExistingLODMeshData>	ExistingLODData;
 
@@ -836,19 +837,23 @@ struct ExistingStaticMeshData
 	bool						ExistingCustomizedCollision;
 	bool						bAutoComputeLODScreenSize;
 
-	int32							ExistingLightMapResolution;
-	int32							ExistingLightMapCoordinateIndex;
+	int32						ExistingLightMapResolution;
+	int32						ExistingLightMapCoordinateIndex;
 
 	TWeakObjectPtr<UAssetImportData> ExistingImportData;
 	TWeakObjectPtr<UThumbnailInfo> ExistingThumbnailInfo;
 
 	UModel*						ExistingCollisionModel;
-	UBodySetup*				ExistingBodySetup;
+	UBodySetup*					ExistingBodySetup;
 
 	float						ExistingStreamingDistanceMultiplier;
 
 	// A mapping of vertex positions to their color in the existing static mesh
 	TMap<FVector, FColor>		ExistingVertexColorData;
+
+	float						LpvBiasMultiplier;
+	bool						bHasNavigationData;
+	FName						LODGroup;
 };
 
 
@@ -861,6 +866,15 @@ ExistingStaticMeshData* SaveExistingStaticMeshData(UStaticMesh* ExistingMesh, bo
 		ExistingMeshDataPtr = new ExistingStaticMeshData();
 
 		FMeshSectionInfoMap OldSectionInfoMap = ExistingMesh->SectionInfoMap;
+		
+		ExistingMeshDataPtr->ExistingMaterials.Empty();
+		if (bSaveMaterials)
+		{
+			for (UMaterialInterface *MaterialInterface : ExistingMesh->Materials)
+			{
+				ExistingMeshDataPtr->ExistingMaterials.Add(MaterialInterface);
+			}
+		}
 
 		ExistingMeshDataPtr->ExistingLODData.AddZeroed(ExistingMesh->SourceModels.Num());
 
@@ -916,6 +930,10 @@ ExistingStaticMeshData* SaveExistingStaticMeshData(UStaticMesh* ExistingMesh, bo
 		ExistingMeshDataPtr->ExistingBodySetup = ExistingMesh->BodySetup;
 
 		ExistingMeshDataPtr->ExistingStreamingDistanceMultiplier = ExistingMesh->StreamingDistanceMultiplier;
+
+		ExistingMeshDataPtr->LpvBiasMultiplier = ExistingMesh->LpvBiasMultiplier;
+		ExistingMeshDataPtr->bHasNavigationData = ExistingMesh->bHasNavigationData;
+		ExistingMeshDataPtr->LODGroup = ExistingMesh->LODGroup;
 	}
 
 	return ExistingMeshDataPtr;
@@ -925,10 +943,11 @@ void RestoreExistingMeshData(struct ExistingStaticMeshData* ExistingMeshDataPtr,
 {
 	if ( ExistingMeshDataPtr && NewMesh )
 	{
-		int32 NumCommonMaterials = FMath::Min(NewMesh->Materials.Num(), ExistingMeshDataPtr->ExistingLODData[0].ExistingMaterials.Num());
+		//Restore the material array
+		int32 NumCommonMaterials = FMath::Min(NewMesh->Materials.Num(), ExistingMeshDataPtr->ExistingMaterials.Num());
 		for (int32 MaterialIndex = 0; MaterialIndex < NumCommonMaterials; ++MaterialIndex)
 		{
-			NewMesh->Materials[MaterialIndex] = ExistingMeshDataPtr->ExistingLODData[0].ExistingMaterials[MaterialIndex];
+			NewMesh->Materials[MaterialIndex] = ExistingMeshDataPtr->ExistingMaterials[MaterialIndex];
 		}
 
 		int32 NumCommonLODs = FMath::Min<int32>(ExistingMeshDataPtr->ExistingLODData.Num(), NewMesh->SourceModels.Num());
@@ -1019,6 +1038,10 @@ void RestoreExistingMeshData(struct ExistingStaticMeshData* ExistingMeshDataPtr,
 				NewMesh->BodySetup->CollisionTraceFlag = ExistingMeshDataPtr->ExistingBodySetup->CollisionTraceFlag;
 			}
 		}
+
+		NewMesh->LpvBiasMultiplier = ExistingMeshDataPtr->LpvBiasMultiplier;
+		NewMesh->bHasNavigationData = ExistingMeshDataPtr->bHasNavigationData;
+		NewMesh->LODGroup = ExistingMeshDataPtr->LODGroup;
 	}
 
 	delete ExistingMeshDataPtr;
diff --git a/Engine/Source/Editor/UnrealEd/Private/UnrealEdEngine.cpp b/Engine/Source/Editor/UnrealEd/Private/UnrealEdEngine.cpp
--- a/Engine/Source/Editor/UnrealEd/Private/UnrealEdEngine.cpp
+++ b/Engine/Source/Editor/UnrealEd/Private/UnrealEdEngine.cpp
@@ -111,7 +111,7 @@ void UUnrealEdEngine::Init(IEngineLoop* InEngineLoop)
 		SpriteIDToIndexMap.Add( SpriteInfo.Category, InfoIndex );
 	}
 
-	if (FPaths::IsProjectFilePathSet())
+	if (FPaths::IsProjectFilePathSet() && GIsEditor)
 	{
 		AutoReimportManager = NewObject<UAutoReimportManager>();
 		AutoReimportManager->Initialize();
@@ -555,8 +555,9 @@ void UUnrealEdEngine::OnSourceControlStateUpdated(const FSourceControlOperationR
 			}
 			else
 			{
+				bool bIsCooking = CookServer && CookServer->IsCookingInEditor() && CookServer->IsCookByTheBookRunning();
 				// Note when cooking in the editor we ignore package notifications.  The cooker is saving packages in a temp location which generates bogus checkout messages otherwise.
-				if ((SourceControlState->CanCheckout() || !SourceControlState->IsCurrent() || SourceControlState->IsCheckedOutOther()) && (!CookServer || !CookServer->IsCookingInEditor() ) )
+				if ((SourceControlState->CanCheckout() || !SourceControlState->IsCurrent() || SourceControlState->IsCheckedOutOther()) && !bIsCooking )
 				{
 					// To get here, either "prompt for checkout on asset modification" is set, or "automatically checkout on asset modification"
 					// is set, but it failed.
diff --git a/Engine/Source/Editor/UnrealEd/Public/Kismet2/StructureEditorUtils.h b/Engine/Source/Editor/UnrealEd/Public/Kismet2/StructureEditorUtils.h
--- a/Engine/Source/Editor/UnrealEd/Public/Kismet2/StructureEditorUtils.h
+++ b/Engine/Source/Editor/UnrealEd/Public/Kismet2/StructureEditorUtils.h
@@ -107,14 +107,28 @@ public:
 
 	static const TArray<FStructVariableDescription>& GetVarDesc(const UUserDefinedStruct* Struct);
 
+	static TArray<FStructVariableDescription>* GetVarDescPtr(UUserDefinedStruct* Struct);
+
+	static const TArray<FStructVariableDescription>* GetVarDescPtr(const UUserDefinedStruct* Struct);
+
 	static FStructVariableDescription* GetVarDescByGuid(UUserDefinedStruct* Struct, FGuid VarGuid)
 	{
-		return Struct ? GetVarDesc(Struct).FindByPredicate(FFindByGuidHelper<FStructVariableDescription>(VarGuid)) : NULL;
+		if (Struct)
+		{
+			auto VarDescArray = GetVarDescPtr(Struct);
+			return VarDescArray ? VarDescArray->FindByPredicate(FFindByGuidHelper<FStructVariableDescription>(VarGuid)) : nullptr;
+		}
+		return nullptr;
 	}
 
 	static const FStructVariableDescription* GetVarDescByGuid(const UUserDefinedStruct* Struct, FGuid VarGuid)
 	{
-		return Struct ? GetVarDesc(Struct).FindByPredicate(FFindByGuidHelper<FStructVariableDescription>(VarGuid)) : NULL;
+		if (Struct)
+		{
+			auto VarDescArray = GetVarDescPtr(Struct);
+			return VarDescArray ? VarDescArray->FindByPredicate(FFindByGuidHelper<FStructVariableDescription>(VarGuid)) : nullptr;
+		}
+		return nullptr;
 	}
 
 	static FGuid GetGuidForProperty(const UProperty* Property);
diff --git a/Engine/Source/Editor/UnrealEd/Public/LevelEditorViewport.h b/Engine/Source/Editor/UnrealEd/Public/LevelEditorViewport.h
--- a/Engine/Source/Editor/UnrealEd/Public/LevelEditorViewport.h
+++ b/Engine/Source/Editor/UnrealEd/Public/LevelEditorViewport.h
@@ -624,7 +624,7 @@ private:
 
 	/** Helper functions for ApplyDeltaTo* functions - modifies scale based on grid settings */
 	void ModifyScale( AActor* InActor, FVector& ScaleDelta, bool bCheckSmallExtent = false ) const;
-	void ValidateScale( const FVector& CurrentScale, const FVector& BoxExtent, FVector& ScaleDelta, bool bCheckSmallExtent = false ) const;
+	void ValidateScale( const FVector& InOriginalPreDragScale, const FVector& CurrentScale, const FVector& BoxExtent, FVector& ScaleDelta, bool bCheckSmallExtent = false ) const;
 
 	/** Project the specified actors into the world according to the current drag parameters */
 	void ProjectActorsIntoWorld(const TArray<AActor*>& Actors, FViewport* Viewport, const FVector& Drag, const FRotator& Rot);
@@ -699,7 +699,7 @@ private:
 	static TArray< TWeakObjectPtr< AActor > > DropPreviewActors;
 
 	/** A map of actor locations before a drag operation */
-	TMap<TWeakObjectPtr<AActor>, FTransform> PreDragActorTransforms;
+	mutable TMap<TWeakObjectPtr<AActor>, FTransform> PreDragActorTransforms;
 
 	/** Bit array representing the visibility of every sprite category in the current viewport */
 	TBitArray<>	SpriteCategoryVisibility;
diff --git a/Engine/Source/Editor/UnrealEd/Public/UnrealEd.h b/Engine/Source/Editor/UnrealEd/Public/UnrealEd.h
--- a/Engine/Source/Editor/UnrealEd/Public/UnrealEd.h
+++ b/Engine/Source/Editor/UnrealEd/Public/UnrealEd.h
@@ -198,8 +198,6 @@
 #include "Classes/Factories/FbxImportUI.h"
 #include "Classes/Factories/FbxSceneImportData.h"
 #include "Classes/Factories/FbxSceneImportOptions.h"
-#include "Classes/Factories/FbxSceneImportOptionsAnimation.h"
-#include "Classes/Factories/FbxSceneImportOptionsMaterial.h"
 #include "Classes/Factories/FbxSceneImportOptionsSkeletalMesh.h"
 #include "Classes/Factories/FbxSceneImportOptionsStaticMesh.h"
 #include "Classes/Commandlets/FileServerCommandlet.h"
diff --git a/Engine/Source/Runtime/AIModule/Classes/EnvironmentQuery/Tests/EnvQueryTest_Distance.h b/Engine/Source/Runtime/AIModule/Classes/EnvironmentQuery/Tests/EnvQueryTest_Distance.h
--- a/Engine/Source/Runtime/AIModule/Classes/EnvironmentQuery/Tests/EnvQueryTest_Distance.h
+++ b/Engine/Source/Runtime/AIModule/Classes/EnvironmentQuery/Tests/EnvQueryTest_Distance.h
@@ -30,9 +30,10 @@ class UEnvQueryTest_Distance : public UEnvQueryTest
 	/** context */
 	UPROPERTY(EditDefaultsOnly, Category=Distance)
 	TSubclassOf<UEnvQueryContext> DistanceTo;
-
+#if WITH_DEV_AUTOMATION_TESTS
 	virtual void RunTest(FEnvQueryInstance& QueryInstance) const override;
 
 	virtual FText GetDescriptionTitle() const override;
 	virtual FText GetDescriptionDetails() const override;
+#endif
 };
diff --git a/Engine/Source/Runtime/AIModule/Classes/EnvironmentQuery/Tests/EnvQueryTest_Dot.h b/Engine/Source/Runtime/AIModule/Classes/EnvironmentQuery/Tests/EnvQueryTest_Dot.h
--- a/Engine/Source/Runtime/AIModule/Classes/EnvironmentQuery/Tests/EnvQueryTest_Dot.h
+++ b/Engine/Source/Runtime/AIModule/Classes/EnvironmentQuery/Tests/EnvQueryTest_Dot.h
@@ -39,6 +39,7 @@ protected:
 	UPROPERTY(EditDefaultsOnly, Category=Dot)
 	bool bAbsoluteValue;
 
+#if WITH_DEV_AUTOMATION_TESTS
 	virtual void RunTest(FEnvQueryInstance& QueryInstance) const override;
 
 	virtual FText GetDescriptionTitle() const override;
@@ -59,4 +60,5 @@ protected:
 
 	/** helper function: check if contexts are updated per item */
 	bool RequiresPerItemUpdates(TSubclassOf<UEnvQueryContext> LineFrom, TSubclassOf<UEnvQueryContext> LineTo, TSubclassOf<UEnvQueryContext> LineDirection, bool bUseDirectionContext) const;
+#endif
 };
diff --git a/Engine/Source/Runtime/AIModule/Classes/EnvironmentQuery/Tests/EnvQueryTest_GameplayTags.h b/Engine/Source/Runtime/AIModule/Classes/EnvironmentQuery/Tests/EnvQueryTest_GameplayTags.h
--- a/Engine/Source/Runtime/AIModule/Classes/EnvironmentQuery/Tests/EnvQueryTest_GameplayTags.h
+++ b/Engine/Source/Runtime/AIModule/Classes/EnvironmentQuery/Tests/EnvQueryTest_GameplayTags.h
@@ -13,11 +13,13 @@ class UEnvQueryTest_GameplayTags : public UEnvQueryTest
 	GENERATED_UCLASS_BODY()
 
 protected:
+#if WITH_DEV_AUTOMATION_TESTS	
 	virtual void RunTest(FEnvQueryInstance& QueryInstance) const override;
 
 	virtual FText GetDescriptionDetails() const override;
 
 	bool SatisfiesTest(IGameplayTagAssetInterface* ItemGameplayTagAssetInterface) const;
+#endif
 
 	UPROPERTY(EditAnywhere, Category=GameplayTagCheck)
 	EGameplayContainerMatchType TagsToMatch;
diff --git a/Engine/Source/Runtime/AIModule/Classes/EnvironmentQuery/Tests/EnvQueryTest_Trace.h b/Engine/Source/Runtime/AIModule/Classes/EnvironmentQuery/Tests/EnvQueryTest_Trace.h
--- a/Engine/Source/Runtime/AIModule/Classes/EnvironmentQuery/Tests/EnvQueryTest_Trace.h
+++ b/Engine/Source/Runtime/AIModule/Classes/EnvironmentQuery/Tests/EnvQueryTest_Trace.h
@@ -31,12 +31,14 @@ class UEnvQueryTest_Trace : public UEnvQueryTest
 	UPROPERTY(EditDefaultsOnly, Category=Trace)
 	TSubclassOf<UEnvQueryContext> Context;
 
+#if WITH_DEV_AUTOMATION_TESTS
 	virtual void RunTest(FEnvQueryInstance& QueryInstance) const override;
 
 	virtual FText GetDescriptionTitle() const override;
 	virtual FText GetDescriptionDetails() const override;
 
 	virtual void PostLoad() override;
+#endif //#if WITH_DEV_AUTOMATION_TESTS
 
 protected:
 
diff --git a/Engine/Source/Runtime/AIModule/Private/EnvironmentQuery/Tests/EnvQueryTest_Distance.cpp b/Engine/Source/Runtime/AIModule/Private/EnvironmentQuery/Tests/EnvQueryTest_Distance.cpp
--- a/Engine/Source/Runtime/AIModule/Private/EnvironmentQuery/Tests/EnvQueryTest_Distance.cpp
+++ b/Engine/Source/Runtime/AIModule/Private/EnvironmentQuery/Tests/EnvQueryTest_Distance.cpp
@@ -5,6 +5,12 @@
 #include "EnvironmentQuery/Items/EnvQueryItemType_VectorBase.h"
 #include "EnvironmentQuery/Tests/EnvQueryTest_Distance.h"
 
+#if !WITH_DEV_AUTOMATION_TESTS
+
+UEnvQueryTest_Distance::UEnvQueryTest_Distance(const FObjectInitializer& ObjectInitializer) : Super(ObjectInitializer) {}
+
+#else
+
 namespace
 {
 	FORCEINLINE float CalcDistance3D(const FVector& PosA, const FVector& PosB)
@@ -142,3 +148,5 @@ FText UEnvQueryTest_Distance::GetDescriptionDetails() const
 {
 	return DescribeFloatTestParams();
 }
+
+#endif //WITH_DEV_AUTOMATION_TESTS
\ No newline at end of file
diff --git a/Engine/Source/Runtime/AIModule/Private/EnvironmentQuery/Tests/EnvQueryTest_Dot.cpp b/Engine/Source/Runtime/AIModule/Private/EnvironmentQuery/Tests/EnvQueryTest_Dot.cpp
--- a/Engine/Source/Runtime/AIModule/Private/EnvironmentQuery/Tests/EnvQueryTest_Dot.cpp
+++ b/Engine/Source/Runtime/AIModule/Private/EnvironmentQuery/Tests/EnvQueryTest_Dot.cpp
@@ -6,6 +6,12 @@
 #include "EnvironmentQuery/Contexts/EnvQueryContext_Item.h"
 #include "EnvironmentQuery/Tests/EnvQueryTest_Dot.h"
 
+#if !WITH_DEV_AUTOMATION_TESTS
+
+UEnvQueryTest_Dot::UEnvQueryTest_Dot(const FObjectInitializer& ObjectInitializer) : Super(ObjectInitializer) {}
+
+#else
+
 UEnvQueryTest_Dot::UEnvQueryTest_Dot(const FObjectInitializer& ObjectInitializer) : Super(ObjectInitializer)
 {
 	Cost = EEnvTestCost::Low;
@@ -216,3 +222,5 @@ FText UEnvQueryTest_Dot::GetDescriptionDetails() const
 {
 	return DescribeFloatTestParams();
 }
+
+#endif //WITH_DEV_AUTOMATION_TESTS
\ No newline at end of file
diff --git a/Engine/Source/Runtime/AIModule/Private/EnvironmentQuery/Tests/EnvQueryTest_GameplayTags.cpp b/Engine/Source/Runtime/AIModule/Private/EnvironmentQuery/Tests/EnvQueryTest_GameplayTags.cpp
--- a/Engine/Source/Runtime/AIModule/Private/EnvironmentQuery/Tests/EnvQueryTest_GameplayTags.cpp
+++ b/Engine/Source/Runtime/AIModule/Private/EnvironmentQuery/Tests/EnvQueryTest_GameplayTags.cpp
@@ -5,6 +5,12 @@
 #include "EnvironmentQuery/Items/EnvQueryItemType_ActorBase.h"
 #include "GameplayTagAssetInterface.h"
 
+#if !WITH_DEV_AUTOMATION_TESTS
+
+UEnvQueryTest_GameplayTags::UEnvQueryTest_GameplayTags(const FObjectInitializer& ObjectInitializer) : Super(ObjectInitializer) {}
+
+#else
+
 UEnvQueryTest_GameplayTags::UEnvQueryTest_GameplayTags(const FObjectInitializer& ObjectInitializer) : Super(ObjectInitializer)
 {
 	Cost = EEnvTestCost::Low;
@@ -85,3 +91,5 @@ FText UEnvQueryTest_GameplayTags::GetDescriptionDetails() const
 {
 	return GameplayTags.ToMatchingText(TagsToMatch, !BoolValue.DefaultValue);
 }
+
+#endif //WITH_DEV_AUTOMATION_TESTS
\ No newline at end of file
diff --git a/Engine/Source/Runtime/AIModule/Private/EnvironmentQuery/Tests/EnvQueryTest_Overlap.cpp b/Engine/Source/Runtime/AIModule/Private/EnvironmentQuery/Tests/EnvQueryTest_Overlap.cpp
--- a/Engine/Source/Runtime/AIModule/Private/EnvironmentQuery/Tests/EnvQueryTest_Overlap.cpp
+++ b/Engine/Source/Runtime/AIModule/Private/EnvironmentQuery/Tests/EnvQueryTest_Overlap.cpp
@@ -5,6 +5,12 @@
 #include "EnvironmentQuery/Items/EnvQueryItemType_VectorBase.h"
 #include "EnvironmentQuery/Tests/EnvQueryTest_Overlap.h"
 
+#if !WITH_DEV_AUTOMATION_TESTS
+
+UEnvQueryTest_Overlap::UEnvQueryTest_Overlap(const FObjectInitializer& ObjectInitializer) : Super(ObjectInitializer) {}
+
+#else
+
 #define LOCTEXT_NAMESPACE "EnvQueryGenerator"
 
 UEnvQueryTest_Overlap::UEnvQueryTest_Overlap(const FObjectInitializer& ObjectInitializer) : Super(ObjectInitializer)
@@ -126,3 +132,5 @@ bool UEnvQueryTest_Overlap::RunOverlapBlocking(const FVector& ItemPos, const FCo
 }
 
 #undef LOCTEXT_NAMESPACE
+
+#endif //WITH_DEV_AUTOMATION_TESTS
\ No newline at end of file
diff --git a/Engine/Source/Runtime/AIModule/Private/EnvironmentQuery/Tests/EnvQueryTest_Pathfinding.cpp b/Engine/Source/Runtime/AIModule/Private/EnvironmentQuery/Tests/EnvQueryTest_Pathfinding.cpp
--- a/Engine/Source/Runtime/AIModule/Private/EnvironmentQuery/Tests/EnvQueryTest_Pathfinding.cpp
+++ b/Engine/Source/Runtime/AIModule/Private/EnvironmentQuery/Tests/EnvQueryTest_Pathfinding.cpp
@@ -9,6 +9,12 @@
 
 #define LOCTEXT_NAMESPACE "EnvQueryGenerator"
 
+#if !WITH_DEV_AUTOMATION_TESTS
+
+UEnvQueryTest_Pathfinding::UEnvQueryTest_Pathfinding(const FObjectInitializer& ObjectInitializer) : Super(ObjectInitializer) {}
+
+#else
+
 UEnvQueryTest_Pathfinding::UEnvQueryTest_Pathfinding(const FObjectInitializer& ObjectInitializer) : Super(ObjectInitializer)
 {
 	Context = UEnvQueryContext_Querier::StaticClass();
@@ -229,3 +235,5 @@ ANavigationData* UEnvQueryTest_Pathfinding::FindNavigationData(UNavigationSystem
 }
 
 #undef LOCTEXT_NAMESPACE
+
+#endif //WITH_DEV_AUTOMATION_TESTS
\ No newline at end of file
diff --git a/Engine/Source/Runtime/AIModule/Private/EnvironmentQuery/Tests/EnvQueryTest_PathfindingBatch.cpp b/Engine/Source/Runtime/AIModule/Private/EnvironmentQuery/Tests/EnvQueryTest_PathfindingBatch.cpp
--- a/Engine/Source/Runtime/AIModule/Private/EnvironmentQuery/Tests/EnvQueryTest_PathfindingBatch.cpp
+++ b/Engine/Source/Runtime/AIModule/Private/EnvironmentQuery/Tests/EnvQueryTest_PathfindingBatch.cpp
@@ -12,6 +12,12 @@
 #include "AI/Navigation/PImplRecastNavMesh.h"
 #endif
 
+#if !WITH_DEV_AUTOMATION_TESTS
+
+UEnvQueryTest_PathfindingBatch::UEnvQueryTest_PathfindingBatch(const FObjectInitializer& ObjectInitializer) : Super(ObjectInitializer) {}
+
+#else
+
 #define LOCTEXT_NAMESPACE "EnvQueryGenerator"
 
 UEnvQueryTest_PathfindingBatch::UEnvQueryTest_PathfindingBatch(const FObjectInitializer& ObjectInitializer) : Super(ObjectInitializer)
@@ -190,3 +196,5 @@ FText UEnvQueryTest_PathfindingBatch::GetDescriptionTitle() const
 }
 
 #undef LOCTEXT_NAMESPACE
+
+#endif //WITH_DEV_AUTOMATION_TESTS
\ No newline at end of file
diff --git a/Engine/Source/Runtime/AIModule/Private/EnvironmentQuery/Tests/EnvQueryTest_Project.cpp b/Engine/Source/Runtime/AIModule/Private/EnvironmentQuery/Tests/EnvQueryTest_Project.cpp
--- a/Engine/Source/Runtime/AIModule/Private/EnvironmentQuery/Tests/EnvQueryTest_Project.cpp
+++ b/Engine/Source/Runtime/AIModule/Private/EnvironmentQuery/Tests/EnvQueryTest_Project.cpp
@@ -5,6 +5,12 @@
 #include "EnvironmentQuery/Tests/EnvQueryTest_Project.h"
 #include "EnvironmentQuery/EnvQueryTraceHelpers.h"
 
+#if !WITH_DEV_AUTOMATION_TESTS
+
+UEnvQueryTest_Project::UEnvQueryTest_Project(const FObjectInitializer& ObjectInitializer) : Super(ObjectInitializer) {}
+
+#else
+
 UEnvQueryTest_Project::UEnvQueryTest_Project(const FObjectInitializer& ObjectInitializer) : Super(ObjectInitializer)
 {
 	Cost = EEnvTestCost::Medium;
@@ -100,3 +106,5 @@ FText UEnvQueryTest_Project::GetDescriptionDetails() const
 {
 	return DescribeBoolTestParams(TEXT("projected"));
 }
+
+#endif //WITH_DEV_AUTOMATION_TESTS
\ No newline at end of file
diff --git a/Engine/Source/Runtime/AIModule/Private/EnvironmentQuery/Tests/EnvQueryTest_Random.cpp b/Engine/Source/Runtime/AIModule/Private/EnvironmentQuery/Tests/EnvQueryTest_Random.cpp
--- a/Engine/Source/Runtime/AIModule/Private/EnvironmentQuery/Tests/EnvQueryTest_Random.cpp
+++ b/Engine/Source/Runtime/AIModule/Private/EnvironmentQuery/Tests/EnvQueryTest_Random.cpp
@@ -4,6 +4,12 @@
 #include "EnvironmentQuery/Tests/EnvQueryTest_Random.h"
 #include "EnvironmentQuery/Items/EnvQueryItemType.h"
 
+#if !WITH_DEV_AUTOMATION_TESTS
+
+UEnvQueryTest_Random::UEnvQueryTest_Random(const FObjectInitializer& ObjectInitializer) : Super(ObjectInitializer) {}
+
+#else
+
 UEnvQueryTest_Random::UEnvQueryTest_Random(const FObjectInitializer& ObjectInitializer) : Super(ObjectInitializer)
 {
 	Cost = EEnvTestCost::Low;
@@ -37,3 +43,5 @@ FText UEnvQueryTest_Random::GetDescriptionDetails() const
 {
 	return DescribeFloatTestParams();
 }
+
+#endif //WITH_DEV_AUTOMATION_TESTS
\ No newline at end of file
diff --git a/Engine/Source/Runtime/AIModule/Private/EnvironmentQuery/Tests/EnvQueryTest_Trace.cpp b/Engine/Source/Runtime/AIModule/Private/EnvironmentQuery/Tests/EnvQueryTest_Trace.cpp
--- a/Engine/Source/Runtime/AIModule/Private/EnvironmentQuery/Tests/EnvQueryTest_Trace.cpp
+++ b/Engine/Source/Runtime/AIModule/Private/EnvironmentQuery/Tests/EnvQueryTest_Trace.cpp
@@ -5,6 +5,12 @@
 #include "EnvironmentQuery/Items/EnvQueryItemType_VectorBase.h"
 #include "EnvironmentQuery/Tests/EnvQueryTest_Trace.h"
 
+#if !WITH_DEV_AUTOMATION_TESTS
+
+UEnvQueryTest_Trace::UEnvQueryTest_Trace(const FObjectInitializer& ObjectInitializer) : Super(ObjectInitializer) {}
+
+#else
+
 #define LOCTEXT_NAMESPACE "EnvQueryGenerator"
 
 UEnvQueryTest_Trace::UEnvQueryTest_Trace(const FObjectInitializer& ObjectInitializer) : Super(ObjectInitializer)
@@ -186,3 +192,5 @@ bool UEnvQueryTest_Trace::RunCapsuleTraceFrom(const FVector& ItemPos, const FVec
 }
 
 #undef LOCTEXT_NAMESPACE
+
+#endif //WITH_DEV_AUTOMATION_TESTS
\ No newline at end of file
diff --git a/Engine/Source/Runtime/AppFramework/Private/Widgets/Colors/SEyeDropperButton.cpp b/Engine/Source/Runtime/AppFramework/Private/Widgets/Colors/SEyeDropperButton.cpp
--- a/Engine/Source/Runtime/AppFramework/Private/Widgets/Colors/SEyeDropperButton.cpp
+++ b/Engine/Source/Runtime/AppFramework/Private/Widgets/Colors/SEyeDropperButton.cpp
@@ -41,14 +41,14 @@ void SEyeDropperButton::Construct(const FArguments& InArgs)
 	);
 }
 
-EActiveTimerReturnType SEyeDropperButton::ActiveTick(double InCurrentTime, float InDeltaTime)
+void SEyeDropperButton::OnPreTick(float InDeltaTime)
 {
-	if ( bWasClickActivated )
+	if (bWasClickActivated)
 	{
-		if ( HasMouseCapture() && bWasLeft && !bWasReEntered )
+		if (HasMouseCapture() && bWasLeft && !bWasReEntered)
 		{
 			FVector2D CurrentCursorPosition = FSlateApplication::Get().GetCursorPos();
-			if ( CurrentCursorPosition != LastCursorPosition )
+			if (CurrentCursorPosition != LastCursorPosition)
 			{
 				// In dropper mode and outside the button - sample the pixel color and push it to the client
 				// Convert the display gamma into a ratio of gamma from the default gamma.
@@ -60,11 +60,11 @@ EActiveTimerReturnType SEyeDropperButton::ActiveTick(double InCurrentTime, float
 			LastCursorPosition = CurrentCursorPosition;
 		}
 
-		return EActiveTimerReturnType::Continue;
+		return;
 	}
-	
+
 	LastCursorPosition.Reset();
-	return EActiveTimerReturnType::Stop;
+	FSlateApplication::Get().OnPreTick().RemoveAll(this);
 }
 
 FReply SEyeDropperButton::OnMouseButtonDown(const FGeometry& MyGeometry, const FPointerEvent& MouseEvent)
@@ -119,7 +119,7 @@ FReply SEyeDropperButton::OnMouseButtonUp(const FGeometry& MyGeometry, const FPo
 
 		Reply.CaptureMouse(this->AsShared());
 
-		RegisterActiveTimer(0, FWidgetActiveTimerDelegate::CreateSP(this, &SEyeDropperButton::ActiveTick));
+		FSlateApplication::Get().OnPreTick().AddSP(this, &SEyeDropperButton::OnPreTick);
 	}
 	bWasClicked = false;
 
diff --git a/Engine/Source/Runtime/AppFramework/Private/Widgets/Colors/SEyeDropperButton.h b/Engine/Source/Runtime/AppFramework/Private/Widgets/Colors/SEyeDropperButton.h
--- a/Engine/Source/Runtime/AppFramework/Private/Widgets/Colors/SEyeDropperButton.h
+++ b/Engine/Source/Runtime/AppFramework/Private/Widgets/Colors/SEyeDropperButton.h
@@ -46,7 +46,6 @@ public:
 private:
 
 	FReply OnClicked();
-	EActiveTimerReturnType ActiveTick(double InCurrentTime, float InDeltaTime);
 
 	EVisibility GetEscapeTextVisibility() const;
 
@@ -54,6 +53,8 @@ private:
 
 	bool IsMouseOver(const FGeometry& MyGeometry, const FPointerEvent& MouseEvent) const;
 
+	void OnPreTick(float DeltaTime);
+
 private:
 
 	/** Invoked when a new value is selected by the dropper */
diff --git a/Engine/Source/Runtime/Core/Private/GenericPlatform/GenericPlatformMath.cpp b/Engine/Source/Runtime/Core/Private/GenericPlatform/GenericPlatformMath.cpp
--- a/Engine/Source/Runtime/Core/Private/GenericPlatform/GenericPlatformMath.cpp
+++ b/Engine/Source/Runtime/Core/Private/GenericPlatform/GenericPlatformMath.cpp
@@ -77,6 +77,7 @@ void FGenericPlatformMath::FmodReportError(float X, float Y)
 	ensureMsgf(Y != 0, TEXT("FMath::FMod(X=%f, Y=%f) : Y is zero, this is invalid and would result in NaN!"), X, Y);
 }
 
+#if WITH_DEV_AUTOMATION_TESTS
 extern float TheCompilerDoesntKnowThisIsAlwaysZero;
 
 void FGenericPlatformMath::AutoTest() 
@@ -139,3 +140,5 @@ void FGenericPlatformMath::AutoTest()
 		check(Value.ToInt() == 1LL);
 	}
 }
+
+#endif //WITH_DEV_AUTOMATION_TESTS
\ No newline at end of file
diff --git a/Engine/Source/Runtime/Core/Private/Internationalization/Text.cpp b/Engine/Source/Runtime/Core/Private/Internationalization/Text.cpp
--- a/Engine/Source/Runtime/Core/Private/Internationalization/Text.cpp
+++ b/Engine/Source/Runtime/Core/Private/Internationalization/Text.cpp
@@ -1228,7 +1228,7 @@ FText FText::FromString( FString String )
 
 FText FText::AsCultureInvariant( FString String )
 {
-	FText NewText = FText( String );
+	FText NewText = FText( MoveTemp(String) );
 	NewText.Flags |= ETextFlag::CultureInvariant;
 
 	return NewText;
@@ -1236,7 +1236,7 @@ FText FText::AsCultureInvariant( FString String )
 
 FText FText::AsCultureInvariant( FText Text )
 {
-	FText NewText = FText( Text );
+	FText NewText = FText( MoveTemp(Text) );
 	NewText.Flags |= ETextFlag::CultureInvariant;
 
 	return NewText;
@@ -1437,4 +1437,286 @@ FScopedTextIdentityPreserver::~FScopedTextIdentityPreserver()
 	}
 }
 
+bool FTextStringHelper::ReadFromString_ComplexText(const TCHAR* Buffer, FText& OutValue, const TCHAR* Namespace, int32* OutNumCharsRead)
+{
+#define LOC_DEFINE_REGION
+	const TCHAR* const Start = Buffer;
+
+	static const FString InvTextMarker = TEXT("INVTEXT");
+	static const FString NsLocTextMarker = TEXT("NSLOCTEXT");
+	static const FString LocTextMarker = TEXT("LOCTEXT");
+
+	auto ExtractQuotedString = [&](FString& OutStr) -> const TCHAR*
+	{
+		int32 CharsRead = 0;
+		if (!FParse::QuotedString(Buffer, OutStr, &CharsRead))
+		{
+			return nullptr;
+		}
+
+		Buffer += CharsRead;
+		return Buffer;
+	};
+
+	auto WalkToCharacter = [&](const TCHAR InChar) -> const TCHAR*
+	{
+		while (*Buffer && *Buffer != InChar && *Buffer != TCHAR('\n') && *Buffer != TCHAR('\r'))
+		{
+			++Buffer;
+		}
+
+		if (*Buffer != InChar)
+		{
+			return nullptr;
+		}
+
+		return Buffer;
+	};
+
+	#define EXTRACT_QUOTED_STRING(S)		\
+		Buffer = ExtractQuotedString(S);	\
+		if (!Buffer)						\
+		{									\
+			return false;					\
+		}
+
+	#define WALK_TO_CHARACTER(C)			\
+		Buffer = WalkToCharacter(TCHAR(C));	\
+		if (!Buffer)						\
+		{									\
+			return false;					\
+		}
+
+	if (FCString::Strstr(Buffer, *InvTextMarker))
+	{
+		// Parsing something of the form: INVTEXT("...")
+		Buffer += InvTextMarker.Len();
+
+		// Walk to the opening bracket
+		WALK_TO_CHARACTER('(');
+
+		// Walk to the opening quote, and then parse out the quoted string
+		FString InvariantString;
+		WALK_TO_CHARACTER('"');
+		EXTRACT_QUOTED_STRING(InvariantString);
+
+		// Walk to the closing bracket, and then move past it to indicate that the value was successfully imported
+		WALK_TO_CHARACTER(')');
+		++Buffer;
+
+		OutValue = FText::AsCultureInvariant(MoveTemp(InvariantString));
+
+		if (OutNumCharsRead)
+		{
+			*OutNumCharsRead = (Buffer - Start);
+		}
+
+		return true;
+	}
+	else if (FCString::Strstr(Buffer, *NsLocTextMarker))
+	{
+		// Parsing something of the form: NSLOCTEXT("...", "...", "...")
+		Buffer += NsLocTextMarker.Len();
+
+		// Walk to the opening bracket
+		WALK_TO_CHARACTER('(');
+
+		// Walk to the opening quote, and then parse out the quoted namespace
+		FString NamespaceString;
+		WALK_TO_CHARACTER('"');
+		EXTRACT_QUOTED_STRING(NamespaceString);
+
+		// Walk to the opening quote, and then parse out the quoted key
+		FString KeyString;
+		WALK_TO_CHARACTER('"');
+		EXTRACT_QUOTED_STRING(KeyString);
+
+		// Walk to the opening quote, and then parse out the quoted source string
+		FString SourceString;
+		WALK_TO_CHARACTER('"');
+		EXTRACT_QUOTED_STRING(SourceString);
+
+		// Walk to the closing bracket, and then move past it to indicate that the value was successfully imported
+		WALK_TO_CHARACTER(')');
+		++Buffer;
+
+		if (KeyString.IsEmpty())
+		{
+			OutValue = FText::AsCultureInvariant(MoveTemp(SourceString));
+		}
+		else
+		{
+			OutValue = FInternationalization::ForUseOnlyByLocMacroAndGraphNodeTextLiterals_CreateText(*SourceString, *NamespaceString, *KeyString);
+		}
+
+		if (OutNumCharsRead)
+		{
+			*OutNumCharsRead = (Buffer - Start);
+		}
+
+		return true;
+	}
+	else if (FCString::Strstr(Buffer, *LocTextMarker))
+	{
+		// Parsing something of the form: LOCTEXT("...", "...")
+		// This only exists as people sometimes do this in config files. We assume an empty namespace should be used
+		Buffer += LocTextMarker.Len();
+
+		// Walk to the opening bracket
+		WALK_TO_CHARACTER('(');
+
+		// Walk to the opening quote, and then parse out the quoted key
+		FString KeyString;
+		WALK_TO_CHARACTER('"');
+		EXTRACT_QUOTED_STRING(KeyString);
+
+		// Walk to the opening quote, and then parse out the quoted source string
+		FString SourceString;
+		WALK_TO_CHARACTER('"');
+		EXTRACT_QUOTED_STRING(SourceString);
+
+		// Walk to the closing bracket, and then move past it to indicate that the value was successfully imported
+		WALK_TO_CHARACTER(')');
+		++Buffer;
+
+		if (KeyString.IsEmpty())
+		{
+			OutValue = FText::AsCultureInvariant(MoveTemp(SourceString));
+		}
+		else
+		{
+			OutValue = FInternationalization::ForUseOnlyByLocMacroAndGraphNodeTextLiterals_CreateText(*SourceString, (Namespace) ? Namespace : TEXT(""), *KeyString);
+		}
+
+		if (OutNumCharsRead)
+		{
+			*OutNumCharsRead = (Buffer - Start);
+		}
+
+		return true;
+	}
+
+	#undef EXTRACT_QUOTED_STRING
+	#undef WALK_TO_CHARACTER
+#undef LOC_DEFINE_REGION
+
+	return false;
+}
+
+bool FTextStringHelper::ReadFromString(const TCHAR* Buffer, FText& OutValue, const TCHAR* Namespace, int32* OutNumCharsRead, const bool bRequiresQuotes)
+{
+	const TCHAR* const Start = Buffer;
+
+	while (FChar::IsWhitespace(*Buffer))
+	{
+		++Buffer;
+	}
+
+	// First, try and parse the text as a complex text export
+	{
+		int32 SubNumCharsRead = 0;
+		if (FTextStringHelper::ReadFromString_ComplexText(Buffer, OutValue, Namespace, &SubNumCharsRead))
+		{
+			Buffer += SubNumCharsRead;
+			if (OutNumCharsRead)
+			{
+				*OutNumCharsRead = (Buffer - Start);
+			}
+			return true;
+		}
+	}
+
+	// This isn't special text, so just parse it from a string
+	if (bRequiresQuotes)
+	{
+		// Parse out the quoted source string
+		FString LiteralString;
+
+		int32 SubNumCharsRead = 0;
+		if (FParse::QuotedString(Buffer, LiteralString, &SubNumCharsRead))
+		{
+			OutValue = FText::FromString(MoveTemp(LiteralString));
+			Buffer += SubNumCharsRead;
+			if (OutNumCharsRead)
+			{
+				*OutNumCharsRead = (Buffer - Start);
+			}
+			return true;
+		}
+
+		return false;
+	}
+	else
+	{
+		FString LiteralString = Buffer;
+
+		// In order to indicate that the value was successfully imported, advance the buffer past the last character that was imported
+		Buffer += LiteralString.Len();
+
+		OutValue = FText::FromString(MoveTemp(LiteralString));
+
+		if (OutNumCharsRead)
+		{
+			*OutNumCharsRead = (Buffer - Start);
+		}
+		return true;
+	}
+
+	return false;
+}
+
+bool FTextStringHelper::WriteToString(FString& Buffer, const FText& Value, const bool bRequiresQuotes)
+{
+#define LOC_DEFINE_REGION
+	const FString& StringValue = FTextInspector::GetDisplayString(Value);
+
+	if (Value.IsCultureInvariant())
+	{
+		// Produces INVTEXT("...")
+		Buffer += TEXT("INVTEXT(\"");
+		Buffer += StringValue.ReplaceCharWithEscapedChar();
+		Buffer += TEXT("\")");
+	}
+	else
+	{
+		bool bIsLocalized = false;
+		FString Namespace;
+		FString Key;
+		const FString* SourceString = FTextInspector::GetSourceString(Value);
+
+		if (SourceString && Value.ShouldGatherForLocalization())
+		{
+			bIsLocalized = FTextLocalizationManager::Get().FindNamespaceAndKeyFromDisplayString(FTextInspector::GetSharedDisplayString(Value), Namespace, Key);
+		}
+
+		if (bIsLocalized)
+		{
+			// Produces NSLOCTEXT("...", "...", "...")
+			Buffer += TEXT("NSLOCTEXT(\"");
+			Buffer += Namespace.ReplaceCharWithEscapedChar();
+			Buffer += TEXT("\", \"");
+			Buffer += Key.ReplaceCharWithEscapedChar();
+			Buffer += TEXT("\", \"");
+			Buffer += SourceString->ReplaceCharWithEscapedChar();
+			Buffer += TEXT("\")");
+		}
+		else
+		{
+			if (bRequiresQuotes)
+			{
+				Buffer += TEXT("\"");
+				Buffer += StringValue.ReplaceCharWithEscapedChar();
+				Buffer += TEXT("\"");
+			}
+			else
+			{
+				Buffer += StringValue;
+			}
+		}
+	}
+#undef LOC_DEFINE_REGION
+
+	return true;
+}
+
 #undef LOCTEXT_NAMESPACE
diff --git a/Engine/Source/Runtime/Core/Private/Mac/MacPlatformProcess.cpp b/Engine/Source/Runtime/Core/Private/Mac/MacPlatformProcess.cpp
--- a/Engine/Source/Runtime/Core/Private/Mac/MacPlatformProcess.cpp
+++ b/Engine/Source/Runtime/Core/Private/Mac/MacPlatformProcess.cpp
@@ -205,6 +205,91 @@ void FMacPlatformProcess::LaunchURL( const TCHAR* URL, const TCHAR* Parms, FStri
 	}
 }
 
+@interface NSAutoReadPipe : NSObject
+
+/** The pipe itself */
+@property (readonly) NSPipe*			Pipe;
+/** A file associated with the pipe from which we shall read data */
+@property (readonly) NSFileHandle*		File;
+/** Buffer that stores the output from the pipe */
+@property (readonly) NSMutableData*		PipeOutput;
+
+/** Initialization function */
+-(id)init;
+
+/** Deallocation function */
+-(void)dealloc;
+
+/** Callback function that is invoked when data is pushed onto the pipe */
+-(void)readData: (NSNotification *)Notification;
+
+/** Shutdown the background reader, and copy all the data from the pipe as a UTF8 encoded string */
+-(void)copyPipeData: (FString&)OutString;
+
+@end
+
+@implementation NSAutoReadPipe
+
+-(id)init
+{
+	[super init];
+	
+	_PipeOutput = [NSMutableData new];
+	_Pipe = [NSPipe new];
+	_File = [_Pipe fileHandleForReading];
+	
+	[[NSNotificationCenter defaultCenter] addObserver: self
+											selector: @selector(readData:)
+											name: NSFileHandleDataAvailableNotification
+											object: _File];
+	
+	[_File waitForDataInBackgroundAndNotify];
+	return self;
+}
+
+-(void)dealloc
+{
+	[[NSNotificationCenter defaultCenter] removeObserver:self];
+	
+	[_Pipe release];
+	[_PipeOutput release];
+	
+	[super dealloc];
+}
+
+-(void)readData: (NSNotification *)Notification
+{
+	NSFileHandle* FileHandle = (NSFileHandle*)Notification.object;
+	
+	// Ensure we're reading from the right file
+	if (ensure(FileHandle == _File))
+	{
+		[_PipeOutput appendData: [FileHandle availableData]];
+		[FileHandle waitForDataInBackgroundAndNotify];
+	}
+}
+
+-(void)copyPipeData: (FString&)OutString
+{
+	[[NSNotificationCenter defaultCenter] removeObserver:self];
+	
+	// Read any remaining data in from the pipe
+	NSData* Data = [_File readDataToEndOfFile];
+	if (Data && [Data length])
+	{
+		[_PipeOutput appendData: Data];
+	}
+	
+	// Encode the data as a string
+	NSString* String = [[NSString alloc] initWithData:_PipeOutput encoding:NSUTF8StringEncoding];
+	
+	OutString = FString(String);
+	
+	[String release];
+}
+
+@end // NSAutoReadPipe
+
 bool FMacPlatformProcess::ExecProcess( const TCHAR* URL, const TCHAR* Params, int32* OutReturnCode, FString* OutStdOut, FString* OutStdErr )
 {
 	SCOPED_AUTORELEASE_POOL;
@@ -306,53 +391,14 @@ bool FMacPlatformProcess::ExecProcess( const TCHAR* URL, const TCHAR* Params, in
 		
 		[ProcessHandle setArguments: Arguments];
 		
-		NSPipe* StdOutPipe = [[NSPipe new] autorelease];
-		
-		[ProcessHandle setStandardOutput: (id)StdOutPipe];
+		NSAutoReadPipe* StdOutPipe = [[NSAutoReadPipe new] autorelease];
+		[ProcessHandle setStandardOutput: (id)[StdOutPipe Pipe]];
 		
-		NSPipe* StdErrPipe = [[NSPipe new] autorelease];
-		
-		[ProcessHandle setStandardError: (id)StdErrPipe];
+		NSAutoReadPipe* StdErrPipe = [[NSAutoReadPipe new] autorelease];
+		[ProcessHandle setStandardError: (id)[StdErrPipe Pipe]];
 		
 		@try
 		{
-			NSMutableString* StdOutString = [NSMutableString new];
-			NSMutableString* StdErrString = [NSMutableString new];
-			NSFileHandle* StdOutFile = nil;
-			NSFileHandle* StdErrFile = nil;
-			
-			if(OutStdOut)
-			{
-				StdOutFile = [StdOutPipe fileHandleForReading];
-				if(StdOutFile)
-				{
-					StdOutFile.readabilityHandler = ^(NSFileHandle* Handle){
-						NSData* FileData = Handle.availableData;
-						if (FileData.length > 0)
-						{
-							NSString* NewString = (NSString*)[[[NSString alloc] initWithData:FileData encoding:NSUTF8StringEncoding] autorelease];
-							[StdOutString appendString:NewString];
-						}
-					};
-				}
-			}
-			
-			if(OutStdErr)
-			{
-				StdErrFile = [StdErrPipe fileHandleForReading];
-				if(StdErrFile)
-				{
-					StdErrFile.readabilityHandler = ^(NSFileHandle* Handle){
-						NSData* FileData = Handle.availableData;
-						if (FileData.length > 0)
-						{
-							NSString* NewString = (NSString*)[[[NSString alloc] initWithData:FileData encoding:NSUTF8StringEncoding] autorelease];
-							[StdErrString appendString:NewString];
-						}
-					};
-				}
-			}
-
 			[ProcessHandle launch];
 			
 			[ProcessHandle waitUntilExit];
@@ -362,35 +408,16 @@ bool FMacPlatformProcess::ExecProcess( const TCHAR* URL, const TCHAR* Params, in
 				*OutReturnCode = [ProcessHandle terminationStatus];
 			}
 			
-			if (OutStdOut && StdOutFile)
+			if (OutStdOut)
 			{
-				StdOutFile.readabilityHandler = nil;
-				NSData* StdOutData = [StdOutFile readDataToEndOfFile];
-				if (StdOutData.length > 0)
-				{
-					NSString* NewString = (NSString*)[[[NSString alloc] initWithData:StdOutData encoding:NSUTF8StringEncoding] autorelease];
-					[StdOutString appendString:NewString];
-				}
-				
-				*OutStdOut = FString(StdOutString);
+				[StdOutPipe copyPipeData: *OutStdOut];
 			}
 			
-			if (OutStdErr && StdErrFile)
+			if (OutStdErr)
 			{
-				StdErrFile.readabilityHandler = nil;
-				NSData* StdErrData = [StdErrFile readDataToEndOfFile];
-				if (StdErrData.length > 0)
-				{
-					NSString* NewString = (NSString*)[[[NSString alloc] initWithData:StdErrData encoding:NSUTF8StringEncoding] autorelease];
-					[StdErrString appendString:NewString];
-				}
-				
-				*OutStdErr = FString(StdErrString);
+				[StdErrPipe copyPipeData: *OutStdErr];
 			}
 			
-			[StdOutString release];
-			[StdErrString release];
-			
 			return true;
 		}
 		@catch (NSException* Exc)
diff --git a/Engine/Source/Runtime/Core/Private/Math/UnrealMath.cpp b/Engine/Source/Runtime/Core/Private/Math/UnrealMath.cpp
--- a/Engine/Source/Runtime/Core/Private/Math/UnrealMath.cpp
+++ b/Engine/Source/Runtime/Core/Private/Math/UnrealMath.cpp
@@ -1186,20 +1186,25 @@ FQuat FQuat::SlerpFullPath_NotNormalized(const FQuat &quat1, const FQuat &quat2,
 
 FQuat FQuat::Squad(const FQuat& quat1, const FQuat& tang1, const FQuat& quat2, const FQuat& tang2, float Alpha)
 {
-	const FQuat Q1 = FQuat::SlerpFullPath_NotNormalized(quat1, quat2, Alpha);
-	//UE_LOG(LogUnrealMath, Log, TEXT("Q1: %f %f %f %f"), Q1.X, Q1.Y, Q1.Z, Q1.W);
-
+	// Always slerp along the short path from quat1 to quat2 to prevent axis flipping.
+	// This approach is taken by OGRE engine, amongst others.
+	const FQuat Q1 = FQuat::Slerp_NotNormalized(quat1, quat2, Alpha);
 	const FQuat Q2 = FQuat::SlerpFullPath_NotNormalized(tang1, tang2, Alpha);
-	//UE_LOG(LogUnrealMath, Log, TEXT("Q2: %f %f %f %f"), Q2.X, Q2.Y, Q2.Z, Q2.W);
+	const FQuat Result = FQuat::SlerpFullPath(Q1, Q2, 2.f * Alpha * (1.f - Alpha));
 
+	return Result;
+}
+
+FQuat FQuat::SquadFullPath(const FQuat& quat1, const FQuat& tang1, const FQuat& quat2, const FQuat& tang2, float Alpha)
+{
+	const FQuat Q1 = FQuat::SlerpFullPath_NotNormalized(quat1, quat2, Alpha);
+	const FQuat Q2 = FQuat::SlerpFullPath_NotNormalized(tang1, tang2, Alpha);
 	const FQuat Result = FQuat::SlerpFullPath(Q1, Q2, 2.f * Alpha * (1.f - Alpha));
-	//FQuat Result = FQuat::Slerp(Q1, Q2, 2.f * Alpha * (1.f - Alpha));
-	//UE_LOG(LogUnrealMath, Log, TEXT("Result: %f %f %f %f"), Result.X, Result.Y, Result.Z, Result.W);
 
 	return Result;
 }
 
-void FQuat::CalcTangents( const FQuat& PrevP, const FQuat& P, const FQuat& NextP, float Tension, FQuat& OutTan )
+void FQuat::CalcTangents(const FQuat& PrevP, const FQuat& P, const FQuat& NextP, float Tension, FQuat& OutTan)
 {
 	const FQuat InvP = P.Inverse();
 	const FQuat Part1 = (InvP * PrevP).Log();
