Fix non-unity compile

diff --git a/Engine/Source/Editor/UnrealEd/Public/LODCluster.h b/Engine/Source/Editor/UnrealEd/Public/LODCluster.h
--- a/Engine/Source/Editor/UnrealEd/Public/LODCluster.h
+++ b/Engine/Source/Editor/UnrealEd/Public/LODCluster.h
@@ -4,6 +4,7 @@
 
 class ULevel;
 class AActor;
+class ALODActor;
 
 /**
  *
