Sequencer: Fixed non-unity build issues on Mac

diff --git a/Engine/Plugins/MovieScene/LevelSequenceEditor/Source/LevelSequenceEditor/Private/CinematicViewport/SCinematicLevelViewport.h b/Engine/Plugins/MovieScene/LevelSequenceEditor/Source/LevelSequenceEditor/Private/CinematicViewport/SCinematicLevelViewport.h
--- a/Engine/Plugins/MovieScene/LevelSequenceEditor/Source/LevelSequenceEditor/Private/CinematicViewport/SCinematicLevelViewport.h
+++ b/Engine/Plugins/MovieScene/LevelSequenceEditor/Source/LevelSequenceEditor/Private/CinematicViewport/SCinematicLevelViewport.h
@@ -2,14 +2,17 @@
 
 #pragma once
 
+#include "LevelEditorViewport.h"
 #include "SEditorViewport.h"
 
+
 class FLevelSequenceEditorToolkit;
 class ISequencer;
 class SCinematicTransportRange;
 class FLevelViewportLayout;
 struct FTypeInterfaceProxy;
 
+
 /** Overridden level viewport client for this viewport */
 struct FCinematicViewportClient : FLevelEditorViewportClient
 {
@@ -17,6 +20,7 @@ struct FCinematicViewportClient : FLevelEditorViewportClient
 	virtual bool IsLevelEditorClient() const override { return true; }
 };
 
+
 /** struct containing UI data - populated once per frame */
 struct FUIData
 {
@@ -34,6 +38,7 @@ struct FUIData
 	FText MasterEndText;
 };
 
+
 /** Enum specifying different camera modes */
 enum class ECinematicCameraMode
 {
@@ -42,6 +47,7 @@ enum class ECinematicCameraMode
 	SpecificCamera
 };
 
+
 class SCinematicLevelViewport : public SCompoundWidget
 {
 public:
diff --git a/Engine/Source/Editor/UnrealEd/Public/EditorViewportClient.h b/Engine/Source/Editor/UnrealEd/Public/EditorViewportClient.h
--- a/Engine/Source/Editor/UnrealEd/Public/EditorViewportClient.h
+++ b/Engine/Source/Editor/UnrealEd/Public/EditorViewportClient.h
@@ -5,16 +5,26 @@
 #include "UnrealWidget.h"
 #include "Framework/Commands/Commands.h"
 #include "Layout/SlateRect.h"
+#include "Editor.h"
 #include "EditorComponents.h"
 
-class FEditorCameraController;
+
+class FCachedJoystickState;
+class FCameraControllerConfig;
 class FCameraControllerUserImpulseData;
+class FDragTool;
+class FEditorCameraController;
+class FEditorModeTools;
+class FEditorViewportClient;
 class FMouseDeltaTracker;
-class FWidget;
-class FCameraControllerConfig;
-class FCachedJoystickState;
 class FPreviewScene;
-class FEditorModeTools;
+class FSceneView;
+class FWidget;
+class HHitProxy;
+class IMatineeBase;
+class SEditorViewport;
+class UActorFactory;
+
 
 /** Delegate called by FEditorViewportClient to check its visibility */
 DECLARE_DELEGATE_RetVal( bool, FViewportStateGetter );
@@ -106,25 +116,25 @@ private:
 struct FViewportCursorLocation
 {
 public:
-	UNREALED_API FViewportCursorLocation( const class FSceneView* View, class FEditorViewportClient* InViewportClient, int32 X, int32 Y );
+	UNREALED_API FViewportCursorLocation( const FSceneView* View, FEditorViewportClient* InViewportClient, int32 X, int32 Y );
 
 	const FVector&		GetOrigin()			const	{ return Origin; }
 	const FVector&		GetDirection()		const	{ return Direction; }
 	const FIntPoint&	GetCursorPos()		const	{ return CursorPos; }
 	ELevelViewportType	GetViewportType()	const;
-	class FEditorViewportClient*	GetViewportClient()	const	{ return ViewportClient; }
+	FEditorViewportClient*	GetViewportClient()	const	{ return ViewportClient; }
 
 private:
 	FVector	Origin;
 	FVector	Direction;
 	FIntPoint CursorPos;
-	class FEditorViewportClient* ViewportClient;
+	FEditorViewportClient* ViewportClient;
 };
 
 struct FViewportClick : public FViewportCursorLocation
 {
 public:
-	UNREALED_API FViewportClick( const class FSceneView* View, class FEditorViewportClient* ViewportClient, FKey InKey, EInputEvent InEvent, int32 X, int32 Y );
+	UNREALED_API FViewportClick( const FSceneView* View, FEditorViewportClient* ViewportClient, FKey InKey, EInputEvent InEvent, int32 X, int32 Y );
 
 	/** @return The 2D screenspace cursor position of the mouse when it was clicked. */
 	const FIntPoint&	GetClickPos()	const	{ return GetCursorPos(); }
@@ -157,8 +167,6 @@ struct FDropQuery
 	FText HintText;
 };
 
-// Forward declare
-class SEditorViewport;
 
 /** 
  * Stores the transformation data for the viewport camera
@@ -455,7 +463,7 @@ public:
 	 * @param HitX			The X location of the mouse
 	 * @param HitY			The Y location of the mouse
 	 */
-	virtual void ProcessClick(class FSceneView& View, class HHitProxy* HitProxy, FKey Key, EInputEvent Event, uint32 HitX, uint32 HitY);
+	virtual void ProcessClick(FSceneView& View, HHitProxy* HitProxy, FKey Key, EInputEvent Event, uint32 HitX, uint32 HitY);
 	
 	/**
 	 * Called when mouse movement tracking begins
@@ -615,7 +623,7 @@ public:
 	 * @param DragToolType	The type of drag tool to make
 	 * @return The new drag tool
 	 */
-	virtual TSharedPtr<class FDragTool> MakeDragTool( EDragTool::Type DragToolType );
+	virtual TSharedPtr<FDragTool> MakeDragTool( EDragTool::Type DragToolType );
 
 	/** @return true if a drag tool can be used */
 	virtual bool CanUseDragTool() const;
@@ -652,7 +660,7 @@ public:
 	 *
 	 * Returns true if preview actors were updated
 	 */
-	virtual bool UpdateDropPreviewActors(int32 MouseX, int32 MouseY, const TArray<UObject*>& DroppedObjects, bool& out_bDroppedObjectsVisible, class UActorFactory* FactoryToUse = NULL) { return false; }
+	virtual bool UpdateDropPreviewActors(int32 MouseX, int32 MouseY, const TArray<UObject*>& DroppedObjects, bool& out_bDroppedObjectsVisible, UActorFactory* FactoryToUse = NULL) { return false; }
 
 	/**
 	 * If dragging an actor for placement, this function destroys the actor.
@@ -899,7 +907,7 @@ public:
 	/**
 	 * Allows custom disabling of camera recoil
 	 */
-	void SetMatineeRecordingWindow(class IMatineeBase* InInterpEd);
+	void SetMatineeRecordingWindow(IMatineeBase* InInterpEd);
 
 	/**
 	 * Returns true if camera recoil is currently allowed
@@ -1308,7 +1316,7 @@ protected:
 	FMouseDeltaTracker*		MouseDeltaTracker;
 		
 	/**InterpEd, should only be set if used for matinee recording*/
-	class IMatineeBase* RecordingInterpEd;
+	IMatineeBase* RecordingInterpEd;
 
 	/** If true, the canvas has been been moved using bMoveCanvas Mode*/
 	bool bHasMouseMovedSinceClick;
diff --git a/Engine/Source/Editor/UnrealEd/Public/SEditorViewport.h b/Engine/Source/Editor/UnrealEd/Public/SEditorViewport.h
--- a/Engine/Source/Editor/UnrealEd/Public/SEditorViewport.h
+++ b/Engine/Source/Editor/UnrealEd/Public/SEditorViewport.h
@@ -1,9 +1,20 @@
 // Copyright 1998-2016 Epic Games, Inc. All Rights Reserved.
 
-
 #pragma once
 
-class UNREALED_API SEditorViewport : public SCompoundWidget
+#include "Editor.h"
+
+
+class FEditorViewportClient;
+class FSceneViewport;
+class FUICommandList;
+class SOverlay;
+class SViewport;
+class UWorld;
+
+
+class UNREALED_API SEditorViewport
+	: public SCompoundWidget
 {
 public:
 	SLATE_BEGIN_ARGS(SEditorViewport) {}
@@ -58,7 +69,7 @@ protected:
 	virtual TSharedPtr<SWidget> MakeViewportToolbar() { return TSharedPtr<SWidget>(nullptr); }
 
 	// Implement this to add an arbitrary set of toolbars or other overlays to the inside of the viewport
-	virtual void PopulateViewportOverlays(TSharedRef<class SOverlay> Overlay) { }
+	virtual void PopulateViewportOverlays(TSharedRef<SOverlay> Overlay) { }
 
 	virtual void BindCommands();
 	virtual const FSlateBrush* OnGetViewportBorderBrush() const { return NULL; }
@@ -175,7 +186,7 @@ protected:
 	TSharedPtr<SOverlay> ViewportOverlay;
 
 	/** Viewport that renders the scene provided by the viewport client */
-	TSharedPtr<class FSceneViewport> SceneViewport;
+	TSharedPtr<FSceneViewport> SceneViewport;
 	
 	/** Widget where the scene viewport is drawn in */
 	TSharedPtr<SViewport> ViewportWidget;
diff --git a/Engine/Source/Editor/UnrealEd/Public/UnrealWidget.h b/Engine/Source/Editor/UnrealEd/Public/UnrealWidget.h
--- a/Engine/Source/Editor/UnrealEd/Public/UnrealWidget.h
+++ b/Engine/Source/Editor/UnrealEd/Public/UnrealWidget.h
@@ -1,14 +1,19 @@
 // Copyright 1998-2016 Epic Games, Inc. All Rights Reserved.
 
- /*=============================================================================
-	UnrealWidget: Editor widgets for control like 3DS Max
-=============================================================================*/
-
 #pragma once
 
+#include "Editor.h"
+
+
+class FPrimitiveDrawInterface;
+class FEditorModeTools;
 class FEditorViewportClient;
+class FSceneView;
+class UMaterialInterface;
+
 
-class FWidget : public FGCObject
+class FWidget
+	: public FGCObject
 {
 public:
 	enum EWidgetMode
@@ -27,7 +32,7 @@ public:
 	/**
 	 * Sets editor mode tools to use in this widget
 	 */
-	UNREALED_API void SetUsesEditorModeTools(class FEditorModeTools* InEditorModeTools);
+	UNREALED_API void SetUsesEditorModeTools(FEditorModeTools* InEditorModeTools);
 
 	/**
 	 * Renders any widget specific HUD text
@@ -384,7 +389,7 @@ private:
 	FColor PlaneColorXY, ScreenSpaceColor, CurrentColor;
 
 	/** Any mode tools being used */
-	class FEditorModeTools* EditorModeTools;
+	FEditorModeTools* EditorModeTools;
 
 	/**
 	 * An extra matrix to apply to the widget before drawing it (allows for local/custom coordinate systems).
