Copying //UE4/Orion-Staging to //UE4/Main (Origin //Orion/Dev-General @ 2861092)

==========================
MAJOR FEATURES + CHANGES
==========================

Change 2861045 on 2016/02/09 by [REDACTED_PERSON]

	Fix debug editor crash from async compute creating commands when it shouldn't.
	#test debug editor

Change 2861030 on 2016/02/09 by [REDACTED_PERSON]

	Engine: Added support for debugging safe zones (visualization on any platform and simulation on platforms that don't natively provide safe zone information)

	r.DebugSafeZone.Mode controls the debug visualization overlay (0..2, default 0)
	- 0: Do not display the safe zone overlay
	- 1: Display the overlay for the title safe zone
	- 2: Display the overlay for the action safe zone

	r.DebugSafeZone.OverlayAlpha controls how opaque the debug visualization overlay is (0..1, default 0.3)

	On platforms that don't natively support safe zones, you can simulate a safe zone for quick/easy testing in the editor:
	- r.DebugSafeZone.TitleRatio controls the title safe zone margins returned in FDisplayMetrics
	- r.DebugActionZone.ActionRatio controls the action safe zone margins returned in FDisplayMetrics
	- These both range from 0..1, and default to 1 indicating 100% of the display is safe. A typical test value would be 0.9

	#tests Tested on Win64 uncooked and PS4 cooked (front-end and game)

Change 2860923 on 2016/02/09 by [REDACTED_PERSON]

	Fix client warning about HTTPChunkInstaller module not existing
	#tests ran Win64 client

Change 2860852 on 2016/02/09 by [REDACTED_PERSON]

	Fixed crash enabling capsule direct shadows in BP
	#tests Editor

Change 2860842 on 2016/02/09 by [REDACTED_PERSON]

	MallocLeakDetection proxy
	#test PS4/PC testing all commands.

Change 2860744 on 2016/02/09 by [REDACTED_PERSON]

	#UE4 - fixed possible crash when refresh auth with invalid response
	#tests login flow

Change 2860739 on 2016/02/09 by [REDACTED_PERSON]

	Sync Markers
	- Reset SyncGroups every frame.
	- ::GetSyncGroupPosition() makes sure there is a valid MarkerSyncContext.

	=> Fixes SyncGroup returning 'valid' positions for TransitionLeaders that were not in between sync markers.

	#tests new riftmage and kurohane networked in PIE

Change 2860736 on 2016/02/09 by [REDACTED_PERSON]

	Fixed issue with iterative cook on the fly invalidating cooked content all the time.
	#test Cook on the fly iterative ps4

Change 2860598 on 2016/02/09 by [REDACTED_PERSON]

	Simple log category change to match existing log messages in LoadMap

	#test: loading, cooking, game

Change 2860559 on 2016/02/09 by [REDACTED_PERSON]

	#orion - Add flag to AIController to control whether it copies the Pawn rotation to ControlRotation if there is no focus point.

	#tests PIE ded server AI with lanes

Change 2860462 on 2016/02/09 by [REDACTED_PERSON]

	Build system improvements
	* Added details to Empty manifest file save error
	* Removed redundent pseudo-dependencies from -showdependency output
	* Monolithic Kinds now a set and branch hacker can specify kind not to build
	#tests Preflight

Change 2860434 on 2016/02/09 by [REDACTED_PERSON]

	NaN checks:
	-Targeting mode checks in orion code
	-Changed the AActor::SetTransform NaN check so that the logging is included in the NaN ensure (rather than getting cut off from the log afterwards).

	#tests golden path vs bots

Change 2860390 on 2016/02/09 by [REDACTED_PERSON]

	Adjust 3D rendering resolution so that it stays approximately the same when user switches display modes

	#tests Tested editor build on PC

Change 2860364 on 2016/02/09 by [REDACTED_PERSON]

	Removed unused editor-only functions causing compiler errors when compiling the game.

	#tests none

Change 2860242 on 2016/02/09 by [REDACTED_PERSON]

	Made a number of DialogueWave quality of life improvements to the editor and specifically SoundCue editing.

	New right-click option on SoundWaves to create a DialogueWave

diff --git a/Engine/Config/BaseEditor.ini b/Engine/Config/BaseEditor.ini
--- a/Engine/Config/BaseEditor.ini
+++ b/Engine/Config/BaseEditor.ini
@@ -360,6 +360,10 @@ IdleTimeToGC=20
 MaxMemoryAllowance=8192
 bUseFullGCAssetClassNames=true
 +FullGCAssetClassNames=/Script/Engine.World
+
++ExcludedIniParams=Game:/Script/BlueprintContext.McpContext
++ExcludedIniParams=Editor:CookSettings
+
 +VersionedIniParams=Engine:/Script/Engine.RendererSettings:Compat.UseDXT5NormalMaps
 +VersionedIniParams=Engine:/Script/AndroidRuntimeSettings.AndroidRuntimeSettings:AndroidAudio
 +VersionedIniParams=Engine:/Script/UnrealEd.CookerSettings:DefaultPVRTCQuality
diff --git a/Engine/Config/BaseEngine.ini b/Engine/Config/BaseEngine.ini
--- a/Engine/Config/BaseEngine.ini
+++ b/Engine/Config/BaseEngine.ini
@@ -1415,6 +1415,9 @@ bUseStreamingPause=false
 ;Motion Controller refactor
 +ActiveClassRedirects=(OldClassName="EControllerHand", NewClassName="/Script/InputCore.EControllerHand")
 
+; AsyncTask Pin Redirects
++K2AsyncTaskPinRedirects=(ProxyClassName="AbilityTask_PlayMontageAndWait", OldPinName="OnComplete", NewPinName="OnBlendOut")
+
 [CoreUObject.Metadata]
 ; Note: UnrealHeaderTool should be rerun after making changes to MetadataRedirects to catch any keys specified in class headers
 +MetadataRedirects=(OldKey="K2Protected", NewKey="BlueprintProtected")
diff --git a/Engine/Config/BaseHardware.ini b/Engine/Config/BaseHardware.ini
--- a/Engine/Config/BaseHardware.ini
+++ b/Engine/Config/BaseHardware.ini
@@ -6,10 +6,10 @@
 ;
 ; What to do when a certain driver or hardware shows problems:
 ;  - If the driver exhibiting the problem is old, add it to the blacklist with "<=". In general it's better to ask users to upgrade.
-;  - If the driver exhibiting the problem is the latest problem, add it to the blacklist with ">=". We should not assume the next driver version fixes the problem.
+;  - If the driver exhibiting the problem is the latest, add it to the blacklist with ">=". We should not assume the next driver version fixes the problem.
 ;
 ; Expectations on QA:
-;  - QA always tests the latest driver version AND the suggested driver version. In most cases these should be the same thing.
+;  - QA always tests the latest driver version.
 ;  - When a new driver is released, QA should test it and if it does not introduce bugs the SuggestedDriverVersion should be updated to the latest driver.
 ;  - If the latest driver is blacklisted, it is not removed until QA verifies it is fixed.
 ;  - If driver specific bugs are found, enter them in JIRA so that they can be tracked.
@@ -32,17 +32,17 @@
 
 [GPU_NVIDIA]
 ; a newer one is out but has issues (see below)
-SuggestedDriverVersion="359.06"
+SuggestedDriverVersion="361.75"
 
 ; Seems this driver was not setting uninitialized global variables
 ; (static, non const) to 0 but leaving them uninitialized.
 ; Also blacklist future driver versions until QA has verified them to work.
-; NOTE: UE-25096 was fixed in //Orion/Dev-General in CL 2828817, need Paragon QA to verify.
-+Blacklist=(DriverVersion=">=361.43", Reason="UE-25096 Viewport flashes black and white when moving in the scene on latest Nvidia drivers")
+; NOTE: UE-25096 has been fixed and verified by QA.
+;+Blacklist=(DriverVersion=">=361.43", Reason="UE-25096 Viewport flashes black and white when moving in the scene on latest NVIDIA drivers")
 
 ; To prevent problems with older drivers.
 ; We might have to adjust that further.
-;+Blacklist=(DriverVersion="<=347.09", Reason="Crashes with Paragon content using more recent GPU features")
++Blacklist=(DriverVersion="<=347.09", Reason="Crashes with Paragon content using more recent GPU features")
 
 ;-----------------------------------------------------------------------------------------------------------------
 
@@ -50,8 +50,8 @@ SuggestedDriverVersion="359.06"
 // latest as of 1/27/2016 is "15.30.1025.1001 12/17/2015 (Crimson Edition 15.12)"
 SuggestedDriverVersion="15.30.1025.1001 12/17/2015"
 
-// 14.201.1001.1004 (7/31/14) - 13.152.1.1000 (9/10/13) are the AMD drivers that seem to be causing crashes when we try to compile a compute shader.
-+Blacklist=(DriverVersion="<=14.201.1001.1004", Reason="Crashes with Paragon content using more recent GPU features")
+// 14.301.1001.0 - 13.152.1.1000 are the AMD drivers that seem to be causing crashes when we try to compile a compute shader.
++Blacklist=(DriverVersion="<=14.301.1001.0", Reason="Crashes with Paragon content using more recent GPU features")
 
 ;-----------------------------------------------------------------------------------------------------------------
 
diff --git a/Engine/Plugins/Experimental/SignificanceManager/Source/SignificanceManager/Private/SignificanceManager.cpp b/Engine/Plugins/Experimental/SignificanceManager/Source/SignificanceManager/Private/SignificanceManager.cpp
--- a/Engine/Plugins/Experimental/SignificanceManager/Source/SignificanceManager/Private/SignificanceManager.cpp
+++ b/Engine/Plugins/Experimental/SignificanceManager/Source/SignificanceManager/Private/SignificanceManager.cpp
@@ -15,8 +15,6 @@ DECLARE_CYCLE_STAT(TEXT("Register Object"), STAT_SignificanceManager_RegisterObj
 DECLARE_CYCLE_STAT(TEXT("Initial Significance Update"), STAT_SignificanceManager_InitialSignificanceUpdate, STATGROUP_SignificanceManager);
 DECLARE_DWORD_ACCUMULATOR_STAT(TEXT("Num Managed Objects"), STAT_SignificanceManager_NumObjects, STATGROUP_SignificanceManager);
 
-const FSignificanceNotify GDummySignificanceNotify;
-
 bool CompareBySignificanceAscending(const USignificanceManager::FManagedObjectInfo& A, const USignificanceManager::FManagedObjectInfo& B) 
 { 
 	return A.GetSignificance() < B.GetSignificance(); 
@@ -109,7 +107,7 @@ void USignificanceManager::BeginDestroy()
 {
 	Super::BeginDestroy();
 
-	for (const TPair<const UObject*, FManagedObjectInfo*>& ObjectToObjectInfoPair : ManagedObjects)
+	for (const TPair<UObject*, FManagedObjectInfo*>& ObjectToObjectInfoPair : ManagedObjects)
 	{
 		delete ObjectToObjectInfoPair.Value;
 	}
@@ -117,12 +115,7 @@ void USignificanceManager::BeginDestroy()
 	ManagedObjectsByTag.Reset();
 }
 
-void USignificanceManager::RegisterObject(const UObject* Object, const FName Tag, FSignificanceFunction SignificanceFunction)
-{
-	RegisterObject(Object, Tag, SignificanceFunction, GDummySignificanceNotify);
-}
-
-void USignificanceManager::RegisterObject(const UObject* Object, const FName Tag, FSignificanceFunction SignificanceFunction, const FSignificanceNotify& SignificanceNotifyDelegate)
+void USignificanceManager::RegisterObject(UObject* Object, const FName Tag, FSignificanceFunction SignificanceFunction, FPostSignificanceFunction InPostSignificanceFunction)
 {
 	INC_DWORD_STAT(STAT_SignificanceManager_NumObjects);
 	SCOPE_CYCLE_COUNTER(STAT_SignificanceManager_RegisterObject);
@@ -130,7 +123,7 @@ void USignificanceManager::RegisterObject(const UObject* Object, const FName Tag
 	check(Object);
 	checkf(!ManagedObjects.Contains(Object), TEXT("'%s' already added to significance manager. Original Tag: '%s' New Tag: '%s'"), *Object->GetName(), *ManagedObjects.FindChecked(Object)->GetTag().ToString(), *Tag.ToString());
 
-	FManagedObjectInfo* ObjectInfo = new FManagedObjectInfo(Object, Tag, SignificanceFunction, SignificanceNotifyDelegate);
+	FManagedObjectInfo* ObjectInfo = new FManagedObjectInfo(Object, Tag, SignificanceFunction, InPostSignificanceFunction);
 	
 	// Calculate initial significance
 	if (Viewpoints.Num())
@@ -180,7 +173,7 @@ void USignificanceManager::RegisterObject(const UObject* Object, const FName Tag
 	}
 }
 
-void USignificanceManager::UnregisterObject(const UObject* Object)
+void USignificanceManager::UnregisterObject(UObject* Object)
 {
 	DEC_DWORD_STAT(STAT_SignificanceManager_NumObjects);
 	SCOPE_CYCLE_COUNTER(STAT_SignificanceManager_UnregisterObject);
@@ -209,6 +202,7 @@ void USignificanceManager::UnregisterAll(FName Tag)
 		for (const FManagedObjectInfo* ManagedObj : *ObjectsWithTag)
 		{
 			ManagedObjects.Remove(ManagedObj->GetObject());
+			ManagedObj->PostSignificanceFunction(ManagedObj->GetObject(), ManagedObj->Significance, 1.0f);
 		}
 		ManagedObjectsByTag.Remove(Tag);
 	}
@@ -294,7 +288,7 @@ void USignificanceManager::FManagedObjectInfo::UpdateSignificance(const TArray<F
 		Significance = 0.f;
 	}
 
-	SignificanceNotifyDelegate.ExecuteIfBound(OldSignificance, Significance);
+	PostSignificanceFunction(Object, OldSignificance, Significance);
 }
 
 void USignificanceManager::Update(const TArray<FTransform>& InViewpoints)
@@ -306,7 +300,7 @@ void USignificanceManager::Update(const TArray<FTransform>& InViewpoints)
 	{
 		SCOPE_CYCLE_COUNTER(STAT_SignificanceManager_SignificanceUpdate);
 		//TODO: Consider parallel for
-		for (TPair<const UObject*, FManagedObjectInfo*>& ObjectToObjectInfoPair : ManagedObjects)
+		for (TPair<UObject*, FManagedObjectInfo*>& ObjectToObjectInfoPair : ManagedObjects)
 		{
 			const UObject* Object = ObjectToObjectInfoPair.Key;
 			FManagedObjectInfo* ObjectInfo = ObjectToObjectInfoPair.Value;
diff --git a/Engine/Plugins/Experimental/SignificanceManager/Source/SignificanceManager/Public/SignificanceManager.h b/Engine/Plugins/Experimental/SignificanceManager/Source/SignificanceManager/Public/SignificanceManager.h
--- a/Engine/Plugins/Experimental/SignificanceManager/Source/SignificanceManager/Public/SignificanceManager.h
+++ b/Engine/Plugins/Experimental/SignificanceManager/Source/SignificanceManager/Public/SignificanceManager.h
@@ -45,9 +45,6 @@ private:
 	static TSubclassOf<USignificanceManager>  SignificanceManagerClass;
 };
 
-DECLARE_DELEGATE_TwoParams(FSignificanceNotify, float, float);
-extern SIGNIFICANCEMANAGER_API const FSignificanceNotify GDummySignificanceNotify;
-
 /* The significance manager provides a framework for registering objects by tag to each have a significance
  * value calculated from which a game specific subclass and game logic can make decisions about what level
  * of detail objects should be at, tick frequency, whether to spawn effects, and other such functionality
@@ -64,39 +61,40 @@ class SIGNIFICANCEMANAGER_API USignificanceManager : public UObject
 	GENERATED_BODY()
 
 public:
-	typedef TFunction<float(const UObject*,const FTransform&)> FSignificanceFunction;
+	typedef TFunction<float(UObject*, const FTransform&)> FSignificanceFunction;
+	typedef TFunction<void(UObject*, float, float)> FPostSignificanceFunction;
 
 	struct FManagedObjectInfo
 	{
 		FManagedObjectInfo()
 			: Object(nullptr)
-			, Significance(0.f)
+			, Significance(-1.0f)
 		{
 		}
 
-		FManagedObjectInfo(const UObject* InObject, FName InTag, FSignificanceFunction InSignificanceFunction, const FSignificanceNotify& Notify)
+		FManagedObjectInfo(UObject* InObject, FName InTag, FSignificanceFunction InSignificanceFunction, FPostSignificanceFunction InPostSignificanceFunction)
 			: Object(InObject)
 			, Tag(InTag)
-			, Significance(0.f)
+			, Significance(-1.0f)
 			, SignificanceFunction(InSignificanceFunction)
-			, SignificanceNotifyDelegate(Notify)
+			, PostSignificanceFunction(InPostSignificanceFunction)
 		{
 		}
 
-		const UObject* GetObject() const { return Object; }
+		UObject* GetObject() const { return Object; }
 		FName GetTag() const { return Tag; }
 		float GetSignificance() const { return Significance; }
 		FSignificanceFunction GetSignificanceFunction() const { return SignificanceFunction; }
-		FSignificanceNotify GetSignificanceNotifyDelegate() const { return SignificanceNotifyDelegate; }	
+		FPostSignificanceFunction GetPostSignificanceNotifyDelegate() const { return PostSignificanceFunction; }
 
 	private:
-		const UObject* Object;
+		UObject* Object;
 		FName Tag;
 		float Significance;
 
 		FSignificanceFunction SignificanceFunction;
 
-		FSignificanceNotify SignificanceNotifyDelegate;
+		FPostSignificanceFunction PostSignificanceFunction;
 
 		void UpdateSignificance(const TArray<FTransform>& ViewPoints);
 
@@ -114,11 +112,10 @@ public:
 	virtual void Update(const TArray<FTransform>& Viewpoints);
 
 	// Overridable function used to register an object as managed by the significance manager
-	virtual void RegisterObject(const UObject* Object, FName Tag, FSignificanceFunction SignificanceFunction);
-	virtual void RegisterObject(const UObject* Object, FName Tag, FSignificanceFunction SignificanceFunction, const FSignificanceNotify& SignificanceNotifyDelegate);
+	virtual void RegisterObject(UObject* Object, FName Tag, FSignificanceFunction SignificanceFunction, FPostSignificanceFunction InPostSignificanceFunction);
 	
 	// Overridable function used to unregister an object as managed by the significance manager
-	virtual void UnregisterObject(const UObject* Object);
+	virtual void UnregisterObject(UObject* Object);
 
 	// Unregisters all objects with the specified tag.
 	void UnregisterAll(FName Tag);
@@ -174,7 +171,7 @@ private:
 	TMap<FName, TArray<const FManagedObjectInfo*>> ManagedObjectsByTag;
 
 	// Reverse lookup map to find the tag for a given object
-	TMap<const UObject*, FManagedObjectInfo*> ManagedObjects;
+	TMap<UObject*, FManagedObjectInfo*> ManagedObjects;
 
 	// Game specific significance class to instantiate
 	UPROPERTY(globalconfig, noclear, EditAnywhere, Category=DefaultClasses, meta=(MetaClass="SignificanceManager", DisplayName="Significance Manager Class"))
diff --git a/Engine/Shaders/BasePassPixelShader.usf b/Engine/Shaders/BasePassPixelShader.usf
--- a/Engine/Shaders/BasePassPixelShader.usf
+++ b/Engine/Shaders/BasePassPixelShader.usf
@@ -823,6 +823,7 @@ void Main(
 		GBuffer.ShadingModelID = SHADINGMODELID_CLOTH;
 		GBuffer.CustomData.rgb = SubsurfaceColor;
 		GBuffer.CustomData.a = saturate( GetMaterialCustomData0(MaterialParameters) );	// Cloth
+		GBuffer.IndirectIrradiance *= 1 - GBuffer.CustomData.a;
 	#elif MATERIAL_SHADINGMODEL_EYE
 		GBuffer.ShadingModelID = SHADINGMODELID_EYE;
 		#if NUM_MATERIAL_OUTPUTS_GETTANGENTOUTPUT > 0
diff --git a/Engine/Shaders/BasePassVertexShader.usf b/Engine/Shaders/BasePassVertexShader.usf
--- a/Engine/Shaders/BasePassVertexShader.usf
+++ b/Engine/Shaders/BasePassVertexShader.usf
@@ -109,7 +109,8 @@ void Main(
 
 	//@todo - get from VF
 	float3 LightingPositionOffset = 0;
-	ComputeVolumeUVs(WorldPosition.xyz - ResolvedView.PreViewTranslation.xyz, LightingPositionOffset, InnerVolumeUVs, OuterVolumeUVs, FinalLerpFactor);
+	float3 WorldPositionForVertexLighting = VertexFactoryGetPositionForVertexLighting(Input, VFIntermediates, WorldPosition.xyz);
+	ComputeVolumeUVs(WorldPositionForVertexLighting - ResolvedView.PreViewTranslation.xyz, LightingPositionOffset, InnerVolumeUVs, OuterVolumeUVs, FinalLerpFactor);
 
 	#if TRANSLUCENCY_LIGHTING_VOLUMETRIC_PERVERTEX_DIRECTIONAL
 	
diff --git a/Engine/Shaders/GpuSkinVertexFactory.usf b/Engine/Shaders/GpuSkinVertexFactory.usf
--- a/Engine/Shaders/GpuSkinVertexFactory.usf
+++ b/Engine/Shaders/GpuSkinVertexFactory.usf
@@ -522,6 +522,11 @@ float4 VertexFactoryGetRasterizedWorldPosition(FVertexFactoryInput Input, FVerte
 	return InWorldPosition;
 }
 
+float3 VertexFactoryGetPositionForVertexLighting(FVertexFactoryInput Input, FVertexFactoryIntermediates Intermediates, float3 TranslatedWorldPosition)
+{
+	return TranslatedWorldPosition;
+}
+
 void CalcTangentToWorld(FVertexFactoryInput Input, FVertexFactoryIntermediates Intermediates, out float3 TangentToWorld0, out float4 TangentToWorld2)
 {
 	float3x3 LocalToWorld = GetLocalToWorld3x3();
diff --git a/Engine/Shaders/LandscapeVertexFactory.usf b/Engine/Shaders/LandscapeVertexFactory.usf
--- a/Engine/Shaders/LandscapeVertexFactory.usf
+++ b/Engine/Shaders/LandscapeVertexFactory.usf
@@ -692,6 +692,11 @@ float4 VertexFactoryGetRasterizedWorldPosition(FVertexFactoryInput Input, FVerte
 #endif
 }
 
+float3 VertexFactoryGetPositionForVertexLighting(FVertexFactoryInput Input, FVertexFactoryIntermediates Intermediates, float3 TranslatedWorldPosition)
+{
+	return TranslatedWorldPosition;
+}
+
 FVertexFactoryInterpolantsVSToPS VertexFactoryGetInterpolantsVSToPS(FVertexFactoryInput Input, FVertexFactoryIntermediates Intermediates, FMaterialVertexParameters VertexParameters)
 {
 	FVertexFactoryInterpolantsVSToPS Interpolants;
diff --git a/Engine/Shaders/ParticleSpriteVertexFactory.usf b/Engine/Shaders/ParticleSpriteVertexFactory.usf
--- a/Engine/Shaders/ParticleSpriteVertexFactory.usf
+++ b/Engine/Shaders/ParticleSpriteVertexFactory.usf
@@ -401,7 +401,9 @@ float3x3 CalcTangentBasis(FVertexFactoryIntermediates Intermediates)
 #endif
 }
 
-#if FEATURE_LEVEL >= FEATURE_LEVEL_SM4
+// Note: C++ geometry setup behavior has to match this define
+#define SUPPORTS_PARTICLE_CUTOUTS (FEATURE_LEVEL >= FEATURE_LEVEL_SM4)
+#if SUPPORTS_PARTICLE_CUTOUTS
 
 /** Number of vertices in each SubUV frame of CutoutGeometry. */
 uint NumCutoutVerticesPerFrame;
@@ -416,7 +418,7 @@ void ComputeBillboardUVs(FVertexFactoryInput Input, out float2 UVForPosition, ou
 	// Encoding the UV flip in the sign of the size data.
 	float2 UVFlip = sign(Input.SizeRotSubImage.xy);
 
-#if FEATURE_LEVEL >= FEATURE_LEVEL_SM4
+#if SUPPORTS_PARTICLE_CUTOUTS
 	BRANCH
 	if (NumCutoutVerticesPerFrame > 0)
 	{
@@ -556,6 +558,16 @@ float4 VertexFactoryGetRasterizedWorldPosition(FVertexFactoryInput Input, FVerte
 #endif
 }
 
+float3 VertexFactoryGetPositionForVertexLighting(FVertexFactoryInput Input, FVertexFactoryIntermediates Intermediates, float3 TranslatedWorldPosition)
+{
+#if SUPPORTS_PARTICLE_CUTOUTS && (TRANSLUCENCY_LIGHTING_VOLUMETRIC_PERVERTEX_DIRECTIONAL || TRANSLUCENCY_LIGHTING_VOLUMETRIC_PERVERTEX_NONDIRECTIONAL)
+	// Do per-vertex lighting with particle position instead of vertex position when we're using SubUV cutouts, since the vertex positions vary per-frame which would cause popping
+	return NumCutoutVerticesPerFrame > 0 ? (Intermediates.ParticleWorldPosition.xyz + ResolvedView.PreViewTranslation.xyz) : TranslatedWorldPosition;
+#else
+	return TranslatedWorldPosition;
+#endif
+}
+
 FVertexFactoryInterpolantsVSToPS VertexFactoryGetInterpolantsVSToPS(FVertexFactoryInput Input, FVertexFactoryIntermediates Intermediates, FMaterialVertexParameters VertexParameters)
 {
 	FVertexFactoryInterpolantsVSToPS Interpolants;
@@ -656,4 +668,4 @@ uint VertexFactoryGetEyeIndex(uint InstanceId)
 {
 	return InstancedEyeIndex & bIsInstancedStereo;
 }
-#endif
\ No newline at end of file
+#endif
diff --git a/Engine/Source/Developer/AssetTools/AssetTools.Build.cs b/Engine/Source/Developer/AssetTools/AssetTools.Build.cs
--- a/Engine/Source/Developer/AssetTools/AssetTools.Build.cs
+++ b/Engine/Source/Developer/AssetTools/AssetTools.Build.cs
@@ -21,6 +21,7 @@ public class AssetTools : ModuleRules
 				"SourceControl",
 				"TextureEditor",
 				"UnrealEd",
+				"PropertyEditor",
 				"Kismet",
 				"Landscape",
                 "Foliage",
@@ -38,6 +39,8 @@ public class AssetTools : ModuleRules
 				"DesktopPlatform",
 				"EditorWidgets",
 				"GameProjectGeneration",
+                "PropertyEditor",
+                "ActorPickerMode",
 				"Kismet",
 				"MainFrame",
 				"MaterialEditor",
@@ -61,7 +64,7 @@ public class AssetTools : ModuleRules
 				"DestructibleMeshEditor",
 				"EditorWidgets",
 				"GameProjectGeneration",
-				"PropertyEditor",
+                "ActorPickerMode",
 				"MainFrame",
 				"MaterialEditor",
 				"Persona",
diff --git a/Engine/Source/Developer/AssetTools/Private/AssetTypeActions/AssetTypeActions_DialogueWave.cpp b/Engine/Source/Developer/AssetTools/Private/AssetTypeActions/AssetTypeActions_DialogueWave.cpp
--- a/Engine/Source/Developer/AssetTools/Private/AssetTypeActions/AssetTypeActions_DialogueWave.cpp
+++ b/Engine/Source/Developer/AssetTools/Private/AssetTypeActions/AssetTypeActions_DialogueWave.cpp
@@ -2,9 +2,46 @@
 
 #include "AssetToolsPrivatePCH.h"
 #include "SoundDefinitions.h"
+#include "AssetTypeActions_SoundBase.h"
+#include "ContentBrowserModule.h"
 
 #define LOCTEXT_NAMESPACE "AssetTypeActions"
 
+void FAssetTypeActions_DialogueWave::GetActions(const TArray<UObject*>& InObjects, FMenuBuilder& MenuBuilder)
+{
+	auto DialogueWaves = GetTypedWeakObjectPtrs<UDialogueWave>(InObjects);
+
+	MenuBuilder.AddMenuEntry(
+		LOCTEXT("Sound_PlaySound", "Play"),
+		LOCTEXT("Sound_PlaySoundTooltip", "Plays the selected sound."),
+		FSlateIcon(FEditorStyle::GetStyleSetName(), "MediaAsset.AssetActions.Play"),
+		FUIAction(
+		FExecuteAction::CreateSP(this, &FAssetTypeActions_DialogueWave::ExecutePlaySound, DialogueWaves),
+		FCanExecuteAction::CreateSP(this, &FAssetTypeActions_DialogueWave::CanExecutePlayCommand, DialogueWaves)
+		)
+		);
+
+	MenuBuilder.AddMenuEntry(
+		LOCTEXT("Sound_StopSound", "Stop"),
+		LOCTEXT("Sound_StopSoundTooltip", "Stops the selected sounds."),
+		FSlateIcon(FEditorStyle::GetStyleSetName(), "MediaAsset.AssetActions.Stop"),
+		FUIAction(
+		FExecuteAction::CreateSP(this, &FAssetTypeActions_DialogueWave::ExecuteStopSound, DialogueWaves),
+		FCanExecuteAction()
+		)
+		);
+
+	MenuBuilder.AddMenuEntry(
+		LOCTEXT("DialogueWave_CreateCue", "Create Cue"),
+		LOCTEXT("DialogueWave_CreateCueTooltip", "Creates a sound cue using this dialogue wave."),
+		FSlateIcon(FEditorStyle::GetStyleSetName(), "ClassIcon.SoundCue"),
+		FUIAction(
+		FExecuteAction::CreateSP(this, &FAssetTypeActions_DialogueWave::ExecuteCreateSoundCue, DialogueWaves),
+		FCanExecuteAction()
+		)
+		);
+}
+
 void FAssetTypeActions_DialogueWave::OpenAssetEditor( const TArray<UObject*>& InObjects, TSharedPtr<IToolkitHost> EditWithinLevelEditor )
 {
 	EToolkitMode::Type Mode = EditWithinLevelEditor.IsValid() ? EToolkitMode::WorldCentric : EToolkitMode::Standalone;
@@ -19,4 +56,134 @@ void FAssetTypeActions_DialogueWave::OpenAssetEditor( const TArray<UObject*>& In
 	}
 }
 
+bool FAssetTypeActions_DialogueWave::CanExecutePlayCommand(TArray<TWeakObjectPtr<UDialogueWave>> Objects) const
+{
+	if (Objects.Num() != 1)
+	{
+		return false;
+	}
+
+	USoundBase* Sound = nullptr;
+
+	auto DialogueWave = Cast<UDialogueWave>(Objects[0].Get());
+	for (int32 i = 0; i < DialogueWave->ContextMappings.Num(); ++i)
+	{
+		const FDialogueContextMapping& ContextMapping = DialogueWave->ContextMappings[i];
+
+		Sound = DialogueWave->GetWaveFromContext(ContextMapping.Context);
+		if (Sound != nullptr)
+		{
+			break;
+		}
+	}
+
+	return Sound != nullptr;
+}
+
+void FAssetTypeActions_DialogueWave::ExecutePlaySound(TArray<TWeakObjectPtr<UDialogueWave>> Objects)
+{
+	for (auto ObjIt = Objects.CreateConstIterator(); ObjIt; ++ObjIt)
+	{
+		UDialogueWave* DialogueWave = (*ObjIt).Get();
+		if (DialogueWave)
+		{
+			// Only play the first valid sound
+			PlaySound(DialogueWave);
+			break;
+		}
+	}
+}
+
+void FAssetTypeActions_DialogueWave::ExecuteStopSound(TArray<TWeakObjectPtr<UDialogueWave>> Objects)
+{
+	StopSound();
+}
+
+void FAssetTypeActions_DialogueWave::PlaySound(UDialogueWave* DialogueWave)
+{
+	USoundBase* Sound = nullptr;
+
+	for (int32 i = 0; i < DialogueWave->ContextMappings.Num(); ++i)
+	{
+		const FDialogueContextMapping& ContextMapping = DialogueWave->ContextMappings[i];
+
+		Sound = DialogueWave->GetWaveFromContext(ContextMapping.Context);
+		if (Sound != nullptr)
+		{
+			break;
+		}
+	}
+
+	if (Sound)
+	{
+		GEditor->PlayPreviewSound(Sound);
+	}
+	else
+	{
+		StopSound();
+	}
+}
+
+void FAssetTypeActions_DialogueWave::StopSound()
+{
+	GEditor->ResetPreviewAudioComponent();
+}
+
+void FAssetTypeActions_DialogueWave::ExecuteCreateSoundCue(TArray<TWeakObjectPtr<UDialogueWave>> Objects)
+{
+	const FString DefaultSuffix = TEXT("_Cue");
+
+	if (Objects.Num() == 1)
+	{
+		auto Object = Objects[0].Get();
+
+		if (Object)
+		{
+			// Determine an appropriate name
+			FString Name;
+			FString PackagePath;
+			CreateUniqueAssetName(Object->GetOutermost()->GetName(), DefaultSuffix, PackagePath, Name);
+
+			// Create the factory used to generate the asset
+			USoundCueFactoryNew* Factory = NewObject<USoundCueFactoryNew>();
+			Factory->InitialDialogueWave = Object;
+
+			FContentBrowserModule& ContentBrowserModule = FModuleManager::LoadModuleChecked<FContentBrowserModule>("ContentBrowser");
+			ContentBrowserModule.Get().CreateNewAsset(Name, FPackageName::GetLongPackagePath(PackagePath), USoundCue::StaticClass(), Factory);
+		}
+	}
+	else
+	{
+		TArray<UObject*> ObjectsToSync;
+
+		for (auto ObjIt = Objects.CreateConstIterator(); ObjIt; ++ObjIt)
+		{
+			auto Object = (*ObjIt).Get();
+			if (Object)
+			{
+				FString Name;
+				FString PackageName;
+				CreateUniqueAssetName(Object->GetOutermost()->GetName(), DefaultSuffix, PackageName, Name);
+
+				// Create the factory used to generate the asset
+				USoundCueFactoryNew* Factory = NewObject<USoundCueFactoryNew>();
+				Factory->InitialDialogueWave = Object;
+
+				FAssetToolsModule& AssetToolsModule = FModuleManager::GetModuleChecked<FAssetToolsModule>("AssetTools");
+				UObject* NewAsset = AssetToolsModule.Get().CreateAsset(Name, FPackageName::GetLongPackagePath(PackageName), USoundCue::StaticClass(), Factory);
+
+				if (NewAsset)
+				{
+					ObjectsToSync.Add(NewAsset);
+				}
+			}
+		}
+
+		if (ObjectsToSync.Num() > 0)
+		{
+			FAssetTools::Get().SyncBrowserToAssets(ObjectsToSync);
+		}
+	}
+}
+
 #undef LOCTEXT_NAMESPACE
\ No newline at end of file
diff --git a/Engine/Source/Developer/AssetTools/Private/AssetTypeActions/AssetTypeActions_DialogueWave.h b/Engine/Source/Developer/AssetTools/Private/AssetTypeActions/AssetTypeActions_DialogueWave.h
--- a/Engine/Source/Developer/AssetTools/Private/AssetTypeActions/AssetTypeActions_DialogueWave.h
+++ b/Engine/Source/Developer/AssetTools/Private/AssetTypeActions/AssetTypeActions_DialogueWave.h
@@ -14,5 +14,22 @@ public:
 	virtual UClass* GetSupportedClass() const override { return UDialogueWave::StaticClass(); }
 	virtual uint32 GetCategories() override { return EAssetTypeCategories::Sounds; }
 	virtual bool CanFilter() override { return true; }
+	virtual bool HasActions(const TArray<UObject*>& InObjects) const override { return true; }
+	virtual void GetActions(const TArray<UObject*>& InObjects, FMenuBuilder& MenuBuilder) override;
 	virtual void OpenAssetEditor( const TArray<UObject*>& InObjects, TSharedPtr<class IToolkitHost> EditWithinLevelEditor = TSharedPtr<IToolkitHost>() ) override;
+
+private:
+
+	bool CanExecutePlayCommand(TArray<TWeakObjectPtr<UDialogueWave>> Objects) const;
+
+	void ExecutePlaySound(TArray<TWeakObjectPtr<UDialogueWave>> Objects);
+
+	void ExecuteStopSound(TArray<TWeakObjectPtr<UDialogueWave>> Objects);
+
+	void PlaySound(UDialogueWave* DialogueWave);
+
+	void StopSound();
+
+	void ExecuteCreateSoundCue(TArray<TWeakObjectPtr<UDialogueWave>> Objects);
+
 };
\ No newline at end of file
diff --git a/Engine/Source/Developer/AssetTools/Private/AssetTypeActions/AssetTypeActions_SoundWave.cpp b/Engine/Source/Developer/AssetTools/Private/AssetTypeActions/AssetTypeActions_SoundWave.cpp
--- a/Engine/Source/Developer/AssetTools/Private/AssetTypeActions/AssetTypeActions_SoundWave.cpp
+++ b/Engine/Source/Developer/AssetTools/Private/AssetTypeActions/AssetTypeActions_SoundWave.cpp
@@ -6,6 +6,10 @@
 #include "ContentBrowserModule.h"
 #include "Sound/SoundWave.h"
 #include "Sound/SoundCue.h"
+#include "Sound/DialogueWave.h"
+#include "ContentBrowserDelegates.h"
+#include "PropertyCustomizationHelpers.h"
+#include "Factories/DialogueWaveFactory.h"
 
 #define LOCTEXT_NAMESPACE "AssetTypeActions"
 
@@ -29,6 +33,11 @@ void FAssetTypeActions_SoundWave::GetActions( const TArray<UObject*>& InObjects,
 		FCanExecuteAction()
 			)
 		);
+
+	MenuBuilder.AddSubMenu(
+		LOCTEXT("SoundWave_CreateDialogue", "Create Dialogue"),
+		LOCTEXT("SoundWave_CreateDialogueTooltip", "Creates a dialogue wave using this sound wave."),
+		FNewMenuDelegate::CreateSP(this, &FAssetTypeActions_SoundWave::FillVoiceMenu, SoundNodes));
 }
 
 void FAssetTypeActions_SoundWave::GetResolvedSourceFilePaths(const TArray<UObject*>& TypeAssets, TArray<FString>& OutSourceFilePaths) const
@@ -102,4 +111,84 @@ void FAssetTypeActions_SoundWave::ExecuteCreateSoundCue(TArray<TWeakObjectPtr<US
 	}
 }
 
+void FAssetTypeActions_SoundWave::ExecuteCreateDialogueWave(const class FAssetData& AssetData, TArray<TWeakObjectPtr<USoundWave>> Objects)
+{
+	const FString DefaultSuffix = TEXT("_Dialogue");
+
+	UDialogueVoice* DialogueVoice = Cast<UDialogueVoice>(AssetData.GetAsset());
+
+	if (Objects.Num() == 1)
+	{
+		auto Object = Objects[0].Get();
+
+		if (Object)
+		{
+			// Determine an appropriate name
+			FString Name;
+			FString PackagePath;
+			CreateUniqueAssetName(Object->GetOutermost()->GetName(), DefaultSuffix, PackagePath, Name);
+
+			// Create the factory used to generate the asset
+			UDialogueWaveFactory* Factory = NewObject<UDialogueWaveFactory>();
+			Factory->InitialSoundWave = Object;
+			Factory->InitialSpeakerVoice = DialogueVoice;
+			Factory->HasSetInitialTargetVoice = true;
+
+			FContentBrowserModule& ContentBrowserModule = FModuleManager::LoadModuleChecked<FContentBrowserModule>("ContentBrowser");
+			ContentBrowserModule.Get().CreateNewAsset(Name, FPackageName::GetLongPackagePath(PackagePath), UDialogueWave::StaticClass(), Factory);
+		}
+	}
+	else
+	{
+		TArray<UObject*> ObjectsToSync;
+
+		for (auto ObjIt = Objects.CreateConstIterator(); ObjIt; ++ObjIt)
+		{
+			auto Object = (*ObjIt).Get();
+			if (Object)
+			{
+				FString Name;
+				FString PackageName;
+				CreateUniqueAssetName(Object->GetOutermost()->GetName(), DefaultSuffix, PackageName, Name);
+
+				// Create the factory used to generate the asset
+				UDialogueWaveFactory* Factory = NewObject<UDialogueWaveFactory>();
+				Factory->InitialSoundWave = Object;
+				Factory->InitialSpeakerVoice = DialogueVoice;
+				Factory->HasSetInitialTargetVoice = true;
+
+				FAssetToolsModule& AssetToolsModule = FModuleManager::GetModuleChecked<FAssetToolsModule>("AssetTools");
+				UObject* NewAsset = AssetToolsModule.Get().CreateAsset(Name, FPackageName::GetLongPackagePath(PackageName), UDialogueWave::StaticClass(), Factory);
+
+				if (NewAsset)
+				{
+					ObjectsToSync.Add(NewAsset);
+				}
+			}
+		}
+
+		if (ObjectsToSync.Num() > 0)
+		{
+			FAssetTools::Get().SyncBrowserToAssets(ObjectsToSync);
+		}
+	}
+}
+
+void FAssetTypeActions_SoundWave::FillVoiceMenu(FMenuBuilder& MenuBuilder, TArray<TWeakObjectPtr<USoundWave>> Objects)
+{
+	TArray<const UClass*> AllowedClasses;
+	AllowedClasses.Add(UDialogueVoice::StaticClass());
+
+	TSharedRef<SWidget> VoicePicker = PropertyCustomizationHelpers::MakeAssetPickerWithMenu(
+		FAssetData(),
+		false,
+		AllowedClasses,
+		PropertyCustomizationHelpers::GetNewAssetFactoriesForClasses(AllowedClasses),
+		FOnShouldFilterAsset(),
+		FOnAssetSelected::CreateSP(this, &FAssetTypeActions_SoundWave::ExecuteCreateDialogueWave, Objects),
+		FSimpleDelegate());
+
+	MenuBuilder.AddWidget(VoicePicker, FText::GetEmpty(), false);
+}
+
 #undef LOCTEXT_NAMESPACE
diff --git a/Engine/Source/Developer/AssetTools/Private/AssetTypeActions/AssetTypeActions_SoundWave.h b/Engine/Source/Developer/AssetTools/Private/AssetTypeActions/AssetTypeActions_SoundWave.h
--- a/Engine/Source/Developer/AssetTools/Private/AssetTypeActions/AssetTypeActions_SoundWave.h
+++ b/Engine/Source/Developer/AssetTools/Private/AssetTypeActions/AssetTypeActions_SoundWave.h
@@ -19,4 +19,9 @@ private:
 
 	/** Creates a SoundCue of the same name for the sound, if one does not already exist */
 	void ExecuteCreateSoundCue(TArray<TWeakObjectPtr<USoundWave>> Objects);
+
+	/** Creates a DialogueWave of the same name for the sound, if one does not already exist */
+	void ExecuteCreateDialogueWave(const class FAssetData& AssetData, TArray<TWeakObjectPtr<USoundWave>> Objects);
+
+	void FillVoiceMenu(FMenuBuilder& MenuBuilder, TArray<TWeakObjectPtr<USoundWave>> Objects);
 };
\ No newline at end of file
diff --git a/Engine/Source/Developer/DerivedDataCache/DerivedDataCache.Build.cs b/Engine/Source/Developer/DerivedDataCache/DerivedDataCache.Build.cs
--- a/Engine/Source/Developer/DerivedDataCache/DerivedDataCache.Build.cs
+++ b/Engine/Source/Developer/DerivedDataCache/DerivedDataCache.Build.cs
@@ -8,7 +8,6 @@ public class DerivedDataCache : ModuleRules
 	public DerivedDataCache(TargetInfo Target)
 	{
 		PrivateDependencyModuleNames.Add("Core");
-        PrivateDependencyModuleNames.Add("CookingStats");
 
 		// Internal (NotForLicensees) module
 		var DDCUtilsModule = Path.Combine("Developer", "NotForLicensees", "DDCUtils", "DDCUtils.Build.cs");
diff --git a/Engine/Source/Developer/DerivedDataCache/Private/DerivedDataBackends.cpp b/Engine/Source/Developer/DerivedDataCache/Private/DerivedDataBackends.cpp
--- a/Engine/Source/Developer/DerivedDataCache/Private/DerivedDataBackends.cpp
+++ b/Engine/Source/Developer/DerivedDataCache/Private/DerivedDataBackends.cpp
@@ -4,7 +4,6 @@
 #include "SecureHash.h"
 #include "DerivedDataBackendInterface.h"
 #include "DerivedDataCacheInterface.h"
-#include "DDCStatsHelper.h"
 #include "MemoryDerivedDataBackend.h"
 #include "DerivedDataBackendAsyncPutWrapper.h"
 #include "PakFileDerivedDataBackend.h"
diff --git a/Engine/Source/Developer/DerivedDataCache/Private/DerivedDataCache.cpp b/Engine/Source/Developer/DerivedDataCache/Private/DerivedDataCache.cpp
--- a/Engine/Source/Developer/DerivedDataCache/Private/DerivedDataCache.cpp
+++ b/Engine/Source/Developer/DerivedDataCache/Private/DerivedDataCache.cpp
@@ -931,10 +931,6 @@ public:
 
 	virtual void StartupModule() override
 	{
-#if WITH_EDITOR
-		// Make sure the CookingStats module gets loaded on the correct thread (used by DDCStats on a background thread)
-		FModuleManager::Get().LoadModule(TEXT("CookingStats"));
-#endif // WITH_EDITOR
 	}
 
 	virtual void ShutdownModule() override
diff --git a/Engine/Source/Developer/DerivedDataCache/Private/FileSystemDerivedDataBackend.cpp b/Engine/Source/Developer/DerivedDataCache/Private/FileSystemDerivedDataBackend.cpp
--- a/Engine/Source/Developer/DerivedDataCache/Private/FileSystemDerivedDataBackend.cpp
+++ b/Engine/Source/Developer/DerivedDataCache/Private/FileSystemDerivedDataBackend.cpp
@@ -6,8 +6,6 @@
 #include "DerivedDataBackendInterface.h"
 #include "DDCCleanup.h"
 
-#include "DDCStatsHelper.h"
-
 #define MAX_BACKEND_KEY_LENGTH (120)
 #define MAX_BACKEND_NUMBERED_SUBFOLDER_LENGTH (9)
 #if PLATFORM_LINUX	// PATH_MAX on Linux is 4096 (getconf PATH_MAX /, also see limits.h), so this value can be larger (note that it is still arbitrary).
diff --git a/Engine/Source/Developer/DerivedDataCache/Private/HierarchicalDerivedDataBackend.h b/Engine/Source/Developer/DerivedDataCache/Private/HierarchicalDerivedDataBackend.h
--- a/Engine/Source/Developer/DerivedDataCache/Private/HierarchicalDerivedDataBackend.h
+++ b/Engine/Source/Developer/DerivedDataCache/Private/HierarchicalDerivedDataBackend.h
@@ -2,9 +2,6 @@
 
 #pragma once
 
-#include "DDCStatsHelper.h"
-
-
 /** 
  * A backend wrapper that implements a cache hierarchy of backends. 
 **/
@@ -65,13 +62,6 @@ public:
 	 */
 	virtual bool CachedDataProbablyExists(const TCHAR* CacheKey) override
 	{
-#if ENABLE_DDC_STATS
-		const static FName NAME_CachedDataProbablyExists(TEXT("CachedDataProbablyExists"));
-		FDDCScopeStatHelper Stat(CacheKey, NAME_CachedDataProbablyExists);
-		const static FName NAME_HierarchicalDDC = FName(TEXT("HierarchicalDDC"));
-		Stat.AddTag(NAME_HierarchicalDDC, FString());
-#endif
-
 		for (int32 CacheIndex = 0; CacheIndex < InnerBackends.Num(); CacheIndex++)
 		{
 			if (InnerBackends[CacheIndex]->CachedDataProbablyExists(CacheKey))
@@ -90,16 +80,6 @@ public:
 	 */
 	virtual bool GetCachedData(const TCHAR* CacheKey, TArray<uint8>& OutData) override
 	{
-#if ENABLE_DDC_STATS
-		const static FName NAME_GetCachedData = FName(TEXT("GetCachedData"));
-		const static FName NAME_HierarchicalDDC = FName(TEXT("HierarchicalDDC"));
-		const static FName NAME_DataSize = FName(TEXT("DataSize"));
-		const static FName NAME_Retrieved(TEXT("Retrieved"));
-
-		FDDCScopeStatHelper Stat(CacheKey, NAME_GetCachedData);
-		Stat.AddTag(NAME_HierarchicalDDC, FString());
-#endif
-
 		for (int32 CacheIndex = 0; CacheIndex < InnerBackends.Num(); CacheIndex++)
 		{
 			if (InnerBackends[CacheIndex]->CachedDataProbablyExists(CacheKey) && InnerBackends[CacheIndex]->GetCachedData(CacheKey, OutData))
@@ -139,16 +119,9 @@ public:
 						}
 					}
 				}
-#if ENABLE_DDC_STATS
-				Stat.AddTag(NAME_DataSize, OutData.Num());
-				Stat.AddTag(NAME_Retrieved, true);
-#endif
 				return true;
 			}
 		}
-#if ENABLE_DDC_STATS
-		Stat.AddTag(NAME_Retrieved, false);
-#endif
 		return false;
 	}
 	/**
@@ -160,15 +133,6 @@ public:
 	 */
 	virtual void PutCachedData(const TCHAR* CacheKey, TArray<uint8>& InData, bool bPutEvenIfExists) override
 	{
-#if ENABLE_DDC_STATS
-		const static FName NAME_PutCachedData = FName(TEXT("PutCachedData"));
-		const static FName NAME_HierarchicalDDC = FName(TEXT("HierarchicalDDC"));
-		const static FName NAME_DataSize = FName(TEXT("DataSize"));
-		FDDCScopeStatHelper Stat(CacheKey, NAME_PutCachedData);
-		Stat.AddTag(NAME_HierarchicalDDC, FString());
-		Stat.AddTag(NAME_DataSize, InData.Num());
-#endif
-
 		if (!bIsWritable)
 		{
 			return; // no point in continuing down the chain
diff --git a/Engine/Source/Developer/LauncherServices/Private/Launcher/LauncherWorker.cpp b/Engine/Source/Developer/LauncherServices/Private/Launcher/LauncherWorker.cpp
--- a/Engine/Source/Developer/LauncherServices/Private/Launcher/LauncherWorker.cpp
+++ b/Engine/Source/Developer/LauncherServices/Private/Launcher/LauncherWorker.cpp
@@ -194,6 +194,12 @@ static void AddDeviceToLaunchCommand(const FString& DeviceId, ITargetDeviceProxy
 			RoleCommands += *(TEXT(" ") + Roles[RoleIndex]->GetUATCommandLine());
 		}
 	}
+
+	if (FParse::Param(FCommandLine::Get(), TEXT("nomcp")))
+	{
+		// if our editor has nomcp then pass it through the launched game
+		RoleCommands += TEXT(" -nomcp");
+	}
 }
 
 FString FLauncherWorker::CreateUATCommand( const ILauncherProfileRef& InProfile, const TArray<FString>& InPlatforms, TArray<FCommandDesc>& OutCommands, FString& CommandStart )
diff --git a/Engine/Source/Developer/SlateReflector/Private/Models/WidgetReflectorNode.cpp b/Engine/Source/Developer/SlateReflector/Private/Models/WidgetReflectorNode.cpp
--- a/Engine/Source/Developer/SlateReflector/Private/Models/WidgetReflectorNode.cpp
+++ b/Engine/Source/Developer/SlateReflector/Private/Models/WidgetReflectorNode.cpp
@@ -554,25 +554,7 @@ FText FWidgetReflectorNodeUtils::GetWidgetVisibilityText(const TSharedPtr<SWidge
 
 FText FWidgetReflectorNodeUtils::GetWidgetReadableLocation(const TSharedPtr<SWidget>& InWidget)
 {
-	FText ReadableLocationText;
-
-	if (InWidget.IsValid())
-	{
-		// UMG widgets have meta-data to help track them
-		TSharedPtr<FReflectionMetaData> MetaData = InWidget->GetMetaData<FReflectionMetaData>();
-		if (MetaData.IsValid() && MetaData->Asset.Get() != nullptr)
-		{
-			const FText AssetNameText = FText::FromName(MetaData->Asset->GetFName());
-			const FText WidgetNameText = FText::FromName(MetaData->Name);
-			ReadableLocationText = FText::Format(LOCTEXT("AssetReadableLocationFmt", "{0} [{1}]"), AssetNameText, WidgetNameText);
-		}
-		else
-		{
-			ReadableLocationText = FText::FromString(InWidget->GetReadableLocation());
-		}
-	}
-	
-	return ReadableLocationText;
+	return (InWidget.IsValid()) ? FText::FromString(FReflectionMetaData::GetWidgetDebugInfo(InWidget.Get())) : FText::GetEmpty();
 }
 
 FString FWidgetReflectorNodeUtils::GetWidgetFile(const TSharedPtr<SWidget>& InWidget)
diff --git a/Engine/Source/Editor/BlueprintGraph/Classes/K2Node_BaseAsyncTask.h b/Engine/Source/Editor/BlueprintGraph/Classes/K2Node_BaseAsyncTask.h
--- a/Engine/Source/Editor/BlueprintGraph/Classes/K2Node_BaseAsyncTask.h
+++ b/Engine/Source/Editor/BlueprintGraph/Classes/K2Node_BaseAsyncTask.h
@@ -10,6 +10,18 @@ class UK2Node_CustomEvent;
 class UEdGraphPin;
 class UK2Node_TemporaryVariable;
 
+/** struct to remap pins for Async Tasks.
+ * a single K2 node is shared by many proxy classes.
+ * This allows redirecting pins by name per proxy class.
+ * Add entries similar to this one in Engine.ini:
+ * +K2AsyncTaskPinRedirects=(ProxyClassName="AbilityTask_PlayMontageAndWait", OldPinName="OnComplete", NewPinName="OnBlendOut")
+ */
+
+struct FAsyncTaskPinRedirectMapInfo
+{
+	TMap<FString, TArray<UClass*> > OldPinToProxyClassMap;
+};
+
 /** !!! The proxy object should have RF_StrongRefOnFrame flag. !!! */
 
 UCLASS(Abstract)
@@ -31,6 +43,8 @@ class BLUEPRINTGRAPH_API UK2Node_BaseAsyncTask : public UK2Node
 	virtual FName GetCornerIcon() const override;
 	virtual void GetMenuActions(FBlueprintActionDatabaseRegistrar& ActionRegistrar) const override;
 	virtual FText GetMenuCategory() const override;
+	/** Whether or not two pins match for purposes of reconnection after reconstruction.  This allows pins that may have had their names changed via reconstruction to be matched to their old values on a node-by-node basis, if needed*/
+	virtual ERedirectType DoPinsMatchForReconstruction(const UEdGraphPin* NewPin, int32 NewPinIndex, const UEdGraphPin* OldPin, int32 OldPinIndex) const override;
 	// End of UK2Node interface
 
 protected:
@@ -74,4 +88,8 @@ protected:
 
 		static const FString& GetAsyncTaskProxyName();
 	};
+
+	// Pin Redirector support
+	static TMap<FString, FAsyncTaskPinRedirectMapInfo> AsyncTaskPinRedirectMap;
+	static bool bAsyncTaskPinRedirectMapInitialized;
 };
diff --git a/Engine/Source/Editor/BlueprintGraph/Private/K2Node_BaseAsyncTask.cpp b/Engine/Source/Editor/BlueprintGraph/Private/K2Node_BaseAsyncTask.cpp
--- a/Engine/Source/Editor/BlueprintGraph/Private/K2Node_BaseAsyncTask.cpp
+++ b/Engine/Source/Editor/BlueprintGraph/Private/K2Node_BaseAsyncTask.cpp
@@ -457,4 +457,58 @@ void UK2Node_BaseAsyncTask::ValidateNodeDuringCompilation(class FCompilerResults
 	}
 }
 
+TMap<FString, FAsyncTaskPinRedirectMapInfo> UK2Node_BaseAsyncTask::AsyncTaskPinRedirectMap;
+bool UK2Node_BaseAsyncTask::bAsyncTaskPinRedirectMapInitialized = false;
+
+UK2Node::ERedirectType UK2Node_BaseAsyncTask::DoPinsMatchForReconstruction(const UEdGraphPin* NewPin, int32 NewPinIndex, const UEdGraphPin* OldPin, int32 OldPinIndex) const
+{
+	if (GConfig && ProxyClass)
+	{
+		// Initialize remap table from INI
+		if (!bAsyncTaskPinRedirectMapInitialized)
+		{
+			bAsyncTaskPinRedirectMapInitialized = true;
+			FConfigSection* PackageRedirects = GConfig->GetSectionPrivate(TEXT("/Script/Engine.Engine"), false, true, GEngineIni);
+			for (FConfigSection::TIterator It(*PackageRedirects); It; ++It)
+			{
+				if (It.Key() == TEXT("K2AsyncTaskPinRedirects"))
+				{
+					FString ProxyClassString;
+					FString OldPinString;
+					FString NewPinString;
+
+					FParse::Value(*It.Value(), TEXT("ProxyClassName="), ProxyClassString);
+					FParse::Value(*It.Value(), TEXT("OldPinName="), OldPinString);
+					FParse::Value(*It.Value(), TEXT("NewPinName="), NewPinString);
+
+					UClass* RedirectProxyClass = FindObject<UClass>(ANY_PACKAGE, *ProxyClassString);
+					if (RedirectProxyClass)
+					{
+						FAsyncTaskPinRedirectMapInfo& PinRedirectInfo = AsyncTaskPinRedirectMap.FindOrAdd(OldPinString);
+						TArray<UClass*>& ProxyClassArray = PinRedirectInfo.OldPinToProxyClassMap.FindOrAdd(NewPinString);
+						ProxyClassArray.AddUnique(RedirectProxyClass);
+					}
+				}
+			}
+		}
+
+		// See if these pins need to be remapped.
+		if (FAsyncTaskPinRedirectMapInfo* PinRedirectInfo = AsyncTaskPinRedirectMap.Find(OldPin->PinName))
+		{
+			if (TArray<UClass*>* ProxyClassArray = PinRedirectInfo->OldPinToProxyClassMap.Find(NewPin->PinName))
+			{
+				for (UClass* RedirectedProxyClass : *ProxyClassArray)
+				{
+					if (ProxyClass->IsChildOf(RedirectedProxyClass))
+					{
+						return UK2Node::ERedirectType_Name;
+					}
+				}
+			}
+		}
+	}
+
+	return Super::DoPinsMatchForReconstruction(NewPin, NewPinIndex, OldPin, OldPinIndex);
+}
+
 #undef LOCTEXT_NAMESPACE
diff --git a/Engine/Source/Editor/Cascade/Private/Cascade.cpp b/Engine/Source/Editor/Cascade/Private/Cascade.cpp
--- a/Engine/Source/Editor/Cascade/Private/Cascade.cpp
+++ b/Engine/Source/Editor/Cascade/Private/Cascade.cpp
@@ -174,6 +174,19 @@ FCascade::~FCascade()
 	CascadeModule->CascadeClosed(this);
 }
 
+void FCascade::OnComponentActivationChange(UParticleSystemComponent* PSC, bool bActivated)
+{
+	check(PSC);
+
+	if (UCascadeParticleSystemComponent* CPSC = Cast<UCascadeParticleSystemComponent>(PSC))
+	{
+		if (FCascade* Cascade = CPSC->CascadePreviewViewportPtr->GetCascade())
+		{
+			PSC->SetRequiredSignificance(Cascade->GetRequiredSignificance());
+		}
+	}
+}
+
 void FCascade::InitCascade(const EToolkitMode::Type Mode, const TSharedPtr< class IToolkitHost >& InitToolkitHost, UObject* ObjectToEdit)
 {
 	ParticleSystem = CastChecked<UParticleSystem>(ObjectToEdit);
@@ -253,6 +266,7 @@ void FCascade::InitCascade(const EToolkitMode::Type Mode, const TSharedPtr< clas
 
 	CurveToReplace = NULL;
 	DetailMode = GlobalDetailMode = GetCachedScalabilityCVars().DetailMode;
+	RequiredSignificance = EParticleSignificanceLevel::Low;
 
 	bIsToggleMotion = false;
 	MotionModeRadius = EditorOptions->MotionModeRadius;
@@ -476,6 +490,11 @@ int32 FCascade::GetDetailMode() const
 	return DetailMode;
 }
 
+EParticleSignificanceLevel FCascade::GetRequiredSignificance() const
+{
+	return RequiredSignificance;
+}
+
 bool FCascade::GetIsModuleShared(UParticleModule* Module)
 {
 	int32 FindCount = 0;
@@ -1537,7 +1556,10 @@ void FCascade::Tick(float DeltaTime)
 			ParticleSystemComponent->ComponentToWorld = FTransform(Position);
 		}
 
-		ParticleSystemComponent->CascadeTickComponent(CurrDeltaTime, LEVELTICK_All);
+		if (ParticleSystemComponent->IsComponentTickEnabled())
+		{
+			ParticleSystemComponent->CascadeTickComponent(CurrDeltaTime, LEVELTICK_All);
+		}
 		ParticleSystemComponent->DoDeferredRenderUpdates_Concurrent();
 		GetFXSystem()->Tick(CurrDeltaTime);
 		TotalTime += CurrDeltaTime;
@@ -2313,6 +2335,30 @@ void FCascade::BindCommands()
 		FCanExecuteAction(),
 		FIsActionChecked::CreateSP(this, &FCascade::IsDetailModeChecked, DM_High));
 
+	ToolkitCommands->MapAction(
+		Commands.Significance_Critical,
+		FExecuteAction::CreateSP(this, &FCascade::OnSignificance, EParticleSignificanceLevel::Critical),
+		FCanExecuteAction(),
+		FIsActionChecked::CreateSP(this, &FCascade::IsSignificanceChecked, EParticleSignificanceLevel::Critical));
+
+	ToolkitCommands->MapAction(
+		Commands.Significance_High,
+		FExecuteAction::CreateSP(this, &FCascade::OnSignificance, EParticleSignificanceLevel::High),
+		FCanExecuteAction(),
+		FIsActionChecked::CreateSP(this, &FCascade::IsSignificanceChecked, EParticleSignificanceLevel::High));
+
+	ToolkitCommands->MapAction(
+		Commands.Significance_Medium,
+		FExecuteAction::CreateSP(this, &FCascade::OnSignificance, EParticleSignificanceLevel::Medium),
+		FCanExecuteAction(),
+		FIsActionChecked::CreateSP(this, &FCascade::IsSignificanceChecked, EParticleSignificanceLevel::Medium));
+
+	ToolkitCommands->MapAction(
+		Commands.Significance_Low,
+		FExecuteAction::CreateSP(this, &FCascade::OnSignificance, EParticleSignificanceLevel::Low),
+		FCanExecuteAction(),
+		FIsActionChecked::CreateSP(this, &FCascade::IsSignificanceChecked, EParticleSignificanceLevel::Low));
+
 	ToolkitCommands->MapAction(
 		Commands.RegenerateLowestLODDuplicatingHighest,
 		FExecuteAction::CreateSP(this, &FCascade::OnRegenerateLowestLODDuplicatingHighest));
@@ -2580,6 +2626,7 @@ void FCascade::RestartParticleSystem()
 	if (ParticleSystemComponent)
 	{
 		ParticleSystemComponent->ResetParticles();
+		ParticleSystemComponent->SetRequiredSignificance(RequiredSignificance);
 		ParticleSystemComponent->ActivateSystem();
 		if (ParticleSystemComponent->Template)
 		{
@@ -3687,6 +3734,7 @@ void FCascade::OnRestartInLevel()
 				PSysComp->ActivateSystem();
 				PSysComp->ConditionalCacheViewRelevanceFlags();
 				PSysComp->MarkRenderStateDirty();
+				PSysComp->SetRequiredSignificance(RequiredSignificance);
 			}
 		}
 	}
@@ -4031,6 +4079,33 @@ bool FCascade::IsDetailModeChecked(EDetailMode InDetailMode) const
 	return DetailMode == InDetailMode;
 }
 
+void FCascade::OnSignificance(EParticleSignificanceLevel InSignificance)
+{
+	if (PreviewViewport.IsValid() && PreviewViewport->GetViewportClient().IsValid())
+	{
+		if (RequiredSignificance == InSignificance)
+		{
+			return;
+		}
+
+		// Set the detail mode values on in-level particle systems
+		for (TObjectIterator<UParticleSystemComponent> It; It; ++It)
+		{
+			if (It->Template == ParticleSystemComponent->Template)
+			{
+				It->SetRequiredSignificance(InSignificance);
+			}
+		}
+
+		RequiredSignificance = InSignificance;
+	}
+}
+
+bool FCascade::IsSignificanceChecked(EParticleSignificanceLevel InSignificance) const
+{
+	return RequiredSignificance == InSignificance;
+}
+
 void FCascade::OnJumpToLowestLOD()
 {
 	if (ParticleSystem->Emitters.Num() == 0)
diff --git a/Engine/Source/Editor/Cascade/Private/Cascade.h b/Engine/Source/Editor/Cascade/Private/Cascade.h
--- a/Engine/Source/Editor/Cascade/Private/Cascade.h
+++ b/Engine/Source/Editor/Cascade/Private/Cascade.h
@@ -86,6 +86,9 @@ public:
 	/** Returns the current detail mode */
 	int32 GetDetailMode() const;
 
+	/** Returns the required significance for the fx in the cascade viewport. */
+	EParticleSignificanceLevel GetRequiredSignificance()const;
+
 	/** Returns true if the module is shared */
 	bool GetIsModuleShared(UParticleModule* Module);
 
@@ -192,6 +195,7 @@ public:
 	 */
 	static bool ConvertAllModulesToSeeded(UParticleSystem* ParticleSystem);
 
+	static void OnComponentActivationChange(UParticleSystemComponent* PSC, bool bActivated);
 private:
 	//~ Begin FEditorUndoClient Interface
 	virtual void PostUndo(bool bSuccess) override;
@@ -341,6 +345,8 @@ private:
 	void OnRegenerateLowestLOD();
 	void OnDetailMode(EDetailMode DetailMode);
 	bool IsDetailModeChecked(EDetailMode DetailMode) const;
+	void OnSignificance(EParticleSignificanceLevel InSignificance);
+	bool IsSignificanceChecked(EParticleSignificanceLevel InSignificance) const;
 	void OnJumpToHighestLOD();
 	void OnAddLODAfterCurrent();
 	void OnAddLODBeforeCurrent();
@@ -428,6 +434,10 @@ private:
 	
 	/** Cascade specific detail mode */
 	int32 DetailMode;
+
+	/** Required significance for cascade. */
+	EParticleSignificanceLevel RequiredSignificance;
+
 	/** Used to track changes in the global detail mode setting */
 	int32 GlobalDetailMode;
 
diff --git a/Engine/Source/Editor/Cascade/Private/CascadeActions.cpp b/Engine/Source/Editor/Cascade/Private/CascadeActions.cpp
--- a/Engine/Source/Editor/Cascade/Private/CascadeActions.cpp
+++ b/Engine/Source/Editor/Cascade/Private/CascadeActions.cpp
@@ -42,6 +42,10 @@ void FCascadeCommands::RegisterCommands()
 	UI_COMMAND(DetailMode_Low, "Low", "Select Low Detail Mode", EUserInterfaceActionType::RadioButton, FInputChord());
 	UI_COMMAND(DetailMode_Medium, "Medium", "Select Medium Detail Mode", EUserInterfaceActionType::RadioButton, FInputChord());
 	UI_COMMAND(DetailMode_High, "High", "Select High Detail Mode", EUserInterfaceActionType::RadioButton, FInputChord());
+	UI_COMMAND(Significance_Critical, "Critical", "Require >= Critical Significance", EUserInterfaceActionType::RadioButton, FInputChord());
+	UI_COMMAND(Significance_High, "High", "Require High >= Significance", EUserInterfaceActionType::RadioButton, FInputChord());
+	UI_COMMAND(Significance_Medium, "Medium", "Require >= Medium Significance", EUserInterfaceActionType::RadioButton, FInputChord());
+	UI_COMMAND(Significance_Low, "Low", "Require >= Low Significance", EUserInterfaceActionType::RadioButton, FInputChord());
 	UI_COMMAND(ToggleGeometry, "Geometry", "Display preview mesh", EUserInterfaceActionType::ToggleButton, FInputChord());
 	UI_COMMAND(ToggleGeometry_Properties, "Geometry Properties", "Display Preview Mesh Properties", EUserInterfaceActionType::Button, FInputChord());
 	UI_COMMAND(ToggleLocalVectorFields, "Vector Fields", "Display Local Vector Fields", EUserInterfaceActionType::ToggleButton, FInputChord());
diff --git a/Engine/Source/Editor/Cascade/Private/CascadeActions.h b/Engine/Source/Editor/Cascade/Private/CascadeActions.h
--- a/Engine/Source/Editor/Cascade/Private/CascadeActions.h
+++ b/Engine/Source/Editor/Cascade/Private/CascadeActions.h
@@ -56,6 +56,10 @@ public:
 	TSharedPtr<FUICommandInfo> DetailMode_Low;
 	TSharedPtr<FUICommandInfo> DetailMode_Medium;
 	TSharedPtr<FUICommandInfo> DetailMode_High;
+	TSharedPtr<FUICommandInfo> Significance_Critical;
+	TSharedPtr<FUICommandInfo> Significance_High;
+	TSharedPtr<FUICommandInfo> Significance_Medium;
+	TSharedPtr<FUICommandInfo> Significance_Low;
 	TSharedPtr<FUICommandInfo> JumpToHighestLOD;
 	TSharedPtr<FUICommandInfo> JumpToHigherLOD;
 	TSharedPtr<FUICommandInfo> AddLODAfterCurrent;
diff --git a/Engine/Source/Editor/Cascade/Private/CascadeModule.cpp b/Engine/Source/Editor/Cascade/Private/CascadeModule.cpp
--- a/Engine/Source/Editor/Cascade/Private/CascadeModule.cpp
+++ b/Engine/Source/Editor/Cascade/Private/CascadeModule.cpp
@@ -26,6 +26,8 @@ public:
 	{
 		MenuExtensibilityManager = MakeShareable(new FExtensibilityManager);
 		ToolBarExtensibilityManager = MakeShareable(new FExtensibilityManager);
+
+		UParticleSystemComponent::OnSystemPreActivationChange.AddStatic(&FCascade::OnComponentActivationChange);
 	}
 
 	/** Called before the module is unloaded, right before the module object is destroyed. */
diff --git a/Engine/Source/Editor/Cascade/Private/CascadePreviewViewportClient.cpp b/Engine/Source/Editor/Cascade/Private/CascadePreviewViewportClient.cpp
--- a/Engine/Source/Editor/Cascade/Private/CascadePreviewViewportClient.cpp
+++ b/Engine/Source/Editor/Cascade/Private/CascadePreviewViewportClient.cpp
@@ -464,6 +464,15 @@ void FCascadeEdPreviewViewportClient::Draw(FViewport* Viewport, FCanvas* Canvas)
 		Canvas->DrawItem( TextItem,  5.0f, Viewport->GetSizeXY().Y - 105.0f );		
 	}
 
+	EParticleSignificanceLevel ReqSignificance = CascadePtr.Pin()->GetRequiredSignificance();
+	if (ReqSignificance != EParticleSignificanceLevel::Low)
+	{
+		FString ReqSigOutput = FString::Printf(TEXT("REQUIRED SIGNIFICANCE: %s"), (ReqSignificance == EParticleSignificanceLevel::Medium) ? TEXT("MEDIUM") : ((ReqSignificance == EParticleSignificanceLevel::High) ? TEXT("HIGH") : TEXT("CRITICAL")));
+		TextItem.SetColor(FLinearColor::Red);
+		TextItem.Text = FText::FromString(ReqSigOutput);
+		Canvas->DrawItem(TextItem, 5.0f, Viewport->GetSizeXY().Y - 120.0f);
+	}
+
 	if (bCaptureScreenShot)
 	{
 		UParticleSystem* ParticleSystem = CascadePtr.Pin()->GetParticleSystem();
diff --git a/Engine/Source/Editor/Cascade/Private/CascadePreviewViewportClient.h b/Engine/Source/Editor/Cascade/Private/CascadePreviewViewportClient.h
--- a/Engine/Source/Editor/Cascade/Private/CascadePreviewViewportClient.h
+++ b/Engine/Source/Editor/Cascade/Private/CascadePreviewViewportClient.h
@@ -62,6 +62,7 @@ public:
 	FEditorCommonDrawHelper& GetDrawHelper();
 	float& GetWireSphereRadius();
 
+	FCascade* GetCascade(){ return CascadePtr.Pin().Get(); }
 private:
 	/** Pointer back to the ParticleSystem editor tool that owns us */
 	TWeakPtr<FCascade> CascadePtr;
diff --git a/Engine/Source/Editor/Cascade/Private/SCascadePreviewToolbar.cpp b/Engine/Source/Editor/Cascade/Private/SCascadePreviewToolbar.cpp
--- a/Engine/Source/Editor/Cascade/Private/SCascadePreviewToolbar.cpp
+++ b/Engine/Source/Editor/Cascade/Private/SCascadePreviewToolbar.cpp
@@ -96,6 +96,18 @@ TSharedRef<SWidget> SCascadePreviewViewportToolBar::GenerateViewMenu() const
 			}
 			Menu.EndSection();
 		}
+
+		static void BuildSignificanceMenu(FMenuBuilder& Menu)
+		{
+			Menu.BeginSection("CascadeSignificance", NSLOCTEXT("Cascade", "SignificanceHeader", "Required Significance"));
+			{
+				Menu.AddMenuEntry(FCascadeCommands::Get().Significance_Critical);
+				Menu.AddMenuEntry(FCascadeCommands::Get().Significance_High);
+				Menu.AddMenuEntry(FCascadeCommands::Get().Significance_Medium);
+				Menu.AddMenuEntry(FCascadeCommands::Get().Significance_Low);
+			}
+			Menu.EndSection();
+		}
 	};
 
 	ViewMenuBuilder.AddSubMenu(NSLOCTEXT("Cascade", "ViewOptionsSubMenu", "View Overlays"), 
@@ -110,6 +122,10 @@ TSharedRef<SWidget> SCascadePreviewViewportToolBar::GenerateViewMenu() const
 								FText::GetEmpty(), 
 								FNewMenuDelegate::CreateStatic(&Local::BuildDetailModesMenu));
 
+	ViewMenuBuilder.AddSubMenu(NSLOCTEXT("Cascade", "SignificanceSubMenu", "Significance"),
+								FText::GetEmpty(),
+								FNewMenuDelegate::CreateStatic(&Local::BuildSignificanceMenu));
+
 	ViewMenuBuilder.BeginSection("CascadeMiscPreview");
 	{
 		// Only display the orbit mode option if orbit cam controls are disabled
diff --git a/Engine/Source/Editor/Localization/Private/LocalizationConfigurationScript.cpp b/Engine/Source/Editor/Localization/Private/LocalizationConfigurationScript.cpp
--- a/Engine/Source/Editor/Localization/Private/LocalizationConfigurationScript.cpp
+++ b/Engine/Source/Editor/Localization/Private/LocalizationConfigurationScript.cpp
@@ -202,6 +202,7 @@ namespace LocalizationConfigurationScript
 			}
 
 			ConfigSection.Add( TEXT("ShouldGatherFromEditorOnlyData"), Target->Settings.GatherFromPackages.ShouldGatherFromEditorOnlyData ? TEXT("true") : TEXT("false") );
+			ConfigSection.Add( TEXT("SkipGatherCache"), Target->Settings.GatherFromPackages.SkipGatherCache ? TEXT("true") : TEXT("false") );
 		}
 
 		// GatherTextFromMetadata
diff --git a/Engine/Source/Editor/Localization/Public/LocalizationTargetTypes.h b/Engine/Source/Editor/Localization/Public/LocalizationTargetTypes.h
--- a/Engine/Source/Editor/Localization/Public/LocalizationTargetTypes.h
+++ b/Engine/Source/Editor/Localization/Public/LocalizationTargetTypes.h
@@ -109,6 +109,7 @@ struct FGatherTextFromPackagesConfiguration
 		: IsEnabled(true)
 		, FileExtensions(GetDefaultPackageFileExtensions())
 		, ShouldGatherFromEditorOnlyData(false)
+		, SkipGatherCache(false)
 	{
 	}
 
@@ -129,9 +130,13 @@ struct FGatherTextFromPackagesConfiguration
 	TArray<FGatherTextFileExtension> FileExtensions;
 
 	/* If enabled, data that is specified as editor-only may be processed for gathering. */
-	UPROPERTY(config, EditAnywhere, Category = "Filter")
+	UPROPERTY(config, EditAnywhere, Category = "Gather")
 	bool ShouldGatherFromEditorOnlyData;
 
+	/* Should we ignore the cached text in the package header and perform a full package load instead? */
+	UPROPERTY(config, EditAnywhere, Category = "Gather", AdvancedDisplay)
+	bool SkipGatherCache;
+
 	LOCALIZATION_API bool Validate(const FString& RootDirectory, FText& OutError) const;
 };
 
diff --git a/Engine/Source/Editor/LocalizationDashboard/Private/LocalizationTargetDetailCustomization.cpp b/Engine/Source/Editor/LocalizationDashboard/Private/LocalizationTargetDetailCustomization.cpp
--- a/Engine/Source/Editor/LocalizationDashboard/Private/LocalizationTargetDetailCustomization.cpp
+++ b/Engine/Source/Editor/LocalizationDashboard/Private/LocalizationTargetDetailCustomization.cpp
@@ -353,6 +353,7 @@ void FLocalizationTargetDetailCustomization::CustomizeDetails(IDetailLayoutBuild
 		DetailWidgetRow.NameContent()
 			[
 				SNew(STextBlock)
+				.Font(DetailBuilder.GetDetailFont())
 				.Text(LOCTEXT("LocalizationTargetLoadingPolicyRowName", "Loading Policy"))
 			];
 		DetailWidgetRow.ValueContent()
diff --git a/Engine/Source/Editor/SoundCueEditor/Private/SoundCueEditor.cpp b/Engine/Source/Editor/SoundCueEditor/Private/SoundCueEditor.cpp
--- a/Engine/Source/Editor/SoundCueEditor/Private/SoundCueEditor.cpp
+++ b/Engine/Source/Editor/SoundCueEditor/Private/SoundCueEditor.cpp
@@ -19,6 +19,7 @@
 #include "SDockTab.h"
 #include "GenericCommands.h"
 #include "Sound/SoundCue.h"
+#include "Sound/SoundNodeDialoguePlayer.h"
 
 #define LOCTEXT_NAMESPACE "SoundCueEditor"
 
@@ -392,7 +393,12 @@ void FSoundCueEditor::SyncInBrowser()
 			{
 				ObjectsToSync.AddUnique(SelectedWave->GetSoundWave());
 			}
-			//ObjectsToSync.AddUnique(SelectedNode->SoundNode);
+
+			USoundNodeDialoguePlayer* SelectedDialogue = Cast<USoundNodeDialoguePlayer>(SelectedNode->SoundNode);
+			if (SelectedDialogue && SelectedDialogue->GetDialogueWave())
+			{
+				ObjectsToSync.AddUnique(SelectedDialogue->GetDialogueWave());
+			}
 		}
 	}
 
@@ -418,6 +424,12 @@ bool FSoundCueEditor::CanSyncInBrowser() const
 			{
 				return true;
 			}
+
+			USoundNodeDialoguePlayer* SelectedDialogue = Cast<USoundNodeDialoguePlayer>(SelectedNode->SoundNode);
+			if (SelectedDialogue && SelectedDialogue->GetDialogueWave())
+			{
+				return true;
+			}
 		}
 	}
 	return false;
diff --git a/Engine/Source/Editor/SoundCueEditor/Private/SoundCueEditorUtilities.cpp b/Engine/Source/Editor/SoundCueEditor/Private/SoundCueEditorUtilities.cpp
--- a/Engine/Source/Editor/SoundCueEditor/Private/SoundCueEditorUtilities.cpp
+++ b/Engine/Source/Editor/SoundCueEditor/Private/SoundCueEditorUtilities.cpp
@@ -8,6 +8,7 @@
 #include "Sound/SoundNodeWavePlayer.h"
 #include "GraphEditor.h"
 #include "Sound/SoundCue.h"
+#include "Sound/SoundNodeDialoguePlayer.h"
 
 bool FSoundCueEditorUtilities::CanPasteNodes(const class UEdGraph* Graph)
 {
@@ -52,6 +53,35 @@ void FSoundCueEditorUtilities::CreateWaveContainers(TArray<USoundWave*>& Selecte
 	}
 }
 
+void FSoundCueEditorUtilities::CreateDialogueContainers(TArray<UDialogueWave*>& SelectedDialogues, USoundCue* SoundCue, TArray<USoundNode*>& OutPlayers, FVector2D Location)
+{
+	const int32 NodeSpacing = 70;
+
+	Location.Y -= static_cast<float>((SelectedDialogues.Num() - 1) * NodeSpacing) / 2.f;
+
+	for (int32 Index = 0; Index < SelectedDialogues.Num(); Index++)
+	{
+		UDialogueWave* NewDialogue = SelectedDialogues[Index];
+		if (NewDialogue)
+		{
+			USoundNodeDialoguePlayer* DialoguePlayer = SoundCue->ConstructSoundNode<USoundNodeDialoguePlayer>();
+
+			DialoguePlayer->SetDialogueWave(NewDialogue);
+
+			if (NewDialogue->ContextMappings.Num() == 1)
+			{
+				DialoguePlayer->DialogueWaveParameter.Context.Speaker = NewDialogue->ContextMappings[0].Context.Speaker;
+				DialoguePlayer->DialogueWaveParameter.Context.Targets = NewDialogue->ContextMappings[0].Context.Targets;
+			}
+
+			DialoguePlayer->GraphNode->NodePosX = Location.X - DialoguePlayer->GetGraphNode()->EstimateNodeWidth();
+			DialoguePlayer->GraphNode->NodePosY = Location.Y + (NodeSpacing * Index);
+
+			OutPlayers.Add(DialoguePlayer);
+		}
+	}
+}
+
 bool FSoundCueEditorUtilities::GetBoundsForSelectedNodes(const UEdGraph* Graph, class FSlateRect& Rect, float Padding)
 {
 	TSharedPtr<class ISoundCueEditor> SoundCueEditor = GetISoundCueEditorForObject(Graph);
diff --git a/Engine/Source/Editor/SoundCueEditor/Public/SoundCueEditorUtilities.h b/Engine/Source/Editor/SoundCueEditor/Public/SoundCueEditorUtilities.h
--- a/Engine/Source/Editor/SoundCueEditor/Public/SoundCueEditorUtilities.h
+++ b/Engine/Source/Editor/SoundCueEditor/Public/SoundCueEditorUtilities.h
@@ -30,6 +30,16 @@ public:
 	 */
 	static void CreateWaveContainers(TArray<class USoundWave*>& SelectedWaves, class USoundCue* SoundCue, TArray<class USoundNode*>& OutPlayers, FVector2D Location);
 
+	/** Adds USoundNodeDialoguePlayer nodes based on selected objects
+	 *
+	 * @param SelectedDialogues List of selected DialogueWaves to create player nodes for
+	 * @param SoundCue The SoundCue that the nodes will be part of
+	 * @param OutPlayers Stores all created nodes
+	 * @param Location Position of first created node
+	 */
+	static void CreateDialogueContainers(TArray<class UDialogueWave*>& SelectedDialogues, USoundCue* SoundCue, TArray<class USoundNode*>& OutPlayers, FVector2D Location);
+
+
 	/** Get the bounding area for the currently selected nodes
 	 *
 	 * @param Graph The Graph we are finding bounds for
diff --git a/Engine/Source/Editor/UnrealEd/Classes/Commandlets/GatherTextFromAssetsCommandlet.h b/Engine/Source/Editor/UnrealEd/Classes/Commandlets/GatherTextFromAssetsCommandlet.h
--- a/Engine/Source/Editor/UnrealEd/Classes/Commandlets/GatherTextFromAssetsCommandlet.h
+++ b/Engine/Source/Editor/UnrealEd/Classes/Commandlets/GatherTextFromAssetsCommandlet.h
@@ -5,23 +5,6 @@
 #include "Sound/DialogueTypes.h"
 #include "GatherTextFromAssetsCommandlet.generated.h"
 
-namespace EAssetTextGatherStatus
-{
-	enum 
-	{
-		BadBitMask = 1
-	};
-	enum Type
-	{
-		//							ErrorType << 1 + BadBit
-		None						= (0 << 1) + 0,
-		MissingKey					= (1 << 1) + 1,
-		MissingKey_Resolved			= (1 << 1) + 0,
-		IdentityConflict			= (2 << 1) + 1,
-		IdentityConflict_Resolved	= (2 << 1) + 0,
-	};
-}
-
 class FLocMetadataValue;
 
 /**
@@ -40,28 +23,9 @@ public:
 	virtual int32 Main(const FString& Params) override;
 	//~ End UCommandlet Interface
 
-private:
-	struct FConflictTracker
-	{
-		struct FEntry
-		{
-			FString PathToSourceString;
-			TSharedPtr<FString, ESPMode::ThreadSafe> SourceString;
-			EAssetTextGatherStatus::Type Status;
-		};
-
-		typedef TArray<FEntry> FEntryArray;
-		typedef TMap<FString, FEntryArray> FKeyTable;
-		typedef TMap<FString, FKeyTable> FNamespaceTable;
-
-		FNamespaceTable Namespaces;
-	};
-
 private:
 	static const FString UsageText;
 
-	FConflictTracker ConflictTracker;
-
 	bool bFixBroken;
 	bool ShouldGatherFromEditorOnlyData;
 };
diff --git a/Engine/Source/Editor/UnrealEd/Classes/CookOnTheSide/CookOnTheFlyServer.h b/Engine/Source/Editor/UnrealEd/Classes/CookOnTheSide/CookOnTheFlyServer.h
--- a/Engine/Source/Editor/UnrealEd/Classes/CookOnTheSide/CookOnTheFlyServer.h
+++ b/Engine/Source/Editor/UnrealEd/Classes/CookOnTheSide/CookOnTheFlyServer.h
@@ -1347,13 +1347,6 @@ private:
 	 */
 	bool SaveCurrentIniSettings( const ITargetPlatform* TargetPlatform ) const;
 
-
-	/**
-	 * Save the accumulated cooker stats to the Saved\Stats folder
-	 * saves to different file name pending running multiprocess cooker / not
-	 */
-	void SaveCookerStats() const;
-
 	/**
 	 * IsCookFlagSet
 	 * 
diff --git a/Engine/Source/Editor/UnrealEd/Classes/Factories/DialogueWaveFactory.h b/Engine/Source/Editor/UnrealEd/Classes/Factories/DialogueWaveFactory.h
--- a/Engine/Source/Editor/UnrealEd/Classes/Factories/DialogueWaveFactory.h
+++ b/Engine/Source/Editor/UnrealEd/Classes/Factories/DialogueWaveFactory.h
@@ -7,7 +7,7 @@
 #pragma once
 #include "DialogueWaveFactory.generated.h"
 
-UCLASS(hidecategories=Object)
+UCLASS(hidecategories=Object, MinimalAPI)
 class UDialogueWaveFactory : public UFactory
 {
 	GENERATED_UCLASS_BODY()
@@ -15,6 +15,22 @@ class UDialogueWaveFactory : public UFactory
 	//~ Begin UFactory Interface
 	virtual UObject* FactoryCreateNew(UClass* Class,UObject* InParent,FName Name,EObjectFlags Flags,UObject* Context,FFeedbackContext* Warn) override;
 	//~ Begin UFactory Interface	
+
+	/** An initial sound wave to place in the newly created dialogue wave */
+	UPROPERTY()
+	class USoundWave* InitialSoundWave;
+
+	/** An initial speaking dialogue voice to place in the newly created dialogue wave */
+	UPROPERTY()
+	class UDialogueVoice* InitialSpeakerVoice;
+
+	/** Whether an initial target dialogue voice should be set */
+	UPROPERTY()
+	bool HasSetInitialTargetVoice;
+
+	/** An initial target dialogue voices to place in the newly created dialogue wave */
+	UPROPERTY()
+	TArray<UDialogueVoice*> InitialTargetVoices;
 };
 
 
diff --git a/Engine/Source/Editor/UnrealEd/Classes/Factories/SoundCueFactoryNew.h b/Engine/Source/Editor/UnrealEd/Classes/Factories/SoundCueFactoryNew.h
--- a/Engine/Source/Editor/UnrealEd/Classes/Factories/SoundCueFactoryNew.h
+++ b/Engine/Source/Editor/UnrealEd/Classes/Factories/SoundCueFactoryNew.h
@@ -20,6 +20,10 @@ class USoundCueFactoryNew : public UFactory
 	/** An initial sound wave to place in the newly created cue */
 	UPROPERTY()
 	class USoundWave* InitialSoundWave;
+
+	/** An initial dialogue wave to place in the newly created cue */
+	UPROPERTY()
+	class UDialogueWave* InitialDialogueWave;
 };
 
 
diff --git a/Engine/Source/Editor/UnrealEd/Classes/SoundCueGraph/SoundCueGraphSchema.h b/Engine/Source/Editor/UnrealEd/Classes/SoundCueGraph/SoundCueGraphSchema.h
--- a/Engine/Source/Editor/UnrealEd/Classes/SoundCueGraph/SoundCueGraphSchema.h
+++ b/Engine/Source/Editor/UnrealEd/Classes/SoundCueGraph/SoundCueGraphSchema.h
@@ -119,6 +119,7 @@ class USoundCueGraphSchema : public UEdGraphSchema
 	virtual void BreakNodeLinks(UEdGraphNode& TargetNode) const override;
 	virtual void BreakPinLinks(UEdGraphPin& TargetPin, bool bSendsNodeNotifcation) const override;
 	virtual void DroppedAssetsOnGraph(const TArray<class FAssetData>& Assets, const FVector2D& GraphPosition, UEdGraph* Graph) const override;
+	virtual void DroppedAssetsOnNode(const TArray<FAssetData>& Assets, const FVector2D& GraphPosition, UEdGraphNode* Node) const override;
 	virtual int32 GetNodeSelectionCount(const UEdGraph* Graph) const override;
 	virtual TSharedPtr<FEdGraphSchemaAction> GetCreateCommentAction() const override;
 	//~ End EdGraphSchema Interface
diff --git a/Engine/Source/Editor/UnrealEd/Private/Commandlets/CookCommandlet.cpp b/Engine/Source/Editor/UnrealEd/Private/Commandlets/CookCommandlet.cpp
--- a/Engine/Source/Editor/UnrealEd/Private/Commandlets/CookCommandlet.cpp
+++ b/Engine/Source/Editor/UnrealEd/Private/Commandlets/CookCommandlet.cpp
@@ -1547,7 +1547,6 @@ bool UCookCommandlet::Cook(const TArray<ITargetPlatform*>& Platforms, TArray<FSt
 		}
 	}
 
-	IConsoleManager::Get().ProcessUserConsoleInput(TEXT("Tex.DerivedDataTimings"), *GWarn, NULL );
 	UPackage::WaitForAsyncFileWrites();
 
 	GetDerivedDataCacheRef().WaitForQuiescence(true);
diff --git a/Engine/Source/Editor/UnrealEd/Private/Commandlets/DerivedDataCacheCommandlet.cpp b/Engine/Source/Editor/UnrealEd/Private/Commandlets/DerivedDataCacheCommandlet.cpp
--- a/Engine/Source/Editor/UnrealEd/Private/Commandlets/DerivedDataCacheCommandlet.cpp
+++ b/Engine/Source/Editor/UnrealEd/Private/Commandlets/DerivedDataCacheCommandlet.cpp
@@ -241,7 +241,6 @@ int32 UDerivedDataCacheCommandlet::Main( const FString& Params )
 		}
 	}
 
-	IConsoleManager::Get().ProcessUserConsoleInput(TEXT("Tex.DerivedDataTimings"), *GWarn, NULL);
 	UE_LOG(LogDerivedDataCacheCommandlet, Display, TEXT("Waiting for shaders to finish."));
 	GShaderCompilingManager->FinishAllCompilation();
 	GDistanceFieldAsyncQueue->BlockUntilAllBuildsComplete();
diff --git a/Engine/Source/Editor/UnrealEd/Private/Commandlets/GatherTextFromAssetsCommandlet.cpp b/Engine/Source/Editor/UnrealEd/Private/Commandlets/GatherTextFromAssetsCommandlet.cpp
--- a/Engine/Source/Editor/UnrealEd/Private/Commandlets/GatherTextFromAssetsCommandlet.cpp
+++ b/Engine/Source/Editor/UnrealEd/Private/Commandlets/GatherTextFromAssetsCommandlet.cpp
@@ -138,54 +138,27 @@ void UGatherTextFromAssetsCommandlet::ProcessGatherableTextDataArray(const FStri
 	{
 		for (const FTextSourceSiteContext& TextSourceSiteContext : GatherableTextData.SourceSiteContexts)
 		{
-			const FString KeyName = TextSourceSiteContext.KeyName;
-
-			// Find existing entry from manifest or manifest dependencies.
-			FContext Context;
-			Context.Key = TextSourceSiteContext.KeyName;
-			const FLocMetadataObject DefaultMetadataObject;
-			Context.KeyMetadataObj = !(FLocMetadataObject::IsMetadataExactMatch(&TextSourceSiteContext.KeyMetaData, &DefaultMetadataObject)) ? MakeShareable(new FLocMetadataObject(TextSourceSiteContext.KeyMetaData)) : nullptr;
-			Context.InfoMetadataObj = !(FLocMetadataObject::IsMetadataExactMatch(&TextSourceSiteContext.InfoMetaData, &DefaultMetadataObject)) ? MakeShareable(new FLocMetadataObject(TextSourceSiteContext.InfoMetaData)) : nullptr;
-			Context.bIsOptional = TextSourceSiteContext.IsOptional;
-			Context.SourceLocation = TextSourceSiteContext.SiteDescription;
-
-			TSharedPtr< FManifestEntry > ExistingEntry = ManifestInfo->GetManifest()->FindEntryByContext( GatherableTextData.NamespaceName, Context );
-
-			if (!ExistingEntry.IsValid())
+			if (!TextSourceSiteContext.IsEditorOnly || ShouldGatherFromEditorOnlyData)
 			{
-				FString FileInfo;
-				ExistingEntry = ManifestInfo->FindDependencyEntryByContext( GatherableTextData.NamespaceName, Context, FileInfo );
-			}
+				if (TextSourceSiteContext.KeyName.IsEmpty())
+				{
+					UE_LOG(LogGatherTextFromAssetsCommandlet, Warning, TEXT("Detected missing key on asset \"%s\"."), *TextSourceSiteContext.SiteDescription);
+					continue;
+				}
 
-			FConflictTracker::FEntry NewEntry;
-			NewEntry.PathToSourceString = TextSourceSiteContext.SiteDescription;
-			NewEntry.SourceString = MakeShareable(new FString(GatherableTextData.SourceData.SourceString));
-			NewEntry.Status = EAssetTextGatherStatus::None;
+				static const FLocMetadataObject DefaultMetadataObject;
 
-			if (TextSourceSiteContext.KeyName.IsEmpty())
-			{
-				NewEntry.Status = EAssetTextGatherStatus::MissingKey;
-			}
+				FContext Context;
+				Context.Key = TextSourceSiteContext.KeyName;
+				Context.KeyMetadataObj = !(FLocMetadataObject::IsMetadataExactMatch(&TextSourceSiteContext.KeyMetaData, &DefaultMetadataObject)) ? MakeShareable(new FLocMetadataObject(TextSourceSiteContext.KeyMetaData)) : nullptr;
+				Context.InfoMetadataObj = !(FLocMetadataObject::IsMetadataExactMatch(&TextSourceSiteContext.InfoMetaData, &DefaultMetadataObject)) ? MakeShareable(new FLocMetadataObject(TextSourceSiteContext.InfoMetaData)) : nullptr;
+				Context.bIsOptional = TextSourceSiteContext.IsOptional;
+				Context.SourceLocation = TextSourceSiteContext.SiteDescription;
 
-			// Entry already exists, check for conflict.
-			if (ExistingEntry.IsValid() && !ExistingEntry->Source.Text.Equals(*NewEntry.SourceString, ESearchCase::CaseSensitive))
-			{
-				NewEntry.Status = EAssetTextGatherStatus::IdentityConflict;
-			}
+				FLocItem Source(GatherableTextData.SourceData.SourceString);
 
-			// Gather if it requires gathering and isn't in a bad state.
-			if ((NewEntry.Status & EAssetTextGatherStatus::BadBitMask) == false)
-			{
-				if (!TextSourceSiteContext.IsEditorOnly || ShouldGatherFromEditorOnlyData)
-				{
-					ManifestInfo->AddEntry(TextSourceSiteContext.SiteDescription, GatherableTextData.NamespaceName, *NewEntry.SourceString, Context);
-				}
+				ManifestInfo->AddEntry(TextSourceSiteContext.SiteDescription, GatherableTextData.NamespaceName, Source, Context);
 			}
-
-			// Add to conflict tracker.
-			FConflictTracker::FKeyTable& KeyTable = ConflictTracker.Namespaces.FindOrAdd(GatherableTextData.NamespaceName);
-			FConflictTracker::FEntryArray& EntryArray = KeyTable.FindOrAdd(TextSourceSiteContext.KeyName);
-			EntryArray.Add(NewEntry);
 		}
 	}
 }
@@ -438,7 +411,12 @@ int32 UGatherTextFromAssetsCommandlet::Main(const FString& Params)
 		UE_LOG(LogGatherTextFromAssetsCommandlet, Warning, TEXT("No files found or none passed the include/exclude criteria."));
 	}
 
-	const bool bSkipGatherCache = FParse::Param(FCommandLine::Get(), TEXT("SkipGatherCache"));
+	bool bSkipGatherCache = FParse::Param(FCommandLine::Get(), TEXT("SkipGatherCache"));
+	if (!bSkipGatherCache)
+	{
+		GetBoolFromConfig(*SectionName, TEXT("SkipGatherCache"), bSkipGatherCache, GatherTextConfigPath);
+	}
+	UE_LOG(LogGatherTextFromAssetsCommandlet, Log, TEXT("SkipGatherCache: %s"), bSkipGatherCache ? TEXT("true") : TEXT("false"));
 
 	TArray< FString > PackageFileNamesToLoad;
 	for (FString& PackageFile : PackageFileNamesToProcess)
@@ -450,7 +428,17 @@ int32 UGatherTextFromAssetsCommandlet::Main(const FString& Params)
 			FPackageFileSummary PackageFileSummary;
 			(*FileReader) << PackageFileSummary;
 
-			bool MustLoadForGather = bSkipGatherCache;
+			bool MustLoadForGather = false;
+
+			// Have we been asked to skip the cache of text that exists in the header of newer packages?
+			if (bSkipGatherCache && PackageFileSummary.GetFileVersionUE4() >= VER_UE4_SERIALIZE_TEXT_IN_PACKAGES)
+			{
+				// Fallback on the old package flag check.
+				if (PackageFileSummary.PackageFlags & PKG_RequiresLocalizationGather)
+				{
+					MustLoadForGather = true;
+				}
+			}
 
 			const FCustomVersion* const EditorVersion = PackageFileSummary.GetCustomVersionContainer().GetVersion(FEditorObjectVersion::GUID);
 
@@ -488,22 +476,19 @@ int32 UGatherTextFromAssetsCommandlet::Main(const FString& Params)
 				PackageFileNamesToLoad.Add(PackageFile);
 			}
 			// Process immediately packages that don't require loading to process.
-			else
+			else if (PackageFileSummary.GatherableTextDataOffset > 0)
 			{
+				FileReader->Seek(PackageFileSummary.GatherableTextDataOffset);
+
 				TArray<FGatherableTextData> GatherableTextDataArray;
+				GatherableTextDataArray.SetNum(PackageFileSummary.GatherableTextDataCount);
 
-				if (PackageFileSummary.GatherableTextDataOffset > 0)
+				for (int32 GatherableTextDataIndex = 0; GatherableTextDataIndex < PackageFileSummary.GatherableTextDataCount; ++GatherableTextDataIndex)
 				{
-					FileReader->Seek(PackageFileSummary.GatherableTextDataOffset);
-
-					for(int32 i = 0; i < PackageFileSummary.GatherableTextDataCount; ++i)
-					{
-						FGatherableTextData* GatherableTextData = new(GatherableTextDataArray) FGatherableTextData;
-						(*FileReader) << *GatherableTextData;
-					}
-
-					ProcessGatherableTextDataArray(PackageFile, GatherableTextDataArray);
+					(*FileReader) << GatherableTextDataArray[GatherableTextDataIndex];
 				}
+
+				ProcessGatherableTextDataArray(PackageFile, GatherableTextDataArray);
 			}
 		}
 	}
@@ -612,45 +597,6 @@ int32 UGatherTextFromAssetsCommandlet::Main(const FString& Params)
 		LoadedPackageFileNames.Empty(PackagesPerBatchCount);
 	}
 
-	for(auto i = ConflictTracker.Namespaces.CreateConstIterator(); i; ++i)
-	{
-		const FString& NamespaceName = i.Key();
-		const FConflictTracker::FKeyTable& KeyTable = i.Value();
-		for(auto j = KeyTable.CreateConstIterator(); j; ++j)
-		{
-			const FString& KeyName = j.Key();
-			const FConflictTracker::FEntryArray& EntryArray = j.Value();
-
-			for(int k = 0; k < EntryArray.Num(); ++k)
-			{
-				const FConflictTracker::FEntry& Entry = EntryArray[k];
-				switch(Entry.Status)
-				{
-				case EAssetTextGatherStatus::MissingKey:
-					{
-						UE_LOG(LogGatherTextFromAssetsCommandlet, Warning, TEXT("Detected missing key on asset \"%s\"."), *Entry.PathToSourceString);
-					}
-					break;
-				case EAssetTextGatherStatus::MissingKey_Resolved:
-					{
-						UE_LOG(LogGatherTextFromAssetsCommandlet, Warning, TEXT("Fixed missing key on asset \"%s\"."), *Entry.PathToSourceString);
-					}
-					break;
-				case EAssetTextGatherStatus::IdentityConflict:
-					{
-						UE_LOG(LogGatherTextFromAssetsCommandlet, Warning, TEXT("Detected duplicate identity with differing source on asset \"%s\"."), *Entry.PathToSourceString);
-					}
-					break;
-				case EAssetTextGatherStatus::IdentityConflict_Resolved:
-					{
-						UE_LOG(LogGatherTextFromAssetsCommandlet, Warning, TEXT("Fixed duplicate identity with differing source on asset \"%s\"."), *Entry.PathToSourceString);
-					}
-					break;
-				}
-			}
-		}
-	}
-
 	return 0;
 }
 
diff --git a/Engine/Source/Editor/UnrealEd/Private/CookOnTheFlyServer.cpp b/Engine/Source/Editor/UnrealEd/Private/CookOnTheFlyServer.cpp
--- a/Engine/Source/Editor/UnrealEd/Private/CookOnTheFlyServer.cpp
+++ b/Engine/Source/Editor/UnrealEd/Private/CookOnTheFlyServer.cpp
@@ -28,9 +28,6 @@
 #include "PhysicsPublic.h"
 #include "AutoSaveUtils.h"
 
-// cooking stats
-#include "CookingStatsModule.h"
-
 // cook by the book requirements
 #include "Commandlets/ChunkManifestGenerator.h"
 #include "Engine/WorldComposition.h"
@@ -49,9 +46,6 @@
 
 DEFINE_LOG_CATEGORY_STATIC(LogCook, Log, All);
 
-// these are the stats which are output to the Saved\Stats\<date>.csv file
-#define ENABLE_COOKER_STATS 1
-
 #define DEBUG_COOKONTHEFLY 0
 #define OUTPUT_TIMING 0
 
@@ -417,268 +411,6 @@ void OutputHierarchyTimers()
 #define OUTPUT_HIERARCHYTIMERS()
 #endif
 
-
-
-#if ENABLE_COOKER_STATS
-#include "TypeHash.h"
-#include "CookingStatsModule.h"
-
-// gather stats about a cooked package
-namespace CookOnTheFlyStats
-{
-	ICookingStats* GetCookingStats()
-	{
-		static ICookingStats* CookingStats = nullptr;
-		static bool bInitialized = false;
-		if (bInitialized == false)
-		{
-			FCookingStatsModule* CookingStatsModule = FModuleManager::LoadModulePtr<FCookingStatsModule>(TEXT("CookingStats"));
-			if (CookingStatsModule)
-			{
-				CookingStats = CookingStatsModule->Get();
-			}
-			bInitialized = true;
-		}
-		return CookingStats;
-	}
-
-
-
-	class FCookerStat
-	{
-	private:
-		bool bStarted;
-		double StartTime;
-		const FName TagName;
-	public:
-		FCookerStat(const FName& InTagName) : TagName(InTagName)
-		{
-			StartTime = FPlatformTime::Seconds();
-			bStarted = true;
-		}
-
-		~FCookerStat()
-		{
-			Stop();
-		}
-
-		void Stop()
-		{
-			ICookingStats* CookingStats = GetCookingStats();
-			if (CookingStats)
-			{
-				static const FName NAME_CookerStat(TEXT("CookerStat"));
-				double Duration = (FPlatformTime::Seconds() - StartTime) * 1000.0f;
-				CookingStats->AddTagValue(NAME_CookerStat, TagName, (float)(Duration));
-			}
-		}
-
-	};
-
-	class FCookerPackageStats
-	{
-	private:
-		const FName Filename; 
-	public:
-		FCookerPackageStats(const FName& InFilename) : Filename(InFilename)
-		{
-
-		};
-		void AddTag(const FName& TagName, const FString& Value)
-		{
-			ICookingStats* CookingStats = GetCookingStats();
-			if (CookingStats)
-			{
-				CookingStats->AddTagValue(Filename, TagName, Value);
-			}
-		}
-
-		void AddTag(const FName& TagName, const float Value)
-		{
-			ICookingStats* CookingStats = GetCookingStats();
-			if (CookingStats)
-			{
-				CookingStats->AddTagValue(Filename, TagName, Value);
-			}
-		}
-
-		bool GetTag(const FName& TagName, FString& Value) const
-		{
-			ICookingStats* CookingStats = GetCookingStats();
-			if (CookingStats)
-			{
-				return CookingStats->GetTagValue(Filename, TagName, Value);
-			}
-			return false;
-		}
-
-		const FName& GetFilename() const 
-		{
-			return Filename;
-		}
-	};
-
-	struct FPackageStat
-	{
-		FPackageStat(const FName& InPackageName, const FName& InStatName)
-		{
-			PackageName = InPackageName;
-			StatName = InStatName;
-		}
-
-		bool operator==(const FPackageStat& InPackageStat) const
-		{
-			return (InPackageStat.PackageName == PackageName) &&
-				(InPackageStat.StatName == StatName);
-		}
-
-		FName PackageName;
-		FName StatName;
-
-	};
-
-	struct FStatInfo
-	{
-		int HitCount;
-		bool bStarted;
-		double StartTime;
-		double Duration;
-	};
-
-	TMap<FPackageStat, FStatInfo> StatInfoMap;
-
-	struct FNameCache
-	{
-	public:
-		FNameCache(const FName& BaseName) 
-		{
-			Duration = FName(*FString::Printf(TEXT("%s_Duration"), *BaseName.ToString()));
-			HitCount = FName(*FString::Printf(TEXT("%s_HitCount"), *BaseName.ToString()));
-		}
-		FName Duration;
-		FName HitCount;
-	};
-	TMap<FName, FNameCache> NameCache;
-
-	const FNameCache& GetNameCache(const FName& BaseName)
-	{
-		const FNameCache* Result = NameCache.Find(BaseName);
-		if (Result == nullptr)
-		{
-			Result = &NameCache.Add(BaseName, FNameCache(BaseName));
-		}
-		return *Result;
-	}
-
-	FStatInfo* GetStatInfo(const FName& StatName, const FName& PackageName)
-	{
-		FPackageStat PackageStat(PackageName, StatName);
-
-		FStatInfo *StatInfo = StatInfoMap.Find(PackageStat);
-		if (StatInfo == nullptr)
-		{
-			StatInfo = &StatInfoMap.Add(PackageStat);
-		}
-		return StatInfo;
-	}
-
-	void StartPackageStat(const FName& Name, const FCookerPackageStats& CookerPackageStats)
-	{
-		FStatInfo* StatInfo = GetStatInfo(Name, CookerPackageStats.GetFilename());
-
-		if (StatInfo->bStarted == false)
-		{
-			StatInfo->StartTime = FPlatformTime::Seconds();
-			StatInfo->bStarted = true;
-		}
-	}
-
-	void EndPackageStat(const FName& Name, FCookerPackageStats& CookerPackageStats)
-	{
-		FStatInfo* StatInfo = GetStatInfo(Name, CookerPackageStats.GetFilename());
-
-		if (StatInfo->bStarted)
-		{
-			++StatInfo->HitCount;
-			double Duration = FPlatformTime::Seconds() - StatInfo->StartTime;
-			StatInfo->Duration += Duration;
-			const FNameCache& CachedFNames = GetNameCache(Name);
-
-			CookerPackageStats.AddTag(CachedFNames.Duration, FString::Printf(TEXT("%fms"), StatInfo->Duration*1000.0f));
-			CookerPackageStats.AddTag(CachedFNames.HitCount, FString::Printf(TEXT("%d"), StatInfo->HitCount));
-
-			StatInfo->bStarted = false;
-		}
-	}
-
-
-
-	class FCookerPackageScopeStat
-	{
-	private:
-		const FName& StatName;
-		FCookerPackageStats& PackageStats;
-	public:
-		FCookerPackageScopeStat(const FName& InStatName, FCookerPackageStats& InPackageStats) : StatName(InStatName), PackageStats(InPackageStats)
-		{
-			Start();
-		}
-
-		~FCookerPackageScopeStat()
-		{
-			End();
-		}
-
-		void Start()
-		{
-			StartPackageStat(StatName, PackageStats);
-		}
-		void End()
-		{
-			EndPackageStat(StatName, PackageStats);
-		}
-	};
-
-	uint32 GetTypeHash(const FPackageStat& PackageStat)
-	{
-		return HashCombine(GetTypeHash(PackageStat.PackageName), GetTypeHash(PackageStat.StatName));
-	}
-}
-
-// global cooking stat which isn't associated with a specific package
-#define SCOPE_COOKING_STAT(name) \
-	static const FName NAME_##name(#name); \
-	CookOnTheFlyStats::FCookerStat Stat_##name(NAME_##name);
-
-#define START_PACKAGE(filename) CookOnTheFlyStats::FCookerPackageStats CookerPackageStats(filename);
-
-#define SCOPE_PACKAGE_STAT(name) \
-	static const FName NAME_##name(#name); \
-	CookOnTheFlyStats::FCookerPackageScopeStat CookerPackageScopeStat##name(NAME_##name, CookerPackageStats);
-
-// static const FName NAME_StartTime(TEXT("StartTime"));
-
-#define START_PACKAGE_STAT(name) \
-	static const FName NAME_##name(TEXT(#name)); \
-	CookOnTheFlyStats::StartPackageStat(NAME_##name, CookerPackageStats);
-
-// CookerPackageStats.AddTag(NAME_##nameStartTime, FString::Printf(TEXT("%f"),FPlatformTime::Seconds()));
-
-#define END_PACKAGE_STAT(name) \
-	static const FName NAME2_##name(TEXT(#name)); \
-	CookOnTheFlyStats::EndPackageStat(NAME2_##name, CookerPackageStats);
-
-
-#else
-
-#define SCOPE_COOKING_STAT(name)
-#define START_PACKAGE(filename)
-#define SCOPE_PACKAGE_STAT(name)
-#define START_PACKAGE_STAT(name)
-#define END_PACKAGE_STAT(name)
-
-#endif
-
 //////////////////////////////////////////////////////////////////////////
 // FCookTimer
 // used as a helper to timeslice cooker functions
@@ -936,10 +668,6 @@ UCookOnTheFlyServer::UCookOnTheFlyServer(const FObjectInitializer& ObjectInitial
 
 UCookOnTheFlyServer::~UCookOnTheFlyServer()
 {
-	if (IsCookOnTheFlyMode())
-	{
-		SaveCookerStats();
-	}
 	check(TickChildCookers() == true);
 	if ( CookByTheBookOptions )
 	{		
@@ -968,8 +696,6 @@ TStatId UCookOnTheFlyServer::GetStatId() const
 
 bool UCookOnTheFlyServer::StartNetworkFileServer( const bool BindAnyPort )
 {
-	SCOPE_COOKING_STAT(StartNetworkFileServer);
-	
 	check( IsCookOnTheFlyMode() );
 	//GetDerivedDataCacheRef().WaitForQuiescence(false);
 
@@ -1613,17 +1339,6 @@ uint32 UCookOnTheFlyServer::TickCookOnTheSide( const float TimeSlice, uint32 &Co
 		}
 	}
 
-	{
-		static double SavedStatsStart = FPlatformTime::Seconds();
-		double Now = FPlatformTime::Seconds();
-		const double Timeout = 60 * 10; // 10 minutes
-		if ((Now - SavedStatsStart) > Timeout)
-		{
-			SaveCookerStats();
-			SavedStatsStart = Now;
-		}
-	}
-
 	while (!GIsRequestingExit || CurrentCookMode == ECookMode::CookByTheBook)
 	{
 		// if we just cooked a map then don't process anything the rest of this tick
@@ -1702,8 +1417,6 @@ uint32 UCookOnTheFlyServer::TickCookOnTheSide( const float TimeSlice, uint32 &Co
 		bool bLastLoadWasMap = false;
 		FString LastLoadedMapName;
 
-		START_PACKAGE(ToBuild.GetFilename());
-
 		const FString BuildFilename = ToBuild.GetFilename().ToString();
 
 		if ( ToBuild.GetFilename() != CurrentReentryData.FileName )
@@ -1744,10 +1457,6 @@ uint32 UCookOnTheFlyServer::TickCookOnTheSide( const float TimeSlice, uint32 &Co
 
 		if ( bShouldCook ) // if we should cook the package then cook it otherwise add it to the list of already cooked packages below
 		{
-			SCOPE_PACKAGE_STAT(AllOfLoadPackage);
-
-
-			SCOPE_TIMER(AllOfLoadPackage);
 			UPackage *Package = NULL;
 			{
 				FString PackageName;
@@ -1906,7 +1615,6 @@ uint32 UCookOnTheFlyServer::TickCookOnTheSide( const float TimeSlice, uint32 &Co
 		// this will load more packages :)
 		if ( IsCookByTheBookMode() )
 		{
-			SCOPE_PACKAGE_STAT(ResolveRedirectors);
 			SCOPE_TIMER(ResolveRedirectors);
 			GRedirectCollector.ResolveStringAssetReference();
 		}
@@ -1979,8 +1687,6 @@ uint32 UCookOnTheFlyServer::TickCookOnTheSide( const float TimeSlice, uint32 &Co
 
 		if ( PackagesToSave.Num() )
 		{
-			START_PACKAGE_STAT(CacheForCookedPlatformData);
-
 			SCOPE_TIMER(CallBeginCacheForCookedPlatformData);
 			// cache the resources for this package for each platform
 			
@@ -2015,14 +1721,11 @@ uint32 UCookOnTheFlyServer::TickCookOnTheSide( const float TimeSlice, uint32 &Co
 			}
 		}
 
-		END_PACKAGE_STAT(CacheForCookedPlatformData);
-
 		int32 FirstUnsolicitedPackage = PackagesToSave.Num();
 
 		// generate a list of other packages which were loaded with this one
 		if ( !IsCookByTheBookMode() || (CookByTheBookOptions->bDisableUnsolicitedPackages == false))
 		{
-			SCOPE_PACKAGE_STAT(UnsolicitedMarkup);
 			SCOPE_TIMER(UnsolicitedMarkup);
 
 			TArray<UObject *> ObjectsInOuter;
@@ -2089,8 +1792,6 @@ uint32 UCookOnTheFlyServer::TickCookOnTheSide( const float TimeSlice, uint32 &Co
 
 		bool bFinishedSave = true;
 
-		START_PACKAGE_STAT(SavePackageCacheForCookedPlatformData);
-
 		if ( PackagesToSave.Num() )
 		{
 			int32 OriginalPackagesToSaveCount = PackagesToSave.Num();
@@ -2236,11 +1937,8 @@ uint32 UCookOnTheFlyServer::TickCookOnTheSide( const float TimeSlice, uint32 &Co
 				}
 
 				ESavePackageResult SavePackageResult = ESavePackageResult::Error;
-				END_PACKAGE_STAT(SavePackageCacheForCookedPlatformData);
 				bool bSucceededSavePackage = false;
 				{
-					SCOPE_PACKAGE_STAT(SaveCookedPackage);
-					
 					SCOPE_TIMER(SaveCookedPackage);
 					uint32 SaveFlags = SAVE_KeepGUID | (bShouldSaveAsync ? SAVE_Async : SAVE_None) | (IsCookFlagSet(ECookInitializationFlags::Unversioned) ? SAVE_Unversioned : 0);
 					GOutputCookingWarnings = IsCookFlagSet(ECookInitializationFlags::OutputVerboseCookerWarnings);
@@ -2263,8 +1961,6 @@ uint32 UCookOnTheFlyServer::TickCookOnTheSide( const float TimeSlice, uint32 &Co
 					}
 					Timer.SavedPackage();
 				}
-				START_PACKAGE_STAT(SavePackageCacheForCookedPlatformData);
-
 
 				if ( IsCookingInEditor() == false )
 				{
@@ -2316,7 +2012,6 @@ uint32 UCookOnTheFlyServer::TickCookOnTheSide( const float TimeSlice, uint32 &Co
 				}
 			}
 		}
-		END_PACKAGE_STAT(SavePackageCacheForCookedPlatformData);
 
 		// if after all this our requested file didn't get saved then we need to mark it as saved (this can happen if the file loaded didn't match the one we saved)
 		// for example say we request A.uasset which doesn't exist however A.umap exists, our load call will succeed as it will find A.umap, then the save will save A.umap
@@ -2823,8 +2518,6 @@ static bool IsMobileHDR()
 
 void UCookOnTheFlyServer::Initialize( ECookMode::Type DesiredCookMode, ECookInitializationFlags InCookFlags, const FString &InOutputDirectoryOverride )
 {
-	SCOPE_COOKING_STAT(Initialize);
-
 	OutputDirectoryOverride = InOutputDirectoryOverride;
 	CurrentCookMode = DesiredCookMode;
 	CookFlags = InCookFlags;
@@ -3018,7 +2711,6 @@ void GetVersionFormatNumbersForIniVersionStrings( TArray<FString>& IniVersionStr
 
 bool UCookOnTheFlyServer::GetCurrentIniVersionStrings( const ITargetPlatform* TargetPlatform, TArray<FString>& IniVersionStrings ) const
 {
-	SCOPE_COOKING_STAT(GetCurrentIniVersionStrings);
 	// there is a list of important ini settings in the Editor config 
 #if INVALIDATE_ON_ANY_INI_CHANGE // use a full list of ini settings to invalidate the cooked content or use the entire engine / editor ini
 	TArray<FString> IniFiles;
@@ -3037,6 +2729,9 @@ bool UCookOnTheFlyServer::GetCurrentIniVersionStrings( const ITargetPlatform* Ta
 	IniFiles.Add(FString(TEXT("Game")));
 
 
+	TArray<FString> ExcludedIniParams;
+	GConfig->GetArray(TEXT("CookSettings"), TEXT("ExcludedIniParams"), ExcludedIniParams, GEditorIni);
+
 	for (const auto& IniFilename : IniFiles)
 	{
 		FConfigFile PlatformIniFile;
@@ -3048,6 +2743,23 @@ bool UCookOnTheFlyServer::GetCurrentIniVersionStrings( const ITargetPlatform* Ta
 
 			for (const auto& IniValue : IniSection.Value)
 			{
+				const FString ValueName = FString::Printf(TEXT("%s:%s"), *SectionName, *IniValue.Key.ToString());
+
+				bool bExclude = false;
+				for (const auto& Exclusion : ExcludedIniParams)
+				{
+					if (ValueName.Contains(Exclusion))
+					{
+						bExclude = true;
+						break;
+					}
+				}
+
+				if (bExclude)
+				{
+					continue;
+				}
+
 				FString ProcessedIniValue = IniValue.Value;
 				// get rid of shit we don't care about
 				ProcessedIniValue.ReplaceInline(TEXT("{"), TEXT(""));
@@ -3180,7 +2892,6 @@ bool UCookOnTheFlyServer::GetCurrentIniVersionStrings( const ITargetPlatform* Ta
 
 bool UCookOnTheFlyServer::GetCookedIniVersionStrings( const ITargetPlatform* TargetPlatform, TArray<FString>& IniVersionStrings ) const
 {
-	SCOPE_COOKING_STAT(GetCookedIniVersionStrings);
 	const FString EditorIni = FPaths::GameDir() / TEXT("CookedIniVersion.txt");
 	const FString SandboxEditorIni = ConvertToFullSandboxPath(*EditorIni, true);
 
@@ -3193,8 +2904,6 @@ bool UCookOnTheFlyServer::GetCookedIniVersionStrings( const ITargetPlatform* Tar
 
 bool UCookOnTheFlyServer::CacheIniVersionStringsMap( const ITargetPlatform* TargetPlatform ) const
 {
-	SCOPE_COOKING_STAT(CacheIniVersionStringsMap);
-
 	// check if the cached ones are filled out
 	const FName TargetPlatformName = FName(*TargetPlatform->PlatformName());
 	TArray<FString>* FoundCookedIniVersionStrings = CachedIniVersionStringsMap.Find( TargetPlatformName );
@@ -3212,8 +2921,6 @@ bool UCookOnTheFlyServer::CacheIniVersionStringsMap( const ITargetPlatform* Targ
 
 bool UCookOnTheFlyServer::IniSettingsOutOfDate( const ITargetPlatform* TargetPlatform ) const
 {
-	SCOPE_COOKING_STAT(IniSettingsOutOfDate);
-
 	TArray<FString> CurrentIniVersionStrings;
 	if ( GetCurrentIniVersionStrings(TargetPlatform, CurrentIniVersionStrings) == false )
 	{
@@ -3556,7 +3263,6 @@ void UCookOnTheFlyServer::GenerateAssetRegistry()
 	}
 	CookFlags |= ECookInitializationFlags::GeneratedAssetRegistry;
 
-	SCOPE_COOKING_STAT(GenerateAssetRegistry);
 	if (IsChildCooker())
 	{
 		// don't generate the asset registry
@@ -4075,36 +3781,6 @@ const FString UCookOnTheFlyServer::GetCookedAssetRegistryFilename(const FString&
 	return CookedAssetRegistryFilename;
 }
 
-void UCookOnTheFlyServer::SaveCookerStats() const
-{
-	static ICookingStats* CookingStats = nullptr;
-	static bool bInitialized = false;
-	if (!bInitialized)
-	{
-		static const FName CookingStatsName("CookingStats");
-		FCookingStatsModule* CookingStatsModule = FModuleManager::GetModulePtr<FCookingStatsModule>(CookingStatsName);
-		if (CookingStatsModule)
-		{
-			CookingStats = CookingStatsModule->Get();
-			bInitialized = true; // we don't want to initialize because the module couldn't be found, try again next time incase there are some stats
-		}
-	}
-	if (CookingStats)
-	{
-		if (IsCookByTheBookMode())
-		{
-			const FString StatsFilename = GetStatsFilename(CookByTheBookOptions->ChildCookFilename);
-			CookingStats->SaveStatsAsCSV(StatsFilename);
-		}
-		else if ( IsCookOnTheFlyMode() )
-		{
-			const FString StatsFilename = GetStatsFilename(FString(TEXT("")));
-			CookingStats->SaveStatsAsCSV(StatsFilename);
-		}
-	}
-
-}
-
 void UCookOnTheFlyServer::CookByTheBookFinished()
 {
 	check( IsInGameThread() );
@@ -4116,8 +3792,6 @@ void UCookOnTheFlyServer::CookByTheBookFinished()
 	GetDerivedDataCacheRef().WaitForQuiescence(true);
 	
 
-	SaveCookerStats();
-
 	if (IsChildCooker())
 	{
 		// if we are the child cooker create a list of all the packages which we think wes hould have cooked but didn't
@@ -4153,7 +3827,6 @@ void UCookOnTheFlyServer::CookByTheBookFinished()
 		}
 
 		check(CookByTheBookOptions->ChildUnsolicitedPackages.Num() == 0);
-		SCOPE_COOKING_STAT(SavingAssetRegistry);
 		SCOPE_TIMER(SavingAssetRegistry);
 		// Save modified asset registry with all streaming chunk info generated during cook
 		const FString& SandboxRegistryFilename = GetSandboxAssetRegistryFilename();
@@ -4403,7 +4076,6 @@ void UCookOnTheFlyServer::CreateSandboxFile()
 
 void UCookOnTheFlyServer::InitializeSandbox()
 {
-	SCOPE_COOKING_STAT(InitializeSandbox);
 	if ( SandboxFile == NULL )
 	{
 		ITargetPlatformManagerModule& TPM = GetTargetPlatformManagerRef();
diff --git a/Engine/Source/Editor/UnrealEd/Private/Factories/EditorFactories.cpp b/Engine/Source/Editor/UnrealEd/Private/Factories/EditorFactories.cpp
--- a/Engine/Source/Editor/UnrealEd/Private/Factories/EditorFactories.cpp
+++ b/Engine/Source/Editor/UnrealEd/Private/Factories/EditorFactories.cpp
@@ -108,6 +108,7 @@
 #include "Factories/SubUVAnimationFactory.h"
 #include "Particles/SubUVAnimation.h"
 #include "Kismet2/BlueprintEditorUtils.h"
+#include "Sound/SoundNodeDialoguePlayer.h"
 
 DEFINE_LOG_CATEGORY_STATIC(LogEditorFactories, Log, All);
 
@@ -1822,7 +1823,18 @@ UDialogueWaveFactory::UDialogueWaveFactory(const FObjectInitializer& ObjectIniti
 
 UObject* UDialogueWaveFactory::FactoryCreateNew(UClass* Class,UObject* InParent,FName Name,EObjectFlags Flags,UObject* Context,FFeedbackContext* Warn)
 {
-	return NewObject<UObject>(InParent, Class, Name, Flags);
+	check(Class == SupportedClass);
+	UDialogueWave* DialogueWave = NewObject<UDialogueWave>(InParent, Name, Flags);
+
+	if (InitialSoundWave)
+	{
+		DialogueWave->SpokenText = InitialSoundWave->SpokenText;
+		DialogueWave->bMature = InitialSoundWave->bMature;
+	}
+
+	DialogueWave->UpdateContext(DialogueWave->ContextMappings[0], InitialSoundWave, InitialSpeakerVoice, InitialTargetVoices);
+
+	return DialogueWave;
 }
 
 /*-----------------------------------------------------------------------------
@@ -2343,6 +2355,21 @@ UObject* USoundCueFactoryNew::FactoryCreateNew( UClass* Class, UObject* InParent
 		WavePlayer->GraphNode->NodePosX = -250;
 		WavePlayer->GraphNode->NodePosY = -35;
 	}
+	else if (InitialDialogueWave)
+	{
+		USoundNodeDialoguePlayer* DialoguePlayer = SoundCue->ConstructSoundNode<USoundNodeDialoguePlayer>();
+		SoundCue->FirstNode = DialoguePlayer;
+		SoundCue->LinkGraphNodesFromSoundNodes();
+		DialoguePlayer->SetDialogueWave(InitialDialogueWave);
+		DialoguePlayer->GraphNode->NodePosX = -250 - DialoguePlayer->GetGraphNode()->EstimateNodeWidth();
+		DialoguePlayer->GraphNode->NodePosY = -35;
+
+		if (InitialDialogueWave->ContextMappings.Num() == 1)
+		{
+			DialoguePlayer->DialogueWaveParameter.Context.Speaker = InitialDialogueWave->ContextMappings[0].Context.Speaker;
+			DialoguePlayer->DialogueWaveParameter.Context.Targets = InitialDialogueWave->ContextMappings[0].Context.Targets;
+		}
+	}
 
 	return SoundCue;
 }
diff --git a/Engine/Source/Editor/UnrealEd/Private/SoundCueGraphNode.cpp b/Engine/Source/Editor/UnrealEd/Private/SoundCueGraphNode.cpp
--- a/Engine/Source/Editor/UnrealEd/Private/SoundCueGraphNode.cpp
+++ b/Engine/Source/Editor/UnrealEd/Private/SoundCueGraphNode.cpp
@@ -13,6 +13,7 @@
 #include "GenericCommands.h"
 #include "Sound/SoundCue.h"
 #include "Engine/Font.h"
+#include "Sound/SoundNodeDialoguePlayer.h"
 
 #define LOCTEXT_NAMESPACE "SoundCueGraphNode"
 
@@ -243,6 +244,10 @@ void USoundCueGraphNode::GetContextMenuActions(const FGraphNodeContextMenuBuilde
 			{
 				Context.MenuBuilder->AddMenuEntry(FSoundCueGraphEditorCommands::Get().BrowserSync);
 			}
+			else if (Cast<USoundNodeDialoguePlayer>(SoundNode))
+			{
+				Context.MenuBuilder->AddMenuEntry(FSoundCueGraphEditorCommands::Get().BrowserSync);
+			}
 		}
 		Context.MenuBuilder->EndSection();
 	}
