Copying //UE4/Orion-Staging to //UE4/Main (Origin //Orion/Dev-General @ 2870388)

==========================
MAJOR FEATURES + CHANGES
==========================

Change 2870336 on 2016/02/17 by [REDACTED_PERSON]

	Continued splitting up Orion Build
	* Restructure from platform based MakeBuild steps in to a PS4, Server, and Windows Client MakeBuild
	* Cook server data only once for both Windows and Linux (windows reuses Linux server data)
	* Split compilation of Win64 Client and Server such that MakeBuild_Server only builds Server and MakeBuild_WindowsClient only builds Client
	#tests Preflight and generated Windows Client and Server work to play game

Change 2870026 on 2016/02/17 by [REDACTED_PERSON]

	Don't allow array shrinking when removing the corruption wrapper trailer.
	Updating CIS Counter

Change 2869725 on 2016/02/17 by [REDACTED_PERSON]

	More analytics and QoS stats added for 0.19.

	#tests Ran Windows client and Linux server on compatible content.

Change 2869705 on 2016/02/16 by [REDACTED_PERSON]

	Fix replicated properties and call RepNotifies of startup actors when scrubbing in replays.
	This is the engine support for fixing OR-6817, towers not respawning when rewinding replays.

	#tests golden path, replays, ps4 nomcp

Change 2869644 on 2016/02/16 by [REDACTED_PERSON]

	#ORION_DEV - Merge MAIN (0.18) at CL# 2869635

	#Tests:none

Change 2869586 on 2016/02/16 by [REDACTED_PERSON]

	Fix texturestreaming RHI flushes.
	#test goldenpath

Change 2869279 on 2016/02/16 by [REDACTED_PERSON]

	fixed minion hit reaction directions
	#orion OR-13953
	#tests PIE: hit minions with various abilities from different angles, checked velocity of death particles when killed by abilities and towers

Change 2869277 on 2016/02/16 by [REDACTED_PERSON]

	During cook, when a package is not ready to save, actually early out of the saving code. Saves somewhere in the 130s to 200s range for cooks.
	#tests local windows cooks, preflight PS4 cooks

Change 2869132 on 2016/02/16 by [REDACTED_PERSON]

	Added a function to AISenseConfig allowing native-code MaxAge configuration #UE4

	#test none required

Change 2868981 on 2016/02/16 by [REDACTED_PERSON]

	remove -LogCookStats cmdline check, always log cook stats. -SendCookAnalytics flag is still used.
	This was requested by NickP.
	#tests local windows cooks

Change 2868975 on 2016/02/16 by [REDACTED_PERSON]

	Don't submit DDC usage stats for zero-sized events.
	#tests local windows cook

Change 2868956 on 2016/02/16 by [REDACTED_PERSON]

	#ORION_DEV - Merge MAIN (0.18) at CL# 2868926

	#Tests:none

Change 2868889 on 2016/02/16 by [REDACTED_PERSON]

	Sequencer: Only allow transport control binding when editing level editor sequencers.

	#tests none

Change 2868663 on 2016/02/16 by [REDACTED_PERSON]

	downgrade warning to display

	#tests compile

Change 2868624 on 2016/02/16 by [REDACTED_PERSON]

	Re-Enable Defrag validation for devgeneral
	#test none

Change 2868493 on 2016/02/16 by [REDACTED_PERSON]

	Added a few more stats to morph target updates to try and narrow down hitches
	#tests pie, -game Win64

Change 2868445 on 2016/02/16 by [REDACTED_PERSON]

	Linux: report crashes due to stack overflow (OR-14519).

	- Reserve memory for alternative stack for signal handlers. Adds about 128KB memory per thread.
	- Force process spawning to use vfork() when no pipes are needed.
	- Ignore all signals except explicitly handled.
	- Prevent signals from being raised while another one is handled.
	- Added "debug threadrecurse" and "debug threadstackoverflow" to test that.

diff --git a/Engine/Plugins/Experimental/SignificanceManager/Source/SignificanceManager/Private/SignificanceManager.cpp b/Engine/Plugins/Experimental/SignificanceManager/Source/SignificanceManager/Private/SignificanceManager.cpp
--- a/Engine/Plugins/Experimental/SignificanceManager/Source/SignificanceManager/Private/SignificanceManager.cpp
+++ b/Engine/Plugins/Experimental/SignificanceManager/Source/SignificanceManager/Private/SignificanceManager.cpp
@@ -3,6 +3,7 @@
 #include "SignificanceManagerPrivate.h"
 #include "SignificanceManager.h"
 #include "DisplayDebugHelpers.h"
+#include "ParallelFor.h"
 
 IMPLEMENT_MODULE( FSignificanceManagerModule, SignificanceManager );
 
@@ -115,6 +116,11 @@ void USignificanceManager::BeginDestroy()
 	ManagedObjectsByTag.Reset();
 }
 
+UWorld* USignificanceManager::GetWorld()const
+{
+	return CastChecked<UWorld>(GetOuter());
+}
+
 void USignificanceManager::RegisterObject(UObject* Object, const FName Tag, FSignificanceFunction SignificanceFunction, FPostSignificanceFunction InPostSignificanceFunction)
 {
 	INC_DWORD_STAT(STAT_SignificanceManager_NumObjects);
@@ -186,6 +192,10 @@ void USignificanceManager::UnregisterObject(UObject* Object)
 		{
 			check(ObjectsWithTag[0] == ObjectInfo);
 			ManagedObjectsByTag.Remove(ObjectInfo->GetTag());
+			if (ObjectInfo->PostSignificanceFunction != nullptr)
+			{
+				ObjectInfo->PostSignificanceFunction(ObjectInfo->GetObject(), ObjectInfo->Significance, 1.0f, true);
+			}
 		}
 		else
 		{
@@ -202,7 +212,10 @@ void USignificanceManager::UnregisterAll(FName Tag)
 		for (const FManagedObjectInfo* ManagedObj : *ObjectsWithTag)
 		{
 			ManagedObjects.Remove(ManagedObj->GetObject());
-			ManagedObj->PostSignificanceFunction(ManagedObj->GetObject(), ManagedObj->Significance, 1.0f);
+			if (ManagedObj->PostSignificanceFunction != nullptr)
+			{
+				ManagedObj->PostSignificanceFunction(ManagedObj->GetObject(), ManagedObj->Significance, 1.0f, true);
+			}
 		}
 		ManagedObjectsByTag.Remove(Tag);
 	}
@@ -288,7 +301,10 @@ void USignificanceManager::FManagedObjectInfo::UpdateSignificance(const TArray<F
 		Significance = 0.f;
 	}
 
-	PostSignificanceFunction(Object, OldSignificance, Significance);
+	if (PostSignificanceFunction != nullptr)
+	{
+		PostSignificanceFunction(Object, OldSignificance, Significance, false);
+	}
 }
 
 void USignificanceManager::Update(const TArray<FTransform>& InViewpoints)
@@ -299,16 +315,21 @@ void USignificanceManager::Update(const TArray<FTransform>& InViewpoints)
 
 	{
 		SCOPE_CYCLE_COUNTER(STAT_SignificanceManager_SignificanceUpdate);
+
+		TArray<FManagedObjectInfo*> ObjArray;
+		ManagedObjects.GenerateValueArray(ObjArray);
+
 		//TODO: Consider parallel for
-		for (TPair<UObject*, FManagedObjectInfo*>& ObjectToObjectInfoPair : ManagedObjects)
+		//for (TPair<UObject*, FManagedObjectInfo*>& ObjectToObjectInfoPair : ManagedObjects)
+		ParallelFor(ObjArray.Num(),
+			[&](int32 Index)
 		{
-			const UObject* Object = ObjectToObjectInfoPair.Key;
-			FManagedObjectInfo* ObjectInfo = ObjectToObjectInfoPair.Value;
+			FManagedObjectInfo* ObjectInfo = ObjArray[Index];
 
-			checkSlow(Object->IsValidLowLevel());
+			checkSlow(ObjectInfo->GetObject()->IsValidLowLevel());
 
 			ObjectInfo->UpdateSignificance(Viewpoints);
-		}
+		});
 	}
 
 	{
diff --git a/Engine/Plugins/Experimental/SignificanceManager/Source/SignificanceManager/Public/SignificanceManager.h b/Engine/Plugins/Experimental/SignificanceManager/Source/SignificanceManager/Public/SignificanceManager.h
--- a/Engine/Plugins/Experimental/SignificanceManager/Source/SignificanceManager/Public/SignificanceManager.h
+++ b/Engine/Plugins/Experimental/SignificanceManager/Source/SignificanceManager/Public/SignificanceManager.h
@@ -62,7 +62,7 @@ class SIGNIFICANCEMANAGER_API USignificanceManager : public UObject
 
 public:
 	typedef TFunction<float(UObject*, const FTransform&)> FSignificanceFunction;
-	typedef TFunction<void(UObject*, float, float)> FPostSignificanceFunction;
+	typedef TFunction<void(UObject*, float, float, bool)> FPostSignificanceFunction;
 
 	struct FManagedObjectInfo
 	{
@@ -75,7 +75,7 @@ public:
 		FManagedObjectInfo(UObject* InObject, FName InTag, FSignificanceFunction InSignificanceFunction, FPostSignificanceFunction InPostSignificanceFunction)
 			: Object(InObject)
 			, Tag(InTag)
-			, Significance(-1.0f)
+			, Significance(1.0f)
 			, SignificanceFunction(InSignificanceFunction)
 			, PostSignificanceFunction(InPostSignificanceFunction)
 		{
@@ -106,6 +106,7 @@ public:
 
 	// Begin UObject overrides
 	virtual void BeginDestroy() override;
+	virtual UWorld* GetWorld()const override;
 	// End UObject overrides
 
 	// Overridable function to update the managed objects' significance
@@ -148,10 +149,10 @@ public:
 		return CastChecked<T>(Get(World), ECastCheckedType::NullAllowed);
 	}
 
-protected:
-
 	// Returns the list of viewpoints currently being represented by the significance manager
 	const TArray<FTransform>& GetViewpoints() const { return Viewpoints; }
+protected:
+
 
 	// Whether the significance manager should be created on a client. Only used from CDO and 
 	uint32 bCreateOnClient:1;
diff --git a/Engine/Shaders/BasePassPixelShader.usf b/Engine/Shaders/BasePassPixelShader.usf
--- a/Engine/Shaders/BasePassPixelShader.usf
+++ b/Engine/Shaders/BasePassPixelShader.usf
@@ -593,7 +593,7 @@ uint GetSelectiveOutputMask()
 #if !WRITES_PRECSHADOWFACTOR_TO_GBUFFER
 	Mask |= SKIP_PRECSHADOW_MASK;
 #endif
-#if WRITES_PRECSHADOWFACTOR_TO_ZERO
+#if WRITES_PRECSHADOWFACTOR_ZERO
 	Mask |= ZERO_PRECSHADOW_MASK;
 #endif
 #if !WRITES_VELOCITY_TO_GBUFFER
diff --git a/Engine/Source/Developer/DerivedDataCache/Private/DerivedDataBackendAsyncPutWrapper.h b/Engine/Source/Developer/DerivedDataCache/Private/DerivedDataBackendAsyncPutWrapper.h
--- a/Engine/Source/Developer/DerivedDataCache/Private/DerivedDataBackendAsyncPutWrapper.h
+++ b/Engine/Source/Developer/DerivedDataCache/Private/DerivedDataBackendAsyncPutWrapper.h
@@ -1,6 +1,7 @@
 // Copyright 1998-2016 Epic Games, Inc. All Rights Reserved.
 
 #pragma once
+#include "DerivedDataCacheUsageStats.h"
 
 /** 
  * Thread safe set helper
@@ -56,16 +57,19 @@ public:
 	FThreadSet*							FilesInFlight;
 	/**If true, then do not attempt skip the put even if CachedDataProbablyExists returns true **/
 	bool								bPutEvenIfExists;
+	/** Usage stats to track thread times. */
+	FDerivedDataCacheUsageStats&        UsageStats;
 
 	/** Constructor
 	*/
-	FCachePutAsyncWorker(const TCHAR* InCacheKey, const TArray<uint8>* InData, FDerivedDataBackendInterface* InInnerBackend, bool InbPutEvenIfExists, FDerivedDataBackendInterface* InInflightCache, FThreadSet* InInFilesInFlight)
+	FCachePutAsyncWorker(const TCHAR* InCacheKey, const TArray<uint8>* InData, FDerivedDataBackendInterface* InInnerBackend, bool InbPutEvenIfExists, FDerivedDataBackendInterface* InInflightCache, FThreadSet* InInFilesInFlight, FDerivedDataCacheUsageStats& InUsageStats)
 		: CacheKey(InCacheKey)
 		, Data(*InData)
 		, InnerBackend(InInnerBackend)
 		, InflightCache(InInflightCache)
 		, FilesInFlight(InInFilesInFlight)
 		, bPutEvenIfExists(InbPutEvenIfExists)
+		, UsageStats(InUsageStats)
 	{
 		check(InnerBackend);
 	}
@@ -73,19 +77,23 @@ public:
 	/** Call the inner backend and when that completes, remove the memory cache */
 	void DoWork()
 	{
+		COOK_STAT(auto Timer = UsageStats.TimePut());
 		bool bOk = true;
-		InnerBackend->PutCachedData(*CacheKey, Data, bPutEvenIfExists);
-		if (InflightCache)
+		const bool bAlreadyExists = InnerBackend->CachedDataProbablyExists(*CacheKey);
+		if (!bAlreadyExists || bPutEvenIfExists)
+		{
+			InnerBackend->PutCachedData(*CacheKey, Data, bPutEvenIfExists);
+			COOK_STAT(Timer.AddHit(Data.Num()));
+		}
+		// if it already existed, don't bother checking if we need to retry. We don't.
+		if (InflightCache && !bAlreadyExists && !InnerBackend->CachedDataProbablyExists(*CacheKey))
 		{
+			// retry
+			InnerBackend->PutCachedData(*CacheKey, Data, false);
 			if (!InnerBackend->CachedDataProbablyExists(*CacheKey))
 			{
-				// retry
-				InnerBackend->PutCachedData(*CacheKey, Data, false);
-				if (!InnerBackend->CachedDataProbablyExists(*CacheKey))
-				{
-					UE_LOG(LogDerivedDataCache, Warning, TEXT("FDerivedDataBackendAsyncPutWrapper: Put failed, keeping in memory copy %s."),*CacheKey);
-					bOk = false;
-				}
+				UE_LOG(LogDerivedDataCache, Warning, TEXT("FDerivedDataBackendAsyncPutWrapper: Put failed, keeping in memory copy %s."),*CacheKey);
+				bOk = false;
 			}
 		}
 		if (bOk && InflightCache)
@@ -154,7 +162,10 @@ public:
 	 */
 	virtual bool CachedDataProbablyExists(const TCHAR* CacheKey) override
 	{
-		return (InflightCache && InflightCache->CachedDataProbablyExists(CacheKey)) || InnerBackend->CachedDataProbablyExists(CacheKey);
+		COOK_STAT(auto Timer = UsageStats.TimeProbablyExists());
+		bool Result = (InflightCache && InflightCache->CachedDataProbablyExists(CacheKey)) || InnerBackend->CachedDataProbablyExists(CacheKey);
+		COOK_STAT(if (Result) {	Timer.AddHit(0); });
+		return Result;
 	}
 	/**
 	 * Synchronous retrieve of a cache item
@@ -165,11 +176,17 @@ public:
 	 */
 	virtual bool GetCachedData(const TCHAR* CacheKey, TArray<uint8>& OutData) override
 	{
+		COOK_STAT(auto Timer = UsageStats.TimeGet());
 		if (InflightCache && InflightCache->GetCachedData(CacheKey, OutData))
 		{
+			COOK_STAT(Timer.AddHit(OutData.Num()));
 			return true;
 		}
 		bool bSuccess = InnerBackend->GetCachedData(CacheKey, OutData);
+		if (bSuccess)
+		{
+			COOK_STAT(Timer.AddHit(OutData.Num()));
+		}
 		return bSuccess;
 	}
 	/**
@@ -181,6 +198,7 @@ public:
 	 */
 	virtual void PutCachedData(const TCHAR* CacheKey, TArray<uint8>& InData, bool bPutEvenIfExists) override
 	{
+		COOK_STAT(auto Timer = UsageStats.TimePut());
 		if (!InnerBackend->IsWritable())
 		{
 			return; // no point in continuing down the chain
@@ -198,8 +216,9 @@ public:
 			}
 			InflightCache->PutCachedData(CacheKey, InData, true); // temp copy stored in memory while the async task waits to complete
 		}
+		COOK_STAT(Timer.UntrackHitOrMiss());
 		FDerivedDataBackend::Get().AddToAsyncCompletionCounter(1);
-		(new FAutoDeleteAsyncTask<FCachePutAsyncWorker>(CacheKey, &InData, InnerBackend, bPutEvenIfExists, InflightCache.GetOwnedPointer(), &FilesInFlight))->StartBackgroundTask();
+		(new FAutoDeleteAsyncTask<FCachePutAsyncWorker>(CacheKey, &InData, InnerBackend, bPutEvenIfExists, InflightCache.GetOwnedPointer(), &FilesInFlight, UsageStats))->StartBackgroundTask();
 	}
 
 	virtual void RemoveCachedData(const TCHAR* CacheKey, bool bTransient) override
@@ -218,7 +237,25 @@ public:
 		}
 		InnerBackend->RemoveCachedData(CacheKey, bTransient);
 	}
+
+	virtual void GatherUsageStats(TMap<FString, FDerivedDataCacheUsageStats>& UsageStatsMap, FString&& GraphPath) override
+	{
+		COOK_STAT(
+			{
+			UsageStatsMap.Add(GraphPath + TEXT(": AsyncPut"), UsageStats);
+			if (InnerBackend)
+			{
+				InnerBackend->GatherUsageStats(UsageStatsMap, GraphPath + TEXT(". 0"));
+			}
+			if (InflightCache)
+			{
+				InflightCache->GatherUsageStats(UsageStatsMap, GraphPath + TEXT(". 1"));
+			}
+		});
+	}
+
 private:
+	FDerivedDataCacheUsageStats UsageStats;
 
 	/** Backend to use for storage, my responsibilities are about async puts **/
 	FDerivedDataBackendInterface*					InnerBackend;
diff --git a/Engine/Source/Developer/DerivedDataCache/Private/DerivedDataBackendCorruptionWrapper.h b/Engine/Source/Developer/DerivedDataCache/Private/DerivedDataBackendCorruptionWrapper.h
--- a/Engine/Source/Developer/DerivedDataCache/Private/DerivedDataBackendCorruptionWrapper.h
+++ b/Engine/Source/Developer/DerivedDataCache/Private/DerivedDataBackendCorruptionWrapper.h
@@ -90,7 +90,13 @@ public:
 	 */
 	virtual bool CachedDataProbablyExists(const TCHAR* CacheKey) override
 	{
-		return InnerBackend->CachedDataProbablyExists(CacheKey);
+		COOK_STAT(auto Timer = UsageStats.TimeProbablyExists());
+		bool Result = InnerBackend->CachedDataProbablyExists(CacheKey);
+		if (Result)
+		{
+			COOK_STAT(Timer.AddHit(0));
+		}
+		return Result;
 	}
 	/**
 	 * Synchronous retrieve of a cache item
@@ -101,6 +107,7 @@ public:
 	 */
 	virtual bool GetCachedData(const TCHAR* CacheKey, TArray<uint8>& OutData) override
 	{
+		COOK_STAT(auto Timer = UsageStats.TimeGet());
 		bool bOk = InnerBackend->GetCachedData(CacheKey, OutData);
 		if (bOk)
 		{
@@ -113,7 +120,7 @@ public:
 			{
 				FDerivedDataTrailer Trailer;
 				FMemory::Memcpy(&Trailer,&OutData[OutData.Num() - sizeof(FDerivedDataTrailer)], sizeof(FDerivedDataTrailer));
-				OutData.RemoveAt(OutData.Num() - sizeof(FDerivedDataTrailer),sizeof(FDerivedDataTrailer));
+				OutData.RemoveAt(OutData.Num() - sizeof(FDerivedDataTrailer),sizeof(FDerivedDataTrailer), false);
 				FDerivedDataTrailer RecomputedTrailer(OutData);
 				if (Trailer == RecomputedTrailer)
 				{
@@ -135,6 +142,10 @@ public:
 		{
 			OutData.Empty();
 		}
+		else
+		{
+			COOK_STAT(Timer.AddHit(OutData.Num()));
+		}
 		return bOk;
 	}
 	/**
@@ -146,11 +157,12 @@ public:
 	 */
 	virtual void PutCachedData(const TCHAR* CacheKey, TArray<uint8>& InData, bool bPutEvenIfExists) override
 	{
+		COOK_STAT(auto Timer = UsageStats.TimePut());
 		if (!InnerBackend->IsWritable())
 		{
 			return; // no point in continuing down the chain
 		}
-
+		COOK_STAT(Timer.AddHit(InData.Num()));
 		// Get rid of the double copy!
 		TArray<uint8> Data;
 		Data.Reset( InData.Num() + sizeof(FDerivedDataTrailer) );
@@ -170,7 +182,21 @@ public:
 		}
 		return InnerBackend->RemoveCachedData(CacheKey, bTransient);
 	}
+
+	virtual void GatherUsageStats(TMap<FString, FDerivedDataCacheUsageStats>& UsageStatsMap, FString&& GraphPath) override
+	{
+		COOK_STAT(
+		{
+			UsageStatsMap.Add(GraphPath + TEXT(": CorruptionWrapper"), UsageStats);
+			if (InnerBackend)
+			{
+				InnerBackend->GatherUsageStats(UsageStatsMap, GraphPath + TEXT(". 0"));
+			}
+		});
+	}
+
 private:
+	FDerivedDataCacheUsageStats UsageStats;
 
 	/** Backend to use for storage, my responsibilities are about corruption **/
 	FDerivedDataBackendInterface* InnerBackend;
diff --git a/Engine/Source/Developer/DerivedDataCache/Private/DerivedDataBackendVerifyWrapper.h b/Engine/Source/Developer/DerivedDataCache/Private/DerivedDataBackendVerifyWrapper.h
--- a/Engine/Source/Developer/DerivedDataCache/Private/DerivedDataBackendVerifyWrapper.h
+++ b/Engine/Source/Developer/DerivedDataCache/Private/DerivedDataBackendVerifyWrapper.h
@@ -28,15 +28,18 @@ public:
 
 	virtual bool CachedDataProbablyExists(const TCHAR* CacheKey) override
 	{
+		COOK_STAT(auto Timer = UsageStats.TimeProbablyExists());
 		FScopeLock ScopeLock(&SynchronizationObject);
 		if (AlreadyTested.Contains(FString(CacheKey)))
 		{
+			COOK_STAT(Timer.AddHit(0));
 			return true;
 		}
 		return false;
 	}
 	virtual bool GetCachedData(const TCHAR* CacheKey, TArray<uint8>& OutData) override
 	{
+		COOK_STAT(auto Timer = UsageStats.TimeGet());
 		bool bAlreadyTested = false;
 		{
 			FScopeLock ScopeLock(&SynchronizationObject);
@@ -47,12 +50,18 @@ public:
 		}
 		if (bAlreadyTested)
 		{
-			return InnerBackend->GetCachedData(CacheKey, OutData);
+			bool bResult = InnerBackend->GetCachedData(CacheKey, OutData);
+			if (bResult)
+			{
+				COOK_STAT(Timer.AddHit(OutData.Num()));
+			}
+			return bResult;
 		}
 		return false;
 	}
 	virtual void PutCachedData(const TCHAR* CacheKey, TArray<uint8>& InData, bool bPutEvenIfExists) override
 	{
+		COOK_STAT(auto Timer = UsageStats.TimePut());
 		bool bAlreadyTested = false;
 		{
 			FScopeLock ScopeLock(&SynchronizationObject);
@@ -67,6 +76,7 @@ public:
 		}
 		if (!bAlreadyTested)
 		{
+			COOK_STAT(Timer.AddHit(InData.Num()));
 			TArray<uint8> OutData;
 			bool bSuccess = InnerBackend->GetCachedData(CacheKey, OutData);
 			if (bSuccess)
@@ -100,7 +110,22 @@ public:
 	{
 		InnerBackend->RemoveCachedData(CacheKey, bTransient);
 	}
+
+	virtual void GatherUsageStats(TMap<FString, FDerivedDataCacheUsageStats>& UsageStatsMap, FString&& GraphPath) override
+	{
+		COOK_STAT(
+		{
+			UsageStatsMap.Add(GraphPath + TEXT(": VerifyWrapper"), UsageStats);
+			if (InnerBackend)
+			{
+				InnerBackend->GatherUsageStats(UsageStatsMap, GraphPath + TEXT(". 0"));
+			}
+		});
+	}
+
 private:
+	FDerivedDataCacheUsageStats UsageStats;
+
 	/** If problems are encountered, do we fix them?						*/
 	bool											bFixProblems;
 	/** Object used for synchronization via a scoped lock						*/
diff --git a/Engine/Source/Developer/DerivedDataCache/Private/DerivedDataBackends.cpp b/Engine/Source/Developer/DerivedDataCache/Private/DerivedDataBackends.cpp
--- a/Engine/Source/Developer/DerivedDataCache/Private/DerivedDataBackends.cpp
+++ b/Engine/Source/Developer/DerivedDataCache/Private/DerivedDataBackends.cpp
@@ -13,6 +13,7 @@
 #include "DerivedDataBackendVerifyWrapper.h"
 #include "DerivedDataUtilsInterface.h"
 #include "Misc/EngineBuildSettings.h"
+#include "DerivedDataCacheUsageStats.h"
 
 DEFINE_LOG_CATEGORY(LogDerivedDataCache);
 
@@ -786,6 +787,14 @@ public:
 		return false;
 	}
 
+	virtual void GatherUsageStats(TMap<FString, FDerivedDataCacheUsageStats>& UsageStats) override
+	{
+		if (RootCache)
+		{
+			RootCache->GatherUsageStats(UsageStats, TEXT(" 0"));
+		}
+	}
+
 private:
 
 	/** Delete the old files in a directory **/
diff --git a/Engine/Source/Developer/DerivedDataCache/Private/DerivedDataCache.cpp b/Engine/Source/Developer/DerivedDataCache/Private/DerivedDataCache.cpp
--- a/Engine/Source/Developer/DerivedDataCache/Private/DerivedDataCache.cpp
+++ b/Engine/Source/Developer/DerivedDataCache/Private/DerivedDataCache.cpp
@@ -8,7 +8,7 @@
 #include "DerivedDataBackendInterface.h"
 #include "DerivedDataPluginInterface.h"
 #include "DDCCleanup.h"
-
+#include "CookStats.h"
 
 DEFINE_STAT(STAT_DDC_NumGets);
 DEFINE_STAT(STAT_DDC_NumPuts);
@@ -20,7 +20,117 @@ DEFINE_STAT(STAT_DDC_PutTime);
 DEFINE_STAT(STAT_DDC_SyncBuildTime);
 DEFINE_STAT(STAT_DDC_ExistTime);
 
-/** 
+#if ENABLE_COOK_STATS
+#include "DerivedDataCacheUsageStats.h"
+namespace DerivedDataCacheCookStats
+{
+	static void CollectStats(FCookStatsManager::AddStatFuncRef AddStat)
+	{
+		TMap<FString, FDerivedDataCacheUsageStats> DDCStats;
+		GetDerivedDataCacheRef().GatherUsageStats(DDCStats);
+		{
+			const FString StatName(TEXT("DDC.Usage"));
+
+			auto LogDDCNodeStats = [&](const TPair<FString, FDerivedDataCacheUsageStats>& StatPair)
+			{
+				auto LogDDCNodeCallStats = [&](const FString& NodeName, const FDerivedDataCacheUsageStats::CallStats& CallStat, const TCHAR* CallName)
+				{
+					auto LogDDCNodeCallStat = [&](const FString& InnerNodeName, const FDerivedDataCacheUsageStats::CallStats& InnerCallStat, const TCHAR* InnerCallName, FDerivedDataCacheUsageStats::CallStats::EHitOrMiss HitOrMiss, bool bGameThread)
+					{
+						// only report the stat if it's non-zero
+						if (CallStat.GetAccumulatedValue(HitOrMiss, FDerivedDataCacheUsageStats::CallStats::EStatType::Cycles, bGameThread) > 0)
+						{
+							AddStat(StatName, FCookStatsManager::CreateKeyValueArray(
+								TEXT("ThreadName"), bGameThread ? TEXT("GameThread") : TEXT("OthrThread"),
+								TEXT("Call"), InnerCallName,
+								TEXT("HitOrMiss"), HitOrMiss == FDerivedDataCacheUsageStats::CallStats::EHitOrMiss::Hit ? TEXT("Hit") : TEXT("Miss"),
+								TEXT("Count"), InnerCallStat.GetAccumulatedValue(HitOrMiss, FDerivedDataCacheUsageStats::CallStats::EStatType::Counter, bGameThread),
+								TEXT("TimeSec"), InnerCallStat.GetAccumulatedValue(HitOrMiss, FDerivedDataCacheUsageStats::CallStats::EStatType::Cycles, bGameThread) * FPlatformTime::GetSecondsPerCycle(),
+								TEXT("MB"), InnerCallStat.GetAccumulatedValue(HitOrMiss, FDerivedDataCacheUsageStats::CallStats::EStatType::Bytes, bGameThread) / (1024.0 * 1024.0),
+								TEXT("MB/s"), InnerCallStat.GetAccumulatedValue(HitOrMiss, FDerivedDataCacheUsageStats::CallStats::EStatType::Cycles, bGameThread) == 0
+									? 0.0
+									: InnerCallStat.GetAccumulatedValue(HitOrMiss, FDerivedDataCacheUsageStats::CallStats::EStatType::Bytes, bGameThread) / (1024.0 * 1024.0) /
+									(InnerCallStat.GetAccumulatedValue(HitOrMiss, FDerivedDataCacheUsageStats::CallStats::EStatType::Cycles, bGameThread) * FPlatformTime::GetSecondsPerCycle()),
+									TEXT("Node"), InnerNodeName
+								));
+						}
+					};
+
+					LogDDCNodeCallStat(NodeName, CallStat, CallName, FDerivedDataCacheUsageStats::CallStats::EHitOrMiss::Hit, true);
+					LogDDCNodeCallStat(NodeName, CallStat, CallName, FDerivedDataCacheUsageStats::CallStats::EHitOrMiss::Miss, true);
+					LogDDCNodeCallStat(NodeName, CallStat, CallName, FDerivedDataCacheUsageStats::CallStats::EHitOrMiss::Hit, false);
+					LogDDCNodeCallStat(NodeName, CallStat, CallName, FDerivedDataCacheUsageStats::CallStats::EHitOrMiss::Miss, false);
+				};
+
+				LogDDCNodeCallStats(StatPair.Key, StatPair.Value.GetStats, TEXT("Get"));
+				LogDDCNodeCallStats(StatPair.Key, StatPair.Value.PutStats, TEXT("Put"));
+				LogDDCNodeCallStats(StatPair.Key, StatPair.Value.ExistsStats, TEXT("Exists"));
+			};
+
+			for (const auto& UsageStatPair : DDCStats)
+			{
+				LogDDCNodeStats(UsageStatPair);
+			}
+		}
+
+		// Now lets add some summary data to that applies some crazy knowledge of how we set up our DDC. The goal 
+		// is to print out the global hit rate, and the hit rate of the local and shared DDC.
+		// This is done by adding up the total get/miss calls the root node receives.
+		// Then we find the FileSystem nodes that correspond to the local and shared cache using some hacky logic to detect a "network drive".
+		// If the DDC graph ever contains more than one local or remote filesystem, this will only find one of them.
+		{
+			TArray<FString, TInlineAllocator<20>> Keys;
+			DDCStats.GenerateKeyArray(Keys);
+			FString* RootKey = Keys.FindByPredicate([](const FString& Key) {return Key.StartsWith(TEXT(" 0:")); });
+			// look for a Filesystem DDC that doesn't have a UNC path. Ugly, yeah, but we only cook on PC at the moment.
+			FString* LocalDDCKey = Keys.FindByPredicate([](const FString& Key) {return Key.Contains(TEXT(": FileSystem.")) && !Key.Contains(TEXT("//")); });
+			// look for a UNC path
+			FString* SharedDDCKey = Keys.FindByPredicate([](const FString& Key) {return Key.Contains(TEXT(": FileSystem.//")); });
+			if (RootKey)
+			{
+				const FDerivedDataCacheUsageStats& RootStats = DDCStats[*RootKey];
+				int64 TotalGetHits =
+					RootStats.GetStats.GetAccumulatedValue(FDerivedDataCacheUsageStats::CallStats::EHitOrMiss::Hit, FDerivedDataCacheUsageStats::CallStats::EStatType::Counter, true) +
+					RootStats.GetStats.GetAccumulatedValue(FDerivedDataCacheUsageStats::CallStats::EHitOrMiss::Hit, FDerivedDataCacheUsageStats::CallStats::EStatType::Counter, false);
+				int64 TotalGetMisses =
+					RootStats.GetStats.GetAccumulatedValue(FDerivedDataCacheUsageStats::CallStats::EHitOrMiss::Miss, FDerivedDataCacheUsageStats::CallStats::EStatType::Counter, true) +
+					RootStats.GetStats.GetAccumulatedValue(FDerivedDataCacheUsageStats::CallStats::EHitOrMiss::Miss, FDerivedDataCacheUsageStats::CallStats::EStatType::Counter, false);
+				int64 TotalGets = TotalGetHits + TotalGetMisses;
+
+				int64 LocalHits = 0;
+				if (LocalDDCKey)
+				{
+					const FDerivedDataCacheUsageStats& LocalDDCStats = DDCStats[*LocalDDCKey];
+					LocalHits =
+						LocalDDCStats.GetStats.GetAccumulatedValue(FDerivedDataCacheUsageStats::CallStats::EHitOrMiss::Hit, FDerivedDataCacheUsageStats::CallStats::EStatType::Counter, true) +
+						LocalDDCStats.GetStats.GetAccumulatedValue(FDerivedDataCacheUsageStats::CallStats::EHitOrMiss::Hit, FDerivedDataCacheUsageStats::CallStats::EStatType::Counter, false);
+				}
+				int64 SharedHits = 0;
+				if (SharedDDCKey)
+				{
+					// The shared DDC is only queried if the local one misses (or there isn't one). So it's hit rate is technically 
+					const FDerivedDataCacheUsageStats& SharedDDCStats = DDCStats[*SharedDDCKey];
+					SharedHits =
+						SharedDDCStats.GetStats.GetAccumulatedValue(FDerivedDataCacheUsageStats::CallStats::EHitOrMiss::Hit, FDerivedDataCacheUsageStats::CallStats::EStatType::Counter, true) +
+						SharedDDCStats.GetStats.GetAccumulatedValue(FDerivedDataCacheUsageStats::CallStats::EHitOrMiss::Hit, FDerivedDataCacheUsageStats::CallStats::EStatType::Counter, false);
+				}
+				AddStat(TEXT("DDC.Summary"), FCookStatsManager::CreateKeyValueArray(
+					TEXT("TotalGetHits"), TotalGetHits,
+					TEXT("TotalGets"), TotalGets,
+					TEXT("TotalHitPct"), (double)TotalGetHits / TotalGets,
+					TEXT("LocalHitPct"), (double)LocalHits / TotalGets,
+					TEXT("SharedHitPct"), (double)SharedHits / TotalGets,
+					TEXT("OtherHitPct"), double(TotalGetHits - LocalHits - SharedHits) / TotalGets,
+					TEXT("MissPct"), (double)TotalGetMisses / TotalGets
+					));
+			}
+		}
+	}
+	FCookStatsManager::FAutoRegisterCallback RegisterCookStats(&CollectStats);
+}
+#endif
+
+/**
  * Implementation of the derived data cache
  * This API is fully threadsafe
 **/
@@ -337,6 +447,10 @@ public:
 		// Used by derived classes to spit out leaked pending rollups
 	}
 
+	virtual void GatherUsageStats(TMap<FString, FDerivedDataCacheUsageStats>& UsageStatsMap) override
+	{
+		FDerivedDataBackend::Get().GatherUsageStats(UsageStatsMap);
+	}
 
 protected:
 	uint32 NextHandle()
diff --git a/Engine/Source/Developer/DerivedDataCache/Private/DerivedDataLimitKeyLengthWrapper.h b/Engine/Source/Developer/DerivedDataCache/Private/DerivedDataLimitKeyLengthWrapper.h
--- a/Engine/Source/Developer/DerivedDataCache/Private/DerivedDataLimitKeyLengthWrapper.h
+++ b/Engine/Source/Developer/DerivedDataCache/Private/DerivedDataLimitKeyLengthWrapper.h
@@ -1,6 +1,7 @@
 // Copyright 1998-2016 Epic Games, Inc. All Rights Reserved.
 
 #pragma once
+#include "DerivedDataCacheUsageStats.h"
 
 /** 
  * A backend wrapper that limits the key size and uses hashing...in this case it wraps the payload and the payload contains the full key to verify the integrity of the hash
@@ -35,9 +36,15 @@ public:
 	 */
 	virtual bool CachedDataProbablyExists(const TCHAR* CacheKey) override
 	{
+		COOK_STAT(auto Timer = UsageStats.TimeProbablyExists());
 		FString NewKey;
 		ShortenKey(CacheKey, NewKey);
-		return InnerBackend->CachedDataProbablyExists(*NewKey);
+		bool Result = InnerBackend->CachedDataProbablyExists(*NewKey);
+		if (Result)
+		{
+			COOK_STAT(Timer.AddHit(0));
+		}
+		return Result;
 	}
 	/**
 	 * Synchronous retrieve of a cache item
@@ -48,6 +55,9 @@ public:
 	 */
 	virtual bool GetCachedData(const TCHAR* CacheKey, TArray<uint8>& OutData) override
 	{
+		COOK_STAT(auto Timer = UsageStats.TimeGet());
+		int64 InnerGetCycles = 0;
+
 		FString NewKey;
 		bool bOk;
 		if (!ShortenKey(CacheKey, NewKey))
@@ -105,6 +115,10 @@ public:
 		{
 			OutData.Empty();
 		}
+		else
+		{
+			COOK_STAT(Timer.AddHit(OutData.Num()));
+		}
 		return bOk;
 	}
 	/**
@@ -116,14 +130,15 @@ public:
 	 */
 	virtual void PutCachedData(const TCHAR* CacheKey, TArray<uint8>& InData, bool bPutEvenIfExists) override
 	{
+		COOK_STAT(auto Timer = UsageStats.TimePut());
 		if (!InnerBackend->IsWritable())
 		{
 			return; // no point in continuing down the chain
 		}
+		COOK_STAT(Timer.AddHit(InData.Num()));
 		FString NewKey;
 		if (!ShortenKey(CacheKey, NewKey))
 		{
-			// no shortening needed
 			InnerBackend->PutCachedData(CacheKey, InData, bPutEvenIfExists);
 			return;
 		}
@@ -146,7 +161,21 @@ public:
 		ShortenKey(CacheKey, NewKey);
 		return InnerBackend->RemoveCachedData(*NewKey, bTransient);
 	}
+
+	virtual void GatherUsageStats(TMap<FString, FDerivedDataCacheUsageStats>& UsageStatsMap, FString&& GraphPath) override
+	{
+		COOK_STAT(
+		{
+			UsageStatsMap.Add(GraphPath + TEXT(": LimitKeyLength"), UsageStats);
+			if (InnerBackend)
+			{
+				InnerBackend->GatherUsageStats(UsageStatsMap, GraphPath + TEXT(". 0"));
+			}
+		});
+	}
+
 private:
+	FDerivedDataCacheUsageStats UsageStats;
 
 	/** Shorten the cache key and return true if shortening was required **/
 	bool ShortenKey(const TCHAR* CacheKey, FString& Result)
diff --git a/Engine/Source/Developer/DerivedDataCache/Private/FileSystemDerivedDataBackend.cpp b/Engine/Source/Developer/DerivedDataCache/Private/FileSystemDerivedDataBackend.cpp
--- a/Engine/Source/Developer/DerivedDataCache/Private/FileSystemDerivedDataBackend.cpp
+++ b/Engine/Source/Developer/DerivedDataCache/Private/FileSystemDerivedDataBackend.cpp
@@ -6,6 +6,8 @@
 #include "DerivedDataBackendInterface.h"
 #include "DDCCleanup.h"
 
+#include "DerivedDataCacheUsageStats.h"
+
 #define MAX_BACKEND_KEY_LENGTH (120)
 #define MAX_BACKEND_NUMBERED_SUBFOLDER_LENGTH (9)
 #if PLATFORM_LINUX	// PATH_MAX on Linux is 4096 (getconf PATH_MAX /, also see limits.h), so this value can be larger (note that it is still arbitrary).
@@ -133,12 +135,7 @@ public:
 	 */
 	virtual bool CachedDataProbablyExists(const TCHAR* CacheKey) override
 	{
-#if ENABLE_DDC_STATS
-		static FName NAME_CachedDataProbablyExists(TEXT("CachedDataProbablyExists"));
-		FDDCScopeStatHelper Stat(CacheKey, NAME_CachedDataProbablyExists);
-		static FName NAME_FileDDCPath(TEXT("FileDDCPath"));
-		Stat.AddTag(NAME_FileDDCPath, CachePath);
-#endif
+		COOK_STAT(auto Timer = UsageStats.TimeProbablyExists());
 		check(!bFailed);
 		FString Filename = BuildFilename(CacheKey);
 		FDateTime TimeStamp = IFileManager::Get().GetTimeStamp(*Filename);
@@ -152,6 +149,10 @@ public:
 				IFileManager::Get().SetTimeStamp(*Filename, FDateTime::UtcNow());
 			}
 		}
+		if (bExists)
+		{
+			COOK_STAT(Timer.AddHit(0));
+		}
 		return bExists;
 	}
 	/**
@@ -163,14 +164,7 @@ public:
 	 */
 	virtual bool GetCachedData(const TCHAR* CacheKey, TArray<uint8>& Data) override
 	{
-#if ENABLE_DDC_STATS
-		static FName NAME_GetCachedData(TEXT("GetCachedData"));
-		FDDCScopeStatHelper Stat(CacheKey, NAME_GetCachedData);
-		static FName NAME_FileDDCPath(TEXT("FileDDCPath"));
-		static FName NAME_Retrieved(TEXT("Retrieved"));
-		Stat.AddTag(NAME_FileDDCPath, CachePath);
-#endif
-
+		COOK_STAT(auto Timer = UsageStats.TimeGet());
 		check(!bFailed);
 		FString Filename = BuildFilename(CacheKey);
 		double StartTime = FPlatformTime::Seconds();
@@ -185,16 +179,11 @@ public:
 			}
 
 			UE_LOG(LogDerivedDataCache, Verbose, TEXT("FileSystemDerivedDataBackend: Cache hit on %s"),*Filename);
-#if ENABLE_DDC_STATS
-			Stat.AddTag(NAME_Retrieved, true);
-#endif
+			COOK_STAT(Timer.AddHit(Data.Num()));
 			return true;
 		}
 		UE_LOG(LogDerivedDataCache, Verbose, TEXT("FileSystemDerivedDataBackend: Cache miss on %s"),*Filename);
 		Data.Empty();
-#if ENABLE_DDC_STATS
-		Stat.AddTag(NAME_Retrieved, false);
-#endif
 		return false;
 	}
 	/**
@@ -206,17 +195,13 @@ public:
 	 */
 	virtual void PutCachedData(const TCHAR* CacheKey, TArray<uint8>& Data, bool bPutEvenIfExists) override
 	{
-#if ENABLE_DDC_STATS
-		static FName NAME_PutCachedData(TEXT("PutCachedData"));
-		FDDCScopeStatHelper Stat(CacheKey, NAME_PutCachedData);
-		static FName NAME_FileDDCPath(TEXT("FileDDCPath"));
-		Stat.AddTag(NAME_FileDDCPath, CachePath);
-#endif
+		COOK_STAT(auto Timer = UsageStats.TimePut());
 		check(!bFailed);
 		if (!bReadOnly)
 		{
 			if (bPutEvenIfExists || !CachedDataProbablyExists(CacheKey))
 			{
+				COOK_STAT(Timer.AddHit(Data.Num()));
 				check(Data.Num());
 				FString Filename = BuildFilename(CacheKey);
 				FString TempFilename(TEXT("temp.")); 
@@ -280,7 +265,13 @@ public:
 		}
 	}
 
+	virtual void GatherUsageStats(TMap<FString, FDerivedDataCacheUsageStats>& UsageStatsMap, FString&& GraphPath) override
+	{
+		COOK_STAT(UsageStatsMap.Add(FString::Printf(TEXT("%s: %s.%s"), *GraphPath, TEXT("FileSystem"), *CachePath), UsageStats));
+	}
+
 private:
+	FDerivedDataCacheUsageStats UsageStats;
 
 	/**
 	 * Threadsafe method to compute the filename from the cachekey, currently just adds a path and an extension.
diff --git a/Engine/Source/Developer/DerivedDataCache/Private/HierarchicalDerivedDataBackend.h b/Engine/Source/Developer/DerivedDataCache/Private/HierarchicalDerivedDataBackend.h
--- a/Engine/Source/Developer/DerivedDataCache/Private/HierarchicalDerivedDataBackend.h
+++ b/Engine/Source/Developer/DerivedDataCache/Private/HierarchicalDerivedDataBackend.h
@@ -2,6 +2,9 @@
 
 #pragma once
 
+#include "DerivedDataCacheUsageStats.h"
+
+
 /** 
  * A backend wrapper that implements a cache hierarchy of backends. 
 **/
@@ -62,10 +65,12 @@ public:
 	 */
 	virtual bool CachedDataProbablyExists(const TCHAR* CacheKey) override
 	{
+		COOK_STAT(auto Timer = UsageStats.TimeProbablyExists());
 		for (int32 CacheIndex = 0; CacheIndex < InnerBackends.Num(); CacheIndex++)
 		{
 			if (InnerBackends[CacheIndex]->CachedDataProbablyExists(CacheKey))
 			{
+				COOK_STAT(Timer.AddHit(0));
 				return true;
 			}
 		}
@@ -80,6 +85,7 @@ public:
 	 */
 	virtual bool GetCachedData(const TCHAR* CacheKey, TArray<uint8>& OutData) override
 	{
+		COOK_STAT(auto Timer = UsageStats.TimeGet());
 		for (int32 CacheIndex = 0; CacheIndex < InnerBackends.Num(); CacheIndex++)
 		{
 			if (InnerBackends[CacheIndex]->CachedDataProbablyExists(CacheKey) && InnerBackends[CacheIndex]->GetCachedData(CacheKey, OutData))
@@ -119,6 +125,7 @@ public:
 						}
 					}
 				}
+				COOK_STAT(Timer.AddHit(OutData.Num()));
 				return true;
 			}
 		}
@@ -133,6 +140,7 @@ public:
 	 */
 	virtual void PutCachedData(const TCHAR* CacheKey, TArray<uint8>& InData, bool bPutEvenIfExists) override
 	{
+		COOK_STAT(auto Timer = UsageStats.TimePut());
 		if (!bIsWritable)
 		{
 			return; // no point in continuing down the chain
@@ -146,6 +154,7 @@ public:
 			}
 			if (InnerBackends[PutCacheIndex]->IsWritable())
 			{
+				COOK_STAT(Timer.AddHit(InData.Num()));
 				if (!bSynchronousPutPeformed)
 				{
 					InnerBackends[PutCacheIndex]->PutCachedData(CacheKey, InData, bPutEvenIfExists);
@@ -170,7 +179,34 @@ public:
 			InnerBackends[PutCacheIndex]->RemoveCachedData(CacheKey, bTransient);
 		}
 	}
+
+	virtual void GatherUsageStats(TMap<FString, FDerivedDataCacheUsageStats>& UsageStatsMap, FString&& GraphPath) override
+	{
+		COOK_STAT(
+		{
+			UsageStatsMap.Add(GraphPath + TEXT(": Hierarchical"), UsageStats);
+			// All the inner backends are actually wrapped by AsyncPut backends in writable cases (most cases in practice)
+			if (AsyncPutInnerBackends.Num() > 0)
+			{
+				int Ndx = 0;
+				for (const auto& InnerBackend : AsyncPutInnerBackends)
+				{
+					InnerBackend->GatherUsageStats(UsageStatsMap, GraphPath + FString::Printf(TEXT(".%2d"), Ndx++));
+				}
+			}
+			else
+			{
+				int Ndx = 0;
+				for (auto InnerBackend : InnerBackends)
+				{
+					InnerBackend->GatherUsageStats(UsageStatsMap, GraphPath + FString::Printf(TEXT(".%2d"), Ndx++));
+				}
+			}
+		});
+	}
+
 private:
+	FDerivedDataCacheUsageStats UsageStats;
 
 	/** Array of backends forming the hierarchical cache...the first element is the fastest cache. **/
 	TArray<FDerivedDataBackendInterface*> InnerBackends;
diff --git a/Engine/Source/Developer/DerivedDataCache/Private/MemoryDerivedDataBackend.h b/Engine/Source/Developer/DerivedDataCache/Private/MemoryDerivedDataBackend.h
--- a/Engine/Source/Developer/DerivedDataCache/Private/MemoryDerivedDataBackend.h
+++ b/Engine/Source/Developer/DerivedDataCache/Private/MemoryDerivedDataBackend.h
@@ -1,6 +1,7 @@
 // Copyright 1998-2016 Epic Games, Inc. All Rights Reserved.
 
 #pragma once
+#include "DerivedDataCacheUsageStats.h"
 
 /** 
  * A simple thread safe, memory based backend. This is used for Async puts and the boot cache.
@@ -35,6 +36,7 @@ public:
 	 */
 	virtual bool CachedDataProbablyExists(const TCHAR* CacheKey) override
 	{
+		COOK_STAT(auto Timer = UsageStats.TimeProbablyExists());
 		FScopeLock ScopeLock(&SynchronizationObject);
 		if (bDisabled)
 		{
@@ -47,7 +49,12 @@ public:
 			return true;
 		}
 
-		return CacheItems.Contains(FString(CacheKey));
+		bool Result = CacheItems.Contains(FString(CacheKey));
+		if (Result)
+		{
+			COOK_STAT(Timer.AddHit(0));
+		}
+		return Result;
 	}
 	/**
 	 * Synchronous retrieve of a cache item
@@ -58,6 +65,7 @@ public:
 	 */
 	virtual bool GetCachedData(const TCHAR* CacheKey, TArray<uint8>& OutData) override
 	{
+		COOK_STAT(auto Timer = UsageStats.TimeGet());
 		FScopeLock ScopeLock(&SynchronizationObject);
 		if (!bDisabled)
 		{
@@ -67,6 +75,7 @@ public:
 				OutData = Item->Data;
 				Item->Age = 0;
 				check(OutData.Num());
+				COOK_STAT(Timer.AddHit(OutData.Num()));
 				return true;
 			}
 		}
@@ -82,6 +91,7 @@ public:
 	 */
 	virtual void PutCachedData(const TCHAR* CacheKey, TArray<uint8>& InData, bool bPutEvenIfExists) override
 	{
+		COOK_STAT(auto Timer = UsageStats.TimePut());
 		FScopeLock ScopeLock(&SynchronizationObject);
 		
 		if (bDisabled || bMaxSizeExceeded)
@@ -110,6 +120,7 @@ public:
 			}
 			else
 			{
+				COOK_STAT(Timer.AddHit(InData.Num()));
 				CacheItems.Add(Key, Val);
 				CalcCacheValueSize(Key, *Val);
 
@@ -290,6 +301,7 @@ public:
 		}
 		
 		CurrentCacheSize = FileSize;
+		CacheFilename = Filename;
 		UE_LOG(LogDerivedDataCache, Log, TEXT("Loaded boot cache %4.2fs %lldMB %s."), float(FPlatformTime::Seconds() - StartTime), DataSize / (1024 * 1024), Filename);
 		return true;
 	}
@@ -309,7 +321,15 @@ public:
 		CurrentCacheSize = SerializationSpecificDataSize;
 	}
 
+	virtual void GatherUsageStats(TMap<FString, FDerivedDataCacheUsageStats>& UsageStatsMap, FString&& GraphPath) override
+	{
+		COOK_STAT(UsageStatsMap.Add(FString::Printf(TEXT("%s: %s.%s"), *GraphPath, TEXT("MemoryBackend"), *CacheFilename), UsageStats));
+	}
+
 private:
+	/** Name of the cache file loaded (if any). */
+	FString CacheFilename;
+	FDerivedDataCacheUsageStats UsageStats;
 
 	struct FCacheValue
 	{
diff --git a/Engine/Source/Developer/DerivedDataCache/Private/PakFileDerivedDataBackend.h b/Engine/Source/Developer/DerivedDataCache/Private/PakFileDerivedDataBackend.h
--- a/Engine/Source/Developer/DerivedDataCache/Private/PakFileDerivedDataBackend.h
+++ b/Engine/Source/Developer/DerivedDataCache/Private/PakFileDerivedDataBackend.h
@@ -3,6 +3,7 @@
 #pragma once
 
 #include "Compression.h"
+#include "DerivedDataCacheUsageStats.h"
 
 /** 
  * A simple thread safe, pak file based backend. 
@@ -88,8 +89,14 @@ public:
 	 */
 	virtual bool CachedDataProbablyExists(const TCHAR* CacheKey) override
 	{
+		COOK_STAT(auto Timer = UsageStats.TimeProbablyExists());
 		FScopeLock ScopeLock(&SynchronizationObject);
-		return CacheItems.Contains(FString(CacheKey));
+		bool Result = CacheItems.Contains(FString(CacheKey));
+		if (Result)
+		{
+			COOK_STAT(Timer.AddHit(0));
+		}
+		return Result;
 	}
 	/**
 	 * Synchronous retrieve of a cache item
@@ -100,6 +107,7 @@ public:
 	 */
 	virtual bool GetCachedData(const TCHAR* CacheKey, TArray<uint8>& OutData) override
 	{
+		COOK_STAT(auto Timer = UsageStats.TimeGet());
 		if (bWriting || bClosed)
 		{
 			return false;
@@ -130,6 +138,7 @@ public:
 				{
 					UE_LOG(LogDerivedDataCache, Verbose, TEXT("FPakFileDerivedDataBackend: Cache hit on %s"), CacheKey);
 					check(OutData.Num());
+					COOK_STAT(Timer.AddHit(OutData.Num()));
 					return true;
 				}
 			}
@@ -150,6 +159,7 @@ public:
 	 */
 	virtual void PutCachedData(const TCHAR* CacheKey, TArray<uint8>& InData, bool bPutEvenIfExists) override
 	{
+		COOK_STAT(auto Timer = UsageStats.TimePut());
 		if (!bWriting || bClosed)
 		{
 			return;
@@ -174,6 +184,7 @@ public:
 				}
 				else
 				{
+					COOK_STAT(Timer.AddHit(InData.Num()));
 					FileHandle->Serialize(InData.GetData(), int64(InData.Num()));
 					UE_LOG(LogDerivedDataCache, Verbose, TEXT("FPakFileDerivedDataBackend: Put %s"), CacheKey);
 					CacheItems.Add(Key,FCacheValue(Offset, InData.Num(), Crc));
@@ -424,7 +435,13 @@ public:
 		return true;
 	}
 
+	virtual void GatherUsageStats(TMap<FString, FDerivedDataCacheUsageStats>& UsageStatsMap, FString&& GraphPath) override
+	{
+		COOK_STAT(UsageStatsMap.Add(FString::Printf(TEXT("%s: %s.%s"), *GraphPath, TEXT("PakFile"), *Filename), UsageStats));
+	}
+
 private:
+	FDerivedDataCacheUsageStats UsageStats;
 
 	struct FCacheValue
 	{
diff --git a/Engine/Source/Developer/DerivedDataCache/Public/DerivedDataBackendInterface.h b/Engine/Source/Developer/DerivedDataCache/Public/DerivedDataBackendInterface.h
--- a/Engine/Source/Developer/DerivedDataCache/Public/DerivedDataBackendInterface.h
+++ b/Engine/Source/Developer/DerivedDataCache/Public/DerivedDataBackendInterface.h
@@ -2,6 +2,9 @@
 
 
 #pragma once
+
+class FDerivedDataCacheUsageStats;
+
 DECLARE_LOG_CATEGORY_EXTERN(LogDerivedDataCache, Log, All);
 
 DECLARE_DWORD_ACCUMULATOR_STAT_EXTERN(TEXT("Num Gets"),STAT_DDC_NumGets,STATGROUP_DDC, );
@@ -68,6 +71,14 @@ public:
 	 * @param	bTransient	true if the data is transient and it is up to the backend to decide when and if to remove cached data.
 	 */
 	virtual void RemoveCachedData(const TCHAR* CacheKey, bool bTransient)=0;
+
+	/**
+	 * Retrieve usage stats for this backend. If the backend holds inner backends, this is expected to be passed down recursively.
+	 * @param	UsageStatsMap		The map of usages. Each backend instance should give itself a unique name if possible (ie, use the filename associated).
+	 * @param	GraphPath			Path to the node in the graph. If you have inner nodes, add their index to the current path as ". <n>".
+	 *								This will create a path such as "0. 1. 0. 2", which can uniquely identify this node.
+	 */
+	virtual void GatherUsageStats(TMap<FString, FDerivedDataCacheUsageStats>& UsageStatsMap, FString&& GraphPath) = 0;
 };
 
 class FDerivedDataBackend
@@ -109,6 +120,6 @@ public:
 	 */
 	virtual bool UnmountPakFile(const TCHAR* PakFilename) = 0;
 
-};
-
+	virtual void GatherUsageStats(TMap<FString, FDerivedDataCacheUsageStats>& UsageStats) = 0;
 
+};
diff --git a/Engine/Source/Developer/DerivedDataCache/Public/DerivedDataCacheInterface.h b/Engine/Source/Developer/DerivedDataCache/Public/DerivedDataCacheInterface.h
--- a/Engine/Source/Developer/DerivedDataCache/Public/DerivedDataCacheInterface.h
+++ b/Engine/Source/Developer/DerivedDataCache/Public/DerivedDataCacheInterface.h
@@ -4,6 +4,8 @@
 #pragma once
 #include "ModuleManager.h"
 
+class FDerivedDataCacheUsageStats;
+
 /** 
  * Opaque class for rollup handling
 **/
@@ -200,6 +202,14 @@ public:
 	 */
 	virtual void GetDirectories(TArray<FString>& OutResults) = 0;
 
+	//--------------------
+	// UsageStats Interface
+	//--------------------
+
+	/**
+	 * Retrieve usage stats by the DDC
+	 */
+	virtual void GatherUsageStats(TMap<FString, FDerivedDataCacheUsageStats>& UsageStats) = 0;
 };
 
 /**
diff --git a/Engine/Source/Developer/DerivedDataCache/Public/DerivedDataCacheUsageStats.h b/Engine/Source/Developer/DerivedDataCache/Public/DerivedDataCacheUsageStats.h
new file mode 100644
--- /dev/null
+++ b/Engine/Source/Developer/DerivedDataCache/Public/DerivedDataCacheUsageStats.h
@@ -0,0 +1,209 @@
+// Copyright 1998-2016 Epic Games, Inc. All Rights Reserved.
+#pragma once
+#include "CookStats.h"
+
+#if ENABLE_COOK_STATS
+#include "ThreadingBase.h"
+#endif
+
+/**
+ * Usage stats for the derived data cache nodes. At the end of the app or commandlet, the DDC
+ * can be asked to gather usage stats for each of the nodes in the DDC graph, which are accumulated
+ * into a TMap of Name->Stats. The Stats portion is this class.
+ * 
+ * The class exposes various high-level routines to time important aspects of the DDC, mostly
+ * focusing on performance of GetCachedData, PutCachedData, and CachedDataProbablyExists. This class
+ * will track time taken, calls made, hits, misses, bytes processed, and do it for two buckets:
+ * 1) the main thread and 2) all other threads.  The reason is because any time spent in the DDC on the
+ * main thread is considered meaningful, as DDC access is generally expected to be async from helper threads.
+ * 
+ * The class goes through a bit of trouble to use thread-safe access calls for task-thread usage, and 
+ * simple, fast accumulators for game-thread usage, since it's guaranteed to not be written to concurrently.
+ * The class also limits itself to checking the thread once at construction.
+ * 
+ * Usage would be something like this in a concrete FDerivedDataBackendInterface implementation:
+ *   class MyBackend : public FDerivedDataBackendInterface
+ *   {
+ *       FDerivedDataCacheUsageStats UsageStats;
+ *   public:
+ *       <override CachedDataProbablyExists>
+ *       { 
+ *           auto Timer = UsageStats.TimeExists();
+ *           ...
+ *       }
+ *       <override GetCachedData>
+ *       {    
+ *           auto Timer = UsageStats.TimeGet();
+ *           ...
+ *           <if it's a cache hit> Timer.AddHit(DataSize);
+ *           // Misses are automatically tracked
+ *       }
+ *       <override PutCachedData>
+ *       {
+ *           auto Timer = UsageStats.TimePut();
+ *           ...
+ *           <if the data will really be Put> Timer.AddHit(DataSize);
+ *           // Misses are automatically tracked
+ *       }
+ *       <override GatherUsageStats>
+ *       {
+ *           // Add this node's UsageStats to the usage map. Your Key name should be UNIQUE to the entire graph (so use the file name, or pointer to this if you have to).
+ *           UsageStatsMap.Add(FString::Printf(TEXT("%s: %s"), *GraphPath, <Some unique name for this node instance>), UsageStats);
+ *       }
+*   }
+ */
+class FDerivedDataCacheUsageStats
+{
+#if ENABLE_COOK_STATS
+public:
+	/** 
+	 * Struct to hold stats for a Get/Put/Exists call. 
+	 * The sub-structs and stuff look intimidating, but it is to unify the concept
+	 * of accumulating any stat on the game or other thread, using raw int64 on the 
+	 * main thread, and thread safe accumulators on other threads.
+	 * 
+	 * Each call type that will be tracked will track call counts, call cycles, and bytes processed,
+	 * grouped by hits and misses. Some stats will by definition be zero (ie, no Miss will ever track non-zero bytes processed)
+	 * but it limits the branching in the tracking code.
+	 */
+	struct CallStats
+	{
+		/** One group of accumulators for hits and misses. */
+		enum class EHitOrMiss : uint8
+		{
+			Hit,
+			Miss,
+			MaxValue,
+		};
+
+		/** Each hit or miss will contain these accumulator stats. */
+		enum class EStatType : uint8
+		{
+			Counter,
+			Cycles,
+			Bytes,
+			MaxValue,
+		};
+
+		/** Contains a pair of accumulators, one for the game thread, one for the other threads. */
+		struct GameAndOtherThreadAccumulator
+		{
+			/** Accumulates a stat. Uses thread safe primitives for non-game thread accumulation. */
+			void Accumulate(int64 Value, bool bIsInGameThread)
+			{
+				if (bIsInGameThread)
+				{
+					GameThread += Value;
+				}
+				else if (Value != 0)
+				{
+					OtherThread.Add(Value);
+				}
+			}
+			/** Access the accumulated values (exposed for more uniform access methds to each accumulator). */
+			int64 GetAccumulatedValue(bool bIsInGameThread) const
+			{
+				return bIsInGameThread ? GameThread : OtherThread.GetValue();
+			}
+
+			int64 GameThread = 0l;
+			FThreadSafeCounter64 OtherThread;
+		};
+
+		/** Make it easier to update an accumulator by providing strongly typed access to the 2D array. */
+		void Accumulate(EHitOrMiss HitOrMiss, EStatType StatType, int64 Value, bool bIsInGameThread)
+		{
+			Accumulators[(uint8)HitOrMiss][(uint8)StatType].Accumulate(Value, bIsInGameThread);
+		}
+
+		/** Make it easier to access an accumulator using a uniform, stronly typed interface. */
+		int64 GetAccumulatedValue(EHitOrMiss HitOrMiss, EStatType StatType, bool bIsInGameThread) const
+		{
+			return Accumulators[(uint8)HitOrMiss][(uint8)StatType].GetAccumulatedValue(bIsInGameThread);
+		}
+	
+	private:
+		/** The actual accumulators. All access should be from the above public functions. */
+		GameAndOtherThreadAccumulator Accumulators[(uint8)EHitOrMiss::MaxValue][(uint8)EStatType::MaxValue];
+	};
+
+	/**
+	 * used to accumulated cycles to the thread-safe counter.
+	 * Will also accumulate hit/miss stats in the dtor as well.
+	 * If AddHit is not called, it will assume a hit. If AddMiss is called, it will convert any previous hit call to a miss.
+	 */
+	class FScopedStatsCounter
+	{
+	public:
+		/** Starts the time, tracks the underlying stat it will update. */
+		explicit FScopedStatsCounter(CallStats& InStats)
+			: Stats(InStats)
+			, StartTime(FPlatformTime::Seconds())
+			, bIsInGameThread(IsInGameThread())
+		{
+		}
+		/** Ends the timer. Flushes the stats to the underlying stats object. */
+		~FScopedStatsCounter()
+		{
+			// Can't safely use FPlatformTime::Cycles() because we might be timing long durations.
+			const int64 CyclesUsed = int64((FPlatformTime::Seconds() - StartTime) / FPlatformTime::GetSecondsPerCycle());
+			if (bTrackHitOrMissCount)
+			{
+				Stats.Accumulate(HitOrMiss, CallStats::EStatType::Counter, 1l, bIsInGameThread);
+			}
+			Stats.Accumulate(HitOrMiss, CallStats::EStatType::Cycles, CyclesUsed, bIsInGameThread);
+			Stats.Accumulate(HitOrMiss, CallStats::EStatType::Bytes, BytesProcessed, bIsInGameThread);
+		}
+		
+		/** Call to indicate a Get or Put "cache hit". Exists calls by definition don't have hits or misses. */
+		void AddHit(int64 InBytesProcessed)
+		{
+			HitOrMiss = CallStats::EHitOrMiss::Hit;
+			BytesProcessed += InBytesProcessed;
+		}
+
+		/** Call to indicate a Get or Put "cache miss". This is optional, as a Miss is assumed by the timer. Exists calls by definition don't have hits or misses. */
+		void AddMiss()
+		{
+			HitOrMiss = CallStats::EHitOrMiss::Miss;
+			BytesProcessed = 0;
+		}
+		/** Used in rare case of async nodes to track sync work without necessarily adding a hit/miss, since we won't know until the async task runs, but we want to track cycles used. */
+		void UntrackHitOrMiss()
+		{
+			bTrackHitOrMissCount = false;
+		}
+
+	private:
+		CallStats& Stats;
+		double StartTime;
+		int64 BytesProcessed = 0;
+		bool bIsInGameThread;
+		bool bTrackHitOrMissCount = true;
+		CallStats::EHitOrMiss HitOrMiss = CallStats::EHitOrMiss::Miss;
+	};
+
+	/** Call this at the top of the CachedDataProbablyExists override. auto Timer = TimeProbablyExists(); */
+	FScopedStatsCounter TimeProbablyExists()
+	{
+		return FScopedStatsCounter(ExistsStats);
+	}
+
+	/** Call this at the top of the GetCachedData override. auto Timer = TimeGet(); Use AddHit on the returned type to track a cache hit. */
+	FScopedStatsCounter TimeGet()
+	{
+		return FScopedStatsCounter(GetStats);
+	}
+
+	/** Call this at the top of the PutCachedData override. auto Timer = TimePut(); Use AddHit on the returned type to track a cache hit. */
+	FScopedStatsCounter TimePut()
+	{
+		return FScopedStatsCounter(PutStats);
+	}
+
+	// expose these publicly for low level access. These should really never be accessed directly except when finished accumulating them.
+	CallStats GetStats;
+	CallStats PutStats;
+	CallStats ExistsStats;
+#endif
+};
diff --git a/Engine/Source/Editor/ContentBrowser/Private/CollectionAssetRegistryBridge.cpp b/Engine/Source/Editor/ContentBrowser/Private/CollectionAssetRegistryBridge.cpp
--- a/Engine/Source/Editor/ContentBrowser/Private/CollectionAssetRegistryBridge.cpp
+++ b/Engine/Source/Editor/ContentBrowser/Private/CollectionAssetRegistryBridge.cpp
@@ -46,10 +46,28 @@ public:
 		}
 		else
 		{
+			// Keep track of visted redirectors in case we loop.
+			TSet<FName> VisitedRedirectors;
+
 			// Use the asset registry to avoid loading the object
 			FAssetData ObjectAssetData = AssetRegistryModule.Get().GetAssetByObjectPath(InObjectPath);
 			while (ObjectAssetData.IsValid() && ObjectAssetData.IsRedirector())
 			{
+				// Check to see if we've already seen this path before, it's possible we might have found a redirector loop.
+				if ( VisitedRedirectors.Contains(ObjectAssetData.ObjectPath) )
+				{
+					UE_LOG(LogContentBrowser, Error, TEXT("Redirector Loop Found!"));
+					for ( FName Redirector : VisitedRedirectors )
+					{
+						UE_LOG(LogContentBrowser, Error, TEXT("Redirector: %s"), *Redirector.ToString());
+					}
+
+					ObjectAssetData = FAssetData();
+					break;
+				}
+
+				VisitedRedirectors.Add(ObjectAssetData.ObjectPath);
+
 				// Get the destination object from the meta-data rather than load the redirector object, as 
 				// loading a redirector will also load the object it points to, which could cause a large hitch
 				const FString* DestinationObjectStrPtr = ObjectAssetData.TagsAndValues.Find("DestinationObject");
diff --git a/Engine/Source/Editor/Sequencer/Private/Sequencer.cpp b/Engine/Source/Editor/Sequencer/Private/Sequencer.cpp
--- a/Engine/Source/Editor/Sequencer/Private/Sequencer.cpp
+++ b/Engine/Source/Editor/Sequencer/Private/Sequencer.cpp
@@ -174,7 +174,11 @@ void FSequencer::InitSequencer(const FSequencerInitParams& InitParams, const TSh
 
 	BindSequencerCommands();
 	ActivateSequencerEditorMode();
-	AttachTransportControlsToViewports();
+
+	if( bIsEditingWithinLevelEditor )
+	{
+		AttachTransportControlsToViewports();
+	}
 
 	for (auto TrackEditor : TrackEditors)
 	{
@@ -206,7 +210,10 @@ FSequencer::~FSequencer()
 	GEditor->GetActorRecordingState().RemoveAll( this );
 	GEditor->UnregisterForUndo( this );
 
-	DetachTransportControlsFromViewports();
+	if( bIsEditingWithinLevelEditor )
+	{
+		DetachTransportControlsFromViewports();
+	}
 
 	for (auto TrackEditor : TrackEditors)
 	{
@@ -249,10 +256,10 @@ void FSequencer::Close()
 			
 		FPropertyEditorModule& EditModule = FModuleManager::Get().GetModuleChecked<FPropertyEditorModule>("PropertyEditor");			
 		EditModule.OnPropertyEditorOpened().RemoveAll(this);
+
+		DetachTransportControlsFromViewports();
 	}
 	
-	DetachTransportControlsFromViewports();
-
 	SequencerWidget.Reset();
 }
 
diff --git a/Engine/Source/Editor/UnrealEd/Private/Commandlets/ChunkManifestGenerator.cpp b/Engine/Source/Editor/UnrealEd/Private/Commandlets/ChunkManifestGenerator.cpp
--- a/Engine/Source/Editor/UnrealEd/Private/Commandlets/ChunkManifestGenerator.cpp
+++ b/Engine/Source/Editor/UnrealEd/Private/Commandlets/ChunkManifestGenerator.cpp
@@ -1172,7 +1172,7 @@ bool FChunkManifestGenerator::CheckChunkAssetsAreNotInChild(const FChunkDependen
 
 void FChunkManifestGenerator::AddPackageAndDependenciesToChunk(FChunkPackageSet* ThisPackageSet, FName InPkgName, const FString& InSandboxFile, int32 ChunkID, FSandboxPlatformFile* SandboxPlatformFile)
 {
-	FChunkPackageSet* InitialPackageSetForThisChunk = ChunkManifests[ChunkID];
+	FChunkPackageSet* InitialPackageSetForThisChunk = ChunkManifests.IsValidIndex(ChunkID) ? ChunkManifests[ChunkID] : nullptr;
 
 	//Add this asset
 	ThisPackageSet->Add(InPkgName, InSandboxFile);
diff --git a/Engine/Source/Editor/UnrealEd/Private/Commandlets/CookCommandlet.cpp b/Engine/Source/Editor/UnrealEd/Private/Commandlets/CookCommandlet.cpp
--- a/Engine/Source/Editor/UnrealEd/Private/Commandlets/CookCommandlet.cpp
+++ b/Engine/Source/Editor/UnrealEd/Private/Commandlets/CookCommandlet.cpp
@@ -26,9 +26,116 @@
 #include "CookerSettings.h"
 #include "ShaderCompiler.h"
 #include "MemoryMisc.h"
+#include "CookStats.h"
+#include "DerivedDataCacheUsageStats.h"
 
 DEFINE_LOG_CATEGORY_STATIC(LogCookCommandlet, Log, All);
 
+#if ENABLE_COOK_STATS
+#include "ScopedTimers.h"
+#include "AnalyticsET.h"
+#include "IAnalyticsProvider.h"
+namespace DetailedCookStats
+{
+	FString CookProject;
+	FString TargetPlatforms;
+	double CookWallTimeSec = 0.0;
+	double StartupWallTimeSec = 0.0;
+	double NewCookTimeSec = 0.0;
+	double StartCookByTheBookTimeSec = 0.0;
+	extern double TickCookOnTheSideTimeSec;
+	extern double TickCookOnTheSideLoadPackagesTimeSec;
+	extern double TickCookOnTheSideResolveRedirectorsTimeSec;
+	extern double TickCookOnTheSideSaveCookedPackageTimeSec;
+	extern double TickCookOnTheSideBeginPackageCacheForCookedPlatformDataTimeSec;
+	extern double TickCookOnTheSideFinishPackageCacheForCookedPlatformDataTimeSec;
+	double TickLoopGCTimeSec = 0.0;
+	double TickLoopRecompileShaderRequestsTimeSec = 0.0;
+	double TickLoopShaderProcessAsyncResultsTimeSec = 0.0;
+	double TickLoopProcessDeferredCommandsTimeSec = 0.0;
+	double TickLoopTickCommandletStatsTimeSec = 0.0;
+
+	static void CollectStats(FCookStatsManager::AddStatFuncRef AddStat)
+	{
+		const FString StatName(TEXT("Cook.Profile"));
+		TArray<FCookStatsManager::KeyValue> Attrs;
+		#define ADD_COOK_STAT_FLT(Path, Name) AddStat(StatName, FCookStatsManager::CreateKeyValueArray(TEXT("Path"), TEXT(Path), TEXT(#Name), Name))
+		ADD_COOK_STAT_FLT(" 0", CookWallTimeSec);
+		ADD_COOK_STAT_FLT(" 0. 0", StartupWallTimeSec);
+		ADD_COOK_STAT_FLT(" 0. 0. 0", NewCookTimeSec);
+		ADD_COOK_STAT_FLT(" 0. 0. 0. 0", StartCookByTheBookTimeSec);
+		ADD_COOK_STAT_FLT(" 0. 0. 0. 1", TickCookOnTheSideTimeSec);
+		ADD_COOK_STAT_FLT(" 0. 0. 0. 1. 0", TickCookOnTheSideLoadPackagesTimeSec);
+		ADD_COOK_STAT_FLT(" 0. 0. 0. 1. 1", TickCookOnTheSideResolveRedirectorsTimeSec);
+		ADD_COOK_STAT_FLT(" 0. 0. 0. 1. 2", TickCookOnTheSideSaveCookedPackageTimeSec);
+		ADD_COOK_STAT_FLT(" 0. 0. 0. 1. 3", TickCookOnTheSideBeginPackageCacheForCookedPlatformDataTimeSec);
+		ADD_COOK_STAT_FLT(" 0. 0. 0. 1. 4", TickCookOnTheSideFinishPackageCacheForCookedPlatformDataTimeSec);
+		ADD_COOK_STAT_FLT(" 0. 0. 0. 2", TickLoopGCTimeSec);
+		ADD_COOK_STAT_FLT(" 0. 0. 0. 3", TickLoopRecompileShaderRequestsTimeSec);
+		ADD_COOK_STAT_FLT(" 0. 0. 0. 4", TickLoopShaderProcessAsyncResultsTimeSec);
+		ADD_COOK_STAT_FLT(" 0. 0. 0. 5", TickLoopProcessDeferredCommandsTimeSec);
+		ADD_COOK_STAT_FLT(" 0. 0. 0. 6", TickLoopTickCommandletStatsTimeSec);
+		#undef ADD_COOK_STAT_FLT
+	}
+	FCookStatsManager::FAutoRegisterCallback RegisterCookStats(&CollectStats);
+
+	/** Gathers the cook stats registered with the FCookStatsManager delegate and logs them as a CSV. */
+	static void LogCookStats(const FString& CookCmdLine)
+	{
+		// Optionally create an analytics provider to send stats to for central collection.
+		TSharedPtr<IAnalyticsProvider> CookAnalytics;
+		if (GIsBuildMachine || FParse::Param(FCommandLine::Get(), TEXT("SendCookAnalytics")))
+		{
+			FAnalyticsET::Config Config;
+			// This value is set by an INI private to Epic.
+			if (GConfig->GetString(TEXT("CookAnalytics"), TEXT("APIServer"), Config.APIServerET, GEngineIni))
+			{
+				Config.UseLegacyProtocol = true;
+				Config.APIKeyET = TEXT("Cook");
+				CookAnalytics = FAnalyticsET::Get().CreateAnalyticsProvider(Config);
+			}
+		}
+		if (CookAnalytics.IsValid())
+		{
+			CookAnalytics->SetUserID(FString(FPlatformProcess::ComputerName()) + FString(TEXT("\\")) + FString(FPlatformProcess::UserName(false)));
+			TArray<FAnalyticsEventAttribute> Attrs;
+			Attrs.Emplace(TEXT("Project"), CookProject);
+			Attrs.Emplace(TEXT("CmdLine"), CookCmdLine);
+			Attrs.Emplace(TEXT("IsBuildMachine"), GIsBuildMachine);
+			Attrs.Emplace(TEXT("TargetPlatforms"), TargetPlatforms);
+			CookAnalytics->StartSession(Attrs);
+		}
+
+		/** this functor will take a collected cooker stat and log it out, optionally sending it to a configured analytics provider. */
+		auto LogAndSendAnalytics = [=](const FString& StatName, const TArray<FCookStatsManager::KeyValue>& StatAttributes)
+		{
+			FString LogLine = StatName;
+			// log each key/value pair, with the equal signs lined up.
+			for (const auto& Attr : StatAttributes)
+			{
+				LogLine += FString::Printf(TEXT(",%s=%s"), *Attr.Key, *Attr.Value);
+			}
+			UE_LOG(LogCookCommandlet, Display, TEXT("%s"), *LogLine);
+			// convert to an analytics event
+			if (CookAnalytics.IsValid())
+			{
+				TArray<FAnalyticsEventAttribute> Attrs;
+				Attrs.Reset(StatAttributes.Num());
+				for (const auto& Attr : StatAttributes)
+				{
+					Attrs.Emplace(Attr.Key, Attr.Value);
+				}
+				CookAnalytics->RecordEvent(StatName, Attrs);
+			}
+		};
+
+		UE_LOG(LogCookCommandlet, Display, TEXT("CookStats"));
+		UE_LOG(LogCookCommandlet, Display, TEXT("---------"));
+		FCookStatsManager::LogCookStats(LogAndSendAnalytics);
+	}
+}
+#endif
+
 UCookerSettings::UCookerSettings(const FObjectInitializer& ObjectInitializer)
 	: Super(ObjectInitializer)
 	, bCompileBlueprintsInDevelopmentMode(true)
@@ -500,6 +607,7 @@ void UCookCommandlet::MaybeMarkPackageAsAlreadyLoaded(UPackage *Package)
 
 int32 UCookCommandlet::Main(const FString& CmdLineParams)
 {
+	COOK_STAT(double CookStartTime = FPlatformTime::Seconds());
 	Params = CmdLineParams;
 	ParseCommandLine(*Params, Tokens, Switches);
 
@@ -523,6 +631,8 @@ int32 UCookCommandlet::Main(const FString& CmdLineParams)
 		}
 	}
 
+	COOK_STAT(DetailedCookStats::CookProject = FApp::GetGameName());
+
 	if ( bCookOnTheFly )
 	{
 		// parse instance identifier
@@ -562,6 +672,15 @@ int32 UCookCommandlet::Main(const FString& CmdLineParams)
 		{
 			NewCook(Platforms, FilesInPath );
 		}
+
+		// Use -LogCookStats to log the results to the command line after the cook (happens automatically on a build machine)
+		COOK_STAT(
+		{
+			double Now = FPlatformTime::Seconds();
+			DetailedCookStats::CookWallTimeSec = Now - GStartTime;
+			DetailedCookStats::StartupWallTimeSec = CookStartTime - GStartTime;
+			DetailedCookStats::LogCookStats(CmdLineParams);
+		});
 	}
 	
 	return 0;
@@ -994,7 +1113,7 @@ void UCookCommandlet::GenerateLongPackageNames(TArray<FString>& FilesInPath)
 
 bool UCookCommandlet::NewCook( const TArray<ITargetPlatform*>& Platforms, TArray<FString>& FilesInPath )
 {
-
+	COOK_STAT(FScopedDurationTimer NewCookTimer(DetailedCookStats::NewCookTimeSec));
 	UCookOnTheFlyServer *CookOnTheFlyServer = NewObject<UCookOnTheFlyServer>();
 
 	struct FScopeRootObject
@@ -1207,7 +1326,21 @@ bool UCookCommandlet::NewCook( const TArray<ITargetPlatform*>& Platforms, TArray
 	StartupOptions.ChildCookIdentifier = ChildCookIdentifier;
 	StartupOptions.NumProcesses = NumProcesses;
 
-	CookOnTheFlyServer->StartCookByTheBook( StartupOptions );
+	COOK_STAT(
+	{
+		for (const auto& Platform : Platforms)
+		{
+			DetailedCookStats::TargetPlatforms += Platform->PlatformName() + TEXT("+");
+		}
+		if (!DetailedCookStats::TargetPlatforms.IsEmpty())
+		{
+			DetailedCookStats::TargetPlatforms.RemoveFromEnd(TEXT("+"));
+		}
+	});
+	{
+		COOK_STAT(FScopedDurationTimer StartCookByTheBookTimer(DetailedCookStats::StartCookByTheBookTimeSec));
+		CookOnTheFlyServer->StartCookByTheBook(StartupOptions);
+	}
 
 	// Garbage collection should happen when either
 	//	1. We have cooked a map (configurable asset type)
@@ -1242,7 +1375,10 @@ bool UCookCommandlet::NewCook( const TArray<ITargetPlatform*>& Platforms, TArray
 			}
 
 
-			GShaderCompilingManager->ProcessAsyncResults( true, false );
+			{
+				COOK_STAT(FScopedDurationTimer ShaderProcessAsyncTimer(DetailedCookStats::TickLoopShaderProcessAsyncResultsTimeSec));
+				GShaderCompilingManager->ProcessAsyncResults(true, false);
+			}
 
 			if (NonMapPackageCountSinceLastGC > 0)
 			{
@@ -1265,20 +1401,27 @@ bool UCookCommandlet::NewCook( const TArray<ITargetPlatform*>& Platforms, TArray
 
 				UE_LOG( LogCookCommandlet, Display, TEXT( "GC..." ) );
 
+				COOK_STAT(FScopedDurationTimer GCTimer(DetailedCookStats::TickLoopGCTimeSec));
 				CollectGarbage(RF_NoFlags);
 			}
 			else
 			{
+				COOK_STAT(FScopedDurationTimer RecompileTimer(DetailedCookStats::TickLoopRecompileShaderRequestsTimeSec));
 				CookOnTheFlyServer->TickRecompileShaderRequests();
 
 				FPlatformProcess::Sleep( 0.0f );
 			}
 
-
-			ProcessDeferredCommands();
+			{
+				COOK_STAT(FScopedDurationTimer ProcessDeferredCommandsTimer(DetailedCookStats::TickLoopProcessDeferredCommandsTimeSec));
+				ProcessDeferredCommands();
+			}
 		}
 
-		FStats::TickCommandletStats();
+		{
+			COOK_STAT(FScopedDurationTimer TickCommandletStatsTimer(DetailedCookStats::TickLoopTickCommandletStatsTimeSec));
+			FStats::TickCommandletStats();
+		}
 	}
 
 	return true;
diff --git a/Engine/Source/Editor/UnrealEd/Private/CookOnTheFlyServer.cpp b/Engine/Source/Editor/UnrealEd/Private/CookOnTheFlyServer.cpp
--- a/Engine/Source/Editor/UnrealEd/Private/CookOnTheFlyServer.cpp
+++ b/Engine/Source/Editor/UnrealEd/Private/CookOnTheFlyServer.cpp
@@ -41,6 +41,7 @@
 #include "ShaderCache.h"
 #include "Engine/LevelStreaming.h"
 #include "Engine/TextureLODSettings.h"
+#include "CookStats.h"
 
 #define LOCTEXT_NAMESPACE "Cooker"
 
@@ -53,6 +54,7 @@ DEFINE_LOG_CATEGORY_STATIC(LogCook, Log, All);
 #define VERIFY_GETDEPENDENTPACKAGES 0 // verify has false hits because old serialization method for generating dependencies had errors (included transient objects which shouldn't be in asset registry), but you can still use verify to build a list then cross check against transient objects.  
 
 #if OUTPUT_TIMING
+#include "ScopedTimers.h"
 
 #define HEIRARCHICAL_TIMER 1
 #define PERPACKAGE_TIMER 0
@@ -411,6 +413,20 @@ void OutputHierarchyTimers()
 #define OUTPUT_HIERARCHYTIMERS()
 #endif
 
+#if ENABLE_COOK_STATS
+#include "ScopedTimers.h"
+namespace DetailedCookStats
+{
+	//Externable so CookCommandlet can pick them up and merge them with it's cook stats
+	double TickCookOnTheSideTimeSec = 0.0;
+	double TickCookOnTheSideLoadPackagesTimeSec = 0.0;
+	double TickCookOnTheSideResolveRedirectorsTimeSec = 0.0;
+	double TickCookOnTheSideSaveCookedPackageTimeSec = 0.0;
+	double TickCookOnTheSideBeginPackageCacheForCookedPlatformDataTimeSec = 0.0;
+	double TickCookOnTheSideFinishPackageCacheForCookedPlatformDataTimeSec = 0.0;
+}
+#endif
+
 //////////////////////////////////////////////////////////////////////////
 // FCookTimer
 // used as a helper to timeslice cooker functions
@@ -1295,9 +1311,9 @@ bool UCookOnTheFlyServer::TickChildCookers()
 
 }
 
-
 uint32 UCookOnTheFlyServer::TickCookOnTheSide( const float TimeSlice, uint32 &CookedPackageCount )
 {
+	COOK_STAT(FScopedDurationTimer TickTimer(DetailedCookStats::TickCookOnTheSideTimeSec));
 	FCookerTimer Timer(TimeSlice, IsRealtimeMode());
 
 	uint32 Result = 0;
@@ -1457,6 +1473,7 @@ uint32 UCookOnTheFlyServer::TickCookOnTheSide( const float TimeSlice, uint32 &Co
 
 		if ( bShouldCook ) // if we should cook the package then cook it otherwise add it to the list of already cooked packages below
 		{
+			COOK_STAT(FScopedDurationTimer LoadPackagesTimer(DetailedCookStats::TickCookOnTheSideLoadPackagesTimeSec));
 			UPackage *Package = NULL;
 			{
 				FString PackageName;
@@ -1616,6 +1633,7 @@ uint32 UCookOnTheFlyServer::TickCookOnTheSide( const float TimeSlice, uint32 &Co
 		if ( IsCookByTheBookMode() )
 		{
 			SCOPE_TIMER(ResolveRedirectors);
+			COOK_STAT(FScopedDurationTimer ResolveRedirectorsTimer(DetailedCookStats::TickCookOnTheSideResolveRedirectorsTimeSec));
 			GRedirectCollector.ResolveStringAssetReference();
 		}
 
@@ -1632,6 +1650,7 @@ uint32 UCookOnTheFlyServer::TickCookOnTheSide( const float TimeSlice, uint32 &Co
 
 		auto BeginPackageCacheForCookedPlatformData = [&]( const TArray<UObject*>& ObjectsInPackage )
 		{
+			COOK_STAT(FScopedDurationTimer DurationTimer(DetailedCookStats::TickCookOnTheSideBeginPackageCacheForCookedPlatformDataTimeSec));
 			if (CurrentReentryData.bBeginCacheFinished)
 				return true;
 			for (; CurrentReentryData.BeginCacheCount < ObjectsInPackage.Num(); ++CurrentReentryData.BeginCacheCount)
@@ -1656,15 +1675,19 @@ uint32 UCookOnTheFlyServer::TickCookOnTheSide( const float TimeSlice, uint32 &Co
 
 		auto FinishPackageCacheForCookedPlatformData = [&]( const TArray<UObject*>& ObjectsInPackage )
 		{
-			for ( const auto& Obj : ObjectsInPackage )
+			COOK_STAT(FScopedDurationTimer DurationTimer(DetailedCookStats::TickCookOnTheSideFinishPackageCacheForCookedPlatformDataTimeSec));
+			for (const auto& Obj : ObjectsInPackage)
 			{
-				for ( const auto& TargetPlatform : TargetPlatforms )
+				for (const auto& TargetPlatform : TargetPlatforms)
 				{
+					COOK_STAT(double CookerStatSavedValue = DetailedCookStats::TickCookOnTheSideBeginPackageCacheForCookedPlatformDataTimeSec);
 					// These begin cache calls should be quick 
 					// because they will just be checking that the data is already cached and kicking off new multithreaded requests if not
 					// all sync requests should have been caught in the first begincache call above
-					Obj->BeginCacheForCookedPlatformData( TargetPlatform );
-					if ( Obj->IsCachedCookedPlatformDataLoaded(TargetPlatform) == false )
+					Obj->BeginCacheForCookedPlatformData(TargetPlatform);
+					// We want to measure inclusive time for this function, but not accumulate into the BeginXXX timer, so subtract these times out of the BeginTimer.
+					COOK_STAT(DetailedCookStats::TickCookOnTheSideBeginPackageCacheForCookedPlatformDataTimeSec = CookerStatSavedValue);
+					if (Obj->IsCachedCookedPlatformDataLoaded(TargetPlatform) == false)
 					{
 #if DEBUG_COOKONTHEFLY
 						UE_LOG(LogCook, Display, TEXT("Object %s isn't cached yet"), *Obj->GetFullName());
@@ -1676,7 +1699,7 @@ uint32 UCookOnTheFlyServer::TickCookOnTheSide( const float TimeSlice, uint32 &Co
 				// if this objects data is cached then we can call FinishedCookedPLatformDataCache
 				// we can only safely call this when we are finished caching this object completely.
 				// this doesn't ever happen for cook in editor or cook on the fly mode
-				if ( bIsAllDataCached && (CurrentCookMode == ECookMode::CookByTheBook) )
+				if (bIsAllDataCached && (CurrentCookMode == ECookMode::CookByTheBook))
 				{
 					// this might be run multiple times for a single object
 					Obj->WillNeverCacheCookedPlatformDataAgain();
@@ -1924,6 +1947,7 @@ uint32 UCookOnTheFlyServer::TickCookOnTheSide( const float TimeSlice, uint32 &Co
 					{
 						// add to back of queue
 						PackagesToSave.Add(Package);
+						continue;
 					}
 				}
 
@@ -1939,6 +1963,7 @@ uint32 UCookOnTheFlyServer::TickCookOnTheSide( const float TimeSlice, uint32 &Co
 				ESavePackageResult SavePackageResult = ESavePackageResult::Error;
 				bool bSucceededSavePackage = false;
 				{
+					COOK_STAT(FScopedDurationTimer TickCookOnTheSideSaveCookedPackageTimer(DetailedCookStats::TickCookOnTheSideSaveCookedPackageTimeSec));
 					SCOPE_TIMER(SaveCookedPackage);
 					uint32 SaveFlags = SAVE_KeepGUID | (bShouldSaveAsync ? SAVE_Async : SAVE_None) | (IsCookFlagSet(ECookInitializationFlags::Unversioned) ? SAVE_Unversioned : 0);
 					GOutputCookingWarnings = IsCookFlagSet(ECookInitializationFlags::OutputVerboseCookerWarnings);
diff --git a/Engine/Source/Editor/UnrealEd/UnrealEd.Build.cs b/Engine/Source/Editor/UnrealEd/UnrealEd.Build.cs
--- a/Engine/Source/Editor/UnrealEd/UnrealEd.Build.cs
+++ b/Engine/Source/Editor/UnrealEd/UnrealEd.Build.cs
@@ -23,7 +23,6 @@ public class UnrealEd : ModuleRules
 		PrivateIncludePathModuleNames.AddRange(
 			new string[] 
 			{
-				"Analytics",
 				"AssetRegistry",
 				"AssetTools",
                 "BehaviorTreeEditor",
@@ -144,7 +143,7 @@ public class UnrealEd : ModuleRules
 				"AddContentDialog",
 				"GameProjectGeneration",
                 "HierarchicalLODUtilities",
-                "Analytics",
+                "AnalyticsET",
                 "PluginWarden"
 			}
 		);
diff --git a/Engine/Source/Programs/AutomationTool/GUBP/GUBP.cs b/Engine/Source/Programs/AutomationTool/GUBP/GUBP.cs
--- a/Engine/Source/Programs/AutomationTool/GUBP/GUBP.cs
+++ b/Engine/Source/Programs/AutomationTool/GUBP/GUBP.cs
@@ -883,6 +883,9 @@ public partial class GUBP : BuildCommand
 					}
 				}
 			}
+			History.AllStarted.Sort();
+			History.AllSucceeded.Sort();
+			History.AllFailed.Sort();
 
 			// Filter everything relative to the current CL
             if (History.AllFailed.Count > 0)
diff --git a/Engine/Source/Programs/AutomationTool/GUBP/LegacyBranchSetup.cs b/Engine/Source/Programs/AutomationTool/GUBP/LegacyBranchSetup.cs
--- a/Engine/Source/Programs/AutomationTool/GUBP/LegacyBranchSetup.cs
+++ b/Engine/Source/Programs/AutomationTool/GUBP/LegacyBranchSetup.cs
@@ -717,7 +717,15 @@ partial class GUBP
 										BranchConfig.AddNode(new GamePlatformMonolithicsNode(BranchConfig, HostPlatform, ActivePlatforms, Branch.BaseEngineProject, Plat, true));
 									}
 	                            }
-                            }
+								if (!BranchConfig.HasNode(GamePlatformMonolithicsKindNode.StaticGetFullName(HostPlatform, Branch.BaseEngineProject, Plat, Kind)))
+								{
+									if (GamePlatformMonolithicsKindNode.HasPrecompiledTargets(Branch.BaseEngineProject, HostPlatform, Plat, Kind))
+									{
+										BranchConfig.AddNode(new GamePlatformMonolithicsKindNode(BranchConfig, HostPlatform, ActivePlatforms, Branch.BaseEngineProject, Plat, Kind, InPrecompiled: true));
+									}
+									BranchConfig.AddNode(new GamePlatformMonolithicsKindNode(BranchConfig, HostPlatform, ActivePlatforms, Branch.BaseEngineProject, Plat, Kind));
+								}
+							}
                         }
                     }
                 }
@@ -1004,7 +1012,15 @@ partial class GUBP
 									}
                                     BranchConfig.AddNode(new GamePlatformMonolithicsNode(BranchConfig, HostPlatform, ActivePlatforms, CodeProj, Plat));
                                 }
-                                var FormalBuildConfigs = Target.Rules.GUBP_GetConfigsForFormalBuilds_MonolithicOnly(HostPlatform);
+								if (!BranchConfig.HasNode(GamePlatformMonolithicsKindNode.StaticGetFullName(HostPlatform, CodeProj, Plat, Kind)))
+								{
+									if (GamePlatformMonolithicsKindNode.HasPrecompiledTargets(CodeProj, HostPlatform, Plat, Kind))
+									{
+										BranchConfig.AddNode(new GamePlatformMonolithicsKindNode(BranchConfig, HostPlatform, ActivePlatforms, CodeProj, Plat, Kind, InPrecompiled: true));
+									}
+									BranchConfig.AddNode(new GamePlatformMonolithicsKindNode(BranchConfig, HostPlatform, ActivePlatforms, CodeProj, Plat, Kind));
+								}
+								var FormalBuildConfigs = Target.Rules.GUBP_GetConfigsForFormalBuilds_MonolithicOnly(HostPlatform);
 								if (!AdditionalPlatforms.Contains(Plat) && (BranchOptions.ProjectsToCook.Contains(CodeProj.GameName) || BranchOptions.ProjectsToCook.Count == 0))
 								{
 									string CookedPlatform = Platform.Platforms[Plat].GetCookPlatform(Kind == TargetRules.TargetType.Server, Kind == TargetRules.TargetType.Client, "");
diff --git a/Engine/Source/Programs/AutomationTool/GUBP/LegacyNodes.cs b/Engine/Source/Programs/AutomationTool/GUBP/LegacyNodes.cs
--- a/Engine/Source/Programs/AutomationTool/GUBP/LegacyNodes.cs
+++ b/Engine/Source/Programs/AutomationTool/GUBP/LegacyNodes.cs
@@ -1075,6 +1075,7 @@ partial class GUBP
     {
         BranchInfo.BranchUProject GameProj;
 		List<UnrealTargetPlatform> ActivePlatforms;
+		protected List<TargetRules.TargetType> ActiveMonolithicKinds;
         UnrealTargetPlatform TargetPlatform;
 		bool WithXp;
 		bool Precompiled; // If true, just builds targets which generate static libraries for the -UsePrecompiled option to UBT. If false, just build those that don't.
@@ -1089,6 +1090,7 @@ partial class GUBP
 			WithXp = InWithXp;
 			Precompiled = InPrecompiled;
 			EnhanceAgentRequirements = BranchConfig.BranchOptions.EnhanceAgentRequirements.Contains(StaticGetFullName(HostPlatform, GameProj, TargetPlatform, WithXp, Precompiled));
+			ActiveMonolithicKinds = BranchInfo.MonolithicKinds;
 
             if (TargetPlatform == UnrealTargetPlatform.PS4 || TargetPlatform == UnrealTargetPlatform.XboxOne)
             {
@@ -1318,7 +1320,7 @@ partial class GUBP
 				Args += " -winxp";
 			}
 
-            foreach (var Kind in BranchInfo.MonolithicKinds)
+            foreach (var Kind in ActiveMonolithicKinds)
             {
                 if (GameProj.Properties.Targets.ContainsKey(Kind))
                 {
@@ -1363,14 +1365,49 @@ partial class GUBP
 		public static string StaticGetArchivedHeadersPath(UnrealTargetPlatform HostPlatform, BranchInfo.BranchUProject GameProj, UnrealTargetPlatform TargetPlatform)
 		{
 			return CommandUtils.CombinePaths(CommandUtils.CmdEnv.LocalRoot, "Engine", "Saved", "Precompiled", "Headers-" + StaticGetFullName(HostPlatform, GameProj, TargetPlatform) + ".zip");
-    }
+		}
 
 		public static string StaticGetBuildDependenciesPath(UnrealTargetPlatform HostPlatform, BranchInfo.BranchUProject GameProj, UnrealTargetPlatform TargetPlatform)
-	{
+		{
 			return CommandUtils.CombinePaths(CommandUtils.CmdEnv.LocalRoot, "Engine", "Saved", "Precompiled", "BuildDependencies-" + StaticGetFullName(HostPlatform, GameProj, TargetPlatform) + ".xml");
 		}
 	}
 
+	public class GamePlatformMonolithicsKindNode : GamePlatformMonolithicsNode
+	{
+		public GamePlatformMonolithicsKindNode(GUBPBranchConfig InBranchConfig, UnrealTargetPlatform InHostPlatform, List<UnrealTargetPlatform> InActivePlatforms, BranchInfo.BranchUProject InGameProj, UnrealTargetPlatform InTargetPlatform, TargetRules.TargetType InTargetKind, bool InWithXp = false, bool InPrecompiled = false)
+			: base(InBranchConfig, InHostPlatform, InActivePlatforms, InGameProj, InTargetPlatform, InWithXp, InPrecompiled)
+		{
+			ActiveMonolithicKinds = new List<TargetRules.TargetType>();
+			ActiveMonolithicKinds.Add(InTargetKind);
+		}
+
+		public static string StaticGetFullName(UnrealTargetPlatform InHostPlatform, BranchInfo.BranchUProject InGameProj, UnrealTargetPlatform InTargetPlatform, TargetRules.TargetType InTargetKind, bool WithXp = false, bool Precompiled = false)
+		{
+			string Name = GamePlatformMonolithicsNode.StaticGetFullName(InHostPlatform, InGameProj, InTargetPlatform, WithXp, Precompiled);
+			return Name + "_" + InTargetKind;
+		}
+
+		public override string GetFullName()
+		{
+			string Name = base.GetFullName();
+			return Name + "_" + ActiveMonolithicKinds[0];
+		}
+
+		public static bool HasPrecompiledTargets(BranchInfo.BranchUProject Project, UnrealTargetPlatform HostPlatform, UnrealTargetPlatform TargetPlatform, TargetRules.TargetType InTargetKind)
+		{
+			if (Project.Properties.Targets.ContainsKey(InTargetKind))
+			{
+				SingleTargetProperties Target = Project.Properties.Targets[InTargetKind];
+				if (Target.Rules.GUBP_GetConfigsForPrecompiledBuilds_MonolithicOnly(HostPlatform, TargetPlatform).Any())
+				{
+					return true;
+				}
+			}
+			return false;
+		}	
+	}
+
     public class SuccessNode : GUBPNode
     {
         public SuccessNode()
diff --git a/Engine/Source/Programs/AutomationTool/GUBP/TempStorage.cs b/Engine/Source/Programs/AutomationTool/GUBP/TempStorage.cs
--- a/Engine/Source/Programs/AutomationTool/GUBP/TempStorage.cs
+++ b/Engine/Source/Programs/AutomationTool/GUBP/TempStorage.cs
@@ -452,7 +452,8 @@ namespace AutomationTool
 					}
 					catch (Exception Ex)
 					{
-						CommandUtils.LogWarning("Failed to delete old manifest folder '{0}', will try one file at a time: {1}", OldManifestInfo.NodeNameDir, Ex);
+                        // Quick fix for race condition in Orion between jobs trying to clear the same storage
+						CommandUtils.Log("Failed to delete old manifest folder '{0}', will try one file at a time: {1}", OldManifestInfo.NodeNameDir, Ex);
 						if (CommandUtils.DeleteDirectory_NoExceptions(true, OldManifestInfo.NodeNameDir))
 						{
 							FoldersManuallyCleaned++;
@@ -488,7 +489,11 @@ namespace AutomationTool
 			}
 			catch (Exception Ex)
 			{
-				CommandUtils.LogWarning("Unable to Clean Temp Directory {0}. Exception: {1}", TempStorageFolder, Ex);
+                // Quick fix for a race condition between jobs cleaning out the shared storage in Orion
+                if (!(Ex is System.IO.DirectoryNotFoundException))
+                {
+                    CommandUtils.LogWarning("Unable to Clean Temp Directory {0}. Exception: {1}", TempStorageFolder, Ex);
+                }
 			}
 		}
 
diff --git a/Engine/Source/Programs/AutomationTool/Scripts/CopyBuildToStagingDirectory.Automation.cs b/Engine/Source/Programs/AutomationTool/Scripts/CopyBuildToStagingDirectory.Automation.cs
--- a/Engine/Source/Programs/AutomationTool/Scripts/CopyBuildToStagingDirectory.Automation.cs
+++ b/Engine/Source/Programs/AutomationTool/Scripts/CopyBuildToStagingDirectory.Automation.cs
@@ -913,8 +913,8 @@ public partial class Project : CommandUtils
 		{
 			if (File.Exists(OutputLocation))
 			{
-				string UnrealPakResponseFileName = CombinePaths(CmdEnv.LogFolder, "PakList_" + PakName + ".txt");
-				if (File.GetLastWriteTimeUtc(OutputLocation) > File.GetLastWriteTimeUtc(UnrealPakResponseFileName))
+				string UnrealPakResponseFileName = CombinePaths(CmdEnv.LogFolder, "PakList_" + Path.GetFileNameWithoutExtension(OutputLocation) + ".txt");
+				if (File.Exists(UnrealPakResponseFileName) && File.GetLastWriteTimeUtc(OutputLocation) > File.GetLastWriteTimeUtc(UnrealPakResponseFileName))
 				{
 					bCopiedExistingPak = true;
 				}
@@ -1786,6 +1786,11 @@ public partial class Project : CommandUtils
 				foreach (var SC in DeployContextList)
 				{
 					CreateStagingManifest(Params, SC);
+					CleanStagingDirectory(Params, SC);
+				}
+
+				foreach (var SC in DeployContextList)
+				{
 					ApplyStagingManifest(Params, SC);
 				}
 			}
diff --git a/Engine/Source/Programs/CrashReporter/CrashReportClient/Private/CrashDescription.cpp b/Engine/Source/Programs/CrashReporter/CrashReportClient/Private/CrashDescription.cpp
--- a/Engine/Source/Programs/CrashReporter/CrashReportClient/Private/CrashDescription.cpp
+++ b/Engine/Source/Programs/CrashReporter/CrashReportClient/Private/CrashDescription.cpp
@@ -183,52 +183,96 @@ FString FPrimaryCrashProperties::EncodeArrayStringAsXMLString( const TArray<FStr
 	return Encoded;
 }
 
-void FPrimaryCrashProperties::SendAnalytics()
+void FPrimaryCrashProperties::SendPreUploadAnalytics()
 {
-	// Connect the crash report client analytics provider.
-	FCrashReportAnalytics::Initialize();
-
-	IAnalyticsProvider& Analytics = FCrashReportAnalytics::GetProvider();
-
 	TArray<FAnalyticsEventAttribute> CrashAttributes;
 
-	CrashAttributes.Add( FAnalyticsEventAttribute( TEXT( "bHasPrimaryData" ), bHasPrimaryData ) );
-	CrashAttributes.Add( FAnalyticsEventAttribute( TEXT( "CrashVersion" ), (int32)CrashVersion ) );
-	CrashAttributes.Add( FAnalyticsEventAttribute( TEXT( "CrashGUID" ), CrashGUID ) );
+	CrashAttributes.Add(FAnalyticsEventAttribute(TEXT("bHasPrimaryData"), bHasPrimaryData));
+	CrashAttributes.Add(FAnalyticsEventAttribute(TEXT("CrashVersion"), (int32)CrashVersion));
+	CrashAttributes.Add(FAnalyticsEventAttribute(TEXT("CrashGUID"), CrashGUID));
 
 	//	AppID = GameName
-	CrashAttributes.Add( FAnalyticsEventAttribute( TEXT( "GameName" ), GameName ) );
+	CrashAttributes.Add(FAnalyticsEventAttribute(TEXT("GameName"), GameName));
 
 	//	AppVersion = EngineVersion
-	CrashAttributes.Add( FAnalyticsEventAttribute( TEXT( "EngineVersion" ), EngineVersion.ToString() ) );
+	CrashAttributes.Add(FAnalyticsEventAttribute(TEXT("EngineVersion"), EngineVersion.ToString()));
 
 	// @see UpdateIDs()
-	CrashAttributes.Add( FAnalyticsEventAttribute( TEXT( "MachineID" ), MachineId.AsString() ) );
-	CrashAttributes.Add( FAnalyticsEventAttribute( TEXT( "UserName" ), UserName.AsString() ) );
-	CrashAttributes.Add( FAnalyticsEventAttribute( TEXT( "EpicAccountId" ), EpicAccountId.AsString() ) );
+	CrashAttributes.Add(FAnalyticsEventAttribute(TEXT("MachineID"), MachineId.AsString()));
+	CrashAttributes.Add(FAnalyticsEventAttribute(TEXT("UserName"), UserName.AsString()));
+	CrashAttributes.Add(FAnalyticsEventAttribute(TEXT("EpicAccountId"), EpicAccountId.AsString()));
 
-	CrashAttributes.Add( FAnalyticsEventAttribute( TEXT( "Platform" ), PlatformFullName.AsString() ) );
-	CrashAttributes.Add( FAnalyticsEventAttribute( TEXT( "TimeOfCrash" ), TimeOfCrash.AsString() ) );
-	CrashAttributes.Add( FAnalyticsEventAttribute( TEXT( "EngineMode" ), EngineMode ) );
-	CrashAttributes.Add( FAnalyticsEventAttribute( TEXT( "AppDefaultLocale" ), AppDefaultLocale ) );
+	CrashAttributes.Add(FAnalyticsEventAttribute(TEXT("Platform"), PlatformFullName.AsString()));
+	CrashAttributes.Add(FAnalyticsEventAttribute(TEXT("TimeOfCrash"), TimeOfCrash.AsString()));
+	CrashAttributes.Add(FAnalyticsEventAttribute(TEXT("EngineMode"), EngineMode));
+	CrashAttributes.Add(FAnalyticsEventAttribute(TEXT("AppDefaultLocale"), AppDefaultLocale));
 
-	CrashAttributes.Add( FAnalyticsEventAttribute( TEXT( "UserActivityHint" ), UserActivityHint.AsString() ) );
-	CrashAttributes.Add( FAnalyticsEventAttribute( TEXT( "GameSessionID" ), GameSessionID.AsString() ) );
-	CrashAttributes.Add( FAnalyticsEventAttribute( TEXT( "DeploymentName" ), DeploymentName ) );
+	CrashAttributes.Add(FAnalyticsEventAttribute(TEXT("UserActivityHint"), UserActivityHint.AsString()));
+	CrashAttributes.Add(FAnalyticsEventAttribute(TEXT("GameSessionID"), GameSessionID.AsString()));
+	CrashAttributes.Add(FAnalyticsEventAttribute(TEXT("DeploymentName"), DeploymentName));
 
-	Analytics.RecordEvent( TEXT( "CrashReportClient.ReportCrash" ), CrashAttributes );
-	// Shutdown analytics.
-	FCrashReportAnalytics::Shutdown();
+	if (FCrashReportAnalytics::IsAvailable())
+	{
+		// Connect the crash report client analytics provider.
+		IAnalyticsProvider& Analytics = FCrashReportAnalytics::GetProvider();
+
+		if (bIsEnsure)
+		{
+			Analytics.RecordEvent(TEXT("CrashReportClient.ReportEnsure"), CrashAttributes);
+		}
+		else
+		{
+			Analytics.RecordEvent(TEXT("CrashReportClient.ReportCrash"), CrashAttributes);
+		}
+	}
 
 	// duplicate the event to QoS Reporter
-	FQoSReporter::Initialize();
-	FQoSReporter::SetBackendDeploymentName(DeploymentName);
+	if (FQoSReporter::IsAvailable())
+	{
+		if (bIsEnsure)
+		{
+			FQoSReporter::GetProvider().RecordEvent(TEXT("CrashReportClient.ReportEnsure"), CrashAttributes);
+		}
+		else
+		{
+			FQoSReporter::GetProvider().RecordEvent(TEXT("CrashReportClient.ReportCrash"), CrashAttributes);
+		}
+	}
+}
 
+void FPrimaryCrashProperties::SendPostUploadAnalytics()
+{
+	TArray<FAnalyticsEventAttribute> CrashAttributes;
+
+	CrashAttributes.Add(FAnalyticsEventAttribute(TEXT("CrashGUID"), CrashGUID));
+
+	if (FCrashReportAnalytics::IsAvailable())
+	{
+		// Connect the crash report client analytics provider.
+		IAnalyticsProvider& Analytics = FCrashReportAnalytics::GetProvider();
+
+		if (bIsEnsure)
+		{
+			Analytics.RecordEvent(TEXT("CrashReportClient.ReportEnsureUploaded"), CrashAttributes);
+		}
+		else
+		{
+			Analytics.RecordEvent(TEXT("CrashReportClient.ReportCrashUploaded"), CrashAttributes);
+		}
+	}
+
+	// duplicate the event to QoS Reporter
 	if (FQoSReporter::IsAvailable())
 	{
-		FQoSReporter::GetProvider().RecordEvent(TEXT("CrashReportClient.ReportCrash"), CrashAttributes);
+		if (bIsEnsure)
+		{
+			FQoSReporter::GetProvider().RecordEvent(TEXT("CrashReportClient.ReportEnsureUploaded"), CrashAttributes);
+		}
+		else
+		{
+			FQoSReporter::GetProvider().RecordEvent(TEXT("CrashReportClient.ReportCrashUploaded"), CrashAttributes);
+		}
 	}
-	FQoSReporter::Shutdown();
 }
 
 void FPrimaryCrashProperties::Save()
@@ -280,6 +324,7 @@ FCrashContext::FCrashContext( const FString& CrashContextFilepath )
 		GetCrashProperty( EngineMode, FGenericCrashContext::RuntimePropertiesTag, TEXT( "EngineMode" ) );
 		GetCrashProperty( DeploymentName, FGenericCrashContext::RuntimePropertiesTag, TEXT( "DeploymentName" ) );
 		GetCrashProperty( AppDefaultLocale, FGenericCrashContext::RuntimePropertiesTag, TEXT( "AppDefaultLocale" ) );
+		GetCrashProperty( bIsEnsure, FGenericCrashContext::RuntimePropertiesTag, TEXT("IsEnsure"));
 
 		if (CrashDumpMode == ECrashDumpMode::FullDump)
 		{
@@ -399,6 +444,7 @@ FCrashWERContext::FCrashWERContext( const FString& WERXMLFilepath )
 		}
 
 		GetCrashProperty(DeploymentName, TEXT("DynamicSignatures"), TEXT("DeploymentName"));
+		GetCrashProperty(bIsEnsure, TEXT("DynamicSignatures"), TEXT("IsEnsure"));
 
 		bHasPrimaryData = true;
 	}
diff --git a/Engine/Source/Programs/CrashReporter/CrashReportClient/Private/CrashReportClient.cpp b/Engine/Source/Programs/CrashReporter/CrashReportClient/Private/CrashReportClient.cpp
--- a/Engine/Source/Programs/CrashReporter/CrashReportClient/Private/CrashReportClient.cpp
+++ b/Engine/Source/Programs/CrashReporter/CrashReportClient/Private/CrashReportClient.cpp
@@ -83,12 +83,6 @@ FReply FCrashReportClient::CloseWithoutSending()
 
 FReply FCrashReportClient::Submit()
 {
-	if (ErrorReport.HasFilesToUpload())
-	{
-		// Send analytics.
-		FPrimaryCrashProperties::Get()->SendAnalytics();
-	}
-
 	bSendData = true;
 	StoreCommentAndUpload();
 	bShouldWindowBeHidden = true;
diff --git a/Engine/Source/Programs/CrashReporter/CrashReportClient/Private/CrashReportClientApp.cpp b/Engine/Source/Programs/CrashReporter/CrashReportClient/Private/CrashReportClientApp.cpp
--- a/Engine/Source/Programs/CrashReporter/CrashReportClient/Private/CrashReportClientApp.cpp
+++ b/Engine/Source/Programs/CrashReporter/CrashReportClient/Private/CrashReportClientApp.cpp
@@ -3,6 +3,7 @@
 #include "CrashReportClientApp.h"
 #include "CrashDescription.h"
 #include "CrashReportAnalytics.h"
+#include "QoSReporter.h"
 
 #if !CRASH_REPORT_UNATTENDED_ONLY
 	#include "SCrashReportClient.h"
@@ -165,18 +166,16 @@ void RunCrashReportClient(const TCHAR* CommandLine)
 	
 	if (ErrorReport.HasFilesToUpload() && FPrimaryCrashProperties::Get() != nullptr)
 	{
+		FCrashReportAnalytics::Initialize();
+		FQoSReporter::Initialize();
+		FQoSReporter::SetBackendDeploymentName(FPrimaryCrashProperties::Get()->DeploymentName);
+
 		if (bUnattended)
 		{
 			// In the unattended mode we don't send any PII.
 			FCrashReportClientUnattended CrashReportClient(ErrorReport);
 			ErrorReport.SetUserComment(NSLOCTEXT("CrashReportClient", "UnattendedMode", "Sent in the unattended mode"));
 
-			if (ErrorReport.HasFilesToUpload())
-			{
-				// Send analytics.
-				FPrimaryCrashProperties::Get()->SendAnalytics();
-			}
-
 			// loop until the app is ready to quit
 			while (!GIsRequestingExit)
 			{
@@ -236,6 +235,10 @@ void RunCrashReportClient(const TCHAR* CommandLine)
 			FSlateApplication::Shutdown();
 #endif // !CRASH_REPORT_UNATTENDED_ONLY
 		}
+
+		// Shutdown analytics.
+		FCrashReportAnalytics::Shutdown();
+		FQoSReporter::Shutdown();
 	}
 
 	FPrimaryCrashProperties::Shutdown();
diff --git a/Engine/Source/Programs/CrashReporter/CrashReportClient/Private/CrashReportClientConfig.cpp b/Engine/Source/Programs/CrashReporter/CrashReportClient/Private/CrashReportClientConfig.cpp
--- a/Engine/Source/Programs/CrashReporter/CrashReportClient/Private/CrashReportClientConfig.cpp
+++ b/Engine/Source/Programs/CrashReporter/CrashReportClient/Private/CrashReportClientConfig.cpp
@@ -77,11 +77,12 @@ const FString FCrashReportClientConfig::GetFullCrashDumpLocationForBranch( const
 	for (const auto& It : FullCrashDumpConfigurations)
 	{
 		const bool bExactMatch = It.bExactMatch;
-		if (bExactMatch && BranchName == It.BranchName)
+		const FString IterBranch = It.BranchName.Replace(TEXT("+"), TEXT("/"));
+		if (bExactMatch && BranchName == IterBranch)
 		{
 			return It.Location;
 		}
-		else if (!bExactMatch && BranchName.Contains( It.BranchName ))
+		else if (!bExactMatch && BranchName.Contains(IterBranch))
 		{
 			return It.Location;
 		}
diff --git a/Engine/Source/Programs/CrashReporter/CrashReportClient/Private/CrashUpload.cpp b/Engine/Source/Programs/CrashReporter/CrashReportClient/Private/CrashUpload.cpp
--- a/Engine/Source/Programs/CrashReporter/CrashReportClient/Private/CrashUpload.cpp
+++ b/Engine/Source/Programs/CrashReporter/CrashReportClient/Private/CrashUpload.cpp
@@ -12,6 +12,10 @@
 #include "CrashDescription.h"
 #include "EngineBuildSettings.h"
 
+// TODO Chris.Wood: switch these when backend migrates to DR
+#define PRIMARY_UPLOAD_RECEIVER 1
+#define PRIMARY_UPLOAD_DATAROUTER 0
+
 #define LOCTEXT_NAMESPACE "CrashReportClient"
 
 namespace CrashUploadDefs
@@ -160,7 +164,7 @@ bool FCrashUploadBase::CompressData(const TArray<FString>& InPendingFiles, FComp
 			else
 			{
 				UE_LOG(CrashReportClientLog, Warning, TEXT("Copying full crash minidump to %s"), *DestinationPath);
-				IFileManager::Get().Copy(*DestinationPath, *PathOfFileToUpload);
+				IFileManager::Get().Copy(*DestinationPath, *PathOfFileToUpload, false);
 			}
 
 			continue;
@@ -398,6 +402,10 @@ bool FCrashUploadToReceiver::SendCheckReportRequest()
 	auto Request = CreateHttpRequest();
 	if (State == EUploadState::CheckingReport)
 	{
+#if PRIMARY_UPLOAD_RECEIVER
+		// first stage of any upload to CRR so send analytics
+		FPrimaryCrashProperties::Get()->SendPreUploadAnalytics();
+#endif
 		AssignReportIdToPostDataBuffer();
 		Request->SetURL(UrlPrefix / TEXT("CheckReport"));
 		Request->SetHeader(TEXT("Content-Type"), TEXT("text/plain; charset=us-ascii"));
@@ -489,7 +497,6 @@ void FCrashUploadToReceiver::PostReportComplete()
 	}
 
 	AssignReportIdToPostDataBuffer();
-
 	
 	auto Request = CreateHttpRequest();
 	Request->SetVerb( TEXT( "POST" ) );
@@ -500,6 +507,10 @@ void FCrashUploadToReceiver::PostReportComplete()
 
 	if (Request->ProcessRequest())
 	{
+#if PRIMARY_UPLOAD_RECEIVER
+		// completed upload to CRR so send analytics
+		FPrimaryCrashProperties::Get()->SendPostUploadAnalytics();
+#endif
 		SetCurrentState(EUploadState::PostingReportComplete);
 	}
 	else
@@ -722,6 +733,11 @@ void FCrashUploadToDataRouter::BeginUpload(const FPlatformErrorReport& PlatformE
 
 void FCrashUploadToDataRouter::CompressAndSendData()
 {
+#if PRIMARY_UPLOAD_DATAROUTER
+	// first stage of any upload to DR so send analytics
+	FPrimaryCrashProperties::Get()->SendPreUploadAnalytics();
+#endif
+
 	FCompressedHeader CompressedHeader;
 	CompressedHeader.DirectoryName = ErrorReport.GetReportDirectoryLeafName();
 	CompressedHeader.FileName = ErrorReport.GetReportDirectoryLeafName() + TEXT(".ue4crash");
@@ -755,6 +771,10 @@ void FCrashUploadToDataRouter::CompressAndSendData()
 
 	if (Request->ProcessRequest())
 	{
+#if PRIMARY_UPLOAD_DATAROUTER
+		// completed upload to DR so send analytics
+		FPrimaryCrashProperties::Get()->SendPostUploadAnalytics();
+#endif
 		return;
 	}
 	else
diff --git a/Engine/Source/Runtime/AIModule/Classes/Perception/AISenseConfig.h b/Engine/Source/Runtime/AIModule/Classes/Perception/AISenseConfig.h
--- a/Engine/Source/Runtime/AIModule/Classes/Perception/AISenseConfig.h
+++ b/Engine/Source/Runtime/AIModule/Classes/Perception/AISenseConfig.h
@@ -33,7 +33,8 @@ public:
 
 	FAISenseID GetSenseID() const;
 	
-	float GetMaxAge() const { return MaxAge == 0.f ? FAIStimulus::NeverHappenedAge : MaxAge; }
+	void SetMaxAge(const float NewMaxAge) { MaxAge = NewMaxAge; }
+	float GetMaxAge() const { return MaxAge == 0.f ? FAIStimulus::NeverHappenedAge : MaxAge; }	
 	bool IsEnabled() const { return bStartsEnabled; }
 
 #if !UE_BUILD_SHIPPING
diff --git a/Engine/Source/Runtime/Analytics/QoSReporter/Private/QoSReporterModule.cpp b/Engine/Source/Runtime/Analytics/QoSReporter/Private/QoSReporterModule.cpp
--- a/Engine/Source/Runtime/Analytics/QoSReporter/Private/QoSReporterModule.cpp
+++ b/Engine/Source/Runtime/Analytics/QoSReporter/Private/QoSReporterModule.cpp
@@ -18,7 +18,7 @@ DEFINE_LOG_CATEGORY(LogQoSReporter);
 IMPLEMENT_MODULE(FQoSReporterModule, QoSReporter);
 
 // helps to version QoS events (date*10 to allow for 10 revisions per day)
-#define QOS_EVENTS_REVISION					201602080
+#define QOS_EVENTS_REVISION					201602160
 
 FString FQoSReporterModule::Config::GetDefaultAppVersion()
 { 
@@ -295,7 +295,7 @@ void FAnalyticsProviderQoSReporter::EventRequestComplete(FHttpRequestPtr HttpReq
 	if (bCountTowardsFailedAttempts)
 	{
 		// FIXME: should use retrial with exponential backoff here
-#if USE_SERVER_PERF_COUNTERS
+#if USE_SERVER_PERF_COUNTERS && !IS_PROGRAM
 		PerfCountersIncrement(TEXT("FailedQoSRequests"), 0, IPerfCounters::Flags::Transient);
 #endif // USE_SERVER_PERF_COUNTERS
 	}
diff --git a/Engine/Source/Runtime/AnimGraphRuntime/Private/AnimNodes/AnimNode_SequenceEvaluator.cpp b/Engine/Source/Runtime/AnimGraphRuntime/Private/AnimNodes/AnimNode_SequenceEvaluator.cpp
--- a/Engine/Source/Runtime/AnimGraphRuntime/Private/AnimNodes/AnimNode_SequenceEvaluator.cpp
+++ b/Engine/Source/Runtime/AnimGraphRuntime/Private/AnimNodes/AnimNode_SequenceEvaluator.cpp
@@ -89,6 +89,6 @@ void FAnimNode_SequenceEvaluator::GatherDebugData(FNodeDebugData& DebugData)
 {
 	FString DebugLine = DebugData.GetNodeName(this);
 	
-	DebugLine += FString::Printf(TEXT("('%s' Play Time: %.3f)"), *GetNameSafe(Sequence), ExplicitTime);
+	DebugLine += FString::Printf(TEXT("('%s' InputTime: %.3f, Time: %.3f)"), *GetNameSafe(Sequence), ExplicitTime, InternalTimeAccumulator);
 	DebugData.AddDebugItem(DebugLine, true);
 }
\ No newline at end of file
diff --git a/Engine/Source/Runtime/Core/Private/GenericPlatform/GenericPlatformCrashContext.cpp b/Engine/Source/Runtime/Core/Private/GenericPlatform/GenericPlatformCrashContext.cpp
--- a/Engine/Source/Runtime/Core/Private/GenericPlatform/GenericPlatformCrashContext.cpp
+++ b/Engine/Source/Runtime/Core/Private/GenericPlatform/GenericPlatformCrashContext.cpp
@@ -119,6 +119,7 @@ void FGenericCrashContext::Initialize()
 }
 
 FGenericCrashContext::FGenericCrashContext()
+	: bIsEnsure(false)
 {
 	CommonBuffer.Reserve( 32768 );
 }
@@ -136,6 +137,7 @@ void FGenericCrashContext::SerializeContentToBuffer()
 	AddCrashProperty( TEXT( "IsInternalBuild" ), NCachedCrashContextProperties::bIsInternalBuild );
 	AddCrashProperty( TEXT( "IsPerforceBuild" ), NCachedCrashContextProperties::bIsPerforceBuild );
 	AddCrashProperty( TEXT( "IsSourceDistribution" ), NCachedCrashContextProperties::bIsSourceDistribution );
+	AddCrashProperty( TEXT( "IsEnsure" ), bIsEnsure );
 
 	AddCrashProperty( TEXT( "SecondsSinceStart" ), NCachedCrashContextProperties::SecondsSinceStart );
 
