Copying //UE4/Dev-Editor to Dev-Main (//UE4/Dev-Main)

==========================
MAJOR FEATURES + CHANGES
==========================

Change 2888098 on 3/1/2016 by [REDACTED_PERSON]

	Adding back the SetWidgetToFocusOnActivate call and deprecating it.  Will need to also do this in 4.11.

Change 2851669 on 2016/02/01 by [REDACTED_PERSON]

	Skeletal mesh import now support _skinXX that are not sequential.
	Static mesh was already supporting non sequential _skinxx ordering

Change 2851672 on 2016/02/01 by [REDACTED_PERSON]

	The proxy camera mesh is now properly reset to zero

Change 2851675 on 2016/02/01 by [REDACTED_PERSON]

	Update the tooltips

Change 2851764 on 2016/02/01 by [REDACTED_PERSON]

	The fbx plus and minus icons are now brighter

Change 2852116 on 2016/02/01 by [REDACTED_PERSON]

	#UE4 Submitting pull request 2013 from GitHub (Pull request 2013). Thanks for the fix hoelzl!

	Fix file selection when reimporting curve from moved CSV file
	When reimporting a data curve after moving the CSV file from which it was generated, the file selection dialog does not present an option to select CVS files. This patch fixes the issue by assigning the correct 'SupportedClass' value for `UReimportCurveFactory` instances.

	#2013

Change 2852375 on 2016/02/02 by [REDACTED_PERSON]

	Spline component improvements: added facility to not restore component instance cache after the construction script has run, so the points can act as inputs to the construction script. Created a new property bInputSplinePointsToConstructionScript for that.  Added SetUpVectorAtSplinePoint, and corrected some bugs.

Change 2852726 on 2016/02/02 by [REDACTED_PERSON]

	Fixed FPropertyChangedEvent::GetArrayIndex when called from PostEditChangeProperty.

Change 2853152 on 2016/02/02 by [REDACTED_PERSON]

	Fixed multi-line editable texts not updating their font when changed in UMG

	Also made all the SetStyle functions use the default if they're passed null (to match SEditableTextBox), and tidied up some of the property panel layout when editing styles.

Change 2853220 on 2016/02/02 by [REDACTED_PERSON]

	We now apply the scene option transform to the vertex of meshes instead of the root node of the scene. This allow re-alignment of the mesh to go with animation.

Change 2853282 on 2016/02/02 by [REDACTED_PERSON]

	Back out changelist 2853220

Change 2854211 on 2016/02/03 by [REDACTED_PERSON]

	Widget Reflector - Limit the minimum scale that can be applied to something more reasonable 50%, instead of 10%.

Change 2854216 on 2016/02/03 by [REDACTED_PERSON]

	Scene Viewport - The scene viewport handles application scale better now, allowing click locations to be interepreted correctly and transformed into pixel hit location, rather than local space widget location, which may not match.

Change 2854220 on 2016/02/03 by [REDACTED_PERSON]

	Slate - Allowing mousewheel or gesture to be routed directly for a widget path like has been done for other mouse events, this permits more kinds of mouse like actions in a VR environment onto widgets in the scene.  (not actually hooked up to do it, but this now permits it at the slate level to be done correctly).

Change 2854449 on 2016/02/03 by [REDACTED_PERSON]

	-Fix the fbx import options, for the scene transform value that was not apply correctly
	-Add an inspector on the scene import option in the reimport dialog

Change 2855659 on 2016/02/04 by [REDACTED_PERSON]

	-Fix the bake pivot when there is a hierarchy, we now accumulate the pivot effect in the hierarchy to place the object at the correct place.

Change 2855922 on 2016/02/04 by [REDACTED_PERSON]

	The animation is now align with the imported skeletal mesh, the bakepivot is also supported

Change 2856989 on 2016/02/05 by [REDACTED_PERSON]

	Some improvements to the source code loc gatherer

	* We now strip any comments out of any pre-processor tokens before we start handling them.
	* Fixed a case where "#if defined" and "#elif defined" would parse incorrectly.
	* Fixed a case where "#define LOCTEXT_NAMESPACE" and "#define LOC_DEFINE_REGION" may not be paired correctly with their corresponding "#undef".

diff --git a/Engine/Config/BaseEditorPerProjectUserSettings.ini b/Engine/Config/BaseEditorPerProjectUserSettings.ini
--- a/Engine/Config/BaseEditorPerProjectUserSettings.ini
+++ b/Engine/Config/BaseEditorPerProjectUserSettings.ini
@@ -153,6 +153,8 @@ bAutoSaveContent=True
 AutoSaveTimeMinutes=10
 ; The number of seconds warning before an autosave
 AutoSaveWarningInSeconds=10
+; Add the game directory by default
++AutoReimportDirectorySettings=(SourceDirectory="/Game/",MountPoint=,Wildcards=((Wildcard="Localization/*")))
 
 ; Whether to automatically perform an SCC checkout on asset modification
 bPromptForCheckoutOnAssetModification=False
diff --git a/Engine/Plugins/2D/Paper2D/Source/Paper2D/Classes/PaperTileLayer.h b/Engine/Plugins/2D/Paper2D/Source/Paper2D/Classes/PaperTileLayer.h
--- a/Engine/Plugins/2D/Paper2D/Source/Paper2D/Classes/PaperTileLayer.h
+++ b/Engine/Plugins/2D/Paper2D/Source/Paper2D/Classes/PaperTileLayer.h
@@ -134,11 +134,11 @@ private:
 	uint32 bLayerCollides:1;
 
 	// Should this layer use a custom thickness for generated collision instead of the thickness setting in the tile map?
-	UPROPERTY(BlueprintReadOnly, Category=Sprite, meta=(AllowPrivateAccess="true"))
+	UPROPERTY(BlueprintReadOnly, EditAnywhere, Category=Sprite, meta=(AllowPrivateAccess="true", InlineEditConditionToggle))
 	uint32 bOverrideCollisionThickness:1;
 
 	// Should this layer use a custom offset for generated collision instead of the layer drawing offset?
-	UPROPERTY(BlueprintReadOnly, Category=Sprite, meta=(AllowPrivateAccess="true"))
+	UPROPERTY(BlueprintReadOnly, EditAnywhere, Category=Sprite, meta=(AllowPrivateAccess="true", InlineEditConditionToggle))
 	uint32 bOverrideCollisionOffset:1;
 
 	// The override thickness of the collision for this layer (when bOverrideCollisionThickness is set)
diff --git a/Engine/Plugins/2D/Paper2D/Source/Paper2D/Classes/PaperTileMapComponent.h b/Engine/Plugins/2D/Paper2D/Source/Paper2D/Classes/PaperTileMapComponent.h
--- a/Engine/Plugins/2D/Paper2D/Source/Paper2D/Classes/PaperTileMapComponent.h
+++ b/Engine/Plugins/2D/Paper2D/Source/Paper2D/Classes/PaperTileMapComponent.h
@@ -64,7 +64,7 @@ private:
 	int32 UseSingleLayerIndex;
 
 	// Should we draw a single layer?
-	UPROPERTY()
+	UPROPERTY(EditAnywhere, Category=Rendering, meta=(InlineEditConditionToggle))
 	bool bUseSingleLayer;
 
 #if WITH_EDITOR
diff --git a/Engine/Plugins/Blendables/LightPropagationVolume/Source/LightPropagationVolumeRuntime/Public/LightPropagationVolumeBlendable.h b/Engine/Plugins/Blendables/LightPropagationVolume/Source/LightPropagationVolumeRuntime/Public/LightPropagationVolumeBlendable.h
--- a/Engine/Plugins/Blendables/LightPropagationVolume/Source/LightPropagationVolumeRuntime/Public/LightPropagationVolumeBlendable.h
+++ b/Engine/Plugins/Blendables/LightPropagationVolume/Source/LightPropagationVolumeRuntime/Public/LightPropagationVolumeBlendable.h
@@ -11,43 +11,43 @@ struct FLightPropagationVolumeSettings
 {
 	GENERATED_USTRUCT_BODY()
 
-	UPROPERTY(BlueprintReadWrite, Category=Overrides, meta=(PinHiddenByDefault))
+	UPROPERTY(BlueprintReadWrite, EditAnywhere, Category=Overrides, meta=(PinHiddenByDefault, InlineEditConditionToggle))
 	uint32 bOverride_LPVIntensity:1;
 
-	UPROPERTY(BlueprintReadWrite, Category=Overrides, meta=(PinHiddenByDefault))
+	UPROPERTY(BlueprintReadWrite, EditAnywhere, Category = Overrides, meta = (PinHiddenByDefault, InlineEditConditionToggle))
 	uint32 bOverride_LPVDirectionalOcclusionIntensity:1;
 
-	UPROPERTY(BlueprintReadWrite, Category=Overrides, meta=(PinHiddenByDefault))
+	UPROPERTY(BlueprintReadWrite, EditAnywhere, Category = Overrides, meta = (PinHiddenByDefault, InlineEditConditionToggle))
 	uint32 bOverride_LPVDirectionalOcclusionRadius:1;
 
-	UPROPERTY(BlueprintReadWrite, Category=Overrides, meta=(PinHiddenByDefault))
+	UPROPERTY(BlueprintReadWrite, EditAnywhere, Category = Overrides, meta = (PinHiddenByDefault, InlineEditConditionToggle))
 	uint32 bOverride_LPVDiffuseOcclusionExponent:1;
 
-	UPROPERTY(BlueprintReadWrite, Category=Overrides, meta=(PinHiddenByDefault))
+	UPROPERTY(BlueprintReadWrite, EditAnywhere, Category = Overrides, meta = (PinHiddenByDefault, InlineEditConditionToggle))
 	uint32 bOverride_LPVSpecularOcclusionExponent:1;
 
-	UPROPERTY(BlueprintReadWrite, Category=Overrides, meta=(PinHiddenByDefault))
+	UPROPERTY(BlueprintReadWrite, EditAnywhere, Category = Overrides, meta = (PinHiddenByDefault, InlineEditConditionToggle))
 	uint32 bOverride_LPVDiffuseOcclusionIntensity:1;
 
-	UPROPERTY(BlueprintReadWrite, Category=Overrides, meta=(PinHiddenByDefault))
+	UPROPERTY(BlueprintReadWrite, EditAnywhere, Category = Overrides, meta = (PinHiddenByDefault, InlineEditConditionToggle))
 	uint32 bOverride_LPVSpecularOcclusionIntensity:1;
 
-	UPROPERTY(BlueprintReadWrite, Category=Overrides, meta=(PinHiddenByDefault))
+	UPROPERTY(BlueprintReadWrite, EditAnywhere, Category = Overrides, meta = (PinHiddenByDefault, InlineEditConditionToggle))
 	uint32 bOverride_LPVSize:1;
 
-	UPROPERTY(BlueprintReadWrite, Category=Overrides, meta=(PinHiddenByDefault))
+	UPROPERTY(BlueprintReadWrite, EditAnywhere, Category = Overrides, meta = (PinHiddenByDefault, InlineEditConditionToggle))
 	uint32 bOverride_LPVSecondaryOcclusionIntensity:1;
 
-	UPROPERTY(BlueprintReadWrite, Category=Overrides, meta=(PinHiddenByDefault))
+	UPROPERTY(BlueprintReadWrite, EditAnywhere, Category = Overrides, meta = (PinHiddenByDefault, InlineEditConditionToggle))
 	uint32 bOverride_LPVSecondaryBounceIntensity:1;
 
-	UPROPERTY(BlueprintReadWrite, Category=Overrides, meta=(PinHiddenByDefault))
+	UPROPERTY(BlueprintReadWrite, EditAnywhere, Category = Overrides, meta = (PinHiddenByDefault, InlineEditConditionToggle))
 	uint32 bOverride_LPVGeometryVolumeBias:1;
 
-	UPROPERTY(BlueprintReadWrite, Category=Overrides, meta=(PinHiddenByDefault))
+	UPROPERTY(BlueprintReadWrite, EditAnywhere, Category = Overrides, meta = (PinHiddenByDefault, InlineEditConditionToggle))
 	uint32 bOverride_LPVVplInjectionBias:1;
 
-	UPROPERTY(BlueprintReadWrite, Category=Overrides, meta=(PinHiddenByDefault))
+	UPROPERTY(BlueprintReadWrite, EditAnywhere, Category = Overrides, meta = (PinHiddenByDefault, InlineEditConditionToggle))
 	uint32 bOverride_LPVEmissiveInjectionIntensity:1;
 
 	/** How strong the dynamic GI from the LPV should be. 0.0 is off, 1.0 is the "normal" value, but higher values can be used to boost the effect*/
diff --git a/Engine/Plugins/Experimental/CodeEditor/Source/CodeEditor/Private/SCodeEditableText.cpp b/Engine/Plugins/Experimental/CodeEditor/Source/CodeEditor/Private/SCodeEditableText.cpp
--- a/Engine/Plugins/Experimental/CodeEditor/Source/CodeEditor/Private/SCodeEditableText.cpp
+++ b/Engine/Plugins/Experimental/CodeEditor/Source/CodeEditor/Private/SCodeEditableText.cpp
@@ -27,7 +27,7 @@ FReply SCodeEditableText::OnKeyChar(const FGeometry& MyGeometry, const FCharacte
 	const TCHAR Character = InCharacterEvent.GetCharacter();
 	if(Character == TEXT('\t'))
 	{
-		if (!GetIsReadOnly())
+		if (!IsTextReadOnly())
 		{
 			FString String;
 			String.AppendChar(Character);
diff --git a/Engine/Plugins/MovieScene/LevelSequenceEditor/Source/LevelSequenceEditor/Private/CinematicLevelViewportLayout.cpp b/Engine/Plugins/MovieScene/LevelSequenceEditor/Source/LevelSequenceEditor/Private/CinematicLevelViewportLayout.cpp
--- a/Engine/Plugins/MovieScene/LevelSequenceEditor/Source/LevelSequenceEditor/Private/CinematicLevelViewportLayout.cpp
+++ b/Engine/Plugins/MovieScene/LevelSequenceEditor/Source/LevelSequenceEditor/Private/CinematicLevelViewportLayout.cpp
@@ -126,7 +126,7 @@ TSharedRef<SWidget> FCinematicLevelViewportLayout_TwoPane::MakeViewportLayout(co
 	// Make newly-created perspective viewports active by default
 	GCurrentLevelEditingViewportClient = &ViewportWidget0->GetLevelViewportClient();
 
-	InitCommonLayoutFromString(SpecificLayoutString);
+	InitCommonLayoutFromString(SpecificLayoutString, NAME_None);
 
 	return SplitterWidget.ToSharedRef();
 }
diff --git a/Engine/Source/Developer/AssetTools/Private/AssetTools.h b/Engine/Source/Developer/AssetTools/Private/AssetTools.h
--- a/Engine/Source/Developer/AssetTools/Private/AssetTools.h
+++ b/Engine/Source/Developer/AssetTools/Private/AssetTools.h
@@ -38,6 +38,7 @@ public:
 	virtual void MigratePackages(const TArray<FName>& PackageNamesToMigrate) const override;
 	virtual void FixupReferencers(const TArray<UObjectRedirector*>& Objects) const override;
 	virtual FAssetPostRenameEvent& OnAssetPostRename() override { return AssetRenameManager->OnAssetPostRenameEvent(); }
+	virtual void ExpandDirectories(const TArray<FString>& Files, const FString& DestinationPath, TArray<TPair<FString, FString>>& FilesAndDestinations) const override;
 
 public:
 	/** Gets the asset tools singleton as a FAssetTools for asset tools module use */
@@ -46,11 +47,7 @@ public:
 	/** Syncs the primary content browser to the specified assets, whether or not it is locked. Most syncs that come from AssetTools -feel- like they came from the content browser, so this is okay. */
 	void SyncBrowserToAssets(const TArray<UObject*>& AssetsToSync);
 	void SyncBrowserToAssets(const TArray<FAssetData>& AssetsToSync);
-
 private:
-	/** Expands any folders found in the files list, and returns a flattened list of destination paths and files.  Mirrors directory structure. */
-	void ExpandDirectories(const TArray<FString>& Files, const FString& DestinationPath, TArray<TPair<FString, FString>>& FilesAndDestinations) const;
-
 	/** Checks to see if a package is marked for delete then ask the user if he would like to check in the deleted file before he can continue. Returns true when it is safe to proceed. */
 	bool CheckForDeletedPackage(const UPackage* Package) const;
 
diff --git a/Engine/Source/Developer/AssetTools/Public/IAssetTools.h b/Engine/Source/Developer/AssetTools/Public/IAssetTools.h
--- a/Engine/Source/Developer/AssetTools/Public/IAssetTools.h
+++ b/Engine/Source/Developer/AssetTools/Public/IAssetTools.h
@@ -174,4 +174,7 @@ public:
 
 	/** Fix up references to the specified redirectors */
 	virtual void FixupReferencers(const TArray<UObjectRedirector*>& Objects) const = 0;
+
+	/** Expands any folders found in the files list, and returns a flattened list of destination paths and files.  Mirrors directory structure. */
+	virtual void ExpandDirectories(const TArray<FString>& Files, const FString& DestinationPath, TArray<TPair<FString, FString>>& FilesAndDestinations) const = 0;
 };
diff --git a/Engine/Source/Developer/SlateReflector/Private/Widgets/SWidgetReflector.cpp b/Engine/Source/Developer/SlateReflector/Private/Widgets/SWidgetReflector.cpp
--- a/Engine/Source/Developer/SlateReflector/Private/Widgets/SWidgetReflector.cpp
+++ b/Engine/Source/Developer/SlateReflector/Private/Widgets/SWidgetReflector.cpp
@@ -470,7 +470,7 @@ void SWidgetReflector::Construct( const FArguments& InArgs )
 					[
 						SNew(SSpinBox<float>)
 						.Value(this, &SWidgetReflector::HandleAppScaleSliderValue)
-						.MinValue(0.1f)
+						.MinValue(0.50f)
 						.MaxValue(3.0f)
 						.Delta(0.01f)
 						.OnValueChanged(this, &SWidgetReflector::HandleAppScaleSliderChanged)
diff --git a/Engine/Source/Editor/ClassViewer/Private/SClassViewer.cpp b/Engine/Source/Editor/ClassViewer/Private/SClassViewer.cpp
--- a/Engine/Source/Editor/ClassViewer/Private/SClassViewer.cpp
+++ b/Engine/Source/Editor/ClassViewer/Private/SClassViewer.cpp
@@ -1214,7 +1214,7 @@ public:
 		, _bIsPlaceable(false)
 		, _bIsInClassViewer( true )
 		, _bDynamicClassLoading( true )
-		, _HighlightText(NULL)
+		, _HighlightText()
 		, _TextColor(FLinearColor(1.0f, 1.0f, 1.0f, 1.0f))
 		{}
 
@@ -1227,7 +1227,7 @@ public:
 		/** true if this item should allow dynamic class loading */
 		SLATE_ARGUMENT( bool, bDynamicClassLoading )
 		/** The text this item should highlight, if any. */
-		SLATE_ARGUMENT( const FText*, HighlightText )
+		SLATE_ARGUMENT( FText, HighlightText )
 		/** The color text this item will use. */
 		SLATE_ARGUMENT( FSlateColor, TextColor )
 		/** The node this item is associated with. */
@@ -1315,7 +1315,7 @@ public:
 				[
 					SNew( STextBlock )
 						.Text( FText::FromString(*ClassName.Get()) )
-						.HighlightText(*InArgs._HighlightText)
+						.HighlightText(InArgs._HighlightText)
 						.ColorAndOpacity( this, &SClassItem::GetTextColor)
 						.ToolTip(Local::GetToolTip(AssociatedNode))
 						.IsEnabled(!bIsRestricted)
@@ -2274,7 +2274,7 @@ TSharedRef< ITableRow > SClassViewer::OnGenerateRowForClassViewer( TSharedPtr<FC
 	TSharedRef< SClassItem > ReturnRow = SNew(SClassItem, OwnerTable)
 		.ClassName(Item->GetClassName(InitOptions.bShowDisplayNames))
 		.bIsPlaceable(Item->IsClassPlaceable())
-		.HighlightText(&SearchBox->GetText())
+		.HighlightText(SearchBox->GetText())
 		.TextColor(Item->IsClassPlaceable()? FLinearColor(0.2f, 0.4f, 0.6f, AlphaValue) : FLinearColor(1.0f, 1.0f, 1.0f, AlphaValue))
 		.AssociatedNode(Item)
 		.bIsInClassViewer( InitOptions.Mode == EClassViewerMode::ClassBrowsing )
diff --git a/Engine/Source/Editor/ComponentVisualizers/Private/SplineComponentVisualizer.cpp b/Engine/Source/Editor/ComponentVisualizers/Private/SplineComponentVisualizer.cpp
--- a/Engine/Source/Editor/ComponentVisualizers/Private/SplineComponentVisualizer.cpp
+++ b/Engine/Source/Editor/ComponentVisualizers/Private/SplineComponentVisualizer.cpp
@@ -172,7 +172,7 @@ void FSplineComponentVisualizer::DrawVisualization(const UActorComponent* Compon
 		const USplineComponent* EditedSplineComp = GetEditedSplineComponent();
 
 		const USplineComponent* Archetype = CastChecked<USplineComponent>(SplineComp->GetArchetype());
-		const bool bIsSplineEditable = SplineComp->bSplineHasBeenEdited || SplineComp->SplineInfo.Points == Archetype->SplineInfo.Points;
+		const bool bIsSplineEditable = SplineComp->bSplineHasBeenEdited || SplineComp->SplineInfo.Points == Archetype->SplineInfo.Points || SplineComp->bInputSplinePointsToConstructionScript;
 
 		const FColor ReadOnlyColor = FColor(255, 0, 255, 255);
 		const FColor NormalColor = bIsSplineEditable ? FColor(SplineComp->EditorUnselectedSplineSegmentColor.ToFColor(true)) : ReadOnlyColor;
@@ -313,6 +313,12 @@ void FSplineComponentVisualizer::DrawVisualization(const UActorComponent* Compon
 						{
 							PDI->DrawLine(OldPos - OldRightVector * OldScale.Y, NewPos - NewRightVector * NewScale.Y, LineColor, SDPG_Foreground);
 							PDI->DrawLine(OldPos + OldRightVector * OldScale.Y, NewPos + NewRightVector * NewScale.Y, LineColor, SDPG_Foreground);
+
+							#if VISUALIZE_SPLINE_UPVECTORS
+							const FVector NewUpVector = SplineComp->GetUpVectorAtSplineInputKey(Key, ESplineCoordinateSpace::World);
+							PDI->DrawLine(NewPos, NewPos + NewUpVector * SplineComp->ScaleVisualizationWidth * 0.5f, LineColor, SDPG_Foreground);
+							PDI->DrawLine(NewPos, NewPos + NewRightVector * SplineComp->ScaleVisualizationWidth * 0.5f, LineColor, SDPG_Foreground);
+							#endif
 						}
 
 						OldPos = NewPos;
diff --git a/Engine/Source/Editor/ContentBrowser/Private/SAssetView.cpp b/Engine/Source/Editor/ContentBrowser/Private/SAssetView.cpp
--- a/Engine/Source/Editor/ContentBrowser/Private/SAssetView.cpp
+++ b/Engine/Source/Editor/ContentBrowser/Private/SAssetView.cpp
@@ -977,7 +977,7 @@ void SAssetView::Tick( const FGeometry& AllottedGeometry, const double InCurrent
 			AssetAwaitingRename->bRenameWhenScrolledIntoview = false;
 			AwaitingRename = nullptr;
 		}
-		else if (FSlateApplication::Get().HasFocusedDescendants(OwnerWindow.ToSharedRef()))
+		else if (OwnerWindow->HasAnyUserFocusOrFocusedDescendants())
 		{
 			AssetAwaitingRename->RenamedRequestEvent.ExecuteIfBound();
 			AssetAwaitingRename->bRenameWhenScrolledIntoview = false;
@@ -1205,8 +1205,72 @@ FReply SAssetView::OnDrop( const FGeometry& MyGeometry, const FDragDropEvent& Dr
 		{
 			if ( ExternalDragDropOp->HasFiles() )
 			{
+				TArray<FString> ImportFiles;
+				TMap<FString, UObject*> ReimportFiles;
 				FAssetToolsModule& AssetToolsModule = FModuleManager::Get().LoadModuleChecked<FAssetToolsModule>("AssetTools");
-				AssetToolsModule.Get().ImportAssets( ExternalDragDropOp->GetFiles(), SourcesData.PackagePaths[0].ToString() );
+				FString RootDestinationPath = SourcesData.PackagePaths[0].ToString();
+				TArray<TPair<FString, FString>> FilesAndDestinations;
+				const TArray<FString>& DragFiles = ExternalDragDropOp->GetFiles();
+				AssetToolsModule.Get().ExpandDirectories(DragFiles, RootDestinationPath, FilesAndDestinations);
+
+				for (int32 FileIdx = 0; FileIdx < FilesAndDestinations.Num(); ++FileIdx)
+				{
+					const FString& Filename = FilesAndDestinations[FileIdx].Key;
+					const FString& DestinationPath = FilesAndDestinations[FileIdx].Value;
+					FString Name = ObjectTools::SanitizeObjectName(FPaths::GetBaseFilename(Filename));
+					FString PackageName = DestinationPath + TEXT("/") + Name;
+
+					// We can not create assets that share the name of a map file in the same location
+					if (FEditorFileUtils::IsMapPackageAsset(PackageName))
+					{
+						//The error message will be log in the import process
+						ImportFiles.Add(Filename);
+						continue;
+					}
+					//Check if package exist in memory
+					UPackage* Pkg = FindPackage(nullptr, *PackageName);
+					bool IsPkgExist = Pkg != nullptr;
+					//check if package exist on file
+					if (!IsPkgExist && !FPackageName::DoesPackageExist(PackageName))
+					{
+						ImportFiles.Add(Filename);
+						continue;
+					}
+					if (Pkg == nullptr)
+					{
+						Pkg = CreatePackage(nullptr, *PackageName);
+						if (Pkg == nullptr)
+						{
+							//Cannot create a package that don't exist on disk or in memory!!!
+							//The error message will be log in the import process
+							ImportFiles.Add(Filename);
+							continue;
+						}
+					}
+					// Make sure the destination package is loaded
+					Pkg->FullyLoad();
+
+					// Check for an existing object
+					UObject* ExistingObject = StaticFindObject(UObject::StaticClass(), Pkg, *Name);
+					if (ExistingObject != nullptr)
+					{
+						ReimportFiles.Add(Filename, ExistingObject);
+					}
+					else
+					{
+						ImportFiles.Add(Filename);
+					}
+				}
+				//Reimport
+				for (auto kvp : ReimportFiles)
+				{
+					FReimportManager::Instance()->Reimport(kvp.Value, false, true, kvp.Key);
+				}
+				//Import
+				if (ImportFiles.Num() > 0)
+				{
+					AssetToolsModule.Get().ImportAssets(ImportFiles, SourcesData.PackagePaths[0].ToString());
+				}
 			}
 
 			return FReply::Handled();
diff --git a/Engine/Source/Editor/FontEditor/Private/SFontEditorViewport.cpp b/Engine/Source/Editor/FontEditor/Private/SFontEditorViewport.cpp
--- a/Engine/Source/Editor/FontEditor/Private/SFontEditorViewport.cpp
+++ b/Engine/Source/Editor/FontEditor/Private/SFontEditorViewport.cpp
@@ -182,15 +182,14 @@ void FFontEditorViewportClient::Draw(FViewport* Viewport, FCanvas* Canvas)
 		static const FVector2D StartPos(4.0f, 4.0f);
 
 		// And draw the text with the foreground color
-		if (Font->FontCacheType == EFontCacheType::Runtime && Font->CompositeFont.DefaultTypeface.Fonts.Num() > 0)
+		if (Font->FontCacheType == EFontCacheType::Runtime)
 		{
 			static const float FontScale = 1.0f;
 
 			TSharedRef<FSlateFontCache> FontCache = FSlateApplication::Get().GetRenderer()->GetFontCache();
-			TSharedRef<FSlateFontMeasure> FontMeasure = FSlateApplication::Get().GetRenderer()->GetFontMeasureService();
 
 			FVector2D CurPos = StartPos;
-			float WidestName = 0.0f;
+			int32 WidestName = 0;
 
 			// Draw and measure each name so we can work out where to start drawing the preview text column
 			for (const FTypefaceEntry& TypefaceEntry : Font->CompositeFont.DefaultTypeface.Fonts)
@@ -198,17 +197,16 @@ void FFontEditorViewportClient::Draw(FViewport* Viewport, FCanvas* Canvas)
 				const FSlateFontInfo FontInfo(Font, Font->LegacyFontSize, TypefaceEntry.Name);
 
 				FCharacterList& CharacterList = FontCache->GetCharacterList(FontInfo, FontScale);
-				const int32 MaxCharHeight = CharacterList.GetMaxHeight();
 
-				const FText EntryNameText = FText::FromName(TypefaceEntry.Name);
+				FShapedGlyphSequenceRef EntryNameShapedText = FontCache->ShapeBidirectionalText(TypefaceEntry.Name.ToString(), FontInfo, FontScale, TextBiDi::ETextDirection::LeftToRight, ETextShapingMethod::Auto);
+				
+				FCanvasShapedTextItem ShapedTextItem(CurPos, EntryNameShapedText, FLinearColor(ForegroundColor));
+				Canvas->DrawItem(ShapedTextItem);
 
-				FCanvasTextItem TextItem(CurPos, EntryNameText, FontInfo, FLinearColor(ForegroundColor));
-				Canvas->DrawItem(TextItem);
+				const FVector2D MeasuredText = ShapedTextItem.DrawnSize;
+				WidestName = FMath::Max(WidestName, EntryNameShapedText->GetMeasuredWidth());
 
-				const FVector2D MeasuredText = FontMeasure->Measure(EntryNameText, FontInfo, FontScale);
-				WidestName = FMath::Max(WidestName, MeasuredText.X);
-
-				CurPos.Y += MaxCharHeight;
+				CurPos.Y += CharacterList.GetMaxHeight() + 8.0f;
 			}
 
 			CurPos = FVector2D(WidestName + 12.0f, StartPos.Y);
@@ -218,73 +216,130 @@ void FFontEditorViewportClient::Draw(FViewport* Viewport, FCanvas* Canvas)
 			{
 				const FSlateFontInfo FontInfo(Font, Font->LegacyFontSize, TypefaceEntry.Name);
 
-				FCharacterList& CharacterList = FontCache->GetCharacterList(FontInfo, FontScale);
-				const int32 MaxCharHeight = CharacterList.GetMaxHeight();
+				FShapedGlyphSequenceRef ShapedPreviewText = FontCache->ShapeBidirectionalText(PreviewText.ToString(), FontInfo, FontScale, TextBiDi::ETextDirection::LeftToRight, ETextShapingMethod::Auto);
 
-				FCanvasTextItem TextItem(CurPos, PreviewText, FontInfo, FLinearColor(ForegroundColor));
-				Canvas->DrawItem(TextItem);
+				FCanvasShapedTextItem ShapedTextItem(CurPos, ShapedPreviewText, FLinearColor(ForegroundColor));
+				Canvas->DrawItem(ShapedTextItem);
 
 				if (bDrawFontMetrics)
 				{
-					const FVector2D MeasuredText = FontMeasure->Measure(PreviewText, FontInfo, FontScale);
-
-					// Draw the bounding box for the actual characters
+					// Draw the bounding box for the separate glyphs
 					{
 						float LineX = 0.0f;
-						FCharacterEntry PreviousCharEntry;
-
-						for (const TCHAR Char : PreviewText.ToString())
+						for (const auto& GlyphToRender : ShapedPreviewText->GetGlyphsToRender())
 						{
-							const FCharacterEntry Entry = CharacterList.GetCharacter(Char, FontInfo.FontFallback);
+							if (GlyphToRender.bIsVisible)
+							{
+								const FShapedGlyphFontAtlasData GlyphAtlasData = FontCache->GetShapedGlyphFontAtlasData(GlyphToRender);
 
-							const bool bIsWhitespace = FChar::IsWhitespace(Char);
+								const float X = CurPos.X + LineX + GlyphAtlasData.HorizontalOffset + GlyphToRender.XOffset;
+								const float Y = CurPos.Y - GlyphAtlasData.VerticalOffset + GlyphToRender.YOffset + ShapedPreviewText->GetTextBaseline() + ShapedPreviewText->GetMaxTextHeight();
 
-							int8 Kerning = 0;
-							if (!bIsWhitespace && PreviousCharEntry.IsCached())
-							{
-								Kerning = CharacterList.GetKerning(PreviousCharEntry, Entry);
+								FCanvasBoxItem BoundingBoxItem(FVector2D(X, Y), FVector2D(GlyphAtlasData.USize, GlyphAtlasData.VSize));
+								BoundingBoxItem.SetColor(FColor::Orange);
+								Canvas->DrawItem(BoundingBoxItem);
 							}
 
-							LineX += Kerning;
-							PreviousCharEntry = Entry;
+							LineX += GlyphToRender.XAdvance;
+						}
+					}
 
-							if (!bIsWhitespace)
+					// Draw the bounding box for the logical glyphs
+					{
+						float LineX = 0.0f;
+						const auto& GlyphsToRender = ShapedPreviewText->GetGlyphsToRender();
+						for (int32 CurrentGlyphIndex = 0; CurrentGlyphIndex < GlyphsToRender.Num(); ++CurrentGlyphIndex)
+						{
+							// A single character may produce multiple glyphs which must be treated as a single logic unit
+							int16 GlyphClusterAdvance = 0;
+							FBox2D GlyphClusterBounds(ForceInitToZero);
+							for (;;)
 							{
-								const float X = CurPos.X + LineX + Entry.HorizontalOffset;
-								const float Y = CurPos.Y - Entry.VerticalOffset + Entry.GlobalDescender + MaxCharHeight;
+								const auto& GlyphToRender = GlyphsToRender[CurrentGlyphIndex];
+								const bool bIsWithinGlyphCluster = GlyphsToRender.IsValidIndex(CurrentGlyphIndex + 1) && GlyphToRender.ClusterIndex == GlyphsToRender[CurrentGlyphIndex + 1].ClusterIndex;
 
-								FCanvasBoxItem BoundingBoxItem(FVector2D(X, Y), FVector2D(Entry.USize, Entry.VSize));
-								BoundingBoxItem.SetColor(FLinearColor::Yellow);
-								Canvas->DrawItem(BoundingBoxItem);
+								if (GlyphToRender.bIsVisible)
+								{
+									const FShapedGlyphFontAtlasData GlyphAtlasData = FontCache->GetShapedGlyphFontAtlasData(GlyphToRender);
+
+									const float X = CurPos.X + LineX + GlyphAtlasData.HorizontalOffset + GlyphToRender.XOffset;
+									const float Y = CurPos.Y - GlyphAtlasData.VerticalOffset + GlyphToRender.YOffset + ShapedPreviewText->GetTextBaseline() + ShapedPreviewText->GetMaxTextHeight();
 
-								FCanvasLineItem BaseLineItem(FVector2D(CurPos.X, CurPos.Y + CharacterList.GetBaseline()), FVector2D(CurPos.X + MeasuredText.X, CurPos.Y + CharacterList.GetBaseline()));
-								BaseLineItem.SetColor(FLinearColor::Red);
-								Canvas->DrawItem(BaseLineItem);
+									GlyphClusterBounds += FBox2D(FVector2D(X, Y), FVector2D(GlyphAtlasData.USize, GlyphAtlasData.VSize));
+								}
+
+								GlyphClusterAdvance += GlyphToRender.XAdvance;
+
+								if (!bIsWithinGlyphCluster)
+								{
+									break;
+								}
+
+								++CurrentGlyphIndex;
 							}
 
-							LineX += Entry.XAdvance;
+							FCanvasBoxItem BoundingBoxItem(GlyphClusterBounds.Min, GlyphClusterBounds.Max);
+							BoundingBoxItem.SetColor(FLinearColor::Yellow);
+							Canvas->DrawItem(BoundingBoxItem);
+
+							LineX += GlyphClusterAdvance;
 						}
 					}
 
 					// Draw the baseline
 					{
-						const FCharacterEntry Entry = CharacterList.GetCharacter(0, FontInfo.FontFallback);
-						const float Y = CurPos.Y /*- Entry.VerticalOffset*/ + Entry.GlobalDescender + MaxCharHeight;
+						const float Y = CurPos.Y + ShapedPreviewText->GetTextBaseline() + ShapedPreviewText->GetMaxTextHeight();
 
-						FCanvasLineItem BaseLineItem(FVector2D(CurPos.X, Y), FVector2D(CurPos.X + MeasuredText.X, Y));
+						FCanvasLineItem BaseLineItem(FVector2D(CurPos.X, Y), FVector2D(CurPos.X + ShapedPreviewText->GetMeasuredWidth(), Y));
 						BaseLineItem.SetColor(FLinearColor::Red);
 						Canvas->DrawItem(BaseLineItem);
 					}
 
 					// Draw the bounding box for the line height
 					{
-						FCanvasBoxItem LineHeightBoxItem(CurPos, FVector2D(MeasuredText.X, MaxCharHeight));
+						FCanvasBoxItem LineHeightBoxItem(CurPos, FVector2D(ShapedPreviewText->GetMeasuredWidth(), ShapedPreviewText->GetMaxTextHeight()));
 						LineHeightBoxItem.SetColor(FLinearColor::Green);
 						Canvas->DrawItem(LineHeightBoxItem);
 					}
 				}
 
-				CurPos.Y += MaxCharHeight;
+				CurPos.Y += ShapedPreviewText->GetMaxTextHeight() + 8.0f;
+			}
+
+			// Draw the key
+			if (bDrawFontMetrics)
+			{
+				const FSlateFontInfo FontInfo = FEditorStyle::GetFontStyle("NormalFont");
+				const float KeyBoxSize = 14.0f;
+
+				struct FKeyDataType
+				{
+					FLinearColor KeyColor;
+					FText KeyText;
+				};
+
+				static const FKeyDataType KeyDataArray[] = {
+					{ FLinearColor::Red,	LOCTEXT("BaselineKeyLabel", "Baseline") },
+					{ FLinearColor::Green,	LOCTEXT("LineBoundsKeyLabel", "Line Bounds") },
+					{ FLinearColor::Yellow,	LOCTEXT("LogicalGlyphBoundsKeyLabel", "Logical Glyph Bounds") },
+					{ FColor::Orange,		LOCTEXT("IsolatedGlyphBoundsKeyLabel", "Isolated Glyph Bounds") },
+				};
+
+				for (const FKeyDataType& KeyData : KeyDataArray)
+				{
+					FCanvasBoxItem KeyBox(CurPos + (KeyBoxSize * 0.25f), FVector2D(KeyBoxSize * 0.5f, KeyBoxSize * 0.5f));
+					KeyBox.SetColor(KeyData.KeyColor);
+					KeyBox.LineThickness = KeyBoxSize * 0.5f;
+					Canvas->DrawItem(KeyBox);
+
+					CurPos.X += KeyBoxSize + 4.0f;
+
+					FShapedGlyphSequenceRef KeyLabelShapedText = FontCache->ShapeBidirectionalText(*KeyData.KeyText.ToString(), FontInfo, FontScale, TextBiDi::ETextDirection::LeftToRight, ETextShapingMethod::Auto);
+					FCanvasShapedTextItem ShapedTextItem(CurPos, KeyLabelShapedText, FLinearColor(ForegroundColor));
+					Canvas->DrawItem(ShapedTextItem);
+
+					CurPos.X += KeyLabelShapedText->GetMeasuredWidth() + 8.0f;
+				}
 			}
 		}
 		else
diff --git a/Engine/Source/Editor/GameplayAbilitiesEditor/Private/SGameplayAttributeWidget.cpp b/Engine/Source/Editor/GameplayAbilitiesEditor/Private/SGameplayAttributeWidget.cpp
--- a/Engine/Source/Editor/GameplayAbilitiesEditor/Private/SGameplayAttributeWidget.cpp
+++ b/Engine/Source/Editor/GameplayAbilitiesEditor/Private/SGameplayAttributeWidget.cpp
@@ -35,12 +35,12 @@ class SAttributeItem : public SComboRow< TSharedPtr<FAttributeViewerNode> >
 public:
 
 	SLATE_BEGIN_ARGS(SAttributeItem)
-		: _HighlightText(nullptr)
+		: _HighlightText()
 		, _TextColor(FLinearColor(1.0f, 1.0f, 1.0f, 1.0f))
 	{}
 
 	/** The text this item should highlight, if any. */
-	SLATE_ARGUMENT(const FText*, HighlightText)
+	SLATE_ARGUMENT(FText, HighlightText)
 	/** The color text this item will use. */
 	SLATE_ARGUMENT(FSlateColor, TextColor)
 	/** The node this item is associated with. */
@@ -68,7 +68,7 @@ public:
 					[
 						SNew(STextBlock)
 						.Text(FText::FromString(*AssociatedNode->AttributeName.Get()))
-						.HighlightText(*InArgs._HighlightText)
+						.HighlightText(InArgs._HighlightText)
 						.ColorAndOpacity(this, &SAttributeItem::GetTextColor)
 						.IsEnabled(true)
 					]
@@ -237,7 +237,7 @@ void SAttributeListWidget::Construct(const FArguments& InArgs)
 TSharedRef< ITableRow > SAttributeListWidget::OnGenerateRowForAttributeViewer(TSharedPtr<FAttributeViewerNode> Item, const TSharedRef< STableViewBase >& OwnerTable)
 {
 	TSharedRef< SAttributeItem > ReturnRow = SNew(SAttributeItem, OwnerTable)
-		.HighlightText(&SearchBoxPtr->GetText())
+		.HighlightText(SearchBoxPtr->GetText())
 		.TextColor(FLinearColor(1.0f, 1.0f, 1.0f, 1.f))
 		.AssociatedNode(Item);
 
diff --git a/Engine/Source/Editor/GraphEditor/Private/KismetPins/SGraphPinString.cpp b/Engine/Source/Editor/GraphEditor/Private/KismetPins/SGraphPinString.cpp
--- a/Engine/Source/Editor/GraphEditor/Private/KismetPins/SGraphPinString.cpp
+++ b/Engine/Source/Editor/GraphEditor/Private/KismetPins/SGraphPinString.cpp
@@ -12,21 +12,44 @@ void SGraphPinString::Construct(const FArguments& InArgs, UEdGraphPin* InGraphPi
 
 TSharedRef<SWidget>	SGraphPinString::GetDefaultValueWidget()
 {
-	return SNew(SBox)
-		.MinDesiredWidth(18.0f)
-		.MaxDesiredHeight(200)
-		[
-			SNew(SMultiLineEditableTextBox)
-			.Style(FEditorStyle::Get(), "Graph.EditableTextBox")
-			.Text(this, &SGraphPinString::GetTypeInValue)
-			.SelectAllTextWhenFocused(true)
-			.Visibility(this, &SGraphPin::GetDefaultValueVisibility)
-			.IsReadOnly(this, &SGraphPinString::GetDefaultValueIsReadOnly)
-			.OnTextCommitted(this, &SGraphPinString::SetTypeInValue)
-			.ForegroundColor(FSlateColor::UseForeground())
-			.WrapTextAt(400)
-			.ModiferKeyForNewLine(EModifierKey::Shift)
-		];
+	// Only allow actual string pins to be multi-line
+	// Other text based pins (such as names and numbers) should be single-line only
+	const bool bIsMultiLine = GraphPinObj->PinType.PinCategory == UEdGraphSchema_K2::PC_String;
+
+	if (bIsMultiLine)
+	{
+		return SNew(SBox)
+			.MinDesiredWidth(18)
+			.MaxDesiredHeight(200)
+			[
+				SNew(SMultiLineEditableTextBox)
+				.Style(FEditorStyle::Get(), "Graph.EditableTextBox")
+				.Text(this, &SGraphPinString::GetTypeInValue)
+				.SelectAllTextWhenFocused(true)
+				.Visibility(this, &SGraphPin::GetDefaultValueVisibility)
+				.IsReadOnly(this, &SGraphPinString::GetDefaultValueIsReadOnly)
+				.OnTextCommitted(this, &SGraphPinString::SetTypeInValue)
+				.ForegroundColor(FSlateColor::UseForeground())
+				.WrapTextAt(400)
+				.ModiferKeyForNewLine(EModifierKey::Shift)
+			];
+	}
+	else
+	{
+		return SNew(SBox)
+			.MinDesiredWidth(18)
+			.MaxDesiredWidth(400)
+			[
+				SNew(SEditableTextBox)
+					.Style(FEditorStyle::Get(), "Graph.EditableTextBox")
+					.Text(this, &SGraphPinString::GetTypeInValue)
+					.SelectAllTextWhenFocused(true)
+					.Visibility(this, &SGraphPin::GetDefaultValueVisibility)
+					.IsReadOnly(this, &SGraphPinString::GetDefaultValueIsReadOnly)
+					.OnTextCommitted(this, &SGraphPinString::SetTypeInValue)
+					.ForegroundColor(FSlateColor::UseForeground())
+			];
+	}
 }
 
 FText SGraphPinString::GetTypeInValue() const
diff --git a/Engine/Source/Editor/InputBindingEditor/Private/Widgets/SChordEditor.cpp b/Engine/Source/Editor/InputBindingEditor/Private/Widgets/SChordEditor.cpp
--- a/Engine/Source/Editor/InputBindingEditor/Private/Widgets/SChordEditor.cpp
+++ b/Engine/Source/Editor/InputBindingEditor/Private/Widgets/SChordEditor.cpp
@@ -1,7 +1,7 @@
 // Copyright 1998-2016 Epic Games, Inc. All Rights Reserved.
 
 #include "InputBindingEditorPrivatePCH.h"
-
+#include "SlateEditableTextLayout.h"
 
 #define LOCTEXT_NAMESPACE "SChordEditor"
 
@@ -30,8 +30,7 @@ void SChordEditor::Construct( const FArguments& InArgs, TSharedPtr<FChordTreeIte
 		.Font( Font )
 	);
 
-	LoadText();
-
+	EditableTextLayout->LoadText();
 }
 
 
@@ -52,18 +51,17 @@ FReply SChordEditor::OnKeyDown( const FGeometry& MyGeometry, const FKeyEvent& In
 			EditingInputChord.Key = Key;
 		}
 	
-		StartChangingText();
+		EditableTextLayout->BeginEditTransation();
 
 		EditingInputChord.bCtrl = InKeyEvent.IsControlDown();
 		EditingInputChord.bAlt = InKeyEvent.IsAltDown();
 		EditingInputChord.bShift = InKeyEvent.IsShiftDown();
 		EditingInputChord.bCmd = InKeyEvent.IsCommandDown();
 
-		LoadText();
-		//@todo checking the length of localized string is not valid, at the very least in this manner [10/11/2013 justin.sargent]
-		SetCaretPosition( EditingInputChord.GetInputText().ToString().Len() );
+		EditableTextLayout->LoadText();
+		EditableTextLayout->GoTo(ETextLocation::EndOfDocument);
 
-		FinishChangingText();
+		EditableTextLayout->EndEditTransaction();
 		
 		OnChordTyped( EditingInputChord );
 		
diff --git a/Engine/Source/Editor/Kismet/Private/BlueprintDetailsCustomization.cpp b/Engine/Source/Editor/Kismet/Private/BlueprintDetailsCustomization.cpp
--- a/Engine/Source/Editor/Kismet/Private/BlueprintDetailsCustomization.cpp
+++ b/Engine/Source/Editor/Kismet/Private/BlueprintDetailsCustomization.cpp
@@ -1117,9 +1117,12 @@ FText FBlueprintVarActionDetails::OnGetTooltipText() const
 	FName VarName = CachedVariableName;
 	if (VarName != NAME_None)
 	{
-		FString Result;
-		FBlueprintEditorUtils::GetBlueprintVariableMetaData(GetPropertyOwnerBlueprint(), VarName, GetLocalVariableScope(CachedVariableProperty.Get()), TEXT("tooltip"), Result);
-		return FText::FromString(Result);
+		if ( UBlueprint* OwnerBlueprint = GetPropertyOwnerBlueprint() )
+		{
+			FString Result;
+			FBlueprintEditorUtils::GetBlueprintVariableMetaData(GetPropertyOwnerBlueprint(), VarName, GetLocalVariableScope(CachedVariableProperty.Get()), TEXT("tooltip"), Result);
+			return FText::FromString(Result);
+		}
 	}
 	return FText();
 }
@@ -1334,17 +1337,21 @@ FText FBlueprintVarActionDetails::OnGetCategoryText() const
 	{
 		const UEdGraphSchema_K2* K2Schema = GetDefault<UEdGraphSchema_K2>();
 
-		FText Category = FBlueprintEditorUtils::GetBlueprintVariableCategory(GetPropertyOwnerBlueprint(), VarName, GetLocalVariableScope(CachedVariableProperty.Get()));
-
-		// Older blueprints will have their name as the default category and whenever it is the same as the default category, display localized text
-		if( Category.EqualTo(FText::FromString(GetPropertyOwnerBlueprint()->GetName())) || Category.EqualTo(K2Schema->VR_DefaultCategory) )
+		if ( UBlueprint* OwnerBlueprint = GetPropertyOwnerBlueprint() )
 		{
-			return K2Schema->VR_DefaultCategory;
-		}
-		else
-		{
-			return Category;
+			FText Category = FBlueprintEditorUtils::GetBlueprintVariableCategory(OwnerBlueprint, VarName, GetLocalVariableScope(CachedVariableProperty.Get()));
+
+			// Older blueprints will have their name as the default category and whenever it is the same as the default category, display localized text
+			if ( Category.EqualTo(FText::FromString(OwnerBlueprint->GetName())) || Category.EqualTo(K2Schema->VR_DefaultCategory) )
+			{
+				return K2Schema->VR_DefaultCategory;
+			}
+			else
+			{
+				return Category;
+			}
 		}
+
 		return FText::FromName(VarName);
 	}
 	return FText();
@@ -1659,10 +1666,13 @@ FText FBlueprintVarActionDetails::OnGetMetaKeyValue(FName Key) const
 	FName VarName = CachedVariableName;
 	if (VarName != NAME_None)
 	{
-		FString Result;
-		FBlueprintEditorUtils::GetBlueprintVariableMetaData(GetPropertyOwnerBlueprint(), VarName, GetLocalVariableScope(CachedVariableProperty.Get()), Key, /*out*/ Result);
+		if ( UBlueprint* Blueprint = GetPropertyOwnerBlueprint() )
+		{
+			FString Result;
+			FBlueprintEditorUtils::GetBlueprintVariableMetaData(GetPropertyOwnerBlueprint(), VarName, GetLocalVariableScope(CachedVariableProperty.Get()), Key, /*out*/ Result);
 
-		return FText::FromString(Result);
+			return FText::FromString(Result);
+		}
 	}
 	return FText();
 }
diff --git a/Engine/Source/Editor/Kismet/Public/SSCSEditor.h b/Engine/Source/Editor/Kismet/Public/SSCSEditor.h
--- a/Engine/Source/Editor/Kismet/Public/SSCSEditor.h
+++ b/Engine/Source/Editor/Kismet/Public/SSCSEditor.h
@@ -82,7 +82,12 @@ public:
 	 */
 	ENodeType GetNodeType() const;
 	/**
-	 * @return The component template or instance represented by this node, if it's a component node.
+	 * @param	bEvenIfPendingKill	If false, nullptr will be returned if the cached component template is pending kill.
+	 *								If true, it will be returned regardless (this is used for recaching the component template if the objects
+	 *								have been reinstanced following construction script execution).
+	 *
+	 * @note	Deliberately non-virtual, for performance reasons.
+	 * @return	The component template or instance represented by this node, if it's a component node.
 	 */
 	UActorComponent* GetComponentTemplate(bool bEvenIfPendingKill = false) const;
 	/**
@@ -379,7 +384,6 @@ protected:
 private:
 	FName InstancedComponentName;
 	TWeakObjectPtr<AActor> InstancedComponentOwnerPtr;
-	TWeakObjectPtr<UActorComponent> ComponentTemplate;
 };
 
 //////////////////////////////////////////////////////////////////////////
diff --git a/Engine/Source/Editor/KismetCompiler/Private/KismetCompilerVMBackend.cpp b/Engine/Source/Editor/KismetCompiler/Private/KismetCompilerVMBackend.cpp
--- a/Engine/Source/Editor/KismetCompiler/Private/KismetCompilerVMBackend.cpp
+++ b/Engine/Source/Editor/KismetCompiler/Private/KismetCompilerVMBackend.cpp
@@ -96,6 +96,15 @@ public:
 		return *this;
 	}
 
+	FArchive& operator<<(EBlueprintTextLiteralType E)
+	{
+		static_assert(sizeof(__underlying_type(EBlueprintTextLiteralType)) == sizeof(uint8), "EBlueprintTextLiteralType is expected to be a uint8");
+
+		uint8 B = (uint8)E;
+		Serialize(&B, 1);
+		return *this;
+	}
+
 	FArchive& operator<<(EPropertyType E)
 	{
 		uint8 B = E; 
@@ -382,9 +391,43 @@ public:
 			{
 				Writer << EX_TextConst;
 				
-				EmitStringLiteral(FTextInspector::GetSourceString(Term->TextLiteral)? *FTextInspector::GetSourceString(Term->TextLiteral) : TEXT(""));
-				EmitStringLiteral(FTextInspector::GetKey(Term->TextLiteral).Get(TEXT("")));
-				EmitStringLiteral(FTextInspector::GetNamespace(Term->TextLiteral).Get(TEXT("")));
+				const FString& StringValue = FTextInspector::GetDisplayString(Term->TextLiteral);
+
+				// What kind of text are we dealing with?
+				if (Term->TextLiteral.IsEmpty())
+				{
+					Writer << EBlueprintTextLiteralType::Empty;
+				}
+				else if (Term->TextLiteral.IsCultureInvariant())
+				{
+					Writer << EBlueprintTextLiteralType::InvariantText;
+					EmitStringLiteral(StringValue);
+				}
+				else
+				{
+					bool bIsLocalized = false;
+					FString Namespace;
+					FString Key;
+					const FString* SourceString = FTextInspector::GetSourceString(Term->TextLiteral);
+
+					if (SourceString && Term->TextLiteral.ShouldGatherForLocalization())
+					{
+						bIsLocalized = FTextLocalizationManager::Get().FindNamespaceAndKeyFromDisplayString(FTextInspector::GetSharedDisplayString(Term->TextLiteral), Namespace, Key);
+					}
+
+					if (bIsLocalized)
+					{
+						Writer << EBlueprintTextLiteralType::LocalizedText;
+						EmitStringLiteral(*SourceString);
+						EmitStringLiteral(Key);
+						EmitStringLiteral(Namespace);
+					}
+					else
+					{
+						Writer << EBlueprintTextLiteralType::LiteralString;
+						EmitStringLiteral(StringValue);
+					}
+				}
 			}
 			else if (CoerceProperty->IsA(UFloatProperty::StaticClass()))
 			{
diff --git a/Engine/Source/Editor/LandscapeEditor/Public/LandscapeEditorObject.h b/Engine/Source/Editor/LandscapeEditor/Public/LandscapeEditorObject.h
--- a/Engine/Source/Editor/LandscapeEditor/Public/LandscapeEditorObject.h
+++ b/Engine/Source/Editor/LandscapeEditor/Public/LandscapeEditorObject.h
@@ -233,7 +233,7 @@ class ULandscapeEditorObject : public UObject
 	float ToolStrength;
 
 	// Enable to make tools blend towards a target value
-	UPROPERTY(NonTransactional)
+	UPROPERTY(Category = "Tool Settings", NonTransactional, EditAnywhere, meta = (InlineEditConditionToggle))
 	bool bUseWeightTargetValue;
 
 	// Enable to make tools blend towards a target value
@@ -259,7 +259,7 @@ class ULandscapeEditorObject : public UObject
 	bool bPickValuePerApply;
 
 	// Enable to flatten towards a target height
-	UPROPERTY(NonTransactional)
+	UPROPERTY(Category = "Tool Settings", NonTransactional, EditAnywhere, meta = (InlineEditConditionToggle))
 	bool bUseFlattenTarget;
 
 	// Target height to flatten towards (in Unreal Units)
@@ -288,7 +288,7 @@ class ULandscapeEditorObject : public UObject
 	int32 SmoothFilterKernelSize;
 
 	// If checked, performs a detail preserving smooth using the specified detail smoothing value
-	UPROPERTY(NonTransactional)
+	UPROPERTY(Category = "Tool Settings", NonTransactional, EditAnywhere, meta = (InlineEditConditionToggle))
 	bool bDetailSmooth;
 
 	// Larger detail smoothing values remove more details, while smaller values preserve more details
@@ -340,7 +340,7 @@ class ULandscapeEditorObject : public UObject
 	float RainDistScale;
 
 	// If checked, performs a detail-preserving smooth to the erosion effect using the specified detail smoothing value
-	UPROPERTY(NonTransactional)
+	UPROPERTY(Category = "Tool Settings", NonTransactional, EditAnywhere, meta = (InlineEditConditionToggle))
 	bool bHErosionDetailSmooth;
 
 	// Larger detail smoothing values remove more details, while smaller values preserve more details
diff --git a/Engine/Source/Editor/LevelEditor/Private/LevelViewportLayout.cpp b/Engine/Source/Editor/LevelEditor/Private/LevelViewportLayout.cpp
--- a/Engine/Source/Editor/LevelEditor/Private/LevelViewportLayout.cpp
+++ b/Engine/Source/Editor/LevelEditor/Private/LevelViewportLayout.cpp
@@ -260,13 +260,11 @@ void  FLevelViewportLayout::EndThrottleForAnimatedResize()
 	}
 }
 
-void FLevelViewportLayout::InitCommonLayoutFromString( const FString& SpecificLayoutString )
+void FLevelViewportLayout::InitCommonLayoutFromString( const FString& SpecificLayoutString, FName DefaultMaximizedViewport )
 {
 	bool bShouldBeMaximized = bIsMaximizeSupported && ViewportLayoutDefs::bDefaultShouldBeMaximized;
 	bool bShouldBeImmersive = ViewportLayoutDefs::bDefaultShouldBeImmersive;
 	
-	FName SavedMaximizedViewport = NAME_None;
-	
 	if (!SpecificLayoutString.IsEmpty())
 	{
 		const FString& IniSection = FLayoutSaveRestore::GetAdditionalLayoutConfigIni();
@@ -277,17 +275,17 @@ void FLevelViewportLayout::InitCommonLayoutFromString( const FString& SpecificLa
 		FString MaximizedViewportString;
 		if (GConfig->GetString(*IniSection, *(SpecificLayoutString + TEXT(".MaximizedViewport")), MaximizedViewportString, GEditorPerProjectIni))
 		{
-			SavedMaximizedViewport = *MaximizedViewportString;
+			DefaultMaximizedViewport = *MaximizedViewportString;
 		}
 	}
 	// Replacement layouts (those selected by the user via a command) don't start maximized so the layout can be seen clearly.
-	if (!bIsReplacement && bIsMaximizeSupported && Viewports.Contains(SavedMaximizedViewport) && (bShouldBeMaximized || bShouldBeImmersive))
+	if (!bIsReplacement && bIsMaximizeSupported && Viewports.Contains(DefaultMaximizedViewport) && (bShouldBeMaximized || bShouldBeImmersive))
 	{
 		// we are not toggling maximize or immersive state but setting it directly
 		const bool bToggle=false;
 		// Do not allow animation at startup as it hitches
 		const bool bAllowAnimation=false;
-		MaximizeViewport(SavedMaximizedViewport, bShouldBeMaximized, bShouldBeImmersive, bAllowAnimation);
+		MaximizeViewport(DefaultMaximizedViewport, bShouldBeMaximized, bShouldBeImmersive, bAllowAnimation);
 	}
 }
 
diff --git a/Engine/Source/Editor/LevelEditor/Private/LevelViewportLayout2x2.cpp b/Engine/Source/Editor/LevelEditor/Private/LevelViewportLayout2x2.cpp
--- a/Engine/Source/Editor/LevelEditor/Private/LevelViewportLayout2x2.cpp
+++ b/Engine/Source/Editor/LevelEditor/Private/LevelViewportLayout2x2.cpp
@@ -101,7 +101,7 @@ TSharedRef<SWidget> FLevelViewportLayout2x2::MakeViewportLayout(const FString& L
 		SplitterWidget->SetSplitterPercentages(SplitterPercentages);
 	}
 
-	InitCommonLayoutFromString(SpecificLayoutString);
+	InitCommonLayoutFromString(SpecificLayoutString, *BottomLeftKey);
 
 	return SplitterWidget.ToSharedRef();
 }
diff --git a/Engine/Source/Editor/LevelEditor/Private/LevelViewportLayoutFourPanes.cpp b/Engine/Source/Editor/LevelEditor/Private/LevelViewportLayoutFourPanes.cpp
--- a/Engine/Source/Editor/LevelEditor/Private/LevelViewportLayoutFourPanes.cpp
+++ b/Engine/Source/Editor/LevelEditor/Private/LevelViewportLayoutFourPanes.cpp
@@ -55,7 +55,7 @@ TSharedRef<SWidget> FLevelViewportLayoutFourPanes::MakeViewportLayout(const FStr
 		ViewportKey0, ViewportKey1, ViewportKey2, ViewportKey3,
 		PrimarySplitterPercentage, SecondarySplitterPercentage0, SecondarySplitterPercentage1);
 
-	InitCommonLayoutFromString(SpecificLayoutString);
+	InitCommonLayoutFromString(SpecificLayoutString, *ViewportKey1);
 
 	return LayoutWidget;
 }
diff --git a/Engine/Source/Editor/LevelEditor/Private/LevelViewportLayoutOnePane.cpp b/Engine/Source/Editor/LevelEditor/Private/LevelViewportLayoutOnePane.cpp
--- a/Engine/Source/Editor/LevelEditor/Private/LevelViewportLayoutOnePane.cpp
+++ b/Engine/Source/Editor/LevelEditor/Private/LevelViewportLayoutOnePane.cpp
@@ -54,7 +54,7 @@ TSharedRef<SWidget> FLevelViewportLayoutOnePane::MakeViewportLayout(const FStrin
 	// Make newly-created perspective viewports active by default
 	GCurrentLevelEditingViewportClient = &ViewportWidget->GetLevelViewportClient();
 
-	InitCommonLayoutFromString(SpecificLayoutString);
+	InitCommonLayoutFromString(SpecificLayoutString, NAME_None);
 
 	return StaticCastSharedRef<SWidget>(ViewportBox.ToSharedRef());
 }
diff --git a/Engine/Source/Editor/LevelEditor/Private/LevelViewportLayoutThreePanes.cpp b/Engine/Source/Editor/LevelEditor/Private/LevelViewportLayoutThreePanes.cpp
--- a/Engine/Source/Editor/LevelEditor/Private/LevelViewportLayoutThreePanes.cpp
+++ b/Engine/Source/Editor/LevelEditor/Private/LevelViewportLayoutThreePanes.cpp
@@ -52,7 +52,7 @@ TSharedRef<SWidget> FLevelViewportLayoutThreePanes::MakeViewportLayout(const FSt
 		ViewportKey0, ViewportKey1, ViewportKey2,
 		PrimarySplitterPercentage, SecondarySplitterPercentage);
 
-	InitCommonLayoutFromString(SpecificLayoutString);
+	InitCommonLayoutFromString(SpecificLayoutString, *ViewportKey1);
 
 	return LayoutWidget;
 }
diff --git a/Engine/Source/Editor/LevelEditor/Public/LevelViewportLayout.h b/Engine/Source/Editor/LevelEditor/Public/LevelViewportLayout.h
--- a/Engine/Source/Editor/LevelEditor/Public/LevelViewportLayout.h
+++ b/Engine/Source/Editor/LevelEditor/Public/LevelViewportLayout.h
@@ -236,7 +236,7 @@ protected:
 	}
 
 	/** Called in MakeViewportLayout() functions for derived types of layout to init values common to all types */
-	void InitCommonLayoutFromString( const FString& SpecificLayoutString );
+	void InitCommonLayoutFromString( const FString& SpecificLayoutString, FName DefaultMaximizedViewport );
 
 	/** Called in SaveLayoutString() functions for derived types of layout to save values common to all types */
 	void SaveCommonLayoutString( const FString& SpecificLayoutString ) const;
diff --git a/Engine/Source/Editor/MeshPaint/Private/MeshPaintEdMode.cpp b/Engine/Source/Editor/MeshPaint/Private/MeshPaintEdMode.cpp
--- a/Engine/Source/Editor/MeshPaint/Private/MeshPaintEdMode.cpp
+++ b/Engine/Source/Editor/MeshPaint/Private/MeshPaintEdMode.cpp
@@ -395,7 +395,8 @@ bool FEdModeMeshPaint::InputKey( FEditorViewportClient* InViewportClient, FViewp
 	bool bHandled = false;
 
 	const bool bIsLeftButtonDown = ( InKey == EKeys::LeftMouseButton && InEvent != IE_Released ) || InViewport->KeyState( EKeys::LeftMouseButton );
-	const bool bIsCtrlDown = ( ( InKey == EKeys::LeftControl || InKey == EKeys::RightControl ) && InEvent != IE_Released ) || InViewport->KeyState( EKeys::LeftControl ) || InViewport->KeyState( EKeys::RightControl );
+	const bool bIsRightButtonDown = (InKey == EKeys::RightMouseButton && InEvent != IE_Released) || InViewport->KeyState(EKeys::RightMouseButton);
+	const bool bIsCtrlDown = ((InKey == EKeys::LeftControl || InKey == EKeys::RightControl) && InEvent != IE_Released) || InViewport->KeyState(EKeys::LeftControl) || InViewport->KeyState(EKeys::RightControl);
 	const bool bIsShiftDown = ( ( InKey == EKeys::LeftShift || InKey == EKeys::RightShift ) && InEvent != IE_Released ) || InViewport->KeyState( EKeys::LeftShift ) || InViewport->KeyState( EKeys::RightShift );
 	const bool bIsAltDown = ( ( InKey == EKeys::LeftAlt || InKey == EKeys::RightAlt ) && InEvent != IE_Released ) || InViewport->KeyState( EKeys::LeftAlt ) || InViewport->KeyState( EKeys::RightAlt );
 
@@ -462,7 +463,7 @@ bool FEdModeMeshPaint::InputKey( FEditorViewportClient* InViewportClient, FViewp
 	if( !bIsAltDown && InViewportClient->IsPerspective() && InViewportClient->EngineShowFlags.ModeWidgets)
 	{
 		// Does the user want to paint right now?
-		const bool bUserWantsPaint = bIsLeftButtonDown && !bIsAltDown;
+		const bool bUserWantsPaint = bIsLeftButtonDown && !bIsRightButtonDown && !bIsAltDown;
 		bool bAnyPaintAbleActorsUnderCursor = false;
 
 		// Stop current tracking if the user is no longer painting
diff --git a/Engine/Source/Editor/PropertyEditor/Private/CategoryPropertyNode.cpp b/Engine/Source/Editor/PropertyEditor/Private/CategoryPropertyNode.cpp
--- a/Engine/Source/Editor/PropertyEditor/Private/CategoryPropertyNode.cpp
+++ b/Engine/Source/Editor/PropertyEditor/Private/CategoryPropertyNode.cpp
@@ -54,7 +54,8 @@ void FCategoryPropertyNode::InitChildNodes()
 			bool bMetaDataAllowVisible = true;
 			if (!bShowHiddenProperties)
 			{
-				FString MetaDataVisibilityCheckString = It->GetMetaData(TEXT("bShowOnlyWhenTrue"));
+				static const FName Name_bShowOnlyWhenTrue("bShowOnlyWhenTrue");
+				FString MetaDataVisibilityCheckString = It->GetMetaData(Name_bShowOnlyWhenTrue);
 				if (MetaDataVisibilityCheckString.Len())
 				{
 					//ensure that the metadata visibility string is actually set to true in order to show this property
@@ -65,12 +66,14 @@ void FCategoryPropertyNode::InitChildNodes()
 
 			if (bMetaDataAllowVisible)
 			{
+				static const FName Name_InlineEditConditionToggle("InlineEditConditionToggle");
+				const bool bOnlyShowAsInlineEditCondition = (*It)->HasMetaData(Name_InlineEditConditionToggle);
 				const bool bShowIfEditableProperty = (*It)->HasAnyPropertyFlags(CPF_Edit);
 				const bool bShowIfDisableEditOnInstance = !(*It)->HasAnyPropertyFlags(CPF_DisableEditOnInstance) || bShouldShowDisableEditOnInstance;
 
 				// Add if we are showing non-editable props and this is the 'None' category, 
 				// or if this is the right category, and we are either showing non-editable
-				if (FObjectEditorUtils::GetCategoryFName(*It) == CategoryName && (bShowHiddenProperties || (bShowIfEditableProperty && bShowIfDisableEditOnInstance)))
+				if (FObjectEditorUtils::GetCategoryFName(*It) == CategoryName && (bShowHiddenProperties || (bShowIfEditableProperty && !bOnlyShowAsInlineEditCondition && bShowIfDisableEditOnInstance)))
 				{
 					Properties.Add(*It);
 				}
diff --git a/Engine/Source/Editor/PropertyEditor/Private/ItemPropertyNode.cpp b/Engine/Source/Editor/PropertyEditor/Private/ItemPropertyNode.cpp
--- a/Engine/Source/Editor/PropertyEditor/Private/ItemPropertyNode.cpp
+++ b/Engine/Source/Editor/PropertyEditor/Private/ItemPropertyNode.cpp
@@ -188,9 +188,11 @@ void FItemPropertyNode::InitChildNodes()
 		for( TFieldIterator<UProperty> It(StructProperty->Struct); It; ++It )
 		{
 			UProperty* StructMember = *It;
+			static const FName Name_InlineEditConditionToggle("InlineEditConditionToggle");
+			const bool bOnlyShowAsInlineEditCondition = StructMember->HasMetaData(Name_InlineEditConditionToggle);
 			const bool bShowIfEditableProperty = StructMember->HasAnyPropertyFlags(CPF_Edit);
 			const bool bShowIfDisableEditOnInstance = !StructMember->HasAnyPropertyFlags(CPF_DisableEditOnInstance) || bShouldShowDisableEditOnInstance;
-			if (bShouldShowHiddenProperties || (bShowIfEditableProperty && bShowIfDisableEditOnInstance))
+			if (bShouldShowHiddenProperties || (bShowIfEditableProperty && !bOnlyShowAsInlineEditCondition && bShowIfDisableEditOnInstance))
 			{
 				TSharedPtr<FItemPropertyNode> NewItemNode( new FItemPropertyNode );//;//CreatePropertyItem(StructMember,INDEX_NONE,this);
 		
diff --git a/Engine/Source/Editor/PropertyEditor/Private/ObjectPropertyNode.cpp b/Engine/Source/Editor/PropertyEditor/Private/ObjectPropertyNode.cpp
--- a/Engine/Source/Editor/PropertyEditor/Private/ObjectPropertyNode.cpp
+++ b/Engine/Source/Editor/PropertyEditor/Private/ObjectPropertyNode.cpp
@@ -335,8 +335,16 @@ void FObjectPropertyNode::InternalInitChildNodes( FName SinglePropertyName )
 			}
 		}
 
+		// 
+		static const FName Name_InlineEditConditionToggle("InlineEditConditionToggle");
+		if (It->HasMetaData(Name_InlineEditConditionToggle))
+		{
+			bHidden = true;
+		}
+
 		bool bMetaDataAllowVisible = true;
-		FString MetaDataVisibilityCheckString = It->GetMetaData(TEXT("bShowOnlyWhenTrue"));
+		static const FName Name_bShowOnlyWhenTrue("bShowOnlyWhenTrue");
+		FString MetaDataVisibilityCheckString = It->GetMetaData(Name_bShowOnlyWhenTrue);
 		if (MetaDataVisibilityCheckString.Len())
 		{
 			//ensure that the metadata visibility string is actually set to true in order to show this property
@@ -451,9 +459,11 @@ void FObjectPropertyNode::InternalInitChildNodes( FName SinglePropertyName )
 		// Iterate over all fields, creating items.
 		for( TFieldIterator<UProperty> It(BaseClass.Get()); It; ++It )
 		{
+			static const FName Name_InlineEditConditionToggle("InlineEditConditionToggle");
+			const bool bOnlyShowAsInlineEditCondition = (*It)->HasMetaData(Name_InlineEditConditionToggle);
 			const bool bShowIfNonHiddenEditableProperty = (*It)->HasAnyPropertyFlags(CPF_Edit) && !FEditorCategoryUtils::IsCategoryHiddenFromClass(BaseClass.Get(), FObjectEditorUtils::GetCategory(*It));
 			const bool bShowIfDisableEditOnInstance = !(*It)->HasAnyPropertyFlags(CPF_DisableEditOnInstance) || bShouldShowDisableEditOnInstance;
-			if (bShouldShowHiddenProperties || (bShowIfNonHiddenEditableProperty && bShowIfDisableEditOnInstance))
+			if (bShouldShowHiddenProperties || (bShowIfNonHiddenEditableProperty && !bOnlyShowAsInlineEditCondition && bShowIfDisableEditOnInstance))
 			{
 				UProperty* CurProp = *It;
 				if( SinglePropertyName == NAME_None || CurProp->GetFName() == SinglePropertyName )
diff --git a/Engine/Source/Editor/PropertyEditor/Private/Presentation/PropertyEditor/PropertyEditor.cpp b/Engine/Source/Editor/PropertyEditor/Private/Presentation/PropertyEditor/PropertyEditor.cpp
--- a/Engine/Source/Editor/PropertyEditor/Private/Presentation/PropertyEditor/PropertyEditor.cpp
+++ b/Engine/Source/Editor/PropertyEditor/Private/Presentation/PropertyEditor/PropertyEditor.cpp
@@ -582,20 +582,41 @@ bool FPropertyEditor::SupportsEditConditionToggle( UProperty* InProperty )
 {
 	bool bShowEditConditionToggle = false;
 
-	if (!InProperty->HasMetaData(TEXT("HideEditConditionToggle")))
+	static const FName Name_HideEditConditionToggle("HideEditConditionToggle");
+	if (!InProperty->HasMetaData(Name_HideEditConditionToggle))
 	{
 		bool bNegateValue = false;
 		UBoolProperty* ConditionalProperty = PropertyCustomizationHelpers::GetEditConditionProperty( InProperty, bNegateValue );
 		if( ConditionalProperty != NULL )
 		{
-			bShowEditConditionToggle = true;
-
-			if( ConditionalProperty->HasAllPropertyFlags( CPF_Edit ) )
+			if (!ConditionalProperty->HasAllPropertyFlags(CPF_Edit))
 			{
-				// Conditionally-dependent property is already exposed for editing, so no need to draw another
-				// check box next to this property's label
-				bShowEditConditionToggle = false;
+				// Warn if the editcondition property is not marked as editable; this will break when the component is added to a Blueprint.
+				UObject* PropertyScope = InProperty->GetOwnerClass();
+				UObject* ConditionalPropertyScope = ConditionalProperty->GetOwnerClass();
+				if (!PropertyScope)
+				{
+					PropertyScope = InProperty->GetOwnerStruct();
+					ConditionalPropertyScope = ConditionalProperty->GetOwnerStruct();
+				}
+				const FString PropertyScopeName = PropertyScope ? PropertyScope->GetName() : FString();
+				const FString ConditionalPropertyScopeName = ConditionalPropertyScope ? ConditionalPropertyScope->GetName() : FString();
+
+				// Show the toggle box anyway, with a warning
+				UE_LOG(LogPropertyNode, Warning, TEXT("Property %s::%s is an editcondition for %s::%s, but is not marked as EditAnywhere."),
+						*ConditionalPropertyScopeName, *ConditionalProperty->GetName(), *PropertyScopeName, *InProperty->GetName());
+
+				bShowEditConditionToggle = true;
 			}
+			else
+			{
+				// If it's editable, then only put an inline toggle box if the metadata specifies it
+				static const FName Name_InlineEditConditionToggle("InlineEditConditionToggle");
+				if (ConditionalProperty->HasMetaData(Name_InlineEditConditionToggle))
+				{
+					bShowEditConditionToggle = true;
+				}
+			}		
 		}
 	}
 
diff --git a/Engine/Source/Editor/PropertyEditor/Private/PropertyNode.cpp b/Engine/Source/Editor/PropertyEditor/Private/PropertyNode.cpp
--- a/Engine/Source/Editor/PropertyEditor/Private/PropertyNode.cpp
+++ b/Engine/Source/Editor/PropertyEditor/Private/PropertyNode.cpp
@@ -130,13 +130,15 @@ void FPropertyNode::InitNode( const FPropertyNodeInitParams& InitParams )
 
 		// true if the property can be expanded into the property window; that is, instead of seeing
 		// a pointer to the object, you see the object's properties.
-		const bool bEditInline = bIsObjectOrInterface && GotReadAddresses && MyProperty->HasMetaData(TEXT("EditInline"));
+		static const FName Name_EditInline("EditInline");
+		const bool bEditInline = bIsObjectOrInterface && GotReadAddresses && MyProperty->HasMetaData(Name_EditInline);
 		SetNodeFlags(EPropertyNodeFlags::EditInline, bEditInline);
 
 		//Get the property max child depth
-		if (Property->HasMetaData(TEXT("MaxPropertyDepth")))
+		static const FName Name_MaxPropertyDepth("MaxPropertyDepth");
+		if (Property->HasMetaData(Name_MaxPropertyDepth))
 		{
-			int32 NewMaxChildDepthAllowed = Property->GetINTMetaData(TEXT("MaxPropertyDepth"));
+			int32 NewMaxChildDepthAllowed = Property->GetINTMetaData(Name_MaxPropertyDepth);
 			//Ensure new depth is valid.  Otherwise just let the parent specified value stand
 			if (NewMaxChildDepthAllowed > 0)
 			{
diff --git a/Engine/Source/Editor/PropertyEditor/Private/SStructureDetailsView.cpp b/Engine/Source/Editor/PropertyEditor/Private/SStructureDetailsView.cpp
--- a/Engine/Source/Editor/PropertyEditor/Private/SStructureDetailsView.cpp
+++ b/Engine/Source/Editor/PropertyEditor/Private/SStructureDetailsView.cpp
@@ -143,29 +143,33 @@ void SStructureDetailsView::Construct(const FArguments& InArgs)
 
 	ChildSlot
 	[
-		SNew( SVerticalBox )
-		+ SVerticalBox::Slot()
-		.AutoHeight()
-		.Padding( 0.0f, 0.0f, 0.0f, 2.0f )
+		SNew( SBox )
+		.Visibility(this, &SStructureDetailsView::GetPropertyEditingVisibility)
 		[
-			FilterRow
-		]
-		+ SVerticalBox::Slot()
-		.FillHeight(1)
-		.Padding(0)
-		[
-			SNew( SHorizontalBox )
-			+ SHorizontalBox::Slot()
+			SNew( SVerticalBox )
+			+ SVerticalBox::Slot()
+			.AutoHeight()
+			.Padding( 0.0f, 0.0f, 0.0f, 2.0f )
 			[
-				DetailTree.ToSharedRef()
+				FilterRow
 			]
-			+ SHorizontalBox::Slot()
-			.AutoWidth()
+			+ SVerticalBox::Slot()
+			.FillHeight(1)
+			.Padding(0)
 			[
-				SNew( SBox )
-				.WidthOverride( 16.0f )
+				SNew( SHorizontalBox )
+				+ SHorizontalBox::Slot()
+				[
+					DetailTree.ToSharedRef()
+				]
+				+ SHorizontalBox::Slot()
+				.AutoWidth()
 				[
-					ExternalScrollbar
+					SNew( SBox )
+					.WidthOverride( 16.0f )
+					[
+						ExternalScrollbar
+					]
 				]
 			]
 		]
@@ -249,4 +253,9 @@ void SStructureDetailsView::CustomUpdatePropertyMap()
 	DetailLayout->DefaultCategory(NAME_None).SetDisplayName(NAME_None, CustomName);
 }
 
+EVisibility SStructureDetailsView::GetPropertyEditingVisibility() const
+{
+	return StructData.IsValid() && StructData->IsValid() && RootNode.IsValid() && RootNode->IsValid() ? EVisibility::Visible : EVisibility::Collapsed;
+}
+
 #undef LOCTEXT_NAMESPACE
diff --git a/Engine/Source/Editor/PropertyEditor/Private/SStructureDetailsView.h b/Engine/Source/Editor/PropertyEditor/Private/SStructureDetailsView.h
--- a/Engine/Source/Editor/PropertyEditor/Private/SStructureDetailsView.h
+++ b/Engine/Source/Editor/PropertyEditor/Private/SStructureDetailsView.h
@@ -110,4 +110,6 @@ public:
 protected:
 
 	virtual void CustomUpdatePropertyMap() override;
+
+	EVisibility GetPropertyEditingVisibility() const;
 };
diff --git a/Engine/Source/Editor/PropertyEditor/Private/StructurePropertyNode.cpp b/Engine/Source/Editor/PropertyEditor/Private/StructurePropertyNode.cpp
--- a/Engine/Source/Editor/PropertyEditor/Private/StructurePropertyNode.cpp
+++ b/Engine/Source/Editor/PropertyEditor/Private/StructurePropertyNode.cpp
@@ -21,10 +21,12 @@ void FStructurePropertyNode::InitChildNodes()
 
 		if (StructMember)
 		{
+			static const FName Name_InlineEditConditionToggle("InlineEditConditionToggle");
+			const bool bOnlyShowAsInlineEditCondition = StructMember->HasMetaData(Name_InlineEditConditionToggle);
 			const bool bShowIfEditableProperty = StructMember->HasAnyPropertyFlags(CPF_Edit);
 			const bool bShowIfDisableEditOnInstance = !StructMember->HasAnyPropertyFlags(CPF_DisableEditOnInstance) || bShouldShowDisableEditOnInstance;
 
-			if (bShouldShowHiddenProperties || (bShowIfEditableProperty && bShowIfDisableEditOnInstance))
+			if (bShouldShowHiddenProperties || (bShowIfEditableProperty && !bOnlyShowAsInlineEditCondition && bShowIfDisableEditOnInstance))
 			{
 				TSharedPtr<FItemPropertyNode> NewItemNode(new FItemPropertyNode);//;//CreatePropertyItem(StructMember,INDEX_NONE,this);
 
diff --git a/Engine/Source/Editor/UnrealEd/Classes/Commandlets/GatherTextFromSourceCommandlet.h b/Engine/Source/Editor/UnrealEd/Classes/Commandlets/GatherTextFromSourceCommandlet.h
--- a/Engine/Source/Editor/UnrealEd/Classes/Commandlets/GatherTextFromSourceCommandlet.h
+++ b/Engine/Source/Editor/UnrealEd/Classes/Commandlets/GatherTextFromSourceCommandlet.h
@@ -33,6 +33,10 @@ private:
 
 		EMacroBlockState EvaluateMacroStack() const;
 
+		void SetDefine( const FString& InDefineCtx );
+
+		void RemoveDefine( const FString& InDefineCtx );
+
 		//Working data
 		FString Filename;
 		int32 LineNumber;
@@ -105,8 +109,6 @@ private:
 		static const FString ElseString;
 		static const FString EndIfString;
 		static const FString DefinedString;
-		static const FString LocNamespaceString;
-		static const FString LocDefRegionString;
 		static const FString IniNamespaceString;
 	};
 
@@ -237,6 +239,7 @@ private:
 	static const FString ChangelistName;
 
 	static FString RemoveStringFromTextMacro(const FString& TextMacro, const FString& IdentForLogging, bool& Error);
+	static FString StripCommentsFromToken(const FString& InToken, FSourceFileParseContext& Context);
 	static bool ParseSourceText(const FString& Text, const TArray<FParsableDescriptor*>& Parsables, FSourceFileParseContext& ParseCtxt);
 
 public:
diff --git a/Engine/Source/Editor/UnrealEd/Classes/Editor/TransBuffer.h b/Engine/Source/Editor/UnrealEd/Classes/Editor/TransBuffer.h
--- a/Engine/Source/Editor/UnrealEd/Classes/Editor/TransBuffer.h
+++ b/Engine/Source/Editor/UnrealEd/Classes/Editor/TransBuffer.h
@@ -32,6 +32,9 @@ public:
 	/** Maximum number of bytes the transaction buffer is allowed to occupy */
 	SIZE_T MaxMemory;
 
+	/** Undo barrier stack */
+	TArray<int32> UndoBarrierStack;
+
 public:
 
 	// Constructor.
@@ -102,6 +105,9 @@ public:
 	virtual SIZE_T GetUndoSize() const override;
 	virtual int32 GetUndoCount( ) const override { return UndoCount; }
 	virtual FUndoSessionContext GetRedoContext() override;
+	virtual void SetUndoBarrier() override;
+	virtual void RemoveUndoBarrier() override;
+	virtual void ClearUndoBarriers() override;
 	virtual bool Undo() override;
 	virtual bool Redo() override;
 	virtual bool EnableObjectSerialization() override;
diff --git a/Engine/Source/Editor/UnrealEd/Classes/Editor/Transactor.h b/Engine/Source/Editor/UnrealEd/Classes/Editor/Transactor.h
--- a/Engine/Source/Editor/UnrealEd/Classes/Editor/Transactor.h
+++ b/Engine/Source/Editor/UnrealEd/Classes/Editor/Transactor.h
@@ -533,12 +533,29 @@ class UTransactor : public UObject
 	 */
 	virtual FUndoSessionContext GetRedoContext () PURE_VIRTUAL(UTransactor::GetRedoDesc,return FUndoSessionContext(););
 
+	/**
+	 * Sets an undo barrier at the current point in the transaction buffer.
+	 * Undoing beyond this point will not be allowed until the barrier is removed.
+	 */
+	virtual void SetUndoBarrier() PURE_VIRTUAL(UTransactor::SetUndoBarrier, );
+
+	/**
+	 * Removes the last set undo barrier from the transaction buffer.
+	 */
+	virtual void RemoveUndoBarrier() PURE_VIRTUAL(UTransactor::RemoveUndoBarrier, );
+
+	/**
+	 * Clears all undo barriers.
+	 */
+	virtual void ClearUndoBarriers() PURE_VIRTUAL(UTransactor::ClearUndoBarriers, );
+
 	/**
 	 * Executes an undo transaction, undoing all actions contained by that transaction.
 	 * 
 	 * @return				true if the transaction was successfully undone
 	 */
 	virtual bool Undo() PURE_VIRTUAL(UTransactor::Undo,return false;);
+
 	/**
 	 * Executes an redo transaction, redoing all actions contained by that transaction.
 	 * 
diff --git a/Engine/Source/Editor/UnrealEd/Classes/Factories/FbxMeshImportData.h b/Engine/Source/Editor/UnrealEd/Classes/Factories/FbxMeshImportData.h
--- a/Engine/Source/Editor/UnrealEd/Classes/Factories/FbxMeshImportData.h
+++ b/Engine/Source/Editor/UnrealEd/Classes/Factories/FbxMeshImportData.h
@@ -34,11 +34,11 @@ class UFbxMeshImportData : public UFbxAssetImportData
 	GENERATED_UCLASS_BODY()
 
 	/** If this option is true the node absolute transform (transform, offset and pivot) will be apply to the mesh vertices. */
-	UPROPERTY(EditAnywhere, AdvancedDisplay, config, Category = ImportSettings, meta = (ImportType = "Mesh"))
+	UPROPERTY(EditAnywhere, AdvancedDisplay, config, Category = ImportSettings, meta = (ImportType = "StaticMesh"))
 	bool bTransformVertexToAbsolute;
 
 	/** - Experimental - If this option is true the inverse node pivot will be apply to the mesh vertices. The pivot from the DCC will then be the origin of the mesh.*/
-	UPROPERTY(EditAnywhere, AdvancedDisplay, config, Category = ImportSettings, meta = (ImportType = "Mesh"))
+	UPROPERTY(EditAnywhere, AdvancedDisplay, config, Category = ImportSettings, meta = (ImportType = "StaticMesh"))
 	bool bBakePivotInVertex;
 
 	/** Enables importing of mesh LODs from FBX LOD groups, if present in the FBX file */
diff --git a/Engine/Source/Editor/UnrealEd/Classes/Factories/FbxSceneImportFactory.h b/Engine/Source/Editor/UnrealEd/Classes/Factories/FbxSceneImportFactory.h
--- a/Engine/Source/Editor/UnrealEd/Classes/Factories/FbxSceneImportFactory.h
+++ b/Engine/Source/Editor/UnrealEd/Classes/Factories/FbxSceneImportFactory.h
@@ -372,6 +372,8 @@ class UNREALED_API UFbxSceneImportFactory : public UFactory
 	static FString DefaultOptionName;
 
 protected:
+	/** Make sure GlobalImportSettings is pointing to the correct options */
+	void ApplyMeshInfoFbxOptions(TSharedPtr<FFbxMeshInfo> MeshInfo);
 
 	/* Compute the path of every node and fill the result in the node. This data will be use by the reimport
 	*  as a unique key for for the reimport status of the node hierarchy.
@@ -389,10 +391,7 @@ protected:
 	* LocalTransform: Must be a valid pointer
 	* PreMultiplyTransform: Can be nullptr
 	*/
-	void ApplyTransformToComponent(USceneComponent *SceneComponent, FTransform *LocalTransform, FTransform *PreMultiplyTransform, FVector &PivotLocation);
-
-	/** This will add the scene transform options to the root node */
-	void ApplySceneTransformOptionsToRootNode(TSharedPtr<FFbxSceneInfo> SceneInfoPtr);
+	void ApplyTransformToComponent(USceneComponent *SceneComponent, FTransform *LocalTransform, FTransform *PreMultiplyTransform, FVector &PivotLocation, FVector &ParentPivotAccumulation);
 
 	/** Import all skeletal mesh from the fbx scene */
 	void ImportAllSkeletalMesh(void* VoidRootNodeToImport, void* VoidFbxImporter, EObjectFlags Flags, int32& NodeIndex, int32& InterestingNodeCount , TSharedPtr<FFbxSceneInfo> SceneInfo);
diff --git a/Engine/Source/Editor/UnrealEd/Classes/Factories/FbxSceneImportOptions.h b/Engine/Source/Editor/UnrealEd/Classes/Factories/FbxSceneImportOptions.h
--- a/Engine/Source/Editor/UnrealEd/Classes/Factories/FbxSceneImportOptions.h
+++ b/Engine/Source/Editor/UnrealEd/Classes/Factories/FbxSceneImportOptions.h
@@ -30,32 +30,32 @@ class UFbxSceneImportOptions : public UObject
 	//UPROPERTY()
 	//FString ContentPath;
 
-	/** If check, a folders hierarchy will be create under the import asset path. All the create folders will match the actor hierarchy. In case there is more then one actor that link to an asset, the shared asset will be place at the first actor place.  */
+	/** If checked, a folder's hierarchy will be created under the import asset path. All the created folders will match the actor hierarchy. In case there is more than one actor that links to an asset, the shared asset will be placed at the first actor's place. */
 	UPROPERTY(EditAnywhere, config, category = ImportOptions)
 	uint32 bCreateContentFolderHierarchy : 1;
 
-	/** If checked, the mobility of all actors or component will be dynamic. If not it will be static  */
+	/** If checked, the mobility of all actors or components will be dynamic. If unchecked, they will be static. */
 	UPROPERTY(EditAnywhere, config, category = ImportOptions)
 	uint32 bImportAsDynamic : 1;
 
-	/** Choose if you want to generate the scene hierarchy with normal level actors or one actor with multiple components or one blueprint asset with multiple components.  */
+	/** Choose if you want to generate the scene hierarchy with normal level actors, one actor with multiple components, or one blueprint asset with multiple components. */
 	UPROPERTY(EditAnywhere, config, category = ImportOptions)
 	TEnumAsByte<enum EFBXSceneOptionsCreateHierarchyType> HierarchyType;
 
-	UPROPERTY(EditAnywhere, config, Category = Transform)
+	UPROPERTY()
 	FVector ImportTranslation;
 
-	UPROPERTY(EditAnywhere, config, Category = Transform)
+	UPROPERTY()
 	FRotator ImportRotation;
 
-	UPROPERTY(EditAnywhere, config, Category = Transform)
+	UPROPERTY()
 	float ImportUniformScale;
 
 	/** If this option is true the node absolute transform (transform, offset and pivot) will be apply to the mesh vertices. */
 	UPROPERTY()
 	bool bTransformVertexToAbsolute;
 
-	/** - Experimental - If this option is true the inverse node pivot will be apply to the mesh vertices. The pivot from the DCC will then be the origin of the mesh. */
+	/** - Experimental - If this option is true the inverse node pivot will be apply to the mesh vertices. The pivot from the DCC will then be the origin of the mesh. This option only work with static meshes.*/
 	UPROPERTY(EditAnywhere, config, category = Meshes)
 	bool bBakePivotInVertex;
 
@@ -67,8 +67,8 @@ class UFbxSceneImportOptions : public UObject
 	UPROPERTY(EditAnywhere, config, category = Meshes)
 	uint32 bImportSkeletalMeshLODs : 1;
 
-	/** If either importing of textures (or materials) is enabled, this option will cause normal map values to be inverted */
-	UPROPERTY(EditAnywhere, config, Category = Textures)
+	/** If importing textures is enabled, this option will cause normal map Y (Green) values to be inverted */
+	UPROPERTY(EditAnywhere, config, category = Material)
 	uint32 bInvertNormalMaps : 1;
 };
 
diff --git a/Engine/Source/Editor/UnrealEd/Classes/Factories/FbxSceneImportOptionsSkeletalMesh.h b/Engine/Source/Editor/UnrealEd/Classes/Factories/FbxSceneImportOptionsSkeletalMesh.h
--- a/Engine/Source/Editor/UnrealEd/Classes/Factories/FbxSceneImportOptionsSkeletalMesh.h
+++ b/Engine/Source/Editor/UnrealEd/Classes/Factories/FbxSceneImportOptionsSkeletalMesh.h
@@ -20,8 +20,8 @@ class UFbxSceneImportOptionsSkeletalMesh : public UObject
 	UPROPERTY(EditAnywhere, AdvancedDisplay, config, Category = SkeletalMesh)
 	uint32 bCreatePhysicsAsset : 1;
 
-	/** Enable this option to use frame 0 as reference pose */
-	UPROPERTY(EditAnywhere, config, Category = SkeletalMesh)
+	/** TODO support T0AsRefPose Enable this option to use frame 0 as reference pose */
+	UPROPERTY()
 	uint32 bUseT0AsRefPose : 1;
 
 	/** If checked, triangles with non-matching smoothing groups will be physically split. */
@@ -71,7 +71,7 @@ class UFbxSceneImportOptionsSkeletalMesh : public UObject
 	UPROPERTY(EditAnywhere, AdvancedDisplay, config, Category = Animation)
 	bool bDeleteExistingMorphTargetCurves;
 
-	void FillSkeletalMeshInmportData(class UFbxSkeletalMeshImportData* SkeletalMeshImportData, class UFbxAnimSequenceImportData* AnimSequenceImportData);
+	void FillSkeletalMeshInmportData(class UFbxSkeletalMeshImportData* SkeletalMeshImportData, class UFbxAnimSequenceImportData* AnimSequenceImportData, class UFbxSceneImportOptions* SceneImportOptions);
 };
 
 
diff --git a/Engine/Source/Editor/UnrealEd/Classes/Factories/FbxTextureImportData.h b/Engine/Source/Editor/UnrealEd/Classes/Factories/FbxTextureImportData.h
--- a/Engine/Source/Editor/UnrealEd/Classes/Factories/FbxTextureImportData.h
+++ b/Engine/Source/Editor/UnrealEd/Classes/Factories/FbxTextureImportData.h
@@ -11,7 +11,7 @@ class UNREALED_API UFbxTextureImportData : public UFbxAssetImportData
 {
 	GENERATED_UCLASS_BODY()
 
-	/** If either importing of textures (or materials) is enabled, this option will cause normal map values to be inverted */
+	/** If importing textures is enabled, this option will cause normal map Y (Green) values to be inverted */
 	UPROPERTY(EditAnywhere, AdvancedDisplay, config, Category=ImportSettings, meta=(OBJRestrict="true"))
 	uint32 bInvertNormalMaps:1;
 
diff --git a/Engine/Source/Editor/UnrealEd/Classes/Niagara/NiagaraNodeInput.h b/Engine/Source/Editor/UnrealEd/Classes/Niagara/NiagaraNodeInput.h
--- a/Engine/Source/Editor/UnrealEd/Classes/Niagara/NiagaraNodeInput.h
+++ b/Engine/Source/Editor/UnrealEd/Classes/Niagara/NiagaraNodeInput.h
@@ -28,7 +28,7 @@ public:
 	UNiagaraDataObject* DataObjectDefault;
 
 	/** Allows code to explicitly disable exposing of certain inputs e.g. system constants such as Delta Time. */
-	UPROPERTY()
+	UPROPERTY(EditAnywhere, Category=Constant, meta=(InlineEditConditionToggle))
 	uint32 bCanBeExposed:1;
 
 	/** When true, and this input is a constant, the input is exposed to the effect editor. */
diff --git a/Engine/Source/Editor/UnrealEd/Classes/Settings/EditorLoadingSavingSettings.h b/Engine/Source/Editor/UnrealEd/Classes/Settings/EditorLoadingSavingSettings.h
--- a/Engine/Source/Editor/UnrealEd/Classes/Settings/EditorLoadingSavingSettings.h
+++ b/Engine/Source/Editor/UnrealEd/Classes/Settings/EditorLoadingSavingSettings.h
@@ -116,8 +116,8 @@ public:
 	bool bAutoCreateAssets;
 	UPROPERTY(EditAnywhere, config, AdvancedDisplay, Category=AutoReimport, meta=(DisplayName="Auto Delete Assets", ToolTip="When enabled, deleting a source content file will automatically prompt the deletion of any related assets."))
 	bool bAutoDeleteAssets;
-	UPROPERTY(EditAnywhere, config, AdvancedDisplay, Category=AutoReimport, meta=(DisplayName="Detect Changes On Restart", ToolTip="When enabled, changes to monitored directories since UE4 was closed will be detected on restart.\n(Not recommended when working in collaboration with others using source control)."))
-	bool bDetectChangesOnRestart;
+	UPROPERTY(EditAnywhere, config, AdvancedDisplay, Category=AutoReimport, meta=(DisplayName="Detect Changes On Startup", ToolTip="When enabled, changes to monitored directories since UE4 was closed will be detected on restart.\n(Not recommended when working in collaboration with others using source control)."))
+	bool bDetectChangesOnStartup;
 
 	/** Internal setting to control whether we should ask the user whether we should automatically delete source files when their assets are deleted */
 	UPROPERTY(config)
diff --git a/Engine/Source/Editor/UnrealEd/Private/AutoReimport/ContentDirectoryMonitor.cpp b/Engine/Source/Editor/UnrealEd/Private/AutoReimport/ContentDirectoryMonitor.cpp
--- a/Engine/Source/Editor/UnrealEd/Private/AutoReimport/ContentDirectoryMonitor.cpp
+++ b/Engine/Source/Editor/UnrealEd/Private/AutoReimport/ContentDirectoryMonitor.cpp
@@ -33,7 +33,7 @@ DirectoryWatcher::FFileCacheConfig GenerateFileCacheConfig(const FString& InPath
 	// We always store paths inside content folders relative to the folder
 	Config.PathType = DirectoryWatcher::EPathType::Relative;
 
-	Config.bDetectChangesSinceLastRun = GetDefault<UEditorLoadingSavingSettings>()->bDetectChangesOnRestart;
+	Config.bDetectChangesSinceLastRun = GetDefault<UEditorLoadingSavingSettings>()->bDetectChangesOnStartup;
 
 	// We only detect changes for when the file *contents* have changed (not its timestamp)
 	Config
diff --git a/Engine/Source/Editor/UnrealEd/Private/Commandlets/GatherTextFromSourceCommandlet.cpp b/Engine/Source/Editor/UnrealEd/Private/Commandlets/GatherTextFromSourceCommandlet.cpp
--- a/Engine/Source/Editor/UnrealEd/Private/Commandlets/GatherTextFromSourceCommandlet.cpp
+++ b/Engine/Source/Editor/UnrealEd/Private/Commandlets/GatherTextFromSourceCommandlet.cpp
@@ -22,8 +22,6 @@ const FString UGatherTextFromSourceCommandlet::FPreProcessorDescriptor::ElIfStri
 const FString UGatherTextFromSourceCommandlet::FPreProcessorDescriptor::ElseString(TEXT("#else"));
 const FString UGatherTextFromSourceCommandlet::FPreProcessorDescriptor::EndIfString(TEXT("#endif"));
 const FString UGatherTextFromSourceCommandlet::FPreProcessorDescriptor::DefinedString(TEXT("defined "));
-const FString UGatherTextFromSourceCommandlet::FPreProcessorDescriptor::LocNamespaceString(TEXT("LOCTEXT_NAMESPACE"));
-const FString UGatherTextFromSourceCommandlet::FPreProcessorDescriptor::LocDefRegionString(TEXT("LOC_DEFINE_REGION"));
 const FString UGatherTextFromSourceCommandlet::FPreProcessorDescriptor::IniNamespaceString(TEXT("["));
 const FString UGatherTextFromSourceCommandlet::FMacroDescriptor::TextMacroString(TEXT("TEXT"));
 const FString UGatherTextFromSourceCommandlet::ChangelistName(TEXT("Update Localization"));
@@ -345,6 +343,94 @@ FString UGatherTextFromSourceCommandlet::RemoveStringFromTextMacro(const FString
 	return Text;
 }
 
+FString UGatherTextFromSourceCommandlet::StripCommentsFromToken(const FString& InToken, FSourceFileParseContext& Context)
+{
+	check(!Context.WithinBlockComment);
+	check(!Context.WithinLineComment);
+	check(!Context.WithinStringLiteral);
+
+	// Remove both block and inline comments from the given token
+	FString StrippedToken;
+	StrippedToken.Reserve(InToken.Len());
+
+	TCHAR WithinQuote = 0;
+	bool bIgnoreNextQuote = false;
+	for (const TCHAR* Char = *InToken; *Char; ++Char)
+	{
+		if (WithinQuote != 0)
+		{
+			StrippedToken += *Char;
+
+			if (!bIgnoreNextQuote)
+			{
+				if (*Char == TEXT('\\'))
+				{
+					bIgnoreNextQuote = true;
+					continue;
+				}
+
+				if (*Char == WithinQuote)
+				{
+					// Found an unescaped closing quote - we are no longer within quotes
+					WithinQuote = 0;
+				}
+			}
+
+			bIgnoreNextQuote = false;
+		}
+		else
+		{
+			if (*Char == TEXT('/'))
+			{
+				const TCHAR* NextChar = Char + 1;
+
+				if (*NextChar == TEXT('/'))
+				{
+					// Found an inline quote - this strips the remainder of the string so just break out of the loop
+					break;
+				}
+
+				if (*NextChar == TEXT('*'))
+				{
+					// Found a block comment - skip all characters until we find the closing quote
+					Context.WithinBlockComment = true;
+					++Char; // Skip over the opening slash, and the for loop will skip over the *
+					continue;
+				}
+			}
+
+			if (Context.WithinBlockComment)
+			{
+				if (*Char == TEXT('*'))
+				{
+					const TCHAR* NextChar = Char + 1;
+
+					if (*NextChar == TEXT('/'))
+					{
+						// Found the end of a block comment
+						Context.WithinBlockComment = false;
+						++Char; // Skip over the opening *, and the for loop will skip over the slash
+						continue;
+					}
+				}
+
+				// Skip over all characters while within a block comment
+				continue;
+			}
+
+			StrippedToken += *Char;
+
+			if (*Char == TEXT('"') || *Char == TEXT('\''))
+			{
+				// We found an opening quote - keep track of it until we find a matching closing quote
+				WithinQuote = *Char;
+			}
+		}
+	}
+
+	return StrippedToken.Trim().TrimTrailing();
+}
+
 bool UGatherTextFromSourceCommandlet::ParseSourceText(const FString& Text, const TArray<UGatherTextFromSourceCommandlet::FParsableDescriptor*>& Parsables, FSourceFileParseContext& ParseCtxt)
 {
 	// Create array of ints, one for each parsable we're looking for.
@@ -631,104 +717,139 @@ UGatherTextFromSourceCommandlet::EMacroBlockState UGatherTextFromSourceCommandle
 	return CachedMacroBlockState.GetValue();
 }
 
-void UGatherTextFromSourceCommandlet::FDefineDescriptor::TryParse(const FString& Text, FSourceFileParseContext& Context) const
+void UGatherTextFromSourceCommandlet::FSourceFileParseContext::SetDefine(const FString& InDefineCtx)
 {
-	// Attempt to parse something of the format
-	// #define <defname>
-	//  or
-	// #define <defname> <value>
+	static const FString LocDefRegionString = TEXT("LOC_DEFINE_REGION");
+	static const FString LocNamespaceString = TEXT("LOCTEXT_NAMESPACE");
 
-	if (!Context.ExcludedRegion && (Context.Filename.EndsWith(TEXT(".inl")) || (!Context.WithinBlockComment && !Context.WithinLineComment && !Context.WithinStringLiteral)) )
+	if (InDefineCtx.Equals(LocDefRegionString, ESearchCase::CaseSensitive))
 	{
-		FString RemainingText = Text.RightChop(GetToken().Len()).Trim();
-
-		if (LocDefRegionString == RemainingText)
+		// #define LOC_DEFINE_REGION
+		if (ExcludedRegion)
 		{
-			// #define LOC_DEFINE_REGION
-			Context.ExcludedRegion = true;
-			Context.EndParsingCurrentLine = true;
+			UE_LOG(LogGatherTextFromSourceCommandlet, Warning, TEXT("Found a '#define LOC_DEFINE_REGION' while still within another '#define LOC_DEFINE_REGION'. File %s at line %d"), *Filename, LineNumber);
 		}
-
-		FString FoundNamespaceString;
-
-		if (RemainingText.StartsWith(LocNamespaceString, ESearchCase::CaseSensitive))
+		else
+		{
+			ExcludedRegion = true;
+		}
+		return;
+	}
+	else if (!ExcludedRegion)
+	{
+		if (InDefineCtx.StartsWith(LocNamespaceString, ESearchCase::CaseSensitive) && InDefineCtx.IsValidIndex(LocNamespaceString.Len()) && (FText::IsWhitespace(InDefineCtx[LocNamespaceString.Len()]) || InDefineCtx[LocNamespaceString.Len()] == TEXT('"')))
 		{
 			// #define LOCTEXT_NAMESPACE <namespace>
-			FoundNamespaceString = LocNamespaceString;
-			if (Context.WithinNamespaceDefine == true)
+			if (WithinNamespaceDefine)
 			{
-				UE_LOG(LogGatherTextFromSourceCommandlet, Warning, TEXT("Non-matching LOCTEXT_NAMESPACE defines found in %s"), *Context.Filename);
+				UE_LOG(LogGatherTextFromSourceCommandlet, Warning, TEXT("Found a '#define LOCTEXT_NAMESPACE' while still within another '#define LOCTEXT_NAMESPACE'. File %s at line %d"), *Filename, LineNumber);
 			}
+			else
+			{
+				FString RemainingText = InDefineCtx.RightChop(LocNamespaceString.Len()).Trim();
 
-			Context.WithinNamespaceDefine = true;
-		}
-
-		if( !FoundNamespaceString.IsEmpty() )
-		{
-			RemainingText = RemainingText.RightChop(FoundNamespaceString.Len()).Trim();
+				bool RemoveStringError;
+				const FString DefineDesc = FString::Printf(TEXT("%s define %s(%d):%s"), *RemainingText, *Filename, LineNumber, *LineText);
+				FString NewNamespace = RemoveStringFromTextMacro(RemainingText, DefineDesc, RemoveStringError);
 
-			bool RemoveStringError;
-			FString DefineDesc = FString::Printf(TEXT("%s define %s(%d):%s"), *FoundNamespaceString, *Context.Filename, Context.LineNumber, *Context.LineText);
-			FString Namespace = RemoveStringFromTextMacro(RemainingText, DefineDesc, RemoveStringError);
-			if (!RemoveStringError)
-			{
-				Context.Namespace = Namespace;
-				Context.EndParsingCurrentLine = true;
+				if (!RemoveStringError)
+				{
+					Namespace = MoveTemp(NewNamespace);
+					WithinNamespaceDefine = true;
+				}
 			}
+			return;
 		}
 	}
 }
 
-void UGatherTextFromSourceCommandlet::FUndefDescriptor::TryParse(const FString& Text, FSourceFileParseContext& Context) const
+void UGatherTextFromSourceCommandlet::FSourceFileParseContext::RemoveDefine(const FString& InDefineCtx)
 {
-	// Attempt to parse something of the format
-	// #undef <defname>
+	static const FString LocDefRegionString = TEXT("LOC_DEFINE_REGION");
+	static const FString LocNamespaceString = TEXT("LOCTEXT_NAMESPACE");
 
-	FString RemainingText = Text.RightChop(GetToken().Len()).Trim();
-
-	if (!Context.ExcludedRegion && (Context.Filename.EndsWith(TEXT(".inl")) || (!Context.WithinBlockComment && !Context.WithinLineComment && !Context.WithinStringLiteral)))
+	if (InDefineCtx.Equals(LocDefRegionString, ESearchCase::CaseSensitive))
 	{
-		if (Context.ExcludedRegion)
+		// #undef LOC_DEFINE_REGION
+		if (!ExcludedRegion)
 		{
-			if (LocDefRegionString == RemainingText)
-			{
-				// #undef LOC_DEFINE_REGION
-				Context.ExcludedRegion = false;
-				Context.EndParsingCurrentLine = true;
-			}
+			UE_LOG(LogGatherTextFromSourceCommandlet, Warning, TEXT("Found an '#undef LOC_DEFINE_REGION' without a corresponding '#define LOC_DEFINE_REGION'. File %s at line %d"), *Filename, LineNumber);
 		}
 		else
 		{
-			if (LocNamespaceString == RemainingText)
+			ExcludedRegion = false;
+		}
+		return;
+	}
+	else if (!ExcludedRegion)
+	{
+		if (InDefineCtx.Equals(LocNamespaceString, ESearchCase::CaseSensitive))
+		{
+			// #undef LOCTEXT_NAMESPACE
+			if (!WithinNamespaceDefine)
 			{
-				Context.EndParsingCurrentLine = true;
-				Context.Namespace.Empty();
-				if (Context.WithinNamespaceDefine == false)
-				{
-					UE_LOG(LogGatherTextFromSourceCommandlet, Warning, TEXT("Non-matching LOCTEXT_NAMESPACE defines found in %s"), *Context.Filename);
-				}
-				Context.WithinNamespaceDefine = false;
+				UE_LOG(LogGatherTextFromSourceCommandlet, Warning, TEXT("Found an '#undef LOCTEXT_NAMESPACE' without a corresponding '#define LOCTEXT_NAMESPACE'. File %s at line %d"), *Filename, LineNumber);
 			}
+			else
+			{
+				Namespace.Empty();
+				WithinNamespaceDefine = false;
+			}
+			return;
 		}
 	}
 }
 
+void UGatherTextFromSourceCommandlet::FDefineDescriptor::TryParse(const FString& Text, FSourceFileParseContext& Context) const
+{
+	// Attempt to parse something of the format
+	// #define <defname>
+	//  or
+	// #define <defname> <value>
+
+	if (!Context.WithinBlockComment && !Context.WithinLineComment && !Context.WithinStringLiteral)
+	{
+		FString RemainingText = Text.RightChop(GetToken().Len()).Trim();
+		RemainingText = StripCommentsFromToken(RemainingText, Context);
+
+		Context.SetDefine(RemainingText);
+		Context.EndParsingCurrentLine = true;
+	}
+}
+
+void UGatherTextFromSourceCommandlet::FUndefDescriptor::TryParse(const FString& Text, FSourceFileParseContext& Context) const
+{
+	// Attempt to parse something of the format
+	// #undef <defname>
+
+	if (!Context.WithinBlockComment && !Context.WithinLineComment && !Context.WithinStringLiteral)
+	{
+		FString RemainingText = Text.RightChop(GetToken().Len()).Trim();
+		RemainingText = StripCommentsFromToken(RemainingText, Context);
+
+		Context.RemoveDefine(RemainingText);
+		Context.EndParsingCurrentLine = true;
+	}
+}
+
 void UGatherTextFromSourceCommandlet::FIfDescriptor::TryParse(const FString& Text, FSourceFileParseContext& Context) const
 {
 	// Attempt to parse something of the format
 	// #if <defname>
 
-	FString RemainingText = Text.RightChop(GetToken().Len()).Trim();
-
 	if (!Context.WithinBlockComment && !Context.WithinLineComment && !Context.WithinStringLiteral)
 	{
+		FString RemainingText = Text.RightChop(GetToken().Len()).Trim();
+		RemainingText = StripCommentsFromToken(RemainingText, Context);
+
 		// Handle "#if defined <defname>"
 		if (RemainingText.StartsWith(DefinedString, ESearchCase::CaseSensitive))
 		{
-			RemainingText = Text.RightChop(DefinedString.Len()).Trim();
+			RemainingText = RemainingText.RightChop(DefinedString.Len()).Trim();
 		}
 
 		Context.PushMacroBlock(RemainingText);
+		Context.EndParsingCurrentLine = true;
 	}
 }
 
@@ -737,11 +858,13 @@ void UGatherTextFromSourceCommandlet::FIfDefDescriptor::TryParse(const FString&
 	// Attempt to parse something of the format
 	// #ifdef <defname>
 
-	FString RemainingText = Text.RightChop(GetToken().Len()).Trim();
-
 	if (!Context.WithinBlockComment && !Context.WithinLineComment && !Context.WithinStringLiteral)
 	{
+		FString RemainingText = Text.RightChop(GetToken().Len()).Trim();
+		RemainingText = StripCommentsFromToken(RemainingText, Context);
+
 		Context.PushMacroBlock(RemainingText);
+		Context.EndParsingCurrentLine = true;
 	}
 }
 
@@ -750,18 +873,20 @@ void UGatherTextFromSourceCommandlet::FElIfDescriptor::TryParse(const FString& T
 	// Attempt to parse something of the format
 	// #elif <defname>
 
-	FString RemainingText = Text.RightChop(GetToken().Len()).Trim();
-
 	if (!Context.WithinBlockComment && !Context.WithinLineComment && !Context.WithinStringLiteral)
 	{
+		FString RemainingText = Text.RightChop(GetToken().Len()).Trim();
+		RemainingText = StripCommentsFromToken(RemainingText, Context);
+
 		// Handle "#elif defined <defname>"
 		if (RemainingText.StartsWith(DefinedString, ESearchCase::CaseSensitive))
 		{
-			RemainingText = Text.RightChop(DefinedString.Len()).Trim();
+			RemainingText = RemainingText.RightChop(DefinedString.Len()).Trim();
 		}
 
 		Context.PopMacroBlock(); // Pop the current #if or #ifdef state
 		Context.PushMacroBlock(RemainingText);
+		Context.EndParsingCurrentLine = true;
 	}
 }
 
@@ -774,6 +899,7 @@ void UGatherTextFromSourceCommandlet::FElseDescriptor::TryParse(const FString& T
 	{
 		Context.PopMacroBlock(); // Pop the current #if or #ifdef state
 		Context.PushMacroBlock(FString());
+		Context.EndParsingCurrentLine = true;
 	}
 }
 
@@ -785,6 +911,7 @@ void UGatherTextFromSourceCommandlet::FEndIfDescriptor::TryParse(const FString&
 	if (!Context.WithinBlockComment && !Context.WithinLineComment && !Context.WithinStringLiteral)
 	{
 		Context.PopMacroBlock(); // Pop the current #if or #ifdef state
+		Context.EndParsingCurrentLine = true;
 	}
 }
 
@@ -903,10 +1030,10 @@ void UGatherTextFromSourceCommandlet::FCommandMacroDescriptor::TryParse(const FS
 	// Attempt to parse something of the format
 	// UI_COMMAND(LocKey, DefaultLangString, DefaultLangTooltipString, <IgnoredParam>, <IgnoredParam>)
 
-	if (!Context.ExcludedRegion && (Context.Filename.EndsWith(TEXT(".inl")) || (!Context.WithinBlockComment && !Context.WithinLineComment && !Context.WithinStringLiteral)) )
+	if (!Context.ExcludedRegion && !Context.WithinBlockComment && !Context.WithinLineComment && !Context.WithinStringLiteral)
 	{
 		TArray<FString> Arguments;
-		if (ParseArgsFromMacro(Text, Arguments, Context))
+		if (ParseArgsFromMacro(StripCommentsFromToken(Text, Context), Arguments, Context))
 		{
 			if (Arguments.Num() != 5)
 			{
@@ -968,10 +1095,10 @@ void UGatherTextFromSourceCommandlet::FStringMacroDescriptor::TryParse(const FSt
 	// Attempt to parse something of the format
 	// MACRONAME(param0, param1, param2)
 
-	if (!Context.ExcludedRegion && (Context.Filename.EndsWith(TEXT(".inl")) || (!Context.WithinBlockComment && !Context.WithinLineComment && !Context.WithinStringLiteral)) )
+	if (!Context.ExcludedRegion && !Context.WithinBlockComment && !Context.WithinLineComment && !Context.WithinStringLiteral)
 	{
 		TArray<FString> ArgArray;
-		if (ParseArgsFromMacro(Text, ArgArray, Context))
+		if (ParseArgsFromMacro(StripCommentsFromToken(Text, Context), ArgArray, Context))
 		{
 			int32 NumArgs = ArgArray.Num();
 
diff --git a/Engine/Source/Editor/UnrealEd/Private/Commandlets/InternationalizationExportCommandlet.cpp b/Engine/Source/Editor/UnrealEd/Private/Commandlets/InternationalizationExportCommandlet.cpp
--- a/Engine/Source/Editor/UnrealEd/Private/Commandlets/InternationalizationExportCommandlet.cpp
+++ b/Engine/Source/Editor/UnrealEd/Private/Commandlets/InternationalizationExportCommandlet.cpp
@@ -413,7 +413,7 @@ bool UInternationalizationExportCommandlet::DoExport( const FString& SourcePath,
 										{
 											const FString KeyName = InfoMetaDataPair.Key;
 											const TSharedPtr<FLocMetadataValue> Value = InfoMetaDataPair.Value;
-											InfoMetaDataStrings.Add(GetConditionedInfoMetaDataForExtractedComment(KeyName, Value->AsString()));
+											InfoMetaDataStrings.Add(GetConditionedInfoMetaDataForExtractedComment(KeyName, Value->ToString()));
 										}
 									}
 									if (InfoMetaDataStrings.Num())
@@ -469,7 +469,7 @@ bool UInternationalizationExportCommandlet::DoExport( const FString& SourcePath,
 												{
 													const FString KeyName = InfoMetaDataPair.Key;
 													const TSharedPtr<FLocMetadataValue> Value = InfoMetaDataPair.Value;
-													InfoMetaDataStrings.Add(FString::Printf(TEXT("InfoMetaData:\t\"%s\" : \"%s\""), *KeyName, *Value->AsString()));
+													InfoMetaDataStrings.Add(FString::Printf(TEXT("InfoMetaData:\t\"%s\" : \"%s\""), *KeyName, *Value->ToString()));
 												}
 											}
 											if (InfoMetaDataStrings.Num())
diff --git a/Engine/Source/Editor/UnrealEd/Private/Editor.cpp b/Engine/Source/Editor/UnrealEd/Private/Editor.cpp
--- a/Engine/Source/Editor/UnrealEd/Private/Editor.cpp
+++ b/Engine/Source/Editor/UnrealEd/Private/Editor.cpp
@@ -223,7 +223,7 @@ void FReimportManager::UpdateReimportPaths( UObject* Obj, const TArray<FString>&
 	}
 }
 
-bool FReimportManager::Reimport( UObject* Obj, bool bAskForNewFileIfMissing, bool bShowNotification )
+bool FReimportManager::Reimport( UObject* Obj, bool bAskForNewFileIfMissing, bool bShowNotification, FString PreferedReimportFile)
 {
 	// Warn that were about to reimport, so prep for it
 	PreReimport.Broadcast( Obj );
@@ -265,9 +265,17 @@ bool FReimportManager::Reimport( UObject* Obj, bool bAskForNewFileIfMissing, boo
 				}
 
 				bValidSourceFilename = true;
-				if ( bAskForNewFileIfMissing && bMissingFiles )
+				if ((bAskForNewFileIfMissing || !PreferedReimportFile.IsEmpty()) && bMissingFiles )
 				{
-					GetNewReimportPath(Obj, SourceFilenames);
+					if (!bAskForNewFileIfMissing && !PreferedReimportFile.IsEmpty())
+					{
+						SourceFilenames.Empty();
+						SourceFilenames.Add(PreferedReimportFile);
+					}
+					else
+					{
+						GetNewReimportPath(Obj, SourceFilenames);
+					}
 					if ( SourceFilenames.Num() == 0 )
 					{
 						// Failed to specify a new filename. Don't show a notification of the failure since the user exited on his own
diff --git a/Engine/Source/Editor/UnrealEd/Private/EditorActor.cpp b/Engine/Source/Editor/UnrealEd/Private/EditorActor.cpp
--- a/Engine/Source/Editor/UnrealEd/Private/EditorActor.cpp
+++ b/Engine/Source/Editor/UnrealEd/Private/EditorActor.cpp
@@ -1286,7 +1286,8 @@ void UUnrealEdEngine::edactHideSelectedStartup( UWorld* InWorld )
 				FBspSurf& CurSurface = *SurfaceIterator;
 				
 				// Set the BSP surface to hide at editor startup, if it's not already set that way
-				if ( ( CurSurface.PolyFlags & PF_Selected ) && !CurSurface.IsHiddenEdAtStartup() && !CurSurface.IsHiddenEd() )
+				const bool bSelected = CurSurface.Actor->IsSelected() || (CurSurface.PolyFlags & PF_Selected);
+				if (bSelected && !CurSurface.IsHiddenEdAtStartup() && !CurSurface.IsHiddenEd())
 				{
 					CurLevelModel.Modify();
 					CurLevelModel.ModifySurf( SurfaceIterator.GetIndex(), false );
@@ -1376,7 +1377,8 @@ void UUnrealEdEngine::edactUnHideSelectedStartup( UWorld* InWorld )
 				FBspSurf& CurSurface = *SurfaceIterator;
 
 				// Mark the selected BSP surface as showing at editor startup if it was currently set to be hidden
-				if ( ( CurSurface.PolyFlags & PF_Selected ) && CurSurface.IsHiddenEdAtStartup() )
+				const bool bSelected = CurSurface.Actor->IsSelected() || (CurSurface.PolyFlags & PF_Selected);
+				if (bSelected && CurSurface.IsHiddenEdAtStartup())
 				{
 					CurLevelModel.Modify();
 					CurLevelModel.ModifySurf( SurfaceIterator.GetIndex(), false );
diff --git a/Engine/Source/Editor/UnrealEd/Private/EditorTransaction.cpp b/Engine/Source/Editor/UnrealEd/Private/EditorTransaction.cpp
--- a/Engine/Source/Editor/UnrealEd/Private/EditorTransaction.cpp
+++ b/Engine/Source/Editor/UnrealEd/Private/EditorTransaction.cpp
@@ -670,6 +670,20 @@ bool UTransBuffer::CanUndo( FText* Text )
 		}
 		return false;
 	}
+	
+	if (UndoBarrierStack.Num())
+	{
+		const int32 UndoBarrier = UndoBarrierStack.Last();
+		if (UndoBuffer.Num() - UndoCount <= UndoBarrier)
+		{
+			if (Text)
+			{
+				*Text = NSLOCTEXT("TransactionSystem", "HitUndoBarrier", "(Hit Undo barrier; can't undo any further)");
+			}
+			return false;
+		}
+	}
+
 	if( UndoBuffer.Num()==UndoCount )
 	{
 		if( Text )
@@ -746,6 +760,27 @@ FUndoSessionContext UTransBuffer::GetRedoContext()
 }
 
 
+void UTransBuffer::SetUndoBarrier()
+{
+	UndoBarrierStack.Push(UndoBuffer.Num() - UndoCount);
+}
+
+
+void UTransBuffer::RemoveUndoBarrier()
+{
+	if (UndoBarrierStack.Num() > 0)
+	{
+		UndoBarrierStack.Pop();
+	}
+}
+
+
+void UTransBuffer::ClearUndoBarriers()
+{
+	UndoBarrierStack.Empty();
+}
+
+
 bool UTransBuffer::Undo()
 {
 	CheckState();
