Copying //UE4/Dev-Editor to Dev-Main (//UE4/Dev-Main)

diff --git a/Engine/Source/Editor/UnrealEd/Private/Factories/ApexDestructibleAssetImport.cpp b/Engine/Source/Editor/UnrealEd/Private/Factories/ApexDestructibleAssetImport.cpp
--- a/Engine/Source/Editor/UnrealEd/Private/Factories/ApexDestructibleAssetImport.cpp
+++ b/Engine/Source/Editor/UnrealEd/Private/Factories/ApexDestructibleAssetImport.cpp
@@ -37,7 +37,7 @@ DEFINE_LOG_CATEGORY_STATIC(LogApexDestructibleAssetImport, Log, All);
 
 // Forward declarations and external processing functions
 struct ExistingSkelMeshData;
-extern ExistingSkelMeshData* SaveExistingSkelMeshData(USkeletalMesh* ExistingSkelMesh);
+extern ExistingSkelMeshData* SaveExistingSkelMeshData(USkeletalMesh* ExistingSkelMesh, bool bSaveMaterials);
 extern void RestoreExistingSkelMeshData(ExistingSkelMeshData* MeshData, USkeletalMesh* SkeletalMesh);
 extern void ProcessImportMeshInfluences(FSkeletalMeshImportData& ImportData);
 extern void ProcessImportMeshMaterials(TArray<FSkeletalMaterial>& Materials, FSkeletalMeshImportData& ImportData);
@@ -236,7 +236,7 @@ ExistingDestMeshData* SaveExistingDestMeshData(UDestructibleMesh* ExistingDestru
 		// Only save off SkelMeshData if it's been created
 		ExistingDestMeshDataPtr->SkelMeshData = NULL;
 		
-		ExistingDestMeshDataPtr->SkelMeshData = SaveExistingSkelMeshData(ExistingDestructibleMesh);
+		ExistingDestMeshDataPtr->SkelMeshData = SaveExistingSkelMeshData(ExistingDestructibleMesh, true);
 		
 		ExistingDestMeshDataPtr->BodySetup = ExistingDestructibleMesh->BodySetup;
 		ExistingDestMeshDataPtr->FractureEffects = ExistingDestructibleMesh->FractureEffects;
diff --git a/Engine/Source/Editor/UnrealEd/Private/Factories/CSVImportFactory.cpp b/Engine/Source/Editor/UnrealEd/Private/Factories/CSVImportFactory.cpp
--- a/Engine/Source/Editor/UnrealEd/Private/Factories/CSVImportFactory.cpp
+++ b/Engine/Source/Editor/UnrealEd/Private/Factories/CSVImportFactory.cpp
@@ -391,6 +391,7 @@ int32 UReimportCurveTableFactory::GetPriority() const
 UReimportCurveFactory::UReimportCurveFactory(const FObjectInitializer& ObjectInitializer)
 	: Super(ObjectInitializer)
 {
+	SupportedClass = UCurveBase::StaticClass();
 }
 
 bool UReimportCurveFactory::CanReimport( UObject* Obj, TArray<FString>& OutFilenames )
diff --git a/Engine/Source/Editor/UnrealEd/Private/Factories/EditorFactories.cpp b/Engine/Source/Editor/UnrealEd/Private/Factories/EditorFactories.cpp
--- a/Engine/Source/Editor/UnrealEd/Private/Factories/EditorFactories.cpp
+++ b/Engine/Source/Editor/UnrealEd/Private/Factories/EditorFactories.cpp
@@ -5607,6 +5607,8 @@ EReimportResult::Type UReimportFbxStaticMeshFactory::Reimport( UObject* Obj )
 	
 	UnFbx::FFbxImporter* FFbxImporter = UnFbx::FFbxImporter::GetInstance();
 	UnFbx::FBXImportOptions* ImportOptions = FFbxImporter->GetImportOptions();
+	//Clean up the options
+	UnFbx::FBXImportOptions::ResetOptions(ImportOptions);
 
 	UFbxStaticMeshImportData* ImportData = Cast<UFbxStaticMeshImportData>(Mesh->AssetImportData);
 	
@@ -5817,6 +5819,8 @@ EReimportResult::Type UReimportFbxSkeletalMeshFactory::Reimport( UObject* Obj )
 
 	UnFbx::FFbxImporter* FFbxImporter = UnFbx::FFbxImporter::GetInstance();
 	UnFbx::FBXImportOptions* ImportOptions = FFbxImporter->GetImportOptions();
+	//Clean up the options
+	UnFbx::FBXImportOptions::ResetOptions(ImportOptions);
 
 	UFbxSkeletalMeshImportData* ImportData = Cast<UFbxSkeletalMeshImportData>(SkeletalMesh->AssetImportData);
 	
@@ -5837,6 +5841,10 @@ EReimportResult::Type UReimportFbxSkeletalMeshFactory::Reimport( UObject* Obj )
 	{
 		// Import data already exists, apply it to the fbx import options
 		ReimportUI->SkeletalMeshImportData = ImportData;
+		//Some options not supported with skeletal mesh
+		ReimportUI->SkeletalMeshImportData->bBakePivotInVertex = false;
+		ReimportUI->SkeletalMeshImportData->bTransformVertexToAbsolute = true;
+
 		ApplyImportUIToImportOptions(ReimportUI, *ImportOptions);
 	}
 	else
@@ -5964,7 +5972,21 @@ bool UReimportFbxAnimSequenceFactory::CanReimport( UObject* Obj, TArray<FString>
 	UAnimSequence* AnimSequence = Cast<UAnimSequence>(Obj);
 	if(AnimSequence)
 	{
-		AnimSequence->AssetImportData->ExtractFilenames(OutFilenames);
+		if (AnimSequence->AssetImportData)
+		{
+			AnimSequence->AssetImportData->ExtractFilenames(OutFilenames);
+
+			UFbxAssetImportData *FbxAssetImportData = Cast<UFbxAssetImportData>(AnimSequence->AssetImportData);
+			if (FbxAssetImportData != nullptr && FbxAssetImportData->bImportAsScene)
+			{
+				//This mesh was import with a scene import, we cannot reimport it
+				return false;
+			}
+		}
+		else
+		{
+			OutFilenames.Add(TEXT(""));
+		}
 		return true;
 	}
 	return false;
@@ -6023,6 +6045,7 @@ EReimportResult::Type UReimportFbxAnimSequenceFactory::Reimport( UObject* Obj )
 	}
 
 	UnFbx::FFbxImporter* Importer = UnFbx::FFbxImporter::GetInstance();
+	UnFbx::FBXImportOptions::ResetOptions(Importer->ImportOptions);
 
 	CurrentFilename = Filename;
 
diff --git a/Engine/Source/Editor/UnrealEd/Private/Factories/SkeletalMeshImport.cpp b/Engine/Source/Editor/UnrealEd/Private/Factories/SkeletalMeshImport.cpp
--- a/Engine/Source/Editor/UnrealEd/Private/Factories/SkeletalMeshImport.cpp
+++ b/Engine/Source/Editor/UnrealEd/Private/Factories/SkeletalMeshImport.cpp
@@ -283,6 +283,31 @@ void ProcessImportMeshInfluences(FSkeletalMeshImportData& ImportData)
 	float TotalWeight		= 0.f;
 	const float MINWEIGHT   = 0.01f;
 
+	//We have to normalize the data before filtering influences
+	//Because influence filtering is base on the normalize value.
+	//Some DCC like Daz studio don't have normalized weight
+	for (int32 i = 0; i < Influences.Num(); i++)
+	{
+		// if less than min weight, or it's more than 8, then we clear it to use weight
+		InfluenceCount++;
+		TotalWeight += Influences[i].Weight;
+		// we have all influence for the same vertex, normalize it now
+		if (i + 1 >= Influences.Num() || Influences[i].VertexIndex != Influences[i+1].VertexIndex)
+		{
+			// Normalize the last set of influences.
+			if (InfluenceCount && (TotalWeight != 1.0f))
+			{
+				float OneOverTotalWeight = 1.f / TotalWeight;
+				for (int r = 0; r < InfluenceCount; r++)
+				{
+					Influences[i - r].Weight *= OneOverTotalWeight;
+				}
+			}
+			// clear to count next one
+			InfluenceCount = 0;
+			TotalWeight = 0.f;
+		}
+	}
 
 	for( int32 i=0; i<Influences.Num(); i++ )
 	{
@@ -691,7 +716,7 @@ void FSavedCustomSortInfo::Restore(USkeletalMesh* NewSkeletalMesh, int32 LODMode
 
 
 
-ExistingSkelMeshData* SaveExistingSkelMeshData(USkeletalMesh* ExistingSkelMesh)
+ExistingSkelMeshData* SaveExistingSkelMeshData(USkeletalMesh* ExistingSkelMesh, bool bSaveMaterials)
 {
 	struct ExistingSkelMeshData* ExistingMeshDataPtr = NULL;
 
@@ -706,7 +731,11 @@ ExistingSkelMeshData* SaveExistingSkelMeshData(USkeletalMesh* ExistingSkelMesh)
 		}
 
 		ExistingMeshDataPtr->ExistingSockets = ExistingSkelMesh->GetMeshOnlySocketList();
-		ExistingMeshDataPtr->ExistingMaterials = ExistingSkelMesh->Materials;
+		ExistingMeshDataPtr->bSaveRestoreMaterials = bSaveMaterials;
+		if (ExistingMeshDataPtr->bSaveRestoreMaterials)
+		{
+			ExistingMeshDataPtr->ExistingMaterials = ExistingSkelMesh->Materials;
+		}
 		ExistingMeshDataPtr->ExistingRetargetBasePose = ExistingSkelMesh->RetargetBasePose;
 
 		if( ImportedResource->LODModels.Num() > 0 &&
@@ -825,16 +854,19 @@ void RestoreExistingSkelMeshData(ExistingSkelMeshData* MeshData, USkeletalMesh*
 	if(MeshData && SkeletalMesh)
 	{
 		// Fix Materials array to be the correct size.
-		if(MeshData->ExistingMaterials.Num() > SkeletalMesh->Materials.Num())
-		{
-			MeshData->ExistingMaterials.RemoveAt( SkeletalMesh->Materials.Num(), MeshData->ExistingMaterials.Num() - SkeletalMesh->Materials.Num() );
-		}
-		else if(SkeletalMesh->Materials.Num() > MeshData->ExistingMaterials.Num())
+		if (MeshData->bSaveRestoreMaterials)
 		{
-			MeshData->ExistingMaterials.AddZeroed( SkeletalMesh->Materials.Num() - MeshData->ExistingMaterials.Num() );
-		}
+			if (MeshData->ExistingMaterials.Num() > SkeletalMesh->Materials.Num())
+			{
+				MeshData->ExistingMaterials.RemoveAt(SkeletalMesh->Materials.Num(), MeshData->ExistingMaterials.Num() - SkeletalMesh->Materials.Num());
+			}
+			else if (SkeletalMesh->Materials.Num() > MeshData->ExistingMaterials.Num())
+			{
+				MeshData->ExistingMaterials.AddZeroed(SkeletalMesh->Materials.Num() - MeshData->ExistingMaterials.Num());
+			}
 
-		SkeletalMesh->Materials = MeshData->ExistingMaterials;
+			SkeletalMesh->Materials = MeshData->ExistingMaterials;
+		}
 
 		// this is not ideal. Ideally we'll have to save only diff with indicating which joints, 
 		// but for now, we allow them to keep the previous pose IF the element count is same
diff --git a/Engine/Source/Editor/UnrealEd/Private/Fbx/FbxFactory.cpp b/Engine/Source/Editor/UnrealEd/Private/Fbx/FbxFactory.cpp
--- a/Engine/Source/Editor/UnrealEd/Private/Fbx/FbxFactory.cpp
+++ b/Engine/Source/Editor/UnrealEd/Private/Fbx/FbxFactory.cpp
@@ -189,6 +189,9 @@ UObject* UFbxFactory::FactoryCreateBinary
 	// logger for all error/warnings
 	// this one prints all messages that are stored in FFbxImporter
 	UnFbx::FFbxImporter* FbxImporter = UnFbx::FFbxImporter::GetInstance();
+	UnFbx::FBXImportOptions* ImportOptions = FbxImporter->GetImportOptions();
+	//Clean up the options
+	UnFbx::FBXImportOptions::ResetOptions(ImportOptions);
 	
 	UnFbx::FFbxLoggerSetter Logger(FbxImporter);
 
@@ -203,7 +206,7 @@ UObject* UFbxFactory::FactoryCreateBinary
 
 	bool bShowImportDialog = bShowOption && !GIsAutomationTesting;
 	bool bImportAll = false;
-	UnFbx::FBXImportOptions* ImportOptions = GetImportOptions(FbxImporter, ImportUI, bShowImportDialog, InParent->GetPathName(), bOperationCanceled, bImportAll, bIsObjFormat, bIsObjFormat, ForcedImportType );
+	ImportOptions = GetImportOptions(FbxImporter, ImportUI, bShowImportDialog, InParent->GetPathName(), bOperationCanceled, bImportAll, bIsObjFormat, bIsObjFormat, ForcedImportType );
 	bOutOperationCanceled = bOperationCanceled;
 	
 	if( bImportAll )
diff --git a/Engine/Source/Editor/UnrealEd/Private/Fbx/FbxMainExport.cpp b/Engine/Source/Editor/UnrealEd/Private/Fbx/FbxMainExport.cpp
--- a/Engine/Source/Editor/UnrealEd/Private/Fbx/FbxMainExport.cpp
+++ b/Engine/Source/Editor/UnrealEd/Private/Fbx/FbxMainExport.cpp
@@ -2145,7 +2145,8 @@ FbxNode* FFbxExporter::ExportStaticMeshToFbx(const UStaticMesh* StaticMesh, int3
 			}
 			check(UVsLayer);
 
-			const char* UVChannelName = ""; // actually UTF8 as required by Fbx, but can't use UE4's UTF8CHAR type because that's a uint8 aka *unsigned* char
+			FString UVChannelNameBuilder = TEXT("UVmap_") + FString::FromInt(TexCoordSourceIndex);
+			const char* UVChannelName = TCHAR_TO_UTF8(*UVChannelNameBuilder); // actually UTF8 as required by Fbx, but can't use UE4's UTF8CHAR type because that's a uint8 aka *unsigned* char
 			if ((LightmapUVChannel >= 0) || ((LightmapUVChannel == -1) && (TexCoordSourceIndex == StaticMesh->LightMapCoordinateIndex)))
 			{
 				UVChannelName = "LightMapUV";
@@ -2460,8 +2461,8 @@ void FFbxExporter::ExportSplineMeshToFbx(const USplineMeshComponent* SplineMeshC
 			Mesh->CreateLayer();
 			UVsLayer = Mesh->GetLayer(TexCoordSourceIndex);
 		}
-
-		const char* UVChannelName = ""; // actually UTF8 as required by Fbx, but can't use UE4's UTF8CHAR type because that's a uint8 aka *unsigned* char
+		FString UVChannelNameBuilder = TEXT("UVmap_") + FString::FromInt(TexCoordSourceIndex);
+		const char* UVChannelName = TCHAR_TO_UTF8(*UVChannelNameBuilder); // actually UTF8 as required by Fbx, but can't use UE4's UTF8CHAR type because that's a uint8 aka *unsigned* char
 		if (TexCoordSourceIndex == StaticMesh->LightMapCoordinateIndex)
 		{
 			UVChannelName = "LightMapUV";
diff --git a/Engine/Source/Editor/UnrealEd/Private/Fbx/FbxMainImport.cpp b/Engine/Source/Editor/UnrealEd/Private/Fbx/FbxMainImport.cpp
--- a/Engine/Source/Editor/UnrealEd/Private/Fbx/FbxMainImport.cpp
+++ b/Engine/Source/Editor/UnrealEd/Private/Fbx/FbxMainImport.cpp
@@ -120,6 +120,15 @@ FBXImportOptions* GetImportOptions( UnFbx::FFbxImporter* FbxImporter, UFbxImport
 		// @todo: we can make this slow as showing progress bar later
 		FSlateApplication::Get().AddModalWindow(Window, ParentWindow, false);
 
+		if (ImportUI->MeshTypeToImport == FBXIT_SkeletalMesh || ImportUI->MeshTypeToImport == FBXIT_Animation)
+		{
+			//Set some hardcoded options for skeletal mesh
+			ImportUI->SkeletalMeshImportData->bBakePivotInVertex = false;
+			ImportOptions->bBakePivotInVertex = false;
+			ImportUI->SkeletalMeshImportData->bTransformVertexToAbsolute = true;
+			ImportOptions->bTransformVertexToAbsolute = true;
+		}
+
 		ImportUI->SaveConfig();
 
 		if( ImportUI->StaticMeshImportData )
@@ -161,6 +170,8 @@ FBXImportOptions* GetImportOptions( UnFbx::FFbxImporter* FbxImporter, UFbxImport
 	{
 		//Automation tests set ImportUI settings directly.  Just copy them over
 		UnFbx::FBXImportOptions* ImportOptions = FbxImporter->GetImportOptions();
+		//Clean up the options
+		UnFbx::FBXImportOptions::ResetOptions(ImportOptions);
 		ApplyImportUIToImportOptions(ImportUI, *ImportOptions);
 		return ImportOptions;
 	}
@@ -192,6 +203,9 @@ void ApplyImportUIToImportOptions(UFbxImportUI* ImportUI, FBXImportOptions& InOu
 		InOutImportOptions.ImportTranslation		= StaticMeshData->ImportTranslation;
 		InOutImportOptions.ImportRotation			= StaticMeshData->ImportRotation;
 		InOutImportOptions.ImportUniformScale		= StaticMeshData->ImportUniformScale;
+		InOutImportOptions.bTransformVertexToAbsolute = StaticMeshData->bTransformVertexToAbsolute;
+		InOutImportOptions.bBakePivotInVertex		= StaticMeshData->bBakePivotInVertex;
+		InOutImportOptions.bImportStaticMeshLODs	= StaticMeshData->bImportMeshLODs;
 	}
 	else if ( ImportUI->MeshTypeToImport == FBXIT_SkeletalMesh )
 	{
@@ -201,6 +215,9 @@ void ApplyImportUIToImportOptions(UFbxImportUI* ImportUI, FBXImportOptions& InOu
 		InOutImportOptions.ImportTranslation			= SkeletalMeshData->ImportTranslation;
 		InOutImportOptions.ImportRotation				= SkeletalMeshData->ImportRotation;
 		InOutImportOptions.ImportUniformScale			= SkeletalMeshData->ImportUniformScale;
+		InOutImportOptions.bTransformVertexToAbsolute	= SkeletalMeshData->bTransformVertexToAbsolute;
+		InOutImportOptions.bBakePivotInVertex			= SkeletalMeshData->bBakePivotInVertex;
+		InOutImportOptions.bImportSkeletalMeshLODs		= SkeletalMeshData->bImportMeshLODs;
 
 		if(ImportUI->bImportAnimations)
 		{
@@ -228,8 +245,6 @@ void ApplyImportUIToImportOptions(UFbxImportUI* ImportUI, FBXImportOptions& InOu
 	InOutImportOptions.bPreserveSmoothingGroups = ImportUI->SkeletalMeshImportData->bPreserveSmoothingGroups;
 	InOutImportOptions.bKeepOverlappingVertices = ImportUI->SkeletalMeshImportData->bKeepOverlappingVertices;
 	InOutImportOptions.bCombineToSingle = ImportUI->bCombineMeshes;
-	InOutImportOptions.bTransformVertexToAbsolute = ImportUI->StaticMeshImportData->bTransformVertexToAbsolute;
-	InOutImportOptions.bBakePivotInVertex = ImportUI->StaticMeshImportData->bBakePivotInVertex;
 	InOutImportOptions.VertexColorImportOption = ImportUI->StaticMeshImportData->VertexColorImportOption;
 	InOutImportOptions.VertexOverrideColor = ImportUI->StaticMeshImportData->VertexOverrideColor;
 	InOutImportOptions.bRemoveDegenerates = ImportUI->StaticMeshImportData->bRemoveDegenerates;
@@ -1105,6 +1120,18 @@ FName FFbxImporter::MakeNameForMesh(FString InName, FbxObject* FbxObject)
 	return OutputName;
 }
 
+FbxAMatrix FFbxImporter::ComputeSkeletalMeshTotalMatrix(FbxNode* Node, FbxNode *RootSkeletalNode)
+{
+	if (ImportOptions->bImportScene && !ImportOptions->bTransformVertexToAbsolute && RootSkeletalNode != nullptr && RootSkeletalNode != Node)
+	{
+		FbxAMatrix GlobalTransform = Scene->GetAnimationEvaluator()->GetNodeGlobalTransform(Node);
+		FbxAMatrix GlobalSkeletalMeshRootTransform = Scene->GetAnimationEvaluator()->GetNodeGlobalTransform(RootSkeletalNode);
+		FbxAMatrix TotalMatrix = GlobalSkeletalMeshRootTransform.Inverse() * GlobalTransform;
+		return TotalMatrix;
+	}
+	return ComputeTotalMatrix(Node);
+}
+
 FbxAMatrix FFbxImporter::ComputeTotalMatrix(FbxNode* Node)
 {
 	FbxAMatrix Geometry;
@@ -1601,6 +1628,24 @@ void FFbxImporter::RecursiveFindFbxSkelMesh(FbxNode* Node, TArray< TArray<FbxNod
 						TempArray->Add(NodeToAdd);
 						outSkelMeshArray.Add(TempArray);
 						SkeletonArray.Add(Link);
+						
+						if (ImportOptions->bImportScene && !ImportOptions->bTransformVertexToAbsolute)
+						{
+							FbxVector4 NodeScaling = NodeToAdd->EvaluateLocalScaling();
+							FbxVector4 NoScale(1.0, 1.0, 1.0);
+							if (NodeScaling != NoScale)
+							{
+								//Scene import cannot import correctly a skeletal mesh with a root node containing scale
+								//Warn the user is skeletal mesh can be wrong
+								AddTokenizedErrorMessage(
+									FTokenizedMessage::Create(
+										EMessageSeverity::Warning,
+										FText::Format(LOCTEXT("FBX_ImportSceneSkeletalMeshRootNodeScaling", "Importing skeletal mesh {0} that dont have a mesh node with no scale is not supported when doing an import scene."), FText::FromString(UTF8_TO_TCHAR(NodeToAdd->GetName())))
+										),
+									FFbxErrors::SkeletalMesh_InvalidRoot
+									);
+							}
+						}
 					}
 
 					bFoundCorrectLink = true;
@@ -1624,9 +1669,32 @@ void FFbxImporter::RecursiveFindFbxSkelMesh(FbxNode* Node, TArray< TArray<FbxNod
 	else
 	{
 		int32 ChildIndex;
-		for (ChildIndex=0; ChildIndex<Node->GetChildCount(); ++ChildIndex)
+		TArray<FbxNode*> ChildNoScale;
+		TArray<FbxNode*> ChildScale;
+		//Sort the node to have the one with no scaling first so we have more chance
+		//to have a root skeletal mesh with no scale. Because scene import do not support
+		//root skeletal mesh containing scale
+		for (ChildIndex = 0; ChildIndex < Node->GetChildCount(); ++ChildIndex)
+		{
+			FbxNode *ChildNode = Node->GetChild(ChildIndex);
+			FbxVector4 ChildScaling = ChildNode->EvaluateLocalScaling();
+			FbxVector4 NoScale(1.0, 1.0, 1.0);
+			if (ChildScaling == NoScale)
+			{
+				ChildNoScale.Add(ChildNode);
+			}
+			else
+			{
+				ChildScale.Add(ChildNode);
+			}
+		}
+		for (FbxNode *ChildNode : ChildNoScale)
+		{
+			RecursiveFindFbxSkelMesh(ChildNode, outSkelMeshArray, SkeletonArray, ExpandLOD);
+		}
+		for (FbxNode *ChildNode : ChildScale)
 		{
-			RecursiveFindFbxSkelMesh(Node->GetChild(ChildIndex), outSkelMeshArray, SkeletonArray, ExpandLOD);
+			RecursiveFindFbxSkelMesh(ChildNode, outSkelMeshArray, SkeletonArray, ExpandLOD);
 		}
 	}
 }
diff --git a/Engine/Source/Editor/UnrealEd/Private/Fbx/FbxMaterialImport.cpp b/Engine/Source/Editor/UnrealEd/Private/Fbx/FbxMaterialImport.cpp
--- a/Engine/Source/Editor/UnrealEd/Private/Fbx/FbxMaterialImport.cpp
+++ b/Engine/Source/Editor/UnrealEd/Private/Fbx/FbxMaterialImport.cpp
@@ -131,6 +131,7 @@ UTexture* UnFbx::FFbxImporter::ImportTexture( FbxFileTexture* FbxTexture, bool b
 			{
 				TextureFact->LODGroup = TEXTUREGROUP_WorldNormalMap;
 				TextureFact->CompressionSettings = TC_Normalmap;
+				TextureFact->bFlipNormalMapGreenChannel = ImportOptions->bInvertNormalMap;
 			}
 			else
 			{
@@ -303,6 +304,10 @@ bool UnFbx::FFbxImporter::CreateAndLinkExpressionForMaterialProperty(
 						// add/find UVSet and set it to the texture
 						FbxString UVSetName = FbxTexture->UVSet.Get();
 						FString LocalUVSetName = UTF8_TO_TCHAR(UVSetName.Buffer());
+						if (LocalUVSetName.IsEmpty())
+						{
+							LocalUVSetName = TEXT("UVmap_0");
+						}
 						int32 SetIndex = UVSet.Find(LocalUVSetName);
 						if( (SetIndex != 0 && SetIndex != INDEX_NONE) || ScaleU != 1.0f || ScaleV != 1.0f )
 						{
diff --git a/Engine/Source/Editor/UnrealEd/Private/Fbx/FbxSceneImportData.cpp b/Engine/Source/Editor/UnrealEd/Private/Fbx/FbxSceneImportData.cpp
--- a/Engine/Source/Editor/UnrealEd/Private/Fbx/FbxSceneImportData.cpp
+++ b/Engine/Source/Editor/UnrealEd/Private/Fbx/FbxSceneImportData.cpp
@@ -108,7 +108,9 @@ UnFbx::FBXImportOptions *JSONToFbxOption(TSharedPtr<FJsonValue> OptionJsonValue,
 
 	//Skeletal mesh options
 	OptionObj->TryGetBoolField("bUpdateSkeletonReferencePose", Option->bUpdateSkeletonReferencePose);
-	OptionObj->TryGetBoolField("bUseT0AsRefPose", Option->bUseT0AsRefPose);
+	//TODO support T0AsRefPose
+	//OptionObj->TryGetBoolField("bUseT0AsRefPose", Option->bUseT0AsRefPose);
+	Option->bUseT0AsRefPose = false;
 	OptionObj->TryGetBoolField("bPreserveSmoothingGroups", Option->bPreserveSmoothingGroups);
 	OptionObj->TryGetBoolField("bImportMeshesInBoneHierarchy", Option->bImportMeshesInBoneHierarchy);
 	OptionObj->TryGetBoolField("bImportMorphTargets", Option->bImportMorph);
diff --git a/Engine/Source/Editor/UnrealEd/Private/Fbx/FbxSceneImportFactory.cpp b/Engine/Source/Editor/UnrealEd/Private/Fbx/FbxSceneImportFactory.cpp
--- a/Engine/Source/Editor/UnrealEd/Private/Fbx/FbxSceneImportFactory.cpp
+++ b/Engine/Source/Editor/UnrealEd/Private/Fbx/FbxSceneImportFactory.cpp
@@ -84,6 +84,12 @@ bool GetFbxSceneImportOptions(UnFbx::FFbxImporter* FbxImporter
 	GlobalImportSettings->bImportTextures = true;
 	//Make sure Material get imported
 	GlobalImportSettings->bImportMaterials = true;
+	//TODO support T0AsRefPose
+	GlobalImportSettings->bUseT0AsRefPose = false;
+
+	GlobalImportSettings->ImportTranslation = FVector(0);
+	GlobalImportSettings->ImportRotation = FRotator(0);
+	GlobalImportSettings->ImportUniformScale = 1.0f;
 
 	GlobalImportSettings->bConvertScene = true;
 
@@ -130,6 +136,9 @@ bool GetFbxSceneImportOptions(UnFbx::FFbxImporter* FbxImporter
 	GlobalImportSettings->bImportStaticMeshLODs = SceneImportOptions->bImportStaticMeshLODs;
 	GlobalImportSettings->bImportSkeletalMeshLODs = SceneImportOptions->bImportSkeletalMeshLODs;
 	GlobalImportSettings->bInvertNormalMap = SceneImportOptions->bInvertNormalMaps;
+	GlobalImportSettings->ImportTranslation = SceneImportOptions->ImportTranslation;
+	GlobalImportSettings->ImportRotation = SceneImportOptions->ImportRotation;
+	GlobalImportSettings->ImportUniformScale = SceneImportOptions->ImportUniformScale;
 
 	//Set the override material into the options
 	for (TSharedPtr<FFbxNodeInfo> NodeInfo : SceneInfoPtr->HierarchyInfo)
@@ -166,7 +175,8 @@ bool GetFbxSceneImportOptions(UnFbx::FFbxImporter* FbxImporter
 	SFbxSceneOptionWindow::CopyFbxOptionsToStaticMeshOptions(GlobalImportSettings, StaticMeshImportData);
 	StaticMeshImportData->SaveConfig();
 
-	//TODO save all other options
+	SFbxSceneOptionWindow::CopyFbxOptionsToSkeletalMeshOptions(GlobalImportSettings, SkeletalMeshImportData);
+	SkeletalMeshImportData->SaveConfig();
 
 	return true;
 }
@@ -695,31 +705,13 @@ void UFbxSceneImportFactory::FillSceneHierarchyPath(TSharedPtr<FFbxSceneInfo> Sc
 
 UFbxSceneImportData* CreateReImportAsset(const FString &PackagePath, const FString &FbxImportFileName, UFbxSceneImportOptions* SceneImportOptions, TSharedPtr<FFbxSceneInfo> SceneInfo, ImportOptionsNameMap &NameOptionsMap)
 {
-	//Create a package
+	//Create or use existing package
+	//The data must have the name of the import file to support drag drop reimport
 	FString FilenameBase = FPaths::GetBaseFilename(FbxImportFileName);
-	FString FbxReImportPkgName = PackagePath + TEXT("/FbxScene_") + FilenameBase;
+	FString FbxReImportPkgName = PackagePath + TEXT("/") + FilenameBase;
 	FbxReImportPkgName = PackageTools::SanitizePackageName(FbxReImportPkgName);
-	FString AssetName = TEXT("FbxScene_") + FilenameBase;
+	FString AssetName = FilenameBase;
 	AssetName = PackageTools::SanitizePackageName(AssetName);
-	UPackage* PkgExist = nullptr;
-	if (!FEditorFileUtils::IsMapPackageAsset(FbxReImportPkgName))
-	{
-		PkgExist = FindPackage(nullptr, *FbxReImportPkgName);
-	}
-	int32 tryCount = 1;
-	while (PkgExist != nullptr)
-	{
-		FbxReImportPkgName = PackagePath + TEXT("/");
-		AssetName = TEXT("FbxScene_") + FilenameBase;
-		AssetName += TEXT("_");
-		AssetName += FString::FromInt(tryCount++);
-		FbxReImportPkgName += AssetName;
-		FbxReImportPkgName = PackageTools::SanitizePackageName(FbxReImportPkgName);
-		if (!FEditorFileUtils::IsMapPackageAsset(FbxReImportPkgName))
-		{
-			PkgExist = FindPackage(nullptr, *FbxReImportPkgName);
-		}
-	}
 	UPackage* Pkg = CreatePackage(nullptr, *FbxReImportPkgName);
 	if (!ensure(Pkg))
 	{
@@ -812,10 +804,14 @@ FFeedbackContext*	Warn
 	Warn->BeginSlowTask(NSLOCTEXT("FbxSceneFactory", "BeginImportingFbxSceneTask", "Importing FBX scene"), true);
 	
 	GlobalImportSettings = FbxImporter->GetImportOptions();
+	UnFbx::FBXImportOptions::ResetOptions(GlobalImportSettings);
 	
 	//Always convert the scene
 	GlobalImportSettings->bConvertScene = true;
 
+	//Set the import option in importscene mode
+	GlobalImportSettings->bImportScene = true;
+
 	//Read the fbx and store the hierarchy's information so we can reuse it after importing all the model in the fbx file
 	if (!FbxImporter->ImportFromFile(*FbxImportFileName, Type, true))
 	{
@@ -907,13 +903,14 @@ FFeedbackContext*	Warn
 	SkeletalMeshImportData->bImportAsScene = true;
 	SkeletalMeshImportData->FbxSceneImportDataReference = ReimportData;
 
+	AnimSequenceImportData->bImportAsScene = true;
+	AnimSequenceImportData->FbxSceneImportDataReference = ReimportData;
+
 	//Get the scene root node
 	FbxNode* RootNodeToImport = nullptr;
 	RootNodeToImport = FbxImporter->Scene->GetRootNode();
-	
-	//Apply the options scene transform to the root node
-	ApplySceneTransformOptionsToRootNode(SceneInfoPtr);
 
+	
 	// For animation and static mesh we assume there is at lease one interesting node by default
 	int32 InterestingNodeCount = 1;
 
@@ -922,13 +919,13 @@ FFeedbackContext*	Warn
 	int32 NodeIndex = 0;
 
 	//////////////////////////////////////////////////////////////////////////
-	// IMPORT ALL STATIC MESH
-	ImportAllStaticMesh(RootNodeToImport, FbxImporter, Flags, NodeIndex, InterestingNodeCount, SceneInfoPtr);
+	//IMPORT ALL SKELETAL MESH
+	ImportAllSkeletalMesh(RootNodeToImport, FbxImporter, Flags, NodeIndex, InterestingNodeCount, SceneInfoPtr);
 	//////////////////////////////////////////////////////////////////////////
 
 	//////////////////////////////////////////////////////////////////////////
-	//IMPORT ALL SKELETAL MESH
-	ImportAllSkeletalMesh(RootNodeToImport, FbxImporter, Flags, NodeIndex, InterestingNodeCount, SceneInfoPtr);
+	// IMPORT ALL STATIC MESH
+	ImportAllStaticMesh(RootNodeToImport, FbxImporter, Flags, NodeIndex, InterestingNodeCount, SceneInfoPtr);
 	//////////////////////////////////////////////////////////////////////////
 
 	UObject *ReturnObject = nullptr;
@@ -967,10 +964,12 @@ FFeedbackContext*	Warn
 			if (SceneImportOptions->HierarchyType == EFBXSceneOptionsCreateHierarchyType::FBXSOCHT_CreateBlueprint && HierarchyActor != nullptr)
 			{
 				//The location+name of the BP is the user select content path + fbx base filename
-				FString FullnameBP = Path + TEXT("/") + FPaths::GetBaseFilename(UFactory::CurrentFilename);
+				FString FullnameBP = Path + TEXT("/FbxScene_") + FPaths::GetBaseFilename(UFactory::CurrentFilename);
 				FullnameBP = PackageTools::SanitizePackageName(FullnameBP);
+				FString AssetName = TEXT("FbxScene_") + FPaths::GetBaseFilename(UFactory::CurrentFilename);
+				UPackage *Pkg = CreatePackageForNode(FullnameBP, AssetName);
 				//Create the blueprint from the actor and replace the actor with a blueprintactor that point on the blueprint
-				UBlueprint* SceneBlueprint = FKismetEditorUtilities::CreateBlueprintFromActor(FullnameBP, HierarchyActor, true);
+				UBlueprint* SceneBlueprint = FKismetEditorUtilities::CreateBlueprintFromActor(Pkg->GetName(), HierarchyActor, true);
 				if (SceneBlueprint != nullptr && ReimportData != nullptr)
 				{
 					//let the scene blueprint be the return object for this import
@@ -1007,22 +1006,6 @@ FFeedbackContext*	Warn
 	return ReturnObject;
 }
 
-void UFbxSceneImportFactory::ApplySceneTransformOptionsToRootNode(TSharedPtr<FFbxSceneInfo> SceneInfoPtr)
-{
-	for (TSharedPtr<FFbxNodeInfo> NodeInfo : SceneInfoPtr->HierarchyInfo)
-	{
-		if (NodeInfo->NodeName.Compare("RootNode") == 0)
-		{
-			FTransform SceneTransformAddition(SceneImportOptions->ImportRotation, SceneImportOptions->ImportTranslation, FVector(SceneImportOptions->ImportUniformScale));
-			FTransform OutTransform = FTransform::Identity;
-			//We apply the scene pre transform before the existing root node transform and store the result in the rootnode transform
-			FTransform::Multiply(&OutTransform, &NodeInfo->Transform, &SceneTransformAddition);
-			NodeInfo->Transform = OutTransform;
-			return;
-		}
-	}
-}
-
 bool UFbxSceneImportFactory::SetStaticMeshComponentOverrideMaterial(UStaticMeshComponent* StaticMeshComponent, TSharedPtr<FFbxNodeInfo> NodeInfo)
 {
 	bool bOverrideMaterial = false;
@@ -1109,6 +1092,52 @@ USceneComponent *CreateLightComponent(AActor *ParentActor, TSharedPtr<FFbxLightI
 	return LightComponent;
 }
 
+FVector GetParentPivotAccumulation(TSharedPtr<FFbxNodeInfo> NodeInfo, TSharedPtr<FFbxSceneInfo> SceneInfoPtr, FTransform &RootTransform)
+{
+	TArray<TSharedPtr<FFbxNodeInfo>> ParentHierarchy;
+	FVector PivotAccumulation(0.0f);
+	TSharedPtr<FFbxNodeInfo> ParentNodeInfo = NodeInfo->ParentNodeInfo;
+	while (ParentNodeInfo.IsValid())
+	{
+		ParentHierarchy.Insert(ParentNodeInfo, 0);
+		ParentNodeInfo = ParentNodeInfo->ParentNodeInfo;
+	}
+	FTransform CurrentGlobalMatrix;
+	for (TSharedPtr<FFbxNodeInfo> ParentNode : ParentHierarchy)
+	{
+		FTransform LocalTransformAdjusted = ParentNode->Transform;
+		FVector PivotLocation(0.0f);
+		if (ParentNode->AttributeInfo.IsValid())
+		{
+			for (TSharedPtr<FFbxNodeInfo> NodeInfoIter : SceneInfoPtr->HierarchyInfo)
+			{
+				if (NodeInfoIter->UniqueId == ParentNode->AttributeInfo->PivotNodeUid)
+				{
+					PivotLocation = NodeInfoIter->PivotRotation;
+					break;
+				}
+			}
+		}
+		FTransform LocalTransform = ParentNode->Transform;
+		if (!PivotLocation.IsNearlyZero())
+		{
+			FTransform ParentPivotTransform;
+			ParentPivotTransform.SetLocation(PivotLocation);
+			FTransform AlmostNextCurrentTransform = LocalTransform *CurrentGlobalMatrix;
+			//Get the final matrix with pivot
+			LocalTransform = ParentPivotTransform * LocalTransform;
+			CurrentGlobalMatrix = LocalTransform * CurrentGlobalMatrix;
+			ParentPivotTransform = CurrentGlobalMatrix * AlmostNextCurrentTransform.Inverse();
+			PivotAccumulation = ParentPivotTransform.GetLocation();
+		}
+		else
+		{
+			CurrentGlobalMatrix = LocalTransform * CurrentGlobalMatrix;
+		}
+	}
+	return PivotAccumulation;
+}
+
 void UFbxSceneImportFactory::CreateLevelActorHierarchy(TSharedPtr<FFbxSceneInfo> SceneInfoPtr)
 {
 	EComponentMobility::Type MobilityType = SceneImportOptions->bImportAsDynamic ? EComponentMobility::Movable : EComponentMobility::Static;
@@ -1140,6 +1169,7 @@ void UFbxSceneImportFactory::CreateLevelActorHierarchy(TSharedPtr<FFbxSceneInfo>
 		//Find the asset that link with this node attribute
 		UObject *AssetToPlace = (NodeInfo->AttributeInfo.IsValid() && AllNewAssets.Contains(NodeInfo->AttributeInfo)) ? AllNewAssets[NodeInfo->AttributeInfo] : nullptr;
 		
+		bool IsSkeletalMesh = false;
 		//create actor
 		AActor *PlacedActor = nullptr;
 		if (AssetToPlace != nullptr)
@@ -1154,9 +1184,26 @@ void UFbxSceneImportFactory::CreateLevelActorHierarchy(TSharedPtr<FFbxSceneInfo>
 				UStaticMeshComponent *StaticMeshComponent = Cast<UStaticMeshComponent>(PlacedActor->GetComponentByClass(UStaticMeshComponent::StaticClass()));
 				SetStaticMeshComponentOverrideMaterial(StaticMeshComponent, NodeInfo);
 			}
+			IsSkeletalMesh = (AssetToPlace->GetClass() == USkeletalMesh::StaticClass());
 		}
 		else if (IsEmptyAttribute(NodeInfo->AttributeType) || NodeInfo->AttributeType.Compare("eMesh") == 0 || NodeInfo->AttributeUniqueId != INVALID_UNIQUE_ID)
 		{
+			if (NodeInfo->AttributeType.Compare("eMesh") == 0)
+			{
+				bool bIsSubSkeletalMesh = true;
+				for (TSharedPtr<FFbxMeshInfo> MeshInfo : SceneInfoPtr->MeshInfo)
+				{
+					if (!MeshInfo->bIsSkelMesh && NodeInfo->AttributeUniqueId == MeshInfo->UniqueId)
+					{
+						bIsSubSkeletalMesh = false;
+						break;
+					}
+				}
+				if (bIsSubSkeletalMesh == true)
+				{
+					continue;
+				}
+			}
 			//Create an empty actor if the node is an empty attribute or the attribute is a mesh(static mesh or skeletal mesh) that was not export
 			UActorFactory* Factory = GEditor->FindActorFactoryByClass(UActorFactoryEmptyActor::StaticClass());
 			FAssetData EmptyActorAssetData = FAssetData(Factory->GetDefaultActorClass(FAssetData()));
@@ -1233,20 +1280,25 @@ void UFbxSceneImportFactory::CreateLevelActorHierarchy(TSharedPtr<FFbxSceneInfo>
 				}
 				//Find the pivot location
 				FVector PivotLocation(0.0f);
-				if (GlobalImportSettings->bBakePivotInVertex && NodeInfo->AttributeInfo.IsValid() && NodeInfo->AttributeInfo->PivotNodeUid != INVALID_UNIQUE_ID)
+				FVector ParentPivotAccumulation(0.0f);
+				if (!IsSkeletalMesh && GlobalImportSettings->bBakePivotInVertex)
 				{
-					for (TSharedPtr<FFbxNodeInfo> NodeInfoIter : SceneInfoPtr->HierarchyInfo)
+					ParentPivotAccumulation -= GetParentPivotAccumulation(NodeInfo, SceneInfoPtr, RootTransform);
+					if (NodeInfo->AttributeInfo.IsValid() && NodeInfo->AttributeInfo->PivotNodeUid != INVALID_UNIQUE_ID)
 					{
-						if (NodeInfoIter->UniqueId == NodeInfo->AttributeInfo->PivotNodeUid)
+						for (TSharedPtr<FFbxNodeInfo> NodeInfoIter : SceneInfoPtr->HierarchyInfo)
 						{
-							PivotLocation = NodeInfoIter->PivotRotation;
-							break;
+							if (NodeInfoIter->UniqueId == NodeInfo->AttributeInfo->PivotNodeUid)
+							{
+								PivotLocation = NodeInfoIter->PivotRotation;
+								break;
+							}
 						}
 					}
 				}
 
 				//Apply the hierarchy local transform to the root component
-				ApplyTransformToComponent(RootComponent, &(NodeInfo->Transform), ParentActor == nullptr ? &RootTransform : nullptr, PivotLocation);
+				ApplyTransformToComponent(RootComponent, &(NodeInfo->Transform), ParentActor == nullptr ? &RootTransform : nullptr, PivotLocation, ParentPivotAccumulation);
 			}
 		}
 		//We select only the first actor
@@ -1310,6 +1362,7 @@ AActor *UFbxSceneImportFactory::CreateActorComponentsHierarchy(TSharedPtr<FFbxSc
 		//Find the asset that link with this node attribute
 		UObject *AssetToPlace = (NodeInfo->AttributeInfo.IsValid() && AllNewAssets.Contains(NodeInfo->AttributeInfo)) ? AllNewAssets[NodeInfo->AttributeInfo] : nullptr;
 
+		bool IsSkeletalMesh = false;
 		//Create the component where the type depend on the asset point by the component
 		//In case there is no asset we create a SceneComponent
 		USceneComponent* SceneComponent = nullptr;
@@ -1333,10 +1386,28 @@ AActor *UFbxSceneImportFactory::CreateActorComponentsHierarchy(TSharedPtr<FFbxSc
 				SkeletalMeshComponent->DepthPriorityGroup = SDPG_World;
 				SceneComponent = SkeletalMeshComponent;
 				SceneComponent->Mobility = MobilityType;
+				IsSkeletalMesh = true;
 			}
 		}
 		else if (IsEmptyAttribute(NodeInfo->AttributeType) || NodeInfo->AttributeType.Compare("eMesh") == 0 || NodeInfo->AttributeUniqueId != INVALID_UNIQUE_ID)
 		{
+			if (NodeInfo->AttributeType.Compare("eMesh") == 0)
+			{
+				bool bIsSubSkeletalMesh = true;
+				for (TSharedPtr<FFbxMeshInfo> MeshInfo : SceneInfoPtr->MeshInfo)
+				{
+					if (!MeshInfo->bIsSkelMesh && NodeInfo->AttributeUniqueId == MeshInfo->UniqueId)
+					{
+						bIsSubSkeletalMesh = false;
+						break;
+					}
+				}
+				if (bIsSubSkeletalMesh == true)
+				{
+					continue;
+				}
+			}
+
 			if (NodeInfo->AttributeType.Compare("eLight") == 0 && SceneInfoPtr->LightInfo.Contains(NodeInfo->AttributeUniqueId))
 			{
 				TSharedPtr<FFbxLightInfo> LightInfo = *SceneInfoPtr->LightInfo.Find(NodeInfo->AttributeUniqueId);
@@ -1404,20 +1475,25 @@ AActor *UFbxSceneImportFactory::CreateActorComponentsHierarchy(TSharedPtr<FFbxSc
 
 		//Find the pivot location
 		FVector PivotLocation(0.0f);
-		if (GlobalImportSettings->bBakePivotInVertex && NodeInfo->AttributeInfo.IsValid() && NodeInfo->AttributeInfo->PivotNodeUid != INVALID_UNIQUE_ID)
+		FVector ParentPivotAccumulation(0.0f);
+		if (!IsSkeletalMesh && GlobalImportSettings->bBakePivotInVertex)
 		{
-			for (TSharedPtr<FFbxNodeInfo> NodeInfoIter : SceneInfoPtr->HierarchyInfo)
+			ParentPivotAccumulation -= GetParentPivotAccumulation(NodeInfo, SceneInfoPtr, RootTransform);
+			if (NodeInfo->AttributeInfo.IsValid() && NodeInfo->AttributeInfo->PivotNodeUid != INVALID_UNIQUE_ID)
 			{
-				if (NodeInfoIter->UniqueId == NodeInfo->AttributeInfo->PivotNodeUid)
+				for (TSharedPtr<FFbxNodeInfo> NodeInfoIter : SceneInfoPtr->HierarchyInfo)
 				{
-					PivotLocation = NodeInfoIter->PivotRotation;
-					break;
+					if (NodeInfoIter->UniqueId == NodeInfo->AttributeInfo->PivotNodeUid)
+					{
+						PivotLocation = NodeInfoIter->PivotRotation;
+						break;
+					}
 				}
 			}
 		}
 
 		//Apply the local transform to the scene component
-		ApplyTransformToComponent(SceneComponent, &(NodeInfo->Transform), ParentRootComponent != nullptr ? nullptr : &RootTransform, PivotLocation);
+		ApplyTransformToComponent(SceneComponent, &(NodeInfo->Transform), ParentRootComponent != nullptr ? nullptr : &RootTransform, PivotLocation, ParentPivotAccumulation);
 		
 	}
 	// End of iteration of the hierarchy
@@ -1426,14 +1502,17 @@ AActor *UFbxSceneImportFactory::CreateActorComponentsHierarchy(TSharedPtr<FFbxSc
 	return RootActorContainer;
 }
 
-void UFbxSceneImportFactory::ApplyTransformToComponent(USceneComponent *SceneComponent, FTransform *LocalTransform, FTransform *PreMultiplyTransform, FVector &PivotLocation)
+void UFbxSceneImportFactory::ApplyTransformToComponent(USceneComponent *SceneComponent, FTransform *LocalTransform, FTransform *PreMultiplyTransform, FVector &PivotLocation, FVector &ParentPivotAccumulation)
 {
 	check(SceneComponent);
 	check(LocalTransform);
 	FTransform LocalTransformAdjusted = (*LocalTransform);
-	if (GlobalImportSettings->bBakePivotInVertex && !PivotLocation.IsNearlyZero())
+	if (GlobalImportSettings->bBakePivotInVertex && (!PivotLocation.IsNearlyZero() || !ParentPivotAccumulation.IsNearlyZero()))
 	{
 		FTransform PivotTransform;
+		PivotTransform.SetLocation(ParentPivotAccumulation);
+		LocalTransformAdjusted = LocalTransformAdjusted * PivotTransform;
+		PivotTransform.SetIdentity();
 		PivotTransform.SetLocation(PivotLocation);
 		LocalTransformAdjusted = PivotTransform * LocalTransformAdjusted;
 	}
@@ -1449,7 +1528,36 @@ void UFbxSceneImportFactory::ApplyTransformToComponent(USceneComponent *SceneCom
 		SceneComponent->SetRelativeTransform(LocalTransformAdjusted);
 	}
 }
-
+void UFbxSceneImportFactory::ApplyMeshInfoFbxOptions(TSharedPtr<FFbxMeshInfo> MeshInfo)
+{
+	if (!MeshInfo.IsValid())
+	{
+		//Use the default options
+		SFbxSceneOptionWindow::CopyFbxOptionsToFbxOptions(GlobalImportSettingsReference, GlobalImportSettings);
+		SFbxSceneOptionWindow::CopyFbxOptionsToSkeletalMeshOptions(GlobalImportSettingsReference, SceneImportOptionsSkeletalMesh);
+		SFbxSceneOptionWindow::CopyFbxOptionsToStaticMeshOptions(GlobalImportSettingsReference, SceneImportOptionsStaticMesh);
+	}
+	else
+	{
+		UnFbx::FBXImportOptions* OverrideImportSettings = GetOptionsFromName(MeshInfo->OptionName);
+		if (OverrideImportSettings != nullptr)
+		{
+			//Use the override options
+			SFbxSceneOptionWindow::CopyFbxOptionsToFbxOptions(OverrideImportSettings, GlobalImportSettings);
+			SFbxSceneOptionWindow::CopyFbxOptionsToSkeletalMeshOptions(OverrideImportSettings, SceneImportOptionsSkeletalMesh);
+			SFbxSceneOptionWindow::CopyFbxOptionsToStaticMeshOptions(OverrideImportSettings, SceneImportOptionsStaticMesh);
+		}
+		else
+		{
+			//Use the default options if we found no options
+			SFbxSceneOptionWindow::CopyFbxOptionsToFbxOptions(GlobalImportSettingsReference, GlobalImportSettings);
+			SFbxSceneOptionWindow::CopyFbxOptionsToSkeletalMeshOptions(GlobalImportSettingsReference, SceneImportOptionsSkeletalMesh);
+			SFbxSceneOptionWindow::CopyFbxOptionsToStaticMeshOptions(GlobalImportSettingsReference, SceneImportOptionsStaticMesh);
+		}
+	}
+	SceneImportOptionsSkeletalMesh->FillSkeletalMeshInmportData(SkeletalMeshImportData, AnimSequenceImportData, SceneImportOptions);
+	SceneImportOptionsStaticMesh->FillStaticMeshInmportData(StaticMeshImportData, SceneImportOptions);
+}
 
 UObject* UFbxSceneImportFactory::ImportOneSkeletalMesh(void* VoidRootNodeToImport, void* VoidFbxImporter, TSharedPtr<FFbxSceneInfo> SceneInfo, EObjectFlags Flags, TArray<void*> &VoidNodeArray, int32 &TotalNumNodes)
 {
@@ -1485,26 +1593,14 @@ UObject* UFbxSceneImportFactory::ImportOneSkeletalMesh(void* VoidRootNodeToImpor
 	//Apply the correct fbx options
 	TSharedPtr<FFbxMeshInfo> MeshInfo = StaticCastSharedPtr<FFbxMeshInfo>(RootNodeInfo->AttributeInfo);
 
-	UnFbx::FBXImportOptions* OverrideImportSettings = GetOptionsFromName(MeshInfo->OptionName);
-	if (OverrideImportSettings != nullptr)
-	{
-		SFbxSceneOptionWindow::CopyFbxOptionsToFbxOptions(OverrideImportSettings, GlobalImportSettings);
-		SFbxSceneOptionWindow::CopyFbxOptionsToSkeletalMeshOptions(OverrideImportSettings, SceneImportOptionsSkeletalMesh);
-		SceneImportOptionsSkeletalMesh->FillSkeletalMeshInmportData(SkeletalMeshImportData, AnimSequenceImportData);
-	}
-	else
-	{
-		//Use the default options if we dont found options
-		SFbxSceneOptionWindow::CopyFbxOptionsToFbxOptions(GlobalImportSettingsReference, GlobalImportSettings);
-		SFbxSceneOptionWindow::CopyFbxOptionsToSkeletalMeshOptions(GlobalImportSettingsReference, SceneImportOptionsSkeletalMesh);
-		SceneImportOptionsSkeletalMesh->FillSkeletalMeshInmportData(SkeletalMeshImportData, AnimSequenceImportData);
-	}
-
-	//This is possible for the user to not bake pivot for a particular root node
-	if (GlobalImportSettings->bBakePivotInVertex && RootNodeInfo->AttributeInfo->PivotNodeUid == INVALID_UNIQUE_ID)
-	{
-		GlobalImportSettings->bBakePivotInVertex = false;
-	}
+	ApplyMeshInfoFbxOptions(MeshInfo);
+	//TODO support bBakePivotInVertex
+	bool Old_bBakePivotInVertex = GlobalImportSettings->bBakePivotInVertex;
+	GlobalImportSettings->bBakePivotInVertex = false;
+	//if (GlobalImportSettings->bBakePivotInVertex && RootNodeInfo->AttributeInfo->PivotNodeUid == INVALID_UNIQUE_ID)
+	//{
+		//GlobalImportSettings->bBakePivotInVertex = false;
+	//}
 
 	// check if there is LODGroup for this skeletal mesh
 	int32 MaxLODLevel = 1;
@@ -1559,12 +1655,15 @@ UObject* UFbxSceneImportFactory::ImportOneSkeletalMesh(void* VoidRootNodeToImpor
 				TSharedPtr<FFbxNodeInfo> ExportNodeInfo;
 				if (FindSceneNodeInfo(SceneInfo, SkelMeshNode->GetUniqueID(), ExportNodeInfo))
 				{
-					FbxNode *NodePivot = FindFbxNodeById(FbxImporter, nullptr, ExportNodeInfo->AttributeInfo->PivotNodeUid);
-					if (NodePivot != nullptr)
+					if (ExportNodeInfo->AttributeInfo.IsValid())
 					{
-						SkelMeshNodePivotArray.Add(NodePivot);
-						bUseSkelMeshNodePivotArray = true;
-						continue;
+						FbxNode *NodePivot = FindFbxNodeById(FbxImporter, nullptr, ExportNodeInfo->AttributeInfo->PivotNodeUid);
+						if (NodePivot != nullptr)
+						{
+							SkelMeshNodePivotArray.Add(NodePivot);
+							bUseSkelMeshNodePivotArray = true;
+							continue;
+						}
 					}
 				}
 				SkelMeshNodePivotArray.Add(SkelMeshNode);
@@ -1633,7 +1732,8 @@ UObject* UFbxSceneImportFactory::ImportOneSkeletalMesh(void* VoidRootNodeToImpor
 			FbxImporter->ImportFbxMorphTarget(SkelMeshNodeArray, Cast<USkeletalMesh>(NewObject), Pkg, LODIndex);
 		}
 	}
-
+	//Put back the options
+	GlobalImportSettings->bBakePivotInVertex = Old_bBakePivotInVertex;
 	return NewObject;
 }
 
@@ -1843,20 +1943,8 @@ UObject* UFbxSceneImportFactory::ImportANode(void* VoidFbxImporter, void* VoidNo
 	//Apply the correct fbx options
 	TSharedPtr<FFbxMeshInfo> MeshInfo = StaticCastSharedPtr<FFbxMeshInfo>(OutNodeInfo->AttributeInfo);
 	
-	UnFbx::FBXImportOptions* OverrideImportSettings = GetOptionsFromName(MeshInfo->OptionName);
-	if (OverrideImportSettings != nullptr)
-	{
-		SFbxSceneOptionWindow::CopyFbxOptionsToFbxOptions(OverrideImportSettings, GlobalImportSettings);
-		SFbxSceneOptionWindow::CopyFbxOptionsToStaticMeshOptions(OverrideImportSettings, SceneImportOptionsStaticMesh);
-		SceneImportOptionsStaticMesh->FillStaticMeshInmportData(StaticMeshImportData, SceneImportOptions);
-	}
-	else
-	{
-		//Use the default options if we dont found options
-		SFbxSceneOptionWindow::CopyFbxOptionsToFbxOptions(GlobalImportSettingsReference, GlobalImportSettings);
-		SFbxSceneOptionWindow::CopyFbxOptionsToStaticMeshOptions(GlobalImportSettingsReference, SceneImportOptionsStaticMesh);
-		SceneImportOptionsStaticMesh->FillStaticMeshInmportData(StaticMeshImportData, SceneImportOptions);
-	}
+	ApplyMeshInfoFbxOptions(MeshInfo);
+	bool Old_bBakePivotInVertex = GlobalImportSettings->bBakePivotInVertex;
 	if (GlobalImportSettings->bBakePivotInVertex && OutNodeInfo->AttributeInfo->PivotNodeUid == INVALID_UNIQUE_ID)
 	{
 		GlobalImportSettings->bBakePivotInVertex = false;
@@ -1886,6 +1974,7 @@ UObject* UFbxSceneImportFactory::ImportANode(void* VoidFbxImporter, void* VoidNo
 		GWarn->StatusUpdate(NodeIndex, Total, FText::Format(NSLOCTEXT("UnrealEd", "Importingf", "Importing ({NodeIndex} of {ArrayLength})"), Args));
 	}
 
+	GlobalImportSettings->bBakePivotInVertex = Old_bBakePivotInVertex;
 	return NewObject;
 }
 
diff --git a/Engine/Source/Editor/UnrealEd/Private/Fbx/FbxSceneImportOptionsSkeletalMesh.cpp b/Engine/Source/Editor/UnrealEd/Private/Fbx/FbxSceneImportOptionsSkeletalMesh.cpp
--- a/Engine/Source/Editor/UnrealEd/Private/Fbx/FbxSceneImportOptionsSkeletalMesh.cpp
+++ b/Engine/Source/Editor/UnrealEd/Private/Fbx/FbxSceneImportOptionsSkeletalMesh.cpp
@@ -22,7 +22,7 @@ UFbxSceneImportOptionsSkeletalMesh::UFbxSceneImportOptionsSkeletalMesh(const FOb
 {
 }
 
-void UFbxSceneImportOptionsSkeletalMesh::FillSkeletalMeshInmportData(UFbxSkeletalMeshImportData* SkeletalMeshImportData, UFbxAnimSequenceImportData* AnimSequenceImportData)
+void UFbxSceneImportOptionsSkeletalMesh::FillSkeletalMeshInmportData(UFbxSkeletalMeshImportData* SkeletalMeshImportData, UFbxAnimSequenceImportData* AnimSequenceImportData, UFbxSceneImportOptions* SceneImportOptions)
 {
 	check(SkeletalMeshImportData != nullptr);
 	SkeletalMeshImportData->bImportMeshesInBoneHierarchy = bImportMeshesInBoneHierarchy;
@@ -31,6 +31,13 @@ void UFbxSceneImportOptionsSkeletalMesh::FillSkeletalMeshInmportData(UFbxSkeleta
 	SkeletalMeshImportData->bPreserveSmoothingGroups = bPreserveSmoothingGroups;
 	SkeletalMeshImportData->bUpdateSkeletonReferencePose = bUpdateSkeletonReferencePose;
 	SkeletalMeshImportData->bUseT0AsRefPose = bUseT0AsRefPose;
+
+	SkeletalMeshImportData->bImportMeshLODs = SceneImportOptions->bImportSkeletalMeshLODs;
+	SkeletalMeshImportData->ImportTranslation = SceneImportOptions->ImportTranslation;
+	SkeletalMeshImportData->ImportRotation = SceneImportOptions->ImportRotation;
+	SkeletalMeshImportData->ImportUniformScale = SceneImportOptions->ImportUniformScale;
+	SkeletalMeshImportData->bTransformVertexToAbsolute = SceneImportOptions->bTransformVertexToAbsolute;
+	SkeletalMeshImportData->bBakePivotInVertex = SceneImportOptions->bBakePivotInVertex;
 	
 	SkeletalMeshImportData->bImportAsScene = true;
 
@@ -40,4 +47,6 @@ void UFbxSceneImportOptionsSkeletalMesh::FillSkeletalMeshInmportData(UFbxSkeleta
 	AnimSequenceImportData->bPreserveLocalTransform = bPreserveLocalTransform;
 	AnimSequenceImportData->bUseDefaultSampleRate = bUseDefaultSampleRate;
 	AnimSequenceImportData->FrameImportRange = FrameImportRange;
+
+	AnimSequenceImportData->bImportAsScene = true;
 }
\ No newline at end of file
diff --git a/Engine/Source/Editor/UnrealEd/Private/Fbx/FbxSkeletalMeshImport.cpp b/Engine/Source/Editor/UnrealEd/Private/Fbx/FbxSkeletalMeshImport.cpp
--- a/Engine/Source/Editor/UnrealEd/Private/Fbx/FbxSkeletalMeshImport.cpp
+++ b/Engine/Source/Editor/UnrealEd/Private/Fbx/FbxSkeletalMeshImport.cpp
@@ -62,7 +62,7 @@
 using namespace UnFbx;
 
 struct ExistingSkelMeshData;
-extern ExistingSkelMeshData* SaveExistingSkelMeshData(USkeletalMesh* ExistingSkelMesh);
+extern ExistingSkelMeshData* SaveExistingSkelMeshData(USkeletalMesh* ExistingSkelMesh, bool bSaveMaterials);
 extern void RestoreExistingSkelMeshData(ExistingSkelMeshData* MeshData, USkeletalMesh* SkeletalMesh);
 
 // Get the geometry deformation local to a node. It is never inherited by the
@@ -764,7 +764,7 @@ bool UnFbx::FFbxImporter::RetrievePoseFromBindPose(const TArray<FbxNode*>& NodeA
 	return (PoseArray.Size() > 0);
 }
 
-bool UnFbx::FFbxImporter::ImportBone(TArray<FbxNode*>& NodeArray, FSkeletalMeshImportData &ImportData, UFbxSkeletalMeshImportData* TemplateData ,TArray<FbxNode*> &SortedLinks, bool& bOutDiffPose, bool bDisableMissingBindPoseWarning, bool & bUseTime0AsRefPose)
+bool UnFbx::FFbxImporter::ImportBone(TArray<FbxNode*>& NodeArray, FSkeletalMeshImportData &ImportData, UFbxSkeletalMeshImportData* TemplateData ,TArray<FbxNode*> &SortedLinks, bool& bOutDiffPose, bool bDisableMissingBindPoseWarning, bool & bUseTime0AsRefPose, FbxNode *SkeletalMeshNode)
 {
 	bOutDiffPose = false;
 	int32 SkelType = 0; // 0 for skeletal mesh, 1 for rigid mesh
@@ -865,8 +865,15 @@ bool UnFbx::FFbxImporter::ImportBone(TArray<FbxNode*>& NodeArray, FSkeletalMeshI
 	if (!bUseTime0AsRefPose && PoseArray.GetCount() == 0)
 	{
 		// add to tokenized error message
-		AddTokenizedErrorMessage(FTokenizedMessage::Create(EMessageSeverity::Warning, LOCTEXT("FbxSkeletaLMeshimport_InvalidBindPose", "Could not find the bind pose.  It will use time 0 as bind pose.")), FFbxErrors::SkeletalMesh_InvalidBindPose);
-		bUseTime0AsRefPose = true;
+		if (ImportOptions->bImportScene)
+		{
+			AddTokenizedErrorMessage(FTokenizedMessage::Create(EMessageSeverity::Warning, FText::Format(LOCTEXT("FbxSkeletaLMeshimport_InvalidBindPose", "Skeletal Mesh '{0}' dont have a bind pose. Scene import do not support yet time 0 as bind pose, there will be no bind pose import"), FText::FromString(NodeArray[0]->GetName()))), FFbxErrors::SkeletalMesh_InvalidBindPose);
+		}
+		else
+		{
+			AddTokenizedErrorMessage(FTokenizedMessage::Create(EMessageSeverity::Warning, LOCTEXT("FbxSkeletaLMeshimport_InvalidBindPose", "Could not find the bind pose.  It will use time 0 as bind pose.")), FFbxErrors::SkeletalMesh_InvalidBindPose);
+			bUseTime0AsRefPose = true;
+		}
 	}
 
 	int32 LinkIndex;
@@ -990,7 +997,7 @@ bool UnFbx::FFbxImporter::ImportBone(TArray<FbxNode*>& NodeArray, FSkeletalMeshI
 			GlobalsPerLink[LinkIndex] = Link->EvaluateGlobalTransform();
 		}
 		
-		if (bUseTime0AsRefPose)
+		if (bUseTime0AsRefPose && ImportOptions->bImportScene)
 		{
 			FbxAMatrix& T0Matrix = Scene->GetAnimationEvaluator()->GetNodeGlobalTransform(Link, 0);
 			if (GlobalsPerLink[LinkIndex] != T0Matrix)
@@ -1059,6 +1066,18 @@ bool UnFbx::FFbxImporter::ImportBone(TArray<FbxNode*>& NodeArray, FSkeletalMeshI
 		JointMatrix.Transform.SetScale3D(Converter.ConvertScale(LocalLinkS));
 	}
 	
+	//In case we do a scene import we need a relative to skeletal mesh transform instead of a global
+	if (ImportOptions->bImportScene && !ImportOptions->bTransformVertexToAbsolute)
+	{
+		FbxAMatrix GlobalSkeletalNodeFbx = Scene->GetAnimationEvaluator()->GetNodeGlobalTransform(SkeletalMeshNode, 0);
+		FTransform GlobalSkeletalNode;
+		GlobalSkeletalNode.SetFromMatrix(Converter.ConvertMatrix(GlobalSkeletalNodeFbx.Inverse()));
+
+		VBone& RootBone = ImportData.RefBonesBinary[RootIdx];
+		FTransform& RootTransform = RootBone.BonePos.Transform;
+		RootTransform.SetFromMatrix(RootTransform.ToMatrixWithScale() * GlobalSkeletalNode.ToMatrixWithScale());
+	}
+
 	if(TemplateData)
 	{
 		FbxAMatrix FbxAddedMatrix;
@@ -1078,72 +1097,32 @@ bool UnFbx::FFbxImporter::ImportBone(TArray<FbxNode*>& NodeArray, FSkeletalMeshI
 	return true;
 }
 
-USkeletalMesh* UnFbx::FFbxImporter::ImportSkeletalMesh(UObject* InParent, TArray<FbxNode*>& NodeArray, const FName& Name, EObjectFlags Flags, UFbxSkeletalMeshImportData* TemplateImportData, bool* bCancelOperation, TArray<FbxShape*> *FbxShapeArray, FSkeletalMeshImportData* OutData, bool bCreateRenderData )
+bool UnFbx::FFbxImporter::FillSkeletalMeshImportData(TArray<FbxNode*>& NodeArray, UFbxSkeletalMeshImportData* TemplateImportData, TArray<FbxShape*> *FbxShapeArray, FSkeletalMeshImportData* OutData)
 {
 	if (NodeArray.Num() == 0)
 	{
-		return NULL;
+		return false;
 	}
 
 	int32 SkelType = 0; // 0 for skeletal mesh, 1 for rigid mesh
-	
+
 	FbxNode* Node = NodeArray[0];
 	// find the mesh by its name
 	FbxMesh* FbxMesh = Node->GetMesh();
 
-	if( !FbxMesh )
+	if (OutData == nullptr)
 	{
-		AddTokenizedErrorMessage(FTokenizedMessage::Create(EMessageSeverity::Warning, FText::Format(LOCTEXT("FbxSkeletaLMeshimport_NodeInvalidSkeletalMesh", "Fbx node: '{0}' is not a valid skeletal mesh"), FText::FromString(Node->GetName()))), FFbxErrors::Generic_Mesh_MeshNotFound);
-		return NULL;
+		return false;
 	}
+
 	if (FbxMesh->GetDeformerCount(FbxDeformer::eSkin) == 0)
 	{
 		SkelType = 1;
 	}
 
-	// warning for missing smoothing group info
-	CheckSmoothingInfo(FbxMesh);
-
-	Parent = InParent;
-	
-
-	struct ExistingSkelMeshData* ExistSkelMeshDataPtr = NULL;
-
-	if ( !FbxShapeArray  )
-	{
-		UObject* ExistingObject = StaticFindObjectFast(UObject::StaticClass(), InParent, *Name.ToString(), false, false, RF_NoFlags, EInternalObjectFlags::PendingKill);
-		USkeletalMesh* ExistingSkelMesh = Cast<USkeletalMesh>(ExistingObject);
-
-		if (ExistingSkelMesh)
-		{
-#if WITH_APEX_CLOTHING
-			//for supporting re-import 
-			ApexClothingUtils::BackupClothingDataFromSkeletalMesh(ExistingSkelMesh);
-#endif// #if WITH_APEX_CLOTHING
-
-			ExistingSkelMesh->PreEditChange(NULL);
-			ExistSkelMeshDataPtr = SaveExistingSkelMeshData(ExistingSkelMesh);
-		}
-		// if any other object exists, we can't import with this name
-		else if (ExistingObject)
-		{
-			AddTokenizedErrorMessage(FTokenizedMessage::Create(EMessageSeverity::Warning, FText::Format(LOCTEXT("FbxSkeletaLMeshimport_OverlappingName", "Same name but different class: '{0}' already exists"), FText::FromString(ExistingObject->GetName()))), FFbxErrors::Generic_SameNameAssetExists);
-			return NULL;
-		}
-	}
-
-	// [from USkeletalMeshFactory::FactoryCreateBinary]
-	USkeletalMesh* SkeletalMesh = NewObject<USkeletalMesh>(InParent, Name, Flags);
-	
-	SkeletalMesh->PreEditChange(NULL);
-	
-	FSkeletalMeshImportData TempData;
 	// Fill with data from buffer - contains the full .FBX file. 	
-	FSkeletalMeshImportData* SkelMeshImportDataPtr = &TempData;
-	if( OutData )
-	{
-		SkelMeshImportDataPtr = OutData;
-	}
+	FSkeletalMeshImportData* SkelMeshImportDataPtr = OutData;
+
 
 	bool bDiffPose;
 	TArray<FbxNode*> SortedLinkArray;
@@ -1152,26 +1131,22 @@ USkeletalMesh* UnFbx::FFbxImporter::ImportSkeletalMesh(UObject* InParent, TArray
 	bool bUseTime0AsRefPose = ImportOptions->bUseT0AsRefPose;
 	// Note: importing morph data causes additional passes through this function, so disable the warning dialogs
 	// from popping up again on each additional pass.  
-	if ( !ImportBone(NodeArray, *SkelMeshImportDataPtr, TemplateImportData, SortedLinkArray,  bDiffPose, (FbxShapeArray != NULL), bUseTime0AsRefPose))
+	if (!ImportBone(NodeArray, *SkelMeshImportDataPtr, TemplateImportData, SortedLinkArray, bDiffPose, (FbxShapeArray != nullptr), bUseTime0AsRefPose, Node))
 	{
 		AddTokenizedErrorMessage(FTokenizedMessage::Create(EMessageSeverity::Error, LOCTEXT("FbxSkeletaLMeshimport_MultipleRootFound", "Multiple roots found")), FFbxErrors::SkeletalMesh_MultipleRoots);
-		// I can't delete object here since this is middle of import
-		// but I can move to transient package, and GC will automatically collect it
-		SkeletalMesh->ClearFlags( RF_Standalone );
-		SkeletalMesh->Rename(NULL, GetTransientPackage());
-		return NULL;
+		return false;
 	}
 
- 	FbxNode* RootNode = Scene->GetRootNode();
- 	if(RootNode && TemplateImportData)
+ 	FbxNode* SceneRootNode = Scene->GetRootNode();
+ 	if(SceneRootNode && TemplateImportData)
  	{
- 		ApplyTransformSettingsToFbxNode(RootNode, TemplateImportData);
+ 		ApplyTransformSettingsToFbxNode(SceneRootNode, TemplateImportData);
  	}
 
 	// Create a list of all unique fbx materials.  This needs to be done as a separate pass before reading geometry
 	// so that we know about all possible materials before assigning material indices to each triangle
 	TArray<FbxSurfaceMaterial*> FbxMaterials;
-	for( int32 NodeIndex = 0; NodeIndex < NodeArray.Num(); ++NodeIndex )
+	for (int32 NodeIndex = 0; NodeIndex < NodeArray.Num(); ++NodeIndex)
 	{
 		Node = NodeArray[NodeIndex];
 
@@ -1179,77 +1154,156 @@ USkeletalMesh* UnFbx::FFbxImporter::ImportSkeletalMesh(UObject* InParent, TArray
 
 		for (int32 MaterialIndex = 0; MaterialIndex < MaterialCount; ++MaterialIndex)
 		{
-			FbxSurfaceMaterial* FbxMaterial = Node->GetMaterial( MaterialIndex );
-			if( !FbxMaterials.Contains( FbxMaterial ) )
+			FbxSurfaceMaterial* FbxMaterial = Node->GetMaterial(MaterialIndex);
+			if (!FbxMaterials.Contains(FbxMaterial))
 			{
-				FbxMaterials.Add( FbxMaterial );
+				FbxMaterials.Add(FbxMaterial);
 
 				VMaterial NewMaterial;
 
 				NewMaterial.MaterialImportName = UTF8_TO_TCHAR(MakeName(FbxMaterial->GetName()));
 				// Add an entry for each unique material
-				SkelMeshImportDataPtr->Materials.Add( NewMaterial );
+				SkelMeshImportDataPtr->Materials.Add(NewMaterial);
 			}
 		}
 	}
 
-	for ( int32 NodeIndex = 0; NodeIndex < NodeArray.Num(); ++NodeIndex)
+	for (int32 NodeIndex = 0; NodeIndex < NodeArray.Num(); ++NodeIndex)
 	{
 		Node = NodeArray[NodeIndex];
+		FbxNode *RootNode = NodeArray[0];
 		FbxMesh = Node->GetMesh();
 		FbxSkin* Skin = (FbxSkin*)FbxMesh->GetDeformer(0, FbxDeformer::eSkin);
-		FbxShape* FbxShape = NULL;
+		FbxShape* FbxShape = nullptr;
 		if (FbxShapeArray)
 		{
 			FbxShape = (*FbxShapeArray)[NodeIndex];
 		}
 
 		// NOTE: This function may invalidate FbxMesh and set it to point to a an updated version
-		if (!FillSkelMeshImporterFromFbx( *SkelMeshImportDataPtr, FbxMesh, Skin, FbxShape, SortedLinkArray, FbxMaterials ) )
+		if (!FillSkelMeshImporterFromFbx( *SkelMeshImportDataPtr, FbxMesh, Skin, FbxShape, SortedLinkArray, FbxMaterials, RootNode) )
 		{
-			// I can't delete object here since this is middle of import
-			// but I can move to transient package, and GC will automatically collect it
-			SkeletalMesh->ClearFlags( RF_Standalone );
-			SkeletalMesh->Rename(NULL, GetTransientPackage());
-			return NULL;
+			return false;
 		}
-		
-		if (bUseTime0AsRefPose && bDiffPose)
+
+		if (bUseTime0AsRefPose && bDiffPose && !ImportOptions->bImportScene)
 		{
 			// deform skin vertex to the frame 0 from bind pose
 			SkinControlPointsToPose(*SkelMeshImportDataPtr, FbxMesh, FbxShape, true);
 		}
 	}
-	
+
 	// reorder material according to "SKinXX" in material name
-	SetMaterialSkinXXOrder(*SkelMeshImportDataPtr );
-	
-	if( ImportOptions->bPreserveSmoothingGroups )
+	SetMaterialSkinXXOrder(*SkelMeshImportDataPtr);
+
+	if (ImportOptions->bPreserveSmoothingGroups)
 	{
 		bool bDuplicateUnSmoothWedges = (ImportOptions->NormalGenerationMethod != EFBXNormalGenerationMethod::MikkTSpace);
 		DoUnSmoothVerts(*SkelMeshImportDataPtr, bDuplicateUnSmoothWedges);
 	}
 	else
 	{
-		SkelMeshImportDataPtr->PointToRawMap.AddUninitialized( SkelMeshImportDataPtr->Points.Num() );
-		for(int32 PointIdx=0; PointIdx<SkelMeshImportDataPtr->Points.Num(); PointIdx++)
+		SkelMeshImportDataPtr->PointToRawMap.AddUninitialized(SkelMeshImportDataPtr->Points.Num());
+		for (int32 PointIdx = 0; PointIdx < SkelMeshImportDataPtr->Points.Num(); PointIdx++)
 		{
 			SkelMeshImportDataPtr->PointToRawMap[PointIdx] = PointIdx;
 		}
 	}
 
-	// process materials from import data
-	ProcessImportMeshMaterials(SkeletalMesh->Materials,*SkelMeshImportDataPtr);
+	return true;
+}
+
+USkeletalMesh* UnFbx::FFbxImporter::ImportSkeletalMesh(UObject* InParent, TArray<FbxNode*>& NodeArray, const FName& Name, EObjectFlags Flags, UFbxSkeletalMeshImportData* TemplateImportData, bool* bCancelOperation, TArray<FbxShape*> *FbxShapeArray, FSkeletalMeshImportData* OutData, bool bCreateRenderData )
+{
+	if (NodeArray.Num() == 0)
+	{
+		return nullptr;
+	}
+
+	int32 SkelType = 0; // 0 for skeletal mesh, 1 for rigid mesh
 	
+	FbxNode* Node = NodeArray[0];
+	// find the mesh by its name
+	FbxMesh* FbxMesh = Node->GetMesh();
+
+	if( !FbxMesh )
+	{
+		AddTokenizedErrorMessage(FTokenizedMessage::Create(EMessageSeverity::Warning, FText::Format(LOCTEXT("FbxSkeletaLMeshimport_NodeInvalidSkeletalMesh", "Fbx node: '{0}' is not a valid skeletal mesh"), FText::FromString(Node->GetName()))), FFbxErrors::Generic_Mesh_MeshNotFound);
+		return nullptr;
+	}
+	if (FbxMesh->GetDeformerCount(FbxDeformer::eSkin) == 0)
+	{
+		SkelType = 1;
+	}
+
+	// warning for missing smoothing group info
+	CheckSmoothingInfo(FbxMesh);
+
+	Parent = InParent;
+	
+
+	struct ExistingSkelMeshData* ExistSkelMeshDataPtr = nullptr;
+
+	if ( !FbxShapeArray  )
+	{
+		UObject* ExistingObject = StaticFindObjectFast(UObject::StaticClass(), InParent, *Name.ToString(), false, false, RF_NoFlags, EInternalObjectFlags::PendingKill);
+		USkeletalMesh* ExistingSkelMesh = Cast<USkeletalMesh>(ExistingObject);
+
+		if (ExistingSkelMesh)
+		{
+#if WITH_APEX_CLOTHING
+			//for supporting re-import 
+			ApexClothingUtils::BackupClothingDataFromSkeletalMesh(ExistingSkelMesh);
+#endif// #if WITH_APEX_CLOTHING
+
+			ExistingSkelMesh->PreEditChange(NULL);
+			ExistSkelMeshDataPtr = SaveExistingSkelMeshData(ExistingSkelMesh,!ImportOptions->bImportScene);
+		}
+		// if any other object exists, we can't import with this name
+		else if (ExistingObject)
+		{
+			AddTokenizedErrorMessage(FTokenizedMessage::Create(EMessageSeverity::Warning, FText::Format(LOCTEXT("FbxSkeletaLMeshimport_OverlappingName", "Same name but different class: '{0}' already exists"), FText::FromString(ExistingObject->GetName()))), FFbxErrors::Generic_SameNameAssetExists);
+			return NULL;
+		}
+	}
+
+	// [from USkeletalMeshFactory::FactoryCreateBinary]
+	USkeletalMesh* SkeletalMesh = NewObject<USkeletalMesh>(InParent, Name, Flags);
+	
+	SkeletalMesh->PreEditChange(NULL);
+	
+	FSkeletalMeshImportData TempData;
+	// Fill with data from buffer - contains the full .FBX file. 	
+	FSkeletalMeshImportData* SkelMeshImportDataPtr = &TempData;
+	if( OutData )
+	{
+		SkelMeshImportDataPtr = OutData;
+	}
+
+	if (FillSkeletalMeshImportData(NodeArray, TemplateImportData, FbxShapeArray, SkelMeshImportDataPtr) == false)
+	{
+		AddTokenizedErrorMessage(FTokenizedMessage::Create(EMessageSeverity::Warning, LOCTEXT("FbxSkeletaLMeshimport_FillupImportData", "Get Import Data has failed.")), FFbxErrors::SkeletalMesh_FillImportDataFailed);
+
+		// I can't delete object here since this is middle of import
+		// but I can move to transient package, and GC will automatically collect it
+		SkeletalMesh->ClearFlags(RF_Standalone);
+		SkeletalMesh->Rename(NULL, GetTransientPackage());
+		return nullptr;
+	}
+
+	// process materials from import data
+	ProcessImportMeshMaterials(SkeletalMesh->Materials, *SkelMeshImportDataPtr);
+
 	// process reference skeleton from import data
 	int32 SkeletalDepth = 0;
-	if(!ProcessImportMeshSkeleton(SkeletalMesh->RefSkeleton,SkeletalDepth,*SkelMeshImportDataPtr))
+	if (!ProcessImportMeshSkeleton(SkeletalMesh->RefSkeleton, SkeletalDepth, *SkelMeshImportDataPtr))
 	{
 		SkeletalMesh->ClearFlags( RF_Standalone );
 		SkeletalMesh->Rename(NULL, GetTransientPackage());
 		return NULL;
 	}
-	UE_LOG(LogFbx, Warning, TEXT("Bones digested - %i  Depth of hierarchy - %i"), SkeletalMesh->RefSkeleton.GetNum(), SkeletalDepth );
+
+	UE_LOG(LogFbx, Warning, TEXT("Bones digested - %i  Depth of hierarchy - %i"), SkeletalMesh->RefSkeleton.GetNum(), SkeletalDepth);
 
 	// process bone influences from import data
 	ProcessImportMeshInfluences(*SkelMeshImportDataPtr);
@@ -1619,8 +1673,13 @@ USkeletalMesh* UnFbx::FFbxImporter::ReimportSkeletalMesh(USkeletalMesh* Mesh, UF
 	if (FbxNodes)
 	{
 		// set import options, how about others?
-		ImportOptions->bImportMaterials = false;
-		ImportOptions->bImportTextures = false;
+		if (!ImportOptions->bImportScene)
+		{
+			//When importing a scene those options are set by the user
+			ImportOptions->bImportMaterials = false;
+			ImportOptions->bImportTextures = false;
+		}
+		//In case of a scene reimport animations are reimport later so its ok to hardcode animation to false here
 		ImportOptions->bImportAnimations = false;
 		// check if there is LODGroup for this skeletal mesh
 		int32 MaxLODLevel = 1;
@@ -1711,9 +1770,9 @@ USkeletalMesh* UnFbx::FFbxImporter::ReimportSkeletalMesh(USkeletalMesh* Mesh, UF
 void UnFbx::FFbxImporter::SetMaterialSkinXXOrder(FSkeletalMeshImportData& ImportData )
 {
 	TArray<int32> MaterialIndexToSkinIndex;
+	TMap<int32, int32> SkinIndexToMaterialIndex;
 	{
 		int32 MaterialCount = ImportData.Materials.Num();
-		MaterialIndexToSkinIndex.AddZeroed(MaterialCount);
 
 		bool bNeedsReorder = false;
 		for(int32 MaterialIndex=0; MaterialIndex < MaterialCount; ++MaterialIndex)
@@ -1734,7 +1793,7 @@ void UnFbx::FFbxImporter::SetMaterialSkinXXOrder(FSkeletalMeshImportData& Import
 						bNeedsReorder = true;
 
 						int32 TmpIndex = FPlatformString::Atoi( *SkinXXNumber );
-						MaterialIndexToSkinIndex[MaterialIndex] = TmpIndex;
+						SkinIndexToMaterialIndex.Add(MaterialIndex, TmpIndex);
 
 						// remove the 'skinXX' suffix from the material name					
 						ImportData.Materials[MaterialIndex].MaterialImportName.LeftChop( Offset );
@@ -1743,6 +1802,23 @@ void UnFbx::FFbxImporter::SetMaterialSkinXXOrder(FSkeletalMeshImportData& Import
 			}
 		}
 
+		//Fill the array MaterialIndexToSkinIndex so we order material by _skinXX order
+		for (auto kvp : SkinIndexToMaterialIndex)
+		{
+			int32 MatIndexToInsert = 0;
+			for (MatIndexToInsert = 0; MatIndexToInsert < MaterialIndexToSkinIndex.Num(); ++MatIndexToInsert)
+			{
+				//the item should always be in the tmap
+				check(SkinIndexToMaterialIndex.Contains(MaterialIndexToSkinIndex[MatIndexToInsert]));
+
+				if (*(SkinIndexToMaterialIndex.Find(MaterialIndexToSkinIndex[MatIndexToInsert])) >= kvp.Value)
+				{
+					break;
+				}
+			}
+			MaterialIndexToSkinIndex.Insert(kvp.Key, MatIndexToInsert);
+		}
+
 		if (bNeedsReorder)
 		{
 			// re-order the materials
@@ -1768,7 +1844,7 @@ void UnFbx::FFbxImporter::SetMaterialSkinXXOrder(FSkeletalMeshImportData& Import
 	}
 }
 
-bool UnFbx::FFbxImporter::FillSkelMeshImporterFromFbx( FSkeletalMeshImportData& ImportData, FbxMesh*& Mesh, FbxSkin* Skin, FbxShape* FbxShape, TArray<FbxNode*> &SortedLinks, const TArray<FbxSurfaceMaterial*>& FbxMaterials )
+bool UnFbx::FFbxImporter::FillSkelMeshImporterFromFbx( FSkeletalMeshImportData& ImportData, FbxMesh*& Mesh, FbxSkin* Skin, FbxShape* FbxShape, TArray<FbxNode*> &SortedLinks, const TArray<FbxSurfaceMaterial*>& FbxMaterials, FbxNode *RootNode)
 {
 	FbxNode* Node = Mesh->GetNode();
 
@@ -1816,7 +1892,10 @@ bool UnFbx::FFbxImporter::FillSkelMeshImporterFromFbx( FSkeletalMeshImportData&
 					{
 						const char* UVSetName = ElementUV->GetName();
 						FString LocalUVSetName = UTF8_TO_TCHAR(UVSetName);
-
+						if (LocalUVSetName.IsEmpty())
+						{
+							LocalUVSetName = TEXT("UVmap_") + FString::FromInt(UVLayerIndex);
+						}
 						UVSets.AddUnique(LocalUVSetName);
 					}
 				}
@@ -1855,16 +1934,6 @@ bool UnFbx::FFbxImporter::FillSkelMeshImporterFromFbx( FSkeletalMeshImportData&
 	if (ImportOptions->bImportMaterials)
 	{
 		CreateNodeMaterials(Node, Materials, UVSets);
-		//Make sure all material use by skeletal mesh are flag to be use by skeletal mesh
-		for (UMaterialInterface *MaterialInterface : Materials)
-		{
-			if (MaterialInterface == nullptr || MaterialInterface->GetMaterial() == nullptr)
-			{
-				continue;
-			}
-			bool bNeedRecompile = false;
-			MaterialInterface->GetMaterial()->SetMaterialUsage(bNeedRecompile, EMaterialUsage::MATUSAGE_SkeletalMesh);
-		}
 	}
 	else if (ImportOptions->bImportTextures)
 	{
@@ -1997,6 +2066,10 @@ bool UnFbx::FFbxImporter::FillSkelMeshImporterFromFbx( FSkeletalMeshImportData&
 					{
 						const char* UVSetName = ElementUV->GetName();
 						FString LocalUVSetName = UTF8_TO_TCHAR(UVSetName);
+						if (LocalUVSetName.IsEmpty())
+						{
+							LocalUVSetName = TEXT("UVmap_") + FString::FromInt(UVLayerIndex);
+						}
 						if (LocalUVSetName == UVSets[UVIndex])
 						{
 							LayerElementUV[UVIndex] = const_cast<FbxLayerElementUV*>(ElementUV);
@@ -2105,7 +2178,7 @@ bool UnFbx::FFbxImporter::FillSkelMeshImporterFromFbx( FSkeletalMeshImportData&
 	// Construct the matrices for the conversion from right handed to left handed system
 	FbxAMatrix TotalMatrix;
 	FbxAMatrix TotalMatrixForNormal;
-	TotalMatrix = ComputeTotalMatrix(Node);
+	TotalMatrix = ComputeSkeletalMeshTotalMatrix(Node, RootNode);
 	TotalMatrixForNormal = TotalMatrix.Inverse();
 	TotalMatrixForNormal = TotalMatrixForNormal.Transpose();
 
@@ -2745,10 +2818,10 @@ bool UnFbx::FFbxImporter::ImportSkeletalMeshLOD(USkeletalMesh* InSkeletalMesh, U
 class FAsyncImportMorphTargetWork : public FNonAbandonableTask
 {
 public:
-	FAsyncImportMorphTargetWork( USkeletalMesh* InTempSkelMesh, int32 InLODIndex, FSkeletalMeshImportData& InImportData, bool bInKeepOverlappingVertices )
-		: TempSkeletalMesh(InTempSkelMesh)
+	FAsyncImportMorphTargetWork(FStaticLODModel* InLODModel, const FReferenceSkeleton& InRefSkeleton, FSkeletalMeshImportData& InImportData, bool bInKeepOverlappingVertices)
+		: LODModel(InLODModel)
+		, RefSkeleton(InRefSkeleton)
 		, ImportData(InImportData)
-		, LODIndex(InLODIndex)
 		, bKeepOverlappingVertices(bInKeepOverlappingVertices)
 	{
 		MeshUtilities = &FModuleManager::Get().LoadModuleChecked<IMeshUtilities>("MeshUtilities");
@@ -2763,11 +2836,14 @@ public:
 		TArray<int32> LODPointToRawMap;
 		ImportData.CopyLODImportData(LODPoints,LODWedges,LODFaces,LODInfluences,LODPointToRawMap);
 
+		check(LODModel);
+
 		IMeshUtilities::MeshBuildOptions BuildOptions;
 		BuildOptions.bKeepOverlappingVertices = bKeepOverlappingVertices;
-	
+
+		// @why we have to empty here? @todo: possibly we save this and reuse?
 		ImportData.Empty();
-		MeshUtilities->BuildSkeletalMesh( TempSkeletalMesh->GetImportedResource()->LODModels[0], TempSkeletalMesh->RefSkeleton, LODInfluences, LODWedges, LODFaces, LODPoints, LODPointToRawMap, BuildOptions);
+		MeshUtilities->BuildSkeletalMesh( *LODModel, RefSkeleton, LODInfluences, LODWedges, LODFaces, LODPoints, LODPointToRawMap, BuildOptions);
 	}
 
 	FORCEINLINE TStatId GetStatId() const
@@ -2776,10 +2852,11 @@ public:
 	}
 
 private:
-	USkeletalMesh* TempSkeletalMesh;
+	FStaticLODModel* LODModel;
+	// @todo not thread safe
+	const FReferenceSkeleton& RefSkeleton;
 	FSkeletalMeshImportData ImportData;
 	IMeshUtilities* MeshUtilities;
-	int32 LODIndex;
 	bool bKeepOverlappingVertices;
 };
 
@@ -2789,7 +2866,7 @@ void UnFbx::FFbxImporter::ImportMorphTargetsInternal( TArray<FbxNode*>& SkelMesh
 	TMap<FString, TArray<FbxShape*>> ShapeNameToShapeArray;
 
 	// Temp arrays to keep track of data being used by threads
-	TArray<USkeletalMesh*> TempMeshes;
+	TArray<FStaticLODModel*> LODModels;
 	TArray<UMorphTarget*> MorphTargets;
 
 	// Array of pending tasks that are not complete
@@ -2910,12 +2987,13 @@ void UnFbx::FFbxImporter::ImportMorphTargetsInternal( TArray<FbxNode*>& SkelMesh
 			// now we get a shape for whole mesh, import to unreal as a morph target
 			// @todo AssetImportData do we need import data for this temp mesh?
 			UFbxSkeletalMeshImportData* TmpMeshImportData = NULL;
-			USkeletalMesh* TmpSkeletalMesh = (USkeletalMesh*)ImportSkeletalMesh(GetTransientPackage(), SkelMeshNodeArray, NAME_None, (EObjectFlags)0, TmpMeshImportData, nullptr, &ShapeArray, &ImportData, false);
-			TempMeshes.Add(TmpSkeletalMesh);
+			FillSkeletalMeshImportData(SkelMeshNodeArray, TmpMeshImportData, &ShapeArray, &ImportData);
+			FStaticLODModel* NewLODModel = new FStaticLODModel;
+			LODModels.Add(NewLODModel);
 			MorphTargets.Add(Result);
 
 			// Process the skeletal mesh on a separate thread
-			FAsyncTask<FAsyncImportMorphTargetWork>* NewWork = new (PendingWork)FAsyncTask<FAsyncImportMorphTargetWork>(TmpSkeletalMesh, LODIndex, ImportData, ImportOptions->bKeepOverlappingVertices);
+			FAsyncTask<FAsyncImportMorphTargetWork>* NewWork = new (PendingWork)FAsyncTask<FAsyncImportMorphTargetWork>(NewLODModel, BaseSkelMesh->RefSkeleton, ImportData, ImportOptions->bKeepOverlappingVertices);
 			NewWork->StartBackgroundTask();
 			++ShapeIndex;
 		}
@@ -2948,20 +3026,22 @@ void UnFbx::FFbxImporter::ImportMorphTargetsInternal( TArray<FbxNode*>& SkelMesh
 
 	// Create morph streams for each morph target we are importing.
 	// This has to happen on a single thread since the skeletal meshes' bulk data is locked and cant be accessed by multiple threads simultaneously
-	for (int32 Index = 0; Index < TempMeshes.Num(); Index++)
+	for (int32 Index = 0; Index < MorphTargets.Num(); Index++)
 	{
 		FFormatNamedArguments Args;
 		Args.Add( TEXT("NumCompleted"), Index+1 );
-		Args.Add( TEXT("NumTasks"), TempMeshes.Num() );
-		GWarn->StatusUpdate( Index+1, TempMeshes.Num(), FText::Format( LOCTEXT("BuildingMorphTargetRenderDataStatus", "Building Morph Target Render Data: {NumCompleted} of {NumTasks}"), Args ) );
+		Args.Add(TEXT("NumTasks"), MorphTargets.Num());
+		GWarn->StatusUpdate(Index + 1, MorphTargets.Num(), FText::Format(LOCTEXT("BuildingMorphTargetRenderDataStatus", "Building Morph Target Render Data: {NumCompleted} of {NumTasks}"), Args));
 
 		UMorphTarget* MorphTarget = MorphTargets[Index];
-		USkeletalMesh* TmpSkeletalMesh = TempMeshes[ Index ];
 
-		FMorphMeshRawSource TargetMeshRawData( TmpSkeletalMesh );
+		FMorphMeshRawSource TargetMeshRawData(*LODModels[Index]);
 		FMorphMeshRawSource BaseMeshRawData( BaseSkelMesh, LODIndex );
 
 		MorphTarget->PostProcess( BaseSkelMesh, BaseMeshRawData, TargetMeshRawData, LODIndex, ImportOptions->ShouldImportNormals() == false );
+
+		delete LODModels[Index];
+		LODModels[Index] = nullptr;
 	}
 
 	GWarn->EndSlowTask();
diff --git a/Engine/Source/Editor/UnrealEd/Private/Fbx/FbxSkeletalMeshImportData.cpp b/Engine/Source/Editor/UnrealEd/Private/Fbx/FbxSkeletalMeshImportData.cpp
--- a/Engine/Source/Editor/UnrealEd/Private/Fbx/FbxSkeletalMeshImportData.cpp
+++ b/Engine/Source/Editor/UnrealEd/Private/Fbx/FbxSkeletalMeshImportData.cpp
@@ -6,6 +6,8 @@ UFbxSkeletalMeshImportData::UFbxSkeletalMeshImportData(const FObjectInitializer&
 	: Super(ObjectInitializer)
 	, bImportMeshesInBoneHierarchy(true)
 {
+	bTransformVertexToAbsolute = true;
+	bBakePivotInVertex = false;
 }
 
 UFbxSkeletalMeshImportData* UFbxSkeletalMeshImportData::GetImportDataForSkeletalMesh(USkeletalMesh* SkeletalMesh, UFbxSkeletalMeshImportData* TemplateForCreation)
diff --git a/Engine/Source/Editor/UnrealEd/Private/Fbx/FbxStaticMeshImport.cpp b/Engine/Source/Editor/UnrealEd/Private/Fbx/FbxStaticMeshImport.cpp
--- a/Engine/Source/Editor/UnrealEd/Private/Fbx/FbxStaticMeshImport.cpp
+++ b/Engine/Source/Editor/UnrealEd/Private/Fbx/FbxStaticMeshImport.cpp
@@ -139,6 +139,10 @@ bool UnFbx::FFbxImporter::BuildStaticMeshFromGeometry(FbxNode* Node, UStaticMesh
 					{
 						const char* UVSetName = ElementUV->GetName();
 						FString LocalUVSetName = UTF8_TO_TCHAR(UVSetName);
+						if (LocalUVSetName.IsEmpty())
+						{
+							LocalUVSetName = TEXT("UVmap_") + FString::FromInt(UVLayerIndex);
+						}
 
 						UVSets.AddUnique(LocalUVSetName);
 					}
@@ -310,6 +314,10 @@ bool UnFbx::FFbxImporter::BuildStaticMeshFromGeometry(FbxNode* Node, UStaticMesh
 					{
 						const char* UVSetName = ElementUV->GetName();
 						FString LocalUVSetName = UTF8_TO_TCHAR(UVSetName);
+						if (LocalUVSetName.IsEmpty())
+						{
+							LocalUVSetName = TEXT("UVmap_") + FString::FromInt(UVLayerIndex);
+						}
 						if (LocalUVSetName == UVSets[UVIndex])
 						{
 							LayerElementUV[UVIndex] = ElementUV;
diff --git a/Engine/Source/Editor/UnrealEd/Private/Fbx/ReimportFbxSceneFactory.cpp b/Engine/Source/Editor/UnrealEd/Private/Fbx/ReimportFbxSceneFactory.cpp
--- a/Engine/Source/Editor/UnrealEd/Private/Fbx/ReimportFbxSceneFactory.cpp
+++ b/Engine/Source/Editor/UnrealEd/Private/Fbx/ReimportFbxSceneFactory.cpp
@@ -51,7 +51,7 @@ UFbxSceneImportData *GetFbxSceneImportData(UObject *Obj)
 				SceneImportData = ImportData->bImportAsScene ? ImportData->FbxSceneImportDataReference : nullptr;
 			}
 		}
-		if (Obj->IsA(USkeletalMesh::StaticClass()))
+		else if (Obj->IsA(USkeletalMesh::StaticClass()))
 		{
 			//Reimport from one of the static mesh
 			USkeletalMesh* SkeletalMesh = Cast<USkeletalMesh>(Obj);
@@ -61,6 +61,16 @@ UFbxSceneImportData *GetFbxSceneImportData(UObject *Obj)
 				SceneImportData = ImportData->bImportAsScene ? ImportData->FbxSceneImportDataReference : nullptr;
 			}
 		}
+		else if (Obj->IsA(UAnimSequence::StaticClass()))
+		{
+			//Reimport from one of the static mesh
+			UAnimSequence* AnimSequence = Cast<UAnimSequence>(Obj);
+			if (AnimSequence != nullptr && AnimSequence->AssetImportData != nullptr)
+			{
+				ImportData = Cast<UFbxAssetImportData>(AnimSequence->AssetImportData);
+				SceneImportData = ImportData->bImportAsScene ? ImportData->FbxSceneImportDataReference : nullptr;
+			}
+		}
 		//TODO: add all type the fbx scene import can create: material, texture, skeletal mesh, animation, ... 
 	}
 	return SceneImportData;
@@ -191,6 +201,12 @@ bool GetFbxSceneReImportOptions(UnFbx::FFbxImporter* FbxImporter
 	GlobalImportSettings->bImportTextures = true;
 	//Make sure Material get imported
 	GlobalImportSettings->bImportMaterials = true;
+	//TODO support T0AsRefPose
+	GlobalImportSettings->bUseT0AsRefPose = false;
+
+	GlobalImportSettings->ImportTranslation = FVector(0);
+	GlobalImportSettings->ImportRotation = FRotator(0);
+	GlobalImportSettings->ImportUniformScale = 1.0f;
 
 	GlobalImportSettings->bConvertScene = true;
 
@@ -240,6 +256,9 @@ bool GetFbxSceneReImportOptions(UnFbx::FFbxImporter* FbxImporter
 	GlobalImportSettings->bImportStaticMeshLODs = SceneImportOptions->bImportStaticMeshLODs;
 	GlobalImportSettings->bImportSkeletalMeshLODs = SceneImportOptions->bImportSkeletalMeshLODs;
 	SceneImportOptions->bInvertNormalMaps = GlobalImportSettings->bInvertNormalMap;
+	GlobalImportSettings->ImportTranslation = SceneImportOptions->ImportTranslation;
+	GlobalImportSettings->ImportRotation = SceneImportOptions->ImportRotation;
+	GlobalImportSettings->ImportUniformScale = SceneImportOptions->ImportUniformScale;
 
 	//Set the override material into the options
 	for (TSharedPtr<FFbxNodeInfo> NodeInfo : SceneInfoPtr->HierarchyInfo)
@@ -270,6 +289,27 @@ bool GetFbxSceneReImportOptions(UnFbx::FFbxImporter* FbxImporter
 	}
 
 	SceneImportOptions->SaveConfig();
+
+	//Save the Default setting copy them in the UObject and save them
+	SFbxSceneOptionWindow::CopyFbxOptionsToStaticMeshOptions(GlobalImportSettings, StaticMeshImportData);
+	StaticMeshImportData->SaveConfig();
+
+	SFbxSceneOptionWindow::CopyFbxOptionsToSkeletalMeshOptions(GlobalImportSettings, SkeletalMeshImportData);
+	SkeletalMeshImportData->SaveConfig();
+
+	//Make sure default option set will not override fbx global setting by making copy of the real default options
+	ImportOptionsNameMap TmpNameOptionsMap;
+	for (auto kvp : NameOptionsMap)
+	{
+		UnFbx::FBXImportOptions *NewOptions = new UnFbx::FBXImportOptions();
+		SFbxSceneOptionWindow::CopyFbxOptionsToFbxOptions(kvp.Value, NewOptions);
+		TmpNameOptionsMap.Add(kvp.Key, NewOptions);
+	}
+	NameOptionsMap.Reset();
+	for (auto kvp : TmpNameOptionsMap)
+	{
+		NameOptionsMap.Add(kvp.Key, kvp.Value);
+	}
 	return true;
 }
 
@@ -280,6 +320,7 @@ EReimportResult::Type UReimportFbxSceneFactory::Reimport(UObject* Obj)
 	{
 		return EReimportResult::Failed;
 	}
+	NameOptionsMap.Reset();
 
 	//We will call other factory store the filename value since UFactory::CurrentFilename is static
 	FbxImportFileName = ReimportData->SourceFbxFile;
@@ -289,12 +330,14 @@ EReimportResult::Type UReimportFbxSceneFactory::Reimport(UObject* Obj)
 	GWarn->BeginSlowTask(NSLOCTEXT("FbxSceneReImportFactory", "BeginReImportingFbxSceneTask", "ReImporting FBX scene"), true);
 
 	GlobalImportSettings = FbxImporter->GetImportOptions();
+	UnFbx::FBXImportOptions::ResetOptions(GlobalImportSettings);
 
 	//Fill the original options
 	for (auto kvp : ReimportData->NameOptionsMap)
 	{
 		if (kvp.Key.Compare(DefaultOptionName) == 0)
 		{
+			//Save the default option to the fbx default import setting
 			SFbxSceneOptionWindow::CopyFbxOptionsToFbxOptions(kvp.Value, GlobalImportSettings);
 			NameOptionsMap.Add(kvp.Key, GlobalImportSettings);
 		}
@@ -329,6 +372,13 @@ EReimportResult::Type UReimportFbxSceneFactory::Reimport(UObject* Obj)
 	GlobalImportSettings->SkeletonForAnimation = nullptr;
 	GlobalImportSettings->PhysicsAsset = nullptr;
 
+	SFbxSceneOptionWindow::CopyFbxOptionsToStaticMeshOptions(GlobalImportSettings, SceneImportOptionsStaticMesh);
+	SFbxSceneOptionWindow::CopyFbxOptionsToSkeletalMeshOptions(GlobalImportSettings, SceneImportOptionsSkeletalMesh);
+	SceneImportOptions->bBakePivotInVertex = GlobalImportSettings->bBakePivotInVertex;
+	SceneImportOptions->bTransformVertexToAbsolute = GlobalImportSettings->bTransformVertexToAbsolute;
+	SceneImportOptions->bImportStaticMeshLODs = GlobalImportSettings->bImportStaticMeshLODs;
+	SceneImportOptions->bImportSkeletalMeshLODs = GlobalImportSettings->bImportSkeletalMeshLODs;
+
 	FString PackageName = "";
 	Obj->GetOutermost()->GetName(PackageName);
 	Path = FPaths::GetPath(PackageName);
@@ -494,7 +544,7 @@ EReimportResult::Type UReimportFbxSceneFactory::Reimport(UObject* Obj)
 	StaticMeshImportData->SaveConfig();
 
 	SFbxSceneOptionWindow::CopyFbxOptionsToSkeletalMeshOptions(GlobalImportSettingsReference, SceneImportOptionsSkeletalMesh);
-	SceneImportOptionsSkeletalMesh->FillSkeletalMeshInmportData(SkeletalMeshImportData, AnimSequenceImportData);
+	SceneImportOptionsSkeletalMesh->FillSkeletalMeshInmportData(SkeletalMeshImportData, AnimSequenceImportData, SceneImportOptions);
 	SkeletalMeshImportData->SaveConfig();
 
 	//Update the blueprint
@@ -588,10 +638,15 @@ UBlueprint *UReimportFbxSceneFactory::UpdateOriginalBluePrint(FString &BluePrint
 	//Find the BluePrint
 	FAssetRegistryModule& AssetRegistryModule = FModuleManager::LoadModuleChecked<FAssetRegistryModule>("AssetRegistry");
 	FAssetData BlueprintAssetData = AssetRegistryModule.Get().GetAssetByObjectPath(FName(*(BluePrintFullName)));
-	UPackage* PkgExist = LoadPackage(nullptr, *BlueprintAssetData.PackageName.ToString(), LOAD_Verify | LOAD_NoWarn);
+
+	UPackage* PkgExist = FindPackage(nullptr, *BlueprintAssetData.PackageName.ToString());
 	if (PkgExist == nullptr)
 	{
-		return nullptr;
+		PkgExist = LoadPackage(nullptr, *BlueprintAssetData.PackageName.ToString(), LOAD_Verify | LOAD_NoWarn);
+		if (PkgExist == nullptr)
+		{
+			return nullptr;
+		}
 	}
 	//Load the package before searching the asset
 	PkgExist->FullyLoad();
@@ -712,28 +767,15 @@ EReimportResult::Type UReimportFbxSceneFactory::ImportSkeletalMesh(void* VoidRoo
 	{
 		return EReimportResult::Failed;
 	}
-
-	//Copy default options to StaticMeshImportData
-	SFbxSceneOptionWindow::CopyFbxOptionsToSkeletalMeshOptions(GlobalImportSettingsReference, SceneImportOptionsSkeletalMesh);
-	SceneImportOptionsSkeletalMesh->FillSkeletalMeshInmportData(SkeletalMeshImportData, AnimSequenceImportData);
-
-	UnFbx::FBXImportOptions* OverrideImportSettings = GetOptionsFromName(MeshInfo->OptionName);
-	if (OverrideImportSettings != nullptr)
-	{
-		SFbxSceneOptionWindow::CopyFbxOptionsToFbxOptions(OverrideImportSettings, GlobalImportSettings);
-		SFbxSceneOptionWindow::CopyFbxOptionsToSkeletalMeshOptions(OverrideImportSettings, SceneImportOptionsSkeletalMesh);
-	}
-	else
-	{
-		SFbxSceneOptionWindow::CopyFbxOptionsToFbxOptions(GlobalImportSettingsReference, GlobalImportSettings);
-		SFbxSceneOptionWindow::CopyFbxOptionsToSkeletalMeshOptions(GlobalImportSettingsReference, SceneImportOptionsSkeletalMesh);
-	}
-	SceneImportOptionsSkeletalMesh->FillSkeletalMeshInmportData(SkeletalMeshImportData, AnimSequenceImportData);
-	//Override the pivot bake option
-	if (GlobalImportSettings->bBakePivotInVertex && MeshInfo->PivotNodeUid == INVALID_UNIQUE_ID)
-	{
-		GlobalImportSettings->bBakePivotInVertex = false;
-	}
+	ApplyMeshInfoFbxOptions(MeshInfo);
+	
+	//TODO support bBakePivotInVertex
+	bool Old_bBakePivotInVertex = GlobalImportSettings->bBakePivotInVertex;
+	GlobalImportSettings->bBakePivotInVertex = false;
+	//if (GlobalImportSettings->bBakePivotInVertex && MeshInfo->PivotNodeUid == INVALID_UNIQUE_ID)
+	//{
+		//GlobalImportSettings->bBakePivotInVertex = false;
+	//}
 
 	TArray< TArray<FbxNode*>* > SkelMeshArray;
 	FbxImporter->FillFbxSkelMeshArrayInScene(RootNodeToImport, SkelMeshArray, false);
@@ -768,6 +810,8 @@ EReimportResult::Type UReimportFbxSceneFactory::ImportSkeletalMesh(void* VoidRoo
 		NewObject = ImportOneSkeletalMesh(VoidRootNodeToImport, VoidFbxImporter, SceneInfoPtr, RF_Public | RF_Standalone, VoidNodeArray, TotalNumNodes);
 	}
 
+	GlobalImportSettings->bBakePivotInVertex = Old_bBakePivotInVertex;
+
 	for (int32 i = 0; i < SkelMeshArray.Num(); i++)
 	{
 		delete SkelMeshArray[i];
@@ -858,29 +902,16 @@ EReimportResult::Type UReimportFbxSceneFactory::ReimportSkeletalMesh(void* VoidF
 		return EReimportResult::Failed;
 	}
 
-	//Copy default options to SkeletalMeshImportData
-	SFbxSceneOptionWindow::CopyFbxOptionsToSkeletalMeshOptions(GlobalImportSettingsReference, SceneImportOptionsSkeletalMesh);
-	SceneImportOptionsSkeletalMesh->FillSkeletalMeshInmportData(SkeletalMeshImportData, AnimSequenceImportData);
-
-	UnFbx::FBXImportOptions* OverrideImportSettings = GetOptionsFromName(MeshInfo->OptionName);
-	if (OverrideImportSettings != nullptr)
-	{
-		SFbxSceneOptionWindow::CopyFbxOptionsToFbxOptions(OverrideImportSettings, GlobalImportSettings);
-		SFbxSceneOptionWindow::CopyFbxOptionsToSkeletalMeshOptions(OverrideImportSettings, SceneImportOptionsSkeletalMesh);
-	}
-	else
-	{
-		SFbxSceneOptionWindow::CopyFbxOptionsToFbxOptions(GlobalImportSettingsReference, GlobalImportSettings);
-		SFbxSceneOptionWindow::CopyFbxOptionsToSkeletalMeshOptions(GlobalImportSettingsReference, SceneImportOptionsSkeletalMesh);
-	}
-	SceneImportOptionsSkeletalMesh->FillSkeletalMeshInmportData(SkeletalMeshImportData, AnimSequenceImportData);
-
-	//Override the pivot bake option
-	if (GlobalImportSettings->bBakePivotInVertex && MeshInfo->PivotNodeUid == INVALID_UNIQUE_ID)
-	{
-		GlobalImportSettings->bBakePivotInVertex = false;
-	}
+	ApplyMeshInfoFbxOptions(MeshInfo);
+	//TODO support bBakePivotInVertex
+	bool Old_bBakePivotInVertex = GlobalImportSettings->bBakePivotInVertex;
+	GlobalImportSettings->bBakePivotInVertex = false;
+	//if (GlobalImportSettings->bBakePivotInVertex && MeshInfo->PivotNodeUid == INVALID_UNIQUE_ID)
+	//{
+		//GlobalImportSettings->bBakePivotInVertex = false;
+	//}
 	TArray<FbxNode*> OutSkeletalMeshArray;
+	EReimportResult::Type ReimportResult = EReimportResult::Succeeded;
 	if (FbxImporter->ReimportSkeletalMesh(Mesh, SkeletalMeshImportData, MeshInfo->UniqueId, &OutSkeletalMeshArray))
 	{
 		Mesh->AssetImportData->Update(FbxImportFileName);
@@ -903,37 +934,134 @@ EReimportResult::Type UReimportFbxSceneFactory::ReimportSkeletalMesh(void* VoidF
 		//3. Reimport matching animation
 		if (GlobalImportSettings->bImportAnimations)
 		{
-			/*
-			USkeleton* Skeleton = Mesh->Skeleton;//  AnimSequence->GetSkeleton();
-			if (Skeleton)
+			TArray<FbxNode*> FBXMeshNodeArray;
+			FbxNode* SkeletonRoot = FbxImporter->FindFBXMeshesByBone(Mesh->Skeleton->GetReferenceSkeleton().GetBoneName(0), true, FBXMeshNodeArray);
+
+			FString AnimName = FbxImporter->MakeNameForMesh(FBXMeshNodeArray[0]->GetName(), FBXMeshNodeArray[0]).ToString();
+			AnimName = (GlobalImportSettings->AnimationName != "") ? GlobalImportSettings->AnimationName : AnimName + TEXT("_Anim");
+
+			TArray<FbxNode*> SortedLinks;
+			FbxImporter->RecursiveBuildSkeleton(SkeletonRoot, SortedLinks);
+
+			if (SortedLinks.Num() != 0)
 			{
-				//For all anim sequence to reimport
+				//Find the number of take
+				int32 ResampleRate = DEFAULT_SAMPLERATE;
+				if (GlobalImportSettings->bResample)
 				{
-					if (UEditorEngine::ReimportFbxAnimation(Skeleton, AnimSequence, ImportData, *Filename))
+					int32 MaxStackResampleRate = FbxImporter->GetMaxSampleRate(SortedLinks, FBXMeshNodeArray);
+					if (MaxStackResampleRate != 0)
 					{
-						// update the data in case the file source has changed
-						ImportData->Update(UFactory::CurrentFilename);
+						ResampleRate = MaxStackResampleRate;
+					}
+				}
+				int32 ValidTakeCount = 0;
+				int32 AnimStackCount = FbxImporter->Scene->GetSrcObjectCount<FbxAnimStack>();
+				for (int32 AnimStackIndex = 0; AnimStackIndex < AnimStackCount; AnimStackIndex++)
+				{
+					FbxAnimStack* CurAnimStack = FbxImporter->Scene->GetSrcObject<FbxAnimStack>(AnimStackIndex);
+
+					FbxTimeSpan AnimTimeSpan = FbxImporter->GetAnimationTimeSpan(SortedLinks[0], CurAnimStack);
+					bool bValidAnimStack = FbxImporter->ValidateAnimStack(SortedLinks, FBXMeshNodeArray, CurAnimStack, ResampleRate, GlobalImportSettings->bImportMorph, AnimTimeSpan);
+					// no animation
+					if (!bValidAnimStack)
+					{
+						continue;
+					}
+					ValidTakeCount++;
+				}
+
+				if (ValidTakeCount > 0)
+				{
+					//Reimport all sequence (reimport existing and import new one)
+					AnimStackCount = FbxImporter->Scene->GetSrcObjectCount<FbxAnimStack>();
+					for (int32 AnimStackIndex = 0; AnimStackIndex < AnimStackCount; AnimStackIndex++)
+					{
+						FbxAnimStack* CurAnimStack = FbxImporter->Scene->GetSrcObject<FbxAnimStack>(AnimStackIndex);
+						FString SequenceName = AnimName;
+						if (ValidTakeCount > 1)
+						{
+							SequenceName += "_";
+							SequenceName += UTF8_TO_TCHAR(CurAnimStack->GetName());
+						}
+
+						// See if this sequence already exists.
+						SequenceName = ObjectTools::SanitizeObjectName(SequenceName);
+						FString 	ParentPath = FString::Printf(TEXT("%s/%s"), *FPackageName::GetLongPackagePath(*Mesh->GetOutermost()->GetName()), *SequenceName);
+						//See if the sequence exist
+						UAnimSequence *DestSeq = nullptr;
+						UPackage *ParentPackage = LoadPackage(nullptr, *ParentPath, LOAD_Verify | LOAD_NoWarn);
+						if (ParentPackage != nullptr)
+						{
+							ParentPackage->FullyLoad();
+						}
+						UObject* Object = FindObjectSafe<UObject>(ANY_PACKAGE, *SequenceName);
+						if (Object != nullptr)
+						{
+							if (ParentPackage == nullptr)
+							{
+								ParentPackage = Object->GetOutermost();
+								ParentPackage->FullyLoad();
+							}
+							//Cast into sequence
+							DestSeq = Cast<UAnimSequence>(Object);
+						}
+						
+						//Get the sequence timespan
+						ResampleRate = DEFAULT_SAMPLERATE;
+						if (FbxImporter->ImportOptions->bResample)
+						{
+							ResampleRate = FbxImporter->GetMaxSampleRate(SortedLinks, FBXMeshNodeArray);
+						}
+						FbxTimeSpan AnimTimeSpan = FbxImporter->GetAnimationTimeSpan(SortedLinks[0], CurAnimStack);
 
-						// Try to find the outer package so we can dirty it up
-						if (AnimSequence->GetOuter())
+						if (DestSeq == nullptr)
 						{
-							AnimSequence->GetOuter()->MarkPackageDirty();
+							//Import a new sequence
+							ParentPackage = CreatePackage(NULL, *ParentPath);
+							Object = LoadObject<UObject>(ParentPackage, *SequenceName, NULL, LOAD_None, NULL);
+							DestSeq = Cast<UAnimSequence>(Object);
+							if (Object && !DestSeq)
+							{
+								FbxImporter->AddTokenizedErrorMessage(FTokenizedMessage::Create(EMessageSeverity::Error, LOCTEXT("Error_AssetExist", "Asset with same name exists. Can't overwrite another asset")), FFbxErrors::Generic_SameNameAssetExists);
+								continue; // Move on to next sequence...
+							}
+							// If not, create new one now.
+							if (!DestSeq)
+							{
+								DestSeq = NewObject<UAnimSequence>(ParentPackage, *SequenceName, RF_Public | RF_Standalone);
+								// Notify the asset registry
+								FAssetRegistryModule::AssetCreated(DestSeq);
+							}
+							else
+							{
+								DestSeq->RecycleAnimSequence();
+							}
+							DestSeq->SetSkeleton(Mesh->Skeleton);
+							// since to know full path, reimport will need to do same
+							UFbxAnimSequenceImportData* ImportData = UFbxAnimSequenceImportData::GetImportDataForAnimSequence(DestSeq, AnimSequenceImportData);
+							ImportData->Update(UFactory::CurrentFilename);
+							FbxImporter->ImportAnimation(Mesh->Skeleton, DestSeq, CurrentFilename, SortedLinks, FBXMeshNodeArray, CurAnimStack, ResampleRate, AnimTimeSpan);
 						}
 						else
 						{
-							AnimSequence->MarkPackageDirty();
+							//Reimport in a existing sequence
+							if (FbxImporter->ValidateAnimStack(SortedLinks, FBXMeshNodeArray, CurAnimStack, ResampleRate, true, AnimTimeSpan))
+							{
+								FbxImporter->ImportAnimation(Mesh->Skeleton, DestSeq, CurrentFilename, SortedLinks, FBXMeshNodeArray, CurAnimStack, ResampleRate, AnimTimeSpan);
+							}
 						}
 					}
 				}
 			}
-			*/
 		}
 	}
 	else
 	{
-		return EReimportResult::Failed;
+		ReimportResult = EReimportResult::Failed;
 	}
-	return EReimportResult::Succeeded;
+	GlobalImportSettings->bBakePivotInVertex = Old_bBakePivotInVertex;
+	return ReimportResult;
 }
 
 EReimportResult::Type UReimportFbxSceneFactory::ReimportStaticMesh(void* VoidFbxImporter, TSharedPtr<FFbxMeshInfo> MeshInfo)
@@ -950,28 +1078,14 @@ EReimportResult::Type UReimportFbxSceneFactory::ReimportStaticMesh(void* VoidFbx
 		FbxImporter->AddTokenizedErrorMessage(FTokenizedMessage::Create(EMessageSeverity::Error, FText::Format(FText::FromString("Reimport Mesh {0} fail, the original staicmesh in the content browser cannot be load."), FText::FromString(MeshInfo->GetImportPath()))), FName(TEXT("Reimport Fbx Scene")));
 		return EReimportResult::Failed;
 	}
-	
-	//Copy default options to StaticMeshImportData
-	SFbxSceneOptionWindow::CopyFbxOptionsToStaticMeshOptions(GlobalImportSettingsReference, SceneImportOptionsStaticMesh);
-	SceneImportOptionsStaticMesh->FillStaticMeshInmportData(StaticMeshImportData, SceneImportOptions);
-
-	UnFbx::FBXImportOptions* OverrideImportSettings = GetOptionsFromName(MeshInfo->OptionName);
-	if (OverrideImportSettings != nullptr)
-	{
-		SFbxSceneOptionWindow::CopyFbxOptionsToFbxOptions(OverrideImportSettings, GlobalImportSettings);
-		SFbxSceneOptionWindow::CopyFbxOptionsToStaticMeshOptions(OverrideImportSettings, SceneImportOptionsStaticMesh);
-	}
-	else
-	{
-		SFbxSceneOptionWindow::CopyFbxOptionsToFbxOptions(GlobalImportSettingsReference, GlobalImportSettings);
-		SFbxSceneOptionWindow::CopyFbxOptionsToStaticMeshOptions(GlobalImportSettingsReference, SceneImportOptionsStaticMesh);
-	}
-	SceneImportOptionsStaticMesh->FillStaticMeshInmportData(StaticMeshImportData, SceneImportOptions);
+	ApplyMeshInfoFbxOptions(MeshInfo);
 	//Override the pivot bake option
+	bool Old_bBakePivotInVertex = GlobalImportSettings->bBakePivotInVertex;
 	if (GlobalImportSettings->bBakePivotInVertex && MeshInfo->PivotNodeUid == INVALID_UNIQUE_ID)
 	{
 		GlobalImportSettings->bBakePivotInVertex = false;
 	}
+	EReimportResult::Type ReimportResult = EReimportResult::Succeeded;
 
 	FbxImporter->ApplyTransformSettingsToFbxNode(FbxImporter->Scene->GetRootNode(), StaticMeshImportData);
 	const TArray<UAssetUserData*>* UserData = Mesh->GetAssetUserDataArray();
@@ -1038,9 +1152,10 @@ EReimportResult::Type UReimportFbxSceneFactory::ReimportStaticMesh(void* VoidFbx
 	}
 	else
 	{
-		return EReimportResult::Failed;
+		ReimportResult = EReimportResult::Failed;
 	}
-	return EReimportResult::Succeeded;
+	GlobalImportSettings->bBakePivotInVertex = Old_bBakePivotInVertex;
+	return ReimportResult;
 }
 
 int32 UReimportFbxSceneFactory::GetPriority() const
