VR Editor is now "experimental" (disabled by default)
- Turn it on using the new "Enable VR Editng" option under "VR" in the Experimental section of Editor Preferences
- Also, the VR Mode button on the tool bar will now be greyed out if no VR hardware is available

diff --git a/Engine/Config/BaseEditorPerProjectUserSettings.ini b/Engine/Config/BaseEditorPerProjectUserSettings.ini
--- a/Engine/Config/BaseEditorPerProjectUserSettings.ini
+++ b/Engine/Config/BaseEditorPerProjectUserSettings.ini
@@ -164,6 +164,8 @@ bPromptForCheckoutOnAssetModification=True
 bSCCAutoAddNewFiles=True
 ; Tool to use for diff'ing text
 TextDiffToolPath=(FilePath="p4merge.exe")
+; Enables a preview of the Unreal Editor in VR.  This adds a new tool bar button that allows you to toggle into "VR Mode" instantly.  This feature is still in development, but your feedback is appreciated!
+bEnableVREditing=false
 
 
 [/Script/UnrealEd.LevelEditorMiscSettings]
diff --git a/Engine/Source/Editor/LevelEditor/Private/LevelEditorActions.cpp b/Engine/Source/Editor/LevelEditor/Private/LevelEditorActions.cpp
--- a/Engine/Source/Editor/LevelEditor/Private/LevelEditorActions.cpp
+++ b/Engine/Source/Editor/LevelEditor/Private/LevelEditorActions.cpp
@@ -2078,6 +2078,12 @@ bool FLevelEditorActionCallbacks::ToggleVR_IsChecked()
 }
 
 
+bool FLevelEditorActionCallbacks::ToggleVR_IsVisible()
+{
+	return GetDefault<UEditorExperimentalSettings>()->bEnableVREditing;
+}
+
+
 bool FLevelEditorActionCallbacks::CanSelectGameModeBlueprint()
 {
 	bool bCheckOutNeeded = false;
diff --git a/Engine/Source/Editor/LevelEditor/Private/SLevelEditor.cpp b/Engine/Source/Editor/LevelEditor/Private/SLevelEditor.cpp
--- a/Engine/Source/Editor/LevelEditor/Private/SLevelEditor.cpp
+++ b/Engine/Source/Editor/LevelEditor/Private/SLevelEditor.cpp
@@ -119,7 +119,10 @@ void SLevelEditor::BindCommands()
 
 	LevelEditorCommands->MapAction(
 		Actions.ToggleVR,
-		FExecuteAction::CreateStatic( &FLevelEditorActionCallbacks::ToggleVR ) );
+		FExecuteAction::CreateStatic( &FLevelEditorActionCallbacks::ToggleVR ),
+		FCanExecuteAction::CreateStatic( &FLevelEditorActionCallbacks::ToggleVR_CanExecute ),
+		FIsActionChecked::CreateStatic( &FLevelEditorActionCallbacks::ToggleVR_IsChecked ),
+		FIsActionButtonVisible::CreateStatic( &FLevelEditorActionCallbacks::ToggleVR_IsVisible ) );
 
 	LevelEditorCommands->MapAction(
 		Actions.WorldProperties,
diff --git a/Engine/Source/Editor/LevelEditor/Public/LevelEditorActions.h b/Engine/Source/Editor/LevelEditor/Public/LevelEditorActions.h
--- a/Engine/Source/Editor/LevelEditor/Public/LevelEditorActions.h
+++ b/Engine/Source/Editor/LevelEditor/Public/LevelEditorActions.h
@@ -633,6 +633,7 @@ public:
 	static void ToggleVR();
 	static bool ToggleVR_CanExecute();
 	static bool ToggleVR_IsChecked();
+	static bool ToggleVR_IsVisible();
 
 	/** Opens delta transform */
 	static void DeltaTransform();
diff --git a/Engine/Source/Editor/UnrealEd/Classes/Settings/EditorExperimentalSettings.h b/Engine/Source/Editor/UnrealEd/Classes/Settings/EditorExperimentalSettings.h
--- a/Engine/Source/Editor/UnrealEd/Classes/Settings/EditorExperimentalSettings.h
+++ b/Engine/Source/Editor/UnrealEd/Classes/Settings/EditorExperimentalSettings.h
@@ -128,6 +128,10 @@ public:
 	UPROPERTY(EditAnywhere, config, Category = Tools, meta = (DisplayName = "Use OpenCL for convex hull decomposition"))
 	bool bUseOpenCLForConvexHullDecomp;
 
+	/** Enables a preview of the Unreal Editor in VR.  This adds a new tool bar button that allows you to toggle into "VR Mode" instantly.  This feature is still in development, but your feedback is appreciated! */
+	UPROPERTY(EditAnywhere, config, Category=VR, meta = (DisplayName="Enable VR Editing"))
+	bool bEnableVREditing;
+
 	/**
 	 * Returns an event delegate that is executed when a setting has changed.
 	 *
