All UbergraphsFrames in children are fully recreated while FKismetEditorUtilities::RecompileBlueprintBytecode

diff --git a/Engine/Source/Editor/UnrealEd/Private/Kismet2/KismetReinstanceUtilities.cpp b/Engine/Source/Editor/UnrealEd/Private/Kismet2/KismetReinstanceUtilities.cpp
--- a/Engine/Source/Editor/UnrealEd/Private/Kismet2/KismetReinstanceUtilities.cpp
+++ b/Engine/Source/Editor/UnrealEd/Private/Kismet2/KismetReinstanceUtilities.cpp
@@ -1733,11 +1733,11 @@ FRecreateUberGraphFrameScope::FRecreateUberGraphFrameScope(UClass* InClass, bool
 		BP_SCOPED_COMPILER_EVENT_STAT(EKismetCompilerStats_RecreateUberGraphPersistentFrame);
 
 		const bool bIncludeDerivedClasses = true;
-		GetObjectsOfClass(RecompiledClass, Objects, bIncludeDerivedClasses);
+		GetObjectsOfClass(RecompiledClass, Objects, bIncludeDerivedClasses, RF_NoFlags);
 
 		for (auto Obj : Objects)
 		{
-			RecompiledClass->DestroyPersistentUberGraphFrame(Obj, true);
+			RecompiledClass->DestroyPersistentUberGraphFrame(Obj);
 		}
 	}
 }
@@ -1749,7 +1749,7 @@ FRecreateUberGraphFrameScope::~FRecreateUberGraphFrameScope()
 	{
 		if (IsValid(Obj))
 		{
-			RecompiledClass->CreatePersistentUberGraphFrame(Obj, false, true);
+			RecompiledClass->CreatePersistentUberGraphFrame(Obj, false);
 		}
 	}
 }
\ No newline at end of file
