Added FixupNeedsLoadForEditorGame commandlet to fixup exports saved with NeedsLoadForClient, NeedsLoadForServer and NeedsLoadForEditorGame set to false and then had NeedsLoadForEditorGame set to true.

- Added -IgnoreChangelist option to ResavePackagesCommandlet
- Removed a few unused things from ResavePackagesCommandlet
- Fixed broken code that came with Dev-Platform integration

diff --git a/Engine/Source/Editor/UnrealEd/Classes/Commandlets/FixupNeedsLoadForEditorGameCommandlet.h b/Engine/Source/Editor/UnrealEd/Classes/Commandlets/FixupNeedsLoadForEditorGameCommandlet.h
new file mode 100644
--- /dev/null
+++ b/Engine/Source/Editor/UnrealEd/Classes/Commandlets/FixupNeedsLoadForEditorGameCommandlet.h
@@ -0,0 +1,21 @@
+// Copyright 1998-2016 Epic Games, Inc. All Rights Reserved.
+
+#pragma once
+#include "Commandlets/ResavePackagesCommandlet.h"
+#include "FixupNeedsLoadForEditorGameCommandlet.generated.h"
+
+UCLASS()
+class UFixupNeedsLoadForEditorGameCommandlet : public UResavePackagesCommandlet
+{
+	GENERATED_BODY()
+
+	/** Map of class names and their default NeedsLoadForEditorGame() return values. */
+	TMap<FName, bool> ResaveClassNeedsLoadForEditorGameValues;
+
+protected:
+
+	//~ Begin UResavePackagesCommandlet Interface
+	virtual int32 InitializeResaveParameters(const TArray<FString>& Tokens, TArray<FString>& MapPathNames) override;
+	virtual void PerformPreloadOperations(FLinkerLoad* PackageLinker, bool& bSavePackage) override;
+	//~ End UCommandlet Interface
+};
diff --git a/Engine/Source/Editor/UnrealEd/Classes/Commandlets/ResavePackagesCommandlet.h b/Engine/Source/Editor/UnrealEd/Classes/Commandlets/ResavePackagesCommandlet.h
--- a/Engine/Source/Editor/UnrealEd/Classes/Commandlets/ResavePackagesCommandlet.h
+++ b/Engine/Source/Editor/UnrealEd/Classes/Commandlets/ResavePackagesCommandlet.h
@@ -4,12 +4,17 @@
 #include "Commandlets/Commandlet.h"
 #include "ResavePackagesCommandlet.generated.h"
 
+// Log category should be accessible by derived classes
+UNREALED_API DECLARE_LOG_CATEGORY_EXTERN(LogContentCommandlet, Log, All);
+
 UCLASS()
 // Added UNREALED_API to expose this to the save packages test
 class UNREALED_API UResavePackagesCommandlet : public UCommandlet
 {
     GENERATED_UCLASS_BODY()
 
+protected:
+
 	enum EBrevity
 	{
 		VERY_VERBOSE,
@@ -67,6 +72,9 @@ class UNREALED_API UResavePackagesCommandlet : public UCommandlet
 	/** Should we build lighting for the packages we are saving? **/
 	bool bShouldBuildLighting;
 
+	/** Ignore package version changelist **/
+	bool bIgnoreChangelist;
+
 	// Running count of packages that got modified and will need to be resaved
 	int32 PackagesRequiringResave;
 
@@ -78,7 +86,7 @@ class UNREALED_API UResavePackagesCommandlet : public UCommandlet
 
 	// The list of switches that were passed on the commandline
 	TArray<FString> Switches;
-protected:
+
 	/**
 	 * Evaluates the command-line to determine which maps to check.  By default all maps are checked
 	 * Provides child classes with a chance to initialize any variables, parse the command line, etc.
@@ -103,7 +111,7 @@ protected:
 	 * @param	bSavePackage	[in]	indicates whether the package is currently going to be saved
 	 *							[out]	set to true to resave the package
 	 */
-	virtual bool PerformPreloadOperations( FLinkerLoad* PackageLinker, bool& bSavePackage );
+	virtual void PerformPreloadOperations( FLinkerLoad* PackageLinker, bool& bSavePackage );
 
 	/**
 	 * Allows the commandlet to perform any additional operations on the object before it is resaved.
@@ -148,6 +156,7 @@ protected:
 
 	// Print out a message only if running in very verbose mode
 	void VerboseMessage(const FString& Message);
+
 public:		
 	//~ Begin UCommandlet Interface
 	virtual int32 Main(const FString& Params) override;
diff --git a/Engine/Source/Editor/UnrealEd/Private/Commandlets/ContentCommandlets.cpp b/Engine/Source/Editor/UnrealEd/Private/Commandlets/ContentCommandlets.cpp
--- a/Engine/Source/Editor/UnrealEd/Private/Commandlets/ContentCommandlets.cpp
+++ b/Engine/Source/Editor/UnrealEd/Private/Commandlets/ContentCommandlets.cpp
@@ -10,7 +10,7 @@
 #include "PackageHelperFunctions.h"
 #include "PackageTools.h"
 
-DEFINE_LOG_CATEGORY_STATIC(LogContentCommandlet, Log, All);
+DEFINE_LOG_CATEGORY(LogContentCommandlet);
 
 #include "AssetRegistryModule.h"
 
@@ -236,37 +236,48 @@ int32 UResavePackagesCommandlet::InitializeResaveParameters( const TArray<FStrin
 	}
 
 	FString ClassList;
-	for ( int32 SwitchIdx = 0; SwitchIdx < Switches.Num(); SwitchIdx++ )
+	for (const FString& CurrentSwitch : Switches)
 	{
-		const FString& CurrentSwitch = Switches[SwitchIdx];
 		if ( FParse::Value(*CurrentSwitch, TEXT("RESAVECLASS="), ClassList, false) )
 		{
 			TArray<FString> ClassNames;
 			ClassList.ParseIntoArray(ClassNames, TEXT(","), true);
-			for ( int32 Idx = 0; Idx < ClassNames.Num(); Idx++ )
+			for (const FString& ClassName : ClassNames)
 			{
-				ResaveClasses.AddUnique(*ClassNames[Idx]);
+				ResaveClasses.AddUnique(*ClassName);
 			}
 
 			break;
 		}
 	}
 
+	/** determine if we should check subclasses of ResaveClasses */
+	bool bIncludeChildClasses = Switches.Contains(TEXT("IncludeChildClasses"));
+	if (bIncludeChildClasses && ResaveClasses.Num() == 0)
+	{
+		// Sanity check fail
+		UE_LOG(LogContentCommandlet, Error, TEXT("AllowSubclasses param requires ResaveClass param."));
+		return 1;
+	}
 
-	bool bIncludeChildClasses = Switches.Contains(TEXT("includeChildClasses"));
 	if (bIncludeChildClasses)
 	{
-		for (const auto& ResaveClassName : ResaveClasses)
-		{
-			const UClass* ResaveClass = FindObject<UClass>(nullptr, *ResaveClassName.ToString(), true);
+		// Can't use ranged for here because the array grows inside of this loop.
+		// Also, no need to iterate over the newly added objects as we know
+		// we have found all of their subclasses too (IsChildOf guarantees that).
+		const int32 NumResaveClasses = ResaveClasses.Num();
+		for (int32 ClassIndex = 0; ClassIndex < NumResaveClasses; ++ClassIndex)
+		{
+			// Find the class object and then all derived classes
+			UClass* ResaveClass = FindObject<UClass>(ANY_PACKAGE, *ResaveClasses[ClassIndex].ToString());
 			if (ResaveClass)
 			{
 				for (TObjectIterator<UClass> It; It; ++It)
 				{
-
-					if (It->IsChildOf(ResaveClass))
+					UClass* MaybeChildClass = *It;
+					if (MaybeChildClass->IsChildOf(ResaveClass))
 					{
-						ResaveClasses.AddUnique(*It->GetFName().ToString());
+						ResaveClasses.AddUnique(MaybeChildClass->GetFName());
 					}
 				}
 			}
@@ -452,42 +463,42 @@ void UResavePackagesCommandlet::LoadAndSaveOnePackage(const FString& Filename)
 
 					if (bAutoCheckOut)
 					{
-						if (SourceControlState.IsValid() && (SourceControlState->IsCheckedOut() || SourceControlState->IsAdded()))
-						{
-							UE_LOG(LogContentCommandlet, Display, TEXT("Revert '%s' from source control..."), *Filename);
-							SourceControlProvider.Execute(ISourceControlOperation::Create<FRevert>(), PackageFilename);
+					if (SourceControlState.IsValid() && (SourceControlState->IsCheckedOut() || SourceControlState->IsAdded()))
+					{
+						UE_LOG(LogContentCommandlet, Display, TEXT("Revert '%s' from source control..."), *Filename);
+						SourceControlProvider.Execute(ISourceControlOperation::Create<FRevert>(), PackageFilename);
 
-							UE_LOG(LogContentCommandlet, Display, TEXT("Deleting '%s' from source control..."), *Filename);
-							SourceControlProvider.Execute(ISourceControlOperation::Create<FDelete>(), PackageFilename);
-						}
-						else if (SourceControlState.IsValid() && SourceControlState->CanCheckout())
-						{
-							UE_LOG(LogContentCommandlet, Display, TEXT("Deleting '%s' from source control..."), *Filename);
-							SourceControlProvider.Execute(ISourceControlOperation::Create<FDelete>(), PackageFilename);
-						}
-						else if (SourceControlState.IsValid() && SourceControlState->IsCheckedOutOther())
-						{
-							UE_LOG(LogContentCommandlet, Warning, TEXT("Couldn't delete '%s' from source control, someone has it checked out, skipping..."), *Filename);
-						}
-						else if (SourceControlState.IsValid() && !SourceControlState->IsSourceControlled())
+						UE_LOG(LogContentCommandlet, Display, TEXT("Deleting '%s' from source control..."), *Filename);
+						SourceControlProvider.Execute(ISourceControlOperation::Create<FDelete>(), PackageFilename);
+					}
+					else if (SourceControlState.IsValid() && SourceControlState->CanCheckout())
+					{
+						UE_LOG(LogContentCommandlet, Display, TEXT("Deleting '%s' from source control..."), *Filename);
+						SourceControlProvider.Execute(ISourceControlOperation::Create<FDelete>(), PackageFilename);
+					}
+					else if (SourceControlState.IsValid() && SourceControlState->IsCheckedOutOther())
+					{
+						UE_LOG(LogContentCommandlet, Warning, TEXT("Couldn't delete '%s' from source control, someone has it checked out, skipping..."), *Filename);
+					}
+					else if (SourceControlState.IsValid() && !SourceControlState->IsSourceControlled())
+					{
+						UE_LOG(LogContentCommandlet, Warning, TEXT("'%s' is not in source control, attempting to delete from disk..."), *Filename);
+						if (!IFileManager::Get().Delete(*Filename, false, true))
 						{
-							UE_LOG(LogContentCommandlet, Warning, TEXT("'%s' is not in source control, attempting to delete from disk..."), *Filename);
-							if (!IFileManager::Get().Delete(*Filename, false, true))
-							{
-								UE_LOG(LogContentCommandlet, Warning, TEXT("  ... failed to delete from disk."), *Filename);
-							}
+							UE_LOG(LogContentCommandlet, Warning, TEXT("  ... failed to delete from disk."), *Filename);
 						}
-						else
+					}
+					else
+					{
+						UE_LOG(LogContentCommandlet, Warning, TEXT("'%s' is in an unknown source control state, attempting to delete from disk..."), *Filename);
+						if (!IFileManager::Get().Delete(*Filename, false, true))
 						{
-							UE_LOG(LogContentCommandlet, Warning, TEXT("'%s' is in an unknown source control state, attempting to delete from disk..."), *Filename);
-							if (!IFileManager::Get().Delete(*Filename, false, true))
-							{
-								UE_LOG(LogContentCommandlet, Warning, TEXT("  ... failed to delete from disk."), *Filename);
-							}
+							UE_LOG(LogContentCommandlet, Warning, TEXT("  ... failed to delete from disk."), *Filename);
 						}
 					}
 				}
 			}
+			}
 
 			// Now based on the computation above we will see if we should actually attempt
 			// to save this package
@@ -603,6 +614,8 @@ int32 UResavePackagesCommandlet::Main( const FString& Params )
 	bAutoCheckIn = bAutoCheckOut && Switches.Contains(TEXT("AutoCheckIn"));
 	/** determine if we are building lighting for the map packages on the pass. **/
 	bShouldBuildLighting = Switches.Contains(TEXT("buildlighting"));
+	/** determine if we can skip the version changelist check */
+	bIgnoreChangelist = Switches.Contains(TEXT("IgnoreChangelist"));
 
 	TArray<FString> PackageNames;
 	int32 ResultCode = InitializeResaveParameters(Tokens, PackageNames);
@@ -711,10 +724,8 @@ FText UResavePackagesCommandlet::GetChangelistDescription() const
 }
 
 
-bool UResavePackagesCommandlet::PerformPreloadOperations( FLinkerLoad* PackageLinker, bool& bSavePackage )
+void UResavePackagesCommandlet::PerformPreloadOperations( FLinkerLoad* PackageLinker, bool& bSavePackage )
 {
-	bool bResult = false;
-
 	const int32 UE4PackageVersion = PackageLinker->Summary.GetFileVersionUE4();
 	const int32 LicenseeUE4PackageVersion = PackageLinker->Summary.GetFileVersionLicenseeUE4();
 
@@ -723,7 +734,6 @@ bool UResavePackagesCommandlet::PerformPreloadOperations( FLinkerLoad* PackageLi
 	if ( MinResaveUE4Version != IGNORE_PACKAGE_VERSION && UE4PackageVersion < MinResaveUE4Version )
 	{
 		bSavePackage = false;
-		bResult = true;
 	}
 
 	// Check if this package meets the maximum requirements.
@@ -734,43 +744,37 @@ bool UResavePackagesCommandlet::PerformPreloadOperations( FLinkerLoad* PackageLi
 
 	// If the package was saved with a higher engine version do not try to resave it. This also addresses problem with people 
 	// building editor locally and resaving content with a 0 CL version (e.g. BUILD_FROM_CL == 0)
-	if (PackageLinker->Summary.SavedByEngineVersion.GetChangelist() > FEngineVersion::Current().GetChangelist())
+	if (!bIgnoreChangelist && PackageLinker->Summary.SavedByEngineVersion.GetChangelist() > FEngineVersion::Current().GetChangelist())
 	{
 		UE_LOG(LogContentCommandlet, Warning, TEXT("Skipping resave of %s due to engine version mismatch (Package:%d, Editor:%d "), 
 			*PackageLinker->GetArchiveName(),
 			PackageLinker->Summary.SavedByEngineVersion.GetChangelist(), 
 			FEngineVersion::Current().GetChangelist());
 		bSavePackage = false;
-		bResult = true;	
 	}
 
 	// If not, don't resave it.
 	if ( !bAllowResave )
 	{
 		bSavePackage = false;
-		bResult = true;
 	}
 
 	// Check if the package contains any instances of the class that needs to be resaved.
 	if ( bSavePackage && ResaveClasses.Num() > 0 )
 	{
 		bSavePackage = false;
-		for (int32 ExportIndex = 0; ExportIndex < PackageLinker->ExportMap.Num(); ExportIndex++)
+		for (int32 ExportIndex = 0; !bSavePackage && ExportIndex < PackageLinker->ExportMap.Num(); ExportIndex++)
 		{
-			if ( ResaveClasses.Contains(PackageLinker->GetExportClassName(ExportIndex)) )
+			FName ExportClassName = PackageLinker->GetExportClassName(ExportIndex);
+			if (ResaveClasses.Contains(ExportClassName))
 			{
 				bSavePackage = true;
 				break;
 			}
 		}
-
-		bResult = true;
 	}
-
-	return bResult;
 }
 
-
 bool UResavePackagesCommandlet::CheckoutFile(const FString& Filename)
 {
 	if (!bAutoCheckOut)
diff --git a/Engine/Source/Editor/UnrealEd/Private/Commandlets/FixupNeedsLoadForEditorGameCommandlet.cpp b/Engine/Source/Editor/UnrealEd/Private/Commandlets/FixupNeedsLoadForEditorGameCommandlet.cpp
new file mode 100644
--- /dev/null
+++ b/Engine/Source/Editor/UnrealEd/Private/Commandlets/FixupNeedsLoadForEditorGameCommandlet.cpp
@@ -0,0 +1,74 @@
+// Copyright 1998-2016 Epic Games, Inc. All Rights Reserved.
+
+/*==================================================================================================
+	FixupNeedsLoadForEditorGameCommandlet.cpp: Fixes outdated NeedsLoadForEditorGame flags on exports
+====================================================================================================*/
+
+#include "UnrealEd.h"
+#include "Commandlets/FixupNeedsLoadForEditorGameCommandlet.h"
+
+int32 UFixupNeedsLoadForEditorGameCommandlet::InitializeResaveParameters(const TArray<FString>& Tokens, TArray<FString>& MapPathNames)
+{
+	int32 Result = Super::InitializeResaveParameters(Tokens, MapPathNames);
+	// We need ResaveClasses to be specified, otherwise we won't know what to update
+	if (Result == 0 && !ResaveClasses.Num())
+	{
+		UE_LOG(LogContentCommandlet, Error, TEXT("FixupNeedsLoadForEditorGame commandlet requires at least one resave class name. Use -RESAVECLASS=ClassA,ClassB,ClassC to specify resave classes."));
+		Result = 1;
+	}
+	else
+	{
+		for (FName& ClassName : ResaveClasses)
+		{			
+			if (!ResaveClassNeedsLoadForEditorGameValues.Contains(ClassName))
+			{
+				UClass* ResaveClass = FindObject<UClass>(ANY_PACKAGE, *ClassName.ToString());
+				if (ResaveClass)
+				{
+					UObject* DefaultObject = ResaveClass->GetDefaultObject();
+					ResaveClassNeedsLoadForEditorGameValues.Add(ClassName, DefaultObject->NeedsLoadForEditorGame());
+				}
+			}
+			else if (Verbosity != UResavePackagesCommandlet::ONLY_ERRORS)
+			{
+				UE_LOG(LogContentCommandlet, Warning, TEXT("Resave Class \"%s\" could not be found. Make sure the class name is valid and that it's a native class."), *ClassName.ToString());
+			}
+		}
+		if (ResaveClassNeedsLoadForEditorGameValues.Num() == 0)
+		{
+			UE_LOG(LogContentCommandlet, Error, TEXT("Got %d classes to resave but none of the exist."), ResaveClasses.Num());
+			Result = 1;
+		}
+	}
+	return Result;
+}
+
+void UFixupNeedsLoadForEditorGameCommandlet::PerformPreloadOperations(FLinkerLoad* PackageLinker, bool& bSavePackage)
+{
+	Super::PerformPreloadOperations(PackageLinker, bSavePackage);
+	if (bSavePackage)
+	{
+		// The package contains a class we want to check but we don't want to save it unless
+		// any of the exports has an outdated (compared to a CDO) NeedsLoadForEditorGame flag.
+		bSavePackage = false;
+		for (int32 ExportIndex = 0; ExportIndex < PackageLinker->ExportMap.Num(); ExportIndex++)
+		{
+			FName ExportClassName = PackageLinker->GetExportClassName(ExportIndex);
+			bool* bNeedsLoadForEditorGameValuePtr = ResaveClassNeedsLoadForEditorGameValues.Find(ExportClassName);
+			if (bNeedsLoadForEditorGameValuePtr)
+			{
+				FObjectExport& Export = PackageLinker->ExportMap[ExportIndex];
+				// The condition below may seem confusing because bNotForEditorGame = !bNeedsLoadForEditorGame
+				// At the moment we only update assets that had bNotForEditorGame set to true but NeedsLoadForEditorGame returns true 
+				// and either bNotForClient or bNotForServer is true. 
+				if (Export.bNotForEditorGame && 
+					  Export.bNotForEditorGame == *bNeedsLoadForEditorGameValuePtr &&
+						(Export.bNotForClient || Export.bNotForServer))
+				{
+					bSavePackage = true;
+					break;
+				}				
+			}
+		}
+	}	
+}
