SuperSearch - Moving the SuperSearchSettings into UnrealEd, so that non-edior modules can use the SuperSearch module.  Fixes build errors.

diff --git a/Engine/Source/Developer/SuperSearch/Private/SuperSearchSettings.cpp b/Engine/Source/Developer/SuperSearch/Private/SuperSearchSettings.cpp
deleted file mode 100644
diff --git a/Engine/Source/Developer/SuperSearch/Private/SuperSearchSettings.h b/Engine/Source/Developer/SuperSearch/Private/SuperSearchSettings.h
deleted file mode 100644
diff --git a/Engine/Source/Editor/UnrealEd/Private/Settings/SuperSearchSettings.cpp b/Engine/Source/Editor/UnrealEd/Private/Settings/SuperSearchSettings.cpp
new file mode 100644
--- /dev/null
+++ b/Engine/Source/Editor/UnrealEd/Private/Settings/SuperSearchSettings.cpp
@@ -0,0 +1,24 @@
+// Copyright 1998-2016 Epic Games, Inc. All Rights Reserved.
+
+#include "UnrealEd.h"
+
+#include "SuperSearchSettings.h"
+
+void USuperSearchSettings::PostInitProperties()
+{
+	Super::PostInitProperties();
+
+	FSuperSearchModule::Get().SetSearchEngine(SearchEngine);
+}
+
+void USuperSearchSettings::PostEditChangeProperty(struct FPropertyChangedEvent& PropertyChangedEvent)
+{
+	Super::PostEditChangeProperty(PropertyChangedEvent);
+
+	const FName Name = PropertyChangedEvent.Property ? PropertyChangedEvent.Property->GetFName() : NAME_None;
+
+	if ( Name == GET_MEMBER_NAME_CHECKED(USuperSearchSettings, SearchEngine) )
+	{
+		FSuperSearchModule::Get().SetSearchEngine(SearchEngine);
+	}
+}
\ No newline at end of file
diff --git a/Engine/Source/Editor/UnrealEd/Private/Settings/SuperSearchSettings.h b/Engine/Source/Editor/UnrealEd/Private/Settings/SuperSearchSettings.h
new file mode 100644
--- /dev/null
+++ b/Engine/Source/Editor/UnrealEd/Private/Settings/SuperSearchSettings.h
@@ -0,0 +1,28 @@
+// Copyright 1998-2016 Epic Games, Inc. All Rights Reserved.
+
+#pragma once
+
+#include "Engine/DeveloperSettings.h"
+#include "SuperSearchModule.h"
+#include "SuperSearchSettings.generated.h"
+
+/**
+ * 
+ */
+UCLASS(config=EditorSettings, meta=( DisplayName="Super Search" ))
+class USuperSearchSettings : public UDeveloperSettings
+{
+	GENERATED_BODY()
+
+public:
+
+	/**
+	 * Set the search engine to use in the editor for looking up questions in the community.
+	 */
+	UPROPERTY(EditAnywhere, config, Category=General)
+	ESearchEngine SearchEngine;
+
+public:
+	virtual void PostInitProperties() override;
+	virtual void PostEditChangeProperty(struct FPropertyChangedEvent& PropertyChangedEvent) override;
+};
\ No newline at end of file
diff --git a/Engine/Source/Editor/UnrealEd/UnrealEd.Build.cs b/Engine/Source/Editor/UnrealEd/UnrealEd.Build.cs
--- a/Engine/Source/Editor/UnrealEd/UnrealEd.Build.cs
+++ b/Engine/Source/Editor/UnrealEd/UnrealEd.Build.cs
@@ -50,7 +50,6 @@ public class UnrealEd : ModuleRules
 				"SourceCodeAccess",
 				"ReferenceViewer",
                 "IntroTutorials",
-                "SuperSearch",
 				"OutputLog",
 				"Landscape",
                 "Niagara",
@@ -122,7 +121,8 @@ public class UnrealEd : ModuleRules
 				"Sockets",
 				"SoundClassEditor",
 				"SoundCueEditor",
-				"SourceControlWindows", 
+				"SourceControlWindows",
+                "SuperSearch",
 				"StatsViewer",
 				"SwarmInterface",
 				"TargetPlatform",
@@ -189,7 +189,6 @@ public class UnrealEd : ModuleRules
 				"ImageWrapper",
 				"Blutility",
 				"IntroTutorials",
-                "SuperSearch",
 				"WorkspaceMenuStructure",
 				"PlacementMode",
 				"NiagaraEditor",
@@ -236,6 +235,7 @@ public class UnrealEd : ModuleRules
 			{
                 "GraphEditor",
 				"Kismet",
+                "SuperSearch",
             }
 		); 
 
