Fixed DDC not generating for VREditor engine content because module wasn't loaded
- "VREditor" module is loaded early at startup when generating DDC in an editor build
- Also, we now ignore VREditor content when cooking with "-SkipEditorContent"
- This addresses a slew of new DDC warnings about Blueprints that couldn't be loaded

diff --git a/Engine/Source/Editor/UnrealEd/Private/Commandlets/CookCommandlet.cpp b/Engine/Source/Editor/UnrealEd/Private/Commandlets/CookCommandlet.cpp
--- a/Engine/Source/Editor/UnrealEd/Private/Commandlets/CookCommandlet.cpp
+++ b/Engine/Source/Editor/UnrealEd/Private/Commandlets/CookCommandlet.cpp
@@ -518,7 +518,9 @@ bool UCookCommandlet::SaveCookedPackage( UPackage* Package, uint32 SaveFlags, bo
 			}
 
 			// don't save Editor resources from the Engine if the target doesn't have editoronly data
-			if (bSkipEditorContent && Name.StartsWith(TEXT("/Engine/Editor")) && !Target->HasEditorOnlyData())
+			if (bSkipEditorContent && 
+				( Name.StartsWith(TEXT("/Engine/Editor")) || Name.StartsWith(TEXT("/Engine/VREditor"))) && 
+				!Target->HasEditorOnlyData())
 			{
 				bCookPackage = false;
 			}
diff --git a/Engine/Source/Editor/UnrealEd/Private/CookOnTheFlyServer.cpp b/Engine/Source/Editor/UnrealEd/Private/CookOnTheFlyServer.cpp
--- a/Engine/Source/Editor/UnrealEd/Private/CookOnTheFlyServer.cpp
+++ b/Engine/Source/Editor/UnrealEd/Private/CookOnTheFlyServer.cpp
@@ -2469,7 +2469,9 @@ ESavePackageResult UCookOnTheFlyServer::SaveCookedPackage(UPackage* Package, uin
 			bool bCookPackage = true;
 
 			// don't save Editor resources from the Engine if the target doesn't have editoronly data
-			if (IsCookFlagSet(ECookInitializationFlags::SkipEditorContent) && Name.StartsWith(TEXT("/Engine/Editor")) && !Target->HasEditorOnlyData())
+			if (IsCookFlagSet(ECookInitializationFlags::SkipEditorContent) && 
+				( Name.StartsWith(TEXT("/Engine/Editor")) || Name.StartsWith(TEXT("/Engine/VREditor")) ) && 
+				!Target->HasEditorOnlyData())
 			{
 				bCookPackage = false;
 			}
diff --git a/Engine/Source/Runtime/Launch/Private/LaunchEngineLoop.cpp b/Engine/Source/Runtime/Launch/Private/LaunchEngineLoop.cpp
--- a/Engine/Source/Runtime/Launch/Private/LaunchEngineLoop.cpp
+++ b/Engine/Source/Runtime/Launch/Private/LaunchEngineLoop.cpp
@@ -2066,7 +2066,9 @@ bool FEngineLoop::LoadStartupCoreModules()
 
 	// Ability tasks are based on GameplayTasks, so we need to make sure that module is loaded as well
 	FModuleManager::Get().LoadModule(TEXT("GameplayTasksEditor"));
-
+	
+	// VREditor needs to be loaded in editor builds early, so engine content Blueprints can be loaded during DDC generation
+	FModuleManager::Get().LoadModule(TEXT("VREditor"));
 	// -----------------------------------------------------
 
 	// HACK: load AbilitySystem editor as early as possible for statically initialized assets (non cooked BT assets needs it)
