Copying //UE4/Dev-Sequencer to //UE4/Main (Source: //UE4/Dev-Sequencer @ 2945541)

diff --git a/Engine/Source/Editor/Sequencer/Public/SequencerCommands.h b/Engine/Source/Editor/Sequencer/Public/SequencerCommands.h
--- a/Engine/Source/Editor/Sequencer/Public/SequencerCommands.h
+++ b/Engine/Source/Editor/Sequencer/Public/SequencerCommands.h
@@ -75,11 +75,17 @@ public:
 	/** Expand/collapse nodes and descendants */
 	TSharedPtr< FUICommandInfo > ToggleExpandCollapseNodesAndDescendants;
 
-	/** Sets the upper bound of the in/out section */
-	TSharedPtr< FUICommandInfo > SetInOutEnd;
+	/** Sets the upper bound of the selection range */
+	TSharedPtr< FUICommandInfo > SetSelectionRangeEnd;
 
-	/** Sets the lower bound of the in/out section */
-	TSharedPtr< FUICommandInfo > SetInOutStart;
+	/** Sets the lower bound of the selection range */
+	TSharedPtr< FUICommandInfo > SetSelectionRangeStart;
+
+	/** Select all keys that fall into the selection range*/
+	TSharedPtr< FUICommandInfo > SelectKeysInSelectionRange;
+
+	/** Delete the time in the selection range */
+	TSharedPtr< FUICommandInfo > DeleteSelectionRange;
 
 	/** Sets a key at the current time for the selected actor */
 	TSharedPtr< FUICommandInfo > SetKey;
@@ -129,9 +135,6 @@ public:
 	/** Turns the range slider on and off. */
 	TSharedPtr< FUICommandInfo > ToggleShowRangeSlider;
 
-	/** Toggles whether to lock the in/out to the start/end range on scrolling. */
-	TSharedPtr< FUICommandInfo > ToggleLockInOutToStartEndRange;
-
 	/** Turns snapping on and off. */
 	TSharedPtr< FUICommandInfo > ToggleIsSnapEnabled;
 
@@ -192,6 +195,9 @@ public:
 	/** Open a panel that enables exporting the sequence to a movie */
 	TSharedPtr< FUICommandInfo > RenderMovie;
 
+	/** Create camera and set it as the current camera cut */
+	TSharedPtr< FUICommandInfo > CreateCamera;
+
 	/** Paste from the sequencer clipboard history */
 	TSharedPtr< FUICommandInfo > PasteFromHistory;
 
diff --git a/Engine/Source/Editor/Sequencer/Sequencer.Build.cs b/Engine/Source/Editor/Sequencer/Sequencer.Build.cs
--- a/Engine/Source/Editor/Sequencer/Sequencer.Build.cs
+++ b/Engine/Source/Editor/Sequencer/Sequencer.Build.cs
@@ -16,6 +16,7 @@ public class Sequencer : ModuleRules
 		PrivateDependencyModuleNames.AddRange(
 			new string[] {
 				"AppFramework", 
+                "CinematicCamera",
 				"Core", 
 				"CoreUObject", 
                 "InputCore",
@@ -52,7 +53,8 @@ public class Sequencer : ModuleRules
 				"LevelEditor",
 				"SceneOutliner",
 				"WorkspaceMenuStructure",
-				"SequenceRecorder"
+				"SequenceRecorder",
+				"SequenceRecorderSections",
 			}
 		);
 
diff --git a/Engine/Source/Editor/SequencerWidgets/Private/STimeRange.cpp b/Engine/Source/Editor/SequencerWidgets/Private/STimeRange.cpp
--- a/Engine/Source/Editor/SequencerWidgets/Private/STimeRange.cpp
+++ b/Engine/Source/Editor/SequencerWidgets/Private/STimeRange.cpp
@@ -102,7 +102,7 @@ void STimeRange::Construct( const STimeRange::FArguments& InArgs, TSharedRef<ITi
 		[
 			SNew(SBox)
 			.Visibility(InArgs._ShowWorkingRange ? EVisibility::Visible : EVisibility::Collapsed)
-			.WidthOverride(36)
+			.WidthOverride(48)
 			.HAlign(HAlign_Center)
 			[
 				WorkingRangeStart
@@ -115,7 +115,7 @@ void STimeRange::Construct( const STimeRange::FArguments& InArgs, TSharedRef<ITi
 		.Padding(2.f)
 		[
 			SNew(SBox)
-			.WidthOverride(36)
+			.WidthOverride(48)
 			.HAlign(HAlign_Center)
 			.Visibility(InArgs._ShowPlaybackRange ? EVisibility::Visible : EVisibility::Collapsed)
 			[
@@ -136,7 +136,7 @@ void STimeRange::Construct( const STimeRange::FArguments& InArgs, TSharedRef<ITi
 		[
 			SNew(SBox)
 			.Visibility(InArgs._ShowViewRange ? EVisibility::Visible : EVisibility::Collapsed)
-			.WidthOverride(36)
+			.WidthOverride(48)
 			.HAlign(HAlign_Center)
 			[
 				ViewRangeStart
@@ -144,11 +144,12 @@ void STimeRange::Construct( const STimeRange::FArguments& InArgs, TSharedRef<ITi
 		]
 
 		+SHorizontalBox::Slot()
-
-		.Padding(2.0f, 4.0f)
-		[
-			InArgs._CenterContent.Widget
-		]
+			.FillWidth(1.0f)
+			.Padding(2.0f, 4.0f)
+			.VAlign(VAlign_Center)
+			[
+				InArgs._CenterContent.Widget
+			]
 		
 		+SHorizontalBox::Slot()
 		.VAlign(VAlign_Center)
@@ -157,7 +158,7 @@ void STimeRange::Construct( const STimeRange::FArguments& InArgs, TSharedRef<ITi
 		[
 			SNew(SBox)
 			.Visibility(InArgs._ShowViewRange ? EVisibility::Visible : EVisibility::Collapsed)
-			.WidthOverride(36)
+			.WidthOverride(48)
 			.HAlign(HAlign_Center)
 			[
 				ViewRangeEnd
@@ -170,7 +171,7 @@ void STimeRange::Construct( const STimeRange::FArguments& InArgs, TSharedRef<ITi
 		.Padding(2.f)
 		[
 			SNew(SBox)
-			.WidthOverride(36)
+			.WidthOverride(48)
 			.HAlign(HAlign_Center)
 			.Visibility(InArgs._ShowPlaybackRange ? EVisibility::Visible : EVisibility::Collapsed)
 			[
@@ -190,7 +191,7 @@ void STimeRange::Construct( const STimeRange::FArguments& InArgs, TSharedRef<ITi
 		.Padding(2.f)
 		[
 			SNew(SBox)
-			.WidthOverride(36)
+			.WidthOverride(48)
 			.HAlign(HAlign_Center)
 			.Visibility(InArgs._ShowWorkingRange ? EVisibility::Visible : EVisibility::Collapsed)
 			[
diff --git a/Engine/Source/Editor/SequencerWidgets/Private/STimeRangeSlider.cpp b/Engine/Source/Editor/SequencerWidgets/Private/STimeRangeSlider.cpp
--- a/Engine/Source/Editor/SequencerWidgets/Private/STimeRangeSlider.cpp
+++ b/Engine/Source/Editor/SequencerWidgets/Private/STimeRangeSlider.cpp
@@ -60,6 +60,13 @@ void STimeRangeSlider::ComputeHandleOffsets(float& LeftHandleOffset, float& Hand
 	RightHandleOffset = HandleOffset + (OutTime - InTime) * UnitsToPixel;
 }
 
+
+FVector2D STimeRangeSlider::ComputeDesiredSize(float) const
+{
+	return FVector2D(4.0f * TimeRangeSliderConstants::HandleSize, TimeRangeSliderConstants::HandleSize);
+}
+
+
 int32 STimeRangeSlider::OnPaint( const FPaintArgs& Args, const FGeometry& AllottedGeometry, const FSlateRect& MyClippingRect, FSlateWindowElementList& OutDrawElements, int32 LayerId, const FWidgetStyle& InWidgetStyle, bool bParentEnabled ) const
 {
 	const int32 BackgroundLayer = LayerId+1;
diff --git a/Engine/Source/Editor/SequencerWidgets/Private/STimeRangeSlider.h b/Engine/Source/Editor/SequencerWidgets/Private/STimeRangeSlider.h
--- a/Engine/Source/Editor/SequencerWidgets/Private/STimeRangeSlider.h
+++ b/Engine/Source/Editor/SequencerWidgets/Private/STimeRangeSlider.h
@@ -14,6 +14,7 @@ public:
 	void Construct( const FArguments& InArgs, TSharedRef<ITimeSliderController> InTimeSliderController );
 
 	// SWidget interface
+	virtual FVector2D ComputeDesiredSize(float) const override;
 	virtual int32 OnPaint( const FPaintArgs& Args, const FGeometry& AllottedGeometry, const FSlateRect& MyClippingRect, FSlateWindowElementList& OutDrawElements, int32 LayerId, const FWidgetStyle& InWidgetStyle, bool bParentEnabled ) const override;
 	virtual FReply OnMouseButtonDown( const FGeometry& MyGeometry, const FPointerEvent& MouseEvent ) override;
 	virtual FReply OnMouseButtonUp( const FGeometry& MyGeometry, const FPointerEvent& MouseEvent ) override;
diff --git a/Engine/Source/Editor/SequencerWidgets/Public/ITimeSlider.h b/Engine/Source/Editor/SequencerWidgets/Public/ITimeSlider.h
--- a/Engine/Source/Editor/SequencerWidgets/Public/ITimeSlider.h
+++ b/Engine/Source/Editor/SequencerWidgets/Public/ITimeSlider.h
@@ -93,18 +93,6 @@ struct FTimeSliderArgs
 	/** Called when the clamp range changes */
 	FOnRangeChanged OnClampRangeChanged;
 
-	/** Attribute defining the playback range for this controller */
-	TAttribute<TRange<float>> InOutRange;
-
-	/** Delegate that is called when the in/out range wants to change */
-	FOnRangeChanged OnInOutRangeChanged;
-
-	/** Called right before the in/out range starts to be dragged */
-	FSimpleDelegate OnBeginInOutRangeDrag;
-
-	/** Called right after the in/out range has finished being dragged */
-	FSimpleDelegate OnEndInOutRangeDrag;
-
 	/** Attribute defining the playback range for this controller */
 	TAttribute<TRange<float>> PlaybackRange;
 
@@ -112,11 +100,22 @@ struct FTimeSliderArgs
 	FOnRangeChanged OnPlaybackRangeChanged;
 
 	/** Called right before the playback range starts to be dragged */
-	FSimpleDelegate OnBeginPlaybackRangeDrag;
+	FSimpleDelegate OnPlaybackRangeBeginDrag;
 
 	/** Called right after the playback range has finished being dragged */
-	FSimpleDelegate OnEndPlaybackRangeDrag;
+	FSimpleDelegate OnPlaybackRangeEndDrag;
+
+	/** Attribute defining the selection range for this controller */
+	TAttribute<TRange<float>> SelectionRange;
+
+	/** Delegate that is called when the selection range wants to change */
+	FOnRangeChanged OnSelectionRangeChanged;
+
+	/** Called right before the selection range starts to be dragged */
+	FSimpleDelegate OnSelectionRangeBeginDrag;
 
+	/** Called right after the selection range has finished being dragged */
+	FSimpleDelegate OnSelectionRangeEndDrag;
 	/** Round the scrub position to an integer during playback */
 	TAttribute<EMovieScenePlayerStatus::Type> PlaybackStatus;
 
diff --git a/Engine/Source/Editor/UMGEditor/Private/Animation/WidgetMaterialTrackEditor.cpp b/Engine/Source/Editor/UMGEditor/Private/Animation/WidgetMaterialTrackEditor.cpp
new file mode 100644
--- /dev/null
+++ b/Engine/Source/Editor/UMGEditor/Private/Animation/WidgetMaterialTrackEditor.cpp
@@ -0,0 +1,38 @@
+// Copyright 1998-2016 Epic Games, Inc. All Rights Reserved.
+
+#include "UMGEditorPrivatePCH.h"
+#include "WidgetMaterialTrackEditor.h"
+#include "MovieSceneWidgetMaterialTrack.h"
+#include "WidgetMaterialTrackUtilities.h"
+
+
+FWidgetMaterialTrackEditor::FWidgetMaterialTrackEditor( TSharedRef<ISequencer> InSequencer )
+	: FMaterialTrackEditor( InSequencer )
+{
+}
+
+
+TSharedRef<ISequencerTrackEditor> FWidgetMaterialTrackEditor::CreateTrackEditor( TSharedRef<ISequencer> OwningSequencer )
+{
+	return MakeShareable( new FWidgetMaterialTrackEditor( OwningSequencer ) );
+}
+
+
+bool FWidgetMaterialTrackEditor::SupportsType( TSubclassOf<UMovieSceneTrack> Type ) const
+{
+	return Type == UMovieSceneWidgetMaterialTrack::StaticClass();
+}
+
+
+UMaterialInterface* FWidgetMaterialTrackEditor::GetMaterialInterfaceForTrack( FGuid ObjectBinding, UMovieSceneMaterialTrack* MaterialTrack )
+{
+	UObject* WidgetObject = GetSequencer()->GetFocusedMovieSceneSequenceInstance()->FindObject( ObjectBinding, *GetSequencer() );
+	UWidget* Widget = Cast<UWidget>( WidgetObject );
+	UMovieSceneWidgetMaterialTrack* WidgetMaterialTrack = Cast<UMovieSceneWidgetMaterialTrack>( MaterialTrack );
+	if ( Widget != nullptr && WidgetMaterialTrack != nullptr )
+	{
+		FSlateBrush* Brush = WidgetMaterialTrackUtilities::GetBrush( Widget, WidgetMaterialTrack->GetBrushPropertyName() );
+		return Cast<UMaterialInterface>(Brush->GetResourceObject());
+	}
+	return nullptr;
+}
diff --git a/Engine/Source/Editor/UMGEditor/Private/UMGEditorModule.cpp b/Engine/Source/Editor/UMGEditor/Private/UMGEditorModule.cpp
--- a/Engine/Source/Editor/UMGEditor/Private/UMGEditorModule.cpp
+++ b/Engine/Source/Editor/UMGEditor/Private/UMGEditorModule.cpp
@@ -13,6 +13,7 @@
 #include "Editor/Sequencer/Public/ISequencerModule.h"
 #include "Animation/MarginTrackEditor.h"
 #include "Animation/Sequencer2DTransformTrackEditor.h"
+#include "Animation/WidgetMaterialTrackEditor.h"
 #include "IUMGModule.h"
 #include "ComponentReregisterContext.h"
 #include "WidgetComponent.h"
@@ -70,8 +71,9 @@ public:
 
 		// Register with the sequencer module that we provide auto-key handlers.
 		ISequencerModule& SequencerModule = FModuleManager::Get().LoadModuleChecked<ISequencerModule>("Sequencer");
-		MarginTrackEditorCreateTrackEditorHandle    = SequencerModule.RegisterTrackEditor_Handle(FOnCreateTrackEditor::CreateStatic(&FMarginTrackEditor::CreateTrackEditor));
-		TransformTrackEditorCreateTrackEditorHandle = SequencerModule.RegisterTrackEditor_Handle(FOnCreateTrackEditor::CreateStatic(&F2DTransformTrackEditor::CreateTrackEditor));
+		MarginTrackEditorCreateTrackEditorHandle          = SequencerModule.RegisterTrackEditor_Handle(FOnCreateTrackEditor::CreateStatic(&FMarginTrackEditor::CreateTrackEditor));
+		TransformTrackEditorCreateTrackEditorHandle       = SequencerModule.RegisterTrackEditor_Handle(FOnCreateTrackEditor::CreateStatic(&F2DTransformTrackEditor::CreateTrackEditor));
+		WidgetMaterialTrackEditorCreateTrackEditorHandle  = SequencerModule.RegisterTrackEditor_Handle(FOnCreateTrackEditor::CreateStatic(&FWidgetMaterialTrackEditor::CreateTrackEditor));
 
 		ISettingsModule* SettingsModule = FModuleManager::GetModulePtr<ISettingsModule>("Settings");
 		if ( SettingsModule != nullptr )
@@ -114,6 +116,15 @@ public:
 		}
 		CreatedAssetTypeActions.Empty();
 
+		// Unregister sequencer track creation delegates
+		ISequencerModule* SequencerModule = FModuleManager::GetModulePtr<ISequencerModule>( "Sequencer" );
+		if ( SequencerModule != nullptr )
+		{
+			SequencerModule->UnRegisterTrackEditor_Handle( MarginTrackEditorCreateTrackEditorHandle );
+			SequencerModule->UnRegisterTrackEditor_Handle( TransformTrackEditorCreateTrackEditorHandle );
+			SequencerModule->UnRegisterTrackEditor_Handle( WidgetMaterialTrackEditorCreateTrackEditorHandle );
+		}
+
 		// Unregister the setting
 		ISettingsModule* SettingsModule = FModuleManager::GetModulePtr<ISettingsModule>("Settings");
 
@@ -195,6 +206,7 @@ private:
 
 	FDelegateHandle MarginTrackEditorCreateTrackEditorHandle;
 	FDelegateHandle TransformTrackEditorCreateTrackEditorHandle;
+	FDelegateHandle WidgetMaterialTrackEditorCreateTrackEditorHandle;
 
 	/** All created asset type actions.  Cached here so that we can unregister it during shutdown. */
 	TArray< TSharedPtr<IAssetTypeActions> > CreatedAssetTypeActions;
diff --git a/Engine/Source/Editor/UMGEditor/Private/WidgetBlueprintEditor.cpp b/Engine/Source/Editor/UMGEditor/Private/WidgetBlueprintEditor.cpp
--- a/Engine/Source/Editor/UMGEditor/Private/WidgetBlueprintEditor.cpp
+++ b/Engine/Source/Editor/UMGEditor/Private/WidgetBlueprintEditor.cpp
@@ -23,6 +23,11 @@
 #include "WidgetBlueprint.h"
 #include "Engine/SimpleConstructionScript.h"
 
+#include "MovieSceneWidgetMaterialTrack.h"
+#include "WidgetMaterialTrackUtilities.h"
+
+#include "ScopedTransaction.h"
+
 #define LOCTEXT_NAMESPACE "UMG"
 
 FWidgetBlueprintEditor::FWidgetBlueprintEditor()
@@ -836,16 +841,36 @@ void FWidgetBlueprintEditor::ExtendSequencerAddTrackMenu( FMenuBuilder& AddTrack
 	if ( ContextObjects.Num() == 1 )
 	{
 		UWidget* Widget = Cast<UWidget>( ContextObjects[0] );
-		if ( Widget != nullptr && Widget->GetParent() != nullptr && Widget->Slot != nullptr )
+		if ( Widget != nullptr )
 		{
-			AddTrackMenuBuilder.BeginSection( "Slot", LOCTEXT( "SlotSection", "Slot" ) );
+			if( Widget->GetParent() != nullptr && Widget->Slot != nullptr )
+			{
+				AddTrackMenuBuilder.BeginSection( "Slot", LOCTEXT( "SlotSection", "Slot" ) );
+				{
+					FUIAction AddSlotAction( FExecuteAction::CreateRaw( this, &FWidgetBlueprintEditor::AddSlotTrack, Widget->Slot ) );
+					FText AddSlotLabel = FText::Format(LOCTEXT("SlotLabelFormat", "{0} Slot"), FText::FromString(Widget->GetParent()->GetName()));
+					FText AddSlotToolTip = FText::Format(LOCTEXT("SlotToolTipFormat", "Add {0} slot"), FText::FromString( Widget->GetParent()->GetName()));
+					AddTrackMenuBuilder.AddMenuEntry(AddSlotLabel, AddSlotToolTip, FSlateIcon(), AddSlotAction);
+				}
+				AddTrackMenuBuilder.EndSection();
+			}
+
+			TArray<UProperty*> MaterialBrushProperties;
+			WidgetMaterialTrackUtilities::GetMaterialBrushProperties( Widget, MaterialBrushProperties );
+			if ( MaterialBrushProperties.Num() > 0 )
 			{
-				FUIAction AddSlotAction( FExecuteAction::CreateRaw( this, &FWidgetBlueprintEditor::AddSlotTrack, Widget->Slot ) );
-				FText AddSlotLabel = FText::Format(LOCTEXT("SlotLabelFormat", "{0} Slot"), FText::FromString(Widget->GetParent()->GetName()));
-				FText AddSlotToolTip = FText::Format(LOCTEXT("SlotToolTipFormat", "Add {0} slot"), FText::FromString( Widget->GetParent()->GetName()));
-				AddTrackMenuBuilder.AddMenuEntry(AddSlotLabel, AddSlotToolTip, FSlateIcon(), AddSlotAction);
+				AddTrackMenuBuilder.BeginSection( "Materials", LOCTEXT( "MaterialsSection", "Materials" ) );
+				{
+					for ( UProperty* MaterialBrushProperty : MaterialBrushProperties )
+					{
+						FUIAction AddMaterialAction( FExecuteAction::CreateRaw( this, &FWidgetBlueprintEditor::AddMaterialTrack, Widget, MaterialBrushProperty->GetFName(), MaterialBrushProperty->GetDisplayNameText() ) );
+						FText AddMaterialLabel = MaterialBrushProperty->GetDisplayNameText();
+						FText AddMaterialToolTip = FText::Format( LOCTEXT( "MaterialToolTipFormat", "Add a material track for the {0} property." ), MaterialBrushProperty->GetDisplayNameText() );
+						AddTrackMenuBuilder.AddMenuEntry( AddMaterialLabel, AddMaterialToolTip, FSlateIcon(), AddMaterialAction );
+					}
+				}
+				AddTrackMenuBuilder.EndSection();
 			}
-			AddTrackMenuBuilder.EndSection();
 		}
 	}
 }
@@ -855,4 +880,26 @@ void FWidgetBlueprintEditor::AddSlotTrack( UPanelSlot* Slot )
 	GetSequencer()->GetHandleToObject( Slot );
 }
 
+void FWidgetBlueprintEditor::AddMaterialTrack( UWidget* Widget, FName MaterialPropertyName, FText MaterialPropertyDisplayName )
+{
+	FGuid WidgetHandle = Sequencer->GetHandleToObject( Widget );
+	if ( WidgetHandle.IsValid() )
+	{
+		UMovieScene* MovieScene = Sequencer->GetFocusedMovieSceneSequence()->GetMovieScene();
+		if( MovieScene->FindTrack( UMovieSceneWidgetMaterialTrack::StaticClass(), WidgetHandle, MaterialPropertyName ) == nullptr)
+		{
+			const FScopedTransaction Transaction( LOCTEXT( "AddWidgetMaterialTrack", "Add widget material track" ) );
+
+			MovieScene->Modify();
+
+			UMovieSceneWidgetMaterialTrack* NewTrack = Cast<UMovieSceneWidgetMaterialTrack>( MovieScene->AddTrack( UMovieSceneWidgetMaterialTrack::StaticClass(), WidgetHandle ) );
+			NewTrack->Modify();
+			NewTrack->SetBrushPropertyName( MaterialPropertyName );
+			NewTrack->SetDisplayName( FText::Format( LOCTEXT( "TrackDisplayNameFormat", "{0} Material"), MaterialPropertyDisplayName ) );
+
+			Sequencer->NotifyMovieSceneDataChanged();
+		}
+	}
+}
+
 #undef LOCTEXT_NAMESPACE
diff --git a/Engine/Source/Editor/UMGEditor/Private/WidgetBlueprintEditor.h b/Engine/Source/Editor/UMGEditor/Private/WidgetBlueprintEditor.h
--- a/Engine/Source/Editor/UMGEditor/Private/WidgetBlueprintEditor.h
+++ b/Engine/Source/Editor/UMGEditor/Private/WidgetBlueprintEditor.h
@@ -190,6 +190,9 @@ private:
 	/** Add an animation track for the supplied slot to the current animation. */
 	void AddSlotTrack( UPanelSlot* Slot );
 
+	/** Add an animation track for the supplied material property to the current animation. */
+	void AddMaterialTrack( UWidget* Widget, FName MaterialPropertyName, FText MaterialPropertyDisplayName );
+
 private:
 	/** The preview scene that owns the preview GUI */
 	FPreviewScene PreviewScene;
diff --git a/Engine/Source/Editor/UnrealEd/Classes/Factories/Factory.h b/Engine/Source/Editor/UnrealEd/Classes/Factories/Factory.h
--- a/Engine/Source/Editor/UnrealEd/Classes/Factories/Factory.h
+++ b/Engine/Source/Editor/UnrealEd/Classes/Factories/Factory.h
@@ -1,19 +1,23 @@
 // Copyright 1998-2016 Epic Games, Inc. All Rights Reserved.
 
+#pragma once
+
+#include "Factory.generated.h"
+
+
 /**
  * Base class for all factories
  * An object responsible for creating and importing new objects.
  * 
  */
-
-#pragma once
-#include "Factory.generated.h"
-
 UCLASS(abstract)
-class UNREALED_API UFactory : public UObject
+class UNREALED_API UFactory
+	: public UObject
 {
 	GENERATED_UCLASS_BODY()
 
+public:
+
 	/** The class manufactured by this factory. */
 	UPROPERTY()
 	TSubclassOf<class UObject>  SupportedClass;
@@ -26,12 +30,6 @@
 	UPROPERTY()
 	TArray<FString> Formats;
 
-protected:
-	/** The default value to return from CanCreateNew() */
-	UPROPERTY()
-	uint32 bCreateNew:1;
-
-public:
 	/** true if the associated editor should be opened after creating a new object. */
 	UPROPERTY()
 	uint32 bEditAfterNew:1;
@@ -53,34 +51,82 @@
 	UPROPERTY()
 	int32 ImportPriority;
 
-	/** This is the import priority that all factories are given in the default constructor. */
-	static const int32 DefaultImportPriority;
+public:
 
-	static FString GetCurrentFilename() { return CurrentFilename; }
+	/**
+	 * @return true if the factory can currently create a new object from scratch.
+	 */
+	virtual bool CanCreateNew() const
+	{
+		return bCreateNew;
+	}
 
-	static FString CurrentFilename;
+	/**
+	 * Whether the specified file can be imported by this factory.
+	 *
+	 * @return true if the file is supported, false otherwise.
+	 */
+	virtual bool FactoryCanImport(const FString& Filename);
 
-	/** For interactive object imports, this value indicates whether the user wants objects to be automatically
-		overwritten (See EAppReturnType), or -1 if the user should be prompted. */
-	static int32 OverwriteYesOrNoToAllState;
+	/**
+	 * Create a new object by importing it from a file name.
+	 *
+	 * The default implementation of this method will load the contents of the entire
+	 * file into a byte buffer and call FactoryCreateBinary. User defined factories
+	 * may override this behavior to process the provided file name on their own.
+	 *
+	 * @param InClass
+	 * @param InParent
+	 * @param InName
+	 * @param Flags
+	 * @param Filename
+	 * @param Parms
+	 * @param Warn
+	 * @param bOutOperationCanceled Will indicate whether the user canceled the import.
+	 * @return The new object.
+	 */
+	virtual UObject* FactoryCreateFile(UClass* InClass, UObject* InParent, FName InName, EObjectFlags Flags, const FString& Filename, const TCHAR* Parms, FFeedbackContext* Warn, bool& bOutOperationCanceled);
 
-	/** If this value is true, warning messages will be shown once for all objects being imported at the same time.  
-		This value will be reset to false each time a new import operation is started. */
-	DEPRECATED(4.8, "bAllowOneTimeWarningMessages is due to be removed in future.")
-	static bool bAllowOneTimeWarningMessages;
+	/**
+	 * Create a new object by class.
+	 *
+	 * @param InClass
+	 * @param InParent
+	 * @param InName
+	 * @param Flags
+	 * @param Context
+	 * @param Warn
+	 * @param CallingContext
+	 * @return The new object.
+	 */
+	virtual UObject* FactoryCreateNew(UClass* InClass, UObject* InParent, FName InName, EObjectFlags Flags, UObject* Context, FFeedbackContext* Warn, FName CallingContext)
+	{
+		return FactoryCreateNew(InClass, InParent, InName, Flags, Context, Warn);
+	}
 
-	//~ Begin UObject Interface.
-	static void AddReferencedObjects(UObject* InThis, FReferenceCollector& Collector);
-	//~ End UObject Interface.
+	/**
+	 * Create a new object by class.
+	 *
+	 * @param InClass
+	 * @param InParent
+	 * @param InName
+	 * @param Flags
+	 * @param Context
+	 * @param Warn
+	 * @return The new object.
+	 */
+	virtual UObject* FactoryCreateNew(UClass* InClass, UObject* InParent, FName InName, EObjectFlags Flags, UObject* Context, FFeedbackContext* Warn)
+	{
+		return nullptr;
+	}
 
-	/** Helper function to sort an array of factories by their import priority - use as a predicate for Sort */
-	static bool SortFactoriesByPriority(const UFactory& A, const UFactory& B);
+	virtual UObject* ImportObject(UClass* InClass, UObject* InOuter, FName InName, EObjectFlags Flags, const FString& Filename, const TCHAR* Parms, bool& OutCanceled);
 
 	/** Returns true if this factory should be shown in the New Asset menu (by default calls CanCreateNew). */
 	virtual bool ShouldShowInNewMenu() const;
 
 	/** Returns an optional override brush name for the new asset menu. If this is not specified, the thumbnail for the supported class will be used. */
-	virtual FName GetNewAssetThumbnailOverride() const { return NAME_None; }
+	virtual FName GetNewAssetThumbnailOverride() const;
 
 	/** Returns the name of the factory for menus */
 	virtual FText GetDisplayName() const;
@@ -109,80 +155,169 @@
 
 	/**
 	 * Resolves SupportedClass for factories which support multiple classes.
-	 * Such factories will have a NULL SupportedClass member.
+	 * Such factories will have a nullptr SupportedClass member.
 	 */
 	virtual UClass* ResolveSupportedClass();
 
-	/**
-	 * Resets the saved state of this factory.  The states are used to suppress messages during multiple object import. 
-	 * It needs to be reset each time a new import is started
-	 */
-	static void ResetState();
-
-	/**
-	 * Pop up message to the user asking whether they wish to overwrite existing state or not
-	 */
-	static void DisplayOverwriteOptionsDialog(const FText& Message);
-
 	/** Opens a dialog to configure the factory properties. Return false if user opted out of configuring properties */
-	virtual bool ConfigureProperties() { return true; }
-
-	/**
-	 * @return true if the factory can currently create a new object from scratch.
-	 */
-	virtual bool CanCreateNew() const { return bCreateNew; }
+	virtual bool ConfigureProperties()
+	{
+		return true;
+	}
 
 	// @todo document
-	virtual UObject* FactoryCreateText( UClass* InClass, UObject* InParent, FName InName, EObjectFlags Flags, UObject* Context, const TCHAR* Type, const TCHAR*& Buffer, const TCHAR* BufferEnd, FFeedbackContext* Warn ) {return NULL;}
-
-	// @todo document
-	// @param Type must not be 0, e.g. TEXT("TGA")
-	virtual UObject* FactoryCreateBinary( UClass* InClass, UObject* InParent, FName InName, EObjectFlags Flags, UObject* Context, const TCHAR* Type, const uint8*& Buffer, const uint8* BufferEnd, FFeedbackContext* Warn) {return NULL;}
-	// @param Type must not be 0, e.g. TEXT("TGA")
-	virtual UObject* FactoryCreateBinary( UClass* InClass, UObject* InParent, FName InName, EObjectFlags Flags, UObject* Context, const TCHAR* Type, const uint8*& Buffer, const uint8* BufferEnd, FFeedbackContext* Warn, bool& bOutOperationCanceled) { return FactoryCreateBinary(InClass, InParent, InName, Flags, Context, Type, Buffer, BufferEnd, Warn); }
-
-	// @todo document
-	virtual UObject* FactoryCreateNew( UClass* InClass, UObject* InParent, FName InName, EObjectFlags Flags, UObject* Context, FFeedbackContext* Warn, FName CallingContext )
-	{
-		return FactoryCreateNew(InClass, InParent, InName, Flags, Context, Warn);
-	}
-	virtual UObject* FactoryCreateNew( UClass* InClass, UObject* InParent, FName InName, EObjectFlags Flags, UObject* Context, FFeedbackContext* Warn ) {return NULL;}
-
-	/**
-	 * @return	true if this factory can deal with the file sent in.
-	 */
-	virtual bool FactoryCanImport( const FString& Filename );
-
-
-	// @todo document
 	virtual bool ImportUntypedBulkDataFromText(const TCHAR*& Buffer, FUntypedBulkData& BulkData);
 
-
-	// @todo document
-	static UObject* StaticImportObject( UClass* Class, UObject* InOuter, FName Name, EObjectFlags Flags, const TCHAR* Filename=TEXT(""), UObject* Context=NULL, UFactory* Factory=NULL, const TCHAR* Parms=NULL, FFeedbackContext* Warn=GWarn, int32 MaxImportFileSize = 0xC100000 );
-	static UObject* StaticImportObject( UClass* Class, UObject* InOuter, FName Name, EObjectFlags Flags, bool& bOutOperationCanceled, const TCHAR* Filename=TEXT(""), UObject* Context=NULL, UFactory* Factory=NULL, const TCHAR* Parms=NULL, FFeedbackContext* Warn=GWarn, int32 MaxImportFileSize = 0xC100000 );
-
 	/** Creates a list of file extensions supported by this factory */
-	void GetSupportedFileExtensions( TArray<FString>& OutExtensions ) const;
+	void GetSupportedFileExtensions(TArray<FString>& OutExtensions) const;
 
-	/** 
-	 * Do clean up after importing is done. Will be called once for multi batch import
-	 */
+	/** Do clean up after importing is done. Will be called once for multi batch import. */
 	virtual void CleanUp() {}
 
 	/**
-	 * Creates an asset if it doesn't exist. If it does exist then it overwrites it if possible. If it can not overwrite then it will delete and replace. If it can not delete, it will return NULL.
+	 * Creates an asset if it doesn't exist. If it does exist then it overwrites it if possible. If it can not overwrite then it will delete and replace. If it can not delete, it will return nullptr.
 	 * 
-	 * @param	InClass		the class of the asset to create
-	 * @param	InPackage	the package to create this object within.
-	 * @param	Name		the name to give the new asset. If no value (NAME_None) is specified, the asset will be given a unique name in the form of ClassName_#.
-	 * @param	InFlags		the ObjectFlags to assign to the new asset.
-	 * @param	Template	if specified, the property values from this object will be copied to the new object, and the new object's ObjectArchetype value will be set to this object.
-	 *						If NULL, the class default object is used instead.
-	 * @return	A pointer to a new asset of the specified type or null if the creation failed.
+	 * @param InClass The class of the asset to create
+	 * @param InPackage The package to create this object within.
+	 * @param Name The name to give the new asset. If no value (NAME_None) is specified, the asset will be given a unique name in the form of ClassName_#.
+	 * @param InFlags The ObjectFlags to assign to the new asset.
+	 * @param Template If specified, the property values from this object will be copied to the new object, and the new object's ObjectArchetype value will be set to this object.
+	 *	               If nullptr, the class default object is used instead.
+	 * @return A pointer to a new asset of the specified type or null if the creation failed.
 	 */
-	UObject* CreateOrOverwriteAsset(UClass* InClass, UObject* InParent, FName InName, EObjectFlags InFlags, UObject* InTemplate = NULL) const;
+	UObject* CreateOrOverwriteAsset(UClass* InClass, UObject* InParent, FName InName, EObjectFlags InFlags, UObject* InTemplate = nullptr) const;
 
 	/** Returns a new starting point name for newly created assets in the content browser */
 	virtual FString GetDefaultNewAssetName() const;
+
+public:
+
+	/**
+	 * Pop up message to the user asking whether they wish to overwrite existing state or not.
+	 *
+	 * @param Message The message text.
+	 **/
+	static void DisplayOverwriteOptionsDialog(const FText& Message);
+
+	/** Get the name of the file currently being imported. */
+	static FString GetCurrentFilename()
+	{
+		return CurrentFilename;
+	}
+
+	/** Get the default import priority for factories. */
+	static int32 GetDefaultImportPriority()
+	{
+		return DefaultImportPriority;
+	}
+
+	/**
+	 * Resets the saved state of this factory.
+	 *
+	 * The states are used to suppress messages during multiple object import. 
+	 * It needs to be reset each time a new import is started
+	 */
+	static void ResetState();
+
+	/** Helper function to sort an array of factories by their import priority - use as a predicate for Sort */
+	static bool SortFactoriesByPriority(const UFactory& A, const UFactory& B);
+
+	// @todo document
+	static UObject* StaticImportObject(UClass* Class, UObject* InOuter, FName Name, EObjectFlags Flags, const TCHAR* Filename = TEXT(""), UObject* Context = nullptr, UFactory* Factory = nullptr, const TCHAR* Parms = nullptr, FFeedbackContext* Warn = GWarn, int32 MaxImportFileSize = 0xC100000);
+	static UObject* StaticImportObject(UClass* Class, UObject* InOuter, FName Name, EObjectFlags Flags, bool& bOutOperationCanceled, const TCHAR* Filename = TEXT(""), UObject* Context = nullptr, UFactory* Factory = nullptr, const TCHAR* Parms = nullptr, FFeedbackContext* Warn = GWarn, int32 MaxImportFileSize = 0xC100000);
+
+public:
+
+	//~ UObject interface
+
+	static void AddReferencedObjects(UObject* InThis, FReferenceCollector& Collector);
+
+protected:
+
+	/** The default value to return from CanCreateNew() */
+	UPROPERTY()
+	uint32 bCreateNew : 1;
+
+protected:
+
+	/**
+	 * Create a new object by importing it from a text buffer.
+	 *
+	 * @param InClass
+	 * @param InParent
+	 * @param InName
+	 * @param Flags
+	 * @param Context
+	 * @param Type (must not be nullptr, i.e. TEXT("TGA"))
+	 * @param Buffer
+	 * @param BufferEnd
+	 * @param Warn
+	 * @return The new object.
+	 */
+	virtual UObject* FactoryCreateText(UClass* InClass, UObject* InParent, FName InName, EObjectFlags Flags, UObject* Context, const TCHAR* Type, const TCHAR*& Buffer, const TCHAR* BufferEnd, FFeedbackContext* Warn)
+	{
+		return nullptr;
+	}
+
+	/**
+	 * Create a new object by importing it from a binary buffer.
+	 *
+	 * @param InClass
+	 * @param InParent
+	 * @param InName
+	 * @param Flags
+	 * @param Context
+	 * @param Type (must not be nullptr, i.e. TEXT("TGA"))
+	 * @param Buffer
+	 * @param BufferEnd
+	 * @param Warn
+	 * @return The new object.
+	 */
+	virtual UObject* FactoryCreateBinary(UClass* InClass, UObject* InParent, FName InName, EObjectFlags Flags, UObject* Context, const TCHAR* Type, const uint8*& Buffer, const uint8* BufferEnd, FFeedbackContext* Warn)
+	{
+		return nullptr;
+	}
+	
+	/**
+	 * Create a new object by importing it from a binary buffer (cancelable).
+	 *
+	 * @param InClass
+	 * @param InParent
+	 * @param InName
+	 * @param Flags
+	 * @param Context
+	 * @param Type (must not be nullptr, i.e. TEXT("TGA"))
+	 * @param Buffer
+	 * @param BufferEnd
+	 * @param Warn
+	 * @param bOutOperationCanceled Will indicate whether the user canceled the import.
+	 * @return The new object.
+	 */
+	virtual UObject* FactoryCreateBinary(UClass* InClass, UObject* InParent, FName InName, EObjectFlags Flags, UObject* Context, const TCHAR* Type, const uint8*& Buffer, const uint8* BufferEnd, FFeedbackContext* Warn, bool& bOutOperationCanceled)
+	{
+		return FactoryCreateBinary(InClass, InParent, InName, Flags, Context, Type, Buffer, BufferEnd, Warn);
+	}
+
+protected:
+
+	/** Name of the file currently being imported. */
+	static FString CurrentFilename;
+
+	/** This is the import priority that all factories are given in the default constructor. */
+	static const int32 DefaultImportPriority;
+
+	/**
+	 * For interactive object imports, this value indicates whether the user wants
+	 * objects to be automatically overwritten (See EAppReturnType), or -1 if the
+	 * user should be prompted.
+	 */
+	static int32 OverwriteYesOrNoToAllState;
+
+	/**
+	 * If this value is true, warning messages will be shown once for all objects
+	 * being imported at the same time. This value will be reset to false each time
+	 * a new import operation is started.
+	 */
+	DEPRECATED(4.8, "bAllowOneTimeWarningMessages is due to be removed in future.")
+	static bool bAllowOneTimeWarningMessages;
 };
diff --git a/Engine/Source/Editor/UnrealEd/Private/AutoReimport/ContentDirectoryMonitor.cpp b/Engine/Source/Editor/UnrealEd/Private/AutoReimport/ContentDirectoryMonitor.cpp
--- a/Engine/Source/Editor/UnrealEd/Private/AutoReimport/ContentDirectoryMonitor.cpp
+++ b/Engine/Source/Editor/UnrealEd/Private/AutoReimport/ContentDirectoryMonitor.cpp
@@ -230,7 +230,7 @@ UObject* AttemptImport(UClass* InFactoryType, UPackage* Package, FName InName, b
 		{
 			if (auto* SupportedClass = Factory->ResolveSupportedClass())
 			{
-				Asset = UFactory::StaticImportObject(SupportedClass, Package, InName, RF_Public | RF_Standalone, bCancelled, *FullFilename, nullptr, Factory);
+				Asset = Factory->ImportObject(SupportedClass, Package, InName, RF_Public | RF_Standalone, FullFilename, nullptr, bCancelled);
 			}
 		}
 		Factory->RemoveFromRoot();
@@ -308,6 +308,7 @@ void FContentDirectoryMonitor::ProcessAdditions(const DirectoryWatcher::FTimeLim
 				UObject* NewAsset = nullptr;
 
 				// Find a relevant factory for this file
+				// @todo import: gmp: show dialog in case of multiple matching factories
 				const FString Ext = FPaths::GetExtension(Addition.Filename.Get(), false);
 				auto* Factories = InFactoriesByExtension.Find(Ext);
 				if (Factories && Factories->Num() != 0)
@@ -537,4 +538,5 @@ void FContentDirectoryMonitor::Abort()
 	}
 }
 
+
 #undef LOCTEXT_NAMESPACE
diff --git a/Engine/Source/Editor/UnrealEd/Private/Editor.cpp b/Engine/Source/Editor/UnrealEd/Private/Editor.cpp
--- a/Engine/Source/Editor/UnrealEd/Private/Editor.cpp
+++ b/Engine/Source/Editor/UnrealEd/Private/Editor.cpp
@@ -498,7 +498,7 @@ FReimportManager::~FReimportManager()
 
 int32 FReimportHandler::GetPriority() const
 {
-	return UFactory::DefaultImportPriority;
+	return UFactory::GetDefaultImportPriority();
 }
 
 /*-----------------------------------------------------------------------------
diff --git a/Engine/Source/Editor/UnrealEd/Private/EditorModes.cpp b/Engine/Source/Editor/UnrealEd/Private/EditorModes.cpp
--- a/Engine/Source/Editor/UnrealEd/Private/EditorModes.cpp
+++ b/Engine/Source/Editor/UnrealEd/Private/EditorModes.cpp
@@ -43,6 +43,8 @@ const FEditorModeID FBuiltinEditorModes::EM_Level(TEXT("EM_Level"));
 const FEditorModeID FBuiltinEditorModes::EM_StreamingLevel(TEXT("EM_StreamingLevel"));
 const FEditorModeID FBuiltinEditorModes::EM_Physics(TEXT("EM_Physics"));
 const FEditorModeID FBuiltinEditorModes::EM_ActorPicker(TEXT("EM_ActorPicker"));
+const FEditorModeID FBuiltinEditorModes::EM_SceneDepthPicker(TEXT("EM_SceneDepthPicker"));
+
 
 /*------------------------------------------------------------------------------
 	Default.
diff --git a/Engine/Source/Editor/UnrealEd/Private/EditorViewportClient.cpp b/Engine/Source/Editor/UnrealEd/Private/EditorViewportClient.cpp
--- a/Engine/Source/Editor/UnrealEd/Private/EditorViewportClient.cpp
+++ b/Engine/Source/Editor/UnrealEd/Private/EditorViewportClient.cpp
@@ -956,6 +956,13 @@ FSceneView* FEditorViewportClient::CalcSceneView(FSceneViewFamily* ViewFamily, c
 	if (bUseControllingActorViewInfo)
 	{
 		View->OverridePostProcessSettings(ControllingActorViewInfo.PostProcessSettings, ControllingActorViewInfo.PostProcessBlendWeight);
+
+		for (int32 ExtraPPBlendIdx = 0; ExtraPPBlendIdx < ControllingActorExtraPostProcessBlends.Num(); ++ExtraPPBlendIdx)
+		{
+			FPostProcessSettings const& PPSettings = ControllingActorExtraPostProcessBlends[ExtraPPBlendIdx];
+			float const Weight = ControllingActorExtraPostProcessBlendWeights[ExtraPPBlendIdx];
+			View->OverridePostProcessSettings(PPSettings, Weight);
+		}
 	}
 	else
 	{
@@ -4240,7 +4247,38 @@ void FEditorViewportClient::OverrideFarClipPlane(const float InFarPlane)
 	FarPlane = InFarPlane;
 }
 
-void FEditorViewportClient::MoveViewportCamera(const FVector& InDrag, const FRotator& InRot, bool bDollyCamera )
+float FEditorViewportClient::GetSceneDepthAtLocation(int32 X, int32 Y)
+{
+	// #todo: in the future we will just sample the depth buffer
+	return 0.f;
+}
+
+FVector FEditorViewportClient::GetHitProxyObjectLocation(int32 X, int32 Y)
+{
+	// #todo: for now we are just getting the actor and using its location for 
+	// depth. in the future we will just sample the depth buffer
+	HHitProxy* const HitProxy = Viewport->GetHitProxy(X, Y);
+	if (HitProxy && HitProxy->IsA(HActor::StaticGetType()))
+	{
+		HActor* const ActorHit = static_cast<HActor*>(HitProxy);
+
+		// dist to component will be more reliable than dist to actor
+		if (ActorHit->PrimComponent != nullptr)
+		{
+			return ActorHit->PrimComponent->GetComponentLocation();
+		}
+
+		if (ActorHit->Actor != nullptr)
+		{
+			return ActorHit->Actor->GetActorLocation();
+		}
+	}
+
+	return FVector::ZeroVector;
+}
+
+
+void FEditorViewportClient::MoveViewportCamera(const FVector& InDrag, const FRotator& InRot, bool bDollyCamera)
 {
 	switch( GetViewportType() )
 	{
diff --git a/Engine/Source/Editor/UnrealEd/Private/Factories/ApexDestructibleAssetImport.cpp b/Engine/Source/Editor/UnrealEd/Private/Factories/ApexDestructibleAssetImport.cpp
--- a/Engine/Source/Editor/UnrealEd/Private/Factories/ApexDestructibleAssetImport.cpp
+++ b/Engine/Source/Editor/UnrealEd/Private/Factories/ApexDestructibleAssetImport.cpp
@@ -1023,7 +1023,7 @@ UDestructibleMesh* ImportDestructibleMeshFromApexDestructibleAsset(UObject* InPa
 		// Store the current file path and timestamp for re-import purposes
 		// @todo AssetImportData make a data class for Apex destructible assets
 		DestructibleMesh->AssetImportData = NewObject<UAssetImportData>(DestructibleMesh);
-		DestructibleMesh->AssetImportData->Update(UFactory::CurrentFilename);
+		DestructibleMesh->AssetImportData->Update(UFactory::GetCurrentFilename());
 	}
 
 	DestructibleMesh->PreEditChange(NULL);
diff --git a/Engine/Source/Editor/UnrealEd/Private/Factories/EditorFactories.cpp b/Engine/Source/Editor/UnrealEd/Private/Factories/EditorFactories.cpp
--- a/Engine/Source/Editor/UnrealEd/Private/Factories/EditorFactories.cpp
+++ b/Engine/Source/Editor/UnrealEd/Private/Factories/EditorFactories.cpp
@@ -210,21 +210,26 @@ UObject* UMaterialFactoryNew::FactoryCreateNew(UClass* Class,UObject* InParent,F
 	if ( InitialTexture != nullptr )
 	{
 		// An initial texture was specified, add it and assign it to the BaseColor
-		UMaterialExpressionTextureSample* Expression = NewObject<UMaterialExpressionTextureSample>(NewMaterial);
-		NewMaterial->Expressions.Add( Expression );
+		UMaterialExpressionTextureSample* TextureSampler = NewObject<UMaterialExpressionTextureSample>(NewMaterial);
+		{
+			TextureSampler->MaterialExpressionEditorX = -250;
+			TextureSampler->Texture = InitialTexture;
+			TextureSampler->AutoSetSampleType();
+		}
 
-		NewMaterial->BaseColor.Expression = Expression;
-		Expression->MaterialExpressionEditorX = -250;
-		Expression->Texture = InitialTexture;
+		NewMaterial->Expressions.Add(TextureSampler);
 
-		TArray<FExpressionOutput> Outputs;
-		Outputs = Expression->GetOutputs();
-		FExpressionOutput* Output = Outputs.GetData();
-		NewMaterial->BaseColor.Mask = Output->Mask;
-		NewMaterial->BaseColor.MaskR = Output->MaskR;
-		NewMaterial->BaseColor.MaskG = Output->MaskG;
-		NewMaterial->BaseColor.MaskB = Output->MaskB;
-		NewMaterial->BaseColor.MaskA = Output->MaskA;
+		FExpressionOutput& Output = TextureSampler->GetOutputs()[0];
+		FExpressionInput& Input = (TextureSampler->SamplerType == SAMPLERTYPE_Normal)
+			? (FExpressionInput&)NewMaterial->Normal
+			: (FExpressionInput&)NewMaterial->BaseColor;
+
+		Input.Expression = TextureSampler;
+		Input.Mask = Output.Mask;
+		Input.MaskR = Output.MaskR;
+		Input.MaskG = Output.MaskG;
+		Input.MaskB = Output.MaskB;
+		Input.MaskA = Output.MaskA;
 
 		NewMaterial->PostEditChange();
 	}
@@ -1914,18 +1919,21 @@ EReimportResult::Type UReimportSoundFactory::Reimport( UObject* Obj )
 	// Suppress the import overwrite dialog, we want to keep existing settings when re-importing
 	USoundFactory::SuppressImportOverwriteDialog();
 
-	if( UFactory::StaticImportObject( SoundWave->GetClass(), SoundWave->GetOuter(), *SoundWave->GetName(), RF_Public|RF_Standalone, *Filename, nullptr, this ) )
+	bool OutCanceled = false;
+	if (ImportObject(SoundWave->GetClass(), SoundWave->GetOuter(), *SoundWave->GetName(), RF_Public | RF_Standalone, Filename, nullptr, OutCanceled) != nullptr)
 	{
 		UE_LOG(LogEditorFactories, Log, TEXT("-- imported successfully") );
 
 		SoundWave->AssetImportData->Update(Filename);
-
-		// Mark the package dirty after the successful import
 		SoundWave->MarkPackageDirty();
 	}
+	else if (OutCanceled)
+	{
+		UE_LOG(LogEditorFactories, Warning, TEXT("-- import canceled"));
+	}
 	else
 	{
-		UE_LOG(LogEditorFactories, Warning, TEXT("-- import failed") );
+		UE_LOG(LogEditorFactories, Warning, TEXT("-- import failed"));
 	}
 
 	return EReimportResult::Succeeded;
@@ -2295,8 +2303,9 @@ EReimportResult::Type UReimportSoundSurroundFactory::Reimport( UObject* Obj )
 
 		FString SpeakerLocation = FPaths::GetBaseFilename(Filename).Right(3);
 		FName ImportName = *(SoundWave->GetName() + SpeakerLocation);
+		bool OutCanceled = false;
 
-		if (UFactory::StaticImportObject(SoundWave->GetClass(), SoundWave->GetOuter(), ImportName, RF_Public|RF_Standalone, *Filename, nullptr, this ))
+		if (ImportObject(SoundWave->GetClass(), SoundWave->GetOuter(), ImportName, RF_Public | RF_Standalone, Filename, nullptr, OutCanceled) != nullptr)
 		{
 			FFormatNamedArguments Arguments;
 			Arguments.Add(TEXT("NameText"), NameText);
@@ -2308,7 +2317,7 @@ EReimportResult::Type UReimportSoundSurroundFactory::Reimport( UObject* Obj )
 
 			bSourceReimported = true;
 		}
-		else
+		else if (!OutCanceled)
 		{
 			FFormatNamedArguments Arguments;
 			Arguments.Add(TEXT("NameText"), NameText);
@@ -5509,7 +5518,9 @@ EReimportResult::Type UReimportTextureFactory::Reimport( UObject* Obj )
 	// Suppress the import overwrite dialog because we know that for explicitly re-importing we want to preserve existing settings
 	UTextureFactory::SuppressImportOverwriteDialog();
 
-	if (UFactory::StaticImportObject(pTex->GetClass(), pTex->GetOuter(), *pTex->GetName(), RF_Public|RF_Standalone, *ResolvedSourceFilePath, nullptr, this))
+	bool OutCanceled = false;
+
+	if (ImportObject(pTex->GetClass(), pTex->GetOuter(), *pTex->GetName(), RF_Public | RF_Standalone, ResolvedSourceFilePath, nullptr, OutCanceled) != nullptr)
 	{
 		UE_LOG(LogEditorFactories, Log, TEXT("-- imported successfully") );
 
@@ -5525,9 +5536,13 @@ EReimportResult::Type UReimportTextureFactory::Reimport( UObject* Obj )
 			pTex->MarkPackageDirty();
 		}
 	}
+	else if (OutCanceled)
+	{
+		UE_LOG(LogEditorFactories, Warning, TEXT("-- import canceled"));
+	}
 	else
 	{
-		UE_LOG(LogEditorFactories, Warning, TEXT("-- import failed") );
+		UE_LOG(LogEditorFactories, Warning, TEXT("-- import failed"));
 	}
 	
 	return EReimportResult::Succeeded;
diff --git a/Engine/Source/Editor/UnrealEd/Private/Factories/Factory.cpp b/Engine/Source/Editor/UnrealEd/Private/Factories/Factory.cpp
--- a/Engine/Source/Editor/UnrealEd/Private/Factories/Factory.cpp
+++ b/Engine/Source/Editor/UnrealEd/Private/Factories/Factory.cpp
@@ -1,84 +1,174 @@
 // Copyright 1998-2016 Epic Games, Inc. All Rights Reserved.
 
-
-// Editor includes.
 #include "UnrealEd.h"
 #include "ObjectTools.h"
 #include "AssetToolsModule.h"
 #include "EditorClassUtils.h"
 
+
 DEFINE_LOG_CATEGORY_STATIC(LogFactory, Log, All);
 
-/*----------------------------------------------------------------------------
-	UFactory.
-----------------------------------------------------------------------------*/
+
 FString UFactory::CurrentFilename(TEXT(""));
 
 // This needs to be greater than 0 to allow factories to have both higher and lower priority than the default
 const int32 UFactory::DefaultImportPriority = 100;
-
 int32 UFactory::OverwriteYesOrNoToAllState = -1;
-
 bool UFactory::bAllowOneTimeWarningMessages = true;
 
+
 UFactory::UFactory(const FObjectInitializer& ObjectInitializer)
 	: Super(ObjectInitializer)
 {
 	ImportPriority = DefaultImportPriority;
 }
 
+
 void UFactory::AddReferencedObjects(UObject* InThis, FReferenceCollector& Collector)
 {	
 	UFactory* This = CastChecked<UFactory>(InThis);
 	UClass* SupportedClass = *This->SupportedClass;
 	UClass* ContextClass = *This->ContextClass;
-	Collector.AddReferencedObject( SupportedClass, This );
-	Collector.AddReferencedObject( ContextClass, This );
+	Collector.AddReferencedObject(SupportedClass, This);
+	Collector.AddReferencedObject(ContextClass, This);
 
-	Super::AddReferencedObjects( This, Collector );
+	Super::AddReferencedObjects(This, Collector);
 }
 
 
-bool UFactory::FactoryCanImport( const FString& Filename )
+UObject* UFactory::FactoryCreateFile(UClass* InClass, UObject* InParent, FName InName, EObjectFlags Flags, const FString& Filename, const TCHAR* Parms, FFeedbackContext* Warn, bool& bOutOperationCanceled)
 {
-	//check extension (only do the following for t3d)
-	if (FPaths::GetExtension(Filename) == TEXT("t3d"))
+	FString FileExtension = FPaths::GetExtension(Filename);
+
+	// load as text
+	if (bText)
 	{
-		//Open file
 		FString Data;
-		if( FFileHelper::LoadFileToString( Data, *Filename ) )
+		if (!FFileHelper::LoadFileToString(Data, *Filename))
 		{
-			const TCHAR* Str= *Data;
+			UE_LOG(LogFactory, Error, TEXT("Failed to load file '%s' to string"), *Filename);
+			return nullptr;
+		}
+
+		ParseParms(Parms);
+		const TCHAR* Ptr = *Data;
+			
+		return FactoryCreateText(InClass, InParent, InName, Flags, nullptr, *FileExtension, Ptr, Ptr + Data.Len(), Warn);
+	}
+
+	// load as binary
+	{
+		TArray<uint8> Data;
+		if (!FFileHelper::LoadFileToArray(Data, *Filename))
+		{
+			UE_LOG(LogFactory, Error, TEXT("Failed to load file '%s' to array"), *Filename);
+			return nullptr;
+		}
+
+		Data.Add(0);
+		ParseParms(Parms);
+		const uint8* Ptr = &Data[0];
+
+		return FactoryCreateBinary(InClass, InParent, InName, Flags, nullptr, *FileExtension, Ptr, Ptr + Data.Num() - 1, Warn, bOutOperationCanceled);
+	}
+}
+
+
+bool UFactory::FactoryCanImport( const FString& Filename )
+{
+	// only T3D is supported
+	if (FPaths::GetExtension(Filename) != TEXT("t3d"))
+	{
+		return false;
+	}
 
-			if( FParse::Command(&Str, TEXT("BEGIN") ) && FParse::Command(&Str, TEXT("OBJECT")) )
+	// open file
+	FString Data;
+
+	if (FFileHelper::LoadFileToString(Data, *Filename))
+	{
+		const TCHAR* Str= *Data;
+		if (FParse::Command(&Str, TEXT("BEGIN")) && FParse::Command(&Str, TEXT("OBJECT")))
+		{
+			FString strClass;
+			if (FParse::Value(Str, TEXT("CLASS="), strClass))
 			{
-				FString strClass;
-				if (FParse::Value(Str, TEXT("CLASS="), strClass))
+				//we found the right syntax, so no error if we don't match
+				if (strClass == SupportedClass->GetName())
 				{
-					//we found the right syntax, so no error if we don't match
-					if (strClass == SupportedClass->GetName())
-					{
-						return true;
-					}
-					return false;
+					return true;
 				}
+
+				return false;
 			}
-			UE_LOG(LogFactory, Warning, TEXT("Factory import failed due to invalid format: %s"), *Filename);
+		}
+
+		UE_LOG(LogFactory, Warning, TEXT("Factory import failed due to invalid format: %s"), *Filename);
+	}
+	else
+	{
+		UE_LOG(LogFactory, Warning, TEXT("Factory import failed due to inability to load file %s"), *Filename);
+	}
+
+	return false;
+}
+
+
+UObject* UFactory::ImportObject(UClass* InClass, UObject* InOuter, FName InName, EObjectFlags InFlags, const FString& Filename, const TCHAR* Parms, bool& OutCanceled)
+{
+	UObject* Result = nullptr;
+	CurrentFilename = Filename;
+
+	if (CanCreateNew())
+	{
+		UE_LOG(LogFactory, Log, TEXT("FactoryCreateNew: %s with %s (%i %i %s)"), *InClass->GetName(), *GetClass()->GetName(), bCreateNew, bText, *Filename);
+		ParseParms(Parms);
+
+		Result = FactoryCreateNew(InClass, InOuter, InName, InFlags, nullptr, GWarn);
+	}
+	else if (!Filename.IsEmpty())
+	{
+		// sanity check the file size of the impending import and prompt
+		// the user if they wish to continue if the file size is very large
+		const int32 FileSize = IFileManager::Get().FileSize(*Filename);
+
+		if (FileSize == INDEX_NONE)
+		{
+			UE_LOG(LogFactory, Error, TEXT("Can't find file '%s' for import"), *Filename);
 		}
 		else
 		{
-			UE_LOG(LogFactory, Warning, TEXT("Factory import failed due to inability to load file %s"), *Filename);
+			UE_LOG(LogFactory, Log, TEXT("FactoryCreateFile: %s with %s (%i %i %s)"), *InClass->GetName(), *GetClass()->GetName(), bCreateNew, bText, *Filename);
+
+			Result = FactoryCreateFile(InClass, InOuter, InName, InFlags, *Filename, Parms, GWarn, OutCanceled);
 		}
 	}
 
-	return false;
+	if (Result != nullptr)
+	{
+		Result->MarkPackageDirty();
+		ULevel::LevelDirtiedEvent.Broadcast();
+		Result->PostEditChange();
+	}
+
+	CurrentFilename = TEXT("");
+
+	return Result;
 }
 
+
 bool UFactory::ShouldShowInNewMenu() const
 {
 	return CanCreateNew();
 }
 
+
+FName UFactory::GetNewAssetThumbnailOverride() const
+{
+	return NAME_None;
+}
+
+
 FText UFactory::GetDisplayName() const
 {
 	FAssetToolsModule& AssetToolsModule = FModuleManager::LoadModuleChecked<FAssetToolsModule>("AssetTools");
@@ -101,15 +191,16 @@ FText UFactory::GetDisplayName() const
 	return FText();
 }
 
+
 uint32 UFactory::GetMenuCategories() const
 {
 	FAssetToolsModule& AssetToolsModule = FModuleManager::LoadModuleChecked<FAssetToolsModule>("AssetTools");
-
 	UClass* LocalSupportedClass = GetSupportedClass();
-	if ( LocalSupportedClass )
+
+	if (LocalSupportedClass)
 	{
 		TWeakPtr<IAssetTypeActions> AssetTypeActions = AssetToolsModule.Get().GetAssetTypeActionsForClass(LocalSupportedClass);
-		if ( AssetTypeActions.IsValid() )
+		if (AssetTypeActions.IsValid())
 		{
 			return AssetTypeActions.Pin()->GetCategories();
 		}
@@ -119,21 +210,25 @@ uint32 UFactory::GetMenuCategories() const
 	return EAssetTypeCategories::Misc;
 }
 
+
 FText UFactory::GetToolTip() const
 {
 	return GetSupportedClass()->GetToolTipText();
 }
 
+
 FString UFactory::GetToolTipDocumentationPage() const
 {
 	return FEditorClassUtils::GetDocumentationPage(GetSupportedClass());
 }
 
+
 FString UFactory::GetToolTipDocumentationExcerpt() const
 {
 	return FEditorClassUtils::GetDocumentationExcerpt(GetSupportedClass());
 }
 
+
 UClass* UFactory::GetSupportedClass() const
 {
 	return SupportedClass;
@@ -145,14 +240,16 @@ bool UFactory::DoesSupportClass(UClass * Class)
 	return (Class == GetSupportedClass());
 }
 
+
 UClass* UFactory::ResolveSupportedClass()
 {
 	// This check forces factories which support multiple classes to overload this method.
-	// In other words, you can't have a SupportedClass of NULL and not overload this method.
+	// In other words, you can't have a SupportedClass of nullptr and not overload this method.
 	check( SupportedClass );
 	return SupportedClass;
 }
 
+
 void UFactory::ResetState()
 {
 	// Resets the state of the 'Yes To All / No To All' prompt for overwriting existing objects on import.
@@ -160,6 +257,7 @@ void UFactory::ResetState()
 	OverwriteYesOrNoToAllState = -1;
 }
 
+
 void UFactory::DisplayOverwriteOptionsDialog(const FText& Message)
 {
 	// Prompt the user for what to do if a 'To All' response wasn't already given.
@@ -171,14 +269,16 @@ void UFactory::DisplayOverwriteOptionsDialog(const FText& Message)
 	}
 }
 
+
 bool UFactory::SortFactoriesByPriority(const UFactory& A, const UFactory& B)
 {
 	// First sort so that higher priorities are earlier in the list
-	if( A.ImportPriority > B.ImportPriority )
+	if (A.ImportPriority > B.ImportPriority)
 	{
 		return true;
 	}
-	else if( A.ImportPriority < B.ImportPriority )
+
+	if (A.ImportPriority < B.ImportPriority)
 	{
 		return false;
 	}
@@ -186,7 +286,8 @@ bool UFactory::SortFactoriesByPriority(const UFactory& A, const UFactory& B)
 	// Then sort so that factories that only create new assets are tried after those that actually import the file data (when they have an equivalent priority)
 	const bool bFactoryAImportsFiles = !A.CanCreateNew();
 	const bool bFactoryBImportsFiles = !B.CanCreateNew();
-	if( bFactoryAImportsFiles && !bFactoryBImportsFiles )
+
+	if (bFactoryAImportsFiles && !bFactoryBImportsFiles)
 	{
 		return true;
 	}
@@ -194,24 +295,26 @@ bool UFactory::SortFactoriesByPriority(const UFactory& A, const UFactory& B)
 	return false;
 }
 
+
 UObject* UFactory::StaticImportObject
 (
-UClass*				Class,
-UObject*			InOuter,
-FName				Name,
-EObjectFlags		Flags,
-const TCHAR*		Filename,
-UObject*			Context,
-UFactory*			InFactory,
-const TCHAR*		Parms,
-FFeedbackContext*	Warn,
-int32				MaxImportFileSize
+	UClass*				Class,
+	UObject*			InOuter,
+	FName				Name,
+	EObjectFlags		Flags,
+	const TCHAR*		Filename,
+	UObject*			Context,
+	UFactory*			InFactory,
+	const TCHAR*		Parms,
+	FFeedbackContext*	Warn,
+	int32				MaxImportFileSize
 )
 {
 	bool bOperationCanceled = false;
 	return StaticImportObject(Class, InOuter, Name, Flags, bOperationCanceled, Filename, Context, InFactory, Parms, Warn, MaxImportFileSize);
 }
 
+
 UObject* UFactory::StaticImportObject
 (
 	UClass*				Class,
@@ -229,11 +332,12 @@ UObject* UFactory::StaticImportObject
 {
 	check(Class);
 
-	CurrentFilename = Filename;
-	FString extension = FPaths::GetExtension(CurrentFilename);
-	// Make list of all applicable factories.
+	UObject* Result = nullptr;
+	FString Extension = FPaths::GetExtension(Filename);
 	TArray<UFactory*> Factories;
-	if( InFactory )
+
+	// Make list of all applicable factories.
+	if (InFactory != nullptr)
 	{
 		// Use just the specified factory.
 		if (ensureMsgf( !InFactory->SupportedClass || Class->IsChildOf(InFactory->SupportedClass), 
@@ -245,12 +349,14 @@ UObject* UFactory::StaticImportObject
 	else
 	{
 		auto TransientPackage = GetTransientPackage();
-		// Try all automatic factories, sorted by priority.
-		for( TObjectIterator<UClass> It; It; ++It )
+
+		// try all automatic factories, sorted by priority
+		for (TObjectIterator<UClass> It; It; ++It)
 		{
-			if( It->IsChildOf( UFactory::StaticClass() ) )
+			if (It->IsChildOf(UFactory::StaticClass()))
 			{
 				UFactory* Default = It->GetDefaultObject<UFactory>();
+
 				if (Class->IsChildOf(Default->SupportedClass) && Default->ImportPriority >= 0)
 				{
 					//Add the factory if there is no extension or the factory don't have any supported file extension or the factory support this file extension.
@@ -258,8 +364,11 @@ UObject* UFactory::StaticImportObject
 					//See UFactory::SortFactoriesByPriority
 					TArray<FString> FactoryExtension;
 					Default->GetSupportedFileExtensions(FactoryExtension);
-					if (extension.IsEmpty() || FactoryExtension.Num() == 0 || FactoryExtension.Contains(extension))
+
+					if (Extension.IsEmpty() || (FactoryExtension.Num() == 0) || (FactoryExtension.Contains(Extension)))
+					{
 						Factories.Add(NewObject<UFactory>(TransientPackage, *It));
+					}
 				}
 			}
 		}
@@ -267,95 +376,40 @@ UObject* UFactory::StaticImportObject
 		Factories.Sort(&UFactory::SortFactoriesByPriority);
 	}
 
-	bool bLoadedFile = false;
-
 	// Try each factory in turn.
-	for( int32 i=0; i<Factories.Num(); i++ )
+	for (auto& Factory : Factories)
 	{
-		UFactory* Factory = Factories[i];
-		UObject* Result = NULL;
-		if( Factory->CanCreateNew() )
-		{
-			UE_LOG(LogFactory, Log,  TEXT("FactoryCreateNew: %s with %s (%i %i %s)"), *Class->GetName(), *Factories[i]->GetClass()->GetName(), Factory->bCreateNew, Factory->bText, Filename );
-			Factory->ParseParms( Parms );
-			Result = Factory->FactoryCreateNew( Class, InOuter, Name, Flags, NULL, Warn );
-		}
-		else if( FCString::Stricmp(Filename,TEXT(""))!=0 )
-		{
-			if( Factory->bText )
-			{
-				//UE_LOG(LogFactory, Log,  TEXT("FactoryCreateText: %s with %s (%i %i %s)"), *Class->GetName(), *Factories(i)->GetClass()->GetName(), Factory->bCreateNew, Factory->bText, Filename );
-				FString Data;
-				if( FFileHelper::LoadFileToString( Data, Filename ) )
-				{
-					bLoadedFile = true;
-					const TCHAR* Ptr = *Data;
-					Factory->ParseParms( Parms );
-					Result = Factory->FactoryCreateText( Class, InOuter, Name, Flags, NULL, *extension, Ptr, Ptr+Data.Len(), Warn );
-				}
-			}
-			else
-			{
-				UE_LOG(LogFactory, Log,  TEXT("FactoryCreateBinary: %s with %s (%i %i %s)"), *Class->GetName(), *Factories[i]->GetClass()->GetName(), Factory->bCreateNew, Factory->bText, Filename );
-				
-				// Sanity check the file size of the impending import and prompt the user if they wish to continue if the file size is very large
-				const int32 FileSize = IFileManager::Get().FileSize( Filename );
-				bool bValidFileSize = true;
-
-				if ( FileSize == INDEX_NONE )
-				{
-					UE_LOG(LogFactory, Error,TEXT("File '%s' does not exist"), Filename );
-					bValidFileSize = false;
-				}
+		Result = Factory->ImportObject(Class, InOuter, Name, Flags, Filename, Parms, bOutOperationCanceled);
 
-				TArray<uint8> Data;
-				if( bValidFileSize && FFileHelper::LoadFileToArray( Data, Filename ) )
-				{
-					bLoadedFile = true;
-					Data.Add( 0 );
-					const uint8* Ptr = &Data[ 0 ];
-					Factory->ParseParms( Parms );
-					Result = Factory->FactoryCreateBinary(Class, InOuter, Name, Flags, NULL, *extension, Ptr, Ptr + Data.Num() - 1, Warn, bOutOperationCanceled);
-				}
-			}
-		}
-		if( Result )
+		if (Result != nullptr)
 		{
-			// prevent UTextureCube created from UTextureFactory			check(Result->IsA(Class));
-			Result->MarkPackageDirty();
-			ULevel::LevelDirtiedEvent.Broadcast();
-			Result->PostEditChange();
-
-			CurrentFilename = TEXT("");
-			return Result;
+			break;
 		}
 	}
 
-	if ( !bLoadedFile && !bOutOperationCanceled )
+	if ((Result == nullptr) && !bOutOperationCanceled)
 	{
-		Warn->Logf( *FText::Format( NSLOCTEXT( "UnrealEd", "NoFindImport", "Can't find file '{0}' for import" ), FText::FromString( FString(Filename) ) ).ToString() );
+		Warn->Logf(*FText::Format(NSLOCTEXT("UnrealEd", "ImportFailed", "Failed to import file '{0}'"), FText::FromString(FString(Filename))).ToString());
 	}
 
-	CurrentFilename = TEXT("");
-
-	return NULL;
+	return Result;
 }
 
 
 void UFactory::GetSupportedFileExtensions(TArray<FString>& OutExtensions) const
 {
-	for ( int32 FormatIdx = 0; FormatIdx < Formats.Num(); ++FormatIdx )
+	for (const auto& Format : Formats)
 	{
-		const FString& Format = Formats[FormatIdx];
 		const int32 DelimiterIdx = Format.Find(TEXT(";"));
 
-		if ( DelimiterIdx != INDEX_NONE )
+		if (DelimiterIdx != INDEX_NONE)
 		{
-			OutExtensions.Add( Format.Left(DelimiterIdx) );
+			OutExtensions.Add(Format.Left(DelimiterIdx));
 		}
 	}
 }
 
+
 bool UFactory::ImportUntypedBulkDataFromText(const TCHAR*& Buffer, FUntypedBulkData& BulkData)
 {
 	FString StrLine;
@@ -373,16 +427,15 @@ bool UFactory::ImportUntypedBulkDataFromText(const TCHAR*& Buffer, FUntypedBulkD
 			/** Number of elements in bulk data array */
 			ElementCount = FCString::Atoi(*ParsedText);
 		}
-		else
-		if (FParse::Value(Str, TEXT("ELEMENTSIZE="), ParsedText))
+		else if (FParse::Value(Str, TEXT("ELEMENTSIZE="), ParsedText))
 		{
 			/** Serialized flags for bulk data */
 			ElementSize = FCString::Atoi(*ParsedText);
 		}
-		else
-		if (FParse::Value(Str, TEXT("BEGIN "), ParsedText) && (ParsedText.ToUpper() == TEXT("BINARYBLOB")))
+		else if (FParse::Value(Str, TEXT("BEGIN "), ParsedText) && (ParsedText.ToUpper() == TEXT("BINARYBLOB")))
 		{
-			uint8* RawData = NULL;
+			uint8* RawData = nullptr;
+
 			/** The bulk data... */
 			while(FParse::Line(&Buffer,StrLine))
 			{
@@ -399,8 +452,7 @@ bool UFactory::ImportUntypedBulkDataFromText(const TCHAR*& Buffer, FUntypedBulkD
 					RawData = (uint8*)RawBulkData;
 					bBulkDataIsLocked = true;
 				}
-				else
-				if (FParse::Value(Str, TEXT("BEGIN "), ParsedText) && (ParsedText.ToUpper() == TEXT("BINARY")))
+				else if (FParse::Value(Str, TEXT("BEGIN "), ParsedText) && (ParsedText.ToUpper() == TEXT("BINARY")))
 				{
 					uint8* BulkDataPointer = RawData;
 					while(FParse::Line(&Buffer,StrLine))
@@ -435,8 +487,7 @@ bool UFactory::ImportUntypedBulkDataFromText(const TCHAR*& Buffer, FUntypedBulkD
 						}
 					}
 				}
-				else
-				if (FParse::Value(Str, TEXT("END "), ParsedText) && (ParsedText.ToUpper() == TEXT("BINARYBLOB")))
+				else if (FParse::Value(Str, TEXT("END "), ParsedText) && (ParsedText.ToUpper() == TEXT("BINARYBLOB")))
 				{
 					BulkData.Unlock();
 					bBulkDataIsLocked = false;
@@ -444,8 +495,7 @@ bool UFactory::ImportUntypedBulkDataFromText(const TCHAR*& Buffer, FUntypedBulkD
 				}
 			}
 		}
-		else
-		if (FParse::Value(Str, TEXT("END "), ParsedText) && (ParsedText.ToUpper() == TEXT("UNTYPEDBULKDATA")))
+		else if (FParse::Value(Str, TEXT("END "), ParsedText) && (ParsedText.ToUpper() == TEXT("UNTYPEDBULKDATA")))
 		{
 			break;
 		}
@@ -459,60 +509,60 @@ bool UFactory::ImportUntypedBulkDataFromText(const TCHAR*& Buffer, FUntypedBulkD
 	return true;
 }
 
+
 UObject* UFactory::CreateOrOverwriteAsset(UClass* InClass, UObject* InParent, FName InName, EObjectFlags InFlags, UObject* InTemplate) const
 {
-	// Creates an asset if it doesn't exist.
-	UObject* ExistingAsset = StaticFindObject(NULL, InParent, *InName.ToString());
-	if ( !ExistingAsset )
+	// create an asset if it doesn't exist
+	UObject* ExistingAsset = StaticFindObject(nullptr, InParent, *InName.ToString());
+	
+	if (!ExistingAsset)
 	{
 		return NewObject<UObject>(InParent, InClass, InName, InFlags, InTemplate);
 	}
 
-	// If it does exist then it overwrites it if possible.
-	if ( ExistingAsset->GetClass()->IsChildOf(InClass) )
+	// overwrite existing asset, if possible
+	if (ExistingAsset->GetClass()->IsChildOf(InClass))
 	{
 		return NewObject<UObject>(InParent, InClass, InName, InFlags, InTemplate);
 	}
 	
-	// If it can not overwrite then it will delete and replace.
-	if ( ObjectTools::DeleteSingleObject( ExistingAsset ) )
+	// otherwise delete and replace
+	if (!ObjectTools::DeleteSingleObject(ExistingAsset))
 	{
-		// Keep InPackage alive through the GC, in case ExistingAsset was the only reason it was around.
-		const bool bRootedPackage = InParent->IsRooted();
-		if ( !bRootedPackage )
-		{
-			InParent->AddToRoot();
-		}
+		return nullptr;
+	}
 
-		// Force GC so we can cleanly create a new asset (and not do an 'in place' replacement)
-		CollectGarbage( GARBAGE_COLLECTION_KEEPFLAGS );
+	// keep InPackage alive through the GC, in case ExistingAsset was the only reason it was around
+	const bool bRootedPackage = InParent->IsRooted();
 
-		if ( !bRootedPackage )
-		{
-			InParent->RemoveFromRoot();
-		}
+	if (!bRootedPackage)
+	{
+		InParent->AddToRoot();
+	}
 
-		// Try to find the existing asset again now that the GC has occurred
-		ExistingAsset = StaticFindObject(NULL, InParent, *InName.ToString());
-		if ( ExistingAsset )
-		{
-			// Even after the delete and GC, the object is still around. Fail this operation.
-			return NULL;
-		}
-		else
-		{
-			// We can now create the asset in the package
-			return NewObject<UObject>(InParent, InClass, InName, InFlags, InTemplate);
-		}
+	// force GC so we can cleanly create a new asset (and not do an 'in place' replacement)
+	CollectGarbage(GARBAGE_COLLECTION_KEEPFLAGS);
+
+	if (!bRootedPackage)
+	{
+		InParent->RemoveFromRoot();
 	}
-	else
+
+	// try to find the existing asset again now that the GC has occurred
+	ExistingAsset = StaticFindObject(nullptr, InParent, *InName.ToString());
+
+	// if the object is still around after GC, fail this operation
+	if (ExistingAsset)
 	{
-		// The delete did not succeed. There are still references to the old content.
-		return NULL;
+		return nullptr;
 	}
+
+	// create the asset in the package
+	return NewObject<UObject>(InParent, InClass, InName, InFlags, InTemplate);
 }
 
+
 FString UFactory::GetDefaultNewAssetName() const
 {
 	return FString(TEXT("New")) + GetSupportedClass()->GetName();
-}
\ No newline at end of file
+}
diff --git a/Engine/Source/Editor/UnrealEd/Private/Factories/TTFontImport.cpp b/Engine/Source/Editor/UnrealEd/Private/Factories/TTFontImport.cpp
--- a/Engine/Source/Editor/UnrealEd/Private/Factories/TTFontImport.cpp
+++ b/Engine/Source/Editor/UnrealEd/Private/Factories/TTFontImport.cpp
@@ -211,19 +211,23 @@ void UTrueTypeFontFactory::SetReimportPaths( UObject* Obj, const TArray<FString>
 EReimportResult::Type UTrueTypeFontFactory::Reimport( UObject* InObject )
 {
 	UFont* FontToReimport = Cast<UFont>(InObject);
-	bool bSuccess = false;
-	if ( FontToReimport != NULL )
+	
+	if (FontToReimport == nullptr)
 	{
-		SetupFontImportOptions();
-		this->ImportOptions->Data = FontToReimport->ImportOptions;
+		return EReimportResult::Failed;
+	}
 
-		if (NULL != UFactory::StaticImportObject(InObject->GetClass(), InObject->GetOuter(), *InObject->GetName(), RF_Public|RF_Standalone, TEXT(""), NULL, this))
-		{
-			bSuccess = true;
-		}
+	SetupFontImportOptions();
+	this->ImportOptions->Data = FontToReimport->ImportOptions;
+
+	bool OutCanceled = false;
+
+	if (ImportObject(InObject->GetClass(), InObject->GetOuter(), *InObject->GetName(), RF_Public | RF_Standalone, TEXT(""), nullptr, OutCanceled) != nullptr)
+	{
+		return EReimportResult::Succeeded;
 	}
 
-	return bSuccess ? EReimportResult::Succeeded : EReimportResult::Failed;
+	return EReimportResult::Failed;
 }
 
 int32 UTrueTypeFontFactory::GetPriority() const
diff --git a/Engine/Source/Editor/UnrealEd/Private/Factories/VectorFieldFactory.cpp b/Engine/Source/Editor/UnrealEd/Private/Factories/VectorFieldFactory.cpp
--- a/Engine/Source/Editor/UnrealEd/Private/Factories/VectorFieldFactory.cpp
+++ b/Engine/Source/Editor/UnrealEd/Private/Factories/VectorFieldFactory.cpp
@@ -322,15 +322,21 @@ EReimportResult::Type UReimportVectorFieldStaticFactory::Reimport( UObject* Obj
 		return EReimportResult::Failed;
 	}
 
-	if (UFactory::StaticImportObject(VectorFieldStatic->GetClass(), VectorFieldStatic->GetOuter(), *VectorFieldStatic->GetName(), RF_Public|RF_Standalone, *ReImportFilename, NULL, this))
+	bool OutCanceled = false;
+
+	if (ImportObject(VectorFieldStatic->GetClass(), VectorFieldStatic->GetOuter(), *VectorFieldStatic->GetName(), RF_Public|RF_Standalone, ReImportFilename, nullptr, OutCanceled) != nullptr)
 	{
 		UE_LOG(LogVectorFieldFactory, Log, TEXT("Reimported successfully") );
 		VectorFieldStatic->AssetImportData->Update(ReImportFilename);
 		VectorFieldStatic->MarkPackageDirty();
 	}
+	else if (OutCanceled)
+	{
+		UE_LOG(LogVectorFieldFactory, Warning, TEXT("-- import canceled"));
+	}
 	else
 	{
-		UE_LOG(LogVectorFieldFactory, Warning, TEXT("Reimport failed") );
+		UE_LOG(LogVectorFieldFactory, Warning, TEXT("-- import failed"));
 	}
 
 	return EReimportResult::Succeeded;
diff --git a/Engine/Source/Editor/UnrealEd/Private/Fbx/FbxSkeletalMeshImport.cpp b/Engine/Source/Editor/UnrealEd/Private/Fbx/FbxSkeletalMeshImport.cpp
--- a/Engine/Source/Editor/UnrealEd/Private/Fbx/FbxSkeletalMeshImport.cpp
+++ b/Engine/Source/Editor/UnrealEd/Private/Fbx/FbxSkeletalMeshImport.cpp
@@ -1367,7 +1367,7 @@ USkeletalMesh* UnFbx::FFbxImporter::ImportSkeletalMesh(UObject* InParent, TArray
 
 		// Store the current file path and timestamp for re-import purposes
 		UFbxSkeletalMeshImportData* ImportData = UFbxSkeletalMeshImportData::GetImportDataForSkeletalMesh(SkeletalMesh, TemplateImportData);
-		SkeletalMesh->AssetImportData->Update(UFactory::CurrentFilename);
+		SkeletalMesh->AssetImportData->Update(UFactory::GetCurrentFilename());
 
 		SkeletalMesh->CalculateInvRefMatrices();
 		SkeletalMesh->PostEditChange();
diff --git a/Engine/Source/Editor/UnrealEd/Private/Fbx/FbxStaticMeshImport.cpp b/Engine/Source/Editor/UnrealEd/Private/Fbx/FbxStaticMeshImport.cpp
--- a/Engine/Source/Editor/UnrealEd/Private/Fbx/FbxStaticMeshImport.cpp
+++ b/Engine/Source/Editor/UnrealEd/Private/Fbx/FbxStaticMeshImport.cpp
@@ -1295,7 +1295,7 @@ UStaticMesh* UnFbx::FFbxImporter::ImportStaticMeshAsSingle(UObject* InParent, TA
 		}
 
 		UFbxStaticMeshImportData* ImportData = UFbxStaticMeshImportData::GetImportDataForStaticMesh(StaticMesh, TemplateImportData);
-		ImportData->Update(UFactory::CurrentFilename);
+		ImportData->Update(UFactory::GetCurrentFilename());
 
 		// @todo This overrides restored values currently but we need to be able to import over the existing settings if the user chose to do so.
 		SrcModel.BuildSettings.bRemoveDegenerates = ImportOptions->bRemoveDegenerates;
diff --git a/Engine/Source/Editor/UnrealEd/Private/Kismet2/DebuggerCommands.cpp b/Engine/Source/Editor/UnrealEd/Private/Kismet2/DebuggerCommands.cpp
--- a/Engine/Source/Editor/UnrealEd/Private/Kismet2/DebuggerCommands.cpp
+++ b/Engine/Source/Editor/UnrealEd/Private/Kismet2/DebuggerCommands.cpp
@@ -1964,15 +1964,25 @@ bool FInternalPlayWorldCommandCallbacks::CanLaunchOnDevice(const FString& Device
 {
 	if (!GUnrealEd->IsPlayingViaLauncher())
 	{
-		TSharedPtr<ITargetDeviceServicesModule> TargetDeviceServicesModule = StaticCastSharedPtr<ITargetDeviceServicesModule>(FModuleManager::Get().LoadModule(TEXT("TargetDeviceServices")));
+		static TWeakPtr<ITargetDeviceProxyManager> DeviceProxyManagerPtr;
 
-		if (TargetDeviceServicesModule.IsValid())
+		if (!DeviceProxyManagerPtr.IsValid())
 		{
-			ITargetDeviceProxyPtr DeviceProxy = TargetDeviceServicesModule->GetDeviceProxyManager()->FindProxy(DeviceName);
+			auto TargetDeviceServicesModule = StaticCastSharedPtr<ITargetDeviceServicesModule>(FModuleManager::Get().LoadModule(TEXT("TargetDeviceServices")));		
+			if (TargetDeviceServicesModule.IsValid())
+			{
+				DeviceProxyManagerPtr = TargetDeviceServicesModule->GetDeviceProxyManager();
+			}
+		}
 
+		TSharedPtr<ITargetDeviceProxyManager> DeviceProxyManager = DeviceProxyManagerPtr.Pin();
+		if (DeviceProxyManager.IsValid())
+		{
+			ITargetDeviceProxyPtr DeviceProxy = DeviceProxyManager->FindProxy(DeviceName);
 			return (DeviceProxy.IsValid() && DeviceProxy->IsConnected());
 		}
 	}
+
 	return false;
 }
 
diff --git a/Engine/Source/Editor/UnrealEd/Private/LevelEditorViewport.cpp b/Engine/Source/Editor/UnrealEd/Private/LevelEditorViewport.cpp
--- a/Engine/Source/Editor/UnrealEd/Private/LevelEditorViewport.cpp
+++ b/Engine/Source/Editor/UnrealEd/Private/LevelEditorViewport.cpp
@@ -2093,6 +2093,7 @@ void FLevelEditorViewportClient::UpdateViewForLockedActor()
 				{
 					bUseControllingActorViewInfo = true;
 					CameraComponent->GetCameraView(0.0f, ControllingActorViewInfo);
+					CameraComponent->GetExtraPostProcessBlends(ControllingActorExtraPostProcessBlends, ControllingActorExtraPostProcessBlendWeights);
 
 					// Post processing is handled by OverridePostProcessingSettings
 					ViewFOV = ControllingActorViewInfo.FOV;
diff --git a/Engine/Source/Editor/UnrealEd/Private/RichCurveEditorCommands.cpp b/Engine/Source/Editor/UnrealEd/Private/RichCurveEditorCommands.cpp
--- a/Engine/Source/Editor/UnrealEd/Private/RichCurveEditorCommands.cpp
+++ b/Engine/Source/Editor/UnrealEd/Private/RichCurveEditorCommands.cpp
@@ -12,7 +12,8 @@ void FRichCurveEditorCommands::RegisterCommands()
 	UI_COMMAND(ZoomToFit, "Fit", "Zoom to Fit", EUserInterfaceActionType::Button, FInputChord(EKeys::F));
 	UI_COMMAND(ZoomToFitAll, "FitAll", "Zoom to Fit All", EUserInterfaceActionType::Button, FInputChord(EKeys::A));
 
-	UI_COMMAND(ToggleSnapping, "Snapping", "Toggle Snapping", EUserInterfaceActionType::ToggleButton, FInputChord());
+	UI_COMMAND(ToggleInputSnapping, "Input Snapping", "Toggle Input Snapping", EUserInterfaceActionType::ToggleButton, FInputChord());
+	UI_COMMAND(ToggleOutputSnapping, "Output Snapping", "Toggle Output Snapping", EUserInterfaceActionType::ToggleButton, FInputChord());
 
 	UI_COMMAND(InterpolationConstant, "Constant", "Constant interpolation", EUserInterfaceActionType::ToggleButton, FInputChord(EKeys::Five));
 	UI_COMMAND(InterpolationLinear, "Linear", "Linear interpolation", EUserInterfaceActionType::ToggleButton, FInputChord(EKeys::Four));
diff --git a/Engine/Source/Editor/UnrealEd/Private/SCurveEditor.cpp b/Engine/Source/Editor/UnrealEd/Private/SCurveEditor.cpp
--- a/Engine/Source/Editor/UnrealEd/Private/SCurveEditor.cpp
+++ b/Engine/Source/Editor/UnrealEd/Private/SCurveEditor.cpp
@@ -44,7 +44,8 @@ void SCurveEditor::Construct(const FArguments& InArgs)
 
 	InputSnap = InArgs._InputSnap;
 	OutputSnap = InArgs._OutputSnap;
-	bSnappingEnabled = InArgs._SnappingEnabled;
+	bInputSnappingEnabled = InArgs._InputSnappingEnabled;
+	bOutputSnappingEnabled = InArgs._OutputSnappingEnabled;
 
 	bZoomToFitVertical = InArgs._ZoomToFitVertical;
 	bZoomToFitHorizontal = InArgs._ZoomToFitHorizontal;
@@ -104,10 +105,15 @@ void SCurveEditor::Construct(const FArguments& InArgs)
 	Commands->MapAction(FRichCurveEditorCommands::Get().ZoomToFitAll,
 		FExecuteAction::CreateSP(this, &SCurveEditor::ZoomToFit, true));
 
-	Commands->MapAction(FRichCurveEditorCommands::Get().ToggleSnapping,
-		FExecuteAction::CreateSP(this, &SCurveEditor::ToggleSnapping),
+	Commands->MapAction(FRichCurveEditorCommands::Get().ToggleInputSnapping,
+		FExecuteAction::CreateSP(this, &SCurveEditor::ToggleInputSnapping),
 		FCanExecuteAction(),
-		FIsActionChecked::CreateSP(this, &SCurveEditor::IsSnappingEnabled));
+		FIsActionChecked::CreateSP(this, &SCurveEditor::IsInputSnappingEnabled));
+
+	Commands->MapAction(FRichCurveEditorCommands::Get().ToggleOutputSnapping,
+		FExecuteAction::CreateSP(this, &SCurveEditor::ToggleOutputSnapping),
+		FCanExecuteAction(),
+		FIsActionChecked::CreateSP(this, &SCurveEditor::IsOutputSnappingEnabled));
 
 	// Interpolation
 	Commands->MapAction(FRichCurveEditorCommands::Get().InterpolationConstant,
@@ -1159,12 +1165,12 @@ void SCurveEditor::PaintMarquee(const FGeometry& AllottedGeometry, const FSlateR
 
 float SCurveEditor::GetInputNumericEntryBoxDelta() const
 {
-	return bSnappingEnabled.Get() ? InputSnap.Get() : 0;
+	return bInputSnappingEnabled.Get() ? InputSnap.Get() : 0;
 }
 
 float SCurveEditor::GetOutputNumericEntryBoxDelta() const
 {
-	return bSnappingEnabled.Get() ? OutputSnap.Get() : 0;
+	return bOutputSnappingEnabled.Get() ? OutputSnap.Get() : 0;
 }
 
 void SCurveEditor::SetCurveOwner(FCurveOwnerInterface* InCurveOwner, bool bCanEdit)
@@ -2515,17 +2521,30 @@ void SCurveEditor::ZoomToFit(const bool bZoomToFitAll)
 	ZoomToFitVertical(bZoomToFitAll);
 }
 
-void SCurveEditor::ToggleSnapping()
+void SCurveEditor::ToggleInputSnapping()
+{
+	if (bInputSnappingEnabled.IsBound() == false)
+	{
+		bInputSnappingEnabled = !bInputSnappingEnabled.Get();
+	}
+}
+
+void SCurveEditor::ToggleOutputSnapping()
 {
-	if (bSnappingEnabled.IsBound() == false)
+	if (bOutputSnappingEnabled.IsBound() == false)
 	{
-		bSnappingEnabled = !bSnappingEnabled.Get();
+		bOutputSnappingEnabled = !bOutputSnappingEnabled.Get();
 	}
 }
 
-bool SCurveEditor::IsSnappingEnabled()
+bool SCurveEditor::IsInputSnappingEnabled()
+{
+	return bInputSnappingEnabled.Get();
+}
+
+bool SCurveEditor::IsOutputSnappingEnabled()
 {
-	return bSnappingEnabled.Get();
+	return bOutputSnappingEnabled.Get();
 }
 
 void SCurveEditor::CreateContextMenu(const FGeometry& InMyGeometry, const FPointerEvent& InMouseEvent)
@@ -3533,12 +3552,17 @@ bool SCurveEditor::IsLinearColorCurve() const
 
 FVector2D SCurveEditor::SnapLocation(FVector2D InLocation)
 {
-	if (bSnappingEnabled.Get())
+	if (bInputSnappingEnabled.Get())
 	{
 		const float InputSnapNow = InputSnap.Get();
-		const float OutputSnapNow = OutputSnap.Get();
 
 		InLocation.X = InputSnapNow != 0 ? FMath::RoundToInt(InLocation.X / InputSnapNow) * InputSnapNow : InLocation.X;
+	}
+
+	if (bOutputSnappingEnabled.Get())
+	{
+		const float OutputSnapNow = OutputSnap.Get();
+
 		InLocation.Y = OutputSnapNow != 0 ? FMath::RoundToInt(InLocation.Y / OutputSnapNow) * OutputSnapNow : InLocation.Y;
 	}
 	return InLocation;
diff --git a/Engine/Source/Editor/UnrealEd/Private/SkeletalMeshEdit.cpp b/Engine/Source/Editor/UnrealEd/Private/SkeletalMeshEdit.cpp
--- a/Engine/Source/Editor/UnrealEd/Private/SkeletalMeshEdit.cpp
+++ b/Engine/Source/Editor/UnrealEd/Private/SkeletalMeshEdit.cpp
@@ -77,7 +77,7 @@ UAnimSequence * UEditorEngine::ImportFbxAnimation( USkeleton* Skeleton, UObject*
 			{
 				// since to know full path, reimport will need to do same
 				UFbxAnimSequenceImportData* ImportData = UFbxAnimSequenceImportData::GetImportDataForAnimSequence(NewAnimation, TemplateImportData);
-				ImportData->Update(UFactory::CurrentFilename);
+				ImportData->Update(UFactory::GetCurrentFilename());
 			}
 		}
 	}
@@ -547,7 +547,7 @@ UAnimSequence * UnFbx::FFbxImporter::ImportAnimations(USkeleton* Skeleton, UObje
 
 		// since to know full path, reimport will need to do same
 		UFbxAnimSequenceImportData* ImportData = UFbxAnimSequenceImportData::GetImportDataForAnimSequence(DestSeq, TemplateImportData);
-		ImportData->Update(UFactory::CurrentFilename);
+		ImportData->Update(UFactory::GetCurrentFilename());
 
 		ImportAnimation(Skeleton, DestSeq, Name, SortedLinks, NodeArray, CurAnimStack, ResampleRate, AnimTimeSpan);
 
diff --git a/Engine/Source/Editor/UnrealEd/Private/Tests/AutomationEditorCommon.cpp b/Engine/Source/Editor/UnrealEd/Private/Tests/AutomationEditorCommon.cpp
--- a/Engine/Source/Editor/UnrealEd/Private/Tests/AutomationEditorCommon.cpp
+++ b/Engine/Source/Editor/UnrealEd/Private/Tests/AutomationEditorCommon.cpp
@@ -96,12 +96,17 @@ namespace AutomationEditorCommonUtils
 				UTextureFactory::SuppressImportOverwriteDialog();
 			}
 
-			ImportedAsset = UFactory::StaticImportObject(ImportAssetType, Pkg, FName(*ObjectName), RF_Public | RF_Standalone, bDummy, *ImportPath, NULL, ImportFactory, NULL, GWarn, 0);
+			bool OutCanceled = false;
+			ImportedAsset = ImportFactory->ImportObject(ImportAssetType, Pkg, FName(*ObjectName), RF_Public | RF_Standalone, ImportPath, nullptr, OutCanceled);
 
-			if (ImportedAsset)
+			if (ImportedAsset != nullptr)
 			{
 				UE_LOG(LogAutomationEditorCommon, Display, TEXT("Imported %s"), *ImportPath);
 			}
+			else if (OutCanceled)
+			{
+				UE_LOG(LogAutomationEditorCommon, Display, TEXT("Canceled import of %s"), *ImportPath);
+			}
 			else
 			{
 				UE_LOG(LogAutomationEditorCommon, Error, TEXT("Failed to import asset using factory %s!"), *ImportFactory->GetName());
diff --git a/Engine/Source/Editor/UnrealEd/Public/EditorModeRegistry.cpp b/Engine/Source/Editor/UnrealEd/Public/EditorModeRegistry.cpp
--- a/Engine/Source/Editor/UnrealEd/Public/EditorModeRegistry.cpp
+++ b/Engine/Source/Editor/UnrealEd/Public/EditorModeRegistry.cpp
@@ -9,6 +9,7 @@
 #include "Editor/MeshPaint/Public/MeshPaintModule.h"
 #include "Editor/GeometryMode/Public/GeometryEdMode.h"
 #include "Editor/ActorPickerMode/Public/ActorPickerMode.h"
+#include "Editor/SceneDepthPickerMode/Public/SceneDepthPickerMode.h"
 #include "Editor/TextureAlignMode/Public/TextureAlignEdMode.h"
 #include "Editor/FoliageEdit/Public/FoliageEditModule.h"
 
@@ -55,6 +56,7 @@ void FEditorModeRegistry::Initialize()
 	FModuleManager::LoadModuleChecked<FTextureAlignModeModule>(TEXT("TextureAlignMode"));
 	FModuleManager::LoadModuleChecked<FGeometryModeModule>(TEXT("GeometryMode"));
 	FModuleManager::LoadModuleChecked<FActorPickerModeModule>(TEXT("ActorPickerMode"));
+	FModuleManager::LoadModuleChecked<FSceneDepthPickerModeModule>(TEXT("SceneDepthPickerMode"));
 	FModuleManager::LoadModuleChecked<IMeshPaintModule>(TEXT("MeshPaint"));
 	FModuleManager::LoadModuleChecked<ILandscapeEditorModule>(TEXT("LandscapeEditor"));
 	FModuleManager::LoadModuleChecked<IFoliageEditModule>(TEXT("FoliageEdit"));
diff --git a/Engine/Source/Editor/UnrealEd/Public/EditorModes.h b/Engine/Source/Editor/UnrealEd/Public/EditorModes.h
--- a/Engine/Source/Editor/UnrealEd/Public/EditorModes.h
+++ b/Engine/Source/Editor/UnrealEd/Public/EditorModes.h
@@ -75,6 +75,9 @@ public:
 	/** Actor picker mode, used to interactively pick actors in the viewport */
 	const static FEditorModeID EM_ActorPicker;
 
+	/** Actor picker mode, used to interactively pick actors in the viewport */
+	const static FEditorModeID EM_SceneDepthPicker;
+
 private:
 	FBuiltinEditorModes() {}
 };
diff --git a/Engine/Source/Editor/UnrealEd/Public/EditorViewportClient.h b/Engine/Source/Editor/UnrealEd/Public/EditorViewportClient.h
--- a/Engine/Source/Editor/UnrealEd/Public/EditorViewportClient.h
+++ b/Engine/Source/Editor/UnrealEd/Public/EditorViewportClient.h
@@ -1007,6 +1007,11 @@ public:
 
 	/** When collision draw mode changes, this function allows hidden objects to be drawn, so hidden colliding objects can be seen */
 	void UpdateHiddenCollisionDrawing();
+	/** Returns the scene depth at the given viewport X,Y */
+	float GetSceneDepthAtLocation(int32 X, int32 Y);
+
+	/** Returns the location of the object at the given viewport X,Y */
+	FVector GetHitProxyObjectLocation(int32 X, int32 Y);
 
 protected:
 	/** Invalidates the viewport widget (if valid) to register its active timer */
@@ -1350,6 +1355,8 @@ protected:
 
 	bool bUseControllingActorViewInfo;
 	FMinimalViewInfo ControllingActorViewInfo;
+	TArray<FPostProcessSettings> ControllingActorExtraPostProcessBlends;
+	TArray<float> ControllingActorExtraPostProcessBlendWeights;
 
 	/* Updated on each mouse drag start */
 	uint32 LastMouseX;
diff --git a/Engine/Source/Editor/UnrealEd/Public/RichCurveEditorCommands.h b/Engine/Source/Editor/UnrealEd/Public/RichCurveEditorCommands.h
--- a/Engine/Source/Editor/UnrealEd/Public/RichCurveEditorCommands.h
+++ b/Engine/Source/Editor/UnrealEd/Public/RichCurveEditorCommands.h
@@ -23,7 +23,8 @@ public:
 	TSharedPtr<FUICommandInfo> ZoomToFitVertical;
 	TSharedPtr<FUICommandInfo> ZoomToFit;
 	TSharedPtr<FUICommandInfo> ZoomToFitAll;
-	TSharedPtr<FUICommandInfo> ToggleSnapping;
+	TSharedPtr<FUICommandInfo> ToggleInputSnapping;
+	TSharedPtr<FUICommandInfo> ToggleOutputSnapping;
 
 	TSharedPtr<FUICommandInfo> InterpolationConstant;
 	TSharedPtr<FUICommandInfo> InterpolationLinear;
diff --git a/Engine/Source/Editor/UnrealEd/Public/SCurveEditor.h b/Engine/Source/Editor/UnrealEd/Public/SCurveEditor.h
--- a/Engine/Source/Editor/UnrealEd/Public/SCurveEditor.h
+++ b/Engine/Source/Editor/UnrealEd/Public/SCurveEditor.h
@@ -151,7 +151,8 @@ public:
 		, _ViewMaxOutput(1.0f)
 		, _InputSnap(0.1f)
 		, _OutputSnap(0.05f)
-		, _SnappingEnabled(false)
+		, _InputSnappingEnabled(false)
+		, _OutputSnappingEnabled(false)
 		, _TimelineLength(5.0f)
 		, _DesiredSize(FVector2D::ZeroVector)
 		, _DrawCurve(true)
@@ -177,7 +178,8 @@ public:
 		SLATE_ATTRIBUTE( float, ViewMaxOutput )
 		SLATE_ATTRIBUTE( float, InputSnap )
 		SLATE_ATTRIBUTE( float, OutputSnap )
-		SLATE_ATTRIBUTE( bool, SnappingEnabled )
+		SLATE_ATTRIBUTE( bool, InputSnappingEnabled )
+		SLATE_ATTRIBUTE( bool, OutputSnappingEnabled )
 		SLATE_ATTRIBUTE( float, TimelineLength )
 		SLATE_ATTRIBUTE( FVector2D, DesiredSize )
 		SLATE_ARGUMENT( bool, DrawCurve )
@@ -383,8 +385,10 @@ private:
 	FReply ZoomToFitHorizontalClicked();
 	FReply ZoomToFitVerticalClicked();
 
-	void ToggleSnapping();
-	bool IsSnappingEnabled();
+	void ToggleInputSnapping();
+	void ToggleOutputSnapping();
+	bool IsInputSnappingEnabled();
+	bool IsOutputSnappingEnabled();
 
 	TOptional<float> OnGetTime() const;
 	void OnTimeComitted(float NewValue, ETextCommit::Type CommitType);
@@ -689,8 +693,11 @@ protected:
 	/** The snapping value for the output domain. */
 	TAttribute<float> OutputSnap;
 
-	/** Whether or not snapping is enabled. */
-	TAttribute<bool> bSnappingEnabled;
+	/** Whether or not input snapping is enabled. */
+	TAttribute<bool> bInputSnappingEnabled;
+
+	/** Whether or not output snapping is enabled. */
+	TAttribute<bool> bOutputSnappingEnabled;
 
 	/** True if you want the curve editor to fit to zoom **/
 	bool bZoomToFitVertical;
diff --git a/Engine/Source/Editor/UnrealEd/UnrealEd.Build.cs b/Engine/Source/Editor/UnrealEd/UnrealEd.Build.cs
--- a/Engine/Source/Editor/UnrealEd/UnrealEd.Build.cs
+++ b/Engine/Source/Editor/UnrealEd/UnrealEd.Build.cs
@@ -225,12 +225,6 @@ public class UnrealEd : ModuleRules
 			DynamicallyLoadedModuleNames.Add("AndroidPlatformEditor");
 		}
 
-        if (Target.Type == TargetRules.TargetType.Editor)
-        {
-            DynamicallyLoadedModuleNames.Add("SequenceRecorder");
-            PrivateIncludePathModuleNames.Add("SequenceRecorder");
-        }
-
 		CircularlyReferencedDependentModules.AddRange(
 			new string[] 
 			{
diff --git a/Engine/Source/Runtime/AnimGraphRuntime/Public/BoneControllers/AnimNode_AnimDynamics.h b/Engine/Source/Runtime/AnimGraphRuntime/Public/BoneControllers/AnimNode_AnimDynamics.h
--- a/Engine/Source/Runtime/AnimGraphRuntime/Public/BoneControllers/AnimNode_AnimDynamics.h
+++ b/Engine/Source/Runtime/AnimGraphRuntime/Public/BoneControllers/AnimNode_AnimDynamics.h
@@ -284,6 +284,8 @@ struct ANIMGRAPHRUNTIME_API FAnimNode_AnimDynamics : public FAnimNode_SkeletalCo
 	virtual void GatherDebugData(FNodeDebugData& DebugData) override;
 	virtual bool HasPreUpdate() const override { return true; }
 	virtual void PreUpdate(const UAnimInstance* InAnimInstance) override;
+	virtual bool NeedsDynamicReset() const { return true; }
+	virtual void ResetDynamics() { RequestInitialise(); }
 	// End of FAnimNode_SkeletalControlBase interface
 
 	void RequestInitialise() { bRequiresInit = true; }
diff --git a/Engine/Source/Runtime/CinematicCamera/Private/CameraRig_Crane.cpp b/Engine/Source/Runtime/CinematicCamera/Private/CameraRig_Crane.cpp
--- a/Engine/Source/Runtime/CinematicCamera/Private/CameraRig_Crane.cpp
+++ b/Engine/Source/Runtime/CinematicCamera/Private/CameraRig_Crane.cpp
@@ -10,16 +10,16 @@ ACameraRig_Crane::ACameraRig_Crane(const FObjectInitializer& ObjectInitializer)
 {
 	PrimaryActorTick.bCanEverTick = true;
 
+	// default control values
 	CraneYaw = 0.f;
 	CranePitch = 0.f;
 	CraneArmLength = 500.f;
 
+	// create the root component
 	TransformComponent = CreateDefaultSubobject<USceneComponent>(TEXT("TransformComponent"));
-
-	// Make the scene component the root component
 	RootComponent = TransformComponent;
 	
-	// Setup camera defaults
+	// create the functional component hierarchy
 	CraneYawControl = CreateDefaultSubobject<USceneComponent>(TEXT("CraneYawControl"));
 	CraneYawControl->SetupAttachment(TransformComponent);
 	CraneYawControl->RelativeLocation = FVector(0.f, 0.f, 70.f);			// pivot height off the ground
@@ -33,7 +33,7 @@ ACameraRig_Crane::ACameraRig_Crane(const FObjectInitializer& ObjectInitializer)
 	CraneCameraMount->SetupAttachment(CranePitchControl);
 	CraneCameraMount->RelativeLocation = FVector(CraneArmLength, 0.f, -15.f);			// negative z == underslung mount
 
-	// preview meshes
+	// create preview meshes
 	if (!IsRunningCommandlet() && !IsRunningDedicatedServer())
 	{
 		static ConstructorHelpers::FObjectFinder<UStaticMesh> CraneBaseMesh(TEXT("/Engine/EditorMeshes/Camera/SM_CraneRig_Base.SM_CraneRig_Base"));
@@ -81,9 +81,9 @@ ACameraRig_Crane::ACameraRig_Crane(const FObjectInitializer& ObjectInitializer)
 			PreviewMesh_CraneMount->CastShadow = false;
 			PreviewMesh_CraneMount->PostPhysicsComponentTick.bCanEverTick = false;
 
-			PreviewMesh_CraneMount->SetupAttachment(CraneCameraMount);
+			PreviewMesh_CraneMount->SetupAttachment(CranePitchControl);
 			PreviewMesh_CraneMount->RelativeRotation = FRotator(0.f, 90.f, 0.f);
-			PreviewMesh_CraneMount->RelativeLocation = FVector(0.f, 0.f, 57.f);
+			PreviewMesh_CraneMount->RelativeLocation = FVector(CraneArmLength, 0.f, 0.f);
 			PreviewMesh_CraneMount->RelativeScale3D = FVector(0.7f, 0.7f, 0.7f);
 		}
 		
@@ -108,7 +108,6 @@ ACameraRig_Crane::ACameraRig_Crane(const FObjectInitializer& ObjectInitializer)
 	}
 }
 
-static const float CraneArmMesh_CounterweightOverhang = 0.f;
 static const float CraneArmMesh_DefaultMeshSize = 29.f * 0.7f;		// size of the mesh in the dimension that will stretch (accounting for the 0.7 scale)
 
 void ACameraRig_Crane::UpdatePreviewMeshes()
@@ -117,15 +116,20 @@ void ACameraRig_Crane::UpdatePreviewMeshes()
 	{
 		// note to explain the math here:
 		// this tailored to the dimensions of the asset (29cm in scaling direction, pivot at the edge)
-
-		float const TotalCraneArmMeshLength = CraneArmMesh_CounterweightOverhang + CraneArmLength;
-		float const CraneArmYScale = TotalCraneArmMeshLength / CraneArmMesh_DefaultMeshSize;
+		float const CraneArmYScale = CraneArmLength / CraneArmMesh_DefaultMeshSize;
 
 		FVector NewRelScale3D = PreviewMesh_CraneArm->RelativeScale3D;
 		NewRelScale3D.Y = CraneArmYScale * 0.7f;
 
 		PreviewMesh_CraneArm->SetRelativeScale3D(NewRelScale3D);
 	}
+
+	if (PreviewMesh_CraneMount)
+	{
+		const float ZOffset = 40.f;		// make it appear in the right place
+		FVector const NewLoc(CraneArmLength, 0.f, ZOffset);
+		PreviewMesh_CraneMount->SetRelativeLocation(NewLoc);
+	}
 }
 
 void ACameraRig_Crane::UpdateCraneComponents()
diff --git a/Engine/Source/Runtime/CinematicCamera/Private/CameraRig_Rail.cpp b/Engine/Source/Runtime/CinematicCamera/Private/CameraRig_Rail.cpp
--- a/Engine/Source/Runtime/CinematicCamera/Private/CameraRig_Rail.cpp
+++ b/Engine/Source/Runtime/CinematicCamera/Private/CameraRig_Rail.cpp
@@ -72,48 +72,60 @@ USplineMeshComponent* ACameraRig_Rail::CreateSplinePreviewSegment()
 
 void ACameraRig_Rail::UpdatePreviewMeshes()
 {
-	if (PreviewRailStaticMesh)
+	if (RailSplineComponent)
 	{
-		int32 const NumSplinePoints = RailSplineComponent->GetNumberOfSplinePoints();
-		int32 const NumNeededPreviewMeshes = NumSplinePoints - 1;
-
-		// make sure our preview mesh array is correctly sized and populated
+		if (PreviewRailStaticMesh)
 		{
-			int32 const NumExistingPreviewMeshes = PreviewRailMeshSegments.Num();
-			if (NumExistingPreviewMeshes > NumNeededPreviewMeshes)
+			int32 const NumSplinePoints = RailSplineComponent->GetNumberOfSplinePoints();
+			int32 const NumNeededPreviewMeshes = NumSplinePoints - 1;
+
+			// make sure our preview mesh array is correctly sized and populated
 			{
-				// we have too many meshes, remove some
-				int32 const NumToRemove = NumExistingPreviewMeshes - NumNeededPreviewMeshes;
-				for (int Idx = 0; Idx < NumToRemove; ++Idx)
+				int32 const NumExistingPreviewMeshes = PreviewRailMeshSegments.Num();
+				if (NumExistingPreviewMeshes > NumNeededPreviewMeshes)
 				{
-					USplineMeshComponent* const ElementToRemove = PreviewRailMeshSegments.Pop();
-					ElementToRemove->UnregisterComponent();
+					// we have too many meshes, remove some
+					int32 const NumToRemove = NumExistingPreviewMeshes - NumNeededPreviewMeshes;
+					for (int Idx = 0; Idx < NumToRemove; ++Idx)
+					{
+						USplineMeshComponent* const ElementToRemove = PreviewRailMeshSegments.Pop();
+						ElementToRemove->UnregisterComponent();
+					}
 				}
+				else if (NumExistingPreviewMeshes < NumNeededPreviewMeshes)
+				{
+					int32 const NumToAdd = NumNeededPreviewMeshes - NumExistingPreviewMeshes;
+
+					for (int32 Idx = 0; Idx < NumToAdd; ++Idx)
+					{
+						USplineMeshComponent* PreviewMesh = CreateSplinePreviewSegment();
+						PreviewRailMeshSegments.Add(PreviewMesh);
+					}
+				}
+				check(PreviewRailMeshSegments.Num() == NumNeededPreviewMeshes);
 			}
-			else if (NumExistingPreviewMeshes < NumNeededPreviewMeshes)
+
+			for (int PtIdx = 0; PtIdx < NumSplinePoints - 1; ++PtIdx)
 			{
-				int32 const NumToAdd = NumNeededPreviewMeshes - NumExistingPreviewMeshes;
+				FVector StartLoc, StartTangent, EndLoc, EndTangent;
+				RailSplineComponent->GetLocationAndTangentAtSplinePoint(PtIdx, StartLoc, StartTangent, ESplineCoordinateSpace::Local);
+				RailSplineComponent->GetLocationAndTangentAtSplinePoint(PtIdx + 1, EndLoc, EndTangent, ESplineCoordinateSpace::Local);
 
-				for (int32 Idx = 0; Idx < NumToAdd; ++Idx)
+				USplineMeshComponent* const SplineMeshComp = PreviewRailMeshSegments[PtIdx];
+				if (SplineMeshComp)
 				{
-					USplineMeshComponent* PreviewMesh = CreateSplinePreviewSegment();
-					PreviewRailMeshSegments.Add(PreviewMesh);
+					SplineMeshComp->SetForwardAxis(ESplineMeshAxis::Z);
+					SplineMeshComp->SetStartAndEnd(StartLoc, StartTangent, EndLoc, EndTangent, true);
 				}
 			}
-			check(PreviewRailMeshSegments.Num() == NumNeededPreviewMeshes);
 		}
 
-		for (int PtIdx = 0; PtIdx < NumSplinePoints - 1; ++PtIdx)
+		// make visualization of the mount follow the contour of the rail
+		if (PreviewMesh_Mount)
 		{
-			FVector StartLoc, StartTangent, EndLoc, EndTangent;
-			RailSplineComponent->GetLocationAndTangentAtSplinePoint(PtIdx, StartLoc, StartTangent, ESplineCoordinateSpace::Local);
-			RailSplineComponent->GetLocationAndTangentAtSplinePoint(PtIdx + 1, EndLoc, EndTangent, ESplineCoordinateSpace::Local);
-
-			USplineMeshComponent* const SplineMeshComp = PreviewRailMeshSegments[PtIdx];
-			SplineMeshComp->SetStartScale(FVector2D(0.1f, 0.1f), false);
-			SplineMeshComp->SetEndScale(FVector2D(0.1f, 0.1f), false);
-			SplineMeshComp->SetForwardAxis(ESplineMeshAxis::Z);
-			SplineMeshComp->SetStartAndEnd(StartLoc, StartTangent, EndLoc, EndTangent, true);
+			float const SplineLen = RailSplineComponent->GetSplineLength();
+			FQuat const RailRot = RailSplineComponent->GetQuaternionAtDistanceAlongSpline(CurrentPositionOnRail*SplineLen, ESplineCoordinateSpace::World);
+			PreviewMesh_Mount->SetWorldRotation(RailRot);
 		}
 	}
 }
@@ -122,7 +134,8 @@ void ACameraRig_Rail::UpdateRailComponents()
 {
 	if (RailSplineComponent && RailCameraMount)
 	{
-		FVector const MountPos = RailSplineComponent->GetLocationAtDistanceAlongSpline(CurrentPositionOnRail, ESplineCoordinateSpace::World);
+		float const SplineLen = RailSplineComponent->GetSplineLength();
+		FVector const MountPos = RailSplineComponent->GetLocationAtDistanceAlongSpline(CurrentPositionOnRail*SplineLen, ESplineCoordinateSpace::World);
 		RailCameraMount->SetWorldLocation(MountPos);
 	}
 
@@ -145,13 +158,6 @@ void ACameraRig_Rail::PostEditChangeProperty(FPropertyChangedEvent& PropertyChan
 {
 	Super::PostEditChangeProperty(PropertyChangedEvent);
 	UpdateRailComponents();
-
-	if (RailSplineComponent)
-	{
-		// clamp position to spline limits
-		float const SplineLen = RailSplineComponent->GetSplineLength();
-		CurrentPositionOnRail = FMath::Clamp(CurrentPositionOnRail, 0.f, SplineLen);
-	}
 }
 #endif
 
