Copying //UE4/Dev-Framework to //UE4/Dev-Main (Source: //UE4/Dev-Framework @ 2972815)

diff --git a/Engine/Source/Editor/UnrealEd/Private/EditorEngine.cpp b/Engine/Source/Editor/UnrealEd/Private/EditorEngine.cpp
--- a/Engine/Source/Editor/UnrealEd/Private/EditorEngine.cpp
+++ b/Engine/Source/Editor/UnrealEd/Private/EditorEngine.cpp
@@ -810,11 +810,11 @@ void UEditorEngine::Init(IEngineLoop* InEngineLoop)
 	ActorFactories.Sort( FCompareUActorFactoryByMenuPriority() );
 
 	// Load game user settings and apply
-	auto GameUserSettings = GetGameUserSettings();
-	if (GameUserSettings)
+	UGameUserSettings* MyGameUserSettings = GetGameUserSettings();
+	if (MyGameUserSettings)
 	{
-		GameUserSettings->LoadSettings();
-		GameUserSettings->ApplySettings(true);
+		MyGameUserSettings->LoadSettings();
+		MyGameUserSettings->ApplySettings(true);
 	}
 
 	UEditorStyleSettings* Settings = GetMutableDefault<UEditorStyleSettings>();
@@ -1535,7 +1535,6 @@ void UEditorEngine::Tick( float DeltaSeconds, bool bIdleMode )
 	IStreamingManager::Get().Tick(DeltaSeconds);
 
 	// Update Audio. This needs to occur after rendering as the rendering code updates the listener position.
-	FAudioDeviceManager* AudioDeviceManager = GEngine->GetAudioDeviceManager();
 	if (AudioDeviceManager)
 	{
 		UWorld* OldGWorld = NULL;
@@ -4222,14 +4221,14 @@ FString UEditorEngine::GetFriendlyName( const UProperty* Property, UStruct* Owne
 	return FoundText.ToString();
 }
 
-AActor* UEditorEngine::UseActorFactoryOnCurrentSelection( UActorFactory* Factory, const FTransform* InActorTransform, EObjectFlags ObjectFlags )
+AActor* UEditorEngine::UseActorFactoryOnCurrentSelection( UActorFactory* Factory, const FTransform* InActorTransform, EObjectFlags InObjectFlags )
 {
 	// ensure that all selected assets are loaded
 	FEditorDelegates::LoadSelectedAssetsIfNeeded.Broadcast();
-	return UseActorFactory(Factory, FAssetData( GetSelectedObjects()->GetTop<UObject>() ), InActorTransform, ObjectFlags );
+	return UseActorFactory(Factory, FAssetData( GetSelectedObjects()->GetTop<UObject>() ), InActorTransform, InObjectFlags );
 }
 
-AActor* UEditorEngine::UseActorFactory( UActorFactory* Factory, const FAssetData& AssetData, const FTransform* InActorTransform, EObjectFlags ObjectFlags )
+AActor* UEditorEngine::UseActorFactory( UActorFactory* Factory, const FAssetData& AssetData, const FTransform* InActorTransform, EObjectFlags InObjectFlags )
 {
 	check( Factory );
 
@@ -4279,7 +4278,7 @@ AActor* UEditorEngine::UseActorFactory( UActorFactory* Factory, const FAssetData
 				const FScopedTransaction Transaction( NSLOCTEXT("UnrealEd", "CreateActor", "Create Actor") );
 
 				// Create the actor.
-				Actor = Factory->CreateActor( Asset, DesiredLevel, ActorTransform, ObjectFlags );
+				Actor = Factory->CreateActor( Asset, DesiredLevel, ActorTransform, InObjectFlags );
 				if(Actor != NULL)
 				{
 					SelectNone( false, true );
@@ -5624,7 +5623,7 @@ bool UEditorEngine::AreAllWindowsHidden() const
 	return bAllHidden;
 }
 
-AActor* UEditorEngine::AddActor(ULevel* InLevel, UClass* Class, const FTransform& Transform, bool bSilent, EObjectFlags ObjectFlags)
+AActor* UEditorEngine::AddActor(ULevel* InLevel, UClass* Class, const FTransform& Transform, bool bSilent, EObjectFlags InObjectFlags)
 {
 	check( Class );
 
@@ -5672,7 +5671,7 @@ AActor* UEditorEngine::AddActor(ULevel* InLevel, UClass* Class, const FTransform
 	AActor* Actor = NULL;
 	{
 		FScopedTransaction Transaction( NSLOCTEXT("UnrealEd", "AddActor", "Add Actor") );
-		if ( !(ObjectFlags & RF_Transactional) )
+		if ( !(InObjectFlags & RF_Transactional) )
 		{
 			// Don't attempt a transaction if the actor we are spawning isn't transactional
 			Transaction.Cancel();
@@ -5684,7 +5683,7 @@ AActor* UEditorEngine::AddActor(ULevel* InLevel, UClass* Class, const FTransform
 		FActorSpawnParameters SpawnInfo;
 		SpawnInfo.OverrideLevel = DesiredLevel;
 		SpawnInfo.SpawnCollisionHandlingOverride = ESpawnActorCollisionHandlingMethod::AlwaysSpawn;
-		SpawnInfo.ObjectFlags = ObjectFlags;
+		SpawnInfo.ObjectFlags = InObjectFlags;
 		const auto Location = Transform.GetLocation();
 		const auto Rotation = Transform.GetRotation().Rotator();
 		Actor = World->SpawnActor( Class, &Location, &Rotation, SpawnInfo );
@@ -5716,7 +5715,7 @@ AActor* UEditorEngine::AddActor(ULevel* InLevel, UClass* Class, const FTransform
 	return Actor;
 }
 
-TArray<AActor*> UEditorEngine::AddExportTextActors(const FString& ExportText, bool bSilent, EObjectFlags ObjectFlags)
+TArray<AActor*> UEditorEngine::AddExportTextActors(const FString& ExportText, bool bSilent, EObjectFlags InObjectFlags)
 {
 	TArray<AActor*> NewActors;
 
@@ -5733,7 +5732,7 @@ TArray<AActor*> UEditorEngine::AddExportTextActors(const FString& ExportText, bo
 	FVector Location;
 	{
 		FScopedTransaction Transaction( NSLOCTEXT("UnrealEd", "AddActor", "Add Actor") );
-		if ( !(ObjectFlags & RF_Transactional) )
+		if ( !(InObjectFlags & RF_Transactional) )
 		{
 			// Don't attempt a transaction if the actor we are spawning isn't transactional
 			Transaction.Cancel();
@@ -5741,7 +5740,7 @@ TArray<AActor*> UEditorEngine::AddExportTextActors(const FString& ExportText, bo
 		// Remove the selection to detect the actors that were created during FactoryCreateText. They will be selected when the operation in complete
 		GEditor->SelectNone( false, true );
 		const TCHAR* Text = *ExportText;
-		if ( Factory->FactoryCreateText( ULevel::StaticClass(), CurrentLevel, CurrentLevel->GetFName(), ObjectFlags, NULL, TEXT("paste"), Text, Text + FCString::Strlen(Text), GWarn ) != NULL )
+		if ( Factory->FactoryCreateText( ULevel::StaticClass(), CurrentLevel, CurrentLevel->GetFName(), InObjectFlags, nullptr, TEXT("paste"), Text, Text + FCString::Strlen(Text), GWarn ) != nullptr )
 		{
 			// Now get the selected actors and calculate a center point between all their locations.
 			USelection* ActorSelection = GEditor->GetSelectedActors();
diff --git a/Engine/Source/Editor/UnrealEd/Private/Factories/ActorFactory.cpp b/Engine/Source/Editor/UnrealEd/Private/Factories/ActorFactory.cpp
--- a/Engine/Source/Editor/UnrealEd/Private/Factories/ActorFactory.cpp
+++ b/Engine/Source/Editor/UnrealEd/Private/Factories/ActorFactory.cpp
@@ -196,13 +196,13 @@ FQuat UActorFactory::AlignObjectToSurfaceNormal(const FVector& InSurfaceNormal,
 	}
 }
 
-AActor* UActorFactory::CreateActor( UObject* Asset, ULevel* InLevel, FTransform SpawnTransform, EObjectFlags ObjectFlags, const FName Name )
+AActor* UActorFactory::CreateActor( UObject* Asset, ULevel* InLevel, FTransform SpawnTransform, EObjectFlags InObjectFlags, const FName Name )
 {
 	AActor* NewActor = NULL;
 
 	if ( PreSpawnActor(Asset, SpawnTransform) )
 	{
-		NewActor = SpawnActor(Asset, InLevel, SpawnTransform, ObjectFlags, Name);
+		NewActor = SpawnActor(Asset, InLevel, SpawnTransform, InObjectFlags, Name);
 
 		if ( NewActor )
 		{
@@ -230,14 +230,14 @@ bool UActorFactory::PreSpawnActor( UObject* Asset, FTransform& InOutLocation)
 	return true;
 }
 
-AActor* UActorFactory::SpawnActor( UObject* Asset, ULevel* InLevel, const FTransform& Transform, EObjectFlags ObjectFlags, const FName Name )
+AActor* UActorFactory::SpawnActor( UObject* Asset, ULevel* InLevel, const FTransform& Transform, EObjectFlags InObjectFlags, const FName Name )
 {
 	AActor* DefaultActor = GetDefaultActor( FAssetData( Asset ) );
 	if ( DefaultActor )
 	{
 		FActorSpawnParameters SpawnInfo;
 		SpawnInfo.OverrideLevel = InLevel;
-		SpawnInfo.ObjectFlags = ObjectFlags;
+		SpawnInfo.ObjectFlags = InObjectFlags;
 		SpawnInfo.Name = Name;
 		return InLevel->OwningWorld->SpawnActor( DefaultActor->GetClass(), &Transform, SpawnInfo );
 	}
@@ -1205,12 +1205,12 @@ bool UActorFactoryEmptyActor::CanCreateActorFrom( const FAssetData& AssetData, F
 	return AssetData.ObjectPath == FName(*AActor::StaticClass()->GetPathName());
 }
 
-AActor* UActorFactoryEmptyActor::SpawnActor( UObject* Asset, ULevel* InLevel, const FTransform& Transform, EObjectFlags ObjectFlags, const FName Name )
+AActor* UActorFactoryEmptyActor::SpawnActor( UObject* Asset, ULevel* InLevel, const FTransform& Transform, EObjectFlags InObjectFlags, const FName Name )
 {
 	AActor* NewActor = nullptr;
 	{
 		// Spawn a temporary actor for dragging around
-		NewActor = Super::SpawnActor(Asset, InLevel, Transform, ObjectFlags, Name);
+		NewActor = Super::SpawnActor(Asset, InLevel, Transform, InObjectFlags, Name);
 
 		USceneComponent* RootComponent = NewObject<USceneComponent>(NewActor, USceneComponent::GetDefaultSceneRootVariableName(), RF_Transactional);
 		RootComponent->Mobility = EComponentMobility::Movable;
@@ -1373,7 +1373,7 @@ bool UActorFactoryClass::PreSpawnActor( UObject* Asset, FTransform& InOutLocatio
 	return false;
 }
 
-AActor* UActorFactoryClass::SpawnActor( UObject* Asset, ULevel* InLevel, const FTransform& Transform, EObjectFlags ObjectFlags, const FName Name )
+AActor* UActorFactoryClass::SpawnActor( UObject* Asset, ULevel* InLevel, const FTransform& Transform, EObjectFlags InObjectFlags, const FName Name )
 {
 	UClass* ActualClass = Cast<UClass>(Asset);
 
@@ -1381,7 +1381,7 @@ AActor* UActorFactoryClass::SpawnActor( UObject* Asset, ULevel* InLevel, const F
 	{
 		FActorSpawnParameters SpawnInfo;
 		SpawnInfo.OverrideLevel = InLevel;
-		SpawnInfo.ObjectFlags = ObjectFlags;
+		SpawnInfo.ObjectFlags = InObjectFlags;
 		SpawnInfo.Name = Name;
 		return InLevel->OwningWorld->SpawnActor( ActualClass, &Transform, SpawnInfo );
 	}
@@ -1979,9 +1979,9 @@ bool UActorFactoryMovieScene::CanCreateActorFrom( const FAssetData& AssetData, F
 	return true;
 }
 
-AActor* UActorFactoryMovieScene::SpawnActor( UObject* Asset, ULevel* InLevel, const FTransform& Transform, EObjectFlags ObjectFlags, const FName Name )
+AActor* UActorFactoryMovieScene::SpawnActor( UObject* Asset, ULevel* InLevel, const FTransform& Transform, EObjectFlags InObjectFlags, const FName Name )
 {
-	ALevelSequenceActor* NewActor = Cast<ALevelSequenceActor>(Super::SpawnActor(Asset, InLevel, Transform, ObjectFlags, Name));
+	ALevelSequenceActor* NewActor = Cast<ALevelSequenceActor>(Super::SpawnActor(Asset, InLevel, Transform, InObjectFlags, Name));
 
 	if (NewActor)
 	{
diff --git a/Engine/Source/Editor/UnrealEd/Private/Factories/EditorFactories.cpp b/Engine/Source/Editor/UnrealEd/Private/Factories/EditorFactories.cpp
--- a/Engine/Source/Editor/UnrealEd/Private/Factories/EditorFactories.cpp
+++ b/Engine/Source/Editor/UnrealEd/Private/Factories/EditorFactories.cpp
@@ -5694,11 +5694,10 @@ EReimportResult::Type UReimportFbxStaticMeshFactory::Reimport( UObject* Obj )
 		Mesh->AssetImportData = ImportData;
 		ReimportUI->StaticMeshImportData = ImportData;
 		
-		bool bOperationCanceled = false;
-		bool bShowOption = true;
-		bool bForceImportType = true;
+		bool bImportOperationCanceled = false;
+		bool bImportForceType = true;
 
-		GetImportOptions( FFbxImporter, ReimportUI, bShowOption, Obj->GetPathName(), bOperationCanceled, bForceImportType, FBXIT_StaticMesh );
+		GetImportOptions( FFbxImporter, ReimportUI, /*bShowOptionDialog=*/true, Obj->GetPathName(), bImportOperationCanceled, bImportForceType, FBXIT_StaticMesh );
 	}
 
 	if( !bOperationCanceled && ensure(ImportData) )
@@ -5918,8 +5917,7 @@ EReimportResult::Type UReimportFbxSkeletalMeshFactory::Reimport( UObject* Obj )
 		SkeletalMesh->AssetImportData = ImportData;
 		ReimportUI->SkeletalMeshImportData = ImportData;
 
-		bool bOperationCanceled = false;
-		bool bShowOption = true;
+		bool bImportOperationCanceled = false;
 		bool bForceImportType = true;
 
 		// arggg... hate this different option class to confuse everybody
@@ -5928,7 +5926,7 @@ EReimportResult::Type UReimportFbxSkeletalMeshFactory::Reimport( UObject* Obj )
 		ImportOptions->bCreatePhysicsAsset = false;
 		ImportOptions->PhysicsAsset = SkeletalMesh->PhysicsAsset;
 
-		ImportOptions = GetImportOptions( FFbxImporter, ReimportUI, bShowOption, Obj->GetPathName(), bOperationCanceled, bForceImportType, FBXIT_SkeletalMesh );
+		ImportOptions = GetImportOptions( FFbxImporter, ReimportUI, /*bShowOptionDialog=*/true, Obj->GetPathName(), bImportOperationCanceled, bForceImportType, FBXIT_SkeletalMesh );
 	}
 
 	if( !bOperationCanceled && ensure(ImportData) )
diff --git a/Engine/Source/Editor/UnrealEd/Private/LevelEditorViewport.cpp b/Engine/Source/Editor/UnrealEd/Private/LevelEditorViewport.cpp
--- a/Engine/Source/Editor/UnrealEd/Private/LevelEditorViewport.cpp
+++ b/Engine/Source/Editor/UnrealEd/Private/LevelEditorViewport.cpp
@@ -2136,11 +2136,11 @@ void TrimLineToFrustum(const FConvexVolume& Frustum, FVector& Start, FVector& En
 	}
 }
 
-void FLevelEditorViewportClient::ProjectActorsIntoWorld(const TArray<AActor*>& Actors, FViewport* Viewport, const FVector& Drag, const FRotator& Rot)
+void FLevelEditorViewportClient::ProjectActorsIntoWorld(const TArray<AActor*>& Actors, FViewport* InViewport, const FVector& Drag, const FRotator& Rot)
 {
 	// Compile an array of selected actors
 	FSceneViewFamilyContext ViewFamily(FSceneViewFamily::ConstructionValues(
-		Viewport, 
+		InViewport, 
 		GetScene(),
 		EngineShowFlags)
 		.SetRealtimeUpdate( IsRealtime() ));
@@ -2181,7 +2181,7 @@ void FLevelEditorViewportClient::ProjectActorsIntoWorld(const TArray<AActor*>& A
 		{
 			// We only snap things that are on screen
 			FVector2D ScreenPos;
-			FIntPoint ViewportSize = Viewport->GetSizeXY();
+			FIntPoint ViewportSize = InViewport->GetSizeXY();
 			if (SceneView->WorldToPixel(NewActorPosition, ScreenPos) && FMath::IsWithin<float>(ScreenPos.X, 0, ViewportSize.X) && FMath::IsWithin<float>(ScreenPos.Y, 0, ViewportSize.Y))
 			{
 				bIsOnScreen = true;
@@ -2250,9 +2250,9 @@ void FLevelEditorViewportClient::ProjectActorsIntoWorld(const TArray<AActor*>& A
 	}
 }
 
-bool FLevelEditorViewportClient::InputWidgetDelta(FViewport* Viewport, EAxisList::Type InCurrentAxis, FVector& Drag, FRotator& Rot, FVector& Scale)
+bool FLevelEditorViewportClient::InputWidgetDelta(FViewport* InViewport, EAxisList::Type InCurrentAxis, FVector& Drag, FRotator& Rot, FVector& Scale)
 {
-	if (GUnrealEd->ComponentVisManager.HandleInputDelta(this, Viewport, Drag, Rot, Scale))
+	if (GUnrealEd->ComponentVisManager.HandleInputDelta(this, InViewport, Drag, Rot, Scale))
 	{
 		return true;
 	}
@@ -2260,7 +2260,7 @@ bool FLevelEditorViewportClient::InputWidgetDelta(FViewport* Viewport, EAxisList
 	bool bHandled = false;
 
 	// Give the current editor mode a chance to use the input first.  If it does, don't apply it to anything else.
-	if (FEditorViewportClient::InputWidgetDelta(Viewport, InCurrentAxis, Drag, Rot, Scale))
+	if (FEditorViewportClient::InputWidgetDelta(InViewport, InCurrentAxis, Drag, Rot, Scale))
 	{
 		bHandled = true;
 	}
@@ -2273,9 +2273,9 @@ bool FLevelEditorViewportClient::InputWidgetDelta(FViewport* Viewport, EAxisList
 			// but still pretend that we have handled the input
 			if (!GEditor->HasLockedActors())
 			{
-				const bool LeftMouseButtonDown = Viewport->KeyState(EKeys::LeftMouseButton);
-				const bool RightMouseButtonDown = Viewport->KeyState(EKeys::RightMouseButton);
-				const bool MiddleMouseButtonDown = Viewport->KeyState(EKeys::MiddleMouseButton);
+				const bool LeftMouseButtonDown = InViewport->KeyState(EKeys::LeftMouseButton);
+				const bool RightMouseButtonDown = InViewport->KeyState(EKeys::RightMouseButton);
+				const bool MiddleMouseButtonDown = InViewport->KeyState(EKeys::MiddleMouseButton);
 
 				// If duplicate dragging . . .
 				if ( IsAltPressed() && (LeftMouseButtonDown || RightMouseButtonDown) )
@@ -2317,7 +2317,7 @@ bool FLevelEditorViewportClient::InputWidgetDelta(FViewport* Viewport, EAxisList
 							}
 						}
 
-						ProjectActorsIntoWorld(SelectedActors, Viewport, Drag, Rot);
+						ProjectActorsIntoWorld(SelectedActors, InViewport, Drag, Rot);
 					}
 					else
 					{
@@ -2432,7 +2432,7 @@ void FLevelEditorViewportClient::SetLastKeyViewport()
 	}
 }
 
-bool FLevelEditorViewportClient::InputKey(FViewport* Viewport, int32 ControllerId, FKey Key, EInputEvent Event, float AmountDepressed, bool bGamepad)
+bool FLevelEditorViewportClient::InputKey(FViewport* InViewport, int32 ControllerId, FKey Key, EInputEvent Event, float AmountDepressed, bool bGamepad)
 {
 	if (bDisableInput)
 	{
@@ -2440,16 +2440,16 @@ bool FLevelEditorViewportClient::InputKey(FViewport* Viewport, int32 ControllerI
 	}
 
 	
-	const int32	HitX = Viewport->GetMouseX();
-	const int32	HitY = Viewport->GetMouseY();
+	const int32	HitX = InViewport->GetMouseX();
+	const int32	HitY = InViewport->GetMouseY();
 
-	FInputEventState InputState( Viewport, Key, Event );
+	FInputEventState InputState( InViewport, Key, Event );
 
 	SetLastKeyViewport();
 
 	// Compute a view.
 	FSceneViewFamilyContext ViewFamily(FSceneViewFamily::ConstructionValues(
-		Viewport,
+		InViewport,
 		GetScene(),
 		EngineShowFlags )
 		.SetRealtimeUpdate( IsRealtime() ) );
@@ -2468,12 +2468,12 @@ bool FLevelEditorViewportClient::InputKey(FViewport* Viewport, int32 ControllerI
 		FSnappingUtils::SnapPointToGrid(GEditor->ClickLocation, FVector::ZeroVector);
 	}
 
-	if (GUnrealEd->ComponentVisManager.HandleInputKey(this, Viewport, Key, Event))
+	if (GUnrealEd->ComponentVisManager.HandleInputKey(this, InViewport, Key, Event))
 	{
 		return true;
 	}
 
-	bool bHandled = FEditorViewportClient::InputKey(Viewport,ControllerId,Key,Event,AmountDepressed,bGamepad);
+	bool bHandled = FEditorViewportClient::InputKey(InViewport,ControllerId,Key,Event,AmountDepressed,bGamepad);
 
 	// Handle input for the player height preview mode. 
 	if (!InputState.IsMouseButtonEvent() && CommandAcceptsInput(*this, Key, GetLevelViewportCommands().EnablePreviewMesh))
@@ -2786,12 +2786,12 @@ void FLevelEditorViewportClient::OnActorMoved(AActor* InActor)
 
 void FLevelEditorViewportClient::NudgeSelectedObjects( const struct FInputEventState& InputState )
 {
-	FViewport* Viewport = InputState.GetViewport();
+	FViewport* InViewport = InputState.GetViewport();
 	EInputEvent Event = InputState.GetInputEvent();
 	FKey Key = InputState.GetKey();
 
-	const int32 MouseX = Viewport->GetMouseX();
-	const int32 MouseY = Viewport->GetMouseY();
+	const int32 MouseX = InViewport->GetMouseX();
+	const int32 MouseY = InViewport->GetMouseY();
 
 	if( Event == IE_Pressed || Event == IE_Repeat )
 	{
@@ -2807,7 +2807,7 @@ void FLevelEditorViewportClient::NudgeSelectedObjects( const struct FInputEventS
 		}
 
 		FIntPoint StartMousePos;
-		Viewport->GetMousePos( StartMousePos );
+		InViewport->GetMousePos( StartMousePos );
 		FKey VirtualKey = EKeys::MouseX;
 		EAxisList::Type VirtualAxis = GetHorizAxis();
 		float VirtualDelta = GEditor->GetGridSize() * (Key == EKeys::Left?-1:1);
@@ -2823,7 +2823,7 @@ void FLevelEditorViewportClient::NudgeSelectedObjects( const struct FInputEventS
 		MouseDeltaTracker->AddDelta( this, VirtualKey , VirtualDelta, 1 );
 		Widget->SetCurrentAxis( VirtualAxis );
 		UpdateMouseDelta();
-		Viewport->SetMouse( StartMousePos.X , StartMousePos.Y );
+		InViewport->SetMouse( StartMousePos.X , StartMousePos.Y );
 	}
 	else if( bIsTracking && Event == IE_Released )
 	{
@@ -2910,7 +2910,7 @@ private:
 	FLevelEditorViewportClient* PrevCurrentLevelEditingViewportClient;
 };
 
-bool FLevelEditorViewportClient::InputAxis(FViewport* Viewport, int32 ControllerId, FKey Key, float Delta, float DeltaTime, int32 NumSamples, bool bGamepad)
+bool FLevelEditorViewportClient::InputAxis(FViewport* InViewport, int32 ControllerId, FKey Key, float Delta, float DeltaTime, int32 NumSamples, bool bGamepad)
 {
 	if (bDisableInput)
 	{
@@ -2921,7 +2921,7 @@ bool FLevelEditorViewportClient::InputAxis(FViewport* Viewport, int32 Controller
 
 	FScopedSetCurrentViewportClient( this );
 
-	return FEditorViewportClient::InputAxis(Viewport, ControllerId, Key, Delta, DeltaTime, NumSamples, bGamepad);
+	return FEditorViewportClient::InputAxis(InViewport, ControllerId, Key, Delta, DeltaTime, NumSamples, bGamepad);
 }
 
 
@@ -3528,14 +3528,14 @@ void FLevelEditorViewportClient::ApplyDeltaToActor( AActor* InActor, const FVect
 	UpdateLockedActorViewports(InActor, true);
 }
 
-EMouseCursor::Type FLevelEditorViewportClient::GetCursor(FViewport* Viewport,int32 X,int32 Y)
+EMouseCursor::Type FLevelEditorViewportClient::GetCursor(FViewport* InViewport,int32 X,int32 Y)
 {
-	EMouseCursor::Type CursorType = FEditorViewportClient::GetCursor(Viewport,X,Y);
+	EMouseCursor::Type CursorType = FEditorViewportClient::GetCursor(InViewport,X,Y);
 
-	HHitProxy* HitProxy = Viewport->GetHitProxy(X,Y);
+	HHitProxy* HitProxy = InViewport->GetHitProxy(X,Y);
 
 	// Don't select widget axes by mouse over while they're being controlled by a mouse drag.
-	if( Viewport->IsCursorVisible() && !bWidgetAxisControlledByDrag && !HitProxy )
+	if( InViewport->IsCursorVisible() && !bWidgetAxisControlledByDrag && !HitProxy )
 	{
 		if( HoveredObjects.Num() > 0 )
 		{
diff --git a/Engine/Source/Editor/WorldBrowser/Private/Tiles/SWorldComposition.cpp b/Engine/Source/Editor/WorldBrowser/Private/Tiles/SWorldComposition.cpp
--- a/Engine/Source/Editor/WorldBrowser/Private/Tiles/SWorldComposition.cpp
+++ b/Engine/Source/Editor/WorldBrowser/Private/Tiles/SWorldComposition.cpp
@@ -306,8 +306,8 @@ public:
 			FVector2D PaintPosition = GraphCoordToPanelCoord(FVector2D::ZeroVector) - (PaintSize*0.5f);
 			float Scale = 0.2f; // Scale down drawing border
 			FSlateLayoutTransform LayoutTransform(Scale, AllottedGeometry.GetAccumulatedLayoutTransform().GetTranslation() + PaintPosition);
-			FSlateRenderTransform RenderTransform(Scale, AllottedGeometry.GetAccumulatedRenderTransform().GetTranslation() + PaintPosition);
-			FPaintGeometry EditableArea(LayoutTransform, RenderTransform, PaintSize/Scale);
+			FSlateRenderTransform SlateRenderTransform(Scale, AllottedGeometry.GetAccumulatedRenderTransform().GetTranslation() + PaintPosition);
+			FPaintGeometry EditableArea(LayoutTransform, SlateRenderTransform, PaintSize/Scale);
 
 			FLinearColor PaintColor = FLinearColor::Yellow;
 			PaintColor.A = 0.4f;
diff --git a/Engine/Source/Editor/WorldBrowser/Private/Tiles/SWorldTileItem.cpp b/Engine/Source/Editor/WorldBrowser/Private/Tiles/SWorldTileItem.cpp
--- a/Engine/Source/Editor/WorldBrowser/Private/Tiles/SWorldTileItem.cpp
+++ b/Engine/Source/Editor/WorldBrowser/Private/Tiles/SWorldTileItem.cpp
@@ -295,8 +295,8 @@ int32 SWorldTileItem::OnPaint(const FPaintArgs& Args, const FGeometry& AllottedG
 			const FVector2D InflateAmount = FVector2D(4, 4);
 			const float Scale = 0.5f; // Scale down image of the borders to make them thinner 
 			FSlateLayoutTransform LayoutTransform(Scale, AllottedGeometry.GetAccumulatedLayoutTransform().GetTranslation() - InflateAmount);
-			FSlateRenderTransform RenderTransform(Scale, AllottedGeometry.GetAccumulatedRenderTransform().GetTranslation() - InflateAmount);
-			FPaintGeometry SelectionGeometry(LayoutTransform, RenderTransform, (AllottedGeometry.GetLocalSize()*AllottedGeometry.Scale + InflateAmount*2)/Scale);
+			FSlateRenderTransform SlateRenderTransform(Scale, AllottedGeometry.GetAccumulatedRenderTransform().GetTranslation() - InflateAmount);
+			FPaintGeometry SelectionGeometry(LayoutTransform, SlateRenderTransform, (AllottedGeometry.GetLocalSize()*AllottedGeometry.Scale + InflateAmount*2)/Scale);
 			FLinearColor HighlightColor = FLinearColor::White;
 			if (PreviewLODIndex)
 			{
diff --git a/Engine/Source/Programs/CrashReporter/CrashReportClient/Private/Windows/WindowsErrorReport.cpp b/Engine/Source/Programs/CrashReporter/CrashReportClient/Private/Windows/WindowsErrorReport.cpp
--- a/Engine/Source/Programs/CrashReporter/CrashReportClient/Private/Windows/WindowsErrorReport.cpp
+++ b/Engine/Source/Programs/CrashReporter/CrashReportClient/Private/Windows/WindowsErrorReport.cpp
@@ -143,9 +143,11 @@ void FWindowsErrorReport::FindMostRecentErrorReports(TArray<FString>& ErrorRepor
 		return true;
 	});
 
-	TCHAR LocalAppDataPath[MAX_PATH];
-	SHGetFolderPath(0, CSIDL_LOCAL_APPDATA, NULL, 0, LocalAppDataPath);
-	PlatformFile.IterateDirectory( *(FString(LocalAppDataPath) / TEXT("Microsoft/Windows/WER/ReportQueue")), ReportFinder);
+	{
+		TCHAR LocalAppDataPath[MAX_PATH];
+		SHGetFolderPath(0, CSIDL_LOCAL_APPDATA, NULL, 0, LocalAppDataPath);
+		PlatformFile.IterateDirectory(*(FString(LocalAppDataPath) / TEXT("Microsoft/Windows/WER/ReportQueue")), ReportFinder);
+	}
 
 	if (ErrorReportPaths.Num() == 0)
 	{
diff --git a/Engine/Source/Programs/UnrealBuildTool/Windows/VCToolChain.cs b/Engine/Source/Programs/UnrealBuildTool/Windows/VCToolChain.cs
--- a/Engine/Source/Programs/UnrealBuildTool/Windows/VCToolChain.cs
+++ b/Engine/Source/Programs/UnrealBuildTool/Windows/VCToolChain.cs
@@ -583,9 +583,12 @@ namespace UnrealBuildTool
 			if (WindowsPlatform.Compiler == WindowsCompiler.VisualStudio2015)
 			{
 				// Disable shadow variable warnings
-				Arguments.Append(" /wd4456"); // 4456 - declaration of 'LocalVariable' hides previous local declaration
-				Arguments.Append(" /wd4458"); // 4458 - declaration of 'parameter' hides class member
-				Arguments.Append(" /wd4459"); // 4459 - declaration of 'LocalVariable' hides global declaration
+				if (CompileEnvironment.Config.bEnableShadowVariableWarning == false)
+				{
+					Arguments.Append(" /wd4456"); // 4456 - declaration of 'LocalVariable' hides previous local declaration
+					Arguments.Append(" /wd4458"); // 4458 - declaration of 'parameter' hides class member
+					Arguments.Append(" /wd4459"); // 4459 - declaration of 'LocalVariable' hides global declaration
+				}
 
 				Arguments.Append(" /wd4463"); // 4463 - overflow; assigning 1 to bit-field that can only hold values from -1 to 0
 
diff --git a/Engine/Source/Programs/UnrealLightmass/Private/ImportExport/LightmassScene.cpp b/Engine/Source/Programs/UnrealLightmass/Private/ImportExport/LightmassScene.cpp
--- a/Engine/Source/Programs/UnrealLightmass/Private/ImportExport/LightmassScene.cpp
+++ b/Engine/Source/Programs/UnrealLightmass/Private/ImportExport/LightmassScene.cpp
@@ -282,32 +282,32 @@ FBoxSphereBounds FScene::GetImportanceBounds() const
 	return ImportanceBoundSphere;
 }
 
-const FLight* FScene::FindLightByGuid(const FGuid& Guid) const
+const FLight* FScene::FindLightByGuid(const FGuid& InGuid) const
 {
 	for (int32 i = 0; i < DirectionalLights.Num(); i++)
 	{
-		if (DirectionalLights[i].Guid == Guid)
+		if (DirectionalLights[i].Guid == InGuid)
 		{
 			return &DirectionalLights[i];
 		}
 	}
 	for (int32 i = 0; i < PointLights.Num(); i++)
 	{
-		if (PointLights[i].Guid == Guid)
+		if (PointLights[i].Guid == InGuid)
 		{
 			return &PointLights[i];
 		}
 	}
 	for (int32 i = 0; i < SpotLights.Num(); i++)
 	{
-		if (SpotLights[i].Guid == Guid)
+		if (SpotLights[i].Guid == InGuid)
 		{
 			return &SpotLights[i];
 		}
 	}
 	for (int32 i = 0; i < SkyLights.Num(); i++)
 	{
-		if (SkyLights[i].Guid == Guid)
+		if (SkyLights[i].Guid == InGuid)
 		{
 			return &SkyLights[i];
 		}
diff --git a/Engine/Source/Programs/UnrealPak/Private/KeyGenerator.cpp b/Engine/Source/Programs/UnrealPak/Private/KeyGenerator.cpp
--- a/Engine/Source/Programs/UnrealPak/Private/KeyGenerator.cpp
+++ b/Engine/Source/Programs/UnrealPak/Private/KeyGenerator.cpp
@@ -11,7 +11,6 @@
 // Global constants
 namespace
 {
-	const TEncryptionInt One(1);
 	const TEncryptionInt Two(2);
 	const TEncryptionInt IterationStep(1000);
 	TArray<TEncryptionInt> PrimeLookupTable;
@@ -42,7 +41,7 @@ public:
 		, MaxValue(InMaxValue)
 	{
 		// Must be odd number
-		check(!(Candidate & One).IsZero());
+		check(!(Candidate & TEncryptionInt::One).IsZero());
 		Thread = FRunnableThread::Create(this, TEXT("FPrimeCheckRunnable"));
 	}
 
@@ -96,7 +95,7 @@ public:
 bool IsPrime(const TEncryptionInt& InValue, bool bUseTasks)
 {
 	// 2 is but we don't care about small numbers here.
-	if ((InValue & One) == 0)
+	if ((InValue & TEncryptionInt::One) == 0)
 	{
 		return false;
 	}
@@ -145,7 +144,7 @@ bool IsPrime(const TEncryptionInt& InValue, bool bUseTasks)
 			TEncryptionInt MaxValue(Factor + Range);
 			Tasks.Add(new FPrimeCheckRunnable(FoundFactors, InValue, Factor, MaxValue)); 
 			Factor = MaxValue;
-			if ((Factor & One) == 0)
+			if ((Factor & TEncryptionInt::One) == 0)
 			{
 				++Factor;
 			}
@@ -299,7 +298,7 @@ public:
 		, MaxValue(InMaxValue)
 	{
 		// Must be an odd number
-		check(!(MinValue & One).IsZero());
+		check(!(MinValue & TEncryptionInt::One).IsZero());
 		Thread = FRunnableThread::Create(this, TEXT("FPrimeFinderRunnable"));
 	}
 
@@ -368,7 +367,7 @@ void GeneratePrimeNumberLookupTable()
 		TEncryptionInt MaxValue(MinPrimeValue + Range);
 		Tasks.Add(new FPrimeFinderRunnable(MinPrimeValue, MaxValue)); 
 		MinPrimeValue = MaxValue;
-		if ((MinPrimeValue & One) == 0)
+		if ((MinPrimeValue & TEncryptionInt::One) == 0)
 		{
 			++MinPrimeValue;
 		}
diff --git a/Engine/Source/Programs/UnrealVersionSelector/Private/Windows/WindowsPlatformInstallation.cpp b/Engine/Source/Programs/UnrealVersionSelector/Private/Windows/WindowsPlatformInstallation.cpp
--- a/Engine/Source/Programs/UnrealVersionSelector/Private/Windows/WindowsPlatformInstallation.cpp
+++ b/Engine/Source/Programs/UnrealVersionSelector/Private/Windows/WindowsPlatformInstallation.cpp
@@ -116,8 +116,8 @@ private:
 
 		for(int32 Idx =  0; Idx < SortedIdentifiers.Num(); Idx++)
 		{
-			const FString &Identifier = SortedIdentifiers[Idx];
-			FString Description = GetInstallationDescription(Identifier, Installations[Identifier]);
+			const FString &SortedIdentifier = SortedIdentifiers[Idx];
+			FString Description = GetInstallationDescription(SortedIdentifier, Installations[SortedIdentifier]);
 			SendDlgItemMessage(hWnd, IDC_BUILDLIST, CB_ADDSTRING, 0, (LPARAM)*Description);
 		}
 
diff --git a/Engine/Source/Runtime/AIModule/Private/BehaviorTree/BehaviorTreeComponent.cpp b/Engine/Source/Runtime/AIModule/Private/BehaviorTree/BehaviorTreeComponent.cpp
--- a/Engine/Source/Runtime/AIModule/Private/BehaviorTree/BehaviorTreeComponent.cpp
+++ b/Engine/Source/Runtime/AIModule/Private/BehaviorTree/BehaviorTreeComponent.cpp
@@ -1025,7 +1025,7 @@ void UBehaviorTreeComponent::TickComponent(float DeltaTime, enum ELevelTick Tick
 		ProcessExecutionRequest();
 	}
 
-	if (InstanceStack.Num() == 0 || !bIsRunning)
+	if (InstanceStack.Num() == 0 || !bIsRunning || bIsPaused)
 	{
 		return;
 	}
diff --git a/Engine/Source/Runtime/AIModule/Private/BrainComponent.cpp b/Engine/Source/Runtime/AIModule/Private/BrainComponent.cpp
--- a/Engine/Source/Runtime/AIModule/Private/BrainComponent.cpp
+++ b/Engine/Source/Runtime/AIModule/Private/BrainComponent.cpp
@@ -147,7 +147,7 @@ void FAIMessageObserver::Unregister()
 	UBrainComponent* OwnerComp = Owner.Get();
 	if (OwnerComp)
 	{
-		OwnerComp->MessageObservers.RemoveSingleSwap(this);
+		OwnerComp->MessageObservers.RemoveSingle(this);
 	}
 }
 
@@ -268,11 +268,14 @@ void UBrainComponent::TickComponent(float DeltaTime, enum ELevelTick TickType, F
 	if (MessagesToProcess.Num() > 0)
 	{
 		const int32 NumMessages = MessagesToProcess.Num();
-		for (auto Message : MessagesToProcess)
+		for (int32 Idx = 0; Idx < NumMessages; Idx++)
 		{
+			// create a copy of message in case MessagesToProcess is changed during loop
+			const FAIMessage MessageCopy(MessagesToProcess[Idx]);
+
 			for (int32 ObserverIndex = 0; ObserverIndex < MessageObservers.Num(); ObserverIndex++)
 			{
-				MessageObservers[ObserverIndex]->OnMessage(Message);
+				MessageObservers[ObserverIndex]->OnMessage(MessageCopy);
 			}
 		}
 		MessagesToProcess.RemoveAt(0, NumMessages, false);
diff --git a/Engine/Source/Runtime/Core/Public/Math/BigInt.h b/Engine/Source/Runtime/Core/Public/Math/BigInt.h
--- a/Engine/Source/Runtime/Core/Public/Math/BigInt.h
+++ b/Engine/Source/Runtime/Core/Public/Math/BigInt.h
@@ -17,6 +17,8 @@ public:
 
 	typedef TBigInt<NumBits, bSigned> BigInt;
 
+	static BigInt One;
+
 private:
 
 	enum
@@ -281,7 +283,6 @@ public:
 	 */
 	void Negate()
 	{
-		static const BigInt One(1LL);
 		for (int32 WordIndex = 0; WordIndex < NumWords; ++WordIndex)
 		{
 			Bits[WordIndex] = ~Bits[WordIndex];
@@ -356,7 +357,6 @@ public:
 	 */
 	void DivideWithRemainder(const BigInt& Divisor, BigInt& Remainder) 
 	{ 
-		static const BigInt One(1LL);
 		BigInt Denominator(Divisor);
 		BigInt Dividend(*this);		
 		
@@ -800,7 +800,6 @@ public:
 
 	FORCEINLINE BigInt& operator++()
 	{
-		static const BigInt One(1);
 		Add(One);
 		return *this;
 	}
@@ -820,7 +819,6 @@ public:
 
 	FORCEINLINE BigInt& operator--()
 	{
-		static const BigInt One(1);
 		Subtract(One);
 		return *this;
 	}
@@ -1011,6 +1009,9 @@ public:
 	}
 };
 
+template<int32 NumBits, bool bSigned>
+TBigInt<NumBits, bSigned> TBigInt<NumBits, bSigned>::One(1LL);
+
 // Predefined big int types
 typedef TBigInt<256> int256;
 typedef TBigInt<512> int512;
@@ -1157,7 +1158,7 @@ namespace FEncryption
 	FORCEINLINE static IntType ModularPow(IntType Base, IntType Exponent, IntType Modulus)
 	{
 		const IntType One(1);
-		IntType Result(One);
+		IntType Result(1);
 		while (Exponent > 0)
 		{
 			if (Exponent.IsFirstBitSet())
diff --git a/Engine/Source/Runtime/Engine/Classes/Components/ActorComponent.h b/Engine/Source/Runtime/Engine/Classes/Components/ActorComponent.h
--- a/Engine/Source/Runtime/Engine/Classes/Components/ActorComponent.h
+++ b/Engine/Source/Runtime/Engine/Classes/Components/ActorComponent.h
@@ -484,6 +484,7 @@ public:
 	 * @return  true if this component met the criteria for actually being ticked.
 	 */
 	bool SetupActorComponentTickFunction(struct FTickFunction* TickFunction);
+
 	/** 
 	 * Set this component's tick functions to be enabled or disabled. Only has an effect if the function is registered
 	 * 
@@ -503,6 +504,19 @@ public:
 	UFUNCTION(BlueprintCallable, Category="Utilities")
 	bool IsComponentTickEnabled() const;
 
+	/** 
+	* Sets the tick interval for this component's primary tick function. Does not enable the tick interval. Takes effect on next tick.
+	* @param TickInterval	The duration between ticks for this component's primary tick function
+	*/
+	UFUNCTION(BlueprintCallable, Category="Utilities")
+	void SetComponentTickInterval(float TickInterval);
+
+	/** 
+	* Returns whether this component has tick enabled or not
+	*/
+	UFUNCTION(BlueprintCallable, Category="Utilities")
+	float GetComponentTickInterval() const;
+
 	/**
 	 * @param InWorld - The world to register the component with.
 	 */
diff --git a/Engine/Source/Runtime/Engine/Classes/Engine/DebugCameraController.h b/Engine/Source/Runtime/Engine/Classes/Engine/DebugCameraController.h
--- a/Engine/Source/Runtime/Engine/Classes/Engine/DebugCameraController.h
+++ b/Engine/Source/Runtime/Engine/Classes/Engine/DebugCameraController.h
@@ -12,7 +12,7 @@
 * or use the "ToggleDebugCamera" console command. Check the debug camera bindings
 * in DefaultPawn.cpp for the camera controls.
 */
-UCLASS(config=Game, hidedropdown)
+UCLASS(config=Game)
 class ENGINE_API ADebugCameraController
 	: public APlayerController
 {
@@ -52,7 +52,8 @@ class ENGINE_API ADebugCameraController
 	/** @todo document */
 	void DecreaseFOV();
 
-	/** @todo document */
+	/** Toggles the display of debug info and input commands for the Debug Camera. */
+	UFUNCTION(BlueprintCallable, Category="Debug Camera")
 	void ToggleDisplay();
 
 	/**
@@ -66,6 +67,9 @@ public:
 	/** @todo document */
 	class AActor* SelectedActor;
 
+	UFUNCTION(BlueprintCallable, Category="Debug Camera")
+	AActor* GetSelectedActor() const;
+	
 	/** @todo document */
 	class UPrimitiveComponent* SelectedComponent;
 
@@ -75,16 +79,20 @@ public:
 	/** @todo document */
 	class UPlayer* OriginalPlayer;
 
-	// Allows control over the speed of the spectator pawn. This scales the speed based on the InitialMaxSpeed.
+	/* Allows control over the speed of the spectator pawn. This scales the speed based on the InitialMaxSpeed. */
+	UPROPERTY(EditDefaultsOnly, BlueprintReadWrite, Category="Debug Camera")
 	float SpeedScale;
 	
-	// Initial max speed of the spectator pawn when we start possession.
+	/* Initial max speed of the spectator pawn when we start possession. */
+	UPROPERTY(EditDefaultsOnly, BlueprintReadOnly, Category="Debug Camera")
 	float InitialMaxSpeed;
 
-	// Initial acceleration of the spectator pawn when we start possession.
+	/* Initial acceleration of the spectator pawn when we start possession. */
+	UPROPERTY(EditDefaultsOnly, BlueprintReadOnly, Category="Debug Camera")
 	float InitialAccel;
 
-	// Initial deceleration of the spectator pawn when we start possession.
+	/* Initial deceleration of the spectator pawn when we start possession. */
+	UPROPERTY(EditDefaultsOnly, BlueprintReadOnly, Category="Debug Camera")
 	float InitialDecel;
 
 protected:
@@ -95,11 +103,25 @@ protected:
 
 public:
 
+	/** 
+	* Function called on activation of debug camera controller.
+	* @param OriginalPC The active player controller before this debug camera controller was possessed by the player.
+	*/
+	UFUNCTION(BlueprintImplementableEvent, meta=(DisplayName = "OnActivate"))
+	void ReceiveOnActivate(class APlayerController* OriginalPC);
+	
 	/** Function called on activation debug camera controller */
-	void OnActivate(class APlayerController* OriginalPC);
-
+	virtual void OnActivate(class APlayerController* OriginalPC);
+
+	/** 
+	* Function called on deactivation of debug camera controller.
+	* @param RestoredPC The Player Controller that the player input is being returned to.
+	*/
+	UFUNCTION(BlueprintImplementableEvent, meta=(DisplayName = "OnDeactivate"))
+	void ReceiveOnDeactivate(class APlayerController* RestoredPC);
+	
 	/** Function called on deactivation debug camera controller */
-	void OnDeactivate(class APlayerController* RestoredPC);
+	virtual void OnDeactivate(class APlayerController* RestoredPC);
 
 	/**
 	 * Builds a list of components that are hidden based upon gameplay
@@ -122,6 +144,14 @@ public:
 
 protected:
 
+	/**
+	 * Called when an actor has been selected with the primary key (e.g. left mouse button).
+	 * @param SelectHitLocation The exact world-space location where the selection trace hit the New Selected Actor.
+	 * @param SelectHitNormal The world-space surface normal of the New Selected Actor at the hit location.
+	 */
+	UFUNCTION(BlueprintImplementableEvent, Category="Debug Camera", meta=(DisplayName="OnActorSelected"))
+	void ReceiveOnActorSelected(AActor* NewSelectedActor, const FVector& SelectHitLocation, const FVector& SelectHitNormal, const FHitResult& Hit);
+	
 	/**
 	 * Called when an actor has been selected with the primary key (e.g. left mouse button).
 	 * @param Hit	Info struct for the selection point.
diff --git a/Engine/Source/Runtime/Engine/Classes/GameFramework/Actor.h b/Engine/Source/Runtime/Engine/Classes/GameFramework/Actor.h
--- a/Engine/Source/Runtime/Engine/Classes/GameFramework/Actor.h
+++ b/Engine/Source/Runtime/Engine/Classes/GameFramework/Actor.h
@@ -1899,7 +1899,7 @@ public:
 	/** 
 	 * Set this actor's tick functions to be enabled or disabled. Only has an effect if the function is registered
 	 * This only modifies the tick function on actor itself
-	 * @param	bEnabled - Rather it should be enabled or not
+	 * @param	bEnabled	Whether it should be enabled or not
 	 */
 	UFUNCTION(BlueprintCallable, Category="Utilities")
 	void SetActorTickEnabled(bool bEnabled);
@@ -1908,6 +1908,17 @@ public:
 	UFUNCTION(BlueprintCallable, Category="Utilities")
 	bool IsActorTickEnabled() const;
 
+	/** 
+	* Sets the tick interval of this actor's primary tick function. Will not enable a disabled tick function. Takes effect on next tick. 
+	* @param TickInterval	The rate at which this actor should be ticking
+	*/
+	UFUNCTION(BlueprintCallable, Category="Utilities")
+	void SetActorTickInterval(float TickInterval);
+
+	/**  Returns the tick interval of this actor's primary tick function */
+	UFUNCTION(BlueprintCallable, Category="Utilities")
+	float GetActorTickInterval() const;
+
 	/**
 	 *	ticks the actor
 	 *	@param	DeltaTime			The time slice of this tick
diff --git a/Engine/Source/Runtime/Engine/Classes/GameFramework/Character.h b/Engine/Source/Runtime/Engine/Classes/GameFramework/Character.h
--- a/Engine/Source/Runtime/Engine/Classes/GameFramework/Character.h
+++ b/Engine/Source/Runtime/Engine/Classes/GameFramework/Character.h
@@ -423,7 +423,7 @@ public:
 	virtual void PawnClientRestart() override;
 	virtual void PossessedBy(AController* NewController) override;
 	virtual void UnPossessed() override;
-	virtual void SetupPlayerInputComponent(class UInputComponent* InInputComponent) override;
+	virtual void SetupPlayerInputComponent(class UInputComponent* PlayerInputComponent) override;
 	virtual void DisplayDebug(class UCanvas* Canvas, const FDebugDisplayInfo& DebugDisplay, float& YL, float& YPos) override;
 	virtual void RecalculateBaseEyeHeight() override;
 	virtual void UpdateNavigationRelevance() override;
diff --git a/Engine/Source/Runtime/Engine/Classes/GameFramework/CheatManager.h b/Engine/Source/Runtime/Engine/Classes/GameFramework/CheatManager.h
--- a/Engine/Source/Runtime/Engine/Classes/GameFramework/CheatManager.h
+++ b/Engine/Source/Runtime/Engine/Classes/GameFramework/CheatManager.h
@@ -54,7 +54,7 @@ struct FDebugTraceInfo
 
 };
 
-UCLASS(Within=PlayerController)
+UCLASS(Within=PlayerController, Blueprintable)
 class ENGINE_API UCheatManager : public UObject
 {
 	GENERATED_UCLASS_BODY()
@@ -63,8 +63,11 @@ class ENGINE_API UCheatManager : public UObject
 	UPROPERTY()
 	class ADebugCameraController* DebugCameraControllerRef;
 
-	/** Debug camera - used to have independent camera without stopping gameplay */
-	UPROPERTY()
+	/** 
+	* Debug camera - used to have independent camera without stopping gameplay. 
+	* If the Outer PlayerController::DebugCameraClass is valid then it is used instead of default DebugCameraController class. 
+	*/
+	UPROPERTY(Category="Debug Camera", EditDefaultsOnly, BlueprintReadOnly)
 	TSubclassOf<class ADebugCameraController>  DebugCameraControllerClass;
 
 	// Trace/Sweep debug start
@@ -359,6 +362,10 @@ class ENGINE_API UCheatManager : public UObject
 	/** streaming level debugging */
 	virtual void SetLevelStreamingStatus(FName PackageName, bool bShouldBeLoaded, bool bShouldBeVisible);
 
+	/** Called when the Cheat Manager is first initialized. */
+	UFUNCTION(BlueprintImplementableEvent, meta=(DisplayName = "Init Cheat Manager", Keywords="Begin Play"))
+	void ReceiveInitCheatManager();
+
 	/** 
 	 * Called when CheatManager is created to allow any needed initialization.  
 	 * This is not an actor, so we need a stand in for PostInitializeComponents 
diff --git a/Engine/Source/Runtime/Engine/Classes/GameFramework/HUD.h b/Engine/Source/Runtime/Engine/Classes/GameFramework/HUD.h
--- a/Engine/Source/Runtime/Engine/Classes/GameFramework/HUD.h
+++ b/Engine/Source/Runtime/Engine/Classes/GameFramework/HUD.h
@@ -31,53 +31,56 @@ class ENGINE_API AHUD : public AActor
 	GENERATED_UCLASS_BODY()
 
 	/** Pre-defined FColors for convenience. */
-	UPROPERTY()
-	FColor WhiteColor;
+	DEPRECATED(4.13, "Use FColor::White instead")
+	static const FColor WhiteColor;
 
-	UPROPERTY()
-	FColor GreenColor;
+	DEPRECATED(4.13, "Use FColor::Green instead")
+	static const FColor GreenColor;
 
-	UPROPERTY()
-	FColor RedColor;
+	DEPRECATED(4.13, "Use FColor::Red instead")
+	static const FColor RedColor;
 
 	/** PlayerController which owns this HUD. */
-	UPROPERTY()
+	UPROPERTY(BlueprintReadOnly, Category=HUD)
 	APlayerController* PlayerOwner;    
 
 	/** Tells whether the game was paused due to lost focus */
-	UPROPERTY(transient)
-	uint32 bLostFocusPaused:1;
+	UPROPERTY(BlueprintReadOnly, Category=HUD)
+	uint8 bLostFocusPaused:1;
 
 	/** Whether or not the HUD should be drawn. */
-	UPROPERTY(config)
-	uint32 bShowHUD:1;    
+	UPROPERTY(EditAnywhere, BlueprintReadWrite, Category=HUD)
+	uint8 bShowHUD:1;    
 
 	/** If true, current ViewTarget shows debug information using its DisplayDebug(). */
-	UPROPERTY()
-	uint32 bShowDebugInfo:1;    
+	UPROPERTY(BlueprintReadWrite, Category=HUD)
+	uint8 bShowDebugInfo:1;    
 
 	/** If true, show hitbox debugging info. */
-	UPROPERTY()
-	uint32 bShowHitBoxDebugInfo:1;    
+	UPROPERTY(BlueprintReadWrite, Category=HUD)
+	uint8 bShowHitBoxDebugInfo:1;    
 
 	/** If true, render actor overlays. */
-	UPROPERTY()
-	uint32 bShowOverlays:1;
+	UPROPERTY(BlueprintReadWrite, Category=HUD)
+	uint8 bShowOverlays:1;
 
 	/** Put shadow on debug strings */
-	UPROPERTY()
-	uint32 bEnableDebugTextShadow : 1;	
+	UPROPERTY(EditAnywhere, BlueprintReadWrite, Category=HUD)
+	uint8 bEnableDebugTextShadow:1;	
 
+private:
+	/** if true show debug info for 'ShowDebugTargetActor', otherwise for Camera Viewtarget */
+	uint8 bShowDebugForReticleTarget:1;
+	
+public:
 	/** Holds a list of Actors that need PostRender() calls. */
 	UPROPERTY()
 	TArray<AActor*> PostRenderedActors;
 
 	/** Used to calculate delta time between HUD rendering. */
-	UPROPERTY(transient)
 	float LastHUDRenderTime;
 
 	/** Time since last HUD render. */
-	UPROPERTY(transient)
 	float RenderDelta;
 
 	/** Array of names specifying what debug info to display for viewtarget actor. */
@@ -126,16 +129,12 @@ public:
 	void ShowDebugForReticleTargetToggle(TSubclassOf<AActor> DesiredClass);
 
 private:
-	/** if true show debug info for 'ShowDebugTargetActor', otherwise for Camera Viewtarget */
-	UPROPERTY(Transient)
-	bool bShowDebugForReticleTarget;
-
 	/** Class filter for selecting 'ShowDebugTargetActor' when 'bShowDebugForReticleTarget' is true. */
-	UPROPERTY(Transient)
+	UPROPERTY()
 	TSubclassOf<AActor> ShowDebugTargetDesiredClass;
 
 	/** Show Debug Actor used if 'bShowDebugForReticleTarget' is true, only updated if trace from reticle hit a new Actor of class 'ShowDebugTargetDesiredClass'*/
-	UPROPERTY(Transient)
+	UPROPERTY()
 	AActor* ShowDebugTargetActor;
 
 public:
diff --git a/Engine/Source/Runtime/Engine/Classes/GameFramework/Pawn.h b/Engine/Source/Runtime/Engine/Classes/GameFramework/Pawn.h
--- a/Engine/Source/Runtime/Engine/Classes/GameFramework/Pawn.h
+++ b/Engine/Source/Runtime/Engine/Classes/GameFramework/Pawn.h
@@ -332,7 +332,7 @@ protected:
 	virtual void DestroyPlayerInputComponent();
 
 	/** Allows a Pawn to set up custom input bindings. Called upon possession by a PlayerController, using the InputComponent created by CreatePlayerInputComponent(). */
-	virtual void SetupPlayerInputComponent(UInputComponent* InInputComponent) { /* No bindings by default.*/ }
+	virtual void SetupPlayerInputComponent(UInputComponent* PlayerInputComponent) { /* No bindings by default.*/ }
 
 public:
 	/**
diff --git a/Engine/Source/Runtime/Engine/Classes/GameFramework/PlayerController.h b/Engine/Source/Runtime/Engine/Classes/GameFramework/PlayerController.h
--- a/Engine/Source/Runtime/Engine/Classes/GameFramework/PlayerController.h
+++ b/Engine/Source/Runtime/Engine/Classes/GameFramework/PlayerController.h
@@ -230,11 +230,11 @@ class ENGINE_API APlayerController : public AController
 	int32 ClientCap;
 	
 	/** Object that manages "cheat" commands.  Not instantiated in shipping builds. */
-	UPROPERTY(transient)
+	UPROPERTY(transient, BlueprintReadOnly, Category=PlayerController)
 	class UCheatManager* CheatManager;
 	
 	/** class of my CheatManager. */
-	UPROPERTY()
+	UPROPERTY(EditDefaultsOnly, BlueprintReadOnly, Category=PlayerController)
 	TSubclassOf<class UCheatManager> CheatClass;
 
 	/** Object that manages player input. */
diff --git a/Engine/Source/Runtime/Engine/Private/Actor.cpp b/Engine/Source/Runtime/Engine/Private/Actor.cpp
--- a/Engine/Source/Runtime/Engine/Private/Actor.cpp
+++ b/Engine/Source/Runtime/Engine/Private/Actor.cpp
@@ -752,6 +752,16 @@ bool AActor::IsActorTickEnabled() const
 	return PrimaryActorTick.IsTickFunctionEnabled();
 }
 
+void AActor::SetActorTickInterval(float TickInterval)
+{
+	PrimaryActorTick.TickInterval = TickInterval;
+}
+
+float AActor::GetActorTickInterval() const
+{
+	return PrimaryActorTick.TickInterval;
+}
+
 bool AActor::Rename( const TCHAR* InName, UObject* NewOuter, ERenameFlags Flags )
 {
 	if (NewOuter)
@@ -1829,9 +1839,9 @@ void AActor::RouteEndPlay(const EEndPlayReason::Type EndPlayReason)
 		{
 			EndPlay(EndPlayReason);
 		}
-
-		UninitializeComponents();
 	}
+
+	UninitializeComponents();
 }
 
 void AActor::EndPlay(const EEndPlayReason::Type EndPlayReason)
@@ -2818,15 +2828,18 @@ void AActor::PostActorConstruction()
 		if (!IsPendingKill())
 		{
 			PostInitializeComponents();
-			if (!bActorInitialized && !IsPendingKill())
+			if (!IsPendingKill())
 			{
-				UE_LOG(LogActor, Fatal, TEXT("%s failed to route PostInitializeComponents.  Please call Super::PostInitializeComponents() in your <className>::PostInitializeComponents() function. "), *GetFullName());
-			}
+				if (!bActorInitialized)
+				{
+					UE_LOG(LogActor, Fatal, TEXT("%s failed to route PostInitializeComponents.  Please call Super::PostInitializeComponents() in your <className>::PostInitializeComponents() function. "), *GetFullName());
+				}
 
-			if (World->HasBegunPlay() && !deferBeginPlayAndUpdateOverlaps)
-			{
-				SCOPE_CYCLE_COUNTER(STAT_ActorBeginPlay);
-				BeginPlay();
+				if (World->HasBegunPlay() && !deferBeginPlayAndUpdateOverlaps)
+				{
+					SCOPE_CYCLE_COUNTER(STAT_ActorBeginPlay);
+					BeginPlay();
+				}
 			}
 		}
 	}
diff --git a/Engine/Source/Runtime/Engine/Private/Character.cpp b/Engine/Source/Runtime/Engine/Private/Character.cpp
--- a/Engine/Source/Runtime/Engine/Private/Character.cpp
+++ b/Engine/Source/Runtime/Engine/Private/Character.cpp
@@ -154,9 +154,9 @@ UPawnMovementComponent* ACharacter::GetMovementComponent() const
 }
 
 
-void ACharacter::SetupPlayerInputComponent(UInputComponent* InInputComponent)
+void ACharacter::SetupPlayerInputComponent(UInputComponent* PlayerInputComponent)
 {
-	check(InInputComponent);
+	check(PlayerInputComponent);
 }
 
 
diff --git a/Engine/Source/Runtime/Engine/Private/CheatManager.cpp b/Engine/Source/Runtime/Engine/Private/CheatManager.cpp
--- a/Engine/Source/Runtime/Engine/Private/CheatManager.cpp
+++ b/Engine/Source/Runtime/Engine/Private/CheatManager.cpp
@@ -615,7 +615,7 @@ void UCheatManager::DisableDebugCamera()
 
 void UCheatManager::InitCheatManager() 
 {
-	// Empty
+	ReceiveInitCheatManager(); //BP Initialization event
 }
 
 void UCheatManager::BeginDestroy()
diff --git a/Engine/Source/Runtime/Engine/Private/Components/ActorComponent.cpp b/Engine/Source/Runtime/Engine/Private/Components/ActorComponent.cpp
--- a/Engine/Source/Runtime/Engine/Private/Components/ActorComponent.cpp
+++ b/Engine/Source/Runtime/Engine/Private/Components/ActorComponent.cpp
@@ -750,6 +750,16 @@ bool UActorComponent::IsComponentTickEnabled() const
 	return PrimaryComponentTick.IsTickFunctionEnabled();
 }
 
+void UActorComponent::SetComponentTickInterval(float TickInterval)
+{
+	PrimaryComponentTick.TickInterval = TickInterval;
+}
+
+float UActorComponent::GetComponentTickInterval() const
+{
+	return PrimaryComponentTick.TickInterval;
+}
+
 static UActorComponent* GTestRegisterComponentTickFunctions = NULL;
 
 void UActorComponent::RegisterComponentTickFunctions(bool bRegister)
diff --git a/Engine/Source/Runtime/Engine/Private/DebugCameraController.cpp b/Engine/Source/Runtime/Engine/Private/DebugCameraController.cpp
--- a/Engine/Source/Runtime/Engine/Private/DebugCameraController.cpp
+++ b/Engine/Source/Runtime/Engine/Private/DebugCameraController.cpp
@@ -135,6 +135,10 @@ void ADebugCameraController::OnFingerMove(ETouchIndex::Type FingerIndex, FVector
 	}
 }
 
+AActor* ADebugCameraController::GetSelectedActor() const
+{
+	return SelectedActor;
+}
 
 void ADebugCameraController::Select( FHitResult const& Hit )
 {
@@ -169,6 +173,9 @@ void ADebugCameraController::Select( FHitResult const& Hit )
 			}
 		}
 	}
+	
+	//BP Event
+	ReceiveOnActorSelected(SelectedActor, Hit.ImpactPoint, Hit.ImpactNormal, Hit);
 }
 
 
@@ -316,6 +323,9 @@ void ADebugCameraController::PostInitializeComponents()
 {
 	Super::PostInitializeComponents();
 
+	GDebugSelectedLightmap = nullptr;
+	GDebugSelectedActor = nullptr; 
+	
 	// if hud is existing, delete it and create new hud for debug camera
 	if ( MyHUD != NULL )
 	{
@@ -374,6 +384,9 @@ void ADebugCameraController::OnActivate( APlayerController* OriginalPC )
 	}
 
 	GetWorld()->AddController(this);
+	
+	//BP Event
+	ReceiveOnActivate(OriginalPC);
 }
 
 
@@ -408,6 +421,9 @@ void ADebugCameraController::OnDeactivate( APlayerController* RestoredPC )
 
 	ChangeState(NAME_Inactive);
 	GetWorld()->RemoveController(this);
+	
+	//BP Event
+	ReceiveOnDeactivate(RestoredPC);
 }
 
 
diff --git a/Engine/Source/Runtime/Engine/Private/HUD.cpp b/Engine/Source/Runtime/Engine/Private/HUD.cpp
--- a/Engine/Source/Runtime/Engine/Private/HUD.cpp
+++ b/Engine/Source/Runtime/Engine/Private/HUD.cpp
@@ -36,6 +36,10 @@ TAutoConsoleVariable<float> GSafeZoneVisualizationAlphaCVar(
 	TEXT("The alpha value of the safe zone overlay (0..1)\n")
 	TEXT(" default: 0.2"));
 
+const FColor AHUD::WhiteColor(255, 255, 255, 255);
+const FColor AHUD::GreenColor(0, 255, 0, 255);
+const FColor AHUD::RedColor(255, 0, 0, 255);
+
 AHUD::AHUD(const FObjectInitializer& ObjectInitializer)
 	: Super(ObjectInitializer)
 {
@@ -44,11 +48,8 @@ AHUD::AHUD(const FObjectInitializer& ObjectInitializer)
 	bHidden = true;
 	bReplicates = false;
 
-	WhiteColor = FColor(255, 255, 255, 255);
-	GreenColor = FColor(0, 255, 0, 255);
-	RedColor = FColor(255, 0, 0, 255);
-
 	bLostFocusPaused = false;
+	bShowHUD = true;
 
 	bCanBeDamaged = false;
 	bEnableDebugTextShadow = false;
diff --git a/Engine/Source/Runtime/Engine/Private/Tests/AutomationCommon.cpp b/Engine/Source/Runtime/Engine/Private/Tests/AutomationCommon.cpp
--- a/Engine/Source/Runtime/Engine/Private/Tests/AutomationCommon.cpp
+++ b/Engine/Source/Runtime/Engine/Private/Tests/AutomationCommon.cpp
@@ -196,7 +196,7 @@ bool FEngineWaitLatentCommand::Update()
 
 bool FStreamAllResourcesLatentCommand::Update()
 {
-	float StartTime = FPlatformTime::Seconds();
+	const float CommandStartTime = FPlatformTime::Seconds();
 
 	int32 StillInFlight = IStreamingManager::Get().StreamAllResources(Duration);
 
@@ -204,11 +204,11 @@ bool FStreamAllResourcesLatentCommand::Update()
 
 	if(StillInFlight)
 	{
-		UE_LOG(LogEngineAutomationLatentCommand, Warning, TEXT("StreamAllResources() waited for %.2fs but %d resources are still in flight."), Time - StartTime, StillInFlight);
+		UE_LOG(LogEngineAutomationLatentCommand, Warning, TEXT("StreamAllResources() waited for %.2fs but %d resources are still in flight."), Time - CommandStartTime, StillInFlight);
 	}
 	else
 	{
-		UE_LOG(LogEngineAutomationLatentCommand, Log, TEXT("StreamAllResources() waited for %.2fs (max duration: %.2f)."), Time - StartTime, Duration);
+		UE_LOG(LogEngineAutomationLatentCommand, Log, TEXT("StreamAllResources() waited for %.2fs (max duration: %.2f)."), Time - CommandStartTime, Duration);
 	}
 
 	return true;
diff --git a/Engine/Source/Runtime/GameMenuBuilder/Private/SGameMenuPageWidget.cpp b/Engine/Source/Runtime/GameMenuBuilder/Private/SGameMenuPageWidget.cpp
--- a/Engine/Source/Runtime/GameMenuBuilder/Private/SGameMenuPageWidget.cpp
+++ b/Engine/Source/Runtime/GameMenuBuilder/Private/SGameMenuPageWidget.cpp
@@ -257,12 +257,12 @@ void SGameMenuPageWidget::Construct(const FArguments& InArgs)
 
 EVisibility SGameMenuPageWidget::GetSubMenuVisiblity() const
 {
-	EVisibility Visibility = EVisibility::Collapsed;
+	EVisibility MenuVisibility = EVisibility::Collapsed;
 	if ( (MenuStyle->LayoutType != GameMenuLayoutType::Single) && ( bMenuHiding == false ) )
 	{
-		Visibility = NextMenu.IsValid() ? EVisibility::Visible : EVisibility::Collapsed;
+		MenuVisibility = NextMenu.IsValid() ? EVisibility::Visible : EVisibility::Collapsed;
 	}
-	return Visibility;
+	return MenuVisibility;
 }
 
 EVisibility SGameMenuPageWidget::GetSlateVisibility() const
diff --git a/Engine/Source/Runtime/GameplayAbilities/Private/Abilities/Tasks/AbilityTask_WaitGameplayTag.cpp b/Engine/Source/Runtime/GameplayAbilities/Private/Abilities/Tasks/AbilityTask_WaitGameplayTag.cpp
--- a/Engine/Source/Runtime/GameplayAbilities/Private/Abilities/Tasks/AbilityTask_WaitGameplayTag.cpp
+++ b/Engine/Source/Runtime/GameplayAbilities/Private/Abilities/Tasks/AbilityTask_WaitGameplayTag.cpp
@@ -39,7 +39,7 @@ void UAbilityTask_WaitGameplayTagAdded::Activate()
 	Super::Activate();
 }
 
-void UAbilityTask_WaitGameplayTagAdded::GameplayTagCallback(const FGameplayTag Tag, int32 NewCount)
+void UAbilityTask_WaitGameplayTagAdded::GameplayTagCallback(const FGameplayTag InTag, int32 NewCount)
 {
 	if (NewCount==1)
 	{
@@ -85,7 +85,7 @@ void UAbilityTask_WaitGameplayTagRemoved::Activate()
 	Super::Activate();
 }
 
-void UAbilityTask_WaitGameplayTagRemoved::GameplayTagCallback(const FGameplayTag Tag, int32 NewCount)
+void UAbilityTask_WaitGameplayTagRemoved::GameplayTagCallback(const FGameplayTag InTag, int32 NewCount)
 {
 	if (NewCount==0)
 	{
diff --git a/Engine/Source/Runtime/GameplayAbilities/Private/GameplayEffectTests.cpp b/Engine/Source/Runtime/GameplayAbilities/Private/GameplayEffectTests.cpp
--- a/Engine/Source/Runtime/GameplayAbilities/Private/GameplayEffectTests.cpp
+++ b/Engine/Source/Runtime/GameplayAbilities/Private/GameplayEffectTests.cpp
@@ -291,10 +291,10 @@ protected:
 	virtual FString GetBeautifiedTestName() const override { return "System.AbilitySystem.GameplayEffects"; }
 	virtual void GetTests(TArray<FString>& OutBeautifiedNames, TArray <FString>& OutTestCommands) const override
 	{
-		for (const FString& TestName : TestFunctionNames)
+		for (const FString& TestFuncName : TestFunctionNames)
 		{
-			OutBeautifiedNames.Add(TestName);
-			OutTestCommands.Add(TestName);
+			OutBeautifiedNames.Add(TestFuncName);
+			OutTestCommands.Add(TestFuncName);
 		}
 	}
 
diff --git a/Engine/Source/Runtime/Online/OnlineFramework/Lobby/Private/LobbyBeaconHost.cpp b/Engine/Source/Runtime/Online/OnlineFramework/Lobby/Private/LobbyBeaconHost.cpp
--- a/Engine/Source/Runtime/Online/OnlineFramework/Lobby/Private/LobbyBeaconHost.cpp
+++ b/Engine/Source/Runtime/Online/OnlineFramework/Lobby/Private/LobbyBeaconHost.cpp
@@ -232,19 +232,19 @@ void ALobbyBeaconHost::ProcessDisconnect(ALobbyBeaconClient* ClientActor)
 	}
 }
 
-void ALobbyBeaconHost::ProcessKickPlayer(ALobbyBeaconClient* Instigator, const FUniqueNetIdRepl& PlayerToKick, const FText& Reason)
+void ALobbyBeaconHost::ProcessKickPlayer(ALobbyBeaconClient* InInstigator, const FUniqueNetIdRepl& PlayerToKick, const FText& Reason)
 {
-	if (Instigator && PlayerToKick.IsValid())
+	if (InInstigator && PlayerToKick.IsValid())
 	{
 		for (AOnlineBeaconClient* ExistingClient : ClientActors)
 		{
-			if (ExistingClient != Instigator)
+			if (ExistingClient != InInstigator)
 			{
 				ALobbyBeaconClient* LBC = Cast<ALobbyBeaconClient>(ExistingClient);
 				if (LBC && LBC->PlayerState && LBC->PlayerState->UniqueId == PlayerToKick)
 				{
 					// Right now the only eligible lobby kick is a party leader telling the server it kicked a party member
-					bool bPartyLeaderKick = (Instigator->PlayerState->UniqueId == LBC->PlayerState->PartyOwnerUniqueId);
+					bool bPartyLeaderKick = (InInstigator->PlayerState->UniqueId == LBC->PlayerState->PartyOwnerUniqueId);
 					if (bPartyLeaderKick)
 					{
 						FText KickReason = NSLOCTEXT("NetworkErrors", "KickedPlayerFromParty", "Kicked from party.");
diff --git a/Engine/Source/Runtime/Online/OnlineSubsystemSteam/Private/SocketsSteam.cpp b/Engine/Source/Runtime/Online/OnlineSubsystemSteam/Private/SocketsSteam.cpp
--- a/Engine/Source/Runtime/Online/OnlineSubsystemSteam/Private/SocketsSteam.cpp
+++ b/Engine/Source/Runtime/Online/OnlineSubsystemSteam/Private/SocketsSteam.cpp
@@ -84,10 +84,10 @@ bool FSocketSteam::HasPendingData(uint32& PendingDataSize)
  * @param		SocketDescription debug description of socket
  * @return		The new (heap-allocated) socket, or NULL if unsuccessful.
  */
-FSocket* FSocketSteam::Accept(const FString& SocketDescription) 
+FSocket* FSocketSteam::Accept(const FString& InSocketDescription) 
 {
 	/** Not supported - connectionless (UDP) only */
-	return NULL;
+	return nullptr;
 }
 
 /**
@@ -98,10 +98,10 @@ FSocket* FSocketSteam::Accept(const FString& SocketDescription)
  *
  * @return		The new (heap-allocated) socket, or NULL if unsuccessful.
  */
-FSocket* FSocketSteam::Accept(FInternetAddr& OutAddr, const FString& SocketDescription) 
+FSocket* FSocketSteam::Accept(FInternetAddr& OutAddr, const FString& InSocketDescription) 
 {
 	/** Not supported - connectionless (UDP) only */
-	return NULL;
+	return nullptr;
 }
 
 /**
diff --git a/Engine/Source/Runtime/Renderer/Private/PostProcess/PostProcessTonemap.cpp b/Engine/Source/Runtime/Renderer/Private/PostProcess/PostProcessTonemap.cpp
--- a/Engine/Source/Runtime/Renderer/Private/PostProcess/PostProcessTonemap.cpp
+++ b/Engine/Source/Runtime/Renderer/Private/PostProcess/PostProcessTonemap.cpp
@@ -1040,12 +1040,12 @@ public:
 
 		{
 			static TConsoleVariableData<int32>* CVar709 = IConsoleManager::Get().FindTConsoleVariableDataInt(TEXT("r.Tonemapper709"));
-			static TConsoleVariableData<float>* CVarGamma = IConsoleManager::Get().FindTConsoleVariableDataFloat(TEXT("r.TonemapperGamma"));
+			static TConsoleVariableData<float>* CVarTonemapperGamma = IConsoleManager::Get().FindTConsoleVariableDataFloat(TEXT("r.TonemapperGamma"));
 			static TConsoleVariableData<int32>* CVar2084 = IConsoleManager::Get().FindTConsoleVariableDataInt(TEXT("r.Tonemapper2084"));
 
 			int32 Rec709 = CVar709->GetValueOnRenderThread();
 			int32 ST2084 = CVar2084->GetValueOnRenderThread();
-			float Gamma = CVarGamma->GetValueOnRenderThread();
+			float Gamma = CVarTonemapperGamma->GetValueOnRenderThread();
 
 			if (PLATFORM_APPLE && Gamma == 0.0f)
 			{
diff --git a/Engine/Source/Runtime/WebBrowser/WebBrowser.Build.cs b/Engine/Source/Runtime/WebBrowser/WebBrowser.Build.cs
--- a/Engine/Source/Runtime/WebBrowser/WebBrowser.Build.cs
+++ b/Engine/Source/Runtime/WebBrowser/WebBrowser.Build.cs
@@ -45,5 +45,7 @@ public class WebBrowser : ModuleRules
 				RuntimeDependencies.Add(new RuntimeDependency("$(EngineDir)/Binaries/" + Target.Platform.ToString() + "/UnrealCEFSubProcess.exe"));
 			}
 		}
+
+        bEnableShadowVariableWarnings = false;
 	}
 }
diff --git a/Templates/FP_FirstPerson/Source/FP_FirstPerson/FP_FirstPersonCharacter.cpp b/Templates/FP_FirstPerson/Source/FP_FirstPerson/FP_FirstPersonCharacter.cpp
--- a/Templates/FP_FirstPerson/Source/FP_FirstPerson/FP_FirstPersonCharacter.cpp
+++ b/Templates/FP_FirstPerson/Source/FP_FirstPerson/FP_FirstPersonCharacter.cpp
@@ -54,29 +54,29 @@ AFP_FirstPersonCharacter::AFP_FirstPersonCharacter()
 //////////////////////////////////////////////////////////////////////////
 // Input
 
-void AFP_FirstPersonCharacter::SetupPlayerInputComponent(class UInputComponent* InputComponent)
+void AFP_FirstPersonCharacter::SetupPlayerInputComponent(class UInputComponent* PlayerInputComponent)
 {
 	// set up gameplay key bindings
-	check(InputComponent);
+	check(PlayerInputComponent);
 
 	// Bind jump events
-	InputComponent->BindAction("Jump", IE_Pressed, this, &ACharacter::Jump);
-	InputComponent->BindAction("Jump", IE_Released, this, &ACharacter::StopJumping);
+	PlayerInputComponent->BindAction("Jump", IE_Pressed, this, &ACharacter::Jump);
+	PlayerInputComponent->BindAction("Jump", IE_Released, this, &ACharacter::StopJumping);
 	
-	InputComponent->BindAction("Fire", IE_Pressed, this, &AFP_FirstPersonCharacter::OnFire);
-	TryEnableTouchscreenMovement(InputComponent);
+	PlayerInputComponent->BindAction("Fire", IE_Pressed, this, &AFP_FirstPersonCharacter::OnFire);
+	TryEnableTouchscreenMovement(PlayerInputComponent);
 
 	// Bind movement events
-	InputComponent->BindAxis("MoveForward", this, &AFP_FirstPersonCharacter::MoveForward);
-	InputComponent->BindAxis("MoveRight", this, &AFP_FirstPersonCharacter::MoveRight);
+	PlayerInputComponent->BindAxis("MoveForward", this, &AFP_FirstPersonCharacter::MoveForward);
+	PlayerInputComponent->BindAxis("MoveRight", this, &AFP_FirstPersonCharacter::MoveRight);
 	
 	// We have 2 versions of the rotation bindings to handle different kinds of devices differently
 	// "turn" handles devices that provide an absolute delta, such as a mouse.
 	// "turnrate" is for devices that we choose to treat as a rate of change, such as an analog joystick
-	InputComponent->BindAxis("Turn", this, &APawn::AddControllerYawInput);
-	InputComponent->BindAxis("TurnRate", this, &AFP_FirstPersonCharacter::TurnAtRate);
-	InputComponent->BindAxis("LookUp", this, &APawn::AddControllerPitchInput);
-	InputComponent->BindAxis("LookUpRate", this, &AFP_FirstPersonCharacter::LookUpAtRate);
+	PlayerInputComponent->BindAxis("Turn", this, &APawn::AddControllerYawInput);
+	PlayerInputComponent->BindAxis("TurnRate", this, &AFP_FirstPersonCharacter::TurnAtRate);
+	PlayerInputComponent->BindAxis("LookUp", this, &APawn::AddControllerPitchInput);
+	PlayerInputComponent->BindAxis("LookUpRate", this, &AFP_FirstPersonCharacter::LookUpAtRate);
 }
 
 void AFP_FirstPersonCharacter::OnFire()
@@ -249,9 +249,9 @@ FHitResult AFP_FirstPersonCharacter::WeaponTrace(const FVector& StartTrace, cons
 	return Hit;
 }
 
-void AFP_FirstPersonCharacter::TryEnableTouchscreenMovement(UInputComponent* InputComponent)
+void AFP_FirstPersonCharacter::TryEnableTouchscreenMovement(UInputComponent* PlayerInputComponent)
 {
-	InputComponent->BindTouch(EInputEvent::IE_Pressed, this, &AFP_FirstPersonCharacter::BeginTouch);
-	InputComponent->BindTouch(EInputEvent::IE_Released, this, &AFP_FirstPersonCharacter::EndTouch);
-	InputComponent->BindTouch(EInputEvent::IE_Repeat, this, &AFP_FirstPersonCharacter::TouchUpdate);	
+	PlayerInputComponent->BindTouch(EInputEvent::IE_Pressed, this, &AFP_FirstPersonCharacter::BeginTouch);
+	PlayerInputComponent->BindTouch(EInputEvent::IE_Released, this, &AFP_FirstPersonCharacter::EndTouch);
+	PlayerInputComponent->BindTouch(EInputEvent::IE_Repeat, this, &AFP_FirstPersonCharacter::TouchUpdate);	
 }
diff --git a/Templates/TP_2DSideScroller/Source/TP_2DSideScroller/TP_2DSideScrollerCharacter.cpp b/Templates/TP_2DSideScroller/Source/TP_2DSideScroller/TP_2DSideScrollerCharacter.cpp
--- a/Templates/TP_2DSideScroller/Source/TP_2DSideScroller/TP_2DSideScrollerCharacter.cpp
+++ b/Templates/TP_2DSideScroller/Source/TP_2DSideScroller/TP_2DSideScrollerCharacter.cpp
@@ -116,15 +116,15 @@ void ATP_2DSideScrollerCharacter::Tick(float DeltaSeconds)
 //////////////////////////////////////////////////////////////////////////
 // Input
 
-void ATP_2DSideScrollerCharacter::SetupPlayerInputComponent(class UInputComponent* InputComponent)
+void ATP_2DSideScrollerCharacter::SetupPlayerInputComponent(class UInputComponent* PlayerInputComponent)
 {
 	// Note: the 'Jump' action and the 'MoveRight' axis are bound to actual keys/buttons/sticks in DefaultInput.ini (editable from Project Settings..Input)
-	InputComponent->BindAction("Jump", IE_Pressed, this, &ACharacter::Jump);
-	InputComponent->BindAction("Jump", IE_Released, this, &ACharacter::StopJumping);
-	InputComponent->BindAxis("MoveRight", this, &ATP_2DSideScrollerCharacter::MoveRight);
+	PlayerInputComponent->BindAction("Jump", IE_Pressed, this, &ACharacter::Jump);
+	PlayerInputComponent->BindAction("Jump", IE_Released, this, &ACharacter::StopJumping);
+	PlayerInputComponent->BindAxis("MoveRight", this, &ATP_2DSideScrollerCharacter::MoveRight);
 
-	InputComponent->BindTouch(IE_Pressed, this, &ATP_2DSideScrollerCharacter::TouchStarted);
-	InputComponent->BindTouch(IE_Released, this, &ATP_2DSideScrollerCharacter::TouchStopped);
+	PlayerInputComponent->BindTouch(IE_Pressed, this, &ATP_2DSideScrollerCharacter::TouchStarted);
+	PlayerInputComponent->BindTouch(IE_Released, this, &ATP_2DSideScrollerCharacter::TouchStopped);
 }
 
 void ATP_2DSideScrollerCharacter::MoveRight(float Value)
diff --git a/Templates/TP_FirstPerson/Source/TP_FirstPerson/TP_FirstPersonCharacter.cpp b/Templates/TP_FirstPerson/Source/TP_FirstPerson/TP_FirstPersonCharacter.cpp
--- a/Templates/TP_FirstPerson/Source/TP_FirstPerson/TP_FirstPersonCharacter.cpp
+++ b/Templates/TP_FirstPerson/Source/TP_FirstPerson/TP_FirstPersonCharacter.cpp
@@ -64,30 +64,30 @@ void ATP_FirstPersonCharacter::BeginPlay()
 //////////////////////////////////////////////////////////////////////////
 // Input
 
-void ATP_FirstPersonCharacter::SetupPlayerInputComponent(class UInputComponent* InputComponent)
+void ATP_FirstPersonCharacter::SetupPlayerInputComponent(class UInputComponent* PlayerInputComponent)
 {
 	// set up gameplay key bindings
-	check(InputComponent);
+	check(PlayerInputComponent);
 
-	InputComponent->BindAction("Jump", IE_Pressed, this, &ACharacter::Jump);
-	InputComponent->BindAction("Jump", IE_Released, this, &ACharacter::StopJumping);
+	PlayerInputComponent->BindAction("Jump", IE_Pressed, this, &ACharacter::Jump);
+	PlayerInputComponent->BindAction("Jump", IE_Released, this, &ACharacter::StopJumping);
 
 	//InputComponent->BindTouch(EInputEvent::IE_Pressed, this, &ATP_FirstPersonCharacter::TouchStarted);
-	if (EnableTouchscreenMovement(InputComponent) == false)
+	if (EnableTouchscreenMovement(PlayerInputComponent) == false)
 	{
-		InputComponent->BindAction("Fire", IE_Pressed, this, &ATP_FirstPersonCharacter::OnFire);
+		PlayerInputComponent->BindAction("Fire", IE_Pressed, this, &ATP_FirstPersonCharacter::OnFire);
 	}
 
-	InputComponent->BindAxis("MoveForward", this, &ATP_FirstPersonCharacter::MoveForward);
-	InputComponent->BindAxis("MoveRight", this, &ATP_FirstPersonCharacter::MoveRight);
+	PlayerInputComponent->BindAxis("MoveForward", this, &ATP_FirstPersonCharacter::MoveForward);
+	PlayerInputComponent->BindAxis("MoveRight", this, &ATP_FirstPersonCharacter::MoveRight);
 
 	// We have 2 versions of the rotation bindings to handle different kinds of devices differently
 	// "turn" handles devices that provide an absolute delta, such as a mouse.
 	// "turnrate" is for devices that we choose to treat as a rate of change, such as an analog joystick
-	InputComponent->BindAxis("Turn", this, &APawn::AddControllerYawInput);
-	InputComponent->BindAxis("TurnRate", this, &ATP_FirstPersonCharacter::TurnAtRate);
-	InputComponent->BindAxis("LookUp", this, &APawn::AddControllerPitchInput);
-	InputComponent->BindAxis("LookUpRate", this, &ATP_FirstPersonCharacter::LookUpAtRate);
+	PlayerInputComponent->BindAxis("Turn", this, &APawn::AddControllerYawInput);
+	PlayerInputComponent->BindAxis("TurnRate", this, &ATP_FirstPersonCharacter::TurnAtRate);
+	PlayerInputComponent->BindAxis("LookUp", this, &APawn::AddControllerPitchInput);
+	PlayerInputComponent->BindAxis("LookUpRate", this, &ATP_FirstPersonCharacter::LookUpAtRate);
 }
 
 void ATP_FirstPersonCharacter::OnFire()
@@ -216,15 +216,15 @@ void ATP_FirstPersonCharacter::LookUpAtRate(float Rate)
 	AddControllerPitchInput(Rate * BaseLookUpRate * GetWorld()->GetDeltaSeconds());
 }
 
-bool ATP_FirstPersonCharacter::EnableTouchscreenMovement(class UInputComponent* InputComponent)
+bool ATP_FirstPersonCharacter::EnableTouchscreenMovement(class UInputComponent* PlayerInputComponent)
 {
 	bool bResult = false;
 	if (FPlatformMisc::GetUseVirtualJoysticks() || GetDefault<UInputSettings>()->bUseMouseForTouch)
 	{
 		bResult = true;
-		InputComponent->BindTouch(EInputEvent::IE_Pressed, this, &ATP_FirstPersonCharacter::BeginTouch);
-		InputComponent->BindTouch(EInputEvent::IE_Released, this, &ATP_FirstPersonCharacter::EndTouch);
-		InputComponent->BindTouch(EInputEvent::IE_Repeat, this, &ATP_FirstPersonCharacter::TouchUpdate);
+		PlayerInputComponent->BindTouch(EInputEvent::IE_Pressed, this, &ATP_FirstPersonCharacter::BeginTouch);
+		PlayerInputComponent->BindTouch(EInputEvent::IE_Released, this, &ATP_FirstPersonCharacter::EndTouch);
+		PlayerInputComponent->BindTouch(EInputEvent::IE_Repeat, this, &ATP_FirstPersonCharacter::TouchUpdate);
 	}
 	return bResult;
 }
diff --git a/Templates/TP_Flying/Source/TP_Flying/TP_FlyingPawn.cpp b/Templates/TP_Flying/Source/TP_Flying/TP_FlyingPawn.cpp
--- a/Templates/TP_Flying/Source/TP_Flying/TP_FlyingPawn.cpp
+++ b/Templates/TP_Flying/Source/TP_Flying/TP_FlyingPawn.cpp
@@ -72,14 +72,14 @@ void ATP_FlyingPawn::NotifyHit(class UPrimitiveComponent* MyComp, class AActor*
 }
 
 
-void ATP_FlyingPawn::SetupPlayerInputComponent(class UInputComponent* InputComponent)
+void ATP_FlyingPawn::SetupPlayerInputComponent(class UInputComponent* PlayerInputComponent)
 {
-	check(InputComponent);
+	check(PlayerInputComponent);
 
 	// Bind our control axis' to callback functions
-	InputComponent->BindAxis("Thrust", this, &ATP_FlyingPawn::ThrustInput);
-	InputComponent->BindAxis("MoveUp", this, &ATP_FlyingPawn::MoveUpInput);
-	InputComponent->BindAxis("MoveRight", this, &ATP_FlyingPawn::MoveRightInput);
+	PlayerInputComponent->BindAxis("Thrust", this, &ATP_FlyingPawn::ThrustInput);
+	PlayerInputComponent->BindAxis("MoveUp", this, &ATP_FlyingPawn::MoveUpInput);
+	PlayerInputComponent->BindAxis("MoveRight", this, &ATP_FlyingPawn::MoveRightInput);
 }
 
 void ATP_FlyingPawn::ThrustInput(float Val)
diff --git a/Templates/TP_Rolling/Source/TP_Rolling/TP_RollingBall.cpp b/Templates/TP_Rolling/Source/TP_Rolling/TP_RollingBall.cpp
--- a/Templates/TP_Rolling/Source/TP_Rolling/TP_RollingBall.cpp
+++ b/Templates/TP_Rolling/Source/TP_Rolling/TP_RollingBall.cpp
@@ -41,13 +41,13 @@ ATP_RollingBall::ATP_RollingBall()
 }
 
 
-void ATP_RollingBall::SetupPlayerInputComponent(class UInputComponent* InputComponent)
+void ATP_RollingBall::SetupPlayerInputComponent(class UInputComponent* PlayerInputComponent)
 {
 	// set up gameplay key bindings
-	InputComponent->BindAxis("MoveRight", this, &ATP_RollingBall::MoveRight);
-	InputComponent->BindAxis("MoveForward", this, &ATP_RollingBall::MoveForward);
+	PlayerInputComponent->BindAxis("MoveRight", this, &ATP_RollingBall::MoveRight);
+	PlayerInputComponent->BindAxis("MoveForward", this, &ATP_RollingBall::MoveForward);
 
-	InputComponent->BindAction("Jump", IE_Pressed, this, &ATP_RollingBall::Jump);
+	PlayerInputComponent->BindAction("Jump", IE_Pressed, this, &ATP_RollingBall::Jump);
 }
 
 void ATP_RollingBall::MoveRight(float Val)
diff --git a/Templates/TP_SideScroller/Source/TP_SideScroller/TP_SideScrollerCharacter.cpp b/Templates/TP_SideScroller/Source/TP_SideScroller/TP_SideScrollerCharacter.cpp
--- a/Templates/TP_SideScroller/Source/TP_SideScroller/TP_SideScrollerCharacter.cpp
+++ b/Templates/TP_SideScroller/Source/TP_SideScroller/TP_SideScrollerCharacter.cpp
@@ -44,15 +44,15 @@ ATP_SideScrollerCharacter::ATP_SideScrollerCharacter()
 //////////////////////////////////////////////////////////////////////////
 // Input
 
-void ATP_SideScrollerCharacter::SetupPlayerInputComponent(class UInputComponent* InputComponent)
+void ATP_SideScrollerCharacter::SetupPlayerInputComponent(class UInputComponent* PlayerInputComponent)
 {
 	// set up gameplay key bindings
-	InputComponent->BindAction("Jump", IE_Pressed, this, &ACharacter::Jump);
-	InputComponent->BindAction("Jump", IE_Released, this, &ACharacter::StopJumping);
-	InputComponent->BindAxis("MoveRight", this, &ATP_SideScrollerCharacter::MoveRight);
+	PlayerInputComponent->BindAction("Jump", IE_Pressed, this, &ACharacter::Jump);
+	PlayerInputComponent->BindAction("Jump", IE_Released, this, &ACharacter::StopJumping);
+	PlayerInputComponent->BindAxis("MoveRight", this, &ATP_SideScrollerCharacter::MoveRight);
 
-	InputComponent->BindTouch(IE_Pressed, this, &ATP_SideScrollerCharacter::TouchStarted);
-	InputComponent->BindTouch(IE_Released, this, &ATP_SideScrollerCharacter::TouchStopped);
+	PlayerInputComponent->BindTouch(IE_Pressed, this, &ATP_SideScrollerCharacter::TouchStarted);
+	PlayerInputComponent->BindTouch(IE_Released, this, &ATP_SideScrollerCharacter::TouchStopped);
 }
 
 void ATP_SideScrollerCharacter::MoveRight(float Value)
diff --git a/Templates/TP_ThirdPerson/Source/TP_ThirdPerson/TP_ThirdPersonCharacter.cpp b/Templates/TP_ThirdPerson/Source/TP_ThirdPerson/TP_ThirdPersonCharacter.cpp
--- a/Templates/TP_ThirdPerson/Source/TP_ThirdPerson/TP_ThirdPersonCharacter.cpp
+++ b/Templates/TP_ThirdPerson/Source/TP_ThirdPerson/TP_ThirdPersonCharacter.cpp
@@ -44,27 +44,27 @@ ATP_ThirdPersonCharacter::ATP_ThirdPersonCharacter()
 //////////////////////////////////////////////////////////////////////////
 // Input
 
-void ATP_ThirdPersonCharacter::SetupPlayerInputComponent(class UInputComponent* InputComponent)
+void ATP_ThirdPersonCharacter::SetupPlayerInputComponent(class UInputComponent* PlayerInputComponent)
 {
 	// Set up gameplay key bindings
-	check(InputComponent);
-	InputComponent->BindAction("Jump", IE_Pressed, this, &ACharacter::Jump);
-	InputComponent->BindAction("Jump", IE_Released, this, &ACharacter::StopJumping);
+	check(PlayerInputComponent);
+	PlayerInputComponent->BindAction("Jump", IE_Pressed, this, &ACharacter::Jump);
+	PlayerInputComponent->BindAction("Jump", IE_Released, this, &ACharacter::StopJumping);
 
-	InputComponent->BindAxis("MoveForward", this, &ATP_ThirdPersonCharacter::MoveForward);
-	InputComponent->BindAxis("MoveRight", this, &ATP_ThirdPersonCharacter::MoveRight);
+	PlayerInputComponent->BindAxis("MoveForward", this, &ATP_ThirdPersonCharacter::MoveForward);
+	PlayerInputComponent->BindAxis("MoveRight", this, &ATP_ThirdPersonCharacter::MoveRight);
 
 	// We have 2 versions of the rotation bindings to handle different kinds of devices differently
 	// "turn" handles devices that provide an absolute delta, such as a mouse.
 	// "turnrate" is for devices that we choose to treat as a rate of change, such as an analog joystick
-	InputComponent->BindAxis("Turn", this, &APawn::AddControllerYawInput);
-	InputComponent->BindAxis("TurnRate", this, &ATP_ThirdPersonCharacter::TurnAtRate);
-	InputComponent->BindAxis("LookUp", this, &APawn::AddControllerPitchInput);
-	InputComponent->BindAxis("LookUpRate", this, &ATP_ThirdPersonCharacter::LookUpAtRate);
+	PlayerInputComponent->BindAxis("Turn", this, &APawn::AddControllerYawInput);
+	PlayerInputComponent->BindAxis("TurnRate", this, &ATP_ThirdPersonCharacter::TurnAtRate);
+	PlayerInputComponent->BindAxis("LookUp", this, &APawn::AddControllerPitchInput);
+	PlayerInputComponent->BindAxis("LookUpRate", this, &ATP_ThirdPersonCharacter::LookUpAtRate);
 
 	// handle touch devices
-	InputComponent->BindTouch(IE_Pressed, this, &ATP_ThirdPersonCharacter::TouchStarted);
-	InputComponent->BindTouch(IE_Released, this, &ATP_ThirdPersonCharacter::TouchStopped);
+	PlayerInputComponent->BindTouch(IE_Pressed, this, &ATP_ThirdPersonCharacter::TouchStarted);
+	PlayerInputComponent->BindTouch(IE_Released, this, &ATP_ThirdPersonCharacter::TouchStopped);
 }
 
 
diff --git a/Templates/TP_TopDown/Source/TP_TopDown/TP_TopDownPlayerController.cpp b/Templates/TP_TopDown/Source/TP_TopDown/TP_TopDownPlayerController.cpp
--- a/Templates/TP_TopDown/Source/TP_TopDown/TP_TopDownPlayerController.cpp
+++ b/Templates/TP_TopDown/Source/TP_TopDown/TP_TopDownPlayerController.cpp
@@ -63,11 +63,11 @@ void ATP_TopDownPlayerController::MoveToTouchLocation(const ETouchIndex::Type Fi
 
 void ATP_TopDownPlayerController::SetNewMoveDestination(const FVector DestLocation)
 {
-	APawn* const Pawn = GetPawn();
-	if (Pawn)
+	APawn* const MyPawn = GetPawn();
+	if (MyPawn)
 	{
 		UNavigationSystem* const NavSys = GetWorld()->GetNavigationSystem();
-		float const Distance = FVector::Dist(DestLocation, Pawn->GetActorLocation());
+		float const Distance = FVector::Dist(DestLocation, MyPawn->GetActorLocation());
 
 		// We need to issue move command only if far enough in order for walk animation to play correctly
 		if (NavSys && (Distance > 120.0f))
diff --git a/Templates/TP_TwinStick/Source/TP_TwinStick/TP_TwinStickPawn.cpp b/Templates/TP_TwinStick/Source/TP_TwinStick/TP_TwinStickPawn.cpp
--- a/Templates/TP_TwinStick/Source/TP_TwinStick/TP_TwinStickPawn.cpp
+++ b/Templates/TP_TwinStick/Source/TP_TwinStick/TP_TwinStickPawn.cpp
@@ -44,15 +44,15 @@ ATP_TwinStickPawn::ATP_TwinStickPawn()
 	bCanFire = true;
 }
 
-void ATP_TwinStickPawn::SetupPlayerInputComponent(class UInputComponent* InputComponent)
+void ATP_TwinStickPawn::SetupPlayerInputComponent(class UInputComponent* PlayerInputComponent)
 {
-	check(InputComponent);
+	check(PlayerInputComponent);
 
 	// set up gameplay key bindings
-	InputComponent->BindAxis(MoveForwardBinding);
-	InputComponent->BindAxis(MoveRightBinding);
-	InputComponent->BindAxis(FireForwardBinding);
-	InputComponent->BindAxis(FireRightBinding);
+	PlayerInputComponent->BindAxis(MoveForwardBinding);
+	PlayerInputComponent->BindAxis(MoveRightBinding);
+	PlayerInputComponent->BindAxis(FireForwardBinding);
+	PlayerInputComponent->BindAxis(FireRightBinding);
 }
 
 void ATP_TwinStickPawn::Tick(float DeltaSeconds)
diff --git a/Templates/TP_Vehicle/Source/TP_Vehicle/TP_VehiclePawn.cpp b/Templates/TP_Vehicle/Source/TP_Vehicle/TP_VehiclePawn.cpp
--- a/Templates/TP_Vehicle/Source/TP_Vehicle/TP_VehiclePawn.cpp
+++ b/Templates/TP_Vehicle/Source/TP_Vehicle/TP_VehiclePawn.cpp
@@ -114,23 +114,23 @@ ATP_VehiclePawn::ATP_VehiclePawn()
 	bInReverseGear = false;
 }
 
-void ATP_VehiclePawn::SetupPlayerInputComponent(class UInputComponent* InputComponent)
+void ATP_VehiclePawn::SetupPlayerInputComponent(class UInputComponent* PlayerInputComponent)
 {
-	Super::SetupPlayerInputComponent(InputComponent);
+	Super::SetupPlayerInputComponent(PlayerInputComponent);
 
 	// set up gameplay key bindings
-	check(InputComponent);
+	check(PlayerInputComponent);
 
-	InputComponent->BindAxis("MoveForward", this, &ATP_VehiclePawn::MoveForward);
-	InputComponent->BindAxis("MoveRight", this, &ATP_VehiclePawn::MoveRight);
-	InputComponent->BindAxis("LookUp");
-	InputComponent->BindAxis("LookRight");
+	PlayerInputComponent->BindAxis("MoveForward", this, &ATP_VehiclePawn::MoveForward);
+	PlayerInputComponent->BindAxis("MoveRight", this, &ATP_VehiclePawn::MoveRight);
+	PlayerInputComponent->BindAxis("LookUp");
+	PlayerInputComponent->BindAxis("LookRight");
 
-	InputComponent->BindAction("Handbrake", IE_Pressed, this, &ATP_VehiclePawn::OnHandbrakePressed);
-	InputComponent->BindAction("Handbrake", IE_Released, this, &ATP_VehiclePawn::OnHandbrakeReleased);
-	InputComponent->BindAction("SwitchCamera", IE_Pressed, this, &ATP_VehiclePawn::OnToggleCamera);
+	PlayerInputComponent->BindAction("Handbrake", IE_Pressed, this, &ATP_VehiclePawn::OnHandbrakePressed);
+	PlayerInputComponent->BindAction("Handbrake", IE_Released, this, &ATP_VehiclePawn::OnHandbrakeReleased);
+	PlayerInputComponent->BindAction("SwitchCamera", IE_Pressed, this, &ATP_VehiclePawn::OnToggleCamera);
 
-	InputComponent->BindAction("ResetVR", IE_Pressed, this, &ATP_VehiclePawn::OnResetVR); 
+	PlayerInputComponent->BindAction("ResetVR", IE_Pressed, this, &ATP_VehiclePawn::OnResetVR); 
 }
 
 void ATP_VehiclePawn::MoveForward(float Val)
diff --git a/Templates/TP_VehicleAdv/Source/TP_VehicleAdv/TP_VehicleAdvPawn.cpp b/Templates/TP_VehicleAdv/Source/TP_VehicleAdv/TP_VehicleAdvPawn.cpp
--- a/Templates/TP_VehicleAdv/Source/TP_VehicleAdv/TP_VehicleAdvPawn.cpp
+++ b/Templates/TP_VehicleAdv/Source/TP_VehicleAdv/TP_VehicleAdvPawn.cpp
@@ -169,23 +169,23 @@ ATP_VehicleAdvPawn::ATP_VehicleAdvPawn()
 	bInReverseGear = false;
 }
 
-void ATP_VehicleAdvPawn::SetupPlayerInputComponent(class UInputComponent* InputComponent)
+void ATP_VehicleAdvPawn::SetupPlayerInputComponent(class UInputComponent* PlayerInputComponent)
 {
-	Super::SetupPlayerInputComponent(InputComponent);
+	Super::SetupPlayerInputComponent(PlayerInputComponent);
 
 	// set up gameplay key bindings
-	check(InputComponent);
+	check(PlayerInputComponent);
 
-	InputComponent->BindAxis("MoveForward", this, &ATP_VehicleAdvPawn::MoveForward);
-	InputComponent->BindAxis("MoveRight", this, &ATP_VehicleAdvPawn::MoveRight);
-	InputComponent->BindAxis(LookUpBinding);
-	InputComponent->BindAxis(LookRightBinding);
+	PlayerInputComponent->BindAxis("MoveForward", this, &ATP_VehicleAdvPawn::MoveForward);
+	PlayerInputComponent->BindAxis("MoveRight", this, &ATP_VehicleAdvPawn::MoveRight);
+	PlayerInputComponent->BindAxis(LookUpBinding);
+	PlayerInputComponent->BindAxis(LookRightBinding);
 
-	InputComponent->BindAction("Handbrake", IE_Pressed, this, &ATP_VehicleAdvPawn::OnHandbrakePressed);
-	InputComponent->BindAction("Handbrake", IE_Released, this, &ATP_VehicleAdvPawn::OnHandbrakeReleased);
-	InputComponent->BindAction("SwitchCamera", IE_Pressed, this, &ATP_VehicleAdvPawn::OnToggleCamera);
+	PlayerInputComponent->BindAction("Handbrake", IE_Pressed, this, &ATP_VehicleAdvPawn::OnHandbrakePressed);
+	PlayerInputComponent->BindAction("Handbrake", IE_Released, this, &ATP_VehicleAdvPawn::OnHandbrakeReleased);
+	PlayerInputComponent->BindAction("SwitchCamera", IE_Pressed, this, &ATP_VehicleAdvPawn::OnToggleCamera);
 
-	InputComponent->BindAction("ResetVR", IE_Pressed, this, &ATP_VehicleAdvPawn::OnResetVR); 
+	PlayerInputComponent->BindAction("ResetVR", IE_Pressed, this, &ATP_VehicleAdvPawn::OnResetVR); 
 }
 
 void ATP_VehicleAdvPawn::MoveForward(float Val)
