Copying //UE4/Dev-Editor to //UE4/Dev-Main (Source: //UE4/Dev-Editor @ 2973866)

diff --git a/Engine/Source/Editor/UnrealEd/Classes/Editor/UnrealEdEngine.h b/Engine/Source/Editor/UnrealEd/Classes/Editor/UnrealEdEngine.h
--- a/Engine/Source/Editor/UnrealEd/Classes/Editor/UnrealEdEngine.h
+++ b/Engine/Source/Editor/UnrealEd/Classes/Editor/UnrealEdEngine.h
@@ -10,6 +10,8 @@
 UENUM()
 enum EPackageNotifyState
 {
+	/** Updating the source control state of the package */
+	NS_Updating,
 	/** The user has been prompted with the balloon taskbar message. */
 	NS_BalloonPrompted,
 	/** The user responded to the balloon task bar message and got the modal prompt to checkout dialog and responded to it. */
@@ -126,10 +128,6 @@ public:
 	UPROPERTY()
 	class AActor* CurrentLODParentActor;
 
-	/** If we have packages that are pending and we should notify the user that they need to be checked out */
-	UPROPERTY()
-	uint32 bNeedToPromptForCheckout:1;
-
 	/** Whether the user needs to be prompted about a package being saved with an engine version newer than the current one or not */
 	UPROPERTY()
 	uint32 bNeedWarningForPkgEngineVer:1;
@@ -150,6 +148,9 @@ public:
 	UPROPERTY()
 	class UCookOnTheFlyServer* CookServer;
 
+	/** A list of packages dirtied this tick */
+	TArray<TWeakObjectPtr<UPackage>> PackagesDirtiedThisTick;
+
 	/** A mapping of packages to their checkout notify state.  This map only contains dirty packages.  Once packages become clean again, they are removed from the map.*/
 	TMap<TWeakObjectPtr<UPackage>, uint8>	PackageToNotifyState;
 
@@ -222,9 +223,9 @@ public:
 	/** called when a package has has its dirty state updated */
 	void OnPackageDirtyStateUpdated( UPackage* Pkg);
 	/** called when a package's source control state is updated */
-	void OnSourceControlStateUpdated(const FSourceControlOperationRef& SourceControlOp, ECommandResult::Type ResultType, TWeakObjectPtr<UPackage> Package);
+	void OnSourceControlStateUpdated(const FSourceControlOperationRef& SourceControlOp, ECommandResult::Type ResultType, TArray<TWeakObjectPtr<UPackage>> Packages);
 	/** called when a package is automatically checked out from source control */
-	void OnPackageCheckedOut(const FSourceControlOperationRef& SourceControlOp, ECommandResult::Type ResultType, TWeakObjectPtr<UPackage> Package);
+	void OnPackagesCheckedOut(const FSourceControlOperationRef& SourceControlOp, ECommandResult::Type ResultType, TArray<TWeakObjectPtr<UPackage>> Packages);
 	/** caled by FCoreDelegate::PostGarbageCollect */
 	void OnPostGarbageCollect();
 	/** called by color picker change event */
@@ -625,6 +626,11 @@ public:
 	 */
 	void PromptToCheckoutModifiedPackages( bool bPromptAll = false );
 
+	/**
+	 * Displays a toast notification or warning when a package is dirtied, indicating that it needs checking out (or that it cannot be checked out)
+	 */
+	void ShowPackageNotification();
+
 	/**
 	 * Checks to see if there are any packages in the PackageToNotifyState map that are not checked out by the user
 	 *
diff --git a/Engine/Source/Editor/UnrealEd/Classes/Factories/DataTableFactory.h b/Engine/Source/Editor/UnrealEd/Classes/Factories/DataTableFactory.h
--- a/Engine/Source/Editor/UnrealEd/Classes/Factories/DataTableFactory.h
+++ b/Engine/Source/Editor/UnrealEd/Classes/Factories/DataTableFactory.h
@@ -4,8 +4,8 @@
 
 #include "DataTableFactory.generated.h"
 
-UCLASS(hidecategories=Object, MinimalAPI)
-class UDataTableFactory : public UFactory
+UCLASS(hidecategories=Object)
+class UNREALED_API UDataTableFactory : public UFactory
 {
 	GENERATED_UCLASS_BODY()
 
diff --git a/Engine/Source/Editor/UnrealEd/Classes/Factories/FbxSceneImportOptions.h b/Engine/Source/Editor/UnrealEd/Classes/Factories/FbxSceneImportOptions.h
--- a/Engine/Source/Editor/UnrealEd/Classes/Factories/FbxSceneImportOptions.h
+++ b/Engine/Source/Editor/UnrealEd/Classes/Factories/FbxSceneImportOptions.h
@@ -6,9 +6,9 @@
 UENUM()
 enum EFBXSceneOptionsCreateHierarchyType
 {
-	FBXSOCHT_CreateLevelActors UMETA(DisplayName = "Create Level Actors"),
-	FBXSOCHT_CreateActorComponents UMETA(DisplayName = "Create one Actor with Components"),
-	FBXSOCHT_CreateBlueprint UMETA(DisplayName = "Create one Blueprint asset"),
+	FBXSOCHT_CreateLevelActors UMETA(DisplayName = "Create Level Actors", ToolTip = "Create an actor for every node in the fbx hierarchy. No reimport of the hierarchy."),
+	FBXSOCHT_CreateActorComponents UMETA(DisplayName = "Create one Actor with Components", ToolTip = "Create one actor and a component for every node in the fbx hierarchy. No reimport of the hierarchy."),
+	FBXSOCHT_CreateBlueprint UMETA(DisplayName = "Create one Blueprint asset", ToolTip = "Create one blueprint and a component for every node in the fbx hierarchy. Hierarchy can be reimport."),
 	FBXSOCHT_MAX
 };
 //	Fbx_AddToBlueprint UMETA(DisplayName = "Add to an existing Blueprint asset"),
@@ -67,8 +67,8 @@ class UFbxSceneImportOptions : public UObject
 	UPROPERTY(EditAnywhere, config, category = Meshes)
 	uint32 bImportSkeletalMeshLODs : 1;
 
-	/** If importing textures is enabled, this option will cause normal map Y (Green) values to be inverted */
-	UPROPERTY(EditAnywhere, config, category = Material)
+	/** If enabled, this option will cause normal map Y (Green) values to be inverted when importing textures */
+	UPROPERTY(EditAnywhere, config, category = Texture)
 	uint32 bInvertNormalMaps : 1;
 };
 
diff --git a/Engine/Source/Editor/UnrealEd/Private/AssetDeleteModel.cpp b/Engine/Source/Editor/UnrealEd/Private/AssetDeleteModel.cpp
--- a/Engine/Source/Editor/UnrealEd/Private/AssetDeleteModel.cpp
+++ b/Engine/Source/Editor/UnrealEd/Private/AssetDeleteModel.cpp
@@ -326,14 +326,14 @@ bool FAssetDeleteModel::CanReplaceReferencesWith( const FAssetData& InAssetData
 		// Get BP native parent classes
 		static const FName ParentClassName("ParentClass");
 		UClass* OriginalBPParentClass = CastChecked<UBlueprint>(PendingDeletes[0]->GetObject())->ParentClass;
-		const FString* BPClassNameToTest = InAssetData.TagsAndValues.Find(ParentClassName);
+		const FString BPClassNameToTest = InAssetData.GetTagValueRef<FString>(ParentClassName);
 
-		if (BPClassNameToTest && !BPClassNameToTest->IsEmpty())
+		if (!BPClassNameToTest.IsEmpty())
 		{
-			UClass* ParentClassToTest = FindObject<UClass>(ANY_PACKAGE, **BPClassNameToTest);
+			UClass* ParentClassToTest = FindObject<UClass>(ANY_PACKAGE, *BPClassNameToTest);
 			if (!ParentClassToTest)
 			{
-				ParentClassToTest = LoadObject<UClass>(nullptr, **BPClassNameToTest);
+				ParentClassToTest = LoadObject<UClass>(nullptr, *BPClassNameToTest);
 			}
 
 			UClass* NativeParentClassToReplace = FBlueprintEditorUtils::FindFirstNativeClass(OriginalBPParentClass);
diff --git a/Engine/Source/Editor/UnrealEd/Private/AssetThumbnail.cpp b/Engine/Source/Editor/UnrealEd/Private/AssetThumbnail.cpp
--- a/Engine/Source/Editor/UnrealEd/Private/AssetThumbnail.cpp
+++ b/Engine/Source/Editor/UnrealEd/Private/AssetThumbnail.cpp
@@ -413,7 +413,7 @@ private:
 
 	const FSlateBrush* GetClassIconBrush() const
 	{
-		return FClassIconFinder::FindIconForClass(ThumbnailClass.Get());
+		return FSlateIconFinder::FindIconBrushForClass(ThumbnailClass.Get());
 	}
 
 	FMargin GetClassIconPadding() const
diff --git a/Engine/Source/Editor/UnrealEd/Private/ClassIconFinder.cpp b/Engine/Source/Editor/UnrealEd/Private/ClassIconFinder.cpp
--- a/Engine/Source/Editor/UnrealEd/Private/ClassIconFinder.cpp
+++ b/Engine/Source/Editor/UnrealEd/Private/ClassIconFinder.cpp
@@ -6,38 +6,11 @@
 #include "ClassIconFinder.h"
 #include "AssetData.h"
 
-TArray < const ISlateStyle* > FClassIconFinder::Styles;
-
-void FClassIconFinder::RegisterIconSource(const ISlateStyle* StyleSet)
-{
-	Styles.AddUnique( StyleSet );
-}
-
-void FClassIconFinder::UnregisterIconSource(const ISlateStyle* StyleSet)
-{
-	Styles.Remove( StyleSet );
-}
-
-const FSlateBrush* FClassIconFinder::LookupBrush(FName IconName)
-{
-	const FSlateBrush* IconBrush = nullptr;
-
-	for ( const ISlateStyle* Style : Styles )
-	{
-		IconBrush = Style->GetOptionalBrush(IconName, nullptr, nullptr);
-		if ( IconBrush )
-		{
-			return IconBrush;
-		}
-	}
-
-	return FEditorStyle::GetOptionalBrush(IconName, nullptr, nullptr);
-}
-
 const FSlateBrush* FClassIconFinder::FindIconForActors(const TArray< TWeakObjectPtr<AActor> >& InActors, UClass*& CommonBaseClass)
 {
 	// Get the common base class of the selected actors
-	const FSlateBrush* CommonIcon = nullptr;
+	FSlateIcon CommonIcon;
+
 	for( int32 ActorIdx = 0; ActorIdx < InActors.Num(); ++ActorIdx )
 	{
 		TWeakObjectPtr<AActor> Actor = InActors[ActorIdx];
@@ -52,31 +25,26 @@ const FSlateBrush* FClassIconFinder::FindIconForActors(const TArray< TWeakObject
 			CommonBaseClass = CommonBaseClass->GetSuperClass();
 		}
 
-		const FSlateBrush* ActorIcon = FindIconForActor(Actor);
+		FSlateIcon ActorIcon = FindSlateIconForActor(Actor);
 
-		if (!CommonIcon)
+		if (!CommonIcon.IsSet())
 		{
 			CommonIcon = ActorIcon;
 		}
+
 		if (CommonIcon != ActorIcon)
 		{
-			CommonIcon = FindIconForClass(CommonBaseClass);
+			CommonIcon = FSlateIconFinder::FindIconForClass(CommonBaseClass);
 		}
 	}
 
-	return CommonIcon;
+	return CommonIcon.GetOptionalIcon();
 }
 
-const FSlateBrush* FClassIconFinder::FindIconForActor( const TWeakObjectPtr<AActor>& InActor )
-{
-	return FClassIconFinder::LookupBrush( FindIconNameForActor(InActor) );
-}
-
-FName FClassIconFinder::FindIconNameForActor( const TWeakObjectPtr<AActor>& InActor )
+FSlateIcon FClassIconFinder::FindSlateIconForActor( const TWeakObjectPtr<AActor>& InActor )
 {
 	// Actor specific overrides to normal per-class icons
 	AActor* Actor = InActor.Get();
-	FName BrushName = NAME_None;
 
 	if ( Actor )
 	{
@@ -85,47 +53,27 @@ FName FClassIconFinder::FindIconNameForActor( const TWeakObjectPtr<AActor>& InAc
 		{
 			if (Brush_Add == Brush->BrushType)
 			{
-				BrushName = TEXT( "ClassIcon.BrushAdditive" );
+				return FSlateIconFinder::FindIcon("ClassIcon.BrushAdditive");
 			}
 			else if (Brush_Subtract == Brush->BrushType)
 			{
-				BrushName = TEXT( "ClassIcon.BrushSubtractive" );
+				return FSlateIconFinder::FindIcon("ClassIcon.BrushSubtractive");
 			}
 		}
 
 		// Actor didn't specify an icon - fallback on the class icon
-		if ( BrushName.IsNone() )
-		{
-			BrushName = FindIconNameForClass( Actor->GetClass() );
-		}
+		return FSlateIconFinder::FindIconForClass(Actor->GetClass());
 	}
 	else
 	{
 		// If the actor reference is NULL it must have been deleted
-		BrushName = TEXT( "ClassIcon.Deleted" );
+		return FSlateIconFinder::FindIcon("ClassIcon.Deleted");
 	}
-
-	return BrushName;
-}
-
-const FSlateBrush* FClassIconFinder::FindIconForClass(const UClass* InClass, const FName& InDefaultName )
-{
-	return FClassIconFinder::LookupBrush( FindIconNameForClass( InClass, InDefaultName ) );
 }
 
-FName FClassIconFinder::FindIconNameForClass(const UClass* InClass, const FName& InDefaultName )
-{
-	return FindIconNameImpl( InClass, InDefaultName );
-}
-
-const FSlateBrush* FClassIconFinder::FindThumbnailForClass(const UClass* InClass, const FName& InDefaultName )
-{
-	return FClassIconFinder::LookupBrush( FindThumbnailNameForClass( InClass, InDefaultName ) );
-}
-
-FName FClassIconFinder::FindThumbnailNameForClass(const UClass* InClass, const FName& InDefaultName )
+const FSlateBrush* FClassIconFinder::FindIconForActor( const TWeakObjectPtr<AActor>& InActor )
 {
-	return FindIconNameImpl( InClass, InDefaultName, TEXT("ClassThumbnail") );
+	return FindSlateIconForActor(InActor).GetOptionalIcon();
 }
 
 const UClass* FClassIconFinder::GetIconClassForBlueprint(const UBlueprint* InBlueprint)
@@ -179,15 +127,14 @@ const UClass* FClassIconFinder::GetIconClassForAssetData(const FAssetData& InAss
 		}
 
 		// We need to use the asset data to get the parent class as the blueprint may not be loaded
-		const FString* ParentClassNamePtr = InAssetData.TagsAndValues.Find(NativeParentClassTag);
-		if ( !ParentClassNamePtr )
+		FString ParentClassName;
+		if ( !InAssetData.GetTagValue(NativeParentClassTag, ParentClassName) )
 		{
-			ParentClassNamePtr = InAssetData.TagsAndValues.Find(ParentClassTag);
+			InAssetData.GetTagValue(ParentClassTag, ParentClassName);
 		}
-		if ( ParentClassNamePtr && !ParentClassNamePtr->IsEmpty() )
+		if ( !ParentClassName.IsEmpty() )
 		{
 			UObject* Outer = nullptr;
-			FString ParentClassName = *ParentClassNamePtr;
 			ResolveName(Outer, ParentClassName, false, false);
 			return FindObject<UClass>(ANY_PACKAGE, *ParentClassName);
 		}
@@ -196,37 +143,3 @@ const UClass* FClassIconFinder::GetIconClassForAssetData(const FAssetData& InAss
 	// Default to using the class for the asset type
 	return AssetClass;
 }
-
-FName FClassIconFinder::FindIconNameImpl(const UClass* InClass, const FName& InDefaultName, const TCHAR* StyleRoot )
-{
-	FName BrushName;
-	const FSlateBrush* Brush = nullptr;
-
-	if ( InClass )
-	{
-		// walk up class hierarchy until we find an icon
-		const UClass* CurrentClass = InClass;
-		while( !Brush && CurrentClass )
-		{
-			BrushName = *FString::Printf( TEXT( "%s.%s" ), StyleRoot, *CurrentClass->GetName() );
-			Brush = FClassIconFinder::LookupBrush( BrushName );
-			CurrentClass = CurrentClass->GetSuperClass();
-		}
-	}
-
-	if( !Brush )
-	{
-		// If we didn't supply an override name for the default icon use default class icon.
-		if( InDefaultName.IsNone() )
-		{
-			BrushName = *FString::Printf( TEXT( "%s.Default" ), StyleRoot);
-		}
-		else
-		{
-			BrushName = InDefaultName;
-		}
-	}
-
-	return BrushName;
-}
-
diff --git a/Engine/Source/Editor/UnrealEd/Private/ComponentTypeRegistry.cpp b/Engine/Source/Editor/UnrealEd/Private/ComponentTypeRegistry.cpp
--- a/Engine/Source/Editor/UnrealEd/Private/ComponentTypeRegistry.cpp
+++ b/Engine/Source/Editor/UnrealEd/Private/ComponentTypeRegistry.cpp
@@ -366,18 +366,13 @@ void FComponentTypeRegistryData::Tick(float)
 		for (auto Asset : PendingAssetData)
 		{
 			const FName BPParentClassName(GET_MEMBER_NAME_CHECKED(UBlueprint, ParentClass));
-
-			bool FoundBPNotify = false;
-			for (auto TagIt = Asset.TagsAndValues.CreateConstIterator(); TagIt; ++TagIt)
+			const FString TagValue = Asset.GetTagValueRef<FString>(BPParentClassName);
+			const FString ObjectPath = FPackageName::ExportTextPathToObjectPath(TagValue);
+			FName ObjectName = FName(*FPackageName::ObjectPathToObjectName(ObjectPath));
+			if (DerivedClassNames.Contains(ObjectName))
 			{
-				FString TagValue = Asset.TagsAndValues.FindRef(BPParentClassName);
-				const FString ObjectPath = FPackageName::ExportTextPathToObjectPath(TagValue);
-				FName ObjectName = FName(*FPackageName::ObjectPathToObjectName(ObjectPath));
-				if (DerivedClassNames.Contains(ObjectName))
-				{
-					bRequiresRefresh = true;
-					break;
-				}
+				bRequiresRefresh = true;
+				break;
 			}
 		}
 		PendingAssetData.Empty();
diff --git a/Engine/Source/Editor/UnrealEd/Private/Dialogs/SDeleteAssetsDialog.cpp b/Engine/Source/Editor/UnrealEd/Private/Dialogs/SDeleteAssetsDialog.cpp
--- a/Engine/Source/Editor/UnrealEd/Private/Dialogs/SDeleteAssetsDialog.cpp
+++ b/Engine/Source/Editor/UnrealEd/Private/Dialogs/SDeleteAssetsDialog.cpp
@@ -64,7 +64,7 @@ public:
 			FFormatNamedArguments Args;
 			Args.Add( TEXT( "AssetCount" ), FText::AsNumber( Item->RemainingDiskReferences ) );
 
-			FText OnDiskCountText = Item->RemainingDiskReferences > 1 ? FText::Format( LOCTEXT( "OnDiskReferences", "{AssetCount} References" ), Args ) : FText::Format( LOCTEXT( "OnDiskReference", "{AssetCount} Reference" ), Args );
+			FText OnDiskCountText = Item->RemainingDiskReferences > 1 ? FText::Format( LOCTEXT( "OnDiskAssetReferences", "{AssetCount} References" ), Args ) : FText::Format( LOCTEXT( "OnDiskAssetReference", "{AssetCount} Reference" ), Args );
 
 			return SNew( STextBlock )
 				.Text( OnDiskCountText )
@@ -461,7 +461,7 @@ FText SDeleteAssetsDialog::GetDeleteSourceContentTooltip() const
 	if (NumPrinted < TotalCount)
 	{
 		Args.Add(FText::AsNumber(TotalCount - NumPrinted));
-		RootText = LOCTEXT("DeleteSourceFiles_Tooltip", "When checked, the following source content files will also be deleted along with the assets:\n\n{0}\n... and {1} more.");
+		RootText = LOCTEXT("DeleteSourceFilesAndMore_Tooltip", "When checked, the following source content files will also be deleted along with the assets:\n\n{0}\n... and {1} more.");
 	}
 	else
 	{
@@ -913,8 +913,8 @@ TSharedPtr<SWidget> SDeleteAssetsDialog::OnGetAssetContextMenu( const TArray<FAs
 	MenuBuilder.BeginSection("AssetOptions", LOCTEXT("AssetOptionsText", "Asset Options"));
 	{
 		MenuBuilder.AddMenuEntry( FGenericCommands::Get().Delete, NAME_None,
-			LOCTEXT( "Delete", "Add to Pending Deletes" ),
-			LOCTEXT( "DeleteTooltip", "Adds the selected assets to the list of pending deleted assets." )
+			LOCTEXT( "AddPendingDelete", "Add to Pending Deletes" ),
+			LOCTEXT( "AddPendingDeleteTooltip", "Adds the selected assets to the list of pending deleted assets." )
 			);
 	}
 	MenuBuilder.EndSection();
diff --git a/Engine/Source/Editor/UnrealEd/Private/DistanceFieldBuildNotification.cpp b/Engine/Source/Editor/UnrealEd/Private/DistanceFieldBuildNotification.cpp
--- a/Engine/Source/Editor/UnrealEd/Private/DistanceFieldBuildNotification.cpp
+++ b/Engine/Source/Editor/UnrealEd/Private/DistanceFieldBuildNotification.cpp
@@ -89,7 +89,7 @@ void FDistanceFieldBuildNotificationImpl::DistanceFieldBuildFinished()
 
 	if (NotificationItem.IsValid())
 	{
-		NotificationItem->SetText( NSLOCTEXT("DistanceFieldBuild", "DistanceFieldBuildInProgress", "Finished building Distance Fields!") );
+		NotificationItem->SetText( NSLOCTEXT("DistanceFieldBuild", "DistanceFieldBuildFinished", "Finished building Distance Fields!") );
 		NotificationItem->SetCompletionState(SNotificationItem::CS_Success);
 		NotificationItem->ExpireAndFadeout();
 
diff --git a/Engine/Source/Editor/UnrealEd/Private/EdGraphNode_Comment.cpp b/Engine/Source/Editor/UnrealEd/Private/EdGraphNode_Comment.cpp
--- a/Engine/Source/Editor/UnrealEd/Private/EdGraphNode_Comment.cpp
+++ b/Engine/Source/Editor/UnrealEd/Private/EdGraphNode_Comment.cpp
@@ -71,10 +71,12 @@ FString UEdGraphNode_Comment::GetDocumentationExcerptName() const
 	return TEXT("UEdGraphNode_Comment");
 }
 
-FName UEdGraphNode_Comment::GetPaletteIcon(FLinearColor& OutColor) const
+FSlateIcon UEdGraphNode_Comment::GetIconAndTint(FLinearColor& OutColor) const
 {
 	OutColor = FLinearColor::White;
-	return TEXT("GraphEditor.Comment_16x");
+	
+	static FSlateIcon Icon("EditorStyle", "GraphEditor.Comment_16x");
+	return Icon;
 }
 
 FText UEdGraphNode_Comment::GetNodeTitle(ENodeTitleType::Type TitleType) const
diff --git a/Engine/Source/Editor/UnrealEd/Private/EditorCommandLineUtils.cpp b/Engine/Source/Editor/UnrealEd/Private/EditorCommandLineUtils.cpp
--- a/Engine/Source/Editor/UnrealEd/Private/EditorCommandLineUtils.cpp
+++ b/Engine/Source/Editor/UnrealEd/Private/EditorCommandLineUtils.cpp
@@ -463,7 +463,7 @@ static void EditorCommandLineUtilsImpl::RunAssetDiffCommand(TSharedPtr<SWindow>
 
 	if (Switches.Contains("echo"))
 	{
-		RaiseEditorMessageBox(LOCTEXT("DiffCommandHelp", "Passed Command Arguments"), 
+		RaiseEditorMessageBox(LOCTEXT("PassedCommandArgs", "Passed Command Arguments"), 
 			FText::FromString(CommandArgs), /*bExitOnClose =*/true);
 		return;
 	}
diff --git a/Engine/Source/Editor/UnrealEd/Private/EditorViewportCommands.cpp b/Engine/Source/Editor/UnrealEd/Private/EditorViewportCommands.cpp
--- a/Engine/Source/Editor/UnrealEd/Private/EditorViewportCommands.cpp
+++ b/Engine/Source/Editor/UnrealEd/Private/EditorViewportCommands.cpp
@@ -4,7 +4,7 @@
 #include "UnrealEd.h"
 #include "EditorViewportCommands.h"
 
-#define LOCTEXT_NAMESPACE ""
+#define LOCTEXT_NAMESPACE "EditorViewportCommands"
 
 void FEditorViewportCommands::RegisterCommands()
 {
diff --git a/Engine/Source/Editor/UnrealEd/Private/Factories/ActorFactory.cpp b/Engine/Source/Editor/UnrealEd/Private/Factories/ActorFactory.cpp
--- a/Engine/Source/Editor/UnrealEd/Private/Factories/ActorFactory.cpp
+++ b/Engine/Source/Editor/UnrealEd/Private/Factories/ActorFactory.cpp
@@ -415,14 +415,14 @@ bool UActorFactoryDeferredDecal::CanCreateActorFrom( const FAssetData& AssetData
 	FAssetData CurrentAssetData = AssetData;
 	while( SanityCheck < 1000 && !CurrentAssetData.GetClass()->IsChildOf( UMaterial::StaticClass() ) )
 	{
-		const FString* ObjectPath = CurrentAssetData.TagsAndValues.Find( TEXT("Parent") );
-		if ( ObjectPath == NULL )
+		const FString ObjectPath = CurrentAssetData.GetTagValueRef<FString>( "Parent" );
+		if ( ObjectPath.IsEmpty() )
 		{
 			OutErrorMsg = NSLOCTEXT("CanCreateActor", "NoMaterial", "A valid material must be specified.");
 			return false;
 		}
 
-		CurrentAssetData = AssetRegistry.GetAssetByObjectPath( **ObjectPath );
+		CurrentAssetData = AssetRegistry.GetAssetByObjectPath( *ObjectPath );
 		if ( !CurrentAssetData.IsValid() )
 		{
 			OutErrorMsg = NSLOCTEXT("CanCreateActor", "NoMaterial", "A valid material must be specified.");
@@ -443,8 +443,8 @@ bool UActorFactoryDeferredDecal::CanCreateActorFrom( const FAssetData& AssetData
 		return false;
 	}
 
-	const FString* MaterialDomain = CurrentAssetData.TagsAndValues.Find( TEXT("MaterialDomain") );
-	if ( MaterialDomain == NULL || *MaterialDomain != TEXT("MD_DeferredDecal") )
+	const FString MaterialDomain = CurrentAssetData.GetTagValueRef<FString>( "MaterialDomain" );
+	if ( MaterialDomain != TEXT("MD_DeferredDecal") )
 	{
 		OutErrorMsg = NSLOCTEXT("CanCreateActor", "NotDecalMaterial", "Only materials with a material domain of DeferredDecal can be specified.");
 		return false;
@@ -611,7 +611,7 @@ UActorFactoryNiagara
 UActorFactoryNiagara::UActorFactoryNiagara(const FObjectInitializer& ObjectInitializer)
 : Super(ObjectInitializer)
 {
-	DisplayName = LOCTEXT("EmitterDisplayName", "NiagaraEffect");
+	DisplayName = LOCTEXT("NiagaraEffectDisplayName", "NiagaraEffect");
 	NewActorClass = ANiagaraActor::StaticClass();
 }
 
@@ -829,14 +829,14 @@ bool UActorFactoryAnimationAsset::CanCreateActorFrom( const FAssetData& AssetDat
 
 	if ( AssetData.GetClass()->IsChildOf( UAnimSequenceBase::StaticClass() ) )
 	{
-		const FString* SkeletonPath = AssetData.TagsAndValues.Find("Skeleton");
-		if ( SkeletonPath == NULL || SkeletonPath->IsEmpty() ) 
+		const FString SkeletonPath = AssetData.GetTagValueRef<FString>("Skeleton");
+		if ( SkeletonPath.IsEmpty() ) 
 		{
 			OutErrorMsg = NSLOCTEXT("CanCreateActor", "NoSkeleton", "UAnimationAssets must have a valid Skeleton.");
 			return false;
 		}
 
-		FAssetData SkeletonData = AssetRegistry.GetAssetByObjectPath( **SkeletonPath );
+		FAssetData SkeletonData = AssetRegistry.GetAssetByObjectPath( *SkeletonPath );
 
 		if ( !SkeletonData.IsValid() )
 		{
@@ -871,14 +871,14 @@ bool UActorFactoryAnimationAsset::CanCreateActorFrom( const FAssetData& AssetDat
 
 	if ( AssetData.GetClass()->IsChildOf( UVertexAnimation::StaticClass() ) )
 	{
-		const FString* BaseSkeletalMeshPath = AssetData.TagsAndValues.Find( TEXT("BaseSkelMesh") );
-		if ( BaseSkeletalMeshPath == NULL || BaseSkeletalMeshPath->IsEmpty() ) 
+		const FString BaseSkeletalMeshPath = AssetData.GetTagValueRef<FString>( "BaseSkelMesh" );
+		if ( BaseSkeletalMeshPath.IsEmpty() ) 
 		{
 			OutErrorMsg = NSLOCTEXT("CanCreateActor", "UVertexAnimationNoSkeleton", "UVertexAnimations must have a valid base skeletal mesh.");
 			return false;
 		}
 
-		SkeletalMeshData = AssetRegistry.GetAssetByObjectPath( **BaseSkeletalMeshPath );
+		SkeletalMeshData = AssetRegistry.GetAssetByObjectPath( *BaseSkeletalMeshPath );
 	}
 
 	if ( !SkeletalMeshData.IsValid() )
@@ -1023,17 +1023,17 @@ bool UActorFactorySkeletalMesh::CanCreateActorFrom( const FAssetData& AssetData,
 
 	if ( !SkeletalMeshData.IsValid() && AssetData.GetClass()->IsChildOf( UAnimBlueprint::StaticClass() ) )
 	{
-		const FString* TargetSkeletonPath = AssetData.TagsAndValues.Find( TEXT("TargetSkeleton") );
-		if ( TargetSkeletonPath == NULL || TargetSkeletonPath->IsEmpty() )
+		const FString TargetSkeletonPath = AssetData.GetTagValueRef<FString>( "TargetSkeleton" );
+		if ( TargetSkeletonPath.IsEmpty() )
 		{
-			OutErrorMsg = NSLOCTEXT("CanCreateActor", "NoTargetSkeleton", "UAnimBlueprints must have a valid Target Skeleton.");
+			OutErrorMsg = NSLOCTEXT("CanCreateActor", "NoAnimBPTargetSkeleton", "UAnimBlueprints must have a valid Target Skeleton.");
 			return false;
 		}
 
-		FAssetData TargetSkeleton = AssetRegistry.GetAssetByObjectPath( **TargetSkeletonPath );
+		FAssetData TargetSkeleton = AssetRegistry.GetAssetByObjectPath( *TargetSkeletonPath );
 		if ( !TargetSkeleton.IsValid() )
 		{
-			OutErrorMsg = NSLOCTEXT("CanCreateActor", "NoTargetSkeleton", "UAnimBlueprints must have a valid Target Skeleton.");
+			OutErrorMsg = NSLOCTEXT("CanCreateActor", "NoAnimBPTargetSkeleton", "UAnimBlueprints must have a valid Target Skeleton.");
 			return false;
 		}
 
@@ -1055,7 +1055,7 @@ bool UActorFactorySkeletalMesh::CanCreateActorFrom( const FAssetData& AssetData,
 		}
 		else
 		{
-			OutErrorMsg = NSLOCTEXT("CanCreateActor", "NoTargetSkeleton", "UAnimBlueprints must have a valid Target Skeleton.");
+			OutErrorMsg = NSLOCTEXT("CanCreateActor", "NoAnimBPTargetSkeleton", "UAnimBlueprints must have a valid Target Skeleton.");
 		}
 	}
 
@@ -1078,7 +1078,7 @@ bool UActorFactorySkeletalMesh::CanCreateActorFrom( const FAssetData& AssetData,
 		}
 		else
 		{
-			OutErrorMsg = NSLOCTEXT("CanCreateActor", "NoTargetSkeleton", "SkeletalMesh must have a valid Target Skeleton.");
+			OutErrorMsg = NSLOCTEXT("CanCreateActor", "NoSkelMeshTargetSkeleton", "SkeletalMesh must have a valid Target Skeleton.");
 		}
 	}
 
@@ -1407,14 +1407,14 @@ bool UActorFactoryBlueprint::CanCreateActorFrom( const FAssetData& AssetData, FT
 		return false;
 	}
 
-	const FString* ParentClassPath = AssetData.TagsAndValues.Find( TEXT("ParentClass") );
-	if ( ParentClassPath == NULL || ParentClassPath->IsEmpty() )
+	const FString ParentClassPath = AssetData.GetTagValueRef<FString>( "ParentClass" );
+	if ( ParentClassPath.IsEmpty() )
 	{
 		OutErrorMsg = NSLOCTEXT("CanCreateActor", "NoBlueprint", "No Blueprint was specified, or the specified Blueprint needs to be compiled.");
 		return false;
 	}
 
-	UClass* ParentClass = FindObject<UClass>(NULL, **ParentClassPath);
+	UClass* ParentClass = FindObject<UClass>(NULL, *ParentClassPath);
 
 	bool bIsActorBased = false;
 	if ( ParentClass != NULL )
@@ -1453,13 +1453,13 @@ AActor* UActorFactoryBlueprint::GetDefaultActor( const FAssetData& AssetData )
 		return NULL;
 	}
 
-	const FString* GeneratedClassPath = AssetData.TagsAndValues.Find("GeneratedClass");
-	if ( GeneratedClassPath == NULL || GeneratedClassPath->IsEmpty() )
+	const FString GeneratedClassPath = AssetData.GetTagValueRef<FString>("GeneratedClass");
+	if ( GeneratedClassPath.IsEmpty() )
 	{
 		return NULL;
 	}
 
-	UClass* GeneratedClass = Cast<UClass>(StaticLoadObject(UClass::StaticClass(), NULL, **GeneratedClassPath, NULL, LOAD_NoWarn, NULL));
+	UClass* GeneratedClass = Cast<UClass>(StaticLoadObject(UClass::StaticClass(), NULL, *GeneratedClassPath, NULL, LOAD_NoWarn, NULL));
 
 	if ( GeneratedClass == NULL )
 	{
diff --git a/Engine/Source/Editor/UnrealEd/Private/Factories/EditorFactories.cpp b/Engine/Source/Editor/UnrealEd/Private/Factories/EditorFactories.cpp
--- a/Engine/Source/Editor/UnrealEd/Private/Factories/EditorFactories.cpp
+++ b/Engine/Source/Editor/UnrealEd/Private/Factories/EditorFactories.cpp
@@ -4686,7 +4686,7 @@ bool UTextureExporterBMP::ExportBinary( UObject* Object, const TCHAR* Type, FArc
 		FMessageLog ExportWarning("EditorErrors");
 		FFormatNamedArguments Arguments;
 		Arguments.Add(TEXT("Name"), FText::FromString(Texture->GetName()));
-		ExportWarning.Warning(FText::Format(LOCTEXT("BitDepthWarning", "{Name}: Texture is RGBA16 and cannot be represented at such high bit depth in .bmp. Color will be scaled to RGBA8."), Arguments));
+		ExportWarning.Warning(FText::Format(LOCTEXT("BitDepthBMPWarning", "{Name}: Texture is RGBA16 and cannot be represented at such high bit depth in .bmp. Color will be scaled to RGBA8."), Arguments));
 		ExportWarning.Open(EMessageSeverity::Warning);
 	}
 
@@ -5068,7 +5068,7 @@ bool UTextureExporterTGA::ExportBinary( UObject* Object, const TCHAR* Type, FArc
 		FMessageLog ExportWarning("EditorErrors");
 		FFormatNamedArguments Arguments;
 		Arguments.Add(TEXT("Name"), FText::FromString(Texture->GetName()));
-		ExportWarning.Warning(FText::Format(LOCTEXT("BitDepthWarning", "{Name}: Texture is RGBA16 and cannot be represented at such high bit depth in .tga. Color will be scaled to RGBA8."), Arguments));
+		ExportWarning.Warning(FText::Format(LOCTEXT("BitDepthTGAWarning", "{Name}: Texture is RGBA16 and cannot be represented at such high bit depth in .tga. Color will be scaled to RGBA8."), Arguments));
 		ExportWarning.Open(EMessageSeverity::Warning);
 	}
 
@@ -5695,9 +5695,11 @@ EReimportResult::Type UReimportFbxStaticMeshFactory::Reimport( UObject* Obj )
 		ReimportUI->StaticMeshImportData = ImportData;
 		
 		bool bImportOperationCanceled = false;
-		bool bImportForceType = true;
-
-		GetImportOptions( FFbxImporter, ReimportUI, /*bShowOptionDialog=*/true, Obj->GetPathName(), bImportOperationCanceled, bImportForceType, FBXIT_StaticMesh );
+		bool bForceImportType = true;
+		bool bShowOptionDialog = true;
+		bool bOutImportAll = false;
+		bool bIsObjFormat = false;
+		GetImportOptions( FFbxImporter, ReimportUI, bShowOptionDialog, Obj->GetPathName(), bOperationCanceled, bOutImportAll, bIsObjFormat, bForceImportType, FBXIT_StaticMesh );
 	}
 
 	if( !bOperationCanceled && ensure(ImportData) )
@@ -5882,6 +5884,10 @@ EReimportResult::Type UReimportFbxSkeletalMeshFactory::Reimport( UObject* Obj )
 
 	UnFbx::FFbxImporter* FFbxImporter = UnFbx::FFbxImporter::GetInstance();
 	UnFbx::FBXImportOptions* ImportOptions = FFbxImporter->GetImportOptions();
+
+	//Pop the message log in case of error
+	UnFbx::FFbxLoggerSetter Logger(FFbxImporter, true);
+
 	//Clean up the options
 	UnFbx::FBXImportOptions::ResetOptions(ImportOptions);
 
@@ -5918,7 +5924,10 @@ EReimportResult::Type UReimportFbxSkeletalMeshFactory::Reimport( UObject* Obj )
 		ReimportUI->SkeletalMeshImportData = ImportData;
 
 		bool bImportOperationCanceled = false;
+		bool bShowOptionDialog = true;
 		bool bForceImportType = true;
+		bool bOutImportAll = false;
+		bool bIsObjFormat = false;
 
 		// arggg... hate this different option class to confuse everybody
 		// @hack to make sure skeleton is set before opening the dialog
@@ -5926,7 +5935,7 @@ EReimportResult::Type UReimportFbxSkeletalMeshFactory::Reimport( UObject* Obj )
 		ImportOptions->bCreatePhysicsAsset = false;
 		ImportOptions->PhysicsAsset = SkeletalMesh->PhysicsAsset;
 
-		ImportOptions = GetImportOptions( FFbxImporter, ReimportUI, /*bShowOptionDialog=*/true, Obj->GetPathName(), bImportOperationCanceled, bForceImportType, FBXIT_SkeletalMesh );
+		ImportOptions = GetImportOptions( FFbxImporter, ReimportUI, bShowOptionDialog, Obj->GetPathName(), bOperationCanceled, bOutImportAll, bIsObjFormat, bForceImportType, FBXIT_SkeletalMesh );
 	}
 
 	if( !bOperationCanceled && ensure(ImportData) )
diff --git a/Engine/Source/Editor/UnrealEd/Private/Fbx/FbxMainExport.cpp b/Engine/Source/Editor/UnrealEd/Private/Fbx/FbxMainExport.cpp
--- a/Engine/Source/Editor/UnrealEd/Private/Fbx/FbxMainExport.cpp
+++ b/Engine/Source/Editor/UnrealEd/Private/Fbx/FbxMainExport.cpp
@@ -615,7 +615,7 @@ void FFbxExporter::ExportModel(UModel* Model, FbxNode* Node, const char* Name)
 		FbxSurfaceMaterial* FbxMaterial;
 		if (MaterialInterface != NULL && MaterialInterface->GetMaterial() != NULL)
 		{
-			FbxMaterial = ExportMaterial(MaterialInterface->GetMaterial());
+			FbxMaterial = ExportMaterial(MaterialInterface);
 		}
 		else
 		{
@@ -1037,32 +1037,33 @@ bool FFbxExporter::FillFbxTextureProperty(const char *PropertyName, const FExpre
 /**
 * Exports the profile_COMMON information for a material.
 */
-FbxSurfaceMaterial* FFbxExporter::ExportMaterial(UMaterial* Material)
+FbxSurfaceMaterial* FFbxExporter::ExportMaterial(UMaterialInterface* MaterialInterface)
 {
-	if (Scene == NULL || Material == NULL) return NULL;
+	if (Scene == nullptr || MaterialInterface == nullptr || MaterialInterface->GetMaterial() == nullptr) return nullptr;
 	
 	// Verify that this material has not already been exported:
-	if (FbxMaterials.Find(Material))
+	if (FbxMaterials.Find(MaterialInterface))
 	{
-		return *FbxMaterials.Find(Material);
+		return *FbxMaterials.Find(MaterialInterface);
 	}
 
 	// Create the Fbx material
-	FbxSurfaceMaterial* FbxMaterial = NULL;
+	FbxSurfaceMaterial* FbxMaterial = nullptr;
+
+	UMaterial *Material = MaterialInterface->GetMaterial();
 	
 	// Set the shading model
 	if (Material->GetShadingModel() == MSM_DefaultLit)
 	{
-		FbxMaterial = FbxSurfacePhong::Create(Scene, TCHAR_TO_UTF8(*Material->GetName()));
+		FbxMaterial = FbxSurfacePhong::Create(Scene, TCHAR_TO_UTF8(*MaterialInterface->GetName()));
 		//((FbxSurfacePhong*)FbxMaterial)->Specular.Set(Material->Specular));
 		//((FbxSurfacePhong*)FbxMaterial)->Shininess.Set(Material->SpecularPower.Constant);
 	}
 	else // if (Material->ShadingModel == MSM_Unlit)
 	{
-		FbxMaterial = FbxSurfaceLambert::Create(Scene, TCHAR_TO_UTF8(*Material->GetName()));
+		FbxMaterial = FbxSurfaceLambert::Create(Scene, TCHAR_TO_UTF8(*MaterialInterface->GetName()));
 	}
 	
-	
 	((FbxSurfaceLambert*)FbxMaterial)->TransparencyFactor.Set(Material->Opacity.Constant);
 
 	// Fill in the profile_COMMON effect with the material information.
@@ -1075,7 +1076,7 @@ FbxSurfaceMaterial* FFbxExporter::ExportMaterial(UMaterial* Material)
 	{
 		((FbxSurfaceLambert*)FbxMaterial)->Emissive.Set(SetMaterialComponent(Material->EmissiveColor, true));
 	}
-	
+
 	//Always set the ambient to zero since we dont have ambient in unreal we want to avoid default value in DCCs
 	((FbxSurfaceLambert*)FbxMaterial)->Ambient.Set(FbxDouble3(0.0, 0.0, 0.0));
 
@@ -1095,7 +1096,7 @@ FbxSurfaceMaterial* FFbxExporter::ExportMaterial(UMaterial* Material)
 		}
 	}
 
-	FbxMaterials.Add(Material, FbxMaterial);
+	FbxMaterials.Add(MaterialInterface, FbxMaterial);
 	
 	return FbxMaterial;
 }
@@ -2184,7 +2185,7 @@ FbxNode* FFbxExporter::ExportStaticMeshToFbx(const UStaticMesh* StaticMesh, int3
 			FIndexArrayView RawIndices = RenderMesh.IndexBuffer.GetArrayView();
 			UMaterialInterface* Material = StaticMesh->GetMaterial(Polygons.MaterialIndex);
 
-			FbxSurfaceMaterial* FbxMaterial = Material ? ExportMaterial(Material->GetMaterial()) : NULL;
+			FbxSurfaceMaterial* FbxMaterial = Material ? ExportMaterial(Material) : NULL;
 			if (!FbxMaterial)
 			{
 				FbxMaterial = CreateDefaultMaterial();
@@ -2297,7 +2298,7 @@ FbxNode* FFbxExporter::ExportStaticMeshToFbx(const UStaticMesh* StaticMesh, int3
 			FIndexArrayView RawIndices = RenderMesh.IndexBuffer.GetArrayView();
 			UMaterialInterface* Material = StaticMesh->GetMaterial(Polygons.MaterialIndex);
 
-			FbxSurfaceMaterial* FbxMaterial = Material ? ExportMaterial(Material->GetMaterial()) : NULL;
+			FbxSurfaceMaterial* FbxMaterial = Material ? ExportMaterial(Material) : NULL;
 			if (!FbxMaterial)
 			{
 				FbxMaterial = CreateDefaultMaterial();
@@ -2497,7 +2498,7 @@ void FFbxExporter::ExportSplineMeshToFbx(const USplineMeshComponent* SplineMeshC
 		FIndexArrayView RawIndices = RenderMesh.IndexBuffer.GetArrayView();
 		UMaterialInterface* Material = StaticMesh->GetMaterial(Polygons.MaterialIndex);
 
-		FbxSurfaceMaterial* FbxMaterial = Material ? ExportMaterial(Material->GetMaterial()) : NULL;
+		FbxSurfaceMaterial* FbxMaterial = Material ? ExportMaterial(Material) : NULL;
 		if (!FbxMaterial)
 		{
 			FbxMaterial = CreateDefaultMaterial();
@@ -2768,7 +2769,7 @@ void FFbxExporter::ExportLandscapeToFbx(ALandscapeProxy* Landscape, const TCHAR*
 	Layer0->SetMaterials(LayerElementMaterials);
 
 	UMaterialInterface* Material = Landscape->GetLandscapeMaterial();
-	FbxSurfaceMaterial* FbxMaterial = Material ? ExportMaterial(Material->GetMaterial()) : NULL;
+	FbxSurfaceMaterial* FbxMaterial = Material ? ExportMaterial(Material) : NULL;
 	if (!FbxMaterial)
 	{
 		FbxMaterial = CreateDefaultMaterial();
diff --git a/Engine/Source/Editor/UnrealEd/Private/Fbx/FbxMainImport.cpp b/Engine/Source/Editor/UnrealEd/Private/Fbx/FbxMainImport.cpp
--- a/Engine/Source/Editor/UnrealEd/Private/Fbx/FbxMainImport.cpp
+++ b/Engine/Source/Editor/UnrealEd/Private/Fbx/FbxMainImport.cpp
@@ -1151,7 +1151,8 @@ FbxAMatrix FFbxImporter::ComputeTotalMatrix(FbxNode* Node)
 		FullPivot[2] = -RotationPivot[2];
 		PivotGeometry.SetT(FullPivot);
 	}
-	FbxAMatrix TotalMatrix = ImportOptions->bTransformVertexToAbsolute ? GlobalTransform * Geometry * PivotGeometry : PivotGeometry;
+	//We must always add the geometric transform. Only Max use the geometric transform which is an offset to the local transform of the node
+	FbxAMatrix TotalMatrix = ImportOptions->bTransformVertexToAbsolute ? GlobalTransform * Geometry * PivotGeometry : Geometry * PivotGeometry;
 
 	return TotalMatrix;
 }
@@ -1691,7 +1692,8 @@ void FFbxImporter::RecursiveFindFbxSkelMesh(FbxNode* Node, TArray< TArray<FbxNod
 			}
 		}
 	}
-	else
+	
+	//Skeletalmesh node can have child so let's always iterate trough child
 	{
 		int32 ChildIndex;
 		TArray<FbxNode*> ChildNoScale;
diff --git a/Engine/Source/Editor/UnrealEd/Private/Fbx/FbxSceneImportFactory.cpp b/Engine/Source/Editor/UnrealEd/Private/Fbx/FbxSceneImportFactory.cpp
--- a/Engine/Source/Editor/UnrealEd/Private/Fbx/FbxSceneImportFactory.cpp
+++ b/Engine/Source/Editor/UnrealEd/Private/Fbx/FbxSceneImportFactory.cpp
@@ -671,6 +671,7 @@ UFbxSceneImportFactory::UFbxSceneImportFactory(const FObjectInitializer& ObjectI
 {
 	SupportedClass = UWorld::StaticClass();
 	Formats.Add(TEXT("fbx;Fbx Scene"));
+	Formats.Add(TEXT("obj;OBJ Scene"));
 
 	bCreateNew = false;
 	bText = false;
@@ -1310,6 +1311,8 @@ void UFbxSceneImportFactory::CreateLevelActorHierarchy(TSharedPtr<FFbxSceneInfo>
 
 				//Apply the hierarchy local transform to the root component
 				ApplyTransformToComponent(RootComponent, &(NodeInfo->Transform), ParentActor == nullptr ? &RootTransform : nullptr, PivotLocation, ParentPivotAccumulation);
+				//Notify people that the component get created/changed
+				RootComponent->PostEditChange();
 			}
 		}
 		//We select only the first actor
@@ -1510,7 +1513,8 @@ AActor *UFbxSceneImportFactory::CreateActorComponentsHierarchy(TSharedPtr<FFbxSc
 
 		//Apply the local transform to the scene component
 		ApplyTransformToComponent(SceneComponent, &(NodeInfo->Transform), ParentRootComponent != nullptr ? nullptr : &RootTransform, PivotLocation, ParentPivotAccumulation);
-		
+		//Notify people that the component get created/changed
+		SceneComponent->PostEditChange();
 	}
 	// End of iteration of the hierarchy
 	//////////////////////////////////////////////////////////////////////////
diff --git a/Engine/Source/Editor/UnrealEd/Private/Fbx/FbxSkeletalMeshExport.cpp b/Engine/Source/Editor/UnrealEd/Private/Fbx/FbxSkeletalMeshExport.cpp
--- a/Engine/Source/Editor/UnrealEd/Private/Fbx/FbxSkeletalMeshExport.cpp
+++ b/Engine/Source/Editor/UnrealEd/Private/Fbx/FbxSkeletalMeshExport.cpp
@@ -242,9 +242,9 @@ FbxNode* FFbxExporter::CreateMesh(const USkeletalMesh* SkelMesh, const TCHAR* Me
 		UMaterialInterface* MatInterface = SkelMesh->Materials[MaterialIndex].MaterialInterface;
 
 		FbxSurfaceMaterial* FbxMaterial = NULL;
-		if(MatInterface && !FbxMaterials.Find(MatInterface->GetMaterial()))
+		if(MatInterface && !FbxMaterials.Find(MatInterface))
 		{
-			FbxMaterial = ExportMaterial(MatInterface->GetMaterial());
+			FbxMaterial = ExportMaterial(MatInterface);
 		}
 		else
 		{
diff --git a/Engine/Source/Editor/UnrealEd/Private/Fbx/FbxSkeletalMeshImport.cpp b/Engine/Source/Editor/UnrealEd/Private/Fbx/FbxSkeletalMeshImport.cpp
--- a/Engine/Source/Editor/UnrealEd/Private/Fbx/FbxSkeletalMeshImport.cpp
+++ b/Engine/Source/Editor/UnrealEd/Private/Fbx/FbxSkeletalMeshImport.cpp
@@ -843,7 +843,7 @@ bool UnFbx::FFbxImporter::ImportBone(TArray<FbxNode*>& NodeArray, FSkeletalMeshI
 		}
 		else
 		{
-			AddTokenizedErrorMessage(FTokenizedMessage::Create(EMessageSeverity::Warning, LOCTEXT("FbxSkeletaLMeshimport_InvalidBindPose", "Could not find the bind pose.  It will use time 0 as bind pose.")), FFbxErrors::SkeletalMesh_InvalidBindPose);
+			AddTokenizedErrorMessage(FTokenizedMessage::Create(EMessageSeverity::Warning, LOCTEXT("FbxSkeletaLMeshimport_MissingBindPose", "Could not find the bind pose.  It will use time 0 as bind pose.")), FFbxErrors::SkeletalMesh_InvalidBindPose);
 			bUseTime0AsRefPose = true;
 		}
 	}
@@ -1802,6 +1802,12 @@ void UnFbx::FFbxImporter::SetMaterialSkinXXOrder(FSkeletalMeshImportData& Import
 			}
 		}
 
+		if (bNeedsReorder && MissingSkinSuffixMaterial.Num() > 0)
+		{
+			AddTokenizedErrorMessage(FTokenizedMessage::Create(EMessageSeverity::Error, LOCTEXT("FbxSkeletaLMeshimport_Skinxx_missing", "Cannot mix skinxx suffix materials with no skinxx material, mesh section order will not be right.")), FFbxErrors::Generic_Mesh_SkinxxNameError);
+			return;
+		}
+
 		//Fill the array MaterialIndexToSkinIndex so we order material by _skinXX order
 		//This ensure we support skinxx suffixe that are not increment by one like _skin00, skin_01, skin_03, skin_04, skin_08... 
 		for (auto kvp : SkinIndexToMaterialIndex)
@@ -3163,7 +3169,19 @@ FFbxLogger::FFbxLogger()
 
 FFbxLogger::~FFbxLogger()
 {
-	if(TokenizedErrorMessages.Num() > 0)
+	bool ShowLogMessage = !ShowLogMessageOnlyIfError;
+	if (ShowLogMessageOnlyIfError)
+	{
+		for (TSharedRef<FTokenizedMessage> TokenMessage : TokenizedErrorMessages)
+		{
+			if (TokenMessage->GetSeverity() == EMessageSeverity::CriticalError || TokenMessage->GetSeverity() == EMessageSeverity::Error)
+			{
+				ShowLogMessage = true;
+				break;
+			}
+		}
+	}
+	if(ShowLogMessage && TokenizedErrorMessages.Num() > 0)
 	{
 		const TCHAR* LogTitle = TEXT("FBXImport");
 		FMessageLogModule& MessageLogModule = FModuleManager::LoadModuleChecked<FMessageLogModule>("MessageLog");
diff --git a/Engine/Source/Editor/UnrealEd/Private/Fbx/SFbxSceneOptionWindow.cpp b/Engine/Source/Editor/UnrealEd/Private/Fbx/SFbxSceneOptionWindow.cpp
--- a/Engine/Source/Editor/UnrealEd/Private/Fbx/SFbxSceneOptionWindow.cpp
+++ b/Engine/Source/Editor/UnrealEd/Private/Fbx/SFbxSceneOptionWindow.cpp
@@ -451,8 +451,8 @@ TSharedRef<SDockTab> SFbxSceneOptionWindow::SpawnSkeletalMeshReimportTab(const F
 								.AutoWidth()
 								[
 									SNew(SButton)
-									.Text(LOCTEXT("FbxOptionWindow_SM_CreateOverride", "Create Override"))
-									.ToolTipText(LOCTEXT("FbxOptionWindow_SM_CreateOverrideTooltip", "Create Override to specify custom import options for some Skeletal meshes.\nTo assign options use context menu on Skeletal meshes."))
+									.Text(LOCTEXT("FbxOptionWindow_SKM_CreateOverride", "Create Override"))
+									.ToolTipText(LOCTEXT("FbxOptionWindow_SKM_CreateOverrideTooltip", "Create Override to specify custom import options for some Skeletal meshes.\nTo assign options use context menu on Skeletal meshes."))
 									.OnClicked(SkeletalMeshReimportListView.Get(), &SFbxSceneSkeletalMeshReimportListView::OnCreateOverrideOptions)
 								]
 								+ SHorizontalBox::Slot()
diff --git a/Engine/Source/Editor/UnrealEd/Private/Fbx/SSceneImportNodeTreeView.cpp b/Engine/Source/Editor/UnrealEd/Private/Fbx/SSceneImportNodeTreeView.cpp
--- a/Engine/Source/Editor/UnrealEd/Private/Fbx/SSceneImportNodeTreeView.cpp
+++ b/Engine/Source/Editor/UnrealEd/Private/Fbx/SSceneImportNodeTreeView.cpp
@@ -2,7 +2,7 @@
 
 #include "UnrealEd.h"
 #include "SSceneImportNodeTreeView.h"
-#include "ClassIconFinder.h"
+#include "SlateIconFinder.h"
 
 #define LOCTEXT_NAMESPACE "SFbxSceneTreeView"
 
@@ -97,7 +97,7 @@ public:
 		{
 			IconClass = UCameraComponent::StaticClass();
 		}
-		const FSlateBrush* ClassIcon = FClassIconFinder::FindIconForClass(IconClass);
+		const FSlateBrush* ClassIcon = FSlateIconFinder::FindIconBrushForClass(IconClass);
 
 		//Prepare the tooltip
 		FString Tooltip = FbxNodeInfo->NodeName;
diff --git a/Engine/Source/Editor/UnrealEd/Private/Fbx/SSceneImportStaticMeshListView.cpp b/Engine/Source/Editor/UnrealEd/Private/Fbx/SSceneImportStaticMeshListView.cpp
--- a/Engine/Source/Editor/UnrealEd/Private/Fbx/SSceneImportStaticMeshListView.cpp
+++ b/Engine/Source/Editor/UnrealEd/Private/Fbx/SSceneImportStaticMeshListView.cpp
@@ -2,7 +2,7 @@
 
 #include "UnrealEd.h"
 #include "SSceneImportStaticMeshListView.h"
-#include "ClassIconFinder.h"
+#include "SlateIconFinder.h"
 #include "Dialogs/DlgPickAssetPath.h"
 #include "Editor/ContentBrowser/Public/ContentBrowserModule.h"
 #include "SFbxSceneOptionWindow.h"
@@ -73,7 +73,7 @@ public:
 		else if (ColumnName == FbxSceneImportStaticMesh::SceneImportClassIconHeaderIdName)
 		{
 			UClass *IconClass = FbxMeshInfo->GetType();
-			const FSlateBrush* ClassIcon = FClassIconFinder::FindIconForClass(IconClass);
+			const FSlateBrush* ClassIcon = FSlateIconFinder::FindIconBrushForClass(IconClass);
 
 			TSharedRef<SOverlay> IconContent = SNew(SOverlay)
 				+ SOverlay::Slot()
diff --git a/Engine/Source/Editor/UnrealEd/Private/Fbx/SSceneReimportStaticMeshListView.cpp b/Engine/Source/Editor/UnrealEd/Private/Fbx/SSceneReimportStaticMeshListView.cpp
--- a/Engine/Source/Editor/UnrealEd/Private/Fbx/SSceneReimportStaticMeshListView.cpp
+++ b/Engine/Source/Editor/UnrealEd/Private/Fbx/SSceneReimportStaticMeshListView.cpp
@@ -82,7 +82,7 @@ public:
 		else if (ColumnName == FbxSceneReimportStaticMesh::ClassIconHeaderIdName && SlateBrush != nullptr)
 		{
 			UClass *IconClass = FbxMeshInfo->GetType();
-			const FSlateBrush* ClassIcon = FClassIconFinder::FindIconForClass(IconClass);
+			const FSlateBrush* ClassIcon = FSlateIconFinder::FindIconBrushForClass(IconClass);
 			TSharedRef<SOverlay> IconContent = SNew(SOverlay)
 				+ SOverlay::Slot()
 				.HAlign(HAlign_Right)
diff --git a/Engine/Source/Editor/UnrealEd/Private/FbxExporter.h b/Engine/Source/Editor/UnrealEd/Private/FbxExporter.h
--- a/Engine/Source/Editor/UnrealEd/Private/FbxExporter.h
+++ b/Engine/Source/Editor/UnrealEd/Private/FbxExporter.h
@@ -153,7 +153,7 @@ private:
 	TMap<FString,int32> FbxNodeNameToIndexMap;
 	TMap<const AActor*, FbxNode*> FbxActors;
 	TMap<const USkeletalMeshComponent*, FbxNode*> FbxSkeletonRoots;
-	TMap<const UMaterial*, FbxSurfaceMaterial*> FbxMaterials;
+	TMap<const UMaterialInterface*, FbxSurfaceMaterial*> FbxMaterials;
 	TMap<const UStaticMesh*, FbxMesh*> FbxMeshes;
 
 	/** The frames-per-second (FPS) used when baking transforms */
@@ -304,7 +304,7 @@ private:
 	/**
 	 * Exports the profile_COMMON information for a material.
 	 */
-	FbxSurfaceMaterial* ExportMaterial(UMaterial* Material);
+	FbxSurfaceMaterial* ExportMaterial(UMaterialInterface* Material);
 	
 	FbxSurfaceMaterial* CreateDefaultMaterial();
 	
diff --git a/Engine/Source/Editor/UnrealEd/Private/FbxImporter.h b/Engine/Source/Editor/UnrealEd/Private/FbxImporter.h
--- a/Engine/Source/Editor/UnrealEd/Private/FbxImporter.h
+++ b/Engine/Source/Editor/UnrealEd/Private/FbxImporter.h
@@ -1174,6 +1174,9 @@ class FFbxLogger
 	/** Error messages **/
 	TArray<TSharedRef<FTokenizedMessage>> TokenizedErrorMessages;
 
+	/* The logger will show the LogMessage only if at least one TokenizedErrorMessage have a severity of Error or CriticalError*/
+	bool ShowLogMessageOnlyIfError;
+
 	friend class FFbxImporter;
 	friend class FFbxLoggerSetter;
 };
@@ -1189,12 +1192,13 @@ class FFbxLoggerSetter
 	FFbxImporter * Importer;
 
 public:
-	FFbxLoggerSetter(FFbxImporter * InImpoter)
+	FFbxLoggerSetter(FFbxImporter * InImpoter, bool ShowLogMessageOnlyIfError = false)
 		: Importer(InImpoter)
 	{
 		// if impoter doesn't have logger, sets it
 		if(Importer->Logger == NULL)
 		{
+			Logger.ShowLogMessageOnlyIfError = ShowLogMessageOnlyIfError;
 			Importer->SetLogger(&Logger);
 		}
 		else
diff --git a/Engine/Source/Editor/UnrealEd/Private/FileHelpers.cpp b/Engine/Source/Editor/UnrealEd/Private/FileHelpers.cpp
--- a/Engine/Source/Editor/UnrealEd/Private/FileHelpers.cpp
+++ b/Engine/Source/Editor/UnrealEd/Private/FileHelpers.cpp
@@ -317,7 +317,7 @@ FString FEditorFileUtils::GetFilterString(EFileInteraction Interaction)
 		Result = TEXT("Unreal Text (*.t3d)|*.t3d|All Files|*.*");
 		break;
 	case FI_ImportScene:
-		Result = TEXT("FBX (*.fbx)|*.fbx|All Files|*.*");
+		Result = TEXT("FBX (*.fbx) OBJ (.obj)|*.fbx;*.obj|FBX (*.fbx)|*.fbx|OBJ (*.obj)|*.obj|All Files|*.*");
 		break;
 
 	case FI_Export:
@@ -1439,7 +1439,10 @@ bool FEditorFileUtils::PromptToCheckoutPackages(bool bCheckDirty, const TArray<U
 	// If any files were just checked out, remove any pending flag to show a notification prompting for checkout.
 	if (PackagesToCheckOut.Num() > 0)
 	{
-		GUnrealEd->bNeedToPromptForCheckout = false;
+		for (UPackage* Package : PackagesToCheckOut)
+		{
+			GUnrealEd->PackageToNotifyState.Add(Package, NS_DialogPrompted);
+		}
 	}
 
 	if (OutPackagesNotNeedingCheckout)
@@ -1940,7 +1943,7 @@ bool FEditorFileUtils::AttemptUnloadInactiveWorldPackage(UPackage* PackageToUnlo
 				case EWorldType::PIE:
 				case EWorldType::Preview:
 				default:
-					OutErrorMessage = NSLOCTEXT("SaveAsImplementation", "ExistingWorldNotInactive", "The level you are attempting to unload is invalid.");
+					OutErrorMessage = NSLOCTEXT("SaveAsImplementation", "ExistingWorldInvalid", "The level you are attempting to unload is invalid.");
 					bContinueUnloadingExistingWorld = false;
 					break;
 			}
diff --git a/Engine/Source/Editor/UnrealEd/Private/HierarchicalLOD.cpp b/Engine/Source/Editor/UnrealEd/Private/HierarchicalLOD.cpp
--- a/Engine/Source/Editor/UnrealEd/Private/HierarchicalLOD.cpp
+++ b/Engine/Source/Editor/UnrealEd/Private/HierarchicalLOD.cpp
@@ -70,7 +70,7 @@ void FHierarchicalLODBuilder::Build()
 		FMessageLog MapCheck("HLODResults");
 		MapCheck.Warning()
 			->AddToken(FUObjectToken::Create(WorldSetting))
-			->AddToken(FTextToken::Create(LOCTEXT("MapCheck_Message_HLODHiddenLevels", "Certain levels are marked as hidden, Hierarchical LODs will not be build for hidden levels.")));
+			->AddToken(FTextToken::Create(LOCTEXT("MapCheck_Message_NoBuildHLODHiddenLevels", "Certain levels are marked as hidden, Hierarchical LODs will not be build for hidden levels.")));
 	}
 	
 }
@@ -414,7 +414,7 @@ void FHierarchicalLODBuilder::ClearHLODs()
 		FMessageLog MapCheck("MapCheck");
 		MapCheck.Warning()
 			->AddToken(FUObjectToken::Create(World->GetWorldSettings()))
-			->AddToken(FTextToken::Create(LOCTEXT("MapCheck_Message_HLODHiddenLevels", "Certain levels are marked as hidden, Hierarchical LODs will not be deleted for hidden levels.")));
+			->AddToken(FTextToken::Create(LOCTEXT("MapCheck_Message_NoDeleteHLODHiddenLevels", "Certain levels are marked as hidden, Hierarchical LODs will not be deleted for hidden levels.")));
 	}
 }
 
@@ -436,7 +436,7 @@ void FHierarchicalLODBuilder::ClearPreviewBuild()
 		FMessageLog MapCheck("MapCheck");
 		MapCheck.Warning()
 			->AddToken(FUObjectToken::Create(World->GetWorldSettings()))
-			->AddToken(FTextToken::Create(LOCTEXT("MapCheck_Message_HLODHiddenLevels", "Certain levels are marked as hidden, Hierarchical LODs will not be deleted for hidden levels.")));
+			->AddToken(FTextToken::Create(LOCTEXT("MapCheck_Message_NoDeleteHLODHiddenLevels", "Certain levels are marked as hidden, Hierarchical LODs will not be deleted for hidden levels.")));
 	}
 }
 
@@ -515,7 +515,7 @@ void FHierarchicalLODBuilder::BuildMeshesForLODActors()
 		FMessageLog MapCheck("MapCheck");
 		MapCheck.Warning()
 			->AddToken(FUObjectToken::Create(WorldSetting))
-			->AddToken(FTextToken::Create(LOCTEXT("MapCheck_Message_HLODHiddenLevels", "Certain levels are marked as hidden, Hierarchical LODs will not be build for hidden levels.")));
+			->AddToken(FTextToken::Create(LOCTEXT("MapCheck_Message_NoBuildHLODHiddenLevels", "Certain levels are marked as hidden, Hierarchical LODs will not be build for hidden levels.")));
 	}
 
 }
diff --git a/Engine/Source/Editor/UnrealEd/Private/Kismet2/BlueprintEditorUtils.cpp b/Engine/Source/Editor/UnrealEd/Private/Kismet2/BlueprintEditorUtils.cpp
--- a/Engine/Source/Editor/UnrealEd/Private/Kismet2/BlueprintEditorUtils.cpp
+++ b/Engine/Source/Editor/UnrealEd/Private/Kismet2/BlueprintEditorUtils.cpp
@@ -54,7 +54,7 @@
 #include "EditorCategoryUtils.h"
 #include "EngineUtils.h"
 #include "Engine/LevelScriptActor.h"
-#include "ClassIconFinder.h"
+#include "SlateIconFinder.h"
 #define LOCTEXT_NAMESPACE "Blueprint"
 
 DEFINE_LOG_CATEGORY(LogBlueprintDebug);
@@ -4893,7 +4893,7 @@ void FBlueprintEditorUtils::RenameLocalVariable(UBlueprint* InBlueprint, const U
 
 		if (LocalVariable && !bHasExistingProperty)
 		{
-			const FScopedTransaction Transaction( LOCTEXT("RenameVariable", "Rename Local Variable") );
+			const FScopedTransaction Transaction( LOCTEXT("RenameLocalVariable", "Rename Local Variable") );
 			InBlueprint->Modify();
 			FunctionEntry->Modify();
 
@@ -8309,7 +8309,7 @@ const FSlateBrush* FBlueprintEditorUtils::GetIconFromPin( const FEdGraphPinType&
 		UClass* VarClass = FindObject<UClass>(ANY_PACKAGE, *PinSubObject->GetName());
 		if( VarClass )
 		{
-			IconBrush = FClassIconFinder::FindIconForClass( VarClass );
+			IconBrush = FSlateIconFinder::FindIconBrushForClass( VarClass );
 		}
 	}
 	return IconBrush;
diff --git a/Engine/Source/Editor/UnrealEd/Private/Kismet2/ComponentEditorUtils.cpp b/Engine/Source/Editor/UnrealEd/Private/Kismet2/ComponentEditorUtils.cpp
--- a/Engine/Source/Editor/UnrealEd/Private/Kismet2/ComponentEditorUtils.cpp
+++ b/Engine/Source/Editor/UnrealEd/Private/Kismet2/ComponentEditorUtils.cpp
@@ -12,7 +12,7 @@
 #include "UnrealExporter.h"
 #include "GenericCommands.h"
 #include "SourceCodeNavigation.h"
-#include "ClassIconFinder.h"
+#include "SlateIconFinder.h"
 #include "AssetEditorManager.h"
 #include "Components/DecalComponent.h"
 
@@ -829,7 +829,7 @@ void FComponentEditorUtils::FillComponentContextMenuOptions(FMenuBuilder& MenuBu
 				MenuBuilder.AddMenuEntry(
 					FText::Format(LOCTEXT("GoToBlueprintForComponent", "Edit {0}"), FText::FromString(Component->GetClass()->ClassGeneratedBy->GetName())),
 					LOCTEXT("EditBlueprintForComponent_ToolTip", "Edits the Blueprint Class that defines this component."),
-					FSlateIcon(FEditorStyle::GetStyleSetName(), FClassIconFinder::FindIconNameForClass(Component->GetClass())),
+					FSlateIconFinder::FindIconForClass(Component->GetClass()),
 					FUIAction(
 					FExecuteAction::CreateStatic(&FComponentEditorUtils::OnEditBlueprintComponent, Component->GetClass()->ClassGeneratedBy),
 					FCanExecuteAction()));
diff --git a/Engine/Source/Editor/UnrealEd/Private/Kismet2/DebuggerCommands.cpp b/Engine/Source/Editor/UnrealEd/Private/Kismet2/DebuggerCommands.cpp
--- a/Engine/Source/Editor/UnrealEd/Private/Kismet2/DebuggerCommands.cpp
+++ b/Engine/Source/Editor/UnrealEd/Private/Kismet2/DebuggerCommands.cpp
@@ -813,7 +813,7 @@ TSharedRef< SWidget > FPlayWorldCommands::GenerateLaunchMenuContent( TSharedRef<
 				// ... generate tooltip text
 				FFormatNamedArguments TooltipArguments;
 				TooltipArguments.Add(TEXT("DisplayName"), PlatformInfo->DisplayName);
-				FText Tooltip = FText::Format(LOCTEXT("LaunchDeviceToolTipText", "Found no connected devices for {DisplayName}"), TooltipArguments);
+				FText Tooltip = FText::Format(LOCTEXT("LaunchNoDevicesToolTipText", "Found no connected devices for {DisplayName}"), TooltipArguments);
 
 				// ... and add a menu entry
 				MenuBuilder.AddMenuEntry(
@@ -1646,12 +1646,12 @@ bool FInternalPlayWorldCommandCallbacks::IsReadyToLaunchOnDevice(FString DeviceI
 		{
 		case ETargetPlatformReadyStatus::CodeUnsupported:
 			// show the message
-			FMessageDialog::Open(EAppMsgType::Ok, LOCTEXT("NotSupported", "Sorry, launching a code-based project for the selected platform is currently not supported. This feature may be available in a future release."));
+			FMessageDialog::Open(EAppMsgType::Ok, LOCTEXT("NotSupported_CodeBased", "Sorry, launching a code-based project for the selected platform is currently not supported. This feature may be available in a future release."));
 			return false;
 
 		case ETargetPlatformReadyStatus::PluginsUnsupported:
 			// show the message
-			FMessageDialog::Open(EAppMsgType::Ok, LOCTEXT("NotSupported", "Sorry, launching a project with third-party plugins is currently not supported for the selected platform. This feature may be available in a future release."));
+			FMessageDialog::Open(EAppMsgType::Ok, LOCTEXT("NotSupported_Plugins", "Sorry, launching a project with third-party plugins is currently not supported for the selected platform. This feature may be available in a future release."));
 			return false;
 		default:
 			break;
diff --git a/Engine/Source/Editor/UnrealEd/Private/Kismet2/StructureEditorUtils.cpp b/Engine/Source/Editor/UnrealEd/Private/Kismet2/StructureEditorUtils.cpp
--- a/Engine/Source/Editor/UnrealEd/Private/Kismet2/StructureEditorUtils.cpp
+++ b/Engine/Source/Editor/UnrealEd/Private/Kismet2/StructureEditorUtils.cpp
@@ -706,7 +706,14 @@ bool FStructureEditorUtils::CanEnableMultiLineText(const UUserDefinedStruct* Str
 	auto VarDesc = GetVarDescByGuid(Struct, VarGuid);
 	if (VarDesc)
 	{
-		auto Property = FindField<UProperty>(Struct, VarDesc->VarName);
+		UProperty* Property = FindField<UProperty>(Struct, VarDesc->VarName);
+
+		// If this is an array, we need to test its inner property as that's the real type
+		if (UArrayProperty* ArrayProperty = Cast<UArrayProperty>(Property))
+		{
+			Property = ArrayProperty->Inner;
+		}
+
 		if (Property)
 		{
 			// Can only set multi-line text on string and text properties
@@ -738,6 +745,7 @@ bool FStructureEditorUtils::ChangeMultiLineTextEnabled(UUserDefinedStruct* Struc
 				Property->RemoveMetaData("MultiLine");
 			}
 		}
+		OnStructureChanged(Struct);
 		return true;
 	}
 	return false;
diff --git a/Engine/Source/Editor/UnrealEd/Private/MaterialGraphSchema.cpp b/Engine/Source/Editor/UnrealEd/Private/MaterialGraphSchema.cpp
--- a/Engine/Source/Editor/UnrealEd/Private/MaterialGraphSchema.cpp
+++ b/Engine/Source/Editor/UnrealEd/Private/MaterialGraphSchema.cpp
@@ -697,7 +697,7 @@ void UMaterialGraphSchema::GetMaterialFunctionActions(FGraphActionMenuBuilder& A
 
 	for (const FAssetData& AssetData : AssetDataList)
 	{
-		const bool bExposeToLibrary = AssetData.TagsAndValues.FindRef("bExposeToLibrary") == TEXT("TRUE");
+		const bool bExposeToLibrary = AssetData.GetTagValueRef<bool>("bExposeToLibrary");
 
 		// If this was a function that was selected to be exposed to the library
 		if ( bExposeToLibrary )
@@ -714,10 +714,10 @@ void UMaterialGraphSchema::GetMaterialFunctionActions(FGraphActionMenuBuilder& A
 			{
 				// Gather the relevant information from the asset data
 				const FString FunctionPathName = AssetData.ObjectPath.ToString();
-				const FString Description = AssetData.TagsAndValues.FindRef("Description");
+				const FString Description = AssetData.GetTagValueRef<FString>("Description");
 				TArray<FString> LibraryCategories;
 				{
-					const FString LibraryCategoriesString = AssetData.TagsAndValues.FindRef("LibraryCategories");
+					const FString LibraryCategoriesString = AssetData.GetTagValueRef<FString>("LibraryCategories");
 					if ( !LibraryCategoriesString.IsEmpty() )
 					{
 						if (UArrayProperty* LibraryCategoriesProperty = FindFieldChecked<UArrayProperty>(UMaterialFunction::StaticClass(), TEXT("LibraryCategories")))
@@ -729,7 +729,7 @@ void UMaterialGraphSchema::GetMaterialFunctionActions(FGraphActionMenuBuilder& A
 				}
 				TArray<FText> LibraryCategoriesText;
 				{
-					const FString LibraryCategoriesString = AssetData.TagsAndValues.FindRef("LibraryCategoriesText");
+					const FString LibraryCategoriesString = AssetData.GetTagValueRef<FString>("LibraryCategoriesText");
 					if ( !LibraryCategoriesString.IsEmpty() )
 					{
 						UArrayProperty* LibraryCategoriesProperty = FindFieldChecked<UArrayProperty>(UMaterialFunction::StaticClass(), GET_MEMBER_NAME_CHECKED(UMaterialFunction, LibraryCategoriesText));
@@ -798,10 +798,8 @@ bool UMaterialGraphSchema::HasCompatibleConnection(const FAssetData& FunctionAss
 {
 	if (TestType != 0)
 	{
-		const FString* CombinedInputTagValue = FunctionAssetData.TagsAndValues.Find(GET_MEMBER_NAME_CHECKED(UMaterialFunction, CombinedInputTypes));
-		const FString* CombinedOutputTagValue = FunctionAssetData.TagsAndValues.Find(GET_MEMBER_NAME_CHECKED(UMaterialFunction, CombinedOutputTypes));
-		uint32 CombinedInputTypes = CombinedInputTagValue ? FCString::Atoi(**CombinedInputTagValue) : 0;
-		uint32 CombinedOutputTypes = CombinedOutputTagValue ? FCString::Atoi(**CombinedOutputTagValue) : 0;
+		uint32 CombinedInputTypes = FunctionAssetData.GetTagValueRef<uint32>(GET_MEMBER_NAME_CHECKED(UMaterialFunction, CombinedInputTypes));
+		uint32 CombinedOutputTypes = FunctionAssetData.GetTagValueRef<uint32>(GET_MEMBER_NAME_CHECKED(UMaterialFunction, CombinedOutputTypes));
 
 		if (CombinedOutputTypes == 0)
 		{
diff --git a/Engine/Source/Editor/UnrealEd/Private/Niagara/NiagaraNodeOp.cpp b/Engine/Source/Editor/UnrealEd/Private/Niagara/NiagaraNodeOp.cpp
--- a/Engine/Source/Editor/UnrealEd/Private/Niagara/NiagaraNodeOp.cpp
+++ b/Engine/Source/Editor/UnrealEd/Private/Niagara/NiagaraNodeOp.cpp
@@ -158,7 +158,7 @@ void UNiagaraNodeOp::Compile(class INiagaraCompiler* Compiler, TArray<FNiagaraNo
 	{
 		FFormatNamedArguments Args;
 		Args.Add(TEXT("OpName"), GetNodeTitle(ENodeTitleType::FullTitle));
-		FText Format = LOCTEXT("InputErrorFormat", "Error compiling {OpName} node.");
+		FText Format = LOCTEXT("NodeErrorFormat", "Error compiling {OpName} node.");
 		Compiler->Error(FText::Format(Format, Args), this, nullptr);
 	}
 }
diff --git a/Engine/Source/Editor/UnrealEd/Private/PlayLevel.cpp b/Engine/Source/Editor/UnrealEd/Private/PlayLevel.cpp
--- a/Engine/Source/Editor/UnrealEd/Private/PlayLevel.cpp
+++ b/Engine/Source/Editor/UnrealEd/Private/PlayLevel.cpp
@@ -1553,7 +1553,7 @@ void UEditorEngine::HandleLaunchCompleted(bool Succeeded, double TotalTime, int3
 		if ((PlayUsingLauncherDeviceId.Left(PlayUsingLauncherDeviceId.Find(TEXT("@"))) == TEXT("IOS") && PlayUsingLauncherDeviceName.Contains(DummyIOSDeviceName)) ||
 			(PlayUsingLauncherDeviceId.Left(PlayUsingLauncherDeviceId.Find(TEXT("@"))) == TEXT("TVOS") && PlayUsingLauncherDeviceName.Contains(DummyTVOSDeviceName)))
 		{
-			CompletionMsg = LOCTEXT("LauncherTaskCompleted", "Deployment complete! Open the app on your device to launch.");
+			CompletionMsg = LOCTEXT("DeploymentTaskCompleted", "Deployment complete! Open the app on your device to launch.");
 			TSharedPtr<SNotificationItem> NotificationItem = NotificationItemPtr.Pin();
 //			NotificationItem->SetExpireDuration(30.0f);
 		}
@@ -1581,7 +1581,7 @@ void UEditorEngine::HandleLaunchCompleted(bool Succeeded, double TotalTime, int3
 		if ((PlayUsingLauncherDeviceId.Left(PlayUsingLauncherDeviceId.Find(TEXT("@"))) == TEXT("IOS") && PlayUsingLauncherDeviceName.Contains(DummyIOSDeviceName)) ||
 			(PlayUsingLauncherDeviceId.Left(PlayUsingLauncherDeviceId.Find(TEXT("@"))) == TEXT("TVOS") && PlayUsingLauncherDeviceName.Contains(DummyTVOSDeviceName)))
 		{
-			CompletionMsg = LOCTEXT("LauncherTaskFailed", "Deployment failed!");
+			CompletionMsg = LOCTEXT("DeploymentTaskFailed", "Deployment failed!");
 		}
 		else
 		{
@@ -2701,7 +2701,7 @@ void UEditorEngine::OnLoginPIEComplete_Deferred(int32 LocalUserNum, bool bWasSuc
 		{
 			if (DataStruct.NetMode != EPlayNetMode::PIE_Client)
 			{
-				FMessageLog(NAME_CategoryPIE).Info(LOCTEXT("LoggedInClient", "Server logged in"));
+				FMessageLog(NAME_CategoryPIE).Info(LOCTEXT("LoggedInServer", "Server logged in"));
 			}
 			else
 			{
@@ -2712,7 +2712,7 @@ void UEditorEngine::OnLoginPIEComplete_Deferred(int32 LocalUserNum, bool bWasSuc
 		{
 			if (DataStruct.NetMode != EPlayNetMode::PIE_Client)
 			{
-				FMessageLog(NAME_CategoryPIE).Warning(LOCTEXT("LoggedInClientFailure", "Server failed to login"));
+				FMessageLog(NAME_CategoryPIE).Warning(LOCTEXT("LoggedInServerFailure", "Server failed to login"));
 			}
 			else
 			{
diff --git a/Engine/Source/Editor/UnrealEd/Private/RichCurveEditorCommands.cpp b/Engine/Source/Editor/UnrealEd/Private/RichCurveEditorCommands.cpp
--- a/Engine/Source/Editor/UnrealEd/Private/RichCurveEditorCommands.cpp
+++ b/Engine/Source/Editor/UnrealEd/Private/RichCurveEditorCommands.cpp
@@ -3,7 +3,7 @@
 #include "UnrealEd.h"
 #include "RichCurveEditorCommands.h"
 
-#define LOCTEXT_NAMESPACE ""
+#define LOCTEXT_NAMESPACE "RichCurveEditorCommands"
 
 void FRichCurveEditorCommands::RegisterCommands()
 {
diff --git a/Engine/Source/Editor/UnrealEd/Private/SClassPickerDialog.cpp b/Engine/Source/Editor/UnrealEd/Private/SClassPickerDialog.cpp
--- a/Engine/Source/Editor/UnrealEd/Private/SClassPickerDialog.cpp
+++ b/Engine/Source/Editor/UnrealEd/Private/SClassPickerDialog.cpp
@@ -7,7 +7,7 @@
 #include "SClassPickerDialog.h"
 #include "EditorClassUtils.h"
 #include "SExpandableArea.h"
-#include "ClassIconFinder.h"
+#include "SlateIconFinder.h"
 
 #define LOCTEXT_NAMESPACE "SClassPicker"
 
@@ -167,7 +167,7 @@ TSharedRef<ITableRow> SClassPickerDialog::GenerateListRow(TSharedPtr<FClassPicke
 {
 	FClassPickerDefaults* Obj = InItem.Get();
 	UClass* ItemClass = LoadClass<UObject>(NULL, *Obj->ClassName, NULL, LOAD_None, NULL);
-	const FSlateBrush* ItemBrush = FClassIconFinder::FindIconForClass(ItemClass);
+	const FSlateBrush* ItemBrush = FSlateIconFinder::FindIconBrushForClass(ItemClass);
 
 	return 
 	SNew(STableRow< TSharedPtr<FClassPickerDefaults> >, OwnerTable)
diff --git a/Engine/Source/Editor/UnrealEd/Private/SCommonEditorViewportToolbarBase.cpp b/Engine/Source/Editor/UnrealEd/Private/SCommonEditorViewportToolbarBase.cpp
--- a/Engine/Source/Editor/UnrealEd/Private/SCommonEditorViewportToolbarBase.cpp
+++ b/Engine/Source/Editor/UnrealEd/Private/SCommonEditorViewportToolbarBase.cpp
@@ -226,7 +226,7 @@ TSharedRef<SWidget> SCommonEditorViewportToolbarBase::GenerateCameraMenu() const
 	// Camera types
 	CameraMenuBuilder.AddMenuEntry( FEditorViewportCommands::Get().Perspective );
 
-	CameraMenuBuilder.BeginSection("LevelViewportCameraType_Ortho", LOCTEXT("CameraTypeHeader_Ortho", "Othographic") );
+	CameraMenuBuilder.BeginSection("LevelViewportCameraType_Ortho", LOCTEXT("CameraTypeHeader_Ortho", "Orthographic") );
 		CameraMenuBuilder.AddMenuEntry(FEditorViewportCommands::Get().Top);
 		CameraMenuBuilder.AddMenuEntry(FEditorViewportCommands::Get().Bottom);
 		CameraMenuBuilder.AddMenuEntry(FEditorViewportCommands::Get().Left);
@@ -269,7 +269,7 @@ TSharedRef<SWidget> SCommonEditorViewportToolbarBase::GenerateShowMenu() const
 		if( ShowMenu[SFG_Normal].Num() > 0 )
 		{
 			// Generate entries for the standard show flags
-			ShowMenuBuilder.BeginSection("LevelViewportShowFlagsCommon", LOCTEXT("CommonShowFlagHeader", "Common") );
+			ShowMenuBuilder.BeginSection("LevelViewportShowFlagsCommon", LOCTEXT("CommonShowFlagHeader", "Common Show Flags") );
 			{
 				for( int32 EntryIndex = 0; EntryIndex < ShowMenu[SFG_Normal].Num(); ++EntryIndex )
 				{
diff --git a/Engine/Source/Editor/UnrealEd/Private/SComponentClassCombo.cpp b/Engine/Source/Editor/UnrealEd/Private/SComponentClassCombo.cpp
--- a/Engine/Source/Editor/UnrealEd/Private/SComponentClassCombo.cpp
+++ b/Engine/Source/Editor/UnrealEd/Private/SComponentClassCombo.cpp
@@ -3,7 +3,7 @@
 #include "UnrealEd.h"
 
 #include "BlueprintGraphDefinitions.h"
-#include "ClassIconFinder.h"
+#include "SlateIconFinder.h"
 #include "ComponentAssetBroker.h"
 #include "ComponentTypeRegistry.h"
 #include "EditorClassUtils.h"
@@ -350,7 +350,7 @@ TSharedRef<ITableRow> SComponentClassCombo::GenerateAddComponentRow( FComponentC
 				.AutoWidth()
 				[
 					SNew(SImage)
-					.Image( FClassIconFinder::FindIconForClass( Entry->GetIconOverrideBrushName() == NAME_None ? Entry->GetIconClass() : nullptr, Entry->GetIconOverrideBrushName() ) )
+					.Image( FSlateIconFinder::FindIconBrushForClass( Entry->GetIconOverrideBrushName() == NAME_None ? Entry->GetIconClass() : nullptr, Entry->GetIconOverrideBrushName() ) )
 				]
 				+SHorizontalBox::Slot()
 				.AutoWidth()
diff --git a/Engine/Source/Editor/UnrealEd/Private/SSkeletonWidget.cpp b/Engine/Source/Editor/UnrealEd/Private/SSkeletonWidget.cpp
--- a/Engine/Source/Editor/UnrealEd/Private/SSkeletonWidget.cpp
+++ b/Engine/Source/Editor/UnrealEd/Private/SSkeletonWidget.cpp
@@ -356,9 +356,9 @@ bool SAnimationRemapSkeleton::OnShouldFilterAsset(const class FAssetData& AssetD
 		{
 			URig * Rig = OldSkeleton->GetRig();
 
-			const FString* Value = AssetData.TagsAndValues.Find(USkeleton::RigTag);
+			const FString Value = AssetData.GetTagValueRef<FString>(USkeleton::RigTag);
 
-			if(Value && Rig->GetFullName() == *Value)
+			if(Rig->GetFullName() == Value)
 			{
 				return false;
 			}
@@ -1732,7 +1732,7 @@ FReply SAnimationRemapAssets::OnBestGuessClicked()
 
 	for(const FAssetData& Data : AssetDataList)
 	{
-		if(Data.TagsAndValues.FindRef("Skeleton") == SkeletonName)
+		if(Data.GetTagValueRef<FString>("Skeleton") == SkeletonName)
 		{
 			CompatibleAssets.Add(Data);
 		}
@@ -1921,7 +1921,7 @@ void SAssetEntryRow::OnAssetSelected(const FAssetData& AssetData)
 
 bool SAssetEntryRow::OnShouldFilterAsset(const FAssetData& AssetData) const
 {
-	if(AssetData.TagsAndValues.FindRef("Skeleton") == SkeletonExportName)
+	if(AssetData.GetTagValueRef<FString>("Skeleton") == SkeletonExportName)
 	{
 		return false;
 	}
diff --git a/Engine/Source/Editor/UnrealEd/Private/Settings/LevelEditorPlaySettingsCustomization.h b/Engine/Source/Editor/UnrealEd/Private/Settings/LevelEditorPlaySettingsCustomization.h
--- a/Engine/Source/Editor/UnrealEd/Private/Settings/LevelEditorPlaySettingsCustomization.h
+++ b/Engine/Source/Editor/UnrealEd/Private/Settings/LevelEditorPlaySettingsCustomization.h
@@ -407,7 +407,7 @@ public:
 			WindowPositionHandle->MarkHiddenByCustomization();
 			CenterNewWindowHandle->MarkHiddenByCustomization();
 
-			PlayInStandaloneCategory.AddCustomRow(LOCTEXT("PlayInStandaloneWindowDetails", "Standalone Window Size"), false)
+			PlayInStandaloneCategory.AddCustomRow(LOCTEXT("StandaloneWindowSizeName", "Standalone Window Size"), false)
 				.NameContent()
 				[
 					SNew(STextBlock)
@@ -421,7 +421,7 @@ public:
 					SNew(SScreenResolutionCustomization, &LayoutBuilder, WindowHeightHandle, WindowWidthHandle)
 				];
 
-			PlayInStandaloneCategory.AddCustomRow(LOCTEXT("PlayInStandaloneWindowDetails", "Standalone Window Position"), false)
+			PlayInStandaloneCategory.AddCustomRow(LOCTEXT("StandaloneWindowPosName", "Standalone Window Position"), false)
 				.NameContent()
 				[
 					SNew(STextBlock)
diff --git a/Engine/Source/Editor/UnrealEd/Private/SkeletalMeshEdit.cpp b/Engine/Source/Editor/UnrealEd/Private/SkeletalMeshEdit.cpp
--- a/Engine/Source/Editor/UnrealEd/Private/SkeletalMeshEdit.cpp
+++ b/Engine/Source/Editor/UnrealEd/Private/SkeletalMeshEdit.cpp
@@ -1074,7 +1074,7 @@ bool UnFbx::FFbxImporter::ImportAnimation(USkeleton* Skeleton, UAnimSequence * D
 	//
 	if (ImportOptions->bImportCustomAttribute)
 	{
-		GWarn->BeginSlowTask( LOCTEXT("BeginImportMorphTargetCurves", "Importing Custom Attirubte Curves"), true);
+		GWarn->BeginSlowTask( LOCTEXT("BeginImportCustomAttributeCurves", "Importing Custom Attribute Curves"), true);
 		const int32 TotalLinks = SortedLinks.Num();
 		int32 CurLinkIndex=0;
 		for(auto Node: SortedLinks)
diff --git a/Engine/Source/Editor/UnrealEd/Private/SoundCueGraphEditorCommands.cpp b/Engine/Source/Editor/UnrealEd/Private/SoundCueGraphEditorCommands.cpp
--- a/Engine/Source/Editor/UnrealEd/Private/SoundCueGraphEditorCommands.cpp
+++ b/Engine/Source/Editor/UnrealEd/Private/SoundCueGraphEditorCommands.cpp
@@ -3,7 +3,7 @@
 #include "UnrealEd.h"
 #include "SoundCueGraphEditorCommands.h"
 
-#define LOCTEXT_NAMESPACE ""
+#define LOCTEXT_NAMESPACE "SoundCueGraphEditorCommands"
 
 void FSoundCueGraphEditorCommands::RegisterCommands()
 {
diff --git a/Engine/Source/Editor/UnrealEd/Private/Tests/EditorAssetAutomationTests.cpp b/Engine/Source/Editor/UnrealEd/Private/Tests/EditorAssetAutomationTests.cpp
--- a/Engine/Source/Editor/UnrealEd/Private/Tests/EditorAssetAutomationTests.cpp
+++ b/Engine/Source/Editor/UnrealEd/Private/Tests/EditorAssetAutomationTests.cpp
@@ -1,4 +1,4 @@
-// Copyright 1998-2016 Epic Games, Inc. All Rights Reserved.
+﻿// Copyright 1998-2016 Epic Games, Inc. All Rights Reserved.
 
 #include "UnrealEd.h"
 #include "Tests/AutomationTestSettings.h"
@@ -1069,7 +1069,25 @@ bool FLogImportExportTestResultsCommand::Update()
 /**
 * Automation test to import, open, screenshot, and export assets
 */
-IMPLEMENT_SIMPLE_AUTOMATION_TEST(FAssetImportEditorTest, "Project.Editor.Content.Asset Import and Export", EAutomationTestFlags::EditorContext | EAutomationTestFlags::EngineFilter);
+IMPLEMENT_COMPLEX_AUTOMATION_TEST(FAssetImportEditorTest, "Project.Editor.Content.Asset Import and Export", EAutomationTestFlags::EditorContext | EAutomationTestFlags::EngineFilter);
+void FAssetImportEditorTest::GetTests(TArray<FString>& OutBeautifiedNames, TArray<FString>& OutTestCommands) const
+{
+	UAutomationTestSettings const* AutomationTestSettings = GetDefault<UAutomationTestSettings>();
+	check(AutomationTestSettings);
+	
+	for ( int32 TestIdx = 0; TestIdx < AutomationTestSettings->ImportExportTestDefinitions.Num(); ++TestIdx )
+	{
+		// Lets just use the filename with no path, and change .'s to _'s so the parser
+		// doesn't get confused and insert a bunch of children tests from all the .'s in the path.
+		// TODO Introduce a way to escape .'s
+		FString CleanFileName = FPaths::GetCleanFilename(AutomationTestSettings->ImportExportTestDefinitions[TestIdx].ImportFilePath.FilePath);
+		CleanFileName = CleanFileName.Replace(TEXT("."), TEXT("_"));
+
+		OutBeautifiedNames.Add(CleanFileName);
+		OutTestCommands.Add(FString::FromInt(TestIdx));
+	}
+}
+
 bool FAssetImportEditorTest::RunTest(const FString& Parameters)
 {
 	UAutomationTestSettings const* AutomationTestSettings = GetDefault<UAutomationTestSettings>();
@@ -1077,11 +1095,10 @@ bool FAssetImportEditorTest::RunTest(const FString& Parameters)
 
 	TSharedPtr<ImportExportAssetHelper::FAssetImportStats> BuildStats = MakeShareable(new ImportExportAssetHelper::FAssetImportStats());
 
-	for (int32 TestIdx = 0; TestIdx < AutomationTestSettings->ImportExportTestDefinitions.Num(); ++TestIdx)
-	{
-		TSharedPtr<ImportExportAssetHelper::FAssetInfo> AssetInfo = MakeShareable(new ImportExportAssetHelper::FAssetInfo(AutomationTestSettings->ImportExportTestDefinitions[TestIdx], &ExecutionInfo, BuildStats));
-		ADD_LATENT_AUTOMATION_COMMAND(FImportExportAssetCommand(AssetInfo));
-	}
+	int32 TestIdx = FCString::Atoi(*Parameters);
+	
+	TSharedPtr<ImportExportAssetHelper::FAssetInfo> AssetInfo = MakeShareable(new ImportExportAssetHelper::FAssetInfo(AutomationTestSettings->ImportExportTestDefinitions[TestIdx], &ExecutionInfo, BuildStats));
+	ADD_LATENT_AUTOMATION_COMMAND(FImportExportAssetCommand(AssetInfo));
 
 	ADD_LATENT_AUTOMATION_COMMAND(FLogImportExportTestResultsCommand(BuildStats));
 
diff --git a/Engine/Source/Editor/UnrealEd/Private/Toolkits/AssetEditorCommonCommands.cpp b/Engine/Source/Editor/UnrealEd/Private/Toolkits/AssetEditorCommonCommands.cpp
--- a/Engine/Source/Editor/UnrealEd/Private/Toolkits/AssetEditorCommonCommands.cpp
+++ b/Engine/Source/Editor/UnrealEd/Private/Toolkits/AssetEditorCommonCommands.cpp
@@ -3,7 +3,7 @@
 #include "UnrealEd.h"
 #include "AssetEditorCommonCommands.h"
 
-#define LOCTEXT_NAMESPACE ""
+#define LOCTEXT_NAMESPACE "AssetEditorCommonCommands"
 
 void FAssetEditorCommonCommands::RegisterCommands()
 {
diff --git a/Engine/Source/Editor/UnrealEd/Private/Toolkits/AssetEditorManager.cpp b/Engine/Source/Editor/UnrealEd/Private/Toolkits/AssetEditorManager.cpp
--- a/Engine/Source/Editor/UnrealEd/Private/Toolkits/AssetEditorManager.cpp
+++ b/Engine/Source/Editor/UnrealEd/Private/Toolkits/AssetEditorManager.cpp
@@ -394,7 +394,7 @@ bool FAssetEditorManager::OpenEditorForAssets( const TArray< UObject* >& Assets,
 
 			if (AssetTypeActions.IsValid())
 			{
-				GWarn->BeginSlowTask(LOCTEXT("OpenEditor", "Opening Editor(s)..."), true);
+				GWarn->BeginSlowTask(LOCTEXT("OpenEditors", "Opening Editor(s)..."), true);
 
 				// Determine the appropriate toolkit mode for the asset type
 				auto ActualToolkitMode = ToolkitMode;
diff --git a/Engine/Source/Editor/UnrealEd/Private/Toolkits/AssetEditorToolkit.cpp b/Engine/Source/Editor/UnrealEd/Private/Toolkits/AssetEditorToolkit.cpp
--- a/Engine/Source/Editor/UnrealEd/Private/Toolkits/AssetEditorToolkit.cpp
+++ b/Engine/Source/Editor/UnrealEd/Private/Toolkits/AssetEditorToolkit.cpp
@@ -8,7 +8,7 @@
 #include "Toolkits/GlobalEditorCommonCommands.h"
 #include "Editor/MainFrame/Public/MainFrame.h"
 #include "Editor/WorkspaceMenuStructure/Public/WorkspaceMenuStructureModule.h"
-#include "ClassIconFinder.h"
+#include "SlateIconFinder.h"
 #include "CollectionManagerModule.h"
 #include "IUserFeedbackModule.h"
 #include "IDocumentation.h"
@@ -601,7 +601,7 @@ const FSlateBrush* FAssetEditorToolkit::GetDefaultTabIcon() const
 
 	for (auto ObjectIt = EditingObjects.CreateConstIterator(); ObjectIt; ++ObjectIt)
 	{
-		const FSlateBrush* ThisAssetBrush = FClassIconFinder::FindIconForClass((*ObjectIt)->GetClass());
+		const FSlateBrush* ThisAssetBrush = FSlateIconFinder::FindIconBrushForClass((*ObjectIt)->GetClass());
 		
 		if (!IconBrush)
 		{
diff --git a/Engine/Source/Editor/UnrealEd/Private/UnrealEdEngine.cpp b/Engine/Source/Editor/UnrealEd/Private/UnrealEdEngine.cpp
--- a/Engine/Source/Editor/UnrealEd/Private/UnrealEdEngine.cpp
+++ b/Engine/Source/Editor/UnrealEd/Private/UnrealEdEngine.cpp
@@ -111,7 +111,7 @@ void UUnrealEdEngine::Init(IEngineLoop* InEngineLoop)
 		SpriteIDToIndexMap.Add( SpriteInfo.Category, InfoIndex );
 	}
 
-	if (FPaths::IsProjectFilePathSet() && GIsEditor)
+	if (FPaths::IsProjectFilePathSet() && GIsEditor && !FApp::IsUnattended())
 	{
 		AutoReimportManager = NewObject<UAutoReimportManager>();
 		AutoReimportManager->Initialize();
@@ -517,102 +517,132 @@ void UUnrealEdEngine::OnPackageDirtyStateUpdated( UPackage* Pkg)
 			!bAlreadyAsked && // Don't ask if we already asked once!
 			(Settings->bPromptForCheckoutOnAssetModification || Settings->bAutomaticallyCheckoutOnAssetModification) )
 		{
-			// Force source control state to be updated
-			ISourceControlProvider& SourceControlProvider = ISourceControlModule::Get().GetProvider();
-
-			TArray<FString> Files;
-			Files.Add(SourceControlHelpers::PackageFilename(Package));
-			SourceControlProvider.Execute(ISourceControlOperation::Create<FUpdateStatus>(), SourceControlHelpers::AbsoluteFilenames(Files), EConcurrency::Asynchronous, FSourceControlOperationComplete::CreateUObject(this, &UUnrealEdEngine::OnSourceControlStateUpdated, TWeakObjectPtr<UPackage>(Package)));
+			PackagesDirtiedThisTick.Add(Package);
+			PackageToNotifyState.Add(Package, NS_Updating);
 		}
 	}
 	else
 	{
 		// This package was saved, the user should be prompted again if they checked in the package
+		PackagesDirtiedThisTick.Remove(Package);
 		PackageToNotifyState.Remove( Package );
 	}
 }
 
-
-void UUnrealEdEngine::OnSourceControlStateUpdated(const FSourceControlOperationRef& SourceControlOp, ECommandResult::Type ResultType, TWeakObjectPtr<UPackage> Package)
+void UUnrealEdEngine::AttemptModifiedPackageNotification()
 {
-	if (ResultType == ECommandResult::Succeeded && Package.IsValid())
+	bool bIsCooking = CookServer && CookServer->IsCookingInEditor() && CookServer->IsCookByTheBookRunning();
+
+	if (PackagesDirtiedThisTick.Num() > 0 && !bIsCooking)
 	{
-		// Get the source control state of the package
+		// Force source control state to be updated
 		ISourceControlProvider& SourceControlProvider = ISourceControlModule::Get().GetProvider();
-		FSourceControlStatePtr SourceControlState = SourceControlProvider.GetState(Package.Get(), EStateCacheUsage::Use);
 
-		if (SourceControlState.IsValid())
+		TArray<FString> Files;
+		TArray<TWeakObjectPtr<UPackage>> Packages;
+		for (const auto& Package : PackagesDirtiedThisTick)
 		{
-			const UEditorLoadingSavingSettings* Settings = GetDefault<UEditorLoadingSavingSettings>();
-			check(Settings->bPromptForCheckoutOnAssetModification || Settings->bAutomaticallyCheckoutOnAssetModification);
-
-			if (Settings->bAutomaticallyCheckoutOnAssetModification && SourceControlState->CanCheckout())
+			if (Package.IsValid())
 			{
-				// Automatically check out asset
-				TArray<FString> Files;
+				Packages.Add(Package);
 				Files.Add(SourceControlHelpers::PackageFilename(Package.Get()));
-				SourceControlProvider.Execute(ISourceControlOperation::Create<FCheckOut>(), SourceControlHelpers::AbsoluteFilenames(Files), EConcurrency::Asynchronous, FSourceControlOperationComplete::CreateUObject(this, &UUnrealEdEngine::OnPackageCheckedOut, TWeakObjectPtr<UPackage>(Package)));
-			}
-			else
-			{
-				bool bIsCooking = CookServer && CookServer->IsCookingInEditor() && CookServer->IsCookByTheBookRunning();
-				// Note when cooking in the editor we ignore package notifications.  The cooker is saving packages in a temp location which generates bogus checkout messages otherwise.
-				if ((SourceControlState->CanCheckout() || !SourceControlState->IsCurrent() || SourceControlState->IsCheckedOutOther()) && !bIsCooking )
-				{
-					// To get here, either "prompt for checkout on asset modification" is set, or "automatically checkout on asset modification"
-					// is set, but it failed.
-
-					// Allow packages that are not checked out to pass through.
-					// Allow packages that are not current or checked out by others pass through.  
-					// The user wont be able to checkout these packages but the checkout dialog will show up with a special icon 
-					// to let the user know they wont be able to checkout the package they are modifying.
-
-					PackageToNotifyState.Add(Package, SourceControlState->CanCheckout() ? NS_PendingPrompt : NS_PendingWarning);
-					// We need to prompt since a new package was added
-					bNeedToPromptForCheckout = true;
-				}
 			}
 		}
+		SourceControlProvider.Execute(ISourceControlOperation::Create<FUpdateStatus>(), SourceControlHelpers::AbsoluteFilenames(Files), EConcurrency::Asynchronous, FSourceControlOperationComplete::CreateUObject(this, &UUnrealEdEngine::OnSourceControlStateUpdated, Packages));
+
 	}
-}
 
+	PackagesDirtiedThisTick.Empty();
+}
 
-void UUnrealEdEngine::OnPackageCheckedOut(const FSourceControlOperationRef& SourceControlOp, ECommandResult::Type ResultType, TWeakObjectPtr<UPackage> Package)
+void UUnrealEdEngine::OnSourceControlStateUpdated(const FSourceControlOperationRef& SourceControlOp, ECommandResult::Type ResultType, TArray<TWeakObjectPtr<UPackage>> Packages)
 {
-	if (Package.IsValid())
+	if (ResultType == ECommandResult::Succeeded)
 	{
+		bool bShowNotification = false;
+
 		// Get the source control state of the package
 		ISourceControlProvider& SourceControlProvider = ISourceControlModule::Get().GetProvider();
-		FSourceControlStatePtr SourceControlState = SourceControlProvider.GetState(Package.Get(), EStateCacheUsage::Use);
 
-		FFormatNamedArguments Arguments;
-		Arguments.Add(TEXT("Package"), FText::FromString(Package->GetName()));
+		TArray<TWeakObjectPtr<UPackage>> PackagesToAutomaticallyCheckOut;
+		TArray<FString> FilesToAutomaticallyCheckOut;
 
-		if (ResultType == ECommandResult::Succeeded)
+		const UEditorLoadingSavingSettings* Settings = GetDefault<UEditorLoadingSavingSettings>();
+		for (const auto& PackagePtr : Packages)
 		{
-			if (SourceControlState.IsValid() && SourceControlState->IsCheckedOut())
-	{
-				FNotificationInfo Notification(FText::Format(NSLOCTEXT("SourceControl", "AutoCheckOutNotification", "Package '{Package}' automatically checked out."), Arguments));
-				Notification.bFireAndForget = true;
-				Notification.ExpireDuration = 4.0f;
-				Notification.bUseThrobber = true;
-
-				FSlateNotificationManager::Get().AddNotification(Notification);
+			if (PackagePtr.IsValid())
+			{
+				UPackage* Package = PackagePtr.Get();
 
-				return;
+				FSourceControlStatePtr SourceControlState = SourceControlProvider.GetState(Package, EStateCacheUsage::Use);
+				if (SourceControlState.IsValid())
+				{
+					if (SourceControlState->CanCheckout())
+					{
+						if (Settings->bAutomaticallyCheckoutOnAssetModification)
+						{
+							PackagesToAutomaticallyCheckOut.Add(PackagePtr);
+							FilesToAutomaticallyCheckOut.Add(SourceControlHelpers::PackageFilename(Package));
+						}
+						else
+						{
+							PackageToNotifyState.Add(PackagePtr, NS_PendingPrompt);
+							bShowNotification = true;
+						}
+					}
+					else if (!SourceControlState->IsCurrent() || SourceControlState->IsCheckedOutOther())
+					{
+						PackageToNotifyState.Add(PackagePtr, NS_PendingWarning);
+						bShowNotification = true;
+					}
+				}
 			}
 		}
 
-		FNotificationInfo ErrorNotification(FText::Format(NSLOCTEXT("SourceControl", "AutoCheckOutFailedNotification", "Unable to automatically check out Package '{Package}'."), Arguments));
+		if (FilesToAutomaticallyCheckOut.Num() > 0)
+		{
+			SourceControlProvider.Execute(ISourceControlOperation::Create<FCheckOut>(), SourceControlHelpers::AbsoluteFilenames(FilesToAutomaticallyCheckOut), EConcurrency::Asynchronous, FSourceControlOperationComplete::CreateUObject(this, &UUnrealEdEngine::OnPackagesCheckedOut, PackagesToAutomaticallyCheckOut));
+		}
+
+		if (bShowNotification)
+		{
+			ShowPackageNotification();
+		}
+	}
+}
+
+
+void UUnrealEdEngine::OnPackagesCheckedOut(const FSourceControlOperationRef& SourceControlOp, ECommandResult::Type ResultType, TArray<TWeakObjectPtr<UPackage>> Packages)
+{
+	if (ResultType == ECommandResult::Succeeded)
+	{
+		FNotificationInfo Notification(NSLOCTEXT("SourceControl", "AutoCheckOutNotification", "Packages automatically checked out."));
+		Notification.bFireAndForget = true;
+		Notification.ExpireDuration = 4.0f;
+		Notification.bUseThrobber = true;
+
+		FSlateNotificationManager::Get().AddNotification(Notification);
+
+		for (const auto& Package : Packages)
+		{
+			PackageToNotifyState.Add(Package, NS_DialogPrompted);
+		}
+	}
+	else
+	{
+		FNotificationInfo ErrorNotification(NSLOCTEXT("SourceControl", "AutoCheckOutFailedNotification", "Unable to automatically check out packages."));
 		ErrorNotification.bFireAndForget = true;
 		ErrorNotification.ExpireDuration = 4.0f;
 		ErrorNotification.bUseThrobber = true;
 
 		FSlateNotificationManager::Get().AddNotification(ErrorNotification);
 
-		// Automatic checkout failed - pop up the notification for manual checkout
-		PackageToNotifyState.Add(Package, SourceControlState->CanCheckout() ? NS_PendingPrompt : NS_PendingWarning);
-		bNeedToPromptForCheckout = true;
+		for (const auto& Package : Packages)
+		{
+			PackageToNotifyState.Add(Package, NS_PendingPrompt);
+		}
+
+		ShowPackageNotification();
 	}
 }
 
diff --git a/Engine/Source/Editor/UnrealEd/Private/UnrealEdSrv.cpp b/Engine/Source/Editor/UnrealEd/Private/UnrealEdSrv.cpp
--- a/Engine/Source/Editor/UnrealEd/Private/UnrealEdSrv.cpp
+++ b/Engine/Source/Editor/UnrealEd/Private/UnrealEdSrv.cpp
@@ -1436,9 +1436,9 @@ bool UUnrealEdEngine::IsUserInteracting()
 	return bUserIsInteracting;
 }
 
-void UUnrealEdEngine::AttemptModifiedPackageNotification()
+void UUnrealEdEngine::ShowPackageNotification()
 {
-	if( bNeedToPromptForCheckout && !FApp::IsUnattended() )
+	if( !FApp::IsUnattended() )
 	{
 		// Defer prompting for checkout if we cant prompt because of the following:
 		// The user is interacting with something,
@@ -1486,9 +1486,6 @@ void UUnrealEdEngine::AttemptModifiedPackageNotification()
 				// For adding notifications.
 				FSlateNotificationManager::Get().AddNotification(ErrorNotification);
 			}
-
-			// No longer have a pending prompt.
-			bNeedToPromptForCheckout = false;
 		}
 	}
 }
diff --git a/Engine/Source/Editor/UnrealEd/Public/ClassIconFinder.h b/Engine/Source/Editor/UnrealEd/Public/ClassIconFinder.h
--- a/Engine/Source/Editor/UnrealEd/Public/ClassIconFinder.h
+++ b/Engine/Source/Editor/UnrealEd/Public/ClassIconFinder.h
@@ -2,16 +2,13 @@
 
 #pragma once
 
+#include "SlateIconFinder.h"
+
 class FAssetData;
 
 class FClassIconFinder
 {
 public:
-	/** Registers a new style set to use to try to find icons. */
-	UNREALED_API static void RegisterIconSource(const ISlateStyle* StyleSet);
-
-	/** Unregisters an style set. */
-	UNREALED_API static void UnregisterIconSource(const ISlateStyle* StyleSet);
 
 	/** Find the best fitting small icon to use for the supplied actor array */
 	UNREALED_API static const FSlateBrush* FindIconForActors(const TArray< TWeakObjectPtr<AActor> >& InActors, UClass*& CommonBaseClass);
@@ -19,34 +16,59 @@ public:
 	/** Find the small icon to use for the supplied actor */
 	UNREALED_API static const FSlateBrush* FindIconForActor( const TWeakObjectPtr<AActor>& InActor );
 	
-	/** Find the small icon name to use for the supplied actor */
-	UNREALED_API static FName FindIconNameForActor( const TWeakObjectPtr<AActor>& InActor );
+	/** Find the small icon to use for the supplied actor */
+	UNREALED_API static FSlateIcon FindSlateIconForActor( const TWeakObjectPtr<AActor>& InActor );
 
-	/** Find the small icon to use for the supplied class */
-	UNREALED_API static const FSlateBrush* FindIconForClass(const UClass* InClass, const FName& InDefaultName = FName() );
+	/** Utility function to convert a Blueprint into the most suitable class possible for use by the icon finder */
+	UNREALED_API static const UClass* GetIconClassForBlueprint(const UBlueprint* InBlueprint);
 
-	/** Find the small icon name to use for the supplied class */
-	UNREALED_API static FName FindIconNameForClass(const UClass* InClass, const FName& InDefaultName = FName() );
+	/** Utility function to convert an asset into the most suitable class possible for use by the icon finder */
+	UNREALED_API static const UClass* GetIconClassForAssetData(const FAssetData& InAssetData, bool* bOutIsClassType = nullptr);
 
 	/** Find the large thumbnail name to use for the supplied class */
-	UNREALED_API static const FSlateBrush* FindThumbnailForClass(const UClass* InClass, const FName& InDefaultName = FName() );
+	UNREALED_API static const FSlateBrush* FindThumbnailForClass(const UClass* InClass, const FName& InDefaultName = FName() )
+	{
+		return FSlateIconFinder::FindCustomIconBrushForClass(InClass, TEXT("ClassThumbnail"), InDefaultName);
+	}
 
-	/** Find the large thumbnail name to use for the supplied class */
-	UNREALED_API static FName FindThumbnailNameForClass(const UClass* InClass, const FName& InDefaultName = FName() );
+public:
 
-	/** Utility function to convert a Blueprint into the most suitable class possible for use by the icon finder */
-	UNREALED_API static const UClass* GetIconClassForBlueprint(const UBlueprint* InBlueprint);
+	/** Registers a new style set to use to try to find icons. */
+	DEPRECATED(4.13, "This function is no longer required.")
+	static void RegisterIconSource(const ISlateStyle* StyleSet){}
 
-	/** Utility function to convert an asset into the most suitable class possible for use by the icon finder */
-	UNREALED_API static const UClass* GetIconClassForAssetData(const FAssetData& InAssetData, bool* bOutIsClassType = nullptr);
+	/** Unregisters an style set. */
+	DEPRECATED(4.13, "This function is no longer required.")
+	static void UnregisterIconSource(const ISlateStyle* StyleSet){}
 
-private:
-	/** Find a thumbnail/icon name to use for the supplied class */
-	UNREALED_API static FName FindIconNameImpl(const UClass* InClass, const FName& InDefaultName = FName(), const TCHAR* StyleRoot = TEXT("ClassIcon") );
+	/** Find the small icon to use for the supplied class */
+	DEPRECATED(4.13, "Please use FSlateIconFinder::FindIconForClass directly (followed by FSlateIcon::GetIcon().")
+	static const FSlateBrush* FindIconForClass(const UClass* InClass, const FName& InDefaultName = FName() )
+	{
+		return FSlateIconFinder::FindIconBrushForClass(InClass, InDefaultName);
+	}
+
+	/** Find the small icon name to use for the supplied class */
+	DEPRECATED(4.13, "Please use FSlateIconFinder::FindIconForClass directly (followed by FSlateIcon::GetStyleName().")
+	static FName FindIconNameForClass(const UClass* InClass, const FName& InDefaultName = FName() )
+	{
+		return FSlateIconFinder::FindIconForClass(InClass, InDefaultName).GetStyleName();
+	}
+
+	/** Find the large thumbnail name to use for the supplied class */
+	DEPRECATED(4.13, "Please use FSlateIconFinder::FindCustomIconForClass directly (followed by FSlateIcon::GetStyleName().")
+	static FName FindThumbnailNameForClass(const UClass* InClass, const FName& InDefaultName = FName() )
+	{
+		return FSlateIconFinder::FindCustomIconForClass(InClass, TEXT("ClassThumbnail"), InDefaultName).GetStyleName();
+	}
+
+	/** Find the small icon name to use for the supplied actor */
+	DEPRECATED(4.13, "Please use FindSlateIconForActor directly.")
+	static FName FindIconNameForActor( const TWeakObjectPtr<AActor>& InActor )
+	{
+		return FindSlateIconForActor(InActor).GetStyleName();
+	}
 
 private:
 	UNREALED_API static const FSlateBrush* LookupBrush(FName IconName);
-
-	/** The available style sets to source icons from */
-	static TArray < const ISlateStyle* > Styles;
 };
diff --git a/Engine/Source/Editor/UnrealEd/Public/DragAndDrop/ClassDragDropOp.h b/Engine/Source/Editor/UnrealEd/Public/DragAndDrop/ClassDragDropOp.h
--- a/Engine/Source/Editor/UnrealEd/Public/DragAndDrop/ClassDragDropOp.h
+++ b/Engine/Source/Editor/UnrealEd/Public/DragAndDrop/ClassDragDropOp.h
@@ -1,7 +1,7 @@
 // Copyright 1998-2016 Epic Games, Inc. All Rights Reserved.
 #pragma once
 
-#include "ClassIconFinder.h"
+#include "SlateIconFinder.h"
 
 class FClassDragDropOp : public FDragDropOperation
 {
@@ -15,7 +15,7 @@ public:
 	virtual TSharedPtr<SWidget> GetDefaultDecorator() const override
 	{
 		// Just use the first class for the cursor decorator.
-		const FSlateBrush* ClassIcon = FClassIconFinder::FindIconForClass( ClassesToDrop[0].Get() );
+		const FSlateBrush* ClassIcon = FSlateIconFinder::FindIconBrushForClass( ClassesToDrop[0].Get() );
 
 		// If the class icon is the default brush, do not put it in the cursor decoration window.
 		if(ClassIcon)
diff --git a/Engine/Source/Editor/UnrealEd/Public/EdGraphNode_Comment.h b/Engine/Source/Editor/UnrealEd/Public/EdGraphNode_Comment.h
--- a/Engine/Source/Editor/UnrealEd/Public/EdGraphNode_Comment.h
+++ b/Engine/Source/Editor/UnrealEd/Public/EdGraphNode_Comment.h
@@ -65,7 +65,7 @@ public:
 	virtual bool ShouldDrawNodeAsComment() const override { return true; }
 	UNREALED_API virtual FString GetDocumentationLink() const override;
 	UNREALED_API virtual FString GetDocumentationExcerptName() const override;
-	UNREALED_API virtual FName GetPaletteIcon(FLinearColor& OutColor) const override;
+	UNREALED_API virtual FSlateIcon GetIconAndTint(FLinearColor& OutColor) const override;
 	//~ End UEdGraphNode Interface
 
 	/** Add a node that will be dragged when this comment is dragged */
diff --git a/Engine/Source/Editor/UnrealEd/Public/GameModeInfoCustomizer.h b/Engine/Source/Editor/UnrealEd/Public/GameModeInfoCustomizer.h
--- a/Engine/Source/Editor/UnrealEd/Public/GameModeInfoCustomizer.h
+++ b/Engine/Source/Editor/UnrealEd/Public/GameModeInfoCustomizer.h
@@ -259,6 +259,8 @@ public:
 
 	void OnMakeSelectedDefaultClassClicked(FName ClassPropertyName)
 	{
+		FEditorDelegates::LoadSelectedAssetsIfNeeded.Broadcast();
+
 		UClassProperty* ClassProp = FindFieldChecked<UClassProperty>(AGameMode::StaticClass(), ClassPropertyName);
 		const UClass* SelectedClass = GEditor->GetFirstSelectedClass(ClassProp->MetaClass);
 		if (SelectedClass)
@@ -298,6 +300,8 @@ public:
 
 	void OnUseSelectedGameModeClicked(IDetailLayoutBuilder* DetailLayout)
 	{
+		FEditorDelegates::LoadSelectedAssetsIfNeeded.Broadcast();
+
 		const UClass* SelectedClass = GEditor->GetFirstSelectedClass(AGameMode::StaticClass());
 		if (SelectedClass)
 		{
diff --git a/Engine/Source/Editor/VREditor/VREditorActions.cpp b/Engine/Source/Editor/VREditor/VREditorActions.cpp
--- a/Engine/Source/Editor/VREditor/VREditorActions.cpp
+++ b/Engine/Source/Editor/VREditor/VREditorActions.cpp
@@ -3,7 +3,7 @@
 #include "VREditorModule.h"
 #include "VREditorActions.h"
 
-#define LOCTEXT_NAMESPACE ""
+#define LOCTEXT_NAMESPACE "VREditorActions"
 
 void FVREditorActions::RegisterCommands()
 {
diff --git a/Engine/Source/Editor/WorkspaceMenuStructure/WorkspaceMenuStructure.Build.cs b/Engine/Source/Editor/WorkspaceMenuStructure/WorkspaceMenuStructure.Build.cs
--- a/Engine/Source/Editor/WorkspaceMenuStructure/WorkspaceMenuStructure.Build.cs
+++ b/Engine/Source/Editor/WorkspaceMenuStructure/WorkspaceMenuStructure.Build.cs
@@ -12,7 +12,7 @@ public class WorkspaceMenuStructure : ModuleRules
 			new string[] {
 				"Core",
 				"EditorStyle",
-				"Slate",
+				"SlateCore",
 			}
 		);
 	}
diff --git a/Engine/Source/Editor/WorldBrowser/Private/LevelCollectionCommands.h b/Engine/Source/Editor/WorldBrowser/Private/LevelCollectionCommands.h
--- a/Engine/Source/Editor/WorldBrowser/Private/LevelCollectionCommands.h
+++ b/Engine/Source/Editor/WorldBrowser/Private/LevelCollectionCommands.h
@@ -1,7 +1,7 @@
 // Copyright 1998-2016 Epic Games, Inc. All Rights Reserved.
 #pragma once
 
-#define LOCTEXT_NAMESPACE ""
+#define LOCTEXT_NAMESPACE "LevelCollectionCommands"
 
 /** The set of commands supported by the WorldView */
 class FLevelCollectionCommands 
diff --git a/Engine/Source/Editor/WorldBrowser/Private/SWorldHierarchyItem.cpp b/Engine/Source/Editor/WorldBrowser/Private/SWorldHierarchyItem.cpp
--- a/Engine/Source/Editor/WorldBrowser/Private/SWorldHierarchyItem.cpp
+++ b/Engine/Source/Editor/WorldBrowser/Private/SWorldHierarchyItem.cpp
@@ -296,7 +296,7 @@ void SWorldHierarchyItem::OnColorPickerCancelled(FLinearColor OriginalColor)
 
 void SWorldHierarchyItem::OnColorPickerInteractiveBegin()
 {
-	GEditor->BeginTransaction(NSLOCTEXT("FColorStructCustomization", "SetColorProperty", "Edit Level Draw Color"));
+	GEditor->BeginTransaction(LOCTEXT("EditLevelDragColor", "Edit Level Draw Color"));
 }
 
 void SWorldHierarchyItem::OnColorPickerInteractiveEnd()
diff --git a/Engine/Source/Editor/WorldBrowser/Private/Tiles/STiledLandscapeImportDlg.cpp b/Engine/Source/Editor/WorldBrowser/Private/Tiles/STiledLandscapeImportDlg.cpp
--- a/Engine/Source/Editor/WorldBrowser/Private/Tiles/STiledLandscapeImportDlg.cpp
+++ b/Engine/Source/Editor/WorldBrowser/Private/Tiles/STiledLandscapeImportDlg.cpp
@@ -528,7 +528,7 @@ FReply STiledLandcapeImportDlg::OnClickedSelectWeighmapTiles(TSharedPtr<FTiledLa
 			const FString ImportFileTypes = TEXT("Weightmap tiles (*.raw, *.png)|*.raw;*.png");
 			bool bOpened = DesktopPlatform->OpenFileDialog(
 								ParentWindowWindowHandle,
-								LOCTEXT("SelectHeightmapTiles", "Select weightmap tiles").ToString(),
+								LOCTEXT("SelectWeightmapTiles", "Select weightmap tiles").ToString(),
 								*FEditorDirectories::Get().GetLastDirectory(ELastDirectory::UNR),
 								TEXT(""),
 								*ImportFileTypes,
@@ -699,7 +699,7 @@ FText STiledLandcapeImportDlg::GenerateConfigurationText(int32 NumComponents, in
 	const FString QuadsStr = FString::Printf(TEXT("%dx%d"), NumQuadsPerSection, NumQuadsPerSection);
 	
 	return FText::Format(
-		LOCTEXT("TiledLandscapeImport_ConfigurationText", "Components: {0} Sections: {1} Quads: {2}"), 
+		LOCTEXT("TiledLandscapeImport_ConfigurationDescFmt", "Components: {0} Sections: {1} Quads: {2}"), 
 		FText::FromString(ComponentsStr),
 		FText::FromString(SectionsStr),
 		FText::FromString(QuadsStr)
