Code from Daniel Lamb to try and narrow down cooking issues

diff --git a/Engine/Source/Editor/UnrealEd/Private/CookOnTheFlyServer.cpp b/Engine/Source/Editor/UnrealEd/Private/CookOnTheFlyServer.cpp
--- a/Engine/Source/Editor/UnrealEd/Private/CookOnTheFlyServer.cpp
+++ b/Engine/Source/Editor/UnrealEd/Private/CookOnTheFlyServer.cpp
@@ -48,7 +48,7 @@
 DEFINE_LOG_CATEGORY_STATIC(LogCook, Log, All);
 
 #define DEBUG_COOKONTHEFLY 0
-#define OUTPUT_TIMING 0
+#define OUTPUT_TIMING 1
 
 #define USEASSETREGISTRYFORDEPENDENTPACKAGES 1
 #define VERIFY_GETDEPENDENTPACKAGES 0 // verify has false hits because old serialization method for generating dependencies had errors (included transient objects which shouldn't be in asset registry), but you can still use verify to build a list then cross check against transient objects.  
@@ -126,13 +126,13 @@ public:
 		Children.Empty();
 	}
 
-	FHierarchicalTimerInfo* FindChild(const FString& Name)
+	FHierarchicalTimerInfo* FindChild(const FString& InName)
 	{
-		FHierarchicalTimerInfo* Child = (FHierarchicalTimerInfo*)(Children.FindRef(Name));
+		FHierarchicalTimerInfo* Child = (FHierarchicalTimerInfo*)(Children.FindRef(InName));
 		if ( !Child )
 		{
-			FString Temp = Name;
-			Child = Children.Add(Name, new FHierarchicalTimerInfo(MoveTemp(Temp)));
+			FString Temp = InName;
+			Child = Children.Add(InName, new FHierarchicalTimerInfo(MoveTemp(Temp)));
 			Child->Parent = this;
 		}
 
@@ -1248,6 +1248,13 @@ void UCookOnTheFlyServer::OnStringAssetReferenceLoadedPackage(const FName& Packa
 		{
 			RequestPackage(StandardPackageName, true); // force to front of queue because we know this package is now loaded
 		}
+		else
+		{
+			if (!FPackageName::IsScriptPackage(PackageFName.ToString()))
+			{
+				UE_LOG(LogCook, Warning, TEXT("Unable to find cached package name for package %s"), *PackageFName.ToString());
+			}
+		}
 	}
 }
 
@@ -2060,16 +2067,7 @@ uint32 UCookOnTheFlyServer::TickCookOnTheSide( const float TimeSlice, uint32 &Co
 			}
 		}
 
-		if ((CookRequests.HasItems() == false) && IsCookByTheBookRunning() && !(Result&COSR_WaitingOnChildCookers))
-		{
-			// make sure we resolve all string asset references and nothing is loaded
-			if (GRedirectCollector.HasAnyStringAssetReferencesToResolve())
-			{
-				// resolve redirectors first
-				GRedirectCollector.ResolveStringAssetReference();
-				continue;
-			}
-		}
+		
 
 		if ( Timer.IsTimeUp() )
 		{
@@ -2078,6 +2076,17 @@ uint32 UCookOnTheFlyServer::TickCookOnTheSide( const float TimeSlice, uint32 &Co
 	}
 	
 
+	if ((CookRequests.HasItems() == false) && IsCookByTheBookRunning() && !(Result&COSR_WaitingOnChildCookers))
+	{
+		// make sure we resolve all string asset references and nothing is loaded
+		if (GRedirectCollector.HasAnyStringAssetReferencesToResolve())
+		{
+			// resolve redirectors first
+			GRedirectCollector.ResolveStringAssetReference();
+		}
+	}
+
+
 	OUTPUT_TIMERS();
 
 	if ( CookByTheBookOptions )
