Copying //UE4/Dev-Rendering to //UE4/Dev-Main (Source: //UE4/Dev-Rendering @ 2998063)

diff --git a/Engine/Source/Developer/ShaderCompilerCommon/Private/ShaderCompilerCommon.cpp b/Engine/Source/Developer/ShaderCompilerCommon/Private/ShaderCompilerCommon.cpp
--- a/Engine/Source/Developer/ShaderCompilerCommon/Private/ShaderCompilerCommon.cpp
+++ b/Engine/Source/Developer/ShaderCompilerCommon/Private/ShaderCompilerCommon.cpp
@@ -350,7 +350,7 @@ namespace CrossCompiler
 		CCTCmdLine += ((CCFlags & HLSLCC_ExpandSubexpressions) == HLSLCC_ExpandSubexpressions) ? TEXT(" -xpxpr") : TEXT("");
 		CCTCmdLine += ((CCFlags & HLSLCC_SeparateShaderObjects) == HLSLCC_SeparateShaderObjects) ? TEXT(" -separateshaders") : TEXT("");
 		CCTCmdLine += ((CCFlags & HLSLCC_PackUniformsIntoUniformBuffers) == HLSLCC_PackUniformsIntoUniformBuffers) ? TEXT(" -packintoubs") : TEXT("");
-		CCTCmdLine += ((CCFlags & HLSLCC_FixAtomicReferences) == HLSLCC_PackUniformsIntoUniformBuffers) ? TEXT(" -fixatomics") : TEXT("");
+		CCTCmdLine += ((CCFlags & HLSLCC_FixAtomicReferences) == HLSLCC_FixAtomicReferences) ? TEXT(" -fixatomics") : TEXT("");
 		FString BatchFile;
 		if (PLATFORM_MAC)
 		{
diff --git a/Engine/Source/Editor/DetailCustomizations/Private/FbxImportUIDetails.cpp b/Engine/Source/Editor/DetailCustomizations/Private/FbxImportUIDetails.cpp
--- a/Engine/Source/Editor/DetailCustomizations/Private/FbxImportUIDetails.cpp
+++ b/Engine/Source/Editor/DetailCustomizations/Private/FbxImportUIDetails.cpp
@@ -73,9 +73,16 @@ void FFbxImportUIDetails::CustomizeDetails( IDetailLayoutBuilder& DetailBuilder
 
 	if(ImportType != FBXIT_Animation)
 	{
-		TSharedRef<IPropertyHandle> ImportSkeletalProp = DetailBuilder.GetProperty(GET_MEMBER_NAME_CHECKED(UFbxImportUI, bImportAsSkeletal));
-		ImportSkeletalProp->SetOnPropertyValueChanged(FSimpleDelegate::CreateSP(this, &FFbxImportUIDetails::MeshImportModeChanged));
-		MeshCategory.AddProperty(ImportSkeletalProp);
+		{
+			TSharedRef<IPropertyHandle> Prop = DetailBuilder.GetProperty(GET_MEMBER_NAME_CHECKED(UFbxImportUI, bImportAsSkeletal));
+			Prop->SetOnPropertyValueChanged(FSimpleDelegate::CreateSP(this, &FFbxImportUIDetails::MeshImportModeChanged));
+			MeshCategory.AddProperty(Prop);
+		}
+		{
+			TSharedRef<IPropertyHandle> Prop = DetailBuilder.GetProperty(GET_MEMBER_NAME_CHECKED(UFbxImportUI, bImportAsSubDSurface));
+			Prop->SetOnPropertyValueChanged(FSimpleDelegate::CreateSP(this, &FFbxImportUIDetails::MeshImportModeChanged));
+			MeshCategory.AddProperty(Prop);
+		}
 	}
 
 	TSharedRef<IPropertyHandle> ImportMeshProp = DetailBuilder.GetProperty(GET_MEMBER_NAME_CHECKED(UFbxImportUI, bImportMesh));
@@ -323,8 +330,7 @@ void FFbxImportUIDetails::MeshImportModeChanged()
 {
 	if(CachedDetailBuilder)
 	{
-		ImportUI->MeshTypeToImport = ImportUI->MeshTypeToImport == FBXIT_SkeletalMesh ? FBXIT_StaticMesh : FBXIT_SkeletalMesh;
-		ImportUI->bImportMesh = true;
+		ImportUI->SetMeshTypeToImport();
 		CachedDetailBuilder->ForceRefreshDetails();
 	}
 }
@@ -335,7 +341,7 @@ void FFbxImportUIDetails::ImportMeshToggleChanged()
 	{
 		if(ImportUI->bImportMesh)
 		{
-			ImportUI->MeshTypeToImport = ImportUI->bImportAsSkeletal ? FBXIT_SkeletalMesh : FBXIT_StaticMesh;
+			ImportUI->SetMeshTypeToImport();
 		}
 		else
 		{
diff --git a/Engine/Source/Editor/DetailCustomizations/Private/FbxImportUIDetails.h b/Engine/Source/Editor/DetailCustomizations/Private/FbxImportUIDetails.h
--- a/Engine/Source/Editor/DetailCustomizations/Private/FbxImportUIDetails.h
+++ b/Engine/Source/Editor/DetailCustomizations/Private/FbxImportUIDetails.h
@@ -19,8 +19,8 @@ public:
 	* @param MetaData the metadata string to validate
 	*/
 	bool IsImportTypeMetaDataValid(EFBXImportType& ImportType, FString& MetaData);
-
-	/** Called if the mesh mode (static / skeletal) changes */
+	
+	/** Called if the mesh mode (static / skeletal / SubDSurface) changes */
 	void MeshImportModeChanged();
 
 	/** Called if the import mesh option for skeletal meshes is changed */
diff --git a/Engine/Source/Editor/DetailCustomizations/Private/PostProcessSettingsCustomization.cpp b/Engine/Source/Editor/DetailCustomizations/Private/PostProcessSettingsCustomization.cpp
--- a/Engine/Source/Editor/DetailCustomizations/Private/PostProcessSettingsCustomization.cpp
+++ b/Engine/Source/Editor/DetailCustomizations/Private/PostProcessSettingsCustomization.cpp
@@ -3,7 +3,6 @@
 #include "DetailCustomizationsPrivatePCH.h"
 #include "PostProcessSettingsCustomization.h"
 #include "ObjectEditorUtils.h"
-#include "LightPropagationVolumeBlendable.h"
 #include "SWidgetSwitcher.h"
 #include "Classes/Engine/BlendableInterface.h"
 
@@ -69,7 +68,7 @@ void FPostProcessSettingsCustomization::CustomizeHeader( TSharedRef<IPropertyHan
 	];
 }
 
-void FWeightedBlendableCustomization::AddDirectAsset(TSharedRef<IPropertyHandle> StructPropertyHandle, UPackage* Package, TSharedPtr<IPropertyHandle> Weight, TSharedPtr<IPropertyHandle> Value)
+void FWeightedBlendableCustomization::AddDirectAsset(TSharedRef<IPropertyHandle> StructPropertyHandle, UPackage* Package, TSharedPtr<IPropertyHandle> Weight, TSharedPtr<IPropertyHandle> Value, UClass* Class)
 {
 	Weight->SetValue(1.0f);
 
@@ -82,7 +81,8 @@ void FWeightedBlendableCustomization::AddDirectAsset(TSharedRef<IPropertyHandle>
 		for(TArray<UObject*>::TConstIterator It = Objects.CreateConstIterator(); It; It++)
 		{
 			UObject* Obj = *It;
-			const UObject* NewObj = NewObject<ULightPropagationVolumeBlendable>(Obj);
+
+			const UObject* NewObj = NewObject<UObject>(Obj, Class);
 
 			FString Str = NewObj->GetPathName();
 
@@ -131,22 +131,49 @@ FReply FWeightedBlendableCustomization::JumpToDirectAsset(TSharedPtr<IPropertyHa
 
 TSharedRef<SWidget> FWeightedBlendableCustomization::GenerateContentWidget(TSharedRef<IPropertyHandle> StructPropertyHandle, UPackage* Package, TSharedPtr<IPropertyHandle> Weight, TSharedPtr<IPropertyHandle> Value)
 {
+	bool bSeparatorIsNeeded = false; 
+
 	FMenuBuilder MenuBuilder(true, NULL);
 	{
-/*
 		for(TObjectIterator<UClass> It; It; ++It)
 		{
-			if(It->IsChildOf(UFactory::StaticClass()))
-			if(It->IsChildOf(UBlendableInterface::StaticClass()))
-*/
-		// todo: iterate through all UFactories that create UClasses with UBlendableInterface
+			if( It->IsChildOf(UFactory::StaticClass()))
+			{
+				UFactory* Factory = It->GetDefaultObject<UFactory>();
+
+				check(Factory);
+
+				UClass* SupportedClass = Factory->GetSupportedClass();
+
+				if(SupportedClass)
+				{
+					if(SupportedClass->ImplementsInterface(UBlendableInterface::StaticClass()))
+					{
+						// At the moment we know about 3 Blendables: Material, UMaterialInstanceConstant, LightPropagationVolumeBlendable
+						// The materials are not that useful to have here (hard to reference) so we suppress them here
+						if(!(
+							SupportedClass == UMaterial::StaticClass() ||
+							SupportedClass == UMaterialInstanceConstant::StaticClass()
+							))
 		{
-			FUIAction Direct2(FExecuteAction::CreateSP(this, &FWeightedBlendableCustomization::AddDirectAsset, StructPropertyHandle, Package, Weight, Value));
-			MenuBuilder.AddMenuEntry(LOCTEXT("Blendable_DirectAsset2", "Light Propagation Volume Blendable"), 
-				LOCTEXT("Blendable_DirectAsset2h", "references Light Propagation Volume Blendable (owned by the containing object)"), FSlateIcon(), Direct2);
+							FUIAction Direct2(FExecuteAction::CreateSP(this, &FWeightedBlendableCustomization::AddDirectAsset, StructPropertyHandle, Package, Weight, Value, SupportedClass));
+
+							FName ClassName = SupportedClass->GetFName();
+						
+							MenuBuilder.AddMenuEntry(FText::FromString(ClassName.GetPlainNameString()),
+								LOCTEXT("Blendable_DirectAsset2h", "Creates an asset that is owned by the containing object"), FSlateIcon(), Direct2);
+
+							bSeparatorIsNeeded = true;
+						}
+					}
+				}
+			}
 		}
 
+		if(bSeparatorIsNeeded)
+		{
 		MenuBuilder.AddMenuSeparator();
+		}
 
 		FUIAction Indirect(FExecuteAction::CreateSP(this, &FWeightedBlendableCustomization::AddIndirectAsset, Weight));
 		MenuBuilder.AddMenuEntry(LOCTEXT("Blendable_IndirectAsset", "Asset reference"), 
@@ -189,6 +216,7 @@ TSharedRef<SWidget> FWeightedBlendableCustomization::GenerateContentWidget(TShar
 	return Switcher;
 }
 
+
 int32 FWeightedBlendableCustomization::ComputeSwitcherIndex(TSharedRef<IPropertyHandle> StructPropertyHandle, UPackage* Package, TSharedPtr<IPropertyHandle> Weight, TSharedPtr<IPropertyHandle> Value) const
 {
 	float WeightValue = 1.0f;
diff --git a/Engine/Source/Editor/DetailCustomizations/Private/PostProcessSettingsCustomization.h b/Engine/Source/Editor/DetailCustomizations/Private/PostProcessSettingsCustomization.h
--- a/Engine/Source/Editor/DetailCustomizations/Private/PostProcessSettingsCustomization.h
+++ b/Engine/Source/Editor/DetailCustomizations/Private/PostProcessSettingsCustomization.h
@@ -58,7 +58,7 @@ public:
 	TSharedRef<SWidget> GenerateContentWidget(TSharedRef<IPropertyHandle> StructPropertyHandle, UPackage* Package, TSharedPtr<IPropertyHandle> Weight, TSharedPtr<IPropertyHandle> Value);
 
 	// "Outer" is the object that has the blendables container
-	void AddDirectAsset(TSharedRef<IPropertyHandle> StructPropertyHandle, UPackage* Package, TSharedPtr<IPropertyHandle> Weight, TSharedPtr<IPropertyHandle> Value);
+	void AddDirectAsset(TSharedRef<IPropertyHandle> StructPropertyHandle, UPackage* Package, TSharedPtr<IPropertyHandle> Weight, TSharedPtr<IPropertyHandle> Value, UClass* Class);
 	// "Outer" is a package
 	void AddIndirectAsset(TSharedPtr<IPropertyHandle> Weight);
 
diff --git a/Engine/Source/Editor/Kismet/Private/BlueprintEditorModule.cpp b/Engine/Source/Editor/Kismet/Private/BlueprintEditorModule.cpp
--- a/Engine/Source/Editor/Kismet/Private/BlueprintEditorModule.cpp
+++ b/Engine/Source/Editor/Kismet/Private/BlueprintEditorModule.cpp
@@ -7,7 +7,7 @@
 #include "BlueprintEditor.h"
 #include "BlueprintEditorUtils.h"
 #include "BlueprintUtilities.h"
-#include "BlueprintWarningsConfigurationPanel.h"
+//#include "BlueprintWarningsConfigurationPanel.h"
 #include "Toolkits/ToolkitManager.h"
 #include "Editor/UnrealEd/Public/Kismet2/DebuggerCommands.h"
 #include "Editor/WorkspaceMenuStructure/Public/WorkspaceMenuStructureModule.h"
@@ -223,14 +223,6 @@ void FBlueprintEditorModule::StartupModule()
 
 	if (ISettingsModule* SettingsModule = FModuleManager::GetModulePtr<ISettingsModule>("Settings"))
 	{
-		ConfigurationPanel = SNew(SBlueprintWarningsConfigurationPanel);
-		SettingsModule->RegisterSettings(
-			"Project", 
-			"Engine", 
-			"Blueprints",
-			LOCTEXT("RuntimeSettingsName", "Blueprints"),
-			LOCTEXT("RuntimeSettingsDescription", "Configure Blueprints runtime behavior"),
-			ConfigurationPanel.ToSharedRef());
 	}
 
 }
diff --git a/Engine/Source/Editor/UnrealEd/Classes/ActorFactories/ActorFactorySubDSurface.h b/Engine/Source/Editor/UnrealEd/Classes/ActorFactories/ActorFactorySubDSurface.h
new file mode 100644
--- /dev/null
+++ b/Engine/Source/Editor/UnrealEd/Classes/ActorFactories/ActorFactorySubDSurface.h
@@ -0,0 +1,17 @@
+// Copyright 1998-2014 Epic Games, Inc. All Rights Reserved.
+
+
+#pragma once
+#include "ActorFactorySubDSurface.generated.h"
+
+UCLASS(MinimalAPI, config=Editor)
+class UActorFactorySubDSurface : public UActorFactory
+{
+	GENERATED_UCLASS_BODY()
+
+	virtual void PostSpawnActor( UObject* Asset, AActor* NewActor ) override;
+	virtual bool CanCreateActorFrom( const FAssetData& AssetData, FText& OutErrorMsg ) override;
+};
+
+
+
diff --git a/Engine/Source/Editor/UnrealEd/Classes/Exporters/RenderTargetExporterHDR.h b/Engine/Source/Editor/UnrealEd/Classes/Exporters/RenderTargetExporterHDR.h
new file mode 100644
--- /dev/null
+++ b/Engine/Source/Editor/UnrealEd/Classes/Exporters/RenderTargetExporterHDR.h
@@ -0,0 +1,21 @@
+// Copyright 1998-2016 Epic Games, Inc. All Rights Reserved.
+
+//~=============================================================================
+// TextureExporterHDR
+//~=============================================================================
+
+#pragma once
+#include "RenderTargetExporterHDR.generated.h"
+
+UCLASS()
+class UNREALED_API URenderTargetExporterHDR : public UExporter
+{
+	GENERATED_UCLASS_BODY()
+
+	//~ Begin UExporter Interface
+	virtual bool ExportBinary(UObject* Object, const TCHAR* Type, FArchive& Ar, FFeedbackContext* Warn, int32 FileIndex = 0, uint32 PortFlags = 0) override;
+	//~ End UExporter Interface
+};
+
+
+
diff --git a/Engine/Source/Editor/UnrealEd/Classes/Exporters/TextureExporterBMP.h b/Engine/Source/Editor/UnrealEd/Classes/Exporters/TextureExporterBMP.h
--- a/Engine/Source/Editor/UnrealEd/Classes/Exporters/TextureExporterBMP.h
+++ b/Engine/Source/Editor/UnrealEd/Classes/Exporters/TextureExporterBMP.h
@@ -12,8 +12,8 @@ class UTextureExporterBMP : public UExporter
 {
 	GENERATED_UCLASS_BODY()
 
-
-	bool ExportBinary( UObject* Object, const TCHAR* Type, FArchive& Ar, FFeedbackContext* Warn, int32 FileIndex = 0, uint32 PortFlags=0 ) override;
+	virtual bool SupportsObject(UObject* Object) const override;
+	virtual bool ExportBinary( UObject* Object, const TCHAR* Type, FArchive& Ar, FFeedbackContext* Warn, int32 FileIndex = 0, uint32 PortFlags=0 ) override;
 };
 
 
diff --git a/Engine/Source/Editor/UnrealEd/Classes/Exporters/TextureExporterHDR.h b/Engine/Source/Editor/UnrealEd/Classes/Exporters/TextureExporterHDR.h
--- a/Engine/Source/Editor/UnrealEd/Classes/Exporters/TextureExporterHDR.h
+++ b/Engine/Source/Editor/UnrealEd/Classes/Exporters/TextureExporterHDR.h
@@ -13,6 +13,7 @@ class UNREALED_API UTextureExporterHDR : public UExporter
 	GENERATED_UCLASS_BODY()
 
 	//~ Begin UExporter Interface
+	virtual bool SupportsObject(UObject* Object) const override;
 	virtual bool ExportBinary(UObject* Object, const TCHAR* Type, FArchive& Ar, FFeedbackContext* Warn, int32 FileIndex = 0, uint32 PortFlags = 0) override;
 	//~ End UExporter Interface
 };
diff --git a/Engine/Source/Editor/UnrealEd/Classes/Exporters/TextureExporterPCX.h b/Engine/Source/Editor/UnrealEd/Classes/Exporters/TextureExporterPCX.h
--- a/Engine/Source/Editor/UnrealEd/Classes/Exporters/TextureExporterPCX.h
+++ b/Engine/Source/Editor/UnrealEd/Classes/Exporters/TextureExporterPCX.h
@@ -14,6 +14,7 @@ class UTextureExporterPCX : public UExporter
 
 
 	//~ Begin UExporter Interface
+	virtual bool SupportsObject(UObject* Object) const override;
 	virtual bool ExportBinary( UObject* Object, const TCHAR* Type, FArchive& Ar, FFeedbackContext* Warn, int32 FileIndex = 0, uint32 PortFlags=0 ) override;
 	//~ End UExporter Interface
 };
diff --git a/Engine/Source/Editor/UnrealEd/Classes/Factories/FbxImportUI.h b/Engine/Source/Editor/UnrealEd/Classes/Factories/FbxImportUI.h
--- a/Engine/Source/Editor/UnrealEd/Classes/Factories/FbxImportUI.h
+++ b/Engine/Source/Editor/UnrealEd/Classes/Factories/FbxImportUI.h
@@ -17,6 +17,9 @@ enum EFBXImportType
 	FBXIT_SkeletalMesh UMETA(DisplayName="Skeletal Mesh"),
 	/** Select Animation if you'd like to import only animation. */
 	FBXIT_Animation UMETA(DisplayName="Animation"),
+	/** Subdivision Surface (Experimental, Early work in progress) */
+	FBXIT_SubDSurface UMETA(DisplayName="SubDSurface"),
+
 	FBXIT_MAX,
 };
 
@@ -46,8 +49,12 @@ class UFbxImportUI : public UObject
 	uint32 bConvertScene:1;
 
 	/** Whether to import the incoming FBX as a skeletal object */
-	UPROPERTY(EditAnywhere, Category = Mesh, meta = (ImportType = "StaticMesh|SkeletalMesh"))
+	UPROPERTY(EditAnywhere, Category = Mesh, meta = (ImportType = "StaticMesh|SkeletalMesh", DisplayName="Skeletal Mesh"))
 	bool bImportAsSkeletal;
+	
+	/** Whether to import the incoming FBX as a Subdivision Surface (could be made a combo box together with bImportAsSkeletal) (Experimental, Early work in progress) */
+	UPROPERTY(EditAnywhere, Category = Mesh, meta = (ImportType = "StaticMesh|SkeletalMesh", DisplayName="Subdivison Surface"))
+	bool bImportAsSubDSurface;
 
 	/** Whether to import the mesh. Allows animation only import when importing a skeletal mesh. */
 	UPROPERTY(EditAnywhere, Category=Mesh, meta=(ImportType="SkeletalMesh"))
@@ -108,6 +115,17 @@ class UFbxImportUI : public UObject
 	//~ Begin UObject Interface
 	virtual bool CanEditChange( const UProperty* InProperty ) const override;
 	//~ End UObject Interface
+
+	/** sets MeshTypeToImport */
+	void SetMeshTypeToImport()
+	{
+		MeshTypeToImport = bImportAsSkeletal ? FBXIT_SkeletalMesh : FBXIT_StaticMesh;
+
+		if(bImportAsSubDSurface)
+		{
+			MeshTypeToImport = FBXIT_SubDSurface;
+		}
+	}
 };
 
 
diff --git a/Engine/Source/Editor/UnrealEd/Classes/Factories/SubDSurfaceFactory.h b/Engine/Source/Editor/UnrealEd/Classes/Factories/SubDSurfaceFactory.h
new file mode 100644
--- /dev/null
+++ b/Engine/Source/Editor/UnrealEd/Classes/Factories/SubDSurfaceFactory.h
@@ -0,0 +1,18 @@
+// Copyright 1998-2014 Epic Games, Inc. All Rights Reserved.
+
+#pragma once
+#include "SubDSurfaceFactory.generated.h"
+
+UCLASS(hidecategories=Object)
+class USubDSurfaceFactory : public UFactory
+{
+	GENERATED_UCLASS_BODY()
+
+
+	// Begin UFactory Interface
+	virtual UObject* FactoryCreateNew(UClass* Class,UObject* InParent,FName Name,EObjectFlags Flags,UObject* Context,FFeedbackContext* Warn) override;
+	// Begin UFactory Interface	
+};
+
+
+
diff --git a/Engine/Source/Editor/UnrealEd/Private/Factories/ActorFactory.cpp b/Engine/Source/Editor/UnrealEd/Private/Factories/ActorFactory.cpp
--- a/Engine/Source/Editor/UnrealEd/Private/Factories/ActorFactory.cpp
+++ b/Engine/Source/Editor/UnrealEd/Private/Factories/ActorFactory.cpp
@@ -34,6 +34,8 @@ ActorFactory.cpp:
 #include "Engine/TriggerSphere.h"
 #include "Engine/TriggerCapsule.h"
 #include "Engine/TextRenderActor.h"
+#include "Engine/SubDSurfaceActor.h"
+#include "Components/SubDSurfaceComponent.h"
 
 #include "Engine/DestructibleMesh.h"
 #include "Engine/BlueprintGeneratedClass.h"
@@ -527,6 +529,47 @@ UActorFactoryTextRender::UActorFactoryTextRender(const FObjectInitializer& Objec
 	bUseSurfaceOrientation = true;
 }
 
+/*-----------------------------------------------------------------------------
+UActorFactorySubDSurface
+-----------------------------------------------------------------------------*/
+UActorFactorySubDSurface::UActorFactorySubDSurface(const FObjectInitializer& ObjectInitializer)
+	: Super(ObjectInitializer)
+{
+	// Property initialization
+	DisplayName = LOCTEXT("SubDSurfaceDisplayName", "Subdivision Surface");
+	NewActorClass = ASubDSurfaceActor::GetPrivateStaticClass(L"...");
+	bUseSurfaceOrientation = true;
+}
+
+bool UActorFactorySubDSurface::CanCreateActorFrom( const FAssetData& AssetData, FText& OutErrorMsg )
+{
+	// We allow creating ASubDSurfaceActor without an existing asset
+	if ( UActorFactory::CanCreateActorFrom( AssetData, OutErrorMsg ) )
+	{
+		return true;
+	}
+
+	if ( !AssetData.IsValid() || !AssetData.GetClass()->IsChildOf( USubDSurface::StaticClass() ) )
+	{
+		OutErrorMsg = NSLOCTEXT("CanCreateActor", "NoSubDSurface", "A valid SubDSurface must be specified.");
+		return false;
+	}
+
+	return true;
+}
+
+void UActorFactorySubDSurface::PostSpawnActor( UObject* Asset, AActor* NewActor )
+{
+	Super::PostSpawnActor(Asset, NewActor);
+
+	USubDSurface* SubDSurface = Cast<USubDSurface>(Asset);
+	ASubDSurfaceActor* SubDSurfaceActor = CastChecked<ASubDSurfaceActor>(NewActor);
+
+	if(SubDSurface && SubDSurfaceActor)
+	{
+		SubDSurfaceActor->SubDSurface->SetMesh(SubDSurface);
+	}
+}
 
 /*-----------------------------------------------------------------------------
 UActorFactoryEmitter
diff --git a/Engine/Source/Editor/UnrealEd/Private/Factories/EditorFactories.cpp b/Engine/Source/Editor/UnrealEd/Private/Factories/EditorFactories.cpp
--- a/Engine/Source/Editor/UnrealEd/Private/Factories/EditorFactories.cpp
+++ b/Engine/Source/Editor/UnrealEd/Private/Factories/EditorFactories.cpp
@@ -100,6 +100,7 @@
 #include "GameFramework/ForceFeedbackEffect.h"
 #include "GameFramework/HapticFeedbackEffect.h"
 #include "Engine/SubsurfaceProfile.h"
+#include "Engine/SubDSurface.h"
 #include "Camera/CameraAnim.h"
 #include "GameFramework/TouchInterface.h"
 #include "Engine/DataTable.h"
@@ -3322,13 +3323,13 @@ UTextureFactory::UTextureFactory(const FObjectInitializer& ObjectInitializer)
 	Formats.Add( TEXT( "tga;Texture" ) );
 	Formats.Add( TEXT( "float;Texture" ) );
 	Formats.Add( TEXT( "psd;Texture" ) );
-	Formats.Add( TEXT( "dds;Texture" ) );
+	Formats.Add( TEXT( "dds;Texture (Cubemap or 2D)" ) );
 	Formats.Add( TEXT( "hdr;Cubemap Texture (LongLat unwrap)" ) );
 	Formats.Add( TEXT( "ies;IES Texture (Standard light profiles)" ) );
 	Formats.Add( TEXT( "png;Texture" ) );
 	Formats.Add( TEXT( "jpg;Texture" ) );
 	Formats.Add( TEXT( "jpeg;Texture" ) );
-	Formats.Add( TEXT( "exr;Texture" ) );
+	Formats.Add( TEXT( "exr;Texture (HDR)" ) );
 
 	bCreateNew = false;
 	bEditorImport = true;
@@ -4601,6 +4602,21 @@ UTextureExporterPCX::UTextureExporterPCX(const FObjectInitializer& ObjectInitial
 
 }
 
+bool UTextureExporterPCX::SupportsObject(UObject* Object) const
+{
+	bool bSupportsObject = false;
+	if (Super::SupportsObject(Object))
+	{
+		UTexture2D* Texture = Cast<UTexture2D>(Object);
+
+		if (Texture)
+		{
+			bSupportsObject = Texture->Source.GetFormat() == TSF_BGRA8;
+		}
+	}
+	return bSupportsObject;
+}
+
 bool UTextureExporterPCX::ExportBinary( UObject* Object, const TCHAR* Type, FArchive& Ar, FFeedbackContext* Warn, int32 FileIndex, uint32 PortFlags )
 {
 	UTexture2D* Texture = CastChecked<UTexture2D>( Object );
@@ -4669,6 +4685,21 @@ UTextureExporterBMP::UTextureExporterBMP(const FObjectInitializer& ObjectInitial
 
 }
 
+bool UTextureExporterBMP::SupportsObject(UObject* Object) const
+{
+	bool bSupportsObject = false;
+	if (Super::SupportsObject(Object))
+	{
+		UTexture2D* Texture = Cast<UTexture2D>(Object);
+
+		if (Texture)
+		{
+			bSupportsObject = Texture->Source.GetFormat() == TSF_BGRA8 || Texture->Source.GetFormat() == TSF_RGBA16;
+		}
+	}
+	return bSupportsObject;
+}
+
 bool UTextureExporterBMP::ExportBinary( UObject* Object, const TCHAR* Type, FArchive& Ar, FFeedbackContext* Warn, int32 FileIndex, uint32 PortFlags )
 {
 	UTexture2D* Texture = CastChecked<UTexture2D>( Object );
@@ -4957,6 +4988,42 @@ public:
 		return false;
 	}
 
+	/**
+	* Writes HDR format image to an FArchive
+	* @param TexRT - A 2D source render target to read from.
+	* @param Ar - Archive object to write HDR data to.
+	* @return true on successful export.
+	*/
+	bool ExportHDR(UTexture2D* Texture, FArchive& Ar)
+	{
+		check(Texture != nullptr);
+
+		Size = FIntPoint(Texture->Source.GetSizeX(), Texture->Source.GetSizeY());
+
+		TArray<uint8> RawData;
+		bool bReadSuccess = Texture->Source.GetMipData(RawData, 0);
+
+		if (Texture->Source.GetFormat() == TSF_BGRA8)
+		{
+			Format = PF_B8G8R8A8;
+		}
+		else if (Texture->Source.GetFormat() == TSF_RGBA16F)
+		{
+			Format = PF_FloatRGBA;
+		}
+		else
+		{
+			bReadSuccess = false;
+		}
+
+		if (bReadSuccess)
+		{
+			WriteHDRImage(RawData, Ar);
+			return true;
+		}
+		return false;
+	}
+
 	/**
 	* Writes HDR format image to an FArchive
 	* This function unwraps the cube image on to a 2D surface.
@@ -4985,10 +5052,10 @@ public:
 };
 
 /*------------------------------------------------------------------------------
-	UTextureExporterHDR implementation.
+	URenderTargetExporterHDR implementation.
 	Exports render targets.
 ------------------------------------------------------------------------------*/
-UTextureExporterHDR::UTextureExporterHDR(const FObjectInitializer& ObjectInitializer)
+URenderTargetExporterHDR::URenderTargetExporterHDR(const FObjectInitializer& ObjectInitializer)
 	: Super(ObjectInitializer)
 {
 	SupportedClass = UTextureRenderTarget::StaticClass();
@@ -4997,7 +5064,7 @@ UTextureExporterHDR::UTextureExporterHDR(const FObjectInitializer& ObjectInitial
 	FormatDescription.Add(TEXT("HDR"));
 }
 
-bool UTextureExporterHDR::ExportBinary(UObject* Object, const TCHAR* Type, FArchive& Ar, FFeedbackContext* Warn, int32 FileIndex, uint32 PortFlags)
+bool URenderTargetExporterHDR::ExportBinary(UObject* Object, const TCHAR* Type, FArchive& Ar, FFeedbackContext* Warn, int32 FileIndex, uint32 PortFlags)
 {
 	UTextureRenderTarget2D* TexRT2D = Cast<UTextureRenderTarget2D>(Object);
 	UTextureRenderTargetCube* TexRTCube = Cast<UTextureRenderTargetCube>(Object);
@@ -5039,6 +5106,46 @@ bool UTextureCubeExporterHDR::ExportBinary(UObject* Object, const TCHAR* Type, F
 	return false;
 }
 
+/*------------------------------------------------------------------------------
+	UTextureExporterHDR implementation.
+	Export UTexture2D as .HDR
+------------------------------------------------------------------------------*/
+UTextureExporterHDR::UTextureExporterHDR(const FObjectInitializer& ObjectInitializer)
+	: Super(ObjectInitializer)
+{
+	SupportedClass = UTexture2D::StaticClass();
+	PreferredFormatIndex = 0;
+	FormatExtension.Add(TEXT("HDR"));
+	FormatDescription.Add(TEXT("HDR"));
+}
+
+bool UTextureExporterHDR::SupportsObject(UObject* Object) const
+{
+	bool bSupportsObject = false;
+	if (Super::SupportsObject(Object))
+	{
+		UTexture2D* Texture = Cast<UTexture2D>(Object);
+
+		if (Texture)
+		{
+			bSupportsObject = Texture->Source.GetFormat() == TSF_BGRA8 || Texture->Source.GetFormat() == TSF_RGBA16F;
+		}
+	}
+	return bSupportsObject;
+}
+
+bool UTextureExporterHDR::ExportBinary(UObject* Object, const TCHAR* Type, FArchive& Ar, FFeedbackContext* Warn, int32 FileIndex, uint32 PortFlags)
+{
+	UTexture2D* Texture = Cast<UTexture2D>(Object);
+
+	FHDRExportHelper Exporter;
+	if (Texture != nullptr)
+	{
+		return Exporter.ExportHDR(Texture, Ar);
+	}
+	return false;
+}
+
 /*------------------------------------------------------------------------------
 	UTextureExporterTGA implementation.
 ------------------------------------------------------------------------------*/
@@ -5052,6 +5159,21 @@ UTextureExporterTGA::UTextureExporterTGA(const FObjectInitializer& ObjectInitial
 	FormatDescription.Add(TEXT("Targa"));
 }
 
+bool UTextureExporterTGA::SupportsObject(UObject* Object) const
+{
+	bool bSupportsObject = false;
+	if (Super::SupportsObject(Object))
+	{
+		UTexture2D* Texture = Cast<UTexture2D>(Object);
+
+		if (Texture)
+		{
+			bSupportsObject = Texture->Source.GetFormat() == TSF_BGRA8 || Texture->Source.GetFormat() == TSF_RGBA16;
+		}
+	}
+	return bSupportsObject;
+}
+
 bool UTextureExporterTGA::ExportBinary( UObject* Object, const TCHAR* Type, FArchive& Ar, FFeedbackContext* Warn, int32 FileIndex, uint32 PortFlags )
 {
 	UTexture2D* Texture = CastChecked<UTexture2D>( Object );
@@ -7277,6 +7399,25 @@ UObject* USubsurfaceProfileFactory::FactoryCreateNew(UClass* InClass, UObject* I
 	return NewObject<USubsurfaceProfile>(InParent, InName, Flags);
 }
 
+
+/*-----------------------------------------------------------------------------
+USubDSurfaceFactory implementation.
+-----------------------------------------------------------------------------*/
+USubDSurfaceFactory::USubDSurfaceFactory(const FObjectInitializer& ObjectInitializer)
+	: Super(ObjectInitializer)
+{
+
+	SupportedClass = USubDSurface::StaticClass();
+	bCreateNew = true;
+	bEditorImport = false;
+	bEditAfterNew = true;
+}
+
+UObject* USubDSurfaceFactory::FactoryCreateNew(UClass* InClass, UObject* InParent, FName InName, EObjectFlags Flags, UObject* Context, FFeedbackContext* Warn)
+{
+	return NewObject<USubDSurface>(InParent, USubDSurface::StaticClass(), InName, Flags);
+}
+
 /*-----------------------------------------------------------------------------
 	UTouchInterfaceFactory implementation.
 -----------------------------------------------------------------------------*/
diff --git a/Engine/Source/Editor/UnrealEd/Private/Fbx/FbxFactory.cpp b/Engine/Source/Editor/UnrealEd/Private/Fbx/FbxFactory.cpp
--- a/Engine/Source/Editor/UnrealEd/Private/Fbx/FbxFactory.cpp
+++ b/Engine/Source/Editor/UnrealEd/Private/Fbx/FbxFactory.cpp
@@ -10,6 +10,7 @@
 #include "FbxErrors.h"
 #include "AssetRegistryModule.h"
 #include "Engine/StaticMesh.h"
+#include "ObjectTools.h"
 
 #define LOCTEXT_NAMESPACE "FBXFactory"
 
@@ -41,7 +42,7 @@ void UFbxFactory::PostInitProperties()
 
 bool UFbxFactory::DoesSupportClass(UClass * Class)
 {
-	return (Class == UStaticMesh::StaticClass() || Class == USkeletalMesh::StaticClass() || Class == UAnimSequence::StaticClass());
+	return (Class == UStaticMesh::StaticClass() || Class == USkeletalMesh::StaticClass() || Class == UAnimSequence::StaticClass() || Class == USubDSurface::StaticClass());
 }
 
 UClass* UFbxFactory::ResolveSupportedClass()
@@ -56,6 +57,10 @@ UClass* UFbxFactory::ResolveSupportedClass()
 	{
 		ImportClass = UAnimSequence::StaticClass();
 	}
+	else if (ImportUI->MeshTypeToImport == FBXIT_SubDSurface)
+	{
+		ImportClass = USubDSurface::StaticClass();
+	}
 	else
 	{
 		ImportClass = UStaticMesh::StaticClass();
@@ -273,9 +278,9 @@ UObject* UFbxFactory::FactoryCreateBinary
 				int32 NodeIndex = 0;
 
 				int32 ImportedMeshCount = 0;
-				UStaticMesh* NewStaticMesh = NULL;
 				if ( ImportUI->MeshTypeToImport == FBXIT_StaticMesh )  // static mesh
 				{
+					UStaticMesh* NewStaticMesh = NULL;
 					if (bCombineMeshes)
 					{
 						TArray<FbxNode*> FbxMeshArray;
@@ -456,7 +461,31 @@ UObject* UFbxFactory::FactoryCreateBinary
 							FFbxErrors::SkeletalMesh_NoMeshFoundOnRoot);
 					}
 				}
-				else if ( ImportUI->MeshTypeToImport == FBXIT_Animation )// animation
+				else if ( ImportUI->MeshTypeToImport == FBXIT_SubDSurface ) // SubDSurface
+				{
+					USubDSurface* NewMesh = NULL;
+					TArray<FbxNode*> FbxMeshArray;
+					FbxImporter->FillFbxMeshArray(RootNodeToImport, FbxMeshArray, FbxImporter);
+					if (FbxMeshArray.Num() > 0)
+					{
+						NewMesh = Cast<USubDSurface>( CreateOrOverwriteAsset(USubDSurface::StaticClass(), InParent, Name, Flags) );
+						
+						bool bOk = FbxImporter->ImportSubDSurface(NewMesh, InParent, FbxMeshArray, Name, Flags, ImportUI->StaticMeshImportData);
+
+						if(!bOk)
+						{
+							FbxImporter->AddTokenizedErrorMessage(FTokenizedMessage::Create(EMessageSeverity::Error, LOCTEXT("FailedToImport_ImportSubDSurface", "Could not import subdivision surface mesh (no quad mesh?).")), 
+								FFbxErrors::SkeletalMesh_ImportSubDSurface);
+
+							ObjectTools::DeleteSingleObject(NewMesh);
+						}
+					}
+
+					ImportedMeshCount = NewMesh ? 1 : 0;
+
+					NewObject = NewMesh;
+				}
+					else if ( ImportUI->MeshTypeToImport == FBXIT_Animation )// animation
 				{
 					if (ImportOptions->SkeletonForAnimation)
 					{
diff --git a/Engine/Source/Editor/UnrealEd/Private/Fbx/FbxMainImport.cpp b/Engine/Source/Editor/UnrealEd/Private/Fbx/FbxMainImport.cpp
--- a/Engine/Source/Editor/UnrealEd/Private/Fbx/FbxMainImport.cpp
+++ b/Engine/Source/Editor/UnrealEd/Private/Fbx/FbxMainImport.cpp
@@ -63,6 +63,7 @@ FBXImportOptions* GetImportOptions( UnFbx::FFbxImporter* FbxImporter, UFbxImport
 		}
 
 		ImportUI->bImportAsSkeletal = ImportUI->MeshTypeToImport == FBXIT_SkeletalMesh;
+		ImportUI->bImportAsSubDSurface = ImportUI->MeshTypeToImport == FBXIT_SubDSurface;
 		ImportUI->bImportMesh = ImportUI->MeshTypeToImport != FBXIT_Animation;
 		ImportUI->bIsObjImport = bIsObjFormat;
 
diff --git a/Engine/Source/Editor/UnrealEd/Private/Fbx/FbxStaticMeshImport.cpp b/Engine/Source/Editor/UnrealEd/Private/Fbx/FbxStaticMeshImport.cpp
--- a/Engine/Source/Editor/UnrealEd/Private/Fbx/FbxStaticMeshImport.cpp
+++ b/Engine/Source/Editor/UnrealEd/Private/Fbx/FbxStaticMeshImport.cpp
@@ -68,98 +68,197 @@ UStaticMesh* UnFbx::FFbxImporter::ImportStaticMesh(UObject* InParent, FbxNode* N
 	return ImportStaticMeshAsSingle(InParent, MeshNodeArray, Name, Flags, ImportData, InStaticMesh, LODIndex);
 }
 
-bool UnFbx::FFbxImporter::BuildStaticMeshFromGeometry(FbxNode* Node, UStaticMesh* StaticMesh, TArray<FFbxMaterial>& MeshMaterials, int LODIndex, FRawMesh& RawMesh,
-													  EVertexColorImportOption::Type VertexColorImportOption, const TMap<FVector, FColor>& ExistingVertexColorData, const FColor& VertexOverrideColor)
+// Wraps some common code useful for multiple fbx import code path
+struct FFBXUVs
 {
-	check(StaticMesh->SourceModels.IsValidIndex(LODIndex));
-	FbxMesh* Mesh = Node->GetMesh();
-	FStaticMeshSourceModel& SrcModel = StaticMesh->SourceModels[LODIndex];
-
-	//remove the bad polygons before getting any data from mesh
-	Mesh->RemoveBadPolygons();
-
-	//Get the base layer of the mesh
-	FbxLayer* BaseLayer = Mesh->GetLayer(0);
-	if (BaseLayer == NULL)
+	// constructor
+	FFBXUVs(FbxMesh* Mesh)
+		: UniqueUVCount(0)
 	{
-		AddTokenizedErrorMessage(FTokenizedMessage::Create(EMessageSeverity::Error, FText::Format(LOCTEXT("Error_NoGeometryInMesh", "There is no geometry information in mesh '{0}'"), FText::FromString(Mesh->GetName()))), FFbxErrors::Generic_Mesh_NoGeometry);
-		return false;
-	}
+		check(Mesh);
 
-	//
-	//	store the UVs in arrays for fast access in the later looping of triangles 
-	//
-	// mapping from UVSets to Fbx LayerElementUV
-	// Fbx UVSets may be duplicated, remove the duplicated UVSets in the mapping 
-	int32 LayerCount = Mesh->GetLayerCount();
-	TArray<FString> UVSets;
-	UVSets.Empty();
-	if (LayerCount > 0)
-	{
-		int32 UVLayerIndex;
-		for (UVLayerIndex = 0; UVLayerIndex<LayerCount; UVLayerIndex++)
+		//
+		//	store the UVs in arrays for fast access in the later looping of triangles 
+		//
+		// mapping from UVSets to Fbx LayerElementUV
+		// Fbx UVSets may be duplicated, remove the duplicated UVSets in the mapping 
+		int32 LayerCount = Mesh->GetLayerCount();
+		if (LayerCount > 0)
 		{
-			FbxLayer* lLayer = Mesh->GetLayer(UVLayerIndex);
-			int UVSetCount = lLayer->GetUVSetCount();
-			if(UVSetCount)
+			int32 UVLayerIndex;
+			for (UVLayerIndex = 0; UVLayerIndex<LayerCount; UVLayerIndex++)
 			{
-				FbxArray<FbxLayerElementUV const*> EleUVs = lLayer->GetUVSets();
-				for (int UVIndex = 0; UVIndex<UVSetCount; UVIndex++)
+				FbxLayer* lLayer = Mesh->GetLayer(UVLayerIndex);
+				int UVSetCount = lLayer->GetUVSetCount();
+				if(UVSetCount)
 				{
-					FbxLayerElementUV const* ElementUV = EleUVs[UVIndex];
-					if (ElementUV)
+					FbxArray<FbxLayerElementUV const*> EleUVs = lLayer->GetUVSets();
+					for (int UVIndex = 0; UVIndex<UVSetCount; UVIndex++)
 					{
-						const char* UVSetName = ElementUV->GetName();
-						FString LocalUVSetName = UTF8_TO_TCHAR(UVSetName);
-						if (LocalUVSetName.IsEmpty())
+						FbxLayerElementUV const* ElementUV = EleUVs[UVIndex];
+						if (ElementUV)
 						{
-							LocalUVSetName = TEXT("UVmap_") + FString::FromInt(UVLayerIndex);
+							const char* UVSetName = ElementUV->GetName();
+							FString LocalUVSetName = UTF8_TO_TCHAR(UVSetName);
+							if (LocalUVSetName.IsEmpty())
+							{
+								LocalUVSetName = TEXT("UVmap_") + FString::FromInt(UVLayerIndex);
+							}
+
+							UVSets.AddUnique(LocalUVSetName);
 						}
+					}
+				}
+			}
+		}
 
-						UVSets.AddUnique(LocalUVSetName);
+
+		// If the the UV sets are named using the following format (UVChannel_X; where X ranges from 1 to 4)
+		// we will re-order them based on these names.  Any UV sets that do not follow this naming convention
+		// will be slotted into available spaces.
+		if(UVSets.Num())
+		{
+			for(int32 ChannelNumIdx = 0; ChannelNumIdx < 4; ChannelNumIdx++)
+			{
+				FString ChannelName = FString::Printf( TEXT("UVChannel_%d"), ChannelNumIdx+1 );
+				int32 SetIdx = UVSets.Find( ChannelName );
+
+				// If the specially formatted UVSet name appears in the list and it is in the wrong spot,
+				// we will swap it into the correct spot.
+				if( SetIdx != INDEX_NONE && SetIdx != ChannelNumIdx )
+				{
+					// If we are going to swap to a position that is outside the bounds of the
+					// array, then we pad out to that spot with empty data.
+					for(int32 ArrSize = UVSets.Num(); ArrSize < ChannelNumIdx+1; ArrSize++)
+					{
+						UVSets.Add ( FString(TEXT("")) );
 					}
+					//Swap the entry into the appropriate spot.
+					UVSets.Swap( SetIdx, ChannelNumIdx );
 				}
 			}
 		}
 	}
 
-
-	// If the the UV sets are named using the following format (UVChannel_X; where X ranges from 1 to 4)
-	// we will re-order them based on these names.  Any UV sets that do not follow this naming convention
-	// will be slotted into available spaces.
-	if( UVSets.Num() > 0 )
+	void Phase2(FbxMesh* Mesh)
 	{
-		for(int32 ChannelNumIdx = 0; ChannelNumIdx < 4; ChannelNumIdx++)
+		//
+		//	store the UVs in arrays for fast access in the later looping of triangles 
+		//
+		UniqueUVCount = UVSets.Num();
+		if (UniqueUVCount > 0)
 		{
-			FString ChannelName = FString::Printf( TEXT("UVChannel_%d"), ChannelNumIdx+1 );
-			int32 SetIdx = UVSets.Find( ChannelName );
-
-			// If the specially formatted UVSet name appears in the list and it is in the wrong spot,
-			// we will swap it into the correct spot.
-			if( SetIdx != INDEX_NONE && SetIdx != ChannelNumIdx )
+			LayerElementUV.AddZeroed(UniqueUVCount);
+			UVReferenceMode.AddZeroed(UniqueUVCount);
+			UVMappingMode.AddZeroed(UniqueUVCount);
+		}
+		for (int32 UVIndex = 0; UVIndex < UniqueUVCount; UVIndex++)
+		{
+			LayerElementUV[UVIndex] = NULL;
+			for (int32 UVLayerIndex = 0, LayerCount = Mesh->GetLayerCount(); UVLayerIndex < LayerCount; UVLayerIndex++)
 			{
-				// If we are going to swap to a position that is outside the bounds of the
-				// array, then we pad out to that spot with empty data.
-				for(int32 ArrSize = UVSets.Num(); ArrSize < ChannelNumIdx+1; ArrSize++)
+				FbxLayer* lLayer = Mesh->GetLayer(UVLayerIndex);
+				int UVSetCount = lLayer->GetUVSetCount();
+				if(UVSetCount)
 				{
-					UVSets.Add ( FString(TEXT("")) );
+					FbxArray<FbxLayerElementUV const*> EleUVs = lLayer->GetUVSets();
+					for (int32 FbxUVIndex = 0; FbxUVIndex<UVSetCount; FbxUVIndex++)
+					{
+						FbxLayerElementUV const* ElementUV = EleUVs[FbxUVIndex];
+						if (ElementUV)
+						{
+							const char* UVSetName = ElementUV->GetName();
+							FString LocalUVSetName = UTF8_TO_TCHAR(UVSetName);
+							if (LocalUVSetName == UVSets[UVIndex])
+							{
+								LayerElementUV[UVIndex] = ElementUV;
+								UVReferenceMode[UVIndex] = ElementUV->GetReferenceMode();
+								UVMappingMode[UVIndex] = ElementUV->GetMappingMode();
+								break;
+							}
+						}
+					}
 				}
-				//Swap the entry into the appropriate spot.
-				UVSets.Swap( SetIdx, ChannelNumIdx );
 			}
 		}
+		UniqueUVCount = FMath::Min<int32>(UniqueUVCount, MAX_MESH_TEXTURE_COORDS);
 	}
 
+	int32 FindLightUVIndex() const
+	{
+		// See if any of our UV set entry names match LightMapUV.
+		for(int32 UVSetIdx = 0; UVSetIdx < UVSets.Num(); UVSetIdx++)
+		{
+			if( UVSets[UVSetIdx] == TEXT("LightMapUV"))
+			{
+				return UVSetIdx;
+			}
+		}
 
-	// See if any of our UV set entry names match LightMapUV.
-	for(int UVSetIdx = 0; UVSetIdx < UVSets.Num(); UVSetIdx++)
+		// not found
+		return 0;
+	}
+	
+	// @param FaceCornerIndex usually TriangleIndex * 3 + CornerIndex but more complicated for mixed n-gons
+	int32 ComputeUVIndex(int32 UVLayerIndex, int32 lControlPointIndex, int32 FaceCornerIndex) const
 	{
-		if( UVSets[UVSetIdx] == TEXT("LightMapUV"))
+		int32 UVMapIndex = (UVMappingMode[UVLayerIndex] == FbxLayerElement::eByControlPoint) ? lControlPointIndex : FaceCornerIndex;
+						
+		int32 Ret;
+
+		if(UVReferenceMode[UVLayerIndex] == FbxLayerElement::eDirect)
+		{
+			Ret = UVMapIndex;
+		}
+		else
 		{
-			StaticMesh->LightMapCoordinateIndex = UVSetIdx;
+			FbxLayerElementArrayTemplate<int>& Array = LayerElementUV[UVLayerIndex]->GetIndexArray();
+			Ret = Array.GetAt(UVMapIndex);
 		}
+
+		return Ret;
+	}
+
+	// todo: is that needed? could the dtor do it?
+	void Cleanup()
+	{
+		//
+		// clean up.  This needs to happen before the mesh is destroyed
+		//
+		LayerElementUV.Empty();
+		UVReferenceMode.Empty();
+		UVMappingMode.Empty();
+	}
+
+	TArray<FString> UVSets;
+	TArray<FbxLayerElementUV const*> LayerElementUV;
+	TArray<FbxLayerElement::EReferenceMode> UVReferenceMode;
+	TArray<FbxLayerElement::EMappingMode> UVMappingMode;
+	int32 UniqueUVCount;
+};
+
+bool UnFbx::FFbxImporter::BuildStaticMeshFromGeometry(FbxNode* Node, UStaticMesh* StaticMesh, TArray<FFbxMaterial>& MeshMaterials, int32 LODIndex,FRawMesh& RawMesh,
+	EVertexColorImportOption::Type VertexColorImportOption, const TMap<FVector, FColor>& ExistingVertexColorData, const FColor& VertexOverrideColor)
+{
+	check(StaticMesh->SourceModels.IsValidIndex(LODIndex));
+	FbxMesh* Mesh = Node->GetMesh();
+	FStaticMeshSourceModel& SrcModel = StaticMesh->SourceModels[LODIndex];
+
+	//remove the bad polygons before getting any data from mesh
+	Mesh->RemoveBadPolygons();
+
+	//Get the base layer of the mesh
+	FbxLayer* BaseLayer = Mesh->GetLayer(0);
+	if (BaseLayer == NULL)
+	{
+		AddTokenizedErrorMessage(FTokenizedMessage::Create(EMessageSeverity::Error, FText::Format(LOCTEXT("Error_NoGeometryInMesh", "There is no geometry information in mesh '{0}'"), FText::FromString(Mesh->GetName()))), FFbxErrors::Generic_Mesh_NoGeometry);
+		return false;
 	}
 
+	FFBXUVs FBXUVs(Mesh);
+
+	StaticMesh->LightMapCoordinateIndex = FBXUVs.FindLightUVIndex();
+	
 	//
 	// create materials
 	//
@@ -167,7 +266,7 @@ bool UnFbx::FFbxImporter::BuildStaticMeshFromGeometry(FbxNode* Node, UStaticMesh
 	TArray<UMaterialInterface*> Materials;
 	if (ImportOptions->bImportMaterials)
 	{
-		CreateNodeMaterials(Node, Materials, UVSets);
+		CreateNodeMaterials(Node, Materials, FBXUVs.UVSets);
 	}
 	else if (ImportOptions->bImportTextures)
 	{
@@ -255,55 +354,8 @@ bool UnFbx::FFbxImporter::BuildStaticMeshFromGeometry(FbxNode* Node, UStaticMesh
 	FbxLayerElement::EMappingMode MaterialMappingMode = LayerElementMaterial ? 
 		LayerElementMaterial->GetMappingMode() : FbxLayerElement::eByPolygon;
 
-	//
-	//	store the UVs in arrays for fast access in the later looping of triangles 
-	//
-	int32 UniqueUVCount = UVSets.Num();
-	TArray<FbxLayerElementUV const*> LayerElementUV;
-	TArray<FbxLayerElement::EReferenceMode> UVReferenceMode;
-	TArray<FbxLayerElement::EMappingMode> UVMappingMode;
-	if (UniqueUVCount > 0)
-	{
-		LayerElementUV.AddZeroed(UniqueUVCount);
-		UVReferenceMode.AddZeroed(UniqueUVCount);
-		UVMappingMode.AddZeroed(UniqueUVCount);
-	}
-	LayerCount = Mesh->GetLayerCount();
-	for (int32 UVIndex = 0; UVIndex < UniqueUVCount; UVIndex++)
-	{
-		LayerElementUV[UVIndex] = NULL;
-		for (int32 UVLayerIndex = 0; UVLayerIndex<LayerCount; UVLayerIndex++)
-		{
-			FbxLayer* lLayer = Mesh->GetLayer(UVLayerIndex);
-			int UVSetCount = lLayer->GetUVSetCount();
-			if(UVSetCount)
-			{
-				FbxArray<FbxLayerElementUV const*> EleUVs = lLayer->GetUVSets();
-				for (int32 FbxUVIndex = 0; FbxUVIndex<UVSetCount; FbxUVIndex++)
-				{
-					FbxLayerElementUV const* ElementUV = EleUVs[FbxUVIndex];
-					if (ElementUV)
-					{
-						const char* UVSetName = ElementUV->GetName();
-						FString LocalUVSetName = UTF8_TO_TCHAR(UVSetName);
-						if (LocalUVSetName.IsEmpty())
-						{
-							LocalUVSetName = TEXT("UVmap_") + FString::FromInt(UVLayerIndex);
-						}
-						if (LocalUVSetName == UVSets[UVIndex])
-						{
-							LayerElementUV[UVIndex] = ElementUV;
-							UVReferenceMode[UVIndex] = ElementUV->GetReferenceMode();
-							UVMappingMode[UVIndex] = ElementUV->GetMappingMode();
-							break;
-						}
-					}
-				}
-			}
-		}
-	}
-	UniqueUVCount = FMath::Min<int32>(UniqueUVCount, MAX_MESH_TEXTURE_COORDS);
-
+	//	todo second phase UV, ok to put in first phase?
+	FBXUVs.Phase2(Mesh);
 
 	//
 	// get the smoothing group layer
@@ -466,8 +518,8 @@ bool UnFbx::FFbxImporter::BuildStaticMeshFromGeometry(FbxNode* Node, UStaticMesh
 			break;
 		}
 	}
-
-	int32 UVCount = FMath::Max( UniqueUVCount, ExistingUVCount );
+	
+	int32 UVCount = FMath::Max( FBXUVs.UniqueUVCount, ExistingUVCount );
 
 	// At least one UV set must exist.  
 	UVCount = FMath::Max( 1, UVCount );
@@ -611,20 +663,22 @@ bool UnFbx::FFbxImporter::BuildStaticMeshFromGeometry(FbxNode* Node, UStaticMesh
 		// In FBX file, the same UV may be saved multiple times, i.e., there may be same UV in LayerElementUV
 		// So we don't import the duplicate UVs
 		int32 UVLayerIndex;
-		for (UVLayerIndex = 0; UVLayerIndex<UniqueUVCount; UVLayerIndex++)
+		for (UVLayerIndex = 0; UVLayerIndex<FBXUVs.UniqueUVCount; UVLayerIndex++)
 		{
-			if (LayerElementUV[UVLayerIndex] != NULL) 
+			if (FBXUVs.LayerElementUV[UVLayerIndex] != NULL) 
 			{
 				for (int32 CornerIndex=0;CornerIndex<3;CornerIndex++)
 				{
 					// If there are odd number negative scale, invert the vertex order for triangles
 					int32 WedgeIndex = WedgeOffset + TriangleIndex * 3 + (OddNegativeScale ? 2 - CornerIndex : CornerIndex);
 
+
 					int lControlPointIndex = Mesh->GetPolygonVertex(TriangleIndex, CornerIndex);
-					int UVMapIndex = (UVMappingMode[UVLayerIndex] == FbxLayerElement::eByControlPoint) ? lControlPointIndex : TriangleIndex*3+CornerIndex;
-					int32 UVIndex = (UVReferenceMode[UVLayerIndex] == FbxLayerElement::eDirect) ? 
-						UVMapIndex : LayerElementUV[UVLayerIndex]->GetIndexArray().GetAt(UVMapIndex);
-					FbxVector2	UVVector = LayerElementUV[UVLayerIndex]->GetDirectArray().GetAt(UVIndex);
+					int UVMapIndex = (FBXUVs.UVMappingMode[UVLayerIndex] == FbxLayerElement::eByControlPoint) ? lControlPointIndex : TriangleIndex*3+CornerIndex;
+					int32 UVIndex = (FBXUVs.UVReferenceMode[UVLayerIndex] == FbxLayerElement::eDirect) ? 
+						UVMapIndex : FBXUVs.LayerElementUV[UVLayerIndex]->GetIndexArray().GetAt(UVMapIndex);
+
+					FbxVector2	UVVector = FBXUVs.LayerElementUV[UVLayerIndex]->GetDirectArray().GetAt(UVIndex);
 
 					RawMesh.WedgeTexCoords[UVLayerIndex][WedgeIndex].X = static_cast<float>(UVVector[0]);
 					RawMesh.WedgeTexCoords[UVLayerIndex][WedgeIndex].Y = 1.f-static_cast<float>(UVVector[1]);   //flip the Y of UVs for DirectX
@@ -666,13 +720,9 @@ bool UnFbx::FFbxImporter::BuildStaticMeshFromGeometry(FbxNode* Node, UStaticMesh
 	
 		RawMesh.FaceMaterialIndices[DestTriangleIndex] = MaterialIndex;
 	}
-
-	//
-	// clean up.  This needs to happen before the mesh is destroyed
-	//
-	LayerElementUV.Empty();
-	UVReferenceMode.Empty();
-	UVMappingMode.Empty();
+	
+	// needed?
+	FBXUVs.Cleanup();
 
 	return true;
 }
@@ -1388,6 +1438,166 @@ UStaticMesh* UnFbx::FFbxImporter::ImportStaticMeshAsSingle(UObject* InParent, TA
 	return StaticMesh;
 }
 
+bool UnFbx::FFbxImporter::ImportSubDSurface(USubDSurface* Out, UObject* InParent, TArray<FbxNode*>& MeshNodeArray, const FName InName, EObjectFlags Flags, UFbxStaticMeshImportData* TemplateImportData)
+{
+	if(!MeshNodeArray.Num())
+	{
+		return false;
+	}
+
+	// Make sure rendering is done - so we are not changing data being used
+	FlushRenderingCommands();
+
+	double StartTime = FPlatformTime::Seconds();
+
+	Parent = InParent;
+
+	// warning for missing smoothing group info
+	CheckSmoothingInfo(MeshNodeArray[0]->GetMesh());
+
+	uint32 ModelCount = (uint32)MeshNodeArray.Num();
+
+	// RemoveBadPolygons(), count Vertex and Indices
+	uint32 TotalIndexCount = 0, TotalVertexCount = 0, TotalPolyCount = 0;
+	{
+		for(uint32 ModelIndex = 0; ModelIndex < ModelCount; ++ModelIndex)
+		{
+			FbxNode* Node = MeshNodeArray[ModelIndex];
+			FbxMesh* FbxMesh = Node->GetMesh();
+
+			int32 Removed = FbxMesh->RemoveBadPolygons();
+
+			TotalVertexCount += FbxMesh->GetControlPointsCount();
+
+			int32 PolyCount = FbxMesh->GetPolygonCount();
+
+			TotalPolyCount += PolyCount;
+
+			for(int32 PolyIndex = 0; PolyIndex < PolyCount; ++PolyIndex)
+			{
+				int32 PolySize = FbxMesh->GetPolygonSize(PolyIndex);
+				
+				TotalIndexCount += PolySize;
+			}
+		}
+	}
+
+	Out->IndicesPerFace.AddUninitialized(TotalIndexCount);
+	Out->VertexCountPerFace.AddUninitialized(TotalPolyCount);
+
+	uint32* IndicesPerFace = Out->IndicesPerFace.GetData();
+	uint32* VertexCountPerFace = Out->VertexCountPerFace.GetData();
+
+	// add position vertex attribute
+	FVector* Positions = 0;
+	{
+		auto Stream = Out->CreateVertexAttributeStream(FName(TEXT("Position")));
+		Positions = Stream->CreateFVectorUninitialized(TotalVertexCount);
+	}
+
+	// todo: support multiple UV
+
+	// add UV vertex attribute
+	FVector2D* UV0 = 0;
+	{
+		auto Stream = Out->CreateVertexAttributeStream(FName(TEXT("UV0")));
+		UV0 = Stream->CreateFVector2DUninitialized(TotalIndexCount);
+	}
+
+	uint32 CurrentIndexStart = 0;
+	for(uint32 ModelIndex = 0; ModelIndex < ModelCount; ++ModelIndex)
+	{
+		FbxNode* Node = MeshNodeArray[ModelIndex];
+		FbxMesh* FbxMesh = Node->GetMesh();
+		uint32 PolyCount = (uint32)FbxMesh->GetPolygonCount();
+
+		FFBXUVs FBXUVs(FbxMesh);
+
+		FBXUVs.Phase2(FbxMesh);
+
+		for(uint32 PolyIndex = 0; PolyIndex < PolyCount; ++PolyIndex)
+		{
+			uint32 CornerCount = FbxMesh->GetPolygonSize(PolyIndex);
+
+			// store VertexCountPerFace
+			*VertexCountPerFace++ = CornerCount;
+
+			for(uint32 Corner = 0; Corner < CornerCount; ++Corner)
+			{
+				uint32 VertexIndex = (uint32)FbxMesh->GetPolygonVertex(PolyIndex, Corner);
+
+				check(VertexIndex < TotalVertexCount);
+
+				// store IndicesPerFace
+				*IndicesPerFace++ = CurrentIndexStart + VertexIndex;
+			}
+		}
+
+		int32 ControlPointsCount = FbxMesh->GetControlPointsCount();
+
+		CurrentIndexStart += ControlPointsCount;
+
+		FbxAMatrix Matrix = ComputeTotalMatrix(Node);
+
+		{
+			
+			FbxVector4* ControlPoints = FbxMesh->GetControlPoints();
+
+			for(int32 ControlPointsIndex = 0; ControlPointsIndex < ControlPointsCount; ControlPointsIndex++ )
+			{
+				FVector Pos = Converter.ConvertPos(Matrix.MultT(ControlPoints[ControlPointsIndex]));
+
+				*Positions++ = Pos;
+			}
+		}
+
+
+
+		//
+		// uvs
+		//
+		// In FBX file, the same UV may be saved multiple times, i.e., there may be same UV in LayerElementUV
+		// So we don't import the duplicate UVs
+		{
+			// We don't support multiple UV yet
+			ensure(FBXUVs.UniqueUVCount == 1);
+
+			int32 UVLayerIndex;
+			for (UVLayerIndex = 0; UVLayerIndex < FBXUVs.UniqueUVCount; UVLayerIndex++)
+			{
+				if (FBXUVs.LayerElementUV[UVLayerIndex] != NULL) 
+				{
+					uint32 UniqueCornerIndex = 0;
+
+					for(uint32 PolyIndex = 0; PolyIndex < PolyCount; ++PolyIndex)
+					{
+						int32 PolySize = FbxMesh->GetPolygonSize(PolyIndex);
+
+						for (int32 CornerIndex = 0; CornerIndex < PolySize; CornerIndex++)
+						{
+							// If there are odd number negative scale, invert the vertex order for triangles
+	//todo						int32 WedgeIndex = WedgeOffset + TriangleIndex * 3 + (OddNegativeScale ? 2 - CornerIndex : CornerIndex);
+		
+							int lControlPointIndex = FbxMesh->GetPolygonVertex(PolyIndex, CornerIndex);
+
+							int UVIndex = FBXUVs.ComputeUVIndex(UVLayerIndex, lControlPointIndex, UniqueCornerIndex);
+	
+							FbxVector2 UVVector = FBXUVs.LayerElementUV[UVLayerIndex]->GetDirectArray().GetAt(UVIndex);
+
+							*UV0++ = FVector2D(static_cast<float>(UVVector[0]), 1.0f - static_cast<float>(UVVector[1]));   //flip the Y of UVs for DirectX
+
+							++UniqueCornerIndex;
+						}
+					}
+				}
+			}
+		}
+	}
+	
+	UE_LOG(LogFbx, Log, TEXT("SubDSurface import: %5.2f seconds"), FPlatformTime::Seconds() - StartTime);
+
+	return true;
+}
 struct FbxSocketNode
 {
 	FName SocketName;
diff --git a/Engine/Source/Editor/UnrealEd/Private/FbxImporter.h b/Engine/Source/Editor/UnrealEd/Private/FbxImporter.h
--- a/Engine/Source/Editor/UnrealEd/Private/FbxImporter.h
+++ b/Engine/Source/Editor/UnrealEd/Private/FbxImporter.h
@@ -9,6 +9,7 @@ class UInterpGroupInst;
 class AMatineeActor;
 class AActor;
 class UInterpTrackMove;
+class USubDSurface;
 
 // Temporarily disable a few warnings due to virtual function abuse in FBX source files
 #pragma warning( push )
@@ -387,6 +388,14 @@ public:
 	*/
 	UNREALED_API UStaticMesh* ImportStaticMeshAsSingle(UObject* InParent, TArray<FbxNode*>& MeshNodeArray, const FName InName, EObjectFlags Flags, UFbxStaticMeshImportData* TemplateImportData, UStaticMesh* InStaticMesh, int LODIndex = 0);
 
+	/**
+	* Creates a SubDSurface mesh from all the meshes in FBX scene with the given name and flags.
+	*
+	* @param MeshNodeArray	Fbx Nodes to import
+	* @param InName	the Unreal Mesh name after import
+	*/
+	UNREALED_API bool ImportSubDSurface(USubDSurface* Out, UObject* InParent, TArray<FbxNode*>& MeshNodeArray, const FName InName, EObjectFlags Flags, UFbxStaticMeshImportData* TemplateImportData);
+
 	void ImportStaticMeshSockets( UStaticMesh* StaticMesh );
 
 	/**
@@ -763,7 +772,6 @@ protected:
 	/**
 	 * Set up the static mesh data from Fbx Mesh.
 	 *
-	 * @param FbxMesh  Fbx Mesh object
 	 * @param StaticMesh Unreal static mesh object to fill data into
 	 * @param LODIndex	LOD level to set up for StaticMesh
 	 * @return bool true if set up successfully
diff --git a/Engine/Source/Editor/UnrealEd/Private/Lightmass/Lightmass.cpp b/Engine/Source/Editor/UnrealEd/Private/Lightmass/Lightmass.cpp
--- a/Engine/Source/Editor/UnrealEd/Private/Lightmass/Lightmass.cpp
+++ b/Engine/Source/Editor/UnrealEd/Private/Lightmass/Lightmass.cpp
@@ -3179,6 +3179,33 @@ static int32 AccumulateVisibility(const TArray<uint8>& OtherCellData, TArray<uin
 	return NumAdded;
 }
 
+struct FPrecomputedVisibilitySortGridCell
+{
+	TArray<FUncompressedPrecomputedVisibilityCell, TInlineAllocator<2>> Cells;
+};
+
+void SpreadVisibilityCell(
+	float CellSize, 
+	float PlayAreaHeight,
+	const FUncompressedPrecomputedVisibilityCell& OtherCell,
+	FUncompressedPrecomputedVisibilityCell& VisibilityCell,
+	int32& QueriesVisibleFromSpreadingNeighbors)
+{
+	// Determine whether the cell is a world space neighbor
+	if (!(OtherCell.Bounds.Min == VisibilityCell.Bounds.Min && OtherCell.Bounds.Max == VisibilityCell.Bounds.Max)
+		&& FMath::Abs(VisibilityCell.Bounds.Min.X - OtherCell.Bounds.Min.X) < CellSize + KINDA_SMALL_NUMBER
+		&& FMath::Abs(VisibilityCell.Bounds.Min.Y - OtherCell.Bounds.Min.Y) < CellSize + KINDA_SMALL_NUMBER
+		// Don't spread from cells below, they're probably below the ground and see too much
+		&& OtherCell.Bounds.Min.Z - VisibilityCell.Bounds.Min.Z > -PlayAreaHeight * 0.5f
+		// Only spread from one cell above
+		&& OtherCell.Bounds.Min.Z - VisibilityCell.Bounds.Min.Z < PlayAreaHeight * 1.5f)
+	{
+		// Combine the neighbor's visibility with the current cell's visibility
+		// This reduces visibility errors at the cost of less effective culling
+		QueriesVisibleFromSpreadingNeighbors += AccumulateVisibility(OtherCell.VisibilityData, VisibilityCell.VisibilityData);
+	}
+}
+
 void FLightmassProcessor::ApplyPrecomputedVisibility()
 {
 	TArray<FUncompressedPrecomputedVisibilityCell> CombinedPrecomputedVisibilityCells;
@@ -3211,29 +3238,91 @@ void FLightmassProcessor::ApplyPrecomputedVisibility()
 
 		int32 TotalNumQueries = 0;
 		int32 QueriesVisibleFromSpreadingNeighbors = 0;
+
 		for (int32 IterationIndex = 0; IterationIndex < VisibilitySpreadingIterations; IterationIndex++)
 		{
-			// Copy the original data since we read from outside the current cell
-			TArray<FUncompressedPrecomputedVisibilityCell> OriginalPrecomputedVisibilityCells(CombinedPrecomputedVisibilityCells);
+			FBox AllCellsBounds(0);
+
 			for (int32 CellIndex = 0; CellIndex < CombinedPrecomputedVisibilityCells.Num(); CellIndex++)
 			{
-				FUncompressedPrecomputedVisibilityCell& CurrentCell = CombinedPrecomputedVisibilityCells[CellIndex];
-				TotalNumQueries += CurrentCell.VisibilityData.Num() * 8;
-				for (int32 OtherCellIndex = 0; OtherCellIndex < OriginalPrecomputedVisibilityCells.Num(); OtherCellIndex++)
+				AllCellsBounds += CombinedPrecomputedVisibilityCells[CellIndex].Bounds;
+			}
+
+			int32 GridSizeX = FMath::TruncToInt(AllCellsBounds.GetSize().X / CellSize + .5f);
+			int32 GridSizeY = FMath::TruncToInt(AllCellsBounds.GetSize().Y / CellSize + .5f);
+
+			const int32 GridSizeMax = 10000;
+
+			if (GridSizeX < GridSizeMax && GridSizeY < GridSizeMax)
+			{
+				TArray<FPrecomputedVisibilitySortGridCell> SortGrid;
+				SortGrid.Empty(GridSizeX * GridSizeY);
+				SortGrid.AddZeroed(GridSizeX * GridSizeY);
+
+				// Add visibility cells into a 2d grid
+				// Note that visibility data is duplicated so that the next pass can read from original neighbor visibility data
+				for (int32 CellIndex = 0; CellIndex < CombinedPrecomputedVisibilityCells.Num(); CellIndex++)
+				{
+					float CellXFloat = (CombinedPrecomputedVisibilityCells[CellIndex].Bounds.GetCenter().X - AllCellsBounds.Min.X) / CellSize;
+					int32 CellX = FMath::Clamp(FMath::TruncToInt(CellXFloat), 0, GridSizeX - 1);
+
+					float CellYFloat = (CombinedPrecomputedVisibilityCells[CellIndex].Bounds.GetCenter().Y - AllCellsBounds.Min.Y) / CellSize;
+					int32 CellY = FMath::Clamp(FMath::TruncToInt(CellYFloat), 0, GridSizeY - 1);
+
+					FPrecomputedVisibilitySortGridCell& GridCell = SortGrid[CellY * GridSizeX + CellX];
+					GridCell.Cells.Add(CombinedPrecomputedVisibilityCells[CellIndex]);
+				}
+
+				// Gather visibility from neighbors, using the 2d grid to accelerate the neighbor search
+				for (int32 CellIndex = 0; CellIndex < CombinedPrecomputedVisibilityCells.Num(); CellIndex++)
 				{
-					const FUncompressedPrecomputedVisibilityCell& OtherCell = OriginalPrecomputedVisibilityCells[OtherCellIndex];
-					// Determine whether the cell is a world space neighbor
-					if (CellIndex != OtherCellIndex
-						&& FMath::Abs(CurrentCell.Bounds.Min.X - OtherCell.Bounds.Min.X) < CellSize + KINDA_SMALL_NUMBER
-						&& FMath::Abs(CurrentCell.Bounds.Min.Y - OtherCell.Bounds.Min.Y) < CellSize + KINDA_SMALL_NUMBER
-						// Don't spread from cells below, they're probably below the ground and see too much
-						&& OtherCell.Bounds.Min.Z - CurrentCell.Bounds.Min.Z > -PlayAreaHeight * 0.5f
-						// Only spread from one cell above
-						&& OtherCell.Bounds.Min.Z - CurrentCell.Bounds.Min.Z < PlayAreaHeight * 1.5f)
+					FUncompressedPrecomputedVisibilityCell& CurrentCell = CombinedPrecomputedVisibilityCells[CellIndex];
+
+					float CellXFloat = (CombinedPrecomputedVisibilityCells[CellIndex].Bounds.GetCenter().X - AllCellsBounds.Min.X) / CellSize;
+					int32 CellX = FMath::Clamp(FMath::TruncToInt(CellXFloat), 0, GridSizeX - 1);
+
+					float CellYFloat = (CombinedPrecomputedVisibilityCells[CellIndex].Bounds.GetCenter().Y - AllCellsBounds.Min.Y) / CellSize;
+					int32 CellY = FMath::Clamp(FMath::TruncToInt(CellYFloat), 0, GridSizeY - 1);
+
+					TotalNumQueries += CurrentCell.VisibilityData.Num() * 8;
+
+					const FPrecomputedVisibilitySortGridCell& GridCell = SortGrid[CellY * GridSizeX + CellX];
+
+					for (int32 YOffset = -1; YOffset <= 1; YOffset++)
+					{
+						for (int32 XOffset = -1; XOffset <= 1; XOffset++)
+						{
+							int32 FinalCellX = CellX + XOffset;
+							int32 FinalCellY = CellY + YOffset;
+
+							if (FinalCellX >= 0 && FinalCellX < GridSizeX && FinalCellY >= 0 && FinalCellY < GridSizeY)
+							{
+								const TArray<FUncompressedPrecomputedVisibilityCell, TInlineAllocator<2>>& CurrentSortCell = SortGrid[FinalCellY * GridSizeX + FinalCellX].Cells;
+
+								for (int32 VisibilityCellIndex = 0; VisibilityCellIndex < CurrentSortCell.Num(); VisibilityCellIndex++)
+								{
+									const FUncompressedPrecomputedVisibilityCell& OtherCell = CurrentSortCell[VisibilityCellIndex];
+									SpreadVisibilityCell(CellSize, PlayAreaHeight, OtherCell, CurrentCell, QueriesVisibleFromSpreadingNeighbors);
+								}
+							}
+						}
+					}
+				}
+			}
+			else
+			{
+				// Brute force O(N^2) neighbor spreading version
+				// Copy the original data since we read from outside the current cell
+				TArray<FUncompressedPrecomputedVisibilityCell> OriginalPrecomputedVisibilityCells(CombinedPrecomputedVisibilityCells);
+				for (int32 CellIndex = 0; CellIndex < CombinedPrecomputedVisibilityCells.Num(); CellIndex++)
+				{
+					FUncompressedPrecomputedVisibilityCell& CurrentCell = CombinedPrecomputedVisibilityCells[CellIndex];
+					TotalNumQueries += CurrentCell.VisibilityData.Num() * 8;
+					for (int32 OtherCellIndex = 0; OtherCellIndex < OriginalPrecomputedVisibilityCells.Num(); OtherCellIndex++)
 					{
-						// Combine the neighbor's visibility with the current cell's visibility
-						// This reduces visibility errors at the cost of less effective culling
-						QueriesVisibleFromSpreadingNeighbors += AccumulateVisibility(OtherCell.VisibilityData, CurrentCell.VisibilityData);
+						const FUncompressedPrecomputedVisibilityCell& OtherCell = OriginalPrecomputedVisibilityCells[OtherCellIndex];
+
+						SpreadVisibilityCell(CellSize, PlayAreaHeight, OtherCell, CurrentCell, QueriesVisibleFromSpreadingNeighbors);
 					}
 				}
 			}
@@ -3341,13 +3430,13 @@ void FLightmassProcessor::ApplyPrecomputedVisibility()
 
 		System.GetWorld()->PersistentLevel->PrecomputedVisibilityHandler.UpdateVisibilityStats(true);
 
-		UE_LOG(LogLightmassSolver, Warning, TEXT("ApplyPrecomputedVisibility %.1fs, %.1f%% of all queries changed to visible from spreading neighbors, compressed %.3fMb to %.3fMb (%.1f ratio) with %u rendering buckets"),
+		UE_LOG(LogStaticLightingSystem, Log, TEXT("ApplyPrecomputedVisibility %.1fs with %u cells, %.1f%% of all queries changed to visible from spreading neighbors, compressed %.3fMb to %.3fMb (%.1f ratio)"),
 			FPlatformTime::Seconds() - StartTime,
+			CombinedPrecomputedVisibilityCells.Num(),
 			100.0f * QueriesVisibleFromSpreadingNeighbors / TotalNumQueries,
 			UncompressedSize / 1024.0f / 1024.0f,
 			TotalCompressedSize / 1024.0f / 1024.0f,
-			UncompressedSize / (float)TotalCompressedSize,
-			CellRenderingBuckets.Num());
+			UncompressedSize / (float)TotalCompressedSize);
 	}
 	else
 	{
diff --git a/Engine/Source/Editor/UnrealEd/Public/UnrealEd.h b/Engine/Source/Editor/UnrealEd/Public/UnrealEd.h
--- a/Engine/Source/Editor/UnrealEd/Public/UnrealEd.h
+++ b/Engine/Source/Editor/UnrealEd/Public/UnrealEd.h
@@ -27,6 +27,7 @@
 #include "Classes/ActorFactories/ActorFactoryBoxVolume.h"
 #include "Classes/ActorFactories/ActorFactoryCameraActor.h"
 #include "Classes/ActorFactories/ActorFactoryCharacter.h"
+#include "Classes/ActorFactories/ActorFactorySubDSurface.h"
 #include "Classes/ActorFactories/ActorFactoryClass.h"
 #include "Classes/ActorFactories/ActorFactoryCylinderVolume.h"
 #include "Classes/ActorFactories/ActorFactoryDeferredDecal.h"
@@ -179,6 +180,7 @@
 #include "Classes/Factories/ReimportSoundSurroundFactory.h"
 #include "Classes/Factories/StructureFactory.h"
 #include "Classes/Factories/SubsurfaceProfileFactory.h"
+#include "Classes/Factories/SubDSurfaceFactory.h"
 #include "Classes/Factories/Texture2dFactoryNew.h"
 #include "Classes/Factories/TextureFactory.h"
 #include "Classes/Factories/ReimportTextureFactory.h"
@@ -288,6 +290,7 @@
 #include "Classes/Exporters/TextureCubeExporterHDR.h"
 #include "Classes/Exporters/TextureExporterBMP.h"
 #include "Classes/Exporters/TextureExporterHDR.h"
+#include "Classes/Exporters/RenderTargetExporterHDR.h"
 #include "Classes/Exporters/TextureExporterPCX.h"
 #include "Classes/Exporters/TextureExporterTGA.h"
 #include "Classes/ThumbnailRendering/ThumbnailRenderer.h"
diff --git a/Engine/Source/Runtime/Core/Private/HAL/ConsoleManager.cpp b/Engine/Source/Runtime/Core/Private/HAL/ConsoleManager.cpp
--- a/Engine/Source/Runtime/Core/Private/HAL/ConsoleManager.cpp
+++ b/Engine/Source/Runtime/Core/Private/HAL/ConsoleManager.cpp
@@ -2233,11 +2233,11 @@ static TAutoConsoleVariable<int32> CVarDetailMode(
 static TAutoConsoleVariable<int32> CVarDBuffer(
 	TEXT("r.DBuffer"),
 	0,
-	TEXT("Experimental DBuffer feature: Generating deferred decals before the BasePass.\n")
-	TEXT("Allows decals to be correctly lit by baked lighting. Receivers need to be rendered in the early zpass.\n")
-	TEXT("At the moment only can be ensures by full enablng this pass: r.EarlyZPassMovable=1 r.EarlyZPass=2\n")
+	TEXT("Enables DBuffer decal material blend modes.\n")
+	TEXT("DBuffer decals are rendered before the base pass, allowing them to affect static lighting and skylighting correctly. \n")
+	TEXT("When enabled, a full prepass will be forced which adds CPU / GPU cost.  Several texture lookups will be done in the base pass to fetch the decal properties, which adds pixel work.\n")
 	TEXT(" 0: off\n")
-	TEXT(" 1: on (needs early pass rendering on all decal receivers and base pass lookups into the DBuffer, costs GPU memory, allows GBuffer compression)"),
+	TEXT(" 1: on"),
 	ECVF_RenderThreadSafe | ECVF_ReadOnly);
 
 static TAutoConsoleVariable<float> CVarSkeletalMeshLODRadiusScale(
diff --git a/Engine/Source/Runtime/Core/Private/Misc/FbxErrors.cpp b/Engine/Source/Runtime/Core/Private/Misc/FbxErrors.cpp
--- a/Engine/Source/Runtime/Core/Private/Misc/FbxErrors.cpp
+++ b/Engine/Source/Runtime/Core/Private/Misc/FbxErrors.cpp
@@ -68,6 +68,8 @@ FName FFbxErrors::SkeletalMesh_RestoreSortingForSectionNumber(TEXT("SkeletalMesh
 
 FName FFbxErrors::SkeletalMesh_NoMeshFoundOnRoot(TEXT("SkeletalMesh_NoMeshFoundOnRoot"));
 
+FName FFbxErrors::SkeletalMesh_ImportSubDSurface(TEXT("SkeletalMesh_ImportSubDSurface"));
+
 FName FFbxErrors::SkeletalMesh_InvalidRoot(TEXT("SkeletalMesh_InvalidRoot"));
 
 FName FFbxErrors::SkeletalMesh_InvalidBone(TEXT("SkeletalMesh_InvalidBone"));
diff --git a/Engine/Source/Runtime/Core/Public/Misc/FbxErrors.h b/Engine/Source/Runtime/Core/Public/Misc/FbxErrors.h
--- a/Engine/Source/Runtime/Core/Public/Misc/FbxErrors.h
+++ b/Engine/Source/Runtime/Core/Public/Misc/FbxErrors.h
@@ -75,6 +75,8 @@ struct CORE_API FFbxErrors
 	static FName SkeletalMesh_RestoreSortingForSectionNumber;
 
 	static FName SkeletalMesh_NoMeshFoundOnRoot;
+	
+	static FName SkeletalMesh_ImportSubDSurface;
 
 	static FName SkeletalMesh_InvalidRoot;
 
diff --git a/Engine/Source/Runtime/Engine/Classes/Components/LightComponent.h b/Engine/Source/Runtime/Engine/Classes/Components/LightComponent.h
--- a/Engine/Source/Runtime/Engine/Classes/Components/LightComponent.h
+++ b/Engine/Source/Runtime/Engine/Classes/Components/LightComponent.h
@@ -90,6 +90,10 @@ class ENGINE_API ULightComponent : public ULightComponentBase
 	UPROPERTY(EditAnywhere, BlueprintReadOnly, Category=Light, AdvancedDisplay, meta=(UIMin = "0.0", UIMax = "1.0", DisplayName = "Shadow Filter Sharpen"))
 	float ShadowSharpen;
 	
+	/** Length of screen space ray trace for sharp contact shadows. Zero is disabled. */
+	UPROPERTY(EditAnywhere, BlueprintReadOnly, Category = Light, AdvancedDisplay, meta = (UIMin = "0.0", UIMax = "1.0"))
+	float ContactShadowLength;
+
 	UPROPERTY()
 	uint32 InverseSquaredFalloff_DEPRECATED:1;
 
diff --git a/Engine/Source/Runtime/Engine/Classes/Components/PrimitiveComponent.h b/Engine/Source/Runtime/Engine/Classes/Components/PrimitiveComponent.h
--- a/Engine/Source/Runtime/Engine/Classes/Components/PrimitiveComponent.h
+++ b/Engine/Source/Runtime/Engine/Classes/Components/PrimitiveComponent.h
@@ -472,6 +472,8 @@ public:
 	UPROPERTY(transient)
 	float LastRenderTime;
 
+	UPROPERTY(transient)
+	float LastRenderTimeOnScreen;
 private:
 	UPROPERTY()
 	TEnumAsByte<enum ECanBeCharacterBase> CanBeCharacterBase_DEPRECATED;
diff --git a/Engine/Source/Runtime/Engine/Classes/Components/ReflectionCaptureComponent.h b/Engine/Source/Runtime/Engine/Classes/Components/ReflectionCaptureComponent.h
--- a/Engine/Source/Runtime/Engine/Classes/Components/ReflectionCaptureComponent.h
+++ b/Engine/Source/Runtime/Engine/Classes/Components/ReflectionCaptureComponent.h
@@ -134,6 +134,9 @@ class UReflectionCaptureComponent : public USceneComponent
 		return FullHDRData;
 	}
 
+	ENGINE_API static int32 GetReflectionCaptureSize_GameThread();
+	ENGINE_API static int32 GetReflectionCaptureSize_RenderThread();
+
 	//~ Begin UActorComponent Interface
 	virtual void CreateRenderState_Concurrent() override;
 	virtual void DestroyRenderState_Concurrent() override;
diff --git a/Engine/Source/Runtime/Engine/Classes/Components/SceneCaptureComponent2D.h b/Engine/Source/Runtime/Engine/Classes/Components/SceneCaptureComponent2D.h
--- a/Engine/Source/Runtime/Engine/Classes/Components/SceneCaptureComponent2D.h
+++ b/Engine/Source/Runtime/Engine/Classes/Components/SceneCaptureComponent2D.h
@@ -1,19 +1,9 @@
 // Copyright 1998-2016 Epic Games, Inc. All Rights Reserved.
 
-
 #pragma once
 #include "Components/SceneCaptureComponent.h"
 #include "SceneCaptureComponent2D.generated.h"
 
-// -> will be exported to EngineDecalClasses.h
-
-UENUM()
-enum ESceneCaptureSource 
-{ 
-	SCS_SceneColorHDR UMETA(DisplayName="Scene Color (HDR)"),
-	SCS_FinalColorLDR UMETA(DisplayName="Final Color (LDR with PostProcess)")
-};
-
 /**
  *	Used to capture a 'snapshot' of the scene from a single plane and feed it to a render target.
  */
@@ -22,11 +12,18 @@ class USceneCaptureComponent2D : public USceneCaptureComponent
 {
 	GENERATED_UCLASS_BODY()
 
+	UPROPERTY(EditAnywhere, BlueprintReadWrite, Category=Projection, meta=(DisplayName = "Projection Type"))
+	TEnumAsByte<ECameraProjectionMode::Type> ProjectionType;
+
 	/** Camera field of view (in degrees). */
-	UPROPERTY(interp, Category=SceneCapture, meta=(DisplayName = "Field of View", UIMin = "5.0", UIMax = "170", ClampMin = "0.001", ClampMax = "360.0"))
+	UPROPERTY(interp, Category=Projection, meta=(DisplayName = "Field of View", UIMin = "5.0", UIMax = "170", ClampMin = "0.001", ClampMax = "360.0"))
 	float FOVAngle;
 
-	/** Temporary render target that can be used by the editor. */
+	/** The desired width (in world units) of the orthographic view (ignored in Perspective mode) */
+	UPROPERTY(EditAnywhere, BlueprintReadWrite, Category=Projection)
+	float OrthoWidth;
+
+	/** Output render target of the scene capture that can be read in materals. */
 	UPROPERTY(EditAnywhere, BlueprintReadWrite, Category=SceneCapture)
 	class UTextureRenderTarget2D* TextureTarget;
 
@@ -40,9 +37,6 @@ class USceneCaptureComponent2D : public USceneCaptureComponent
 	UPROPERTY(interp, Category=PostProcessVolume, BlueprintReadWrite, meta=(UIMin = "0.0", UIMax = "1.0"))
 	float PostProcessBlendWeight;
 
-private:
-
-public:
 	//~ Begin UActorComponent Interface
 	virtual void OnRegister() override;
 	virtual void SendRenderTransform_Concurrent() override;
@@ -56,6 +50,7 @@ public:
 
 	//~ Begin UObject Interface
 #if WITH_EDITOR
+	virtual bool CanEditChange(const UProperty* InProperty) const override;
 	virtual void PostEditChangeProperty(FPropertyChangedEvent& PropertyChangedEvent) override;
 #endif // WITH_EDITOR
 	
@@ -67,9 +62,18 @@ public:
 	UFUNCTION(BlueprintCallable, Category="Rendering")
 	ENGINE_API void AddOrUpdateBlendable(TScriptInterface<IBlendableInterface> InBlendableObject, float InWeight = 1.0f) { PostProcessSettings.AddBlendable(InBlendableObject, InWeight); }
 
-	/** Render the scene to the texture */
+	/** Render the scene to the texture the next time the main view is rendered. */
+	ENGINE_API void CaptureSceneDeferred();
+
+	// For backwards compatibility
+	void UpdateContent() { CaptureSceneDeferred(); }
+
+	/** 
+	 * Render the scene to the texture target immediately.  
+	 * This should not be used if bCaptureEveryFrame is enabled, or the scene capture will render redundantly. 
+	 */
 	UFUNCTION(BlueprintCallable,Category = "Rendering|SceneCapture")
-	ENGINE_API void UpdateContent();
+	ENGINE_API void CaptureScene();
 
 	ENGINE_API static void UpdateDeferredCaptures( FSceneInterface* Scene );
 };
diff --git a/Engine/Source/Runtime/Engine/Classes/Components/SceneCaptureComponentCube.h b/Engine/Source/Runtime/Engine/Classes/Components/SceneCaptureComponentCube.h
--- a/Engine/Source/Runtime/Engine/Classes/Components/SceneCaptureComponentCube.h
+++ b/Engine/Source/Runtime/Engine/Classes/Components/SceneCaptureComponentCube.h
@@ -29,9 +29,18 @@ public:
 #endif // WITH_EDITOR
 	//~ End UObject Interface
 
-	/** Render the scene to the texture */
+	/** Render the scene to the texture the next time the main view is rendered. */
+	ENGINE_API void CaptureSceneDeferred();
+
+	/** 
+	 * Render the scene to the texture target immediately.  
+	 * This should not be used if bCaptureEveryFrame is enabled, or the scene capture will render redundantly. 
+	 */
 	UFUNCTION(BlueprintCallable,Category = "Rendering|SceneCapture")
-	ENGINE_API void UpdateContent();
+	ENGINE_API void CaptureScene();
+
+	// For backwards compatibility
+	void UpdateContent() { CaptureSceneDeferred(); }
 
 	ENGINE_API static void UpdateDeferredCaptures(FSceneInterface* Scene);
 };
\ No newline at end of file
diff --git a/Engine/Source/Runtime/Engine/Classes/Components/SubDSurfaceComponent.h b/Engine/Source/Runtime/Engine/Classes/Components/SubDSurfaceComponent.h
new file mode 100644
--- /dev/null
+++ b/Engine/Source/Runtime/Engine/Classes/Components/SubDSurfaceComponent.h
@@ -0,0 +1,64 @@
+// Copyright 1998-2014 Epic Games, Inc. All Rights Reserved.
+
+
+#pragma once
+#include "Components/PrimitiveComponent.h"
+#include "SubDSurfaceComponent.generated.h"
+
+/**
+ * Subdivision Surface Component (Experimental, Early work in progress)
+ */
+UCLASS(ClassGroup=Rendering, hidecategories=(Object,LOD,Physics,TextureStreaming,Activation,"Components|Activation",Collision), editinlinenew, meta=(BlueprintSpawnableComponent = ""))
+class ENGINE_API USubDSurfaceComponent : public UPrimitiveComponent
+{
+	GENERATED_UCLASS_BODY()
+
+	// asset
+private:
+	UPROPERTY(EditAnywhere, Category=StaticMesh)
+	class USubDSurface* Mesh;
+public:
+
+	/* Refinement Level of the SubD mesh */
+	UPROPERTY(EditAnywhere, BlueprintReadWrite, Category=StaticMesh)
+	int32 DebugLevel;
+
+	/** Change the SubDSurface used by this instance. */
+	UFUNCTION(BlueprintCallable, Category="Components|Mesh")
+	virtual bool SetMesh(class USubDSurface* NewMesh);
+
+	UPROPERTY()
+	UStaticMeshComponent* DisplayMeshComponent;
+
+public:
+
+	// -----------------------------
+
+	// Begin UObject interface.
+#if WITH_EDITOR
+	virtual void PostEditChangeProperty(FPropertyChangedEvent& PropertyChangedEvent) override;
+#endif // WITH_EDITOR
+	// End UObject interface.
+
+	// Begin UPrimitiveComponent interface.
+	virtual void GetUsedMaterials( TArray<UMaterialInterface*>& OutMaterials ) const;
+	virtual int32 GetNumMaterials() const;
+	virtual UMaterialInterface* GetMaterial(int32 ElementIndex) const;
+	virtual bool ShouldRecreateProxyOnUpdateTransform() const override;
+	virtual void SetMaterial(int32 ElementIndex, UMaterialInterface* InMaterial);
+	// End UPrimitiveComponent interface.
+
+	// Begin USceneComponent interface.
+	virtual FBoxSphereBounds CalcBounds(const FTransform& LocalToWorld) const override;
+	// Begin USceneComponent interface.
+
+	// --
+
+	void SetDisplayMeshComponent(UStaticMeshComponent* InDisplayMeshComponent);
+	void RecreateMeshData();
+};
+
+
+
+
+
diff --git a/Engine/Source/Runtime/Engine/Classes/Engine/BlendableInterface.h b/Engine/Source/Runtime/Engine/Classes/Engine/BlendableInterface.h
--- a/Engine/Source/Runtime/Engine/Classes/Engine/BlendableInterface.h
+++ b/Engine/Source/Runtime/Engine/Classes/Engine/BlendableInterface.h
@@ -33,7 +33,7 @@ enum EBlendableLocation
 
 /** Dummy class needed to support Cast<IBlendableInterface>(Object). */
 UINTERFACE()
-class UBlendableInterface : public UInterface
+class ENGINE_API UBlendableInterface : public UInterface
 {
 	GENERATED_UINTERFACE_BODY()
 };
diff --git a/Engine/Source/Runtime/Engine/Classes/Engine/EngineTypes.h b/Engine/Source/Runtime/Engine/Classes/Engine/EngineTypes.h
--- a/Engine/Source/Runtime/Engine/Classes/Engine/EngineTypes.h
+++ b/Engine/Source/Runtime/Engine/Classes/Engine/EngineTypes.h
@@ -308,6 +308,17 @@ namespace ETranslucentSortPolicy
 	};
 }
 
+UENUM()
+enum ESceneCaptureSource 
+{ 
+	SCS_SceneColorHDR UMETA(DisplayName="SceneColor (HDR) in RGB, Opacity in A"),
+	SCS_FinalColorLDR UMETA(DisplayName="Final Color (LDR) in RGB"),
+	SCS_SceneColorSceneDepth UMETA(DisplayName="SceneColor (HDR) in RGB, SceneDepth in A"),
+	SCS_SceneDepth UMETA(DisplayName="SceneDepth in R"),
+	SCS_Normal UMETA(DisplayName="Normal in RGB"),
+	SCS_BaseColor UMETA(DisplayName="BaseColor in RGB")
+};
+
 USTRUCT()
 struct FLightingChannels
 {
diff --git a/Engine/Source/Runtime/Engine/Classes/Engine/RendererSettings.h b/Engine/Source/Runtime/Engine/Classes/Engine/RendererSettings.h
--- a/Engine/Source/Runtime/Engine/Classes/Engine/RendererSettings.h
+++ b/Engine/Source/Runtime/Engine/Classes/Engine/RendererSettings.h
@@ -209,7 +209,7 @@ class ENGINE_API URendererSettings : public UDeveloperSettings
 		ToolTip="When adaptive tessellation is enabled it will try to tessellate a mesh so that each triangle contains the specified number of pixels. The tessellation multiplier specified in the material can increase or decrease the amount of tessellation."))
 	float TessellationAdaptivePixelsPerTriangle;
 
-	UPROPERTY(config, EditAnywhere, Category=Postprocessing, meta=(
+	UPROPERTY(config, EditAnywhere, Category=Translucency, meta=(
 		ConsoleVariable="r.SeparateTranslucency",
 		ToolTip="Allow translucency to be rendered to a separate render targeted and composited after depth of field. Prevents translucency from appearing out of focus."))
 	uint32 bSeparateTranslucency:1;
@@ -302,7 +302,7 @@ class ENGINE_API URendererSettings : public UDeveloperSettings
 		ConfigRestartRequired=true))
 	uint32 bSelectiveBasePassOutputs:1;
 
-	UPROPERTY(config, EditAnywhere, Category = Optimizations, meta = (
+	UPROPERTY(config, EditAnywhere, Category = Lighting, meta = (
 		ConsoleVariable = "r.AllowGlobalClipPlane", DisplayName = "Support global clip plane for Planar Reflections",
 		ToolTip = "Whether to support the global clip plane needed for planar reflections.  Enabling this increases BasePass triangle cost by ~15% regardless of whether planar reflections are active. Changing this setting requires restarting the editor.",
 		ConfigRestartRequired = true))
diff --git a/Engine/Source/Runtime/Engine/Classes/Engine/SubDSurface.h b/Engine/Source/Runtime/Engine/Classes/Engine/SubDSurface.h
new file mode 100644
--- /dev/null
+++ b/Engine/Source/Runtime/Engine/Classes/Engine/SubDSurface.h
@@ -0,0 +1,135 @@
+// Copyright 1998-2014 Epic Games, Inc. All Rights Reserved.
+#pragma once
+
+#include "../Curves/CurveFloat.h"
+#include "SubDSurface.generated.h"
+
+struct FSubDMeshSection
+{
+	/** Start face (n-gon) index where this section starts  */
+	uint32 FirstIndex;
+	/** Number of faces (n-gons) for this section */
+	uint32 FaceCount;
+
+	/** The index of the material with which to render this section. 65k should be enough */
+	uint16 MaterialIndex;
+
+	/** If true, collision is enabled for this section. */
+	uint32 bEnableCollision:1;
+	/** If true, this section will cast a shadow. */
+	uint32 bCastShadow:1;
+
+	/** Constructor. */
+	FSubDMeshSection()
+		: FirstIndex(0)
+		, FaceCount(0)
+		, MaterialIndex(0)
+		, bEnableCollision(false)
+		, bCastShadow(true)
+	{}
+
+	/** Serializer. */
+//	friend FArchive& operator<<(FArchive& Ar,FSubDMeshSection& Section);
+};
+
+// typical vertex attribute formats but as they are only CPU processed at the moment we could chnage that
+UENUM()
+enum EVertexAttributeStreamType
+{
+	VAST_unknown,
+	VAST_float,
+	VAST_float2,		// FVector2D e.g. UV
+	VAST_float3,		// FVector e.g. position
+	VAST_float4,
+//	VAST_byteARGB,
+};
+
+UCLASS(MinimalAPI)
+class UVertexAttributeStream : public UObject
+{
+	GENERATED_UCLASS_BODY()
+
+	// @return same as what GetFVector() would return
+	ENGINE_API FVector* CreateFVectorUninitialized(uint32 Count);
+	// @return same as what GetFVector2D() would return
+	ENGINE_API FVector2D* CreateFVector2DUninitialized(uint32 Count);
+
+	// @return 0 if there is no data or it's not the right type
+	ENGINE_API FVector* GetFVector(uint32& OutCount);
+
+	// @return 0 if there is no data or it's not the right type
+	ENGINE_API FVector2D* GetFVector2D(uint32& OutCount);
+
+	ENGINE_API uint32 Num() const;
+	
+	// --------------------------------------------------------
+
+	// e.g. FName(TEXT("Position"))
+	UPROPERTY()
+	FName Usage;
+
+private:
+
+	// e.g. VAST_unknown
+	UPROPERTY()
+	TEnumAsByte<enum EVertexAttributeStreamType> AttributeType;
+
+	// actual type depends on AttributeType, position is indexed by IndicesPerFace, other attributes by CornerID (e.g. 10 quads result in 40 attributes)
+	UPROPERTY()
+	TArray<uint8> Data;
+
+	ENGINE_API uint32 GetTypeSize() const;
+};
+
+/**
+ * Subdivision Surface Asset (Experimental, Early work in progress)
+ */
+UCLASS(autoexpandcategories = SubDSurface, MinimalAPI)
+class USubDSurface : public UObject
+{
+	GENERATED_UCLASS_BODY()
+
+	// becomes TopologyDescriptor::numVertsPerFace, TopologyDescriptor::numFaces = .Num()
+	// 32bit for OpenSubDiv TopologyDescriptor but it could be limited to 4 bits
+	UPROPERTY()
+	TArray<uint32> VertexCountPerFace;
+
+	// becomes TopologyDescriptor::vertIndicesPerFace e.g. two triangles: 0,1,2, 2,3,0
+	UPROPERTY()
+	TArray<uint32> IndicesPerFace;
+
+	// various streams e.g "Position"
+	UPROPERTY()
+	TArray<UVertexAttributeStream*> VertexAttributeStreams;
+	
+	TArray<FSubDMeshSection> Sections;
+
+	/** Materials used by this mesh. Individual sections index in to this array. */
+	UPROPERTY()
+	TArray<UMaterialInterface*> Materials;
+
+	// --------------------------
+
+	// todo
+	ENGINE_API uint32 GetVertexCount() const;
+
+	// @return 0 if not found
+	ENGINE_API UVertexAttributeStream* CreateVertexAttributeStream(FName InUsage);
+
+	// @return 0 if not found
+	ENGINE_API UVertexAttributeStream* FindStreamByUsage(FName InUsage);
+
+/* later
+    int           numCreases;
+    Index const * creaseVertexIndexPairs;
+    float const * creaseWeights;
+
+    int           numCorners;
+    Index const * cornerVertexIndices;
+    float const * cornerWeights;
+*/
+
+	/** A fence which is used to keep track of the rendering thread releasing the mesh resources. */
+//	FRenderCommandFence ReleaseResourcesFence;
+};
+
diff --git a/Engine/Source/Runtime/Engine/Classes/Engine/SubDSurfaceActor.h b/Engine/Source/Runtime/Engine/Classes/Engine/SubDSurfaceActor.h
new file mode 100644
--- /dev/null
+++ b/Engine/Source/Runtime/Engine/Classes/Engine/SubDSurfaceActor.h
@@ -0,0 +1,43 @@
+// Copyright 1998-2014 Epic Games, Inc. All Rights Reserved.
+
+/**
+ * Subdivision Surface Actor (Experimental, Early work in progress)
+ */
+
+#pragma once
+
+#include "SubDSurfaceActor.generated.h"
+
+UCLASS(hidecategories=(Collision, Attachment, Actor))
+class ENGINE_API ASubDSurfaceActor : public AActor
+{
+	GENERATED_UCLASS_BODY()
+
+	friend class UActorFactorySubDSurface;
+
+private:
+	/** Component to render the actor, used GetSubDSurface() to access */
+	UPROPERTY(Category = SubDSurfaceActor, VisibleAnywhere, BlueprintReadOnly, meta = (ExposeFunctionCategories = "Rendering|Components|SubDSurface", AllowPrivateAccess = "true"))
+	class USubDSurfaceComponent* SubDSurface;
+
+	// later this all will be done by USubDSurfaceComponent
+	UPROPERTY(Category = SubDSurfaceActor, VisibleAnywhere, BlueprintReadOnly, meta = (ExposeFunctionCategories = "Mesh,Rendering,Physics,Components|StaticMesh", AllowPrivateAccess = "true"))
+	UStaticMeshComponent* DisplayMeshComponent;
+
+
+#if WITH_EDITORONLY_DATA
+	// Reference to the billboard component
+	UPROPERTY()
+	UBillboardComponent* SpriteComponent;
+
+	/** Returns SpriteComponent subobject **/
+	UBillboardComponent* GetSpriteComponent() const { return SpriteComponent; }
+#endif
+
+public:
+	/** Returns SubDSurface subobject **/
+	class USubDSurfaceComponent* GetSubDSurface() const { return SubDSurface; }
+};
+
+
+
diff --git a/Engine/Source/Runtime/Engine/Classes/Engine/World.h b/Engine/Source/Runtime/Engine/Classes/Engine/World.h
--- a/Engine/Source/Runtime/Engine/Classes/Engine/World.h
+++ b/Engine/Source/Runtime/Engine/Classes/Engine/World.h
@@ -660,6 +660,13 @@ private:
 	UPROPERTY(Transient)
 	TArray<class UMaterialParameterCollectionInstance*> ParameterCollectionInstances;
 
+	/** 
+	 * Canvas object used for drawing to render targets from blueprint functions eg DrawMaterialToRenderTarget.
+	 * This is cached as UCanvas creation takes >100ms.
+	 */
+	UPROPERTY(Transient)
+	UCanvas* CanvasForRenderingToTarget;
+
 public:
 	/** Set the pointer to the Navgation system. */
 	void SetNavigationSystem( UNavigationSystem* InNavigationSystem);
@@ -1997,6 +2004,9 @@ public:
 	/** Updates this world's scene with the list of instances, and optionally updates each instance's uniform buffer. */
 	void UpdateParameterCollectionInstances(bool bUpdateInstanceUniformBuffers);
 
+	/** Gets the canvas object for rendering to a render target.  Will allocate one if needed. */
+	UCanvas* GetCanvasForRenderingToTarget();
+
 	/** Struct containing a collection of optional parameters for initialization of a World. */
 	struct InitializationValues
 	{
diff --git a/Engine/Source/Runtime/Engine/Engine.Build.cs b/Engine/Source/Runtime/Engine/Engine.Build.cs
--- a/Engine/Source/Runtime/Engine/Engine.Build.cs
+++ b/Engine/Source/Runtime/Engine/Engine.Build.cs
@@ -98,6 +98,7 @@ public class Engine : ModuleRules
 				"Landscape",
                 "UMG",
 				"Projects",
+//                "RawMesh",          // causes WARNING: Non-editor build cannot depend on non-redistributable modules. Details:
 				"Niagara",
                 "Internationalization",
                 "MaterialShaderQualitySettings",
@@ -329,6 +330,7 @@ public class Engine : ModuleRules
 				"Vorbis",
 				"VorbisFile",
 				"libOpus"
+			   // "OpenSubdiv"
 				);
 
 			if (UEBuildConfiguration.bCompileLeanAndMeanUE == false)
diff --git a/Engine/Source/Runtime/Engine/Private/Animation/AnimSequence.cpp b/Engine/Source/Runtime/Engine/Private/Animation/AnimSequence.cpp
--- a/Engine/Source/Runtime/Engine/Private/Animation/AnimSequence.cpp
+++ b/Engine/Source/Runtime/Engine/Private/Animation/AnimSequence.cpp
@@ -480,6 +480,11 @@ void UAnimSequence::PostLoad()
 
 	USkeleton* CurrentSkeleton = GetSkeleton();
 
+	if (CurrentSkeleton)
+	{
+		VerifyCurveNames<FFloatCurve>(CurrentSkeleton, USkeleton::AnimCurveMappingName, CompressedCurveData.FloatCurves);
+	}
+
 #if WITH_EDITOR
 	if (CurrentSkeleton)
 	{
@@ -504,7 +509,7 @@ void UAnimSequence::PostLoad()
 		}
 	}
 
-	// Compressed curve flags are not authorative (they come from the DDC). Keep them up to date with
+	// Compressed curve flags are not authoritative (they come from the DDC). Keep them up to date with
 	// actual anim flags 
 	for (FFloatCurve& Curve : RawCurveData.FloatCurves)
 	{
diff --git a/Engine/Source/Runtime/Engine/Private/Components/LightComponent.cpp b/Engine/Source/Runtime/Engine/Private/Components/LightComponent.cpp
--- a/Engine/Source/Runtime/Engine/Private/Components/LightComponent.cpp
+++ b/Engine/Source/Runtime/Engine/Private/Components/LightComponent.cpp
@@ -215,6 +215,7 @@ FLightSceneProxy::FLightSceneProxy(const ULightComponent* InLightComponent)
 	, IndirectLightingScale(InLightComponent->IndirectLightingIntensity)
 	, ShadowBias(InLightComponent->ShadowBias)
 	, ShadowSharpen(InLightComponent->ShadowSharpen)
+	, ContactShadowLength(InLightComponent->ContactShadowLength)
 	, MinRoughness(InLightComponent->MinRoughness)
 	, LightGuid(InLightComponent->LightGuid)
 	, ShadowMapChannel(InLightComponent->ShadowMapChannel)
@@ -339,6 +340,7 @@ ULightComponent::ULightComponent(const FObjectInitializer& ObjectInitializer)
 	IndirectLightingIntensity = 1.0f;
 	ShadowBias = 0.5f;
 	ShadowSharpen = 0.0f;
+	ContactShadowLength = 0.0f;
 	bUseIESBrightness = false;
 	IESBrightnessScale = 1.0f;
 	IESTexture = NULL;
@@ -568,6 +570,7 @@ void ULightComponent::PostEditChangeProperty(FPropertyChangedEvent& PropertyChan
 		PropertyName != GET_MEMBER_NAME_STRING_CHECKED(ULightComponent, DisabledBrightness) &&
 		PropertyName != GET_MEMBER_NAME_STRING_CHECKED(ULightComponent, ShadowBias) &&
 		PropertyName != GET_MEMBER_NAME_STRING_CHECKED(ULightComponent, ShadowSharpen) &&
+		PropertyName != GET_MEMBER_NAME_STRING_CHECKED(ULightComponent, ContactShadowLength) &&
 		PropertyName != GET_MEMBER_NAME_STRING_CHECKED(ULightComponent, bEnableLightShaftBloom) &&
 		PropertyName != GET_MEMBER_NAME_STRING_CHECKED(ULightComponent, BloomScale) &&
 		PropertyName != GET_MEMBER_NAME_STRING_CHECKED(ULightComponent, BloomThreshold) &&
diff --git a/Engine/Source/Runtime/Engine/Private/Components/PrimitiveComponent.cpp b/Engine/Source/Runtime/Engine/Private/Components/PrimitiveComponent.cpp
--- a/Engine/Source/Runtime/Engine/Private/Components/PrimitiveComponent.cpp
+++ b/Engine/Source/Runtime/Engine/Private/Components/PrimitiveComponent.cpp
@@ -153,6 +153,7 @@ UPrimitiveComponent::UPrimitiveComponent(const FObjectInitializer& ObjectInitial
 	PostPhysicsComponentTick.TickGroup = TG_PostPhysics;
 
 	LastRenderTime = -1000.0f;
+	LastRenderTimeOnScreen = -1000.0f;
 	BoundsScale = 1.0f;
 	MinDrawDistance = 0.0f;
 	DepthPriorityGroup = SDPG_World;
diff --git a/Engine/Source/Runtime/Engine/Private/Components/ReflectionCaptureComponent.cpp b/Engine/Source/Runtime/Engine/Private/Components/ReflectionCaptureComponent.cpp
--- a/Engine/Source/Runtime/Engine/Private/Components/ReflectionCaptureComponent.cpp
+++ b/Engine/Source/Runtime/Engine/Private/Components/ReflectionCaptureComponent.cpp
@@ -45,6 +45,24 @@ ENGINE_API TAutoConsoleVariable<int32> CVarReflectionCaptureSize(
 	TEXT("Set the resolution for all reflection capture cubemaps. Should be set via project's Render Settings. Must be power of 2. Defaults to 128.\n")
 	);
 
+static int32 SanatizeReflectionCaptureSize(int32 ReflectionCaptureSize)
+{
+	static const int32 MaxReflectionCaptureSize = 1024;
+	static const int32 MinReflectionCaptureSize = 1;
+
+	return FMath::Clamp(ReflectionCaptureSize, MinReflectionCaptureSize, MaxReflectionCaptureSize);
+}
+
+int32 UReflectionCaptureComponent::GetReflectionCaptureSize_GameThread()
+{
+	return SanatizeReflectionCaptureSize(CVarReflectionCaptureSize.GetValueOnGameThread());
+}
+
+int32 UReflectionCaptureComponent::GetReflectionCaptureSize_RenderThread()
+{
+	return SanatizeReflectionCaptureSize(CVarReflectionCaptureSize.GetValueOnRenderThread());
+}
+
 void UWorld::UpdateAllReflectionCaptures()
 {
 	if ( FeatureLevel < ERHIFeatureLevel::SM4)
@@ -1054,11 +1072,12 @@ void UReflectionCaptureComponent::PostLoad()
 	bool bRetainAllFeatureLevelData = GIsEditor && GMaxRHIFeatureLevel >= ERHIFeatureLevel::SM4;
 	bool bEncodedDataRequired = bRetainAllFeatureLevelData || (GMaxRHIFeatureLevel == ERHIFeatureLevel::ES2 || GMaxRHIFeatureLevel == ERHIFeatureLevel::ES3_1);
 	bool bFullDataRequired = GMaxRHIFeatureLevel >= ERHIFeatureLevel::SM4;
+	const int32 ReflectionCaptureSize = UReflectionCaptureComponent::GetReflectionCaptureSize_GameThread();
 
 	// If we're loading on a platform that doesn't require cooked data, attempt to load missing data from the DDC
 	if (!FPlatformProperties::RequiresCookedData())
 	{
-		if (!FullHDRData || FullHDRData->CubemapSize != CVarReflectionCaptureSize.GetValueOnGameThread())
+		if (!FullHDRData || FullHDRData->CubemapSize != ReflectionCaptureSize)
 		{
 			bDerivedDataDirty = true;
 			delete FullHDRData;
@@ -1101,7 +1120,7 @@ void UReflectionCaptureComponent::PostLoad()
 			INC_MEMORY_STAT_BY(STAT_ReflectionCaptureMemory, EncodedHDRDerivedData->CapturedData.GetAllocatedSize());
 		}
 
-		int32 CubemapSize = CVarReflectionCaptureSize.GetValueOnGameThread();
+		int32 CubemapSize = ReflectionCaptureSize;
 		if (FullHDRData)
 		{
 			CubemapSize = FullHDRData->CubemapSize;
@@ -1420,9 +1439,11 @@ void UReflectionCaptureComponent::UpdateReflectionCaptureContents(UWorld* WorldT
 
 				if (!ReflectionComponent->SM4FullHDRCubemapTexture)
 				{
+					const int32 ReflectionCaptureSize = UReflectionCaptureComponent::GetReflectionCaptureSize_GameThread();
+
 					// Create the cubemap if it wasn't already - this happens when updating a reflection capture that doesn't have valid DDC
 					ReflectionComponent->SM4FullHDRCubemapTexture = new FReflectionTextureCubeResource();
-					ReflectionComponent->SM4FullHDRCubemapTexture->SetupParameters(CVarReflectionCaptureSize.GetValueOnGameThread(), FMath::CeilLogTwo(CVarReflectionCaptureSize.GetValueOnGameThread()) + 1, PF_FloatRGBA, nullptr);
+					ReflectionComponent->SM4FullHDRCubemapTexture->SetupParameters(ReflectionCaptureSize, FMath::CeilLogTwo(ReflectionCaptureSize) + 1, PF_FloatRGBA, nullptr);
 					BeginInitResource(ReflectionComponent->SM4FullHDRCubemapTexture);
 					ReflectionComponent->MarkRenderStateDirty();
 				}
@@ -1497,8 +1518,10 @@ void UReflectionCaptureComponent::PreFeatureLevelChange(ERHIFeatureLevel::Type P
 	{
 		if (SM4FullHDRCubemapTexture == nullptr)
 		{
+			const int32 ReflectionCaptureSize = UReflectionCaptureComponent::GetReflectionCaptureSize_GameThread();
+
 			SM4FullHDRCubemapTexture = new FReflectionTextureCubeResource();
-			SM4FullHDRCubemapTexture->SetupParameters(CVarReflectionCaptureSize.GetValueOnGameThread(), FMath::CeilLogTwo(CVarReflectionCaptureSize.GetValueOnGameThread()) + 1, PF_FloatRGBA, nullptr);
+			SM4FullHDRCubemapTexture->SetupParameters(ReflectionCaptureSize, FMath::CeilLogTwo(ReflectionCaptureSize) + 1, PF_FloatRGBA, nullptr);
 			BeginInitResource(SM4FullHDRCubemapTexture);
 		}
 	}
