Copying //UE4/Dev-Framework to Dev-Main (//UE4/Dev-Main)

diff --git a/Engine/Source/Editor/UnrealEd/Private/UnrealWidget.cpp b/Engine/Source/Editor/UnrealEd/Private/UnrealWidget.cpp
--- a/Engine/Source/Editor/UnrealEd/Private/UnrealWidget.cpp
+++ b/Engine/Source/Editor/UnrealEd/Private/UnrealWidget.cpp
@@ -16,6 +16,129 @@ static const float OUTER_AXIS_CIRCLE_RADIUS = 56.0f;
 static const float ROTATION_TEXT_RADIUS = 75.0f;
 static const int32 AXIS_CIRCLE_SIDES = 24;
 
+/*
+ *  Simple struct used to create and group data related to the current window's / viewport's space,
+ *  orientation, and scale.
+ */
+struct FSpaceDescriptor
+{
+    /*
+     *  Creates a new FSpaceDescriptor.
+     *
+     *  @param View         The virtual view for the space.
+     *  @param Viewport     The real viewport for the space.
+     *  @param InLocation   The location of the camera in the virtual space.
+     */
+    FSpaceDescriptor(const FSceneView* View, const FEditorViewportClient* Viewport, const FVector& InLocation) :
+        bIsPerspective(View->ViewMatrices.ProjMatrix.M[3][3] < 1.0f),
+        bIsLocalSpace(Viewport->GetWidgetCoordSystemSpace() == COORD_Local),
+        bIsOrthoXY(!bIsPerspective && FMath::Abs(View->ViewMatrices.ViewMatrix.M[2][2]) > 0.0f),
+        bIsOrthoXZ(!bIsPerspective && FMath::Abs(View->ViewMatrices.ViewMatrix.M[1][2]) > 0.0f),
+        bIsOrthoYZ(!bIsPerspective && FMath::Abs(View->ViewMatrices.ViewMatrix.M[0][2]) > 0.0f),
+        UniformScale(View->WorldToScreen(InLocation).W * (4.0f / View->UnscaledViewRect.Width() / View->ViewMatrices.ProjMatrix.M[0][0])),
+        Scale(CreateScale())
+    {
+    }
+
+    // Wether or not the view is perspective.
+    const bool bIsPerspective;
+
+    // Whether or not the view is in local space.
+    const bool bIsLocalSpace;
+
+    // Whether or not the view is orthogonal to the XY plane.
+    const bool bIsOrthoXY;
+
+    // Whether or not the view is orthogonal to the XZ plane.
+    const bool bIsOrthoXZ;
+
+    // Whether or not the view is orthogonal to the YZ plane.
+    const bool bIsOrthoYZ;
+
+    // The uniform scale for the space.
+    const float UniformScale;
+
+    // The scale vector for the space based on orientation.
+    const FVector Scale;
+
+    /*
+     *  Used to determine whether or not the X axis should be drawn.
+     *
+     *  @param AxisToDraw   The desired axis to draw.
+     *
+     *  @return True if the axis should be drawn. False otherwise.
+     */
+    bool ShouldDrawAxisX(const EAxisList::Type AxisToDraw)
+    {
+        return ShouldDrawAxis(EAxisList::X, AxisToDraw, bIsOrthoYZ);
+    }
+
+    /*
+     *  Used to determine whether or not the Y axis should be drawn.
+     *
+     *  @param AxisToDraw   The desired axis to draw.
+     *
+     *  @return True if the axis should be drawn. False otherwise.
+     */
+    bool ShouldDrawAxisY(const EAxisList::Type AxisToDraw)
+    {
+        return ShouldDrawAxis(EAxisList::Y, AxisToDraw, bIsOrthoXZ);
+    }
+
+    /*
+     *  Used to determine whether or not the Z axis should be drawn.
+     *
+     *  @param AxisToDraw   The desired axis to draw.
+     *
+     *  @return True if the axis should be drawn. False otherwise.
+     */
+    bool ShouldDrawAxisZ(const EAxisList::Type AxisToDraw)
+    {
+        return ShouldDrawAxis(EAxisList::Z, AxisToDraw, bIsOrthoXY);
+    }
+        
+private:
+
+    /*
+     *  Creates a space scale vector from the determined orientation and uniform scale.
+     *
+     *  @return Space scale vector.
+     */
+    FVector CreateScale()
+    {
+        if (bIsOrthoXY)
+        {
+            return FVector(UniformScale, UniformScale, 1.0f);
+        }
+        else if (bIsOrthoXZ)
+        {
+            return FVector(UniformScale, 1.0f, UniformScale);
+        }
+        else if (bIsOrthoYZ)
+        {
+            return FVector(1.0f, UniformScale, UniformScale);
+        }
+        else
+        {
+            return FVector(UniformScale, UniformScale, UniformScale);
+        }
+    }
+    
+    /*
+     *  Used to determine whether or not a specific axis should be drawn.
+     *
+     *  @param AxisToCheck  The axis to check.
+     *  @param AxisToDraw   The desired axis to draw.
+     *  @param bIsOrtho     Whether or not the axis to check is orthogonal to the viewing orientation.
+     *
+     *  @return True if the axis should be drawn. False otherwise.
+     */
+    bool ShouldDrawAxis(const EAxisList::Type AxisToCheck, const EAxisList::Type AxisToDraw, const bool bIsOrtho)
+    {
+        return (AxisToCheck & AxisToDraw) && (bIsPerspective || bIsLocalSpace || !bIsOrtho);
+    }
+};
+
 FWidget::FWidget()
 {
 	EditorModeTools = NULL;
@@ -410,72 +533,43 @@ void FWidget::Render_Translate( const FSceneView* View, FPrimitiveDrawInterface*
 
 	// Figure out axis matrices
 	FMatrix WidgetMatrix = CustomCoordSystem * FTranslationMatrix( InLocation );
-
-	bool bIsPerspective = ( View->ViewMatrices.ProjMatrix.M[3][3] < 1.0f );
-	const bool bIsOrthoXY = !bIsPerspective && FMath::Abs(View->ViewMatrices.ViewMatrix.M[2][2]) > 0.0f;
-	const bool bIsOrthoXZ = !bIsPerspective && FMath::Abs(View->ViewMatrices.ViewMatrix.M[1][2]) > 0.0f;
-	const bool bIsOrthoYZ = !bIsPerspective && FMath::Abs(View->ViewMatrices.ViewMatrix.M[0][2]) > 0.0f;
-
-	// For local space widgets, we always want to draw all three axis, since they may not be aligned with
-	// the orthographic projection anyway.
-	const bool bIsLocalSpace = ( ViewportClient->GetWidgetCoordSystemSpace() == COORD_Local );
-
 	const EAxisList::Type DrawAxis = GetAxisToDraw( ViewportClient->GetWidgetMode() );
-
 	const bool bDisabled = IsWidgetDisabled();
 
-	FVector Scale;
-	float UniformScale = View->WorldToScreen(InLocation).W * ( 4.0f / View->UnscaledViewRect.Width() / View->ViewMatrices.ProjMatrix.M[0][0] );
-
-	if ( bIsOrthoXY )
-	{
-		Scale = FVector(UniformScale, UniformScale, 1.0f);
-	}
-	else if ( bIsOrthoXZ )
-	{
-		Scale = FVector(UniformScale, 1.0f, UniformScale);
-	}
-	else if ( bIsOrthoYZ )
-	{
-		Scale = FVector(1.0f, UniformScale, UniformScale);
-	}
-	else
-	{
-		Scale = FVector(UniformScale, UniformScale, UniformScale);
-	}
+    FSpaceDescriptor Space(View, ViewportClient, InLocation);
 
 	// Draw the axis lines with arrow heads
-	if( DrawAxis&EAxisList::X && (bIsPerspective || bIsLocalSpace || !bIsOrthoYZ) )
+	if( Space.ShouldDrawAxisX(DrawAxis) )
 	{
 		UMaterialInstanceDynamic* XMaterial = ( CurrentAxis&EAxisList::X ? CurrentAxisMaterial : AxisMaterialX );
-		Render_Axis( View, PDI, EAxisList::X, WidgetMatrix, XMaterial, XColor, XAxisDir, Scale, bDrawWidget );
+		Render_Axis( View, PDI, EAxisList::X, WidgetMatrix, XMaterial, XColor, XAxisDir, Space.Scale, bDrawWidget );
 	}
 
-	if( DrawAxis&EAxisList::Y && (bIsPerspective || bIsLocalSpace || !bIsOrthoXZ) )
+	if( Space.ShouldDrawAxisY(DrawAxis) )
 	{
 		UMaterialInstanceDynamic* YMaterial = ( CurrentAxis&EAxisList::Y ? CurrentAxisMaterial : AxisMaterialY );
-		Render_Axis( View, PDI, EAxisList::Y, WidgetMatrix, YMaterial, YColor, YAxisDir, Scale, bDrawWidget );
+		Render_Axis( View, PDI, EAxisList::Y, WidgetMatrix, YMaterial, YColor, YAxisDir, Space.Scale, bDrawWidget );
 	}
 
-	if( DrawAxis&EAxisList::Z && (bIsPerspective || bIsLocalSpace || !bIsOrthoXY) )
+	if( Space.ShouldDrawAxisZ(DrawAxis) )
 	{
 		UMaterialInstanceDynamic* ZMaterial = ( CurrentAxis&EAxisList::Z ? CurrentAxisMaterial : AxisMaterialZ );
-		Render_Axis( View, PDI, EAxisList::Z, WidgetMatrix, ZMaterial, ZColor, ZAxisDir, Scale, bDrawWidget );
+		Render_Axis( View, PDI, EAxisList::Z, WidgetMatrix, ZMaterial, ZColor, ZAxisDir, Space.Scale, bDrawWidget );
 	}
 
 	// Draw the grabbers
 	if( bDrawWidget )
 	{
-		FVector CornerPos = FVector(7, 0, 7) * UniformScale;
-		FVector AxisSize = FVector(12, 1.2, 12) * UniformScale;
-		float CornerLength = 1.2f * UniformScale;
+		FVector CornerPos = FVector(7, 0, 7) * Space.UniformScale;
+		FVector AxisSize = FVector(12, 1.2, 12) * Space.UniformScale;
+		float CornerLength = 1.2f * Space.UniformScale;
 
 		// After the primitives have been scaled and transformed, we apply this inverse scale that flattens the dimension
 		// that was scaled up to prevent it from intersecting with the near plane.  In perspective this won't have any effect,
 		// but in the ortho viewports it will prevent scaling in the direction of the camera and thus intersecting the near plane.
-		FVector FlattenScale = FVector(Scale.Component(0) == 1.0f ? 1.0f / UniformScale : 1.0f, Scale.Component(1) == 1.0f ? 1.0f / UniformScale : 1.0f, Scale.Component(2) == 1.0f ? 1.0f / UniformScale : 1.0f);
+		FVector FlattenScale = FVector(Space.Scale.Component(0) == 1.0f ? 1.0f / Space.UniformScale : 1.0f, Space.Scale.Component(1) == 1.0f ? 1.0f / Space.UniformScale : 1.0f, Space.Scale.Component(2) == 1.0f ? 1.0f / Space.UniformScale : 1.0f);
 
-		if (bIsPerspective || bIsLocalSpace || bIsOrthoXY)
+		if (Space.bIsPerspective || Space.bIsLocalSpace || Space.bIsOrthoXY)
 		{
 			if( (DrawAxis&EAxisList::XY) == EAxisList::XY )							// Top
 			{
@@ -490,7 +584,7 @@ void FWidget::Render_Translate( const FSceneView* View, FPrimitiveDrawInterface*
 			}
 		}
 
-		if (bIsPerspective || bIsLocalSpace || bIsOrthoXZ)		// Front
+		if (Space.bIsPerspective || Space.bIsLocalSpace || Space.bIsOrthoXZ)		// Front
 		{
 			if( (DrawAxis&EAxisList::XZ) == EAxisList::XZ ) 
 			{
@@ -505,7 +599,7 @@ void FWidget::Render_Translate( const FSceneView* View, FPrimitiveDrawInterface*
 			}
 		}
 
-		if( bIsPerspective || bIsLocalSpace || bIsOrthoYZ )		// Side
+		if( Space.bIsPerspective || Space.bIsLocalSpace || Space.bIsOrthoYZ )		// Side
 		{
 			if( (DrawAxis&EAxisList::YZ) == EAxisList::YZ ) 
 			{
@@ -522,7 +616,7 @@ void FWidget::Render_Translate( const FSceneView* View, FPrimitiveDrawInterface*
 	}
 
 	// Draw screen-space movement handle (circle)
-	if( bDrawWidget && ( DrawAxis & EAxisList::Screen ) && bIsPerspective )
+	if( bDrawWidget && ( DrawAxis & EAxisList::Screen ) && Space.bIsPerspective )
 	{
 		PDI->SetHitProxy( new HWidgetAxis(EAxisList::Screen, bDisabled) );
 		const FVector CameraXAxis = View->ViewMatrices.ViewMatrix.GetColumn(0);
@@ -530,7 +624,7 @@ void FWidget::Render_Translate( const FSceneView* View, FPrimitiveDrawInterface*
 		const FVector CameraZAxis = View->ViewMatrices.ViewMatrix.GetColumn(2);
 
 		UMaterialInstanceDynamic* XYZMaterial = ( CurrentAxis&EAxisList::Screen) ? CurrentAxisMaterial : OpaquePlaneMaterialXY;
-		DrawSphere( PDI, InLocation, 4.0f * Scale, 10, 5, XYZMaterial->GetRenderProxy(false), SDPG_Foreground );
+		DrawSphere( PDI, InLocation, 4.0f * Space.Scale, 10, 5, XYZMaterial->GetRenderProxy(false), SDPG_Foreground );
 
 		PDI->SetHitProxy( NULL );
 	}
@@ -594,52 +688,25 @@ void FWidget::Render_Scale( const FSceneView* View,FPrimitiveDrawInterface* PDI,
 	UMaterialInstanceDynamic* ZMaterial = ( CurrentAxis&EAxisList::Z ? CurrentAxisMaterial : AxisMaterialZ );
 	UMaterialInstanceDynamic* XYZMaterial = ( CurrentAxis&EAxisList::XYZ ? CurrentAxisMaterial : OpaquePlaneMaterialXY );
 
-	// Figure out axis matrices
-
 	FMatrix WidgetMatrix = CustomCoordSystem * FTranslationMatrix( InLocation );
-	// Determine viewport
-
 	const EAxisList::Type DrawAxis = GetAxisToDraw( ViewportClient->GetWidgetMode() );
-	const bool bIsPerspective = ( View->ViewMatrices.ProjMatrix.M[3][3] < 1.0f );
-	const bool bIsOrthoXY = !bIsPerspective && FMath::Abs(View->ViewMatrices.ViewMatrix.M[2][2]) > 0.0f;
-	const bool bIsOrthoXZ = !bIsPerspective && FMath::Abs(View->ViewMatrices.ViewMatrix.M[1][2]) > 0.0f;
-	const bool bIsOrthoYZ = !bIsPerspective && FMath::Abs(View->ViewMatrices.ViewMatrix.M[0][2]) > 0.0f;
-
-	FVector Scale;
-	const float UniformScale = View->WorldToScreen(InLocation).W * (4.0f / View->UnscaledViewRect.Width() / View->ViewMatrices.ProjMatrix.M[0][0]);
-
-	if ( bIsOrthoXY )
-	{
-		Scale = FVector(UniformScale, UniformScale, 1.0f);
-	}
-	else if ( bIsOrthoXZ )
-	{
-		Scale = FVector(UniformScale, 1.0f, UniformScale);
-	}
-	else if ( bIsOrthoYZ )
-	{
-		Scale = FVector(1.0f, UniformScale, UniformScale);
-	}
-	else
-	{
-		Scale = FVector(UniformScale, UniformScale, UniformScale);
-	}
+    FSpaceDescriptor Space(View, ViewportClient, InLocation);
 
 	// Draw the axis lines with cube heads	
-	if( !bIsOrthoYZ && DrawAxis&EAxisList::X )
-	{
-		Render_Axis( View, PDI, EAxisList::X, WidgetMatrix, XMaterial, XColor, XAxisDir, Scale, bDrawWidget, true );
-	}
+    if (Space.ShouldDrawAxisX(DrawAxis))
+    {
+        Render_Axis(View, PDI, EAxisList::X, WidgetMatrix, XMaterial, XColor, XAxisDir, Space.Scale, bDrawWidget, true);
+    }
 
-	if( !bIsOrthoXZ && DrawAxis&EAxisList::Y )
-	{
-		Render_Axis( View, PDI, EAxisList::Y, WidgetMatrix, YMaterial, YColor, YAxisDir, Scale, bDrawWidget, true );
-	}
+    if (Space.ShouldDrawAxisY(DrawAxis))
+    {
+        Render_Axis(View, PDI, EAxisList::Y, WidgetMatrix, YMaterial, YColor, YAxisDir, Space.Scale, bDrawWidget, true);
+    }
 
-	if( !bIsOrthoXY &&  DrawAxis&EAxisList::Z )
-	{
-		Render_Axis( View, PDI, EAxisList::Z, WidgetMatrix, ZMaterial, ZColor, ZAxisDir, Scale, bDrawWidget, true );
-	}
+    if (Space.ShouldDrawAxisZ(DrawAxis))
+    {
+        Render_Axis(View, PDI, EAxisList::Z, WidgetMatrix, ZMaterial, ZColor, ZAxisDir, Space.Scale, bDrawWidget, true);
+    }
 
 	// Draw grabber handles and center cube
 	if ( bDrawWidget )
@@ -647,32 +714,32 @@ void FWidget::Render_Scale( const FSceneView* View,FPrimitiveDrawInterface* PDI,
 		const bool bDisabled = IsWidgetDisabled();
 
 		// Grabber handles
-		if( !bIsOrthoYZ && !bIsOrthoXZ && ((DrawAxis&(EAxisList::X|EAxisList::Y)) == (EAxisList::X|EAxisList::Y)) )
+		if( !Space.bIsOrthoYZ && !Space.bIsOrthoXZ && ((DrawAxis&(EAxisList::X|EAxisList::Y)) == (EAxisList::X|EAxisList::Y)) )
 		{
 			PDI->SetHitProxy( new HWidgetAxis(EAxisList::XY, bDisabled) );
 			{
-				PDI->DrawLine( WidgetMatrix.TransformPosition(FVector(24,0,0) * Scale), WidgetMatrix.TransformPosition(FVector(12,12,0) * Scale), XColor, SDPG_Foreground );
-				PDI->DrawLine( WidgetMatrix.TransformPosition(FVector(12,12,0) * Scale), WidgetMatrix.TransformPosition(FVector(0,24,0) * Scale), YColor, SDPG_Foreground );
+				PDI->DrawLine( WidgetMatrix.TransformPosition(FVector(24,0,0) * Space.Scale), WidgetMatrix.TransformPosition(FVector(12,12,0) * Space.Scale), XColor, SDPG_Foreground );
+				PDI->DrawLine( WidgetMatrix.TransformPosition(FVector(12,12,0) * Space.Scale), WidgetMatrix.TransformPosition(FVector(0,24,0) * Space.Scale), YColor, SDPG_Foreground );
 			}
 			PDI->SetHitProxy( NULL );
 		}
 
-		if( !bIsOrthoYZ && !bIsOrthoXY && ((DrawAxis&(EAxisList::X|EAxisList::Z)) == (EAxisList::X|EAxisList::Z)) )
+		if( !Space.bIsOrthoYZ && !Space.bIsOrthoXY && ((DrawAxis&(EAxisList::X|EAxisList::Z)) == (EAxisList::X|EAxisList::Z)) )
 		{
 			PDI->SetHitProxy( new HWidgetAxis(EAxisList::XZ, bDisabled) );
 			{
-				PDI->DrawLine( WidgetMatrix.TransformPosition(FVector(24,0,0) * Scale), WidgetMatrix.TransformPosition(FVector(12,0,12) * Scale), XColor, SDPG_Foreground );
-				PDI->DrawLine( WidgetMatrix.TransformPosition(FVector(12,0,12) * Scale), WidgetMatrix.TransformPosition(FVector(0,0,24) * Scale), ZColor, SDPG_Foreground );
+				PDI->DrawLine( WidgetMatrix.TransformPosition(FVector(24,0,0) * Space.Scale), WidgetMatrix.TransformPosition(FVector(12,0,12) * Space.Scale), XColor, SDPG_Foreground );
+				PDI->DrawLine( WidgetMatrix.TransformPosition(FVector(12,0,12) * Space.Scale), WidgetMatrix.TransformPosition(FVector(0,0,24) * Space.Scale), ZColor, SDPG_Foreground );
 			}
 			PDI->SetHitProxy( NULL );
 		}
 
-		if( !bIsOrthoXY && !bIsOrthoXZ && ((DrawAxis&(EAxisList::Y|EAxisList::Z)) == (EAxisList::Y|EAxisList::Z)) )
+		if( !Space.bIsOrthoXY && !Space.bIsOrthoXZ && ((DrawAxis&(EAxisList::Y|EAxisList::Z)) == (EAxisList::Y|EAxisList::Z)) )
 		{
 			PDI->SetHitProxy( new HWidgetAxis(EAxisList::YZ, bDisabled) );
 			{
-				PDI->DrawLine( WidgetMatrix.TransformPosition(FVector(0,24,0) * Scale), WidgetMatrix.TransformPosition(FVector(0,12,12) * Scale), YColor, SDPG_Foreground );
-				PDI->DrawLine( WidgetMatrix.TransformPosition(FVector(0,12,12) * Scale), WidgetMatrix.TransformPosition(FVector(0,0,24) * Scale), ZColor, SDPG_Foreground );
+				PDI->DrawLine( WidgetMatrix.TransformPosition(FVector(0,24,0) * Space.Scale), WidgetMatrix.TransformPosition(FVector(0,12,12) * Space.Scale), YColor, SDPG_Foreground );
+				PDI->DrawLine( WidgetMatrix.TransformPosition(FVector(0,12,12) * Space.Scale), WidgetMatrix.TransformPosition(FVector(0,0,24) * Space.Scale), ZColor, SDPG_Foreground );
 			}
 			PDI->SetHitProxy( NULL );
 		}
@@ -682,7 +749,7 @@ void FWidget::Render_Scale( const FSceneView* View,FPrimitiveDrawInterface* PDI,
 		{
 			PDI->SetHitProxy( new HWidgetAxis(EAxisList::XYZ, bDisabled) );
 
-			Render_Cube(PDI, WidgetMatrix, XYZMaterial, Scale * 4 );
+			Render_Cube(PDI, WidgetMatrix, XYZMaterial, Space.Scale * 4 );
 
 			PDI->SetHitProxy( NULL );
 		}
@@ -711,63 +778,35 @@ void FWidget::Render_TranslateRotateZ( const FSceneView* View, FPrimitiveDrawInt
 
 	// Figure out axis matrices
 	FMatrix AxisMatrix = CustomCoordSystem * FTranslationMatrix( InLocation );
-
-	bool bIsPerspective = ( View->ViewMatrices.ProjMatrix.M[3][3] < 1.0f );
-	const bool bIsOrthoXY = !bIsPerspective && FMath::Abs(View->ViewMatrices.ViewMatrix.M[2][2]) > 0.0f;
-	const bool bIsOrthoXZ = !bIsPerspective && FMath::Abs(View->ViewMatrices.ViewMatrix.M[1][2]) > 0.0f;
-	const bool bIsOrthoYZ = !bIsPerspective && FMath::Abs(View->ViewMatrices.ViewMatrix.M[0][2]) > 0.0f;
-
-	// For local space widgets, we always want to draw all three axis, since they may not be aligned with
-	// the orthographic projection anyway.
-	bool bIsLocalSpace = ( ViewportClient->GetWidgetCoordSystemSpace() == COORD_Local );
-
 	EAxisList::Type DrawAxis = GetAxisToDraw( ViewportClient->GetWidgetMode() );
 
-	FVector Scale;
-	float UniformScale = View->WorldToScreen(InLocation).W * ( 4.0f / View->UnscaledViewRect.Width() / View->ViewMatrices.ProjMatrix.M[0][0] );
-
-	if ( bIsOrthoXY )
-	{
-		Scale = FVector(UniformScale, UniformScale, 1.0f);
-	}
-	else if ( bIsOrthoXZ )
-	{
-		Scale = FVector(UniformScale, 1.0f, UniformScale);
-	}
-	else if ( bIsOrthoYZ )
-	{
-		Scale = FVector(1.0f, UniformScale, UniformScale);
-	}
-	else
-	{
-		Scale = FVector(UniformScale, UniformScale, UniformScale);
-	}
+	FSpaceDescriptor Space(View, ViewportClient, InLocation);
 
 	// Draw the grabbers
 	if( bDrawWidget )
 	{
 		// Draw the axis lines with arrow heads
-		if( DrawAxis&EAxisList::X && (bIsPerspective || bIsLocalSpace || View->ViewMatrices.ViewMatrix.M[0][2] != -1.f) )
+		if( DrawAxis&EAxisList::X && (Space.bIsPerspective || Space.bIsLocalSpace || View->ViewMatrices.ViewMatrix.M[0][2] != -1.f) )
 		{
-			Render_Axis( View, PDI, EAxisList::X, AxisMatrix, XMaterial, XColor, XAxisDir, Scale, bDrawWidget );
+			Render_Axis( View, PDI, EAxisList::X, AxisMatrix, XMaterial, XColor, XAxisDir, Space.Scale, bDrawWidget );
 		}
 
-		if( DrawAxis&EAxisList::Y && (bIsPerspective || bIsLocalSpace || View->ViewMatrices.ViewMatrix.M[1][2] != -1.f) )
+		if( DrawAxis&EAxisList::Y && (Space.bIsPerspective || Space.bIsLocalSpace || View->ViewMatrices.ViewMatrix.M[1][2] != -1.f) )
 		{
-			Render_Axis( View, PDI, EAxisList::Y, AxisMatrix, YMaterial, YColor, YAxisDir, Scale, bDrawWidget );
+			Render_Axis( View, PDI, EAxisList::Y, AxisMatrix, YMaterial, YColor, YAxisDir, Space.Scale, bDrawWidget );
 		}
 
-		if( DrawAxis&EAxisList::Z && (bIsPerspective || bIsLocalSpace || View->ViewMatrices.ViewMatrix.M[0][1] != 1.f) )
+		if( DrawAxis&EAxisList::Z && (Space.bIsPerspective || Space.bIsLocalSpace || View->ViewMatrices.ViewMatrix.M[0][1] != 1.f) )
 		{
-			Render_Axis( View, PDI, EAxisList::Z, AxisMatrix, ZMaterial, ZColor, ZAxisDir, Scale, bDrawWidget );
+			Render_Axis( View, PDI, EAxisList::Z, AxisMatrix, ZMaterial, ZColor, ZAxisDir, Space.Scale, bDrawWidget );
 		}
 
 		const bool bDisabled = IsWidgetDisabled();
 
-		const float ScaledRadius = (TRANSLATE_ROTATE_AXIS_CIRCLE_RADIUS * UniformScale) + GetDefault<ULevelEditorViewportSettings>()->TransformWidgetSizeAdjustment;
+		const float ScaledRadius = (TRANSLATE_ROTATE_AXIS_CIRCLE_RADIUS * Space.UniformScale) + GetDefault<ULevelEditorViewportSettings>()->TransformWidgetSizeAdjustment;
 
 		//ZRotation
-		if( DrawAxis&EAxisList::ZRotation && (bIsPerspective || bIsLocalSpace || View->ViewMatrices.ViewMatrix.M[0][2] != -1.f) )
+		if( DrawAxis&EAxisList::ZRotation && (Space.bIsPerspective || Space.bIsLocalSpace || View->ViewMatrices.ViewMatrix.M[0][2] != -1.f) )
 		{
 			PDI->SetHitProxy( new HWidgetAxis(EAxisList::ZRotation, bDisabled) );
 			{
@@ -780,7 +819,7 @@ void FWidget::Render_TranslateRotateZ( const FSceneView* View, FPrimitiveDrawInt
 		}
 
 		//XY Plane
-		if( bIsPerspective || bIsLocalSpace || View->ViewMatrices.ViewMatrix.M[0][1] != 1.f )
+		if( Space.bIsPerspective || Space.bIsLocalSpace || View->ViewMatrices.ViewMatrix.M[0][1] != 1.f )
 		{
 			if( (DrawAxis & EAxisList::XY) == EAxisList::XY ) 
 			{
@@ -821,29 +860,22 @@ void FWidget::Render_2D(const FSceneView* View, FPrimitiveDrawInterface* PDI, FE
 	// Figure out axis matrices
 	FMatrix WidgetMatrix = CustomCoordSystem * FTranslationMatrix(InLocation);
 
-	bool bIsPerspective = (View->ViewMatrices.ProjMatrix.M[3][3] < 1.0f);
-	const bool bIsOrthoXY = !bIsPerspective && FMath::Abs(View->ViewMatrices.ViewMatrix.M[2][2]) > 0.0f;
-	const bool bIsOrthoXZ = !bIsPerspective && FMath::Abs(View->ViewMatrices.ViewMatrix.M[1][2]) > 0.0f;
-	const bool bIsOrthoYZ = !bIsPerspective && FMath::Abs(View->ViewMatrices.ViewMatrix.M[0][2]) > 0.0f;
-
-	// For local space widgets, we always want to draw all three axis, since they may not be aligned with
-	// the orthographic projection anyway.
-	const bool bIsLocalSpace = (ViewportClient->GetWidgetCoordSystemSpace() == COORD_Local);
+    FSpaceDescriptor Space(View, ViewportClient, InLocation);
 
 	EAxisList::Type DrawAxis = EAxisList::None;
-	if (bIsOrthoXY)
+	if (Space.bIsOrthoXY)
 	{
 		DrawAxis = EAxisList::X;
 	}
-	else if (bIsOrthoXZ)
+	else if (Space.bIsOrthoXZ)
 	{
 		DrawAxis = EAxisList::XZ;
 	}
-	else if (bIsOrthoYZ)
+	else if (Space.bIsOrthoYZ)
 	{
 		DrawAxis = EAxisList::Z;
 	}
-	else if (bIsPerspective)
+	else if (Space.bIsPerspective)
 	{
 		// Find the best plane to move on
 		const FVector CameraZAxis = View->ViewMatrices.ViewMatrix.GetColumn(2);
@@ -874,28 +906,8 @@ void FWidget::Render_2D(const FSceneView* View, FPrimitiveDrawInterface* PDI, FE
 
 	const bool bDisabled = IsWidgetDisabled();
 
-	FVector Scale;
-	float UniformScale = View->WorldToScreen(InLocation).W * (4.0f / View->UnscaledViewRect.Width() / View->ViewMatrices.ProjMatrix.M[0][0]);
-
-	if (bIsOrthoXY)
-	{
-		Scale = FVector(UniformScale, UniformScale, 1.0f);
-	}
-	else if (bIsOrthoXZ)
-	{
-		Scale = FVector(UniformScale, 1.0f, UniformScale);
-	}
-	else if (bIsOrthoYZ)
-	{
-		Scale = FVector(1.0f, UniformScale, UniformScale);
-	}
-	else
-	{
-		Scale = FVector(UniformScale, UniformScale, UniformScale);
-	}
-
 	// Radius
-	const float ScaledRadius = (TWOD_AXIS_CIRCLE_RADIUS * UniformScale) + GetDefault<ULevelEditorViewportSettings>()->TransformWidgetSizeAdjustment;
+	const float ScaledRadius = (TWOD_AXIS_CIRCLE_RADIUS * Space.UniformScale) + GetDefault<ULevelEditorViewportSettings>()->TransformWidgetSizeAdjustment;
 	const int32 CircleSides = (GetDefault<ULevelEditorViewportSettings>()->TransformWidgetSizeAdjustment > 0)
 		? AXIS_CIRCLE_SIDES + (GetDefault<ULevelEditorViewportSettings>()->TransformWidgetSizeAdjustment / 5)
 		: AXIS_CIRCLE_SIDES;
@@ -903,40 +915,40 @@ void FWidget::Render_2D(const FSceneView* View, FPrimitiveDrawInterface* PDI, FE
 	// Draw the grabbers
 	if (bDrawWidget)
 	{
-		FVector CornerPos = FVector(7, 0, 7) * UniformScale;
-		FVector AxisSize = FVector(12, 1.2, 12) * UniformScale;
-		float CornerLength = 1.2f * UniformScale;
+		FVector CornerPos = FVector(7, 0, 7) * Space.UniformScale;
+		FVector AxisSize = FVector(12, 1.2, 12) * Space.UniformScale;
+		float CornerLength = 1.2f * Space.UniformScale;
 
 		// Draw the axis lines with arrow heads
-		if (DrawAxis&EAxisList::X && (bIsPerspective || bIsLocalSpace || !bIsOrthoYZ))
+		if (Space.ShouldDrawAxisX(DrawAxis))
 		{
 			UMaterialInstanceDynamic* XMaterial = (CurrentAxis&EAxisList::X ? CurrentAxisMaterial : AxisMaterialX);
-			Render_Axis(View, PDI, EAxisList::X, WidgetMatrix, XMaterial, XColor, XAxisDir, Scale, bDrawWidget);
+			Render_Axis(View, PDI, EAxisList::X, WidgetMatrix, XMaterial, XColor, XAxisDir, Space.Scale, bDrawWidget);
 		}
 
-		if (DrawAxis&EAxisList::Y && (bIsPerspective || bIsLocalSpace || !bIsOrthoXZ))
+		if (Space.ShouldDrawAxisY(DrawAxis))
 		{
 			UMaterialInstanceDynamic* YMaterial = (CurrentAxis&EAxisList::Y ? CurrentAxisMaterial : AxisMaterialY);
-			Render_Axis(View, PDI, EAxisList::Y, WidgetMatrix, YMaterial, YColor, YAxisDir, Scale, bDrawWidget);
+			Render_Axis(View, PDI, EAxisList::Y, WidgetMatrix, YMaterial, YColor, YAxisDir, Space.Scale, bDrawWidget);
 		}
 
-		if (DrawAxis&EAxisList::Z && (bIsPerspective || bIsLocalSpace || !bIsOrthoXY))
+		if (Space.ShouldDrawAxisZ(DrawAxis))
 		{
 			UMaterialInstanceDynamic* ZMaterial = (CurrentAxis&EAxisList::Z ? CurrentAxisMaterial : AxisMaterialZ);
-			Render_Axis(View, PDI, EAxisList::Z, WidgetMatrix, ZMaterial, ZColor, ZAxisDir, Scale, bDrawWidget);
+			Render_Axis(View, PDI, EAxisList::Z, WidgetMatrix, ZMaterial, ZColor, ZAxisDir, Space.Scale, bDrawWidget);
 		}
 
 		// After the primitives have been scaled and transformed, we apply this inverse scale that flattens the dimension
 		// that was scaled up to prevent it from intersecting with the near plane.  In perspective this won't have any effect,
 		// but in the ortho viewports it will prevent scaling in the direction of the camera and thus intersecting the near plane.
-		FVector FlattenScale = FVector(Scale.Component(0) == 1.0f ? 1.0f / UniformScale : 1.0f, Scale.Component(1) == 1.0f ? 1.0f / UniformScale : 1.0f, Scale.Component(2) == 1.0f ? 1.0f / UniformScale : 1.0f);
+		FVector FlattenScale = FVector(Space.Scale.Component(0) == 1.0f ? 1.0f / Space.UniformScale : 1.0f, Space.Scale.Component(1) == 1.0f ? 1.0f / Space.UniformScale : 1.0f, Space.Scale.Component(2) == 1.0f ? 1.0f / Space.UniformScale : 1.0f);
 		float ArrowRadius = ScaledRadius * 2.0f;
 		float ArrowStartRadius = ScaledRadius * 1.3f;
 
 		uint8 HoverAlpha = 0xff;
 		uint8 NormalAlpha = 0x2f;
 
-		if (((DrawAxis&EAxisList::XZ) == EAxisList::XZ) && (bIsPerspective || bIsLocalSpace || bIsOrthoXZ)) // Front
+		if (((DrawAxis&EAxisList::XZ) == EAxisList::XZ) && (Space.bIsPerspective || Space.bIsLocalSpace || Space.bIsOrthoXZ)) // Front
 		{
 			uint8 Alpha = ((CurrentAxis&EAxisList::XZ) == EAxisList::XZ ? HoverAlpha : NormalAlpha);
 			PDI->SetHitProxy(new HWidgetAxis(EAxisList::XZ, bDisabled));
diff --git a/Engine/Source/Editor/UnrealEd/Private/VertexSnapping.cpp b/Engine/Source/Editor/UnrealEd/Private/VertexSnapping.cpp
--- a/Engine/Source/Editor/UnrealEd/Private/VertexSnapping.cpp
+++ b/Engine/Source/Editor/UnrealEd/Private/VertexSnapping.cpp
@@ -173,8 +173,7 @@ public:
 		: ComponentToWorldIT( InSkinnedMeshComp->ComponentToWorld.ToInverseMatrixWithScale().GetTransposed() )
 		, SkinnedMeshComponent( InSkinnedMeshComp )
 		, LODModel( InSkinnedMeshComp->GetSkeletalMeshResource()->LODModels[0] )
-		, CurrentChunkIndex( 0 )
-		, RigidVertexIndex( 0 )
+		, CurrentSectionIndex( 0 )
 		, SoftVertexIndex( 0 )
 	{
 		
@@ -183,51 +182,30 @@ public:
 	/** FVertexIterator interface */
 	virtual FVector Position() override
 	{
-		const FSkelMeshChunk& Chunk = LODModel.Chunks[CurrentChunkIndex];
-
-		if( Chunk.RigidVertices.IsValidIndex( RigidVertexIndex ) )
-		{
-			return SkinnedMeshComponent->ComponentToWorld.TransformPosition( Chunk.RigidVertices[RigidVertexIndex].Position );
-		}
-		else
-		{
-			return SkinnedMeshComponent->ComponentToWorld.TransformPosition( Chunk.SoftVertices[SoftVertexIndex].Position );
-		}
+		const FSkelMeshSection& Section = LODModel.Sections[CurrentSectionIndex];
+		return SkinnedMeshComponent->ComponentToWorld.TransformPosition(Section.SoftVertices[SoftVertexIndex].Position );
 	}
 
 	virtual FVector Normal() override
 	{
-		const FSkelMeshChunk& Chunk = LODModel.Chunks[CurrentChunkIndex];
-
-		if( Chunk.RigidVertices.IsValidIndex( RigidVertexIndex ) )
-		{
-			return ComponentToWorldIT.TransformVector( Chunk.RigidVertices[RigidVertexIndex].TangentZ );
-		}
-		else
-		{
-			return ComponentToWorldIT.TransformVector( Chunk.SoftVertices[SoftVertexIndex].TangentZ );
-		}
+		const FSkelMeshSection& Section = LODModel.Sections[CurrentSectionIndex];
+		return ComponentToWorldIT.TransformVector(Section.SoftVertices[SoftVertexIndex].TangentZ );
 	}
 
 protected:
 	virtual void Advance() override
 	{
 		// First advance the rigid vertex in the current chunk
-		const FSkelMeshChunk& Chunk = LODModel.Chunks[CurrentChunkIndex];
+		const FSkelMeshSection& Section = LODModel.Sections[CurrentSectionIndex];
 
-		if( Chunk.RigidVertices.IsValidIndex( RigidVertexIndex + 1 ) )
-		{
-			++RigidVertexIndex;
-		}
-		else if( Chunk.SoftVertices.IsValidIndex( SoftVertexIndex + 1 ) )
+		if(Section.SoftVertices.IsValidIndex( SoftVertexIndex + 1 ) )
 		{
 			++SoftVertexIndex;
 		}
 		else
 		{
-			// out of rigid and soft verts in this chunk.  Advance chunks
-			++CurrentChunkIndex;
-			RigidVertexIndex = 0;
+			// out of soft verts in this section.  Advance sections
+			++CurrentSectionIndex;
 			SoftVertexIndex = 0;
 		}
 	}
@@ -235,10 +213,10 @@ protected:
 	virtual bool HasMoreVertices() override
 	{
 		bool bHasMoreVerts = false;
-		if( LODModel.Chunks.IsValidIndex(CurrentChunkIndex) )
+		if (LODModel.Sections.IsValidIndex(CurrentSectionIndex))
 		{
-			const FSkelMeshChunk& Chunk = LODModel.Chunks[CurrentChunkIndex];
-			bHasMoreVerts = Chunk.RigidVertices.IsValidIndex( RigidVertexIndex ) || Chunk.SoftVertices.IsValidIndex( SoftVertexIndex );
+			const FSkelMeshSection& Section = LODModel.Sections[CurrentSectionIndex];
+			bHasMoreVerts = Section.SoftVertices.IsValidIndex( SoftVertexIndex );
 		}
 	
 		return bHasMoreVerts;
@@ -251,10 +229,8 @@ private:
 	USkinnedMeshComponent* SkinnedMeshComponent;
 	/** Skeletal mesh render data */
 	FStaticLODModel& LODModel;
-	/** Current chunk the iterator is on */
-	uint32 CurrentChunkIndex;
-	/** Current rigid vertex index the iterator is on */
-	uint32 RigidVertexIndex;
+	/** Current section the iterator is on */
+	uint32 CurrentSectionIndex;
 	/** Current Soft vertex index the iterator is on */
 	uint32 SoftVertexIndex; 
 };
diff --git a/Engine/Source/Editor/UnrealEd/Public/AnimationCompressionPanel.h b/Engine/Source/Editor/UnrealEd/Public/AnimationCompressionPanel.h
--- a/Engine/Source/Editor/UnrealEd/Public/AnimationCompressionPanel.h
+++ b/Engine/Source/Editor/UnrealEd/Public/AnimationCompressionPanel.h
@@ -2,6 +2,17 @@
 
 #pragma once
 
+#include "AnimationCompressionPanel.generated.h"
+
+UCLASS()
+class UCompressionHolder : public UObject
+{
+	GENERATED_UCLASS_BODY()
+
+	UPROPERTY(Instanced, Category = Compression, EditAnywhere)
+	class UAnimCompress* Compression;
+};
+
 /** 
 * FDlgAnimCompression
 * 
@@ -65,57 +76,11 @@ private:
 	 */
 	void					ApplyAlgorithm(class UAnimCompress* Algorithm);
 
-	/**
-	 * Creates a Slate radio button 
-	 *
-	 * @param	RadioText	Text label for the radio button
-	 * @param	ButtonId	The ID for the radio button
-	 * @return				The created radio button widget
-	 */
-	TSharedRef<SWidget> CreateRadioButton( const FText& RadioText, int32 ButtonId )
-	{
-		return
-			SNew(SCheckBox)
-			.Style(FEditorStyle::Get(), "RadioButton")
-			.IsChecked( this, &SAnimationCompressionPanel::IsRadioChecked, ButtonId )
-			.OnCheckStateChanged( this, &SAnimationCompressionPanel::OnRadioChanged, ButtonId )
-			[
-				SNew(STextBlock).Text(RadioText)
-			];
-	}
-
-	/**
-	 * Returns the state of the radio button
-	 *
-	 * @param	ButtonId	The ID for the radio button
-	 * @return				The status of the radio button
-	 */
-	ECheckBoxState IsRadioChecked( int32 ButtonId ) const
-	{
-		return (CurrentCompressionChoice == ButtonId)
-			? ECheckBoxState::Checked
-			: ECheckBoxState::Unchecked;
-	}
-
-	/**
-	 * Handler for all radio button clicks
-	 *
-	 * @param	NewRadioState	The new state of the radio button
-	 * @param	RadioThatChanged	The ID of the radio button that has changed. 
-	 */
-	void OnRadioChanged( ECheckBoxState NewRadioState, int32 RadioThatChanged )
-	{
-		if (NewRadioState == ECheckBoxState::Checked)
-		{
-			CurrentCompressionChoice = RadioThatChanged;
-		}
-	}
-
-	TArray<class UAnimCompress*>			AnimationCompressionAlgorithms;
 	TArray< TWeakObjectPtr<UAnimSequence> >	AnimSequences;
 
-	int32 CurrentCompressionChoice;
-
 	/** Pointer to the window which holds this Widget, required for modal control */
-	TSharedPtr<SWindow> ParentWindow;
+	TWeakPtr<SWindow> ParentWindow;
+
+	/** Our compression display object */
+	UCompressionHolder* CompressionHolder;
 };
\ No newline at end of file
diff --git a/Engine/Source/Editor/UnrealEd/Public/ApexClothingUtils.h b/Engine/Source/Editor/UnrealEd/Public/ApexClothingUtils.h
--- a/Engine/Source/Editor/UnrealEd/Public/ApexClothingUtils.h
+++ b/Engine/Source/Editor/UnrealEd/Public/ApexClothingUtils.h
@@ -45,7 +45,6 @@ struct FClothingBackup
 	bool						bBackedUp;
 	TArray<FClothingAssetData>	ClothingAssets;
 	TArray<FSkelMeshSection>	Sections;
-	TArray<FSkelMeshChunk>		Chunks;
 	TArray<uint32>				IndexBuffer;
 
 	void	Clear()
@@ -53,7 +52,6 @@ struct FClothingBackup
 		bBackedUp = false;
 		ClothingAssets.Empty();
 		Sections.Empty();
-		Chunks.Empty();
 	}
 };
 #endif // #if WITH_APEX_CLOTHING
diff --git a/Engine/Source/Editor/UnrealEd/Public/SSkeletonWidget.h b/Engine/Source/Editor/UnrealEd/Public/SSkeletonWidget.h
--- a/Engine/Source/Editor/UnrealEd/Public/SSkeletonWidget.h
+++ b/Engine/Source/Editor/UnrealEd/Public/SSkeletonWidget.h
@@ -566,7 +566,6 @@ private:
 
 	/** Best guess functions to try and match asset names */
 	const FAssetData* FindBestGuessMatch(const FAssetData& AssetName, const TArray<FAssetData>& PossibleAssets) const;
-	int32 GetStringDistance(const FString& First, const FString& Second) const;
 
 	/** The retargetting context we're managing*/
 	FAnimationRetargetContext* RetargetContext;
diff --git a/Engine/Source/Editor/UnrealEd/Public/ThumbnailHelpers.h b/Engine/Source/Editor/UnrealEd/Public/ThumbnailHelpers.h
--- a/Engine/Source/Editor/UnrealEd/Public/ThumbnailHelpers.h
+++ b/Engine/Source/Editor/UnrealEd/Public/ThumbnailHelpers.h
@@ -181,7 +181,7 @@ protected:
 
 private:
 	/** The skeletal mesh actor used to display all animation thumbnails */
-	class ASkeletalMeshActor* PreviewActor;
+	class AAnimationThumbnailSkeletalMeshActor* PreviewActor;
 
 	/** Animation we are generating the thumbnail for */
 	class UBlendSpaceBase* PreviewAnimation;
diff --git a/Engine/Source/Editor/UnrealEd/Public/UnrealEd.h b/Engine/Source/Editor/UnrealEd/Public/UnrealEd.h
--- a/Engine/Source/Editor/UnrealEd/Public/UnrealEd.h
+++ b/Engine/Source/Editor/UnrealEd/Public/UnrealEd.h
@@ -118,6 +118,7 @@
 #include "Classes/Factories/AnimCompositeFactory.h"
 #include "Classes/Factories/AnimMontageFactory.h"
 #include "Classes/Factories/AnimSequenceFactory.h"
+#include "Classes/Factories/PoseAssetFactory.h"
 #include "Classes/Factories/BlendSpaceFactory1D.h"
 #include "Classes/Factories/AimOffsetBlendSpaceFactory1D.h"
 #include "Classes/Factories/BlendSpaceFactoryNew.h"
diff --git a/Engine/Source/Programs/AutomationTool/GUBP/LegacyBranchSetup.cs b/Engine/Source/Programs/AutomationTool/GUBP/LegacyBranchSetup.cs
--- a/Engine/Source/Programs/AutomationTool/GUBP/LegacyBranchSetup.cs
+++ b/Engine/Source/Programs/AutomationTool/GUBP/LegacyBranchSetup.cs
@@ -1144,6 +1144,32 @@ partial class GUBP
 			BranchConfig.AddNode(new BuildEngineLocalizationNode(BranchOptions.EngineLocalizationBranchSuffix));
 		}
 
+		if(BranchConfig.BranchName != null && BranchConfig.BranchName.StartsWith("//UE4/Dev-", StringComparison.InvariantCultureIgnoreCase))
+		{
+			List<BranchInfo.BranchUProject> CodeProjects = new List<BranchInfo.BranchUProject>();
+			CodeProjects.Add(BranchConfig.Branch.BaseEngineProject);
+			CodeProjects.AddRange(BranchConfig.Branch.CodeProjects);
+
+			foreach(BranchInfo.BranchUProject CodeProject in CodeProjects)
+			{
+				SingleTargetProperties Properties;
+				if(CodeProject.Properties.Targets.TryGetValue(TargetRules.TargetType.Game, out Properties))
+				{
+					foreach(UnrealTargetPlatform HostPlatform in HostPlatforms)
+					{
+						HashSet<UnrealTargetPlatform> TargetPlatforms = new HashSet<UnrealTargetPlatform>();
+						TargetPlatforms.UnionWith(Properties.Rules.GUBP_GetPlatforms_MonolithicOnly(HostPlatform));
+						TargetPlatforms.UnionWith(Properties.Rules.GUBP_GetBuildOnlyPlatforms_MonolithicOnly(HostPlatform));
+						foreach(UnrealTargetPlatform TargetPlatform in TargetPlatforms)
+						{
+							GUBP.SingleTargetNode TargetNode = new SingleTargetNode(BranchConfig, HostPlatform, CodeProject, Properties.TargetName, TargetPlatform, UnrealTargetConfiguration.Development);
+							BranchConfig.AddNode(TargetNode);
+						}
+					}
+				}
+			}
+		}
+
         BranchConfig.AddNode(new WaitForTestShared(this));
 		List<UnrealTargetPlatform> HostPlatformsToWaitFor = BranchConfig.HostPlatforms;
 		HostPlatformsToWaitFor.RemoveAll(HostPlatform => BranchOptions.ExcludePlatformsForEditor.Contains(HostPlatform));
diff --git a/Engine/Source/Programs/AutomationTool/GUBP/LegacyNodes.cs b/Engine/Source/Programs/AutomationTool/GUBP/LegacyNodes.cs
--- a/Engine/Source/Programs/AutomationTool/GUBP/LegacyNodes.cs
+++ b/Engine/Source/Programs/AutomationTool/GUBP/LegacyNodes.cs
@@ -1435,6 +1435,72 @@ partial class GUBP
 		}	
 	}
 
+    public class SingleTargetNode : CompileNode
+    {
+        BranchInfo.BranchUProject GameProj;
+		string TargetName;
+        UnrealTargetPlatform TargetPlatform;
+		UnrealTargetConfiguration TargetConfiguration;
+
+		public SingleTargetNode(GUBPBranchConfig InBranchConfig, UnrealTargetPlatform InHostPlatform, BranchInfo.BranchUProject InGameProj, string InTargetName, UnrealTargetPlatform InTargetPlatform, UnrealTargetConfiguration InTargetConfiguration)
+            : base(InBranchConfig, InHostPlatform)
+        {
+            GameProj = InGameProj;
+			TargetName = InTargetName;
+            TargetPlatform = InTargetPlatform;
+			TargetConfiguration = InTargetConfiguration;
+
+			AddPseudodependency(RootEditorNode.StaticGetFullName(InHostPlatform));
+
+            if (TargetPlatform == UnrealTargetPlatform.PS4 || TargetPlatform == UnrealTargetPlatform.XboxOne)
+            {
+				// Required for PS4MapFileUtil/XboxOnePDBFileUtil
+				AddDependency(ToolsNode.StaticGetFullName(InHostPlatform));
+			}
+
+			if(!InBranchConfig.BranchOptions.bTargetPlatformsInParallel)
+			{
+				AgentSharingGroup = "TargetPlatforms" + StaticGetHostPlatformSuffix(InHostPlatform);
+			}
+        }
+
+        public static string StaticGetFullName(UnrealTargetPlatform InHostPlatform, BranchInfo.BranchUProject InGameProj, UnrealTargetPlatform InTargetPlatform, UnrealTargetConfiguration InTargetConfiguration)
+        {
+			return InGameProj.GameName + "_" + InTargetPlatform + "_" + InTargetConfiguration.ToString() + StaticGetHostPlatformSuffix(InHostPlatform);
+        }
+
+        public override string GetFullName()
+        {
+            return StaticGetFullName(HostPlatform, GameProj, TargetPlatform, TargetConfiguration);
+        }
+
+        public override bool DeleteBuildProducts()
+        {
+            return true;
+        }
+
+		public override int CISFrequencyQuantumShift(GUBP.GUBPBranchConfig BranchConfig)
+        {
+			return 15;
+        }
+
+        public override UE4Build.BuildAgenda GetAgenda(GUBP bp)
+        {
+           string Args = "-nobuilduht -skipactionhistory -CopyAppBundleBackToDevice";
+
+            UE4Build.BuildAgenda Agenda = new UE4Build.BuildAgenda();
+			if (GameProj.GameName == BranchConfig.Branch.BaseEngineProject.GameName)
+			{
+				Agenda.AddTargets(new string[] { TargetName }, TargetPlatform, TargetConfiguration, InAddArgs: Args);
+			}
+			else
+			{
+				Agenda.AddTargets(new string[] { TargetName }, TargetPlatform, TargetConfiguration, GameProj.FilePath, InAddArgs: Args);
+			}
+            return Agenda;
+        }
+	}
+
     public class SuccessNode : GUBPNode
     {
         public SuccessNode()
diff --git a/Engine/Source/Programs/UnrealHeaderTool/Private/CodeGenerator.cpp b/Engine/Source/Programs/UnrealHeaderTool/Private/CodeGenerator.cpp
--- a/Engine/Source/Programs/UnrealHeaderTool/Private/CodeGenerator.cpp
+++ b/Engine/Source/Programs/UnrealHeaderTool/Private/CodeGenerator.cpp
@@ -976,7 +976,7 @@ FString FNativeClassHeaderGenerator::PropertyNew(FString& Meta, UProperty* Prop,
 	}
 	else if (UBoolProperty* BoolProperty = Cast<UBoolProperty>(Prop))
 	{
-		if (Cast<UArrayProperty>(BoolProperty->GetOuter()) || Cast<UMapProperty>(BoolProperty->GetOuter()))
+		if (Cast<UArrayProperty>(BoolProperty->GetOuter()) || Cast<UMapProperty>(BoolProperty->GetOuter()) || Cast<USetProperty>(BoolProperty->GetOuter()))
 		{
 			ExtraArgs = FString(TEXT(", 0"));  // this is an array of C++ bools so the mask is irrelevant.
 		}
@@ -1151,13 +1151,20 @@ void FNativeClassHeaderGenerator::OutputProperty(FString& Meta, FOutputDevice& O
 		OutputDevice.Log(*PropertyNew(Meta, ArrayProperty->Inner, InnerOuterString, PropMacroOuterArray, TEXT("_Inner"), Spaces));
 	}
 
-	if (UMapProperty* MapProperty = Cast<UMapProperty>(Prop))
+	else if (UMapProperty* MapProperty = Cast<UMapProperty>(Prop))
 	{
 		FString InnerOuterString = FString::Printf(TEXT("NewProp_%s"), *PropName);
 		FString PropMacroOuterMap = TEXT("FObjectInitializer(), EC_CppProperty, ");
 		OutputDevice.Log(*PropertyNew(Meta, MapProperty->KeyProp,   InnerOuterString, PropMacroOuterMap + TEXT("0"), TEXT("_KeyProp"),   Spaces));
 		OutputDevice.Log(*PropertyNew(Meta, MapProperty->ValueProp, InnerOuterString, PropMacroOuterMap + TEXT("1"), TEXT("_ValueProp"), Spaces));
 	}
+
+	else if (USetProperty* SetProperty = Cast<USetProperty>(Prop))
+	{
+		FString InnerOuterString = FString::Printf(TEXT("NewProp_%s"), *Prop->GetName());
+		FString PropMacroOuterArray = TEXT("FObjectInitializer(), EC_CppProperty, 0");
+		OutputDevice.Log(*PropertyNew(Meta, SetProperty->ElementProp, InnerOuterString, PropMacroOuterArray, TEXT("_ElementProp"), Spaces));
+	}
 }
 
 static bool IsAlwaysAccessible(UScriptStruct* Script)
@@ -1223,6 +1230,13 @@ static void FindNoExportStructs(TArray<UScriptStruct*>& Structs, UStruct* Start)
 					FindNoExportStructs(Structs, ValueStructProp->Struct);
 				}
 			}
+			else if (USetProperty* SetProp = Cast<USetProperty>(Prop))
+			{
+				if (UStructProperty* ElementStructProp = Cast<UStructProperty>(SetProp->ElementProp))
+				{
+					FindNoExportStructs(Structs, ElementStructProp->Struct);
+				}
+			}
 		}
 		Start = Start->GetSuperStruct();
 	}
@@ -3399,10 +3413,14 @@ void ExportProtoDeclaration(FOutputDevice& Out, const FString& MessageName, TFie
 					FError::Throwf(TEXT("ExportProtoDeclaration - Unhandled property type '%s': %s"), *PropClass->GetName(), *Property->GetPathName());
 				}
 			}
-			else if (UMapProperty* MapProperty = Cast<UMapProperty>(Property))
+			else if (Property->IsA<UMapProperty>())
 			{
 				FError::Throwf(TEXT("ExportProtoDeclaration - Map properties not yet supported '%s': %s"), *PropClass->GetName(), *Property->GetPathName());
 			}
+			else if (Property->IsA<USetProperty>())
+			{
+				FError::Throwf(TEXT("ExportProtoDeclaration - Set properties not yet supported '%s': %s"), *PropClass->GetName(), *Property->GetPathName());
+			}
 			else if(Property->ArrayDim != 1)
 			{
 				bIsRepeated = true;
@@ -3586,6 +3604,12 @@ void ExportMCPDeclaration(FOutputDevice& Out, const FString& MessageName, TField
 			{
 				continue;
 			}
+			// TODO Implement sets
+			USetProperty* SetProperty = Cast<USetProperty>(Property);
+			if (SetProperty != NULL)
+			{
+				continue;
+			}
 			/*bool bIsRepeated = false;
 			UArrayProperty* ArrayProperty = Cast<UArrayProperty>(Property);
 			if (ArrayProperty != NULL)
@@ -3816,6 +3840,7 @@ void FNativeClassHeaderGenerator::ExportEventParm(UFunction* Function, FUHTStrin
 			Cast<UTextProperty>(Prop) ||
 			Cast<UArrayProperty>(Prop) ||
 			Cast<UMapProperty>(Prop) ||
+			Cast<USetProperty>(Prop) ||
 			Cast<UInterfaceProperty>(Prop)
 			)
 		{
@@ -3891,6 +3916,7 @@ FString FNativeClassHeaderGenerator::GetNullParameterValue( UProperty* Prop, boo
 	}
 	else if ( PropClass == UArrayProperty::StaticClass()
 		||    PropClass == UMapProperty::StaticClass()
+		||    PropClass == USetProperty::StaticClass()
 		||    PropClass == UDelegateProperty::StaticClass()
 		||    PropClass == UMulticastDelegateProperty::StaticClass() )
 	{
diff --git a/Engine/Source/Programs/UnrealHeaderTool/Private/HeaderParser.cpp b/Engine/Source/Programs/UnrealHeaderTool/Private/HeaderParser.cpp
--- a/Engine/Source/Programs/UnrealHeaderTool/Private/HeaderParser.cpp
+++ b/Engine/Source/Programs/UnrealHeaderTool/Private/HeaderParser.cpp
@@ -3298,7 +3298,7 @@ FIndexRange*                    ParsedVarIndexRange
 		VarType.PropertyFlags |= Flags;
 
 		GetVarType(AllClasses, Scope, VarProperty, Disallow, &VarType, EPropertyDeclarationStyle::None, VariableCategory);
-		if (VarProperty.ArrayType != EArrayType::None || VarProperty.MapKeyProp.IsValid())
+		if (VarProperty.IsContainer())
 		{
 			FError::Throwf(TEXT("Nested containers are not supported.") );
 		}
@@ -3347,7 +3347,7 @@ FIndexRange*                    ParsedVarIndexRange
 
 		FToken MapKeyType;
 		GetVarType(AllClasses, Scope, MapKeyType, Disallow, &VarType, EPropertyDeclarationStyle::None, VariableCategory);
-		if (VarProperty.ArrayType != EArrayType::None || VarProperty.MapKeyProp.IsValid())
+		if (MapKeyType.IsContainer())
 		{
 			FError::Throwf(TEXT("Nested containers are not supported.") );
 		}
@@ -3364,13 +3364,13 @@ FIndexRange*                    ParsedVarIndexRange
 		}
 
 		GetVarType(AllClasses, Scope, VarProperty, Disallow, &VarType, EPropertyDeclarationStyle::None, VariableCategory);
-		if (VarProperty.ArrayType != EArrayType::None || VarProperty.MapKeyProp.IsValid())
+		if (VarProperty.IsContainer())
 		{
 			FError::Throwf(TEXT("Nested containers are not supported.") );
 		}
 
-		OriginalVarTypeFlags |= VarProperty.PropertyFlags & (CPF_ContainsInstancedReference | CPF_InstancedReference); // propagate these to the array, we will fix them later
-		OriginalVarTypeFlags |= MapKeyType .PropertyFlags & (CPF_ContainsInstancedReference | CPF_InstancedReference); // propagate these to the array, we will fix them later
+		OriginalVarTypeFlags |= VarProperty.PropertyFlags & (CPF_ContainsInstancedReference | CPF_InstancedReference); // propagate these to the map value, we will fix them later
+		OriginalVarTypeFlags |= MapKeyType .PropertyFlags & (CPF_ContainsInstancedReference | CPF_InstancedReference); // propagate these to the map key, we will fix them later
 		VarType.PropertyFlags = OriginalVarTypeFlags;
 		FToken* MapKeyProp = new FToken(MapKeyType);
 		VarProperty.MapKeyProp = MakeShareable<FToken>(MapKeyProp);
@@ -3401,6 +3401,54 @@ FIndexRange*                    ParsedVarIndexRange
 			FError::Throwf(TEXT("Found '%s' - explicit allocators are not supported in TMap properties."), AllocatorToken.Identifier);
 		}
 	}
+	else if ( VarType.Matches(TEXT("TSet")) )
+	{
+		RequireSymbol( TEXT("<"), TEXT("'tset'") );
+
+		// GetVarType() clears the property flags of the array var, so use dummy 
+		// flags when getting the inner property
+		uint64 OriginalVarTypeFlags = VarType.PropertyFlags;
+		VarType.PropertyFlags |= Flags;
+
+		GetVarType(AllClasses, Scope, VarProperty, Disallow, &VarType, EPropertyDeclarationStyle::None, VariableCategory);
+		if (VarProperty.IsContainer())
+		{
+			FError::Throwf(TEXT("Nested containers are not supported.") );
+		}
+
+		if (VarProperty.Type == CPT_Struct)
+		{
+			FError::Throwf(TEXT("USTRUCTs are not currently supported as element types."));
+		}
+
+		OriginalVarTypeFlags |= VarProperty.PropertyFlags & (CPF_ContainsInstancedReference | CPF_InstancedReference); // propagate these to the set, we will fix them later
+		VarType.PropertyFlags = OriginalVarTypeFlags;
+		VarProperty.ArrayType = EArrayType::Set;
+
+		FToken CloseTemplateToken;
+		if (!GetToken(CloseTemplateToken, /*bNoConsts=*/ true, ESymbolParseOption::CloseTemplateBracket))
+		{
+			FError::Throwf(TEXT("Missing token while parsing TArray."));
+		}
+
+		if (CloseTemplateToken.TokenType != TOKEN_Symbol || FCString::Stricmp(CloseTemplateToken.Identifier, TEXT(">")))
+		{
+			// If we didn't find a comma, report it
+			if (FCString::Stricmp(CloseTemplateToken.Identifier, TEXT(",")))
+			{
+				FError::Throwf(TEXT("Expected '>' but found '%s'"), CloseTemplateToken.Identifier);
+			}
+
+			// If we found a comma, read the next thing, assume it's an allocator, and report that
+			FToken AllocatorToken;
+			if (!GetToken(AllocatorToken, /*bNoConsts=*/ true, ESymbolParseOption::CloseTemplateBracket))
+			{
+				FError::Throwf(TEXT("Expected '>' but found '%s'"), CloseTemplateToken.Identifier);
+			}
+
+			FError::Throwf(TEXT("Found '%s' - explicit allocators are not supported in TSet properties."), AllocatorToken.Identifier);
+		}
+	}
 	else if ( VarType.Matches(TEXT("FString")) )
 	{
 		VarProperty = FPropertyBase(CPT_String);
@@ -4061,7 +4109,7 @@ UProperty* FHeaderParser::GetVarNameAndDim
 	int32 OuterContextCount = 0;
 	UField* Existing = FindField(Scope, VarProperty.Identifier, true, UField::StaticClass(), NULL);
 
-	if (Existing != NULL)
+	if (Existing != nullptr)
 	{
 		bool bErrorDueToShadowing = true;
 
@@ -4108,7 +4156,7 @@ UProperty* FHeaderParser::GetVarNameAndDim
 			}
 		}
 
-		if (VarProperty.ArrayType == EArrayType::Dynamic || VarProperty.MapKeyProp.IsValid())
+		if (VarProperty.IsContainer())
 		{
 			FError::Throwf(TEXT("Static arrays of containers are not allowed"));
 		}
@@ -4127,7 +4175,7 @@ UProperty* FHeaderParser::GetVarNameAndDim
 		// Only static arrays are declared with [].  Dynamic arrays use TArray<> instead.
 		VarProperty.ArrayType = EArrayType::Static;
 
-		UEnum* Enum = NULL;
+		UEnum* Enum = nullptr;
 
 		if (*Dimensions.String)
 		{
@@ -4218,10 +4266,10 @@ UProperty* FHeaderParser::GetVarNameAndDim
 	FName PropertyName(VarProperty.Identifier, FindFlag);
 
 	// Add property.
-	UProperty* NewProperty = NULL;
+	UProperty* NewProperty = nullptr;
 
 	{
-		UProperty* Prev = NULL;
+		UProperty* Prev = nullptr;
 	    for (TFieldIterator<UProperty> It(Scope, EFieldIteratorFlags::ExcludeSuper); It; ++It)
 		{
 			Prev = *It;
@@ -4229,6 +4277,7 @@ UProperty* FHeaderParser::GetVarNameAndDim
 
 		UArrayProperty* Array             = nullptr;
 		UMapProperty*   Map               = nullptr;
+		USetProperty*   Set               = nullptr; // TODO: Set Property
 		UProperty*      NewMapKeyProperty = nullptr;
 		UObject*        NewScope          = Scope;
 		int32           ArrayDim          = 1; // 1 = not a static array, 2 = static array
@@ -4243,6 +4292,13 @@ UProperty* FHeaderParser::GetVarNameAndDim
 		{
 			ArrayDim = 2;
 		}
+		else if (VarProperty.ArrayType == EArrayType::Set)
+		{
+			Set               = new (EC_InternalUseOnlyConstructor, Scope, PropertyName, ObjectFlags) USetProperty(FObjectInitializer());
+			UHTMakefile.AddSetProperty(CurrentSrcFile, Set);
+			NewScope          = Set;
+			ObjectFlags       = RF_Public;
+		}
 		else if (VarProperty.MapKeyProp.IsValid())
 		{
 			Map               = new (EC_InternalUseOnlyConstructor, Scope, PropertyName, ObjectFlags) UMapProperty(FObjectInitializer());
@@ -4293,13 +4349,22 @@ UProperty* FHeaderParser::GetVarNameAndDim
 			NewProperty = Map;
 		}
 
+		if (Set)
+		{
+			Set->ElementProp = NewProperty;
+
+			PropagateFlags(CPF_PropagateToSetElement, VarProperty, NewProperty);
+
+			NewProperty = Set;
+		}
+
 		NewProperty->ArrayDim = ArrayDim;
 		if (ArrayDim == 2)
 		{
 			GArrayDimensions.Add(NewProperty, Dimensions.String);
 		}
 		NewProperty->PropertyFlags = VarProperty.PropertyFlags;
-		if (Prev != NULL)
+		if (Prev != nullptr)
 		{
 			NewProperty->Next = Prev->Next;
 			Prev->Next = NewProperty;
diff --git a/Engine/Source/Programs/UnrealHeaderTool/Private/ParserHelper.cpp b/Engine/Source/Programs/UnrealHeaderTool/Private/ParserHelper.cpp
--- a/Engine/Source/Programs/UnrealHeaderTool/Private/ParserHelper.cpp
+++ b/Engine/Source/Programs/UnrealHeaderTool/Private/ParserHelper.cpp
@@ -132,6 +132,7 @@ const TCHAR* FPropertyBase::GetPropertyTypeText( EPropertyType Type )
 		CASE_TEXT(CPT_WeakObjectReference);
 		CASE_TEXT(CPT_LazyObjectReference);
 		CASE_TEXT(CPT_Map);
+		CASE_TEXT(CPT_Set);
 		CASE_TEXT(CPT_MAX);
 	}
 
diff --git a/Engine/Source/Programs/UnrealHeaderTool/Private/ParserHelper.h b/Engine/Source/Programs/UnrealHeaderTool/Private/ParserHelper.h
--- a/Engine/Source/Programs/UnrealHeaderTool/Private/ParserHelper.h
+++ b/Engine/Source/Programs/UnrealHeaderTool/Private/ParserHelper.h
@@ -45,7 +45,8 @@ struct EArrayType
 	{
 		None,
 		Static,
-		Dynamic
+		Dynamic,
+		Set
 	};
 };
 
@@ -427,6 +428,11 @@ public:
 		return Type == CPT_ObjectReference || Type == CPT_Interface || Type == CPT_WeakObjectReference || Type == CPT_LazyObjectReference || Type == CPT_AssetObjectReference;
 	}
 
+	bool IsContainer() const
+	{
+		return (ArrayType != EArrayType::None || MapKeyProp.IsValid());
+	}
+
 	/**
 	 * Determines whether this token's type is compatible with another token's type.
 	 *
diff --git a/Engine/Source/Programs/UnrealHeaderTool/Private/UHTMakefile/MakefileHelpers.h b/Engine/Source/Programs/UnrealHeaderTool/Private/UHTMakefile/MakefileHelpers.h
--- a/Engine/Source/Programs/UnrealHeaderTool/Private/UHTMakefile/MakefileHelpers.h
+++ b/Engine/Source/Programs/UnrealHeaderTool/Private/UHTMakefile/MakefileHelpers.h
@@ -128,6 +128,7 @@ enum class ESerializedObjectType : uint32
 	EInterfaceProperty,
 	EStructProperty,
 	EMapProperty,
+	ESetProperty,
 	EArrayProperty,
 	EUnrealSourceFile,
 	EFileScope,
diff --git a/Engine/Source/Programs/UnrealHeaderTool/Private/UHTMakefile/PropertyArchiveProxy.cpp b/Engine/Source/Programs/UnrealHeaderTool/Private/UHTMakefile/PropertyArchiveProxy.cpp
--- a/Engine/Source/Programs/UnrealHeaderTool/Private/UHTMakefile/PropertyArchiveProxy.cpp
+++ b/Engine/Source/Programs/UnrealHeaderTool/Private/UHTMakefile/PropertyArchiveProxy.cpp
@@ -87,6 +87,15 @@ FArchive& operator<<(FArchive& Ar, FMapPropertyArchiveProxy& MapPropertyArchiveP
 	return Ar;
 }
 
+FArchive& operator<<(FArchive& Ar, FSetPropertyArchiveProxy& SetPropertyArchiveProxy)
+{
+	Ar << static_cast<FPropertyArchiveProxy&>(SetPropertyArchiveProxy);
+	Ar << SetPropertyArchiveProxy.ElementPropIndex;
+	Ar << SetPropertyArchiveProxy.SetLayout;
+
+	return Ar;
+}
+
 FBytePropertyArchiveProxy::FBytePropertyArchiveProxy(FUHTMakefile& UHTMakefile, const UByteProperty* ByteProperty)
 	: FPropertyArchiveProxy(UHTMakefile, ByteProperty)
 {
@@ -562,6 +571,28 @@ void FMapPropertyArchiveProxy::Resolve(UMapProperty* MapProperty, const FUHTMake
 	MapProperty->MapLayout = MapLayout;
 }
 
+FSetPropertyArchiveProxy::FSetPropertyArchiveProxy(FUHTMakefile& UHTMakefile, const USetProperty* SetProperty)
+	: FPropertyArchiveProxy(UHTMakefile, SetProperty)
+{
+	ElementPropIndex = UHTMakefile.GetPropertyIndex(SetProperty->ElementProp);
+	SetLayout = SetProperty->SetLayout;
+}
+
+USetProperty* FSetPropertyArchiveProxy::CreateSetProperty(const FUHTMakefile& UHTMakefile) const
+{
+	UObject* Outer = UHTMakefile.GetObjectByIndex(OuterIndex);
+	USetProperty* SetProperty = new (EC_InternalUseOnlyConstructor, Outer, Name.CreateName(UHTMakefile), (EObjectFlags)ObjectFlagsUint32) USetProperty(FObjectInitializer());
+	PostConstruct(SetProperty, UHTMakefile);
+	return SetProperty;
+}
+
+void FSetPropertyArchiveProxy::Resolve(USetProperty* SetProperty, const FUHTMakefile& UHTMakefile) const
+{
+	FPropertyArchiveProxy::Resolve(SetProperty, UHTMakefile);
+	SetProperty->ElementProp = UHTMakefile.GetPropertyByIndex(ElementPropIndex);
+	SetProperty->SetLayout = SetLayout;
+}
+
 UArrayProperty* FArrayPropertyArchiveProxy::CreateArrayProperty(const FUHTMakefile& UHTMakefile) const
 {
 	UObject* Outer = UHTMakefile.GetObjectByIndex(OuterIndex);
diff --git a/Engine/Source/Programs/UnrealHeaderTool/Private/UHTMakefile/PropertyArchiveProxy.h b/Engine/Source/Programs/UnrealHeaderTool/Private/UHTMakefile/PropertyArchiveProxy.h
--- a/Engine/Source/Programs/UnrealHeaderTool/Private/UHTMakefile/PropertyArchiveProxy.h
+++ b/Engine/Source/Programs/UnrealHeaderTool/Private/UHTMakefile/PropertyArchiveProxy.h
@@ -401,6 +401,25 @@ struct FMapPropertyArchiveProxy : public FPropertyArchiveProxy
 	FScriptMapLayout MapLayout;
 };
 
+struct FSetPropertyArchiveProxy : public FPropertyArchiveProxy
+{
+	FSetPropertyArchiveProxy() { }
+	FSetPropertyArchiveProxy(FUHTMakefile& UHTMakefile, const USetProperty* SetProperty);
+
+	USetProperty* CreateSetProperty(const FUHTMakefile& UHTMakefile) const;
+	static void AddReferencedNames(const USetProperty* Property, FUHTMakefile& UHTMakefile)
+	{
+		FPropertyArchiveProxy::AddReferencedNames(Property, UHTMakefile);
+	}
+
+	void Resolve(USetProperty* SetProperty, const FUHTMakefile& UHTMakefile) const;
+
+	friend FArchive& operator<<(FArchive& Ar, FSetPropertyArchiveProxy& SetPropertyArchiveProxy);
+
+	int32 ElementPropIndex;
+	FScriptSetLayout SetLayout;
+};
+
 struct FArrayPropertyArchiveProxy : public FPropertyArchiveProxy
 {
 	FArrayPropertyArchiveProxy() { }
@@ -417,4 +436,4 @@ struct FArrayPropertyArchiveProxy : public FPropertyArchiveProxy
 	friend FArchive& operator<<(FArchive& Ar, FArrayPropertyArchiveProxy& ArrayPropertyArchiveProxy);
 
 	int32 InnerIndex;
-};
+};
\ No newline at end of file
diff --git a/Engine/Source/Programs/UnrealHeaderTool/Private/UHTMakefile/UHTMakefile.cpp b/Engine/Source/Programs/UnrealHeaderTool/Private/UHTMakefile/UHTMakefile.cpp
--- a/Engine/Source/Programs/UnrealHeaderTool/Private/UHTMakefile/UHTMakefile.cpp
+++ b/Engine/Source/Programs/UnrealHeaderTool/Private/UHTMakefile/UHTMakefile.cpp
@@ -237,6 +237,12 @@ void FUHTMakefile::SetupProxies()
 		MapPropertyProxies.Add(FMapPropertyArchiveProxy(*this, MapProperty));
 	}
 
+	SetPropertyProxies.Empty(SetProperties.Num());
+	for (USetProperty* SetProperty : SetProperties)
+	{
+		SetPropertyProxies.Add(FSetPropertyArchiveProxy(*this, SetProperty));
+	}
+
 	ArrayPropertyProxies.Empty(ArrayProperties.Num());
 	for (UArrayProperty* ArrayProperty : ArrayProperties)
 	{
@@ -403,6 +409,7 @@ int32 FUHTMakefile::GetPropertyCount() const
 		+ InterfaceProperties.Num()
 		+ StructProperties.Num()
 		+ MapProperties.Num()
+		+ SetProperties.Num()
 		+ ArrayProperties.Num();
 }
 
@@ -435,6 +442,7 @@ int32 FUHTMakefile::GetPropertyProxiesCount() const
 		+ InterfacePropertyProxies.Num()
 		+ StructPropertyProxies.Num()
 		+ MapPropertyProxies.Num()
+		+ SetPropertyProxies.Num()
 		+ ArrayPropertyProxies.Num();
 }
 
@@ -760,6 +768,9 @@ void FUHTMakefile::CreateObject(ESerializedObjectType ObjectType, int32 Index)
 	case ESerializedObjectType::EMapProperty:
 		CreateMapProperty(Index);
 		break;
+	case ESerializedObjectType::ESetProperty:
+		CreateSetProperty(Index);
+		break;
 	case ESerializedObjectType::EArrayProperty:
 		CreateArrayProperty(Index);
 		break;
@@ -924,6 +935,9 @@ void FUHTMakefile::ResolveObject(ESerializedObjectType ObjectType, int32 Index)
 	case ESerializedObjectType::EMapProperty:
 		ResolveMapProperty(Index);
 		break;
+	case ESerializedObjectType::ESetProperty:
+		ResolveSetProperty(Index);
+		break;
 	case ESerializedObjectType::EArrayProperty:
 		ResolveArrayProperty(Index);
 		break;
@@ -1532,6 +1546,13 @@ int32 FUHTMakefile::GetPropertyIndex(const UProperty* Property) const
 	}
 	Result += MapProperties.Num();
 
+	CurrentIndex = SetPropertyIndexes.Find(Cast<USetProperty>(Property));
+	if (CurrentIndex)
+	{
+		return Result + *CurrentIndex;
+	}
+	Result += SetProperties.Num();
+
 	CurrentIndex = ArrayPropertyIndexes.Find(Cast<UArrayProperty>(Property));
 	if (CurrentIndex)
 	{
@@ -1711,6 +1732,12 @@ UProperty* FUHTMakefile::GetPropertyByIndex(int32 Index) const
 	}
 	Index -= MapPropertyProxies.Num();
 
+	if (SetProperties.IsValidIndex(Index) && SetPropertyProxies.IsValidIndex(Index))
+	{
+		return SetProperties[Index];
+	}
+	Index -= SetPropertyProxies.Num();
+
 	if (ArrayProperties.IsValidIndex(Index) && ArrayPropertyProxies.IsValidIndex(Index))
 	{
 		return ArrayProperties[Index];
@@ -2161,6 +2188,11 @@ void FUHTMakefile::SetupAndSaveReferencedNames(FArchive& Ar)
 		FMapPropertyArchiveProxy::AddReferencedNames(MapProperty, *this);
 	}
 
+	for (USetProperty* SetProperty : SetProperties)
+	{
+		FSetPropertyArchiveProxy::AddReferencedNames(SetProperty, *this);
+	}
+
 	for (UArrayProperty* ArrayProperty : ArrayProperties)
 	{
 		FArrayPropertyArchiveProxy::AddReferencedNames(ArrayProperty, *this);
@@ -2290,6 +2322,7 @@ FArchive& operator<<(FArchive& Ar, FUHTMakefile& UHTMakefile)
 	Ar << UHTMakefile.InterfacePropertyProxies;
 	Ar << UHTMakefile.StructPropertyProxies;
 	Ar << UHTMakefile.MapPropertyProxies;
+	Ar << UHTMakefile.SetPropertyProxies;
 	Ar << UHTMakefile.ArrayPropertyProxies;
 	Ar << UHTMakefile.DelegateFunctionProxies;
 	Ar << UHTMakefile.FunctionProxies;
diff --git a/Engine/Source/Programs/UnrealHeaderTool/Private/UHTMakefile/UHTMakefile.h b/Engine/Source/Programs/UnrealHeaderTool/Private/UHTMakefile/UHTMakefile.h
--- a/Engine/Source/Programs/UnrealHeaderTool/Private/UHTMakefile/UHTMakefile.h
+++ b/Engine/Source/Programs/UnrealHeaderTool/Private/UHTMakefile/UHTMakefile.h
@@ -811,6 +811,24 @@ public:
 		MapPropertyProxies[Index].Resolve(MapProperties[Index], *this);
 	}
 
+	void AddSetProperty(FUnrealSourceFile* UnrealSourceFile, USetProperty* SetProperty)
+	{
+		int32 Index = SetProperties.Add(SetProperty);
+		GetHeaderDescriptor(UnrealSourceFile).AddEntry(AddObject(ESerializedObjectType::ESetProperty, Index));
+		SetPropertyIndexes.Add(SetProperty, Index);
+	}
+	void CreateSetProperty(int32 Index)
+	{
+		checkSlow(Index == SetProperties.Num());
+		USetProperty* SetProperty = SetPropertyProxies[Index].CreateSetProperty(*this);
+		Index = SetProperties.Add(SetProperty);
+		SetPropertyIndexes.Add(SetProperty, Index);
+	}
+	void ResolveSetProperty(int32 Index)
+	{
+		SetPropertyProxies[Index].Resolve(SetProperties[Index], *this);
+	}
+
 	void AddArrayProperty(FUnrealSourceFile* UnrealSourceFile, UArrayProperty* ArrayProperty)
 	{
 		int32 Index = ArrayProperties.Add(ArrayProperty);
@@ -1565,6 +1583,10 @@ private:
 	TArray<UMapProperty*> MapProperties;
 	TMap<UMapProperty*, int32> MapPropertyIndexes;
 
+	TArray<FSetPropertyArchiveProxy> SetPropertyProxies;
+	TArray<USetProperty*> SetProperties;
+	TMap<USetProperty*, int32> SetPropertyIndexes;
+
 	TArray<FArrayPropertyArchiveProxy> ArrayPropertyProxies;
 	TArray<UArrayProperty*> ArrayProperties;
 	TMap<UArrayProperty*, int32> ArrayPropertyIndexes;
diff --git a/Engine/Source/Programs/UnrealHeaderTool/Resources/UHTDebugging/TestObject.h b/Engine/Source/Programs/UnrealHeaderTool/Resources/UHTDebugging/TestObject.h
--- a/Engine/Source/Programs/UnrealHeaderTool/Resources/UHTDebugging/TestObject.h
+++ b/Engine/Source/Programs/UnrealHeaderTool/Resources/UHTDebugging/TestObject.h
@@ -39,6 +39,9 @@ public:
 	UPROPERTY()
 	TMap<int32, bool> TestMap;
 
+	UPROPERTY()
+	TSet<int32> TestSet;
+
 	UFUNCTION()
 	void CodeGenTestForEnumClasses(ECppEnum Val);
 
diff --git a/Engine/Source/Runtime/AIModule/Private/BehaviorTree/BTNode.cpp b/Engine/Source/Runtime/AIModule/Private/BehaviorTree/BTNode.cpp
--- a/Engine/Source/Runtime/AIModule/Private/BehaviorTree/BTNode.cpp
+++ b/Engine/Source/Runtime/AIModule/Private/BehaviorTree/BTNode.cpp
@@ -74,7 +74,7 @@ void UBTNode::InitializeInSubtree(UBehaviorTreeComponent& OwnerComp, uint8* Node
 		UBTNode* NodeInstance = OwnerComp.NodeInstances.IsValidIndex(NextInstancedIndex) ? OwnerComp.NodeInstances[NextInstancedIndex] : NULL;
 		if (NodeInstance == NULL)
 		{
-			NodeInstance = NewObject<UBTNode>(&OwnerComp, GetClass(), NAME_None, RF_NoFlags, (UObject*)(this));
+			NodeInstance = (UBTNode*)StaticDuplicateObject(this, &OwnerComp);
 			NodeInstance->InitializeNode(GetParentNode(), GetExecutionIndex(), GetMemoryOffset(), GetTreeDepth());
 			NodeInstance->bIsInstanced = true;
 
diff --git a/Engine/Source/Runtime/Analytics/AnalyticsET/Private/IAnalyticsProviderET.cpp b/Engine/Source/Runtime/Analytics/AnalyticsET/Private/IAnalyticsProviderET.cpp
--- a/Engine/Source/Runtime/Analytics/AnalyticsET/Private/IAnalyticsProviderET.cpp
+++ b/Engine/Source/Runtime/Analytics/AnalyticsET/Private/IAnalyticsProviderET.cpp
@@ -11,7 +11,137 @@
 #include "Http.h"
 #include "EngineVersion.h"
 
-DEFINE_LOG_CATEGORY_STATIC(LogAnalyticsDumpEventPayload, Display, All);
+/** When enabled (and -AnalyticsTrackPerf is specified on the command line, will log out analytics flush timings on a regular basis to Saved/AnalyticsTiming.csv. */
+#define ANALYTICS_PERF_TRACKING_ENABLED !UE_BUILD_SHIPPING
+#if ANALYTICS_PERF_TRACKING_ENABLED
+
+/** Measures analytics bandwidth. Only active when -AnalyticsTrackPerf is specified on the command line. */
+struct FAnalyticsPerfTracker : FTickerObjectBase
+{
+	FAnalyticsPerfTracker()
+	{
+		bEnabled = FParse::Param(FCommandLine::Get(), TEXT("ANALYTICSTRACKPERF"));
+		if (bEnabled)
+		{
+			LogFile.SetSuppressEventTag(true);
+			LogFile.Serialize(TEXT("Date,CL,RunID,Time,WindowSeconds,ProfiledSeconds,Frames,Flushes,Events,Bytes"), ELogVerbosity::Log, FName());
+			LastSubmitTime = StartTime;
+			StartDate = FDateTime::UtcNow().ToIso8601();
+			CL = LexicalConversion::ToString(FEngineVersion::Current().GetChangelist());
+		}
+	}
+
+	/** Called once per flush */
+	void RecordFlush(uint64 Bytes, uint64 NumEvents, double TimeSec)
+	{
+		if (bEnabled)
+		{
+			++FlushesThisWindow;
+			BytesThisWindow += Bytes;
+			NumEventsThisWindow += NumEvents;
+			TimeThisWindow += TimeSec;
+		}
+	}
+
+	static FAnalyticsPerfTracker& Get()
+	{
+		static FAnalyticsPerfTracker GTracker;
+		return GTracker;
+	}
+
+	bool IsEnabled() const { return bEnabled; }
+
+	void SetRunID(const FString& InRunID)
+	{
+		if (bEnabled)
+		{
+			RunID = InRunID;
+			StartDate = FDateTime::UtcNow().ToIso8601();
+		}
+	}
+
+private:
+	/** Check to see if we need to log another window of time. */
+	virtual bool Tick(float DeltaTime) override
+	{
+		if (bEnabled)
+		{
+			++FramesThisWindow;
+			double Now = FPlatformTime::Seconds();
+			if (WindowExpired(Now))
+			{
+				LogFile.Serialize(*FString::Printf(TEXT("%s,%s,%s,%f,%f,%f,%d,%d,%d,%d"),
+					*StartDate,
+					*CL,
+					*RunID,
+					Now - StartTime,
+					Now - LastSubmitTime,
+					TimeThisWindow,
+					FramesThisWindow,
+					FlushesThisWindow,
+					NumEventsThisWindow,
+					BytesThisWindow),
+					ELogVerbosity::Log, FName(), Now);
+				ResetWindow(Now);
+			}
+			return true;
+		}
+		else
+		{
+			return false;
+		}
+	}
+
+	/** Helper to reset our window in Tick. */
+	bool WindowExpired(double Now)
+	{
+		return Now > LastSubmitTime + 60.0;
+	}
+
+	/** Helper to reset our window in Tick. */
+	void ResetWindow(double Now)
+	{
+		LastSubmitTime = Now;
+		TimeThisWindow = 0.0;
+		BytesThisWindow = 0;
+		NumEventsThisWindow = 0;
+		FlushesThisWindow = 0;
+		FramesThisWindow = 0;
+	}
+
+	/** log file to use. */
+	FOutputDeviceFile LogFile{ *FPaths::Combine(*FPaths::GameSavedDir(), TEXT("AnalyticsTiming.csv")) };
+	FString StartDate;
+	FString CL;
+	FString RunID = FGuid().ToString().ToLower();
+	// Window tracking data
+	double LastSubmitTime = 0.0;
+	double TimeThisWindow = 0.0;
+	uint64 BytesThisWindow = 0;
+	uint64 NumEventsThisWindow = 0;
+	int FlushesThisWindow = 0;
+	int FramesThisWindow = 0;
+	// time when the first measurement was made.
+	double StartTime = FPlatformTime::Seconds();
+	/** Controls whether metrics gathering is enabled. */
+	bool bEnabled = false;
+};
+
+/** Used to set the RunID between matches in game code. Must be carefully called only in situations where ANALYTICS_PERF_TRACKING_ENABLED = 1 */
+ANALYTICSET_API void SetAnayticsETPerfTrackingRunID(const FString& RunID)
+{
+	FAnalyticsPerfTracker::Get().SetRunID(RunID);
+}
+
+#define ANALYTICS_FLUSH_TRACKING_BEGIN() double FlushStartTime = FPlatformTime::Seconds()
+#define ANALYTICS_FLUSH_TRACKING_END(NumBytes, NumEvents) FAnalyticsPerfTracker::Get().RecordFlush(NumBytes, NumEvents, FPlatformTime::Seconds() - FlushStartTime)
+
+#else
+
+#define ANALYTICS_FLUSH_TRACKING_BEGIN(...)
+#define ANALYTICS_FLUSH_TRACKING_END(...)
+
+#endif
 
 class FAnalyticsProviderET :
 	public IAnalyticsProviderET,
@@ -182,6 +312,10 @@ bool FAnalyticsProviderET::Tick(float DeltaSeconds)
 {
 	QUICK_SCOPE_CYCLE_COUNTER(STAT_FAnalyticsProviderET_Tick);
 
+	// There are much better ways to do this, but since most events are recorded and handled on the same (game) thread,
+	// this is probably mostly fine for now, and simply favoring not crashing at the moment
+	FScopeLock ScopedLock(&CachedEventsCS);
+
 	if (CachedEvents.Num() > 0)
 	{
 		// Countdown to flush
@@ -216,7 +350,7 @@ bool FAnalyticsProviderET::StartSession(const TArray<FAnalyticsEventAttribute>&
 bool FAnalyticsProviderET::StartSession(TArray<FAnalyticsEventAttribute>&& Attributes)
 {
 	UE_LOG(LogAnalytics, Log, TEXT("[%s] AnalyticsET::StartSession"), *APIKey);
-
+	
 	// end/flush previous session before staring new one
 	if (bSessionInProgress)
 	{
@@ -256,15 +390,19 @@ void FAnalyticsProviderET::EndSession()
 
 void FAnalyticsProviderET::FlushEvents()
 {
+	// There are much better ways to do this, but since most events are recorded and handled on the same (game) thread,
+	// this is probably mostly fine for now, and simply favoring not crashing at the moment
+	FScopeLock ScopedLock(&CachedEventsCS);
+
 	// Make sure we don't try to flush too many times. When we are not caching events it's possible this can be called when there are no events in the array.
 	if (CachedEvents.Num() == 0)
 	{
 		return;
 	}
 
-	// There are much better ways to do this, but since most events are recorded and handled on the same (game) thread,
-	// this is probably mostly fine for now, and simply favoring not crashing at the moment
-	FScopeLock ScopedLock(&CachedEventsCS);
+	ANALYTICS_FLUSH_TRACKING_BEGIN();
+	int EventCount = CachedEvents.Num();
+	int PayloadSize = 0;
 
 	if(ensure(FModuleManager::Get().IsModuleLoaded("HTTP")))
 	{
@@ -315,23 +453,23 @@ void FAnalyticsProviderET::FlushEvents()
 				*FPlatformHttp::UrlEncode(UserID),
 				*FPlatformHttp::UrlEncode(AppEnvironment),
 				*FPlatformHttp::UrlEncode(UploadType));
+			PayloadSize = URLPath.Len() + Payload.Len();
 
-			// Recreate the URLPath for logging because we do not want to escape the parameters when logging.
-			// We cannot simply UrlEncode the entire Path after logging it because UrlEncode(Params) != UrlEncode(Param1) & UrlEncode(Param2) ...
-			FString LogString = FString::Printf(TEXT("[%s] AnalyticsET URL:datarouter/api/v1/public/data?SessionID=%s&AppID=%s&AppVersion=%s&UserID=%s&AppEnvironment=%s&UploadType=%s. Payload:%s"),
-				*APIKey,
-				*SessionID,
-				*APIKey,
-				*AppVersion,
-				*UserID,
-				*AppEnvironment,
-				*UploadType,
-				*Payload);
-			UE_LOG(LogAnalytics, VeryVerbose, TEXT("%s"), *LogString);
-
-			// Duplicate the same log message with a separate category. This is used as an "last chance" backup on the servers in the unlikely case 
-			// if the backend lost the events due to overload - then we can scrape the logs manually for them.
-			UE_LOG(LogAnalyticsDumpEventPayload, Log, TEXT("%s"), *LogString);
+			if (UE_LOG_ACTIVE(LogAnalytics, VeryVerbose))
+			{
+				// Recreate the URLPath for logging because we do not want to escape the parameters when logging.
+				// We cannot simply UrlEncode the entire Path after logging it because UrlEncode(Params) != UrlEncode(Param1) & UrlEncode(Param2) ...
+				FString LogString = FString::Printf(TEXT("[%s] AnalyticsET URL:datarouter/api/v1/public/data?SessionID=%s&AppID=%s&AppVersion=%s&UserID=%s&AppEnvironment=%s&UploadType=%s. Payload:%s"),
+					*APIKey,
+					*SessionID,
+					*APIKey,
+					*AppVersion,
+					*UserID,
+					*AppEnvironment,
+					*UploadType,
+					*Payload);
+				UE_LOG(LogAnalytics, VeryVerbose, TEXT("%s"), *LogString);
+			}
 
 			// Create/send Http request for an event
 			TSharedRef<IHttpRequest> HttpRequest = FHttpModule::Get().CreateRequest();
@@ -387,6 +525,7 @@ void FAnalyticsProviderET::FlushEvents()
 					*FPlatformHttp::UrlEncode(UserID), 
 					*FPlatformHttp::UrlEncode(Event.EventName), 
 					*EventParams));
+				PayloadSize = HttpRequest->GetURL().Len();
 				HttpRequest->SetVerb(TEXT("GET"));
 				HttpRequest->OnProcessRequestComplete().BindRaw(this, &FAnalyticsProviderET::EventRequestComplete);
 				HttpRequest->ProcessRequest();
@@ -396,6 +535,7 @@ void FAnalyticsProviderET::FlushEvents()
 		FlushEventsCountdown = MaxCachedElapsedTime;
 		CachedEvents.Empty();
 	}
+	ANALYTICS_FLUSH_TRACKING_END(PayloadSize, EventCount);
 }
 
 void FAnalyticsProviderET::SetUserID(const FString& InUserID)
diff --git a/Engine/Source/Runtime/AnimGraphRuntime/Private/AnimNodes/AnimNode_PoseBlendNode.cpp b/Engine/Source/Runtime/AnimGraphRuntime/Private/AnimNodes/AnimNode_PoseBlendNode.cpp
new file mode 100644
--- /dev/null
+++ b/Engine/Source/Runtime/AnimGraphRuntime/Private/AnimNodes/AnimNode_PoseBlendNode.cpp
@@ -0,0 +1,57 @@
+// Copyright 1998-2016 Epic Games, Inc. All Rights Reserved.
+
+#include "AnimGraphRuntimePrivatePCH.h"
+#include "AnimNodes/AnimNode_PoseBlendNode.h"
+#include "Animation/AnimInstanceProxy.h"
+
+/////////////////////////////////////////////////////
+// FAnimPoseByNameNode
+
+void FAnimNode_PoseBlendNode::Initialize(const FAnimationInitializeContext& Context)
+{
+	FAnimNode_AssetPlayerBase::Initialize(Context);
+
+	SourcePose.Initialize(Context);
+}
+
+void FAnimNode_PoseBlendNode::CacheBones(const FAnimationCacheBonesContext& Context) 
+{
+	SourcePose.CacheBones(Context);
+}
+
+void FAnimNode_PoseBlendNode::UpdateAssetPlayer(const FAnimationUpdateContext& Context)
+{
+	EvaluateGraphExposedInputs.Execute(Context);
+	SourcePose.Update(Context);
+}
+
+void FAnimNode_PoseBlendNode::Evaluate(FPoseContext& Output)
+{
+	SourcePose.Evaluate(Output);
+
+	if ((PoseAsset!= NULL) && (Output.AnimInstanceProxy->IsSkeletonCompatible(PoseAsset->GetSkeleton())))
+	{
+		// only give pose curve, we don't set any more curve here
+		PoseAsset->GetAnimationPose(Output.Pose, Output.Curve, FAnimExtractContext());
+	}
+	// @else it doesn't do anything - it might contain invalid transform but I don't know if overriding transform is the right way 
+	// because the node might contain incoming transform
+}
+
+void FAnimNode_PoseBlendNode::OverrideAsset(UAnimationAsset* NewAsset)
+{
+	if(UPoseAsset* NewPoseAsset = Cast<UPoseAsset>(NewAsset))
+	{
+		PoseAsset = NewPoseAsset;
+	}
+}
+
+void FAnimNode_PoseBlendNode::GatherDebugData(FNodeDebugData& DebugData)
+{
+	FString DebugLine = DebugData.GetNodeName(this);
+	
+	DebugLine += FString::Printf(TEXT("('%s')"), *GetNameSafe(PoseAsset));
+	DebugData.AddDebugItem(DebugLine, true);
+	SourcePose.GatherDebugData(DebugData.BranchFlow(1.f));
+}
+
diff --git a/Engine/Source/Runtime/AnimGraphRuntime/Private/AnimNodes/AnimNode_PoseByName.cpp b/Engine/Source/Runtime/AnimGraphRuntime/Private/AnimNodes/AnimNode_PoseByName.cpp
new file mode 100644
--- /dev/null
+++ b/Engine/Source/Runtime/AnimGraphRuntime/Private/AnimNodes/AnimNode_PoseByName.cpp
@@ -0,0 +1,71 @@
+// Copyright 1998-2016 Epic Games, Inc. All Rights Reserved.
+
+#include "AnimGraphRuntimePrivatePCH.h"
+#include "AnimNodes/AnimNode_PoseByName.h"
+#include "Animation/AnimInstanceProxy.h"
+
+/////////////////////////////////////////////////////
+// FAnimPoseByNameNode
+
+void FAnimNode_PoseByName::Initialize(const FAnimationInitializeContext& Context)
+{
+	FAnimNode_AssetPlayerBase::Initialize(Context);
+
+	if (PoseName != NAME_None)
+	{
+		USkeleton* Skeleton = Context.AnimInstanceProxy->GetSkeleton();
+		PoseUID = Skeleton->GetUIDByName(USkeleton::AnimCurveMappingName, PoseName);
+	}
+	else
+	{
+		PoseUID = FSmartNameMapping::MaxUID;
+	}
+}
+
+void FAnimNode_PoseByName::CacheBones(const FAnimationCacheBonesContext& Context) 
+{
+}
+
+void FAnimNode_PoseByName::UpdateAssetPlayer(const FAnimationUpdateContext& Context)
+{
+	EvaluateGraphExposedInputs.Execute(Context);
+}
+
+void FAnimNode_PoseByName::Evaluate(FPoseContext& Output)
+{
+	if ((PoseAsset!= NULL) && (Output.AnimInstanceProxy->IsSkeletonCompatible(PoseAsset->GetSkeleton())))
+	{
+		FBlendedCurve PoseCurve;
+		PoseCurve.InitFrom(Output.Curve);
+
+		if (PoseUID != FSmartNameMapping::MaxUID)
+		{
+			PoseCurve.Set(PoseUID, 1.f, ACF_DrivesPose);
+		}
+
+		// only give pose curve, we don't set any more curve here
+		PoseAsset->GetAnimationPose(Output.Pose, PoseCurve, FAnimExtractContext());
+
+		// this curve data doesn't go out, but it is intentional
+	}
+	else
+	{
+		Output.ResetToRefPose();
+	}
+}
+
+void FAnimNode_PoseByName::OverrideAsset(UAnimationAsset* NewAsset)
+{
+	if(UPoseAsset* NewPoseAsset = Cast<UPoseAsset>(NewAsset))
+	{
+		PoseAsset = NewPoseAsset;
+	}
+}
+
+void FAnimNode_PoseByName::GatherDebugData(FNodeDebugData& DebugData)
+{
+	FString DebugLine = DebugData.GetNodeName(this);
+	
+	DebugLine += FString::Printf(TEXT("('%s' Pose: %s)"), *GetNameSafe(PoseAsset), *PoseName.ToString());
+	DebugData.AddDebugItem(DebugLine, true);
+}
\ No newline at end of file
diff --git a/Engine/Source/Runtime/AnimGraphRuntime/Private/BoneControllers/AnimNode_BoneDrivenController.cpp b/Engine/Source/Runtime/AnimGraphRuntime/Private/BoneControllers/AnimNode_BoneDrivenController.cpp
--- a/Engine/Source/Runtime/AnimGraphRuntime/Private/BoneControllers/AnimNode_BoneDrivenController.cpp
+++ b/Engine/Source/Runtime/AnimGraphRuntime/Private/BoneControllers/AnimNode_BoneDrivenController.cpp
@@ -185,9 +185,11 @@ void FAnimNode_BoneDrivenController::EvaluateComponentSpaceInternal(FComponentSp
 	{
 		//	Morph target and Material parameter curves
 		USkeleton* Skeleton = Context.AnimInstanceProxy->GetSkeleton();
-		FSmartNameMapping::UID NameUID;
-		Skeleton->AddSmartNameAndModify(USkeleton::AnimCurveMappingName, ParameterName, NameUID);
-		Context.Curve.Set(NameUID, FinalDriverValue, (DestinationMode == EDrivenDestinationMode::MorphTarget) ? EAnimCurveFlags::ACF_DrivesMorphTarget : EAnimCurveFlags::ACF_DrivesMaterial);
+		FSmartNameMapping::UID NameUID = Skeleton->GetUIDByName(USkeleton::AnimCurveMappingName, ParameterName);
+		if (NameUID != FSmartNameMapping::MaxUID)
+		{
+			Context.Curve.Set(NameUID, FinalDriverValue, (DestinationMode == EDrivenDestinationMode::MorphTarget) ? EAnimCurveFlags::ACF_DrivesMorphTarget : EAnimCurveFlags::ACF_DrivesMaterial);
+		}
 	}
 }
 
diff --git a/Engine/Source/Runtime/AnimGraphRuntime/Private/BoneControllers/AnimNode_CopyBoneDelta.cpp b/Engine/Source/Runtime/AnimGraphRuntime/Private/BoneControllers/AnimNode_CopyBoneDelta.cpp
--- a/Engine/Source/Runtime/AnimGraphRuntime/Private/BoneControllers/AnimNode_CopyBoneDelta.cpp
+++ b/Engine/Source/Runtime/AnimGraphRuntime/Private/BoneControllers/AnimNode_CopyBoneDelta.cpp
@@ -1,6 +1,7 @@
 // Copyright 1998-2016 Epic Games, Inc. All Rights Reserved.
 
 #include "AnimGraphRuntimePrivatePCH.h"
+#include "Animation/AnimTypes.h"
 #include "AnimNode_CopyBoneDelta.h"
 
 FAnimNode_CopyBoneDelta::FAnimNode_CopyBoneDelta()
diff --git a/Engine/Source/Runtime/AnimGraphRuntime/Private/BoneControllers/AnimNode_Fabrik.cpp b/Engine/Source/Runtime/AnimGraphRuntime/Private/BoneControllers/AnimNode_Fabrik.cpp
--- a/Engine/Source/Runtime/AnimGraphRuntime/Private/BoneControllers/AnimNode_Fabrik.cpp
+++ b/Engine/Source/Runtime/AnimGraphRuntime/Private/BoneControllers/AnimNode_Fabrik.cpp
@@ -1,6 +1,7 @@
 // Copyright 1998-2016 Epic Games, Inc. All Rights Reserved.
 
 #include "AnimGraphRuntimePrivatePCH.h"
+#include "Animation/AnimTypes.h"
 #include "BoneControllers/AnimNode_Fabrik.h"
 
 /////////////////////////////////////////////////////
diff --git a/Engine/Source/Runtime/AnimGraphRuntime/Public/AnimNodes/AnimNode_PoseBlendNode.h b/Engine/Source/Runtime/AnimGraphRuntime/Public/AnimNodes/AnimNode_PoseBlendNode.h
new file mode 100644
--- /dev/null
+++ b/Engine/Source/Runtime/AnimGraphRuntime/Public/AnimNodes/AnimNode_PoseBlendNode.h
@@ -0,0 +1,49 @@
+// Copyright 1998-2016 Epic Games, Inc. All Rights Reserved.
+
+#pragma once
+#include "Animation/AnimNodeBase.h"
+#include "Animation/PoseAsset.h"
+#include "Animation/AnimNode_AssetPlayerBase.h"
+#include "AnimNode_PoseBlendNode.generated.h"
+
+// Evaluates a point in an anim sequence, using a specific time input rather than advancing time internally.
+// Typically the playback position of the animation for this node will represent something other than time, like jump height.
+// This node will not trigger any notifies present in the associated sequence.
+USTRUCT()
+struct ANIMGRAPHRUNTIME_API FAnimNode_PoseBlendNode : public FAnimNode_AssetPlayerBase
+{
+	GENERATED_USTRUCT_BODY()
+public:
+	UPROPERTY(EditAnywhere, EditFixedSize, BlueprintReadWrite, Category = Links)
+	FPoseLink SourcePose;
+
+	// The animation sequence asset to evaluate
+	UPROPERTY(EditAnywhere, BlueprintReadWrite, Category = Settings, meta = (PinHiddenByDefault))
+	UPoseAsset* PoseAsset;
+
+public:	
+	FAnimNode_PoseBlendNode()
+	{
+	}
+
+	// FAnimNode_AssetPlayerBase interface
+	virtual float GetCurrentAssetTime() { return 0.f; }
+	virtual float GetCurrentAssetLength() { return 0.f; }
+	// End of FAnimNode_AssetPlayerBase interface
+
+	// FAnimNode_Base interface
+	virtual void Initialize(const FAnimationInitializeContext& Context) override;
+	virtual void CacheBones(const FAnimationCacheBonesContext& Context) override;
+	virtual void UpdateAssetPlayer(const FAnimationUpdateContext& Context) override;
+	virtual void Evaluate(FPoseContext& Output) override;
+	virtual void OverrideAsset(UAnimationAsset* NewAsset) override;
+	virtual void GatherDebugData(FNodeDebugData& DebugData) override;
+	// End of FAnimNode_Base interface
+
+	// FAnimNode_AssetPlayerBase Interface
+	virtual float GetAccumulatedTime() {return 0.f;}
+	virtual void SetAccumulatedTime(const float& NewTime) {}
+	virtual UAnimationAsset* GetAnimAsset() {return PoseAsset;}
+	// End of FAnimNode_AssetPlayerBase Interface
+};
+
diff --git a/Engine/Source/Runtime/AnimGraphRuntime/Public/AnimNodes/AnimNode_PoseByName.h b/Engine/Source/Runtime/AnimGraphRuntime/Public/AnimNodes/AnimNode_PoseByName.h
new file mode 100644
--- /dev/null
+++ b/Engine/Source/Runtime/AnimGraphRuntime/Public/AnimNodes/AnimNode_PoseByName.h
@@ -0,0 +1,54 @@
+// Copyright 1998-2016 Epic Games, Inc. All Rights Reserved.
+
+#pragma once
+#include "Animation/AnimNodeBase.h"
+#include "Animation/PoseAsset.h"
+#include "Animation/AnimNode_AssetPlayerBase.h"
+#include "AnimNode_PoseByName.generated.h"
+
+// Evaluates a point in an anim sequence, using a specific time input rather than advancing time internally.
+// Typically the playback position of the animation for this node will represent something other than time, like jump height.
+// This node will not trigger any notifies present in the associated sequence.
+USTRUCT()
+struct ANIMGRAPHRUNTIME_API FAnimNode_PoseByName : public FAnimNode_AssetPlayerBase
+{
+	GENERATED_USTRUCT_BODY()
+public:
+	// The animation sequence asset to evaluate
+	UPROPERTY(EditAnywhere, BlueprintReadWrite, Category = Settings, meta = (PinHiddenByDefault))
+	UPoseAsset* PoseAsset;
+
+	UPROPERTY(EditAnywhere, BlueprintReadWrite, Category = Settings, meta = (PinHiddenByDefault))
+	FName PoseName;
+
+public:	
+	FAnimNode_PoseByName()
+		: PoseUID(FSmartNameMapping::MaxUID)
+	{
+	}
+
+	// FAnimNode_AssetPlayerBase interface
+	virtual float GetCurrentAssetTime() { return 0.f; }
+	virtual float GetCurrentAssetLength() { return 0.f; }
+	// End of FAnimNode_AssetPlayerBase interface
+
+	// FAnimNode_Base interface
+	virtual void Initialize(const FAnimationInitializeContext& Context) override;
+	virtual void CacheBones(const FAnimationCacheBonesContext& Context) override;
+	virtual void UpdateAssetPlayer(const FAnimationUpdateContext& Context) override;
+	virtual void Evaluate(FPoseContext& Output) override;
+	virtual void OverrideAsset(UAnimationAsset* NewAsset) override;
+	virtual void GatherDebugData(FNodeDebugData& DebugData) override;
+	// End of FAnimNode_Base interface
+
+	// FAnimNode_AssetPlayerBase Interface
+	virtual float GetAccumulatedTime() {return 0.f;}
+	virtual void SetAccumulatedTime(const float& NewTime) {}
+	virtual UAnimationAsset* GetAnimAsset() {return PoseAsset;}
+	// End of FAnimNode_AssetPlayerBase Interface
+
+private:
+	// PoseUID
+	FSmartNameMapping::UID PoseUID;
+};
+
diff --git a/Engine/Source/Runtime/Core/Private/UObject/UnrealNames.cpp b/Engine/Source/Runtime/Core/Private/UObject/UnrealNames.cpp
--- a/Engine/Source/Runtime/Core/Private/UObject/UnrealNames.cpp
+++ b/Engine/Source/Runtime/Core/Private/UObject/UnrealNames.cpp
@@ -114,6 +114,18 @@ void FNameEntry::AppendNameToString( FString& String ) const
 	}
 }
 
+void FNameEntry::AppendNameToPathString(FString& String) const
+{
+	if (IsWide())
+	{
+		String /= WideName;
+	}
+	else
+	{
+		String /= AnsiName;
+	}
+}
+
 /**
  * @return length of name
  */
diff --git a/Engine/Source/Runtime/Core/Public/HAL/IConsoleManager.h b/Engine/Source/Runtime/Core/Public/HAL/IConsoleManager.h
--- a/Engine/Source/Runtime/Core/Public/HAL/IConsoleManager.h
+++ b/Engine/Source/Runtime/Core/Public/HAL/IConsoleManager.h
@@ -786,7 +786,7 @@ public:
 	T GetValueOnGameThread() const
 	{
 		// compiled out in shipping for performance (we can change in development later), if this get triggered you need to call GetValueOnRenderThread() or GetValueOnAnyThread(), the last one is a bit slower
-		ensure(GetShadowIndex() == 0);	// ensure to not block content creators, #if to optimize in shipping
+		checkCode(ensure(GetShadowIndex() == 0));	// ensure to not block content creators, #if to optimize in shipping
 		return ShadowedValue[0];
 	}
 
@@ -795,7 +795,7 @@ public:
 	{
 #if !defined(__clang__) // @todo Mac: figure out how to make this compile
 		// compiled out in shipping for performance (we can change in development later), if this get triggered you need to call GetValueOnGameThread() or GetValueOnAnyThread(), the last one is a bit slower
-		ensure(IsInParallelRenderingThread());	// ensure to not block content creators, #if to optimize in shipping
+		checkCode(ensure(IsInParallelRenderingThread()));	// ensure to not block content creators, #if to optimize in shipping
 #endif
 		return ShadowedValue[1];
 	}
@@ -816,7 +816,7 @@ private: // ----------------------------------------------------
 	{	
 		if (bForceGameThread)
 		{
-			ensure(!IsInActualRenderingThread());
+			checkCode(ensure(!IsInActualRenderingThread()));
 			return 0;
 		}
 		return IsInGameThread() ? 0 : 1;
diff --git a/Engine/Source/Runtime/Core/Public/Internationalization/Text.h b/Engine/Source/Runtime/Core/Public/Internationalization/Text.h
--- a/Engine/Source/Runtime/Core/Public/Internationalization/Text.h
+++ b/Engine/Source/Runtime/Core/Public/Internationalization/Text.h
@@ -428,7 +428,7 @@ public:
 	friend class FTextFormatHelper;
 	friend class FTextSnapshot;
 	friend class FTextInspector;
-	friend class UStruct;
+	friend class UTextProperty;
 	friend class UGatherTextFromAssetsCommandlet;
 	friend class FTextHistory_NamedFormat;
 	friend class FTextHistory_ArgumentDataFormat;
diff --git a/Engine/Source/Runtime/Core/Public/ProfilingDebugging/ScopedTimers.h b/Engine/Source/Runtime/Core/Public/ProfilingDebugging/ScopedTimers.h
--- a/Engine/Source/Runtime/Core/Public/ProfilingDebugging/ScopedTimers.h
+++ b/Engine/Source/Runtime/Core/Public/ProfilingDebugging/ScopedTimers.h
@@ -52,7 +52,6 @@ public:
 	explicit FScopedDurationTimer(double& AccumulatorIn)
 		: FDurationTimer(AccumulatorIn)
 	{
-		Start();
 	}
 
 	/** Dtor, updating seconds with time delta. */
diff --git a/Engine/Source/Runtime/Core/Public/UObject/FrameworkObjectVersion.h b/Engine/Source/Runtime/Core/Public/UObject/FrameworkObjectVersion.h
--- a/Engine/Source/Runtime/Core/Public/UObject/FrameworkObjectVersion.h
+++ b/Engine/Source/Runtime/Core/Public/UObject/FrameworkObjectVersion.h
@@ -29,6 +29,12 @@ struct CORE_API FFrameworkObjectVersion
 		// which causes issues with undo. Restore the flag at this version
 		FixNonTransactionalPins,
 
+		// Create new struct for SmartName, and use that for CurveName
+		SmartNameRefactor,
+		
+		// Add Reference Skeleton to Rig
+		AddSourceReferenceSkeletonToRig, 
+
 		// -----<new versions can be added above this line>-------------------------------------------------
 		VersionPlusOne,
 		LatestVersion = VersionPlusOne - 1
diff --git a/Engine/Source/Runtime/Core/Public/UObject/NameTypes.h b/Engine/Source/Runtime/Core/Public/UObject/NameTypes.h
--- a/Engine/Source/Runtime/Core/Public/UObject/NameTypes.h
+++ b/Engine/Source/Runtime/Core/Public/UObject/NameTypes.h
@@ -188,12 +188,19 @@ public:
 	 *
 	 * @param	String	String to append this name to
 	 */
-	void AppendNameToString( FString& String ) const;
+	CORE_API void AppendNameToString( FString& String ) const;
+
+	/**
+	* Appends this name entry to the passed in string, adding path separator between strings (with FString operator/).
+	*
+	* @param	String	String to append this name to
+	*/
+	CORE_API void AppendNameToPathString( FString& String ) const;
 
 	/**
 	 * @return length of name
 	 */
-	int32 GetNameLength() const;
+	CORE_API int32 GetNameLength() const;
 
 	/**
 	 * Compares name using the compare method provided.
diff --git a/Engine/Source/Runtime/CoreUObject/Private/Serialization/PropertyLocalizationDataGathering.cpp b/Engine/Source/Runtime/CoreUObject/Private/Serialization/PropertyLocalizationDataGathering.cpp
--- a/Engine/Source/Runtime/CoreUObject/Private/Serialization/PropertyLocalizationDataGathering.cpp
+++ b/Engine/Source/Runtime/CoreUObject/Private/Serialization/PropertyLocalizationDataGathering.cpp
@@ -162,6 +162,7 @@ void FPropertyLocalizationDataGatherer::GatherLocalizationDataFromChildTextPrope
 	const UTextProperty* const TextProperty = Cast<const UTextProperty>(Property);
 	const UArrayProperty* const ArrayProperty = Cast<const UArrayProperty>(Property);
 	const UMapProperty* const MapProperty = Cast<const UMapProperty>(Property);
+	const USetProperty* const SetProperty = Cast<const USetProperty>(Property);
 	const UStructProperty* const StructProperty = Cast<const UStructProperty>(Property);
 	const UObjectPropertyBase* const ObjectProperty = Cast<const UObjectPropertyBase>(Property);
 
@@ -210,6 +211,23 @@ void FPropertyLocalizationDataGatherer::GatherLocalizationDataFromChildTextPrope
 					GatherLocalizationDataFromChildTextProperties(PathToElement + FString::Printf(TEXT("(%d - Value)"), j), MapProperty->ValueProp, MapPairPtr + MapProperty->MapLayout.ValueOffset, ChildPropertyGatherTextFlags);
 				}
 			}
+			// Property is a set property.
+			else if (SetProperty)
+			{
+				// Iterate over all elements of the Set.
+				FScriptSetHelper ScriptSetHelper(SetProperty, ElementValueAddress);
+				const int32 ElementCount = ScriptSetHelper.Num();
+				for(int32 j = 0; j < ElementCount; ++j)
+				{
+					if (!ScriptSetHelper.IsValidIndex(j))
+					{
+						continue;
+					}
+
+					const uint8* ElementPtr = ScriptSetHelper.GetElementPtr(j);
+					GatherLocalizationDataFromChildTextProperties(PathToElement + FString::Printf(TEXT("(%d)"), j), SetProperty->ElementProp, ElementPtr + SetProperty->SetLayout.ElementOffset, ChildPropertyGatherTextFlags);
+				}
+			}
 			// Property is a struct property.
 			else if (StructProperty)
 			{
