Fix stale thumbnail actor pointer crash when recompiling blueprints

diff --git a/Engine/Source/Editor/UnrealEd/Private/ThumbnailHelpers.cpp b/Engine/Source/Editor/UnrealEd/Private/ThumbnailHelpers.cpp
--- a/Engine/Source/Editor/UnrealEd/Private/ThumbnailHelpers.cpp
+++ b/Engine/Source/Editor/UnrealEd/Private/ThumbnailHelpers.cpp
@@ -975,7 +975,7 @@ FClassActorThumbnailScene::FClassActorThumbnailScene()
 
 void FClassActorThumbnailScene::SpawnPreviewActor(UClass* InClass)
 {
-	if (PreviewActor)
+	if (PreviewActor.IsValid())
 	{
 		if (PreviewActor->GetClass() == InClass)
 		{
@@ -994,7 +994,7 @@ void FClassActorThumbnailScene::SpawnPreviewActor(UClass* InClass)
 		SpawnInfo.ObjectFlags = RF_Transient;
 		PreviewActor = GetWorld()->SpawnActor<AActor>(InClass, SpawnInfo);
 
-		if (PreviewActor)
+		if (PreviewActor.IsValid())
 		{
 			const FBoxSphereBounds Bounds = GetPreviewActorBounds();
 			const float BoundsZOffset = GetBoundsZOffset(Bounds);
@@ -1029,7 +1029,7 @@ bool FClassActorThumbnailScene::IsValidComponentForVisualization(UActorComponent
 FBoxSphereBounds FClassActorThumbnailScene::GetPreviewActorBounds() const
 {
 	FBoxSphereBounds Bounds(ForceInitToZero);
-	if (PreviewActor && PreviewActor->GetRootComponent())
+	if (PreviewActor.IsValid() && PreviewActor->GetRootComponent())
 	{
 		TArray<USceneComponent*> PreviewComponents;
 		PreviewActor->GetRootComponent()->GetChildrenComponents(true, PreviewComponents);
diff --git a/Engine/Source/Editor/UnrealEd/Public/ThumbnailHelpers.h b/Engine/Source/Editor/UnrealEd/Public/ThumbnailHelpers.h
--- a/Engine/Source/Editor/UnrealEd/Public/ThumbnailHelpers.h
+++ b/Engine/Source/Editor/UnrealEd/Public/ThumbnailHelpers.h
@@ -236,7 +236,7 @@ protected:
 
 private:
 
-	class AActor* PreviewActor;
+	TWeakObjectPtr<class AActor> PreviewActor;
 };
 
 class UNREALED_API FBlueprintThumbnailScene : public FClassActorThumbnailScene
