[NO NON-PII COMMIT MESSAGE]

diff --git a/Engine/Source/Editor/UnrealEd/Private/Kismet2/DebuggerCommands.cpp b/Engine/Source/Editor/UnrealEd/Private/Kismet2/DebuggerCommands.cpp
--- a/Engine/Source/Editor/UnrealEd/Private/Kismet2/DebuggerCommands.cpp
+++ b/Engine/Source/Editor/UnrealEd/Private/Kismet2/DebuggerCommands.cpp
@@ -21,6 +21,7 @@
 #include "PlatformInfo.h"
 
 #include "IHeadMountedDisplay.h"
+#include "IVREditorModule.h"
 
 //@TODO: Remove this dependency
 #include "Editor/LevelEditor/Public/LevelEditor.h"
@@ -1850,7 +1851,13 @@ void FInternalPlayWorldCommandCallbacks::TogglePlayPause_Clicked()
 
 bool FInternalPlayWorldCommandCallbacks::CanShowNonPlayWorldOnlyActions()
 {
-	return !HasPlayWorld();
+	// @todo vreditor: Don't display Simulate if we're currently in VR mode to prevent crash in blueprint (see UE-27984)
+	bool bIsInVREditor = false;
+	if (IVREditorModule::IsAvailable())
+	{
+		bIsInVREditor = IVREditorModule::Get().IsVREditorEnabled();
+	}
+	return (!HasPlayWorld() && !bIsInVREditor);
 }
 
 bool FInternalPlayWorldCommandCallbacks::CanShowVulkanNonPlayWorldOnlyActions()
