Copying //UE4/Dev-Blueprints to //UE4/Dev-Main (Source: //UE4/Dev-Blueprints @ 3025888)

diff --git a/Engine/Source/Editor/MaterialEditor/Private/FindInMaterial.cpp b/Engine/Source/Editor/MaterialEditor/Private/FindInMaterial.cpp
--- a/Engine/Source/Editor/MaterialEditor/Private/FindInMaterial.cpp
+++ b/Engine/Source/Editor/MaterialEditor/Private/FindInMaterial.cpp
@@ -11,22 +11,22 @@
 // FFindInMaterialResult
 
 FFindInMaterialResult::FFindInMaterialResult(const FString& InValue)
-:Value(InValue), DuplicationIndex(0), Class(NULL), Pin(NULL), GraphNode(NULL)
+:Value(InValue), DuplicationIndex(0), Class(NULL), Pin(), GraphNode(NULL)
 {
 }
 
 FFindInMaterialResult::FFindInMaterialResult(const FString& InValue, TSharedPtr<FFindInMaterialResult>& InParent, UClass* InClass, int InDuplicationIndex)
-: Parent(InParent), Value(InValue), DuplicationIndex(InDuplicationIndex), Class(InClass), Pin(NULL), GraphNode(NULL)
+: Parent(InParent), Value(InValue), DuplicationIndex(InDuplicationIndex), Class(InClass), Pin(), GraphNode(NULL)
 {
 }
 
 FFindInMaterialResult::FFindInMaterialResult(const FString& InValue, TSharedPtr<FFindInMaterialResult>& InParent, UEdGraphPin* InPin)
-: Parent(InParent), Value(InValue), DuplicationIndex(0), Class(InPin->GetClass()), Pin(InPin), GraphNode(NULL)
+: Parent(InParent), Value(InValue), DuplicationIndex(0), Class(nullptr), Pin(InPin), GraphNode(NULL)
 {
 }
 
 FFindInMaterialResult::FFindInMaterialResult(const FString& InValue, TSharedPtr<FFindInMaterialResult>& InParent, UEdGraphNode* InNode)
-: Parent(InParent), Value(InValue), DuplicationIndex(0), Class(InNode->GetClass()), Pin(NULL), GraphNode(InNode)
+: Parent(InParent), Value(InValue), DuplicationIndex(0), Class(InNode->GetClass()), Pin(), GraphNode(InNode)
 {
 	if (GraphNode.IsValid())
 	{
@@ -40,16 +40,16 @@ FReply FFindInMaterialResult::OnClick(TWeakPtr<class FMaterialEditor> MaterialEd
 	{
 		MaterialEditor.Pin()->JumpToNode(GraphNode.Get());
 	}
-	else if (Pin.IsValid())
+	else if (UEdGraphPin* ResolvedPin = Pin.Get())
 	{
-		MaterialEditor.Pin()->JumpToNode(Pin.Get()->GetOwningNode());
+		MaterialEditor.Pin()->JumpToNode(ResolvedPin->GetOwningNode());
 	}
 	return FReply::Handled();
 }
 
 FText FFindInMaterialResult::GetCategory() const
 {
-	if (Class == UEdGraphPin::StaticClass())
+	if (Class == nullptr && Pin.Get())
 	{
 		return LOCTEXT("PinCategory", "Pin");
 	}
@@ -64,13 +64,13 @@ TSharedRef<SWidget> FFindInMaterialResult::CreateIcon() const
 {
 	FSlateColor IconColor = FSlateColor::UseForeground();
 	const FSlateBrush* Brush = NULL;
-	if (Pin.IsValid())
+	if (UEdGraphPin* ResolvedPin = Pin.Get())
 	{
-		if (Pin->PinType.bIsArray)
+		if (ResolvedPin->PinType.bIsArray)
 		{
 			Brush = FEditorStyle::GetBrush(TEXT("GraphEditor.ArrayPinIcon"));
 		}
-		else if (Pin->PinType.bIsReference)
+		else if (ResolvedPin->PinType.bIsReference)
 		{
 			Brush = FEditorStyle::GetBrush(TEXT("GraphEditor.RefPinIcon"));
 		}
@@ -78,8 +78,8 @@ TSharedRef<SWidget> FFindInMaterialResult::CreateIcon() const
 		{
 			Brush = FEditorStyle::GetBrush(TEXT("GraphEditor.PinIcon"));
 		}
-		const UEdGraphSchema* Schema = Pin->GetSchema();
-		IconColor = Schema->GetPinTypeColor(Pin->PinType);
+		const UEdGraphSchema* Schema = ResolvedPin->GetSchema();
+		IconColor = Schema->GetPinTypeColor(ResolvedPin->PinType);
 	}
 	else if (GraphNode.IsValid())
 	{
@@ -101,27 +101,29 @@ FString FFindInMaterialResult::GetValueText() const
 {
 	FString DisplayVal;
 
-	if (Pin.IsValid()
-		&& (!Pin->DefaultValue.IsEmpty() || !Pin->AutogeneratedDefaultValue.IsEmpty() || Pin->DefaultObject || !Pin->DefaultTextValue.IsEmpty())
-		&& (Pin->Direction == EGPD_Input)
-		&& (Pin->LinkedTo.Num() == 0))
+	if (UEdGraphPin* ResolvedPin = Pin.Get())
 	{
-		if (Pin->DefaultObject)
+		if ((!ResolvedPin->DefaultValue.IsEmpty() || !ResolvedPin->AutogeneratedDefaultValue.IsEmpty() || ResolvedPin->DefaultObject || !ResolvedPin->DefaultTextValue.IsEmpty())
+			&& (ResolvedPin->Direction == EGPD_Input)
+			&& (ResolvedPin->LinkedTo.Num() == 0))
 		{
-			const AActor* DefaultActor = Cast<AActor>(Pin->DefaultObject);
-			DisplayVal = FString::Printf(TEXT("(%s)"), DefaultActor ? *DefaultActor->GetActorLabel() : *Pin->DefaultObject->GetName());
-		}
-		else if (!Pin->DefaultValue.IsEmpty())
-		{
-			DisplayVal = FString::Printf(TEXT("(%s)"), *Pin->DefaultValue);
-		}
-		else if (!Pin->AutogeneratedDefaultValue.IsEmpty())
-		{
-			DisplayVal = FString::Printf(TEXT("(%s)"), *Pin->AutogeneratedDefaultValue);
-		}
-		else if (!Pin->DefaultTextValue.IsEmpty())
-		{
-			DisplayVal = FString::Printf(TEXT("(%s)"), *Pin->DefaultTextValue.ToString());
+			if (ResolvedPin->DefaultObject)
+			{
+				const AActor* DefaultActor = Cast<AActor>(ResolvedPin->DefaultObject);
+				DisplayVal = FString::Printf(TEXT("(%s)"), DefaultActor ? *DefaultActor->GetActorLabel() : *ResolvedPin->DefaultObject->GetName());
+			}
+			else if (!ResolvedPin->DefaultValue.IsEmpty())
+			{
+				DisplayVal = FString::Printf(TEXT("(%s)"), *ResolvedPin->DefaultValue);
+			}
+			else if (!ResolvedPin->AutogeneratedDefaultValue.IsEmpty())
+			{
+				DisplayVal = FString::Printf(TEXT("(%s)"), *ResolvedPin->AutogeneratedDefaultValue);
+			}
+			else if (!ResolvedPin->DefaultTextValue.IsEmpty())
+			{
+				DisplayVal = FString::Printf(TEXT("(%s)"), *ResolvedPin->DefaultTextValue.ToString());
+			}
 		}
 	}
 
diff --git a/Engine/Source/Editor/MaterialEditor/Public/FindInMaterial.h b/Engine/Source/Editor/MaterialEditor/Public/FindInMaterial.h
--- a/Engine/Source/Editor/MaterialEditor/Public/FindInMaterial.h
+++ b/Engine/Source/Editor/MaterialEditor/Public/FindInMaterial.h
@@ -49,7 +49,7 @@ public:
 	UClass* Class;
 
 	/** The pin that this search result refers to */
-	TWeakObjectPtr<UEdGraphPin> Pin;
+	FEdGraphPinReference Pin;
 
 	/** The graph node that this search result refers to (if not by asset registry or UK2Node) */
 	TWeakObjectPtr<UEdGraphNode> GraphNode;
diff --git a/Engine/Source/Editor/Matinee/Private/MatineeTransaction.cpp b/Engine/Source/Editor/Matinee/Private/MatineeTransaction.cpp
--- a/Engine/Source/Editor/Matinee/Private/MatineeTransaction.cpp
+++ b/Engine/Source/Editor/Matinee/Private/MatineeTransaction.cpp
@@ -26,9 +26,7 @@ void FMatineeTransaction::SaveObject( UObject* Object )
 		Object->IsA( UInterpGroupInst::StaticClass() ) ||
 		Object->IsA( UInterpTrackInst::StaticClass() ) ||
 		Object->IsA( UMatineeOptions::StaticClass() ) ||
-		Object->IsA( UK2Node_MatineeController::StaticClass() ) || 
-		// or bIsMatineeControllerPin...
-		(Object->IsA<UEdGraphPin>() && Object->GetOuter()->IsA<UK2Node_MatineeController>()) )
+		Object->IsA( UK2Node_MatineeController::StaticClass() ) )
 	{
 		// Save the object.
 		new( Records )FObjectRecord( this, Object, NULL, 0, 0, 0, 0, NULL, NULL, NULL );
diff --git a/Engine/Source/Editor/NiagaraEditor/Private/NiagaraCompiler.cpp b/Engine/Source/Editor/NiagaraEditor/Private/NiagaraCompiler.cpp
--- a/Engine/Source/Editor/NiagaraEditor/Private/NiagaraCompiler.cpp
+++ b/Engine/Source/Editor/NiagaraEditor/Private/NiagaraCompiler.cpp
@@ -453,7 +453,7 @@ bool FNiagaraCompiler::MergeFunctionIntoMainGraph(UNiagaraNodeFunctionCall* Func
 			{
 				Callstack.Append(StackScript->GetPathName()); Callstack.Append(TEXT("\n"));
 			}
-			MessageLog.Error(TEXT("Reentrant function call!\n%s"), *Callstack);
+			MessageLog.Error(*FString::Printf(TEXT("Reentrant function call!\n%s"), *Callstack));
 			return false;
 		}
 		FunctionStack.Add(FuncScript);
diff --git a/Engine/Source/Editor/PropertyEditor/Private/Presentation/PropertyEditor/PropertyEditor.cpp b/Engine/Source/Editor/PropertyEditor/Private/Presentation/PropertyEditor/PropertyEditor.cpp
--- a/Engine/Source/Editor/PropertyEditor/Private/Presentation/PropertyEditor/PropertyEditor.cpp
+++ b/Engine/Source/Editor/PropertyEditor/Private/Presentation/PropertyEditor/PropertyEditor.cpp
@@ -329,13 +329,38 @@ void FPropertyEditor::SetEditConditionState( bool bShouldEnable )
 	for ( int32 ValueIdx = 0; ValueIdx < PropertyEditConditions.Num(); ValueIdx++ )
 	{
 		uint8* ValueAddr = PropertyEditConditions[ValueIdx].Address;
-		if ( XOR(bShouldEnable, PropertyEditConditions[ValueIdx].bNegateValue) )
-		{
-			EditConditionProperty->SetPropertyValue( ValueAddr, true );
-		}
-		else
+		const bool OldValue = EditConditionProperty->GetPropertyValue( ValueAddr );
+		const bool NewValue = XOR(bShouldEnable, PropertyEditConditions[ValueIdx].bNegateValue);
+		EditConditionProperty->SetPropertyValue( ValueAddr, NewValue );
+
+		// Propagate the value change to any instances if we're editing a template object.
+		FObjectPropertyNode* ObjectNode = PropertyNode->FindObjectItemParent();
+		if (ObjectNode != nullptr)
 		{
-			EditConditionProperty->SetPropertyValue( ValueAddr, false );
+			FPropertyNode* ParentNode = PropertyNode->GetParentNode();
+			check(ParentNode != nullptr);
+
+			for (int32 ObjIndex = 0; ObjIndex < ObjectNode->GetNumObjects(); ++ObjIndex)
+			{
+				TWeakObjectPtr<UObject> ObjectWeakPtr = ObjectNode->GetUObject(ObjIndex);
+				UObject* Object = ObjectWeakPtr.Get();
+				if (Object != nullptr && Object->IsTemplate())
+				{
+					TArray<UObject*> ArchetypeInstances;
+					Object->GetArchetypeInstances(ArchetypeInstances);
+					for (int32 InstanceIndex = 0; InstanceIndex < ArchetypeInstances.Num(); ++InstanceIndex)
+					{
+						// Only propagate if the current value on the instance matches the previous value on the template.
+						uint8* BaseOffset = ParentNode->GetValueAddress((uint8*)ArchetypeInstances[InstanceIndex]);
+						ValueAddr = EditConditionProperty->ContainerPtrToValuePtr<uint8>(BaseOffset);
+						const bool CurValue = EditConditionProperty->GetPropertyValue(ValueAddr);
+						if (OldValue == CurValue)
+						{
+							EditConditionProperty->SetPropertyValue(ValueAddr, NewValue);
+						}
+					}
+				}
+			}
 		}
 	}
 
diff --git a/Engine/Source/Editor/UnrealEd/Classes/Editor/TransBuffer.h b/Engine/Source/Editor/UnrealEd/Classes/Editor/TransBuffer.h
--- a/Engine/Source/Editor/UnrealEd/Classes/Editor/TransBuffer.h
+++ b/Engine/Source/Editor/UnrealEd/Classes/Editor/TransBuffer.h
@@ -114,7 +114,6 @@ public:
 	virtual bool DisableObjectSerialization() override;
 	virtual bool IsObjectSerializationEnabled() override { return DisallowObjectSerialization == 0; }
 	virtual void SetPrimaryUndoObject( UObject* Object ) override;
-	virtual ITransaction* CreateInternalTransaction() override;
 	virtual bool IsObjectInTransationBuffer( const UObject* Object ) const override;
 	virtual bool ContainsPieObject() const override;
 	virtual bool IsActive() override
diff --git a/Engine/Source/Editor/UnrealEd/Classes/Editor/Transactor.h b/Engine/Source/Editor/UnrealEd/Classes/Editor/Transactor.h
--- a/Engine/Source/Editor/UnrealEd/Classes/Editor/Transactor.h
+++ b/Engine/Source/Editor/UnrealEd/Classes/Editor/Transactor.h
@@ -41,13 +41,6 @@ struct FUndoSessionContext
  * warning: The undo buffer cannot be made persistent because of its dependence on offsets
  * of arrays from their owning UObjects.
  *
- * warning: Transactions which rely on Preload calls cannot be garbage collected
- * since references to objects point to the most recent version of the object, not
- * the ordinally correct version which was referred to at the time of serialization.
- * Therefore, Preload-sensitive transactions may only be performed using
- * a temporary UTransactor::CreateInternalTransaction transaction, not a
- * garbage-collectable UTransactor::Begin transaction.
- *
  * warning: UObject::Serialize implicitly assumes that class properties do not change
  * in between transaction resets.
  */
@@ -149,6 +142,18 @@ protected:
 		/** True if record should serialize data as binary blob (more compact). False to use tagged serialization (more robust) */
 		bool				bWantsBinarySerialization;
 
+		/** Copy of data that will be used when the transaction is flipped */
+		TArray<uint8> FlipData;
+
+		/** Copy of ReferencedObjects that will be used when the transaction is flipped */
+		TArray<FPersistentObjectRef> FlipReferencedObjects;
+
+		/** Copy of ReferencedNames that will be used when the transaction is flipped */
+		TArray<FName> FlipReferencedNames;
+
+		/** Copy of ObjectAnnotation that will be used when the transaction is flipped */
+		TSharedPtr<ITransactionObjectAnnotation> FlipObjectAnnotation;
+
 		// Constructors.
 		FObjectRecord()
 		{}
@@ -157,6 +162,8 @@ protected:
 		// Functions.
 		void SerializeContents( FArchive& Ar, int32 InOper );
 		void Restore( FTransaction* Owner );
+		void Save( FTransaction* Owner );
+		void Load( FTransaction* Owner );
 
 		/** Used by GC to collect referenced objects. */
 		void AddReferencedObjects( FReferenceCollector& Collector );
@@ -332,7 +339,7 @@ protected:
 	/** Used to prevent objects from being serialized to a transaction more than once. */
 	ObjectMapType			ObjectMap;
 
-	/** If true, on apply flip the direction of iteration over object records. */
+	/** If true, on apply flip the direction of iteration over object records. The only client for which this is false is FMatineeTransaction */
 	bool					bFlip;
 	/** Used to track direction to iterate over transaction's object records. Typically -1 for Undo, 1 for Redo */
 	int32					Inc;
@@ -594,7 +601,4 @@ class UTransactor : public UObject
 
 	/** @return True if this record contains a reference to a pie object */
 	virtual bool ContainsPieObject() const { return false; }
-
-	// @todo document
-	virtual ITransaction* CreateInternalTransaction() PURE_VIRTUAL(UTransactor::CreateInternalTransaction,return NULL;);
 };
diff --git a/Engine/Source/Editor/UnrealEd/Classes/Settings/EditorExperimentalSettings.h b/Engine/Source/Editor/UnrealEd/Classes/Settings/EditorExperimentalSettings.h
--- a/Engine/Source/Editor/UnrealEd/Classes/Settings/EditorExperimentalSettings.h
+++ b/Engine/Source/Editor/UnrealEd/Classes/Settings/EditorExperimentalSettings.h
@@ -67,6 +67,10 @@ public:
 	UPROPERTY(EditAnywhere, config, Category=Blueprints, meta=(DisplayName="Blueprint Performance Analysis Tools"))
 	bool bBlueprintPerformanceAnalysisTools;
 
+	/** Recent sample bias state. */
+	UPROPERTY(EditAnywhere, config, Category=Blueprints, meta=(DisplayName="Enable Blueprint Profiler Sample Bias"))
+	bool bEnableBlueprintProfilerRecentSampleBias;
+
 	/** Bias to weight prominance of newer samples against hositorical samples. */
 	UPROPERTY(EditAnywhere, config, Category=Blueprints, meta=(DisplayName="Blueprint Profiler Sample Bias", ClampMin=0.01, ClampMax=1.0))
 	float BlueprintProfilerRecentSampleBias;
diff --git a/Engine/Source/Editor/UnrealEd/Classes/UserDefinedStructure/UserDefinedStructEditorData.h b/Engine/Source/Editor/UnrealEd/Classes/UserDefinedStructure/UserDefinedStructEditorData.h
--- a/Engine/Source/Editor/UnrealEd/Classes/UserDefinedStructure/UserDefinedStructEditorData.h
+++ b/Engine/Source/Editor/UnrealEd/Classes/UserDefinedStructure/UserDefinedStructEditorData.h
@@ -107,7 +107,9 @@ public:
 	static void AddReferencedObjects(UObject* InThis, FReferenceCollector& Collector);
 
 	// UObject interface.
+	virtual TSharedPtr<ITransactionObjectAnnotation> GetTransactionAnnotation() const override;
 	virtual void PostEditUndo() override;
+	virtual void PostEditUndo(TSharedPtr<ITransactionObjectAnnotation> TransactionAnnotation) override;
 	virtual void PostLoadSubobjects(struct FObjectInstancingGraph* OuterInstanceGraph) override;
 	// End of UObject interface.
 
diff --git a/Engine/Source/Editor/UnrealEd/Private/ComponentTypeRegistry.cpp b/Engine/Source/Editor/UnrealEd/Private/ComponentTypeRegistry.cpp
--- a/Engine/Source/Editor/UnrealEd/Private/ComponentTypeRegistry.cpp
+++ b/Engine/Source/Editor/UnrealEd/Private/ComponentTypeRegistry.cpp
@@ -70,7 +70,20 @@ void FComponentTypeRegistryData::AddBasicShapeComponents(TArray<FComponentClassC
 		if (SMC)
 		{
 			const FString MaterialName = TEXT("/Engine/BasicShapes/BasicShapeMaterial.BasicShapeMaterial");
-			SMC->SetMaterial(0, FindOrLoadObject<UMaterial>(MaterialName));
+			UMaterial* MaterialAsset = FindOrLoadObject<UMaterial>(MaterialName);
+			SMC->SetMaterial(0, MaterialAsset);
+
+			// If the component object is an archetype (template), propagate the material setting to any instances, as instances
+			// of the archetype will end up being created BEFORE we are able to set the override material on the template object.
+			if (SMC->HasAnyFlags(RF_ArchetypeObject))
+			{
+				TArray<UObject*> ArchetypeInstances;
+				SMC->GetArchetypeInstances(ArchetypeInstances);
+				for (UObject* ArchetypeInstance : ArchetypeInstances)
+				{
+					CastChecked<UStaticMeshComponent>(ArchetypeInstance)->SetMaterial(0, MaterialAsset);
+				}
+			}
 		}
 	};
 
diff --git a/Engine/Source/Editor/UnrealEd/Private/EdGraphToken.cpp b/Engine/Source/Editor/UnrealEd/Private/EdGraphToken.cpp
new file mode 100644
--- /dev/null
+++ b/Engine/Source/Editor/UnrealEd/Private/EdGraphToken.cpp
@@ -0,0 +1,53 @@
+// Copyright 1998-2016 Epic Games, Inc. All Rights Reserved.
+
+#include "UnrealEd.h"
+#include "EdGraphToken.h"
+#include "Kismet2/CompilerResultsLog.h"
+
+TSharedRef<IMessageToken> FEdGraphToken::Create(const UObject* InObject, const FCompilerResultsLog* Log, UEdGraphNode*& OutSourceNode)
+{
+	const UObject* SourceNode = Log->FindSourceObject(InObject);
+	if (OutSourceNode == nullptr)
+	{
+		OutSourceNode = const_cast<UEdGraphNode*>(Cast<UEdGraphNode>(SourceNode));
+	}
+	return MakeShareable(new FEdGraphToken(Log->FindSourceObject(InObject), nullptr));
+}
+
+TSharedRef<IMessageToken> FEdGraphToken::Create(const UEdGraphPin* InPin, const FCompilerResultsLog* Log, UEdGraphNode*& OutSourceNode)
+{
+	const UObject* SourceNode = InPin ? Log->FindSourceObject(InPin->GetOwningNode()) : nullptr;
+	if (OutSourceNode == nullptr)
+	{
+		OutSourceNode = const_cast<UEdGraphNode*>(Cast<UEdGraphNode>(SourceNode));
+	}
+	return MakeShareable(new FEdGraphToken(SourceNode, Log->FindSourcePin(InPin)));
+}
+
+const UEdGraphPin* FEdGraphToken::GetPin() const
+{
+	return PinBeingReferenced.Get();
+}
+
+const UObject* FEdGraphToken::GetGraphObject() const
+{
+	return ObjectBeingReferenced.Get();
+}
+
+FEdGraphToken::FEdGraphToken(const UObject* InObject, const UEdGraphPin* InPin)
+	: ObjectBeingReferenced(InObject)
+	, PinBeingReferenced(InPin)
+{
+	if (InPin)
+	{
+		CachedText = FText::FromString(InPin->GetName());
+	}
+	else if (InObject)
+	{
+		CachedText = FText::FromString(InObject->GetName());
+	}
+	else
+	{
+		CachedText = NSLOCTEXT("MessageLog", "NoneObjectToken", "<None>");
+	}
+}
diff --git a/Engine/Source/Editor/UnrealEd/Private/EdGraphUtilities.cpp b/Engine/Source/Editor/UnrealEd/Private/EdGraphUtilities.cpp
--- a/Engine/Source/Editor/UnrealEd/Private/EdGraphUtilities.cpp
+++ b/Engine/Source/Editor/UnrealEd/Private/EdGraphUtilities.cpp
@@ -208,12 +208,22 @@ UEdGraph* FEdGraphUtilities::CloneGraph(UEdGraph* InSource, UObject* NewOuter, F
 
 			MessageLog->NotifyIntermediateObjectCreation(Dest, Source);
 
-			// During compilation, set cloned nodes to a non-conditional enabled state.
-			if (bCloningForCompile)
+			UEdGraphNode* SrcNode = Cast<UEdGraphNode>(Source);
+			UEdGraphNode* DstNode = Cast<UEdGraphNode>(Dest);
+			if (SrcNode && DstNode)
 			{
-				const UEdGraphNode* SrcNode = Cast<UEdGraphNode>(Source);
-				UEdGraphNode* DstNode = Cast<UEdGraphNode>(Dest);
-				if(SrcNode && DstNode)
+				// associate pins, no known case of StaticDuplicateObjectEx resulting in a different number of pins, but
+				// if that does happen we just associate as many pins as we can:
+				ensure(SrcNode->Pins.Num() == DstNode->Pins.Num());
+				for (int32 I = 0; I < SrcNode->Pins.Num() && I < DstNode->Pins.Num(); ++I)
+				{
+					if (ensure(DstNode->Pins[I] && SrcNode->Pins[I]))
+					{
+						MessageLog->NotifyIntermediatePinCreation(DstNode->Pins[I], SrcNode->Pins[I]);
+					}
+				}
+
+				if (bCloningForCompile)
 				{
 					DstNode->EnabledState = SrcNode->IsNodeEnabled() ? ENodeEnabledState::Enabled : ENodeEnabledState::Disabled;
 				}
@@ -453,8 +463,8 @@ void FEdGraphUtilities::FNodeVisitor::TraverseNodes(UEdGraphNode* Node)
 
 void FWeakGraphPinPtr::operator=(const class UEdGraphPin* Pin)
 {
-	PinObjectPtr = Pin;
-	if(PinObjectPtr.IsValid())
+	PinReference = Pin;
+	if(Pin && !Pin->IsPendingKill())
 	{
 		PinName = Pin->PinName;
 		NodeObjectPtr = Pin->GetOwningNode();
@@ -471,7 +481,7 @@ UEdGraphPin* FWeakGraphPinPtr::Get()
 	if(Node != NULL)
 	{
 		// If pin is no longer valid or has a different owner, attempt to fix up the reference
-		UEdGraphPin* Pin = PinObjectPtr.Get();
+		UEdGraphPin* Pin = PinReference.Get();
 		if(Pin == NULL || Pin->GetOuter() != Node)
 		{
 			for(auto PinIter = Node->Pins.CreateConstIterator(); PinIter; ++PinIter)
@@ -480,7 +490,7 @@ UEdGraphPin* FWeakGraphPinPtr::Get()
 				if(TestPin->PinName.Equals(PinName))
 				{
 					Pin = TestPin;
-					PinObjectPtr = Pin;
+					PinReference = Pin;
 					break;
 				}
 			}
diff --git a/Engine/Source/Editor/UnrealEd/Private/EditorTransaction.cpp b/Engine/Source/Editor/UnrealEd/Private/EditorTransaction.cpp
--- a/Engine/Source/Editor/UnrealEd/Private/EditorTransaction.cpp
+++ b/Engine/Source/Editor/UnrealEd/Private/EditorTransaction.cpp
@@ -139,29 +139,47 @@ void FTransaction::FObjectRecord::SerializeContents( FArchive& Ar, int32 InOper
 
 void FTransaction::FObjectRecord::Restore( FTransaction* Owner )
 {
+	// only used by FMatineeTransaction:
 	if( !bRestored )
 	{
 		bRestored = true;
-		TArray<uint8> FlipData;
-		TArray<FPersistentObjectRef> FlipReferencedObjects;
-		TArray<FName> FlipReferencedNames;
-		TSharedPtr<ITransactionObjectAnnotation> FlipObjectAnnotation;
-		if( Owner->bFlip )
-		{
-			FlipObjectAnnotation = Object->GetTransactionAnnotation();
-			FWriter Writer( FlipData, FlipReferencedObjects, FlipReferencedNames, bWantsBinarySerialization );
-			SerializeContents( Writer, -Oper );
-		}
+		check(!Owner->bFlip);
 		FTransaction::FObjectRecord::FReader Reader( Owner, Data, ReferencedObjects, ReferencedNames, bWantsBinarySerialization );
 		SerializeContents( Reader, Oper );
-		if( Owner->bFlip )
-		{
-			Exchange( ObjectAnnotation, FlipObjectAnnotation );
-			Exchange( Data, FlipData );
-			Exchange( ReferencedObjects, FlipReferencedObjects );
-			Exchange( ReferencedNames, FlipReferencedNames );
-			Oper *= -1;
-		}
+	}
+}
+
+void FTransaction::FObjectRecord::Save(FTransaction* Owner)
+{
+	// common undo/redo path, before applying undo/redo buffer we save current state:
+	check(Owner->bFlip);
+	if (!bRestored)
+	{
+		FlipData.Empty();
+		FlipReferencedObjects.Empty();
+		FlipReferencedNames.Empty();
+		FlipObjectAnnotation = TSharedPtr<ITransactionObjectAnnotation>();
+
+		FlipObjectAnnotation = Object->GetTransactionAnnotation();
+		FWriter Writer(FlipData, FlipReferencedObjects, FlipReferencedNames, bWantsBinarySerialization);
+		SerializeContents(Writer, -Oper);
+	}
+}
+
+void FTransaction::FObjectRecord::Load(FTransaction* Owner)
+{
+	// common undo/redo path, we apply the saved state and then swap it for the state we cached in ::Save above
+	check(Owner->bFlip);
+	if (!bRestored)
+	{
+		bRestored = true;
+		FTransaction::FObjectRecord::FReader Reader(Owner, Data, ReferencedObjects, ReferencedNames, bWantsBinarySerialization);
+		SerializeContents(Reader, Oper);
+		Exchange(ObjectAnnotation, FlipObjectAnnotation);
+		Exchange(Data, FlipData);
+		Exchange(ReferencedObjects, FlipReferencedObjects);
+		Exchange(ReferencedNames, FlipReferencedNames);
+		Oper *= -1;
 	}
 }
 
@@ -408,9 +426,24 @@ void FTransaction::Apply()
 			ChangedObjects.Add(Object, Record.ObjectAnnotation);
 		}
 	}
-	for( int32 i=Start; i!=End; i+=Inc )
+
+	if (bFlip)
 	{
-		Records[i].Restore( this );
+		for (int32 i = Start; i != End; i += Inc)
+		{
+			Records[i].Save(this);
+		}
+		for (int32 i = Start; i != End; i += Inc)
+		{
+			Records[i].Load(this);
+		}
+	}
+	else
+	{
+		for (int32 i = Start; i != End; i += Inc)
+		{
+			Records[i].Restore(this);
+		}
 	}
 
 	// An Actor's components must always get its PostEditUndo before the owning Actor so do a quick sort
@@ -861,11 +894,6 @@ bool UTransBuffer::DisableObjectSerialization()
 	return ++DisallowObjectSerialization == 0;
 }
 
-ITransaction* UTransBuffer::CreateInternalTransaction()
-{
-	return new FTransaction( TEXT("Internal") );
-}
-
 
 SIZE_T UTransBuffer::GetUndoSize() const
 {
diff --git a/Engine/Source/Editor/UnrealEd/Private/Factories/EditorFactories.cpp b/Engine/Source/Editor/UnrealEd/Private/Factories/EditorFactories.cpp
--- a/Engine/Source/Editor/UnrealEd/Private/Factories/EditorFactories.cpp
+++ b/Engine/Source/Editor/UnrealEd/Private/Factories/EditorFactories.cpp
@@ -5411,7 +5411,7 @@ void FCustomizableTextObjectFactory::ProcessBuffer(UObject* InParent, EObjectFla
 	while( FParse::Line(&Buffer,StrLine) )
 	{
 		const TCHAR* Str = *StrLine;
-		if( GetBEGIN(&Str,TEXT("OBJECT")) )
+		if( GetBEGIN(&Str, TEXT("OBJECT")) || (NestedDepth == 0 && GetBEGIN(&Str, TEXT("ACTOR"))) )
 		{
 			++NestedDepth;
 			if (OmittedOuterObj > 0)
@@ -5471,7 +5471,7 @@ void FCustomizableTextObjectFactory::ProcessBuffer(UObject* InParent, EObjectFla
 					{
 						ObjDepth++;
 					}
-					else if( GetEND(&PropStr, TEXT("OBJECT")) )
+					else if( GetEND(&PropStr, TEXT("OBJECT")) || (ObjDepth == 1 && GetEND(&PropStr, TEXT("ACTOR"))) )
 					{
 						bEndLine = true;
 
@@ -5496,7 +5496,7 @@ void FCustomizableTextObjectFactory::ProcessBuffer(UObject* InParent, EObjectFla
 				NewObjects.Add(CreatedObject);
 			}
 		}
-		else if (GetEND(&Str, TEXT("OBJECT")))
+		else if (GetEND(&Str, TEXT("OBJECT")) || (NestedDepth == 1 && GetEND(&Str, TEXT("ACTOR"))))
 		{
 			--NestedDepth;
 		}
@@ -5536,7 +5536,7 @@ bool FCustomizableTextObjectFactory::CanCreateObjectsFromText( const FString& Te
 	while( FParse::Line(&Buffer,StrLine) )
 	{
 		const TCHAR* Str = *StrLine;
-		if( GetBEGIN(&Str,TEXT("OBJECT")) )
+		if( GetBEGIN(&Str,TEXT("OBJECT")) || (NestedDepth == 0 && GetBEGIN(&Str, TEXT("ACTOR"))) )
 		{
 			++NestedDepth;
 			if (OmittedOuterObj > 0)
@@ -5565,7 +5565,7 @@ bool FCustomizableTextObjectFactory::CanCreateObjectsFromText( const FString& Te
 				}
 			}
 		}
-		else if ( GetEND(&Str, TEXT("OBJECT")) )
+		else if ( GetEND(&Str, TEXT("OBJECT")) || (NestedDepth == 1 && GetEND(&Str, TEXT("ACTOR"))) )
 		{
 			--NestedDepth;
 		}
diff --git a/Engine/Source/Editor/UnrealEd/Private/Kismet2/BlueprintEditorUtils.cpp b/Engine/Source/Editor/UnrealEd/Private/Kismet2/BlueprintEditorUtils.cpp
--- a/Engine/Source/Editor/UnrealEd/Private/Kismet2/BlueprintEditorUtils.cpp
+++ b/Engine/Source/Editor/UnrealEd/Private/Kismet2/BlueprintEditorUtils.cpp
@@ -29,7 +29,7 @@
 #include "ClassViewerModule.h"
 #include "ClassViewerFilter.h"
 #include "Editor/Kismet/Public/FindInBlueprintManager.h"
-
+#include "InstancedReferenceSubobjectHelper.h"
 #include "BlueprintEditor.h"
 #include "BlueprintEditorSettings.h"
 #include "Editor/UnrealEd/Public/Kismet2/Kismet2NameValidators.h"
@@ -569,8 +569,16 @@ void FBlueprintEditorUtils::RefreshExternalBlueprintDependencyNodes(UBlueprint*
 					if (!bShouldRefresh)
 					{
 						UClass* OwnerClass = Struct->GetOwnerClass();
-						bShouldRefresh |= OwnerClass && 
-							(OwnerClass->IsChildOf(RefreshOnlyChild) || OwnerClass->GetAuthoritativeClass()->IsChildOf(RefreshOnlyChild));
+						if (ensureMsgf(!OwnerClass || !OwnerClass->GetClass()->IsChildOf<UBlueprintGeneratedClass>() || OwnerClass->ClassGeneratedBy
+							, TEXT("Malformed Blueprint class (%s) - bad node dependency, unable to determine if the %s node (%s) should be refreshed or not. Currently compiling: %s")
+							, *OwnerClass->GetName()
+							, *Node->GetClass()->GetName()
+							, *Node->GetPathName()
+							, *Blueprint->GetName()) )
+						{
+							bShouldRefresh |= OwnerClass &&
+								(OwnerClass->IsChildOf(RefreshOnlyChild) || OwnerClass->GetAuthoritativeClass()->IsChildOf(RefreshOnlyChild));
+						}						
 					}
 					if (bShouldRefresh)
 					{
@@ -686,6 +694,8 @@ void FBlueprintEditorUtils::PatchNewCDOIntoLinker(UObject* CDO, FLinkerLoad* Lin
 		// Patch the new CDO in, and update the Export.Object
 		CDO->SetLinker(Linker, ExportIndex);
 		Linker->ExportMap[ExportIndex].Object = CDO;
+
+		PatchCDOSubobjectsIntoExport(OldCDO, CDO);
 	}
 }
 
@@ -1853,32 +1863,56 @@ void FBlueprintEditorUtils::PatchCDOSubobjectsIntoExport(UObject* PreviousCDO, U
 {
 	if (PreviousCDO && NewCDO)
 	{
-		// Collect the instanced components in both the old and new CDOs
-		TArray<UObject*> OldComponents, NewComponents;
-		PreviousCDO->CollectDefaultSubobjects(OldComponents, true);
-		NewCDO->CollectDefaultSubobjects(NewComponents, true);
-
-		TMap<FName, UObject*> NewComponentsMap;
-		for (auto NewCompIt = NewComponents.CreateIterator(); NewCompIt; ++NewCompIt)
+		struct PatchCDOSubobjectsIntoExport_Impl
 		{
-			UObject* NewComponent = *NewCompIt;
-			if (NewComponent)
+			static void PatchSubObjects(UObject* OldObj, UObject* NewObj)
 			{
-				NewComponentsMap.Add(NewComponent->GetFName(), NewComponent);
-			}
-		}
+				TArray<UObject*> NewSubObjects;
+				GetObjectsWithOuter(NewObj, NewSubObjects, /*bIncludeNestedSubObjects =*/false);
 
-		// For all components common to both, patch the linker table with the new version of the component, so things that reference the default (e.g. InternalArchetypes) will have the updated version
-		for (auto OldCompIt = OldComponents.CreateIterator(); OldCompIt; ++OldCompIt)
-		{
-			UObject* OldComponent = (*OldCompIt);
-			const FName OldComponentName = OldComponent->GetFName();
-			UObject** NewComponentPtr = NewComponentsMap.Find(OldComponentName);
-			if (NewComponentPtr && *NewComponentPtr)
-			{
-				FLinkerLoad::PRIVATE_PatchNewObjectIntoExport(OldComponent, *NewComponentPtr);
+				TMap<FName, UObject*> SubObjLookupTable;
+				for (UObject* NewSubObj : NewSubObjects)
+				{
+					if (NewSubObj != nullptr)
+					{
+						SubObjLookupTable.Add(NewSubObj->GetFName(), NewSubObj);
+					}
+				}
+
+				TArray<UObject*> OldSubObjects;
+				GetObjectsWithOuter(OldObj, OldSubObjects, /*bIncludeNestedSubObjects =*/false);
+
+				for (UObject* OldSubObj : OldSubObjects)
+				{
+					if (UObject** NewSubObjPtr = SubObjLookupTable.Find(OldSubObj->GetFName()))
+					{
+						UObject* NewSubObj = *NewSubObjPtr;
+						if (NewSubObj->IsDefaultSubobject() && OldSubObj->IsDefaultSubobject())
+						{
+							FLinkerLoad::PRIVATE_PatchNewObjectIntoExport(OldSubObj, NewSubObj);
+
+							UClass* SubObjClass = OldSubObj->GetClass();
+							if (SubObjClass && SubObjClass->HasAnyClassFlags(CLASS_HasInstancedReference))
+							{
+								TSet<FInstancedSubObjRef> OldInstancedValues;
+								FFindInstancedReferenceSubobjectHelper::GetInstancedSubObjects(OldSubObj, OldInstancedValues);
+
+								for (const FInstancedSubObjRef& OldInstancedObj : OldInstancedValues)
+								{
+									if (UObject* NewInstancedObj = OldInstancedObj.PropertyPath.Resolve(NewSubObj))
+									{
+										FLinkerLoad::PRIVATE_PatchNewObjectIntoExport(OldInstancedObj, NewInstancedObj);
+									}
+								}
+							}
+						}
+
+						PatchSubObjects(OldSubObj, NewSubObj);
+					}
+				}
 			}
-		}
+		};
+		PatchCDOSubobjectsIntoExport_Impl::PatchSubObjects(PreviousCDO, NewCDO);
 		NewCDO->CheckDefaultSubobjects();
 	}
 }
@@ -2048,7 +2082,10 @@ void FBlueprintEditorUtils::PostDuplicateBlueprint(UBlueprint* Blueprint, bool b
 			GetAllNodesOfClass(Blueprint, AllGraphNodes);
 			for(auto& Node : AllGraphNodes)
 			{
-				Node->CreateNewGuid();
+				if (!FBlueprintDuplicationScopeFlags::HasAnyFlag(FBlueprintDuplicationScopeFlags::TheSameNodeGuid))
+				{
+					Node->CreateNewGuid();
+				}
 
 				// Some variable nodes must be fixed up on duplicate, this cannot wait for individual 
 				// node calls to PostDuplicate because it happens after compilation and will still result in errors
@@ -2245,7 +2282,7 @@ void FBlueprintEditorUtils::MarkBlueprintAsStructurallyModified(UBlueprint* Blue
 }
 
 // Blueprint has changed in some manner that invalidates the compiled data (link made/broken, default value changed, etc...)
-void FBlueprintEditorUtils::MarkBlueprintAsModified(UBlueprint* Blueprint)
+void FBlueprintEditorUtils::MarkBlueprintAsModified(UBlueprint* Blueprint, FPropertyChangedEvent PropertyChangedEvent)
 {
 	Blueprint->bCachedDependenciesUpToDate = false;
 	if (Blueprint->Status != BS_BeingCreated)
@@ -2264,7 +2301,10 @@ void FBlueprintEditorUtils::MarkBlueprintAsModified(UBlueprint* Blueprint)
 
 		Blueprint->Status = BS_Dirty;
 		Blueprint->MarkPackageDirty();
-		Blueprint->PostEditChange();
+		// Previously, PostEditChange() was called on the Blueprint which creates an empty FPropertyChangedEvent. In
+		// certain cases, we needed to be able to pass along specific FPropertyChangedEvent that initially triggered
+		// this call so that we could keep the Blueprint from refreshing under certain conditions.
+		Blueprint->PostEditChangeProperty(PropertyChangedEvent);
 	}
 }
 
@@ -5246,8 +5286,8 @@ bool FBlueprintEditorUtils::ValidateAllFunctionGraphs(UBlueprint* InBlueprint, U
 
 void FBlueprintEditorUtils::ValidateBlueprintVariableMetadata(FBPVariableDescription& VarDesc)
 {
-	// Remove bitflag enum type metadata if the enum type is no longer a bitflags type.
-	if (VarDesc.HasMetaData(FBlueprintMetadata::MD_Bitmask))
+	// Remove bitflag enum type metadata if the enum type name is missing or if the enum type is no longer a bitflags type.
+	if (VarDesc.HasMetaData(FBlueprintMetadata::MD_BitmaskEnum))
 	{
 		FString BitmaskEnumTypeName = VarDesc.GetMetaData(FBlueprintMetadata::MD_BitmaskEnum);
 		if (!BitmaskEnumTypeName.IsEmpty())
@@ -5258,6 +5298,10 @@ void FBlueprintEditorUtils::ValidateBlueprintVariableMetadata(FBPVariableDescrip
 				VarDesc.RemoveMetaData(FBlueprintMetadata::MD_BitmaskEnum);
 			}
 		}
+		else
+		{
+			VarDesc.RemoveMetaData(FBlueprintMetadata::MD_BitmaskEnum);
+		}
 	}
 }
 
@@ -6275,6 +6319,17 @@ void FBlueprintEditorUtils::UpdateComponentTemplates(UBlueprint* Blueprint)
 		{
 			ensure(Blueprint->ComponentTemplates.Contains(ActorComp));
 
+			// fix up AddComponent nodes that don't have their own unique template objects
+			if (ReferencedTemplates.Contains(ActorComp))
+			{
+				UE_LOG(LogBlueprint, Warning,
+					TEXT("Blueprint '%s' has an AddComponent node '%s' with a non-unique component template name (%s). Moving it to a new template object with a unique name. Re-save the Blueprint to remove this warning on the next load."),
+					*Blueprint->GetPathName(), *ComponentNode->GetPathName(), *ActorComp->GetName());
+
+				ComponentNode->MakeNewComponentTemplate();
+				ActorComp = ComponentNode->GetTemplateFromNode();
+			}
+
 			// fix up existing content to be sure these are flagged as archetypes and are transactional
 			ActorComp->SetFlags(RF_ArchetypeObject|RF_Transactional);	
 			ReferencedTemplates.Add(ActorComp);
@@ -6448,15 +6503,15 @@ void FBlueprintEditorUtils::UpdateStalePinWatches( UBlueprint* Blueprint )
 	TSet<UEdGraphPin*> AllPins;
 
 	// Find all unique pins being watched
-	for (auto PinIt = Blueprint->PinWatches.CreateIterator(); PinIt; ++PinIt)
+	for (auto PinIt = Blueprint->WatchedPins.CreateIterator(); PinIt; ++PinIt)
 	{
-		UEdGraphPin* Pin (*PinIt);
+		UEdGraphPin* Pin (PinIt->Get());
 		if (Pin == NULL)
 		{
 			continue;
 		}
 
-		UEdGraphNode* OwningNode = Cast<UEdGraphNode>(Pin->GetOuter());
+		UEdGraphNode* OwningNode = Pin->GetOwningNode();
 		// during node reconstruction, dead pins get moved to the transient 
 		// package (so just in case this blueprint got saved with dead pin watches)
 		if (OwningNode == NULL)
@@ -6473,13 +6528,14 @@ void FBlueprintEditorUtils::UpdateStalePinWatches( UBlueprint* Blueprint )
 	}
 
 	// Refresh watched pins with unique pins (throw away null or duplicate watches)
-	if (Blueprint->PinWatches.Num() != AllPins.Num())
+	if (Blueprint->WatchedPins.Num() != AllPins.Num())
 	{
-		Blueprint->PinWatches.Empty();
+		Blueprint->WatchedPins.Empty();
 		for (auto PinIt = AllPins.CreateIterator(); PinIt; ++PinIt)
 		{
-			Blueprint->PinWatches.Add(*PinIt);
+			Blueprint->WatchedPins.Add(*PinIt);
 		}
+
 		Blueprint->Status = BS_Dirty;
 	}
 }
diff --git a/Engine/Source/Editor/UnrealEd/Private/Kismet2/CompilerResultsLog.cpp b/Engine/Source/Editor/UnrealEd/Private/Kismet2/CompilerResultsLog.cpp
--- a/Engine/Source/Editor/UnrealEd/Private/Kismet2/CompilerResultsLog.cpp
+++ b/Engine/Source/Editor/UnrealEd/Private/Kismet2/CompilerResultsLog.cpp
@@ -33,11 +33,25 @@ void FBacktrackMap::NotifyIntermediateObjectCreation(UObject* NewObject, UObject
 	SourceBacktrackMap.Add(Cast<UObject const>(NewObject), SourceObject);
 }
 
+/** Update the pin source backtrack map to note that NewPin was most closely generated/caused by the SourcePin */
+void FBacktrackMap::NotifyIntermediatePinCreation(UEdGraphPin* NewPin, UEdGraphPin* SourcePin)
+{
+	check(NewPin->GetOwningNode() && SourcePin->GetOwningNode());
+	// Chase the source to make sure it's really a top-level ('source code') node
+	while (UEdGraphPin** SourceOfSource = PinSourceBacktrackMap.Find(SourcePin))
+	{
+		SourcePin = *SourceOfSource;
+	}
+
+	// Record the backtrack link
+	PinSourceBacktrackMap.Add(NewPin, SourcePin);
+}
+
 /** Returns the true source object for the passed in object */
 UObject* FBacktrackMap::FindSourceObject(UObject* PossiblyDuplicatedObject)
 {
 	UObject** RemappedIfExisting = SourceBacktrackMap.Find(Cast<UObject const>(PossiblyDuplicatedObject));
-	if (RemappedIfExisting != NULL)
+	if (RemappedIfExisting != nullptr)
 	{
 		return *RemappedIfExisting;
 	}
@@ -48,12 +62,12 @@ UObject* FBacktrackMap::FindSourceObject(UObject* PossiblyDuplicatedObject)
 	}
 }
 
-UObject const* FBacktrackMap::FindSourceObject(UObject const* PossiblyDuplicatedObject)
+UObject const* FBacktrackMap::FindSourceObject(UObject const* PossiblyDuplicatedObject) const
 {
-	UObject** RemappedIfExisting = SourceBacktrackMap.Find(PossiblyDuplicatedObject);
-	if (RemappedIfExisting != NULL)
+	UObject *const* RemappedIfExisting = SourceBacktrackMap.Find(PossiblyDuplicatedObject);
+	if (RemappedIfExisting != nullptr)
 	{
-		return Cast<UObject const>(*RemappedIfExisting);
+		return *RemappedIfExisting;
 	}
 	else
 	{
@@ -62,6 +76,41 @@ UObject const* FBacktrackMap::FindSourceObject(UObject const* PossiblyDuplicated
 	}
 }
 
+UEdGraphPin* FBacktrackMap::FindSourcePin(UEdGraphPin* PossiblyDuplicatedPin)
+{
+	UEdGraphPin** RemappedIfExisting = PinSourceBacktrackMap.Find(PossiblyDuplicatedPin);
+	if (RemappedIfExisting != nullptr)
+	{
+		return *RemappedIfExisting;
+	}
+	else
+	{
+		// Not in the map, maybe its owning node was duplicated - and then maybe he GUID matches
+		// some node on the original node:
+		if (UObject* OriginalOwner = FindSourceObject(PossiblyDuplicatedPin->GetOwningNode()))
+		{
+			if (UEdGraphNode* AsNode = Cast<UEdGraphNode>(OriginalOwner))
+			{
+				FGuid TargetGuid = PossiblyDuplicatedPin->PinId;
+				UEdGraphPin** ExistingPin = AsNode->Pins.FindByPredicate([TargetGuid](const UEdGraphPin* TargetPin) { return TargetPin && TargetPin->PinId == TargetGuid; });
+				if (ExistingPin != nullptr)
+				{
+					return *ExistingPin;
+				}
+			}
+		}
+		
+		// No source object found, just return itself:
+		return PossiblyDuplicatedPin;
+	}
+}
+
+UEdGraphPin const* FBacktrackMap::FindSourcePin(UEdGraphPin const* PossiblyDuplicatedPin) const
+{
+	return const_cast<FBacktrackMap*>(this)->FindSourcePin(const_cast<UEdGraphPin*>(PossiblyDuplicatedPin));
+}
+
+
 //////////////////////////////////////////////////////////////////////////
 // FCompilerResultsLog
 
@@ -187,117 +236,59 @@ void FCompilerResultsLog::NotifyIntermediateObjectCreation(UObject* NewObject, U
 	SourceBacktrackMap.NotifyIntermediateObjectCreation(NewObject, SourceObject);
 }
 
+void FCompilerResultsLog::NotifyIntermediatePinCreation(UEdGraphPin* NewPin, UEdGraphPin* SourcePPin)
+{
+	SourceBacktrackMap.NotifyIntermediatePinCreation(NewPin, SourcePPin);
+}
+
 /** Returns the true source object for the passed in object */
 UObject* FCompilerResultsLog::FindSourceObject(UObject* PossiblyDuplicatedObject)
 {
 	return SourceBacktrackMap.FindSourceObject(PossiblyDuplicatedObject);
 }
 
-UObject const* FCompilerResultsLog::FindSourceObject(UObject const* PossiblyDuplicatedObject)
+UObject const* FCompilerResultsLog::FindSourceObject(UObject const* PossiblyDuplicatedObject) const
 {
 	return SourceBacktrackMap.FindSourceObject(PossiblyDuplicatedObject);
 }
 
-/** Create a tokenized message record from a message containing @@ indicating where each UObject* in the ArgPtr list goes and place it in the MessageLog. */
-void FCompilerResultsLog::InternalLogMessage(const EMessageSeverity::Type& Severity, const TCHAR* Message, va_list ArgPtr)
+UEdGraphPin* FCompilerResultsLog::FindSourcePin(UEdGraphPin* PossiblyDuplicatedPin)
 {
-	UEdGraphNode* OwnerNode = nullptr;
-
-	// Create the tokenized message
-	TSharedRef<FTokenizedMessage> Line = FTokenizedMessage::Create( Severity );
-	Messages.Add(Line);
-
-	const TCHAR* DelimiterStr = TEXT("@@");
-	int32 DelimLength = FCString::Strlen(DelimiterStr);
-
-	const TCHAR* Start = Message;
-	if (Start && DelimLength)
-	{
-		while (const TCHAR* At = FCString::Strstr(Start, DelimiterStr))
-		{
-			// Found a delimiter, create a token from the preceding text
-			Line->AddToken( FTextToken::Create( FText::FromString( FString(At - Start, Start) ) ) );
-			Start += DelimLength + (At - Start);
-
-			// And read the object and add another token for the object
-			UObject* ObjectArgument = va_arg(ArgPtr, UObject*);
-
-			FText ObjText;
-			if (ObjectArgument)
-			{
-				// Remap object references to the source nodes
-				ObjectArgument = FindSourceObject(ObjectArgument);
-
-				if (ObjectArgument)
-				{
-					UEdGraphNode* Node = Cast<UEdGraphNode>(ObjectArgument);
-					const UEdGraphPin* Pin = (Node? nullptr : Cast<UEdGraphPin>(ObjectArgument));
-
-					//Get owner node reference, consider the first 
-					if (OwnerNode == nullptr)
-					{
-						OwnerNode = (Pin ? Pin->GetOwningNodeUnchecked() : Node);
-					}
-
-					if (ObjectArgument->GetOutermost() == GetTransientPackage())
-					{
-						ObjText = LOCTEXT("Transient", "(transient)");					
-					}
-					else if (Node != NULL)
-					{
-						ObjText = Node->GetNodeTitle(ENodeTitleType::ListView);
-					}
-					else if (Pin != NULL)
-					{
-						ObjText = Pin->GetDisplayName();
-					}
-					else
-					{
-						ObjText = FText::FromString( ObjectArgument->GetName() );
-					}
-				}
-				else
-				{
-					ObjText = LOCTEXT("None", "(none)");
-				}
-
-			}
-			else
-			{
-				ObjText = LOCTEXT("None", "(none)");
-			}
-			
-			Line->AddToken( FUObjectToken::Create( ObjectArgument, ObjText ) );
-		}
-		Line->AddToken( FTextToken::Create( FText::FromString( Start ) ) );
-	}
+	return SourceBacktrackMap.FindSourcePin(PossiblyDuplicatedPin);
+}
 
-	va_end(ArgPtr);
+const UEdGraphPin* FCompilerResultsLog::FindSourcePin(const UEdGraphPin* PossiblyDuplicatedPin) const
+{
+	return SourceBacktrackMap.FindSourcePin(PossiblyDuplicatedPin);
+}
 
-	// Register node error/warning.
-	AnnotateNode(OwnerNode, Line);
+void FCompilerResultsLog::InternalLogMessage(const TSharedRef<FTokenizedMessage>& Message, UEdGraphNode* SourceNode)
+{
+	const EMessageSeverity::Type Severity = Message->GetSeverity();
+	Messages.Add(Message);
+	AnnotateNode(SourceNode, Message);
 
-	if( !bSilentMode && (!bLogInfoOnly || (Severity == EMessageSeverity::Info)) )
+	if (!bSilentMode && (!bLogInfoOnly || (Severity == EMessageSeverity::Info)))
 	{
-		if(Severity == EMessageSeverity::CriticalError || Severity == EMessageSeverity::Error)
+		if (Severity == EMessageSeverity::CriticalError || Severity == EMessageSeverity::Error)
 		{
-			if(IsRunningCommandlet())
+			if (IsRunningCommandlet())
 			{
-				UE_LOG(LogBlueprint, Error, TEXT("Compiler %s from Source: %s"), *Line->ToText().ToString(), *SourcePath);
+				UE_LOG(LogBlueprint, Error, TEXT("Compiler %s from Source: %s"), *Message->ToText().ToString(), *SourcePath);
 			}
 			else
 			{
 				// in editor the compiler log is 'rich' and we don't need to annotate with the blueprint name:
-				UE_LOG(LogBlueprint, Error, TEXT("Compiler %s"), *Line->ToText().ToString());
+				UE_LOG(LogBlueprint, Error, TEXT("Compiler %s"), *Message->ToText().ToString());
 			}
 		}
-		else if(Severity == EMessageSeverity::Warning || Severity == EMessageSeverity::PerformanceWarning)
+		else if (Severity == EMessageSeverity::Warning || Severity == EMessageSeverity::PerformanceWarning)
 		{
-			UE_LOG(LogBlueprint, Warning, TEXT("Compiler %s"), *Line->ToText().ToString());
+			UE_LOG(LogBlueprint, Warning, TEXT("Compiler %s"), *Message->ToText().ToString());
 		}
 		else
 		{
-			UE_LOG(LogBlueprint, Log, TEXT("Compiler %s"), *Line->ToText().ToString());
+			UE_LOG(LogBlueprint, Log, TEXT("Compiler %s"), *Message->ToText().ToString());
 		}
 	}
 }
@@ -458,47 +449,6 @@ void FCompilerResultsLog::GetGlobalModuleCompilerDump(const FString& LogDump, EC
 	MessageLog.AddMessages(ParseCompilerLogDump(LogDump));
 }
 
-// Note: Message is not a fprintf string!  It should be preformatted, but can contain @@ to indicate object references, which are the varargs
-void FCompilerResultsLog::Error(const TCHAR* Message, ...)
-{
-	va_list ArgPtr;
-	va_start(ArgPtr, Message);
-	ErrorVA(Message, ArgPtr);
-}
-
-// Note: Message is not a fprintf string!  It should be preformatted, but can contain @@ to indicate object references, which are the varargs
-void FCompilerResultsLog::Warning(const TCHAR* Message, ...)
-{
-	va_list ArgPtr;
-	va_start(ArgPtr, Message);
-	WarningVA(Message, ArgPtr);
-}
-
-// Note: Message is not a fprintf string!  It should be preformatted, but can contain @@ to indicate object references, which are the varargs
-void FCompilerResultsLog::Note(const TCHAR* Message, ...)
-{
-	va_list ArgPtr;
-	va_start(ArgPtr, Message);
-	NoteVA(Message, ArgPtr);
-}
-
-void FCompilerResultsLog::ErrorVA(const TCHAR* Message, va_list ArgPtr)
-{
-	++NumErrors;
-	InternalLogMessage(EMessageSeverity::Error, Message, ArgPtr);
-}
-
-void FCompilerResultsLog::WarningVA(const TCHAR* Message, va_list ArgPtr)
-{
-	++NumWarnings;
-	InternalLogMessage(EMessageSeverity::Warning, Message, ArgPtr);
-}
-
-void FCompilerResultsLog::NoteVA(const TCHAR* Message, va_list ArgPtr)
-{
-	InternalLogMessage(EMessageSeverity::Info, Message, ArgPtr);
-}
-
 void FCompilerResultsLog::Append(FCompilerResultsLog const& Other)
 {
 	for (TSharedRef<FTokenizedMessage> const& Message : Other.Messages)
@@ -522,22 +472,31 @@ void FCompilerResultsLog::Append(FCompilerResultsLog const& Other)
 		UEdGraphNode* OwnerNode = nullptr;
 		for (TSharedRef<IMessageToken> const& Token : Message->GetMessageTokens())
 		{
-			if (Token->GetType() != EMessageToken::Object)
-			{
-				continue;
-			}
-			
-			FWeakObjectPtr ObjectPtr = ((FUObjectToken&)Token.Get()).GetObject();
-			if (!ObjectPtr.IsValid())
+			if (Token->GetType() == EMessageToken::Object)
 			{
-				continue;
+				FWeakObjectPtr ObjectPtr = ((FUObjectToken&)Token.Get()).GetObject();
+				if (!ObjectPtr.IsValid())
+				{
+					continue;
+				}
+				UObject* ObjectArgument = ObjectPtr.Get();
+
+				OwnerNode = Cast<UEdGraphNode>(ObjectArgument);
+				break;
 			}
-			UObject* ObjectArgument = ObjectPtr.Get();
-			
-			OwnerNode = Cast<UEdGraphNode>(ObjectArgument);
-			if (UEdGraphPin const* Pin = Cast<UEdGraphPin>(ObjectArgument))
+			else if (Token->GetType() == EMessageToken::EdGraph)
 			{
-				OwnerNode = Pin->GetOwningNodeUnchecked();
+				FEdGraphToken& AsEdGraphToken = ((FEdGraphToken&)Token.Get());
+				const UEdGraphPin* PinBeingReferenced = AsEdGraphToken.GetPin();
+				OwnerNode = Cast<UEdGraphNode>((UObject*)AsEdGraphToken.GetGraphObject());
+				if (OwnerNode == nullptr)
+				{
+					if (PinBeingReferenced)
+					{
+						OwnerNode = PinBeingReferenced->GetOwningNodeUnchecked();
+					}
+				}
+				break;
 			}
 		}
 		AnnotateNode(OwnerNode, Message);
diff --git a/Engine/Source/Editor/UnrealEd/Private/Kismet2/ComponentEditorUtils.cpp b/Engine/Source/Editor/UnrealEd/Private/Kismet2/ComponentEditorUtils.cpp
--- a/Engine/Source/Editor/UnrealEd/Private/Kismet2/ComponentEditorUtils.cpp
+++ b/Engine/Source/Editor/UnrealEd/Private/Kismet2/ComponentEditorUtils.cpp
@@ -87,25 +87,46 @@ protected:
 	virtual bool CanCreateClass(UClass* ObjectClass, bool& bOmitSubObjs) const override
 	{
 		// Only allow actor component types to be created
-		return ObjectClass->IsChildOf(UActorComponent::StaticClass());
+		bool bCanCreate = ObjectClass->IsChildOf(UActorComponent::StaticClass());
+
+		// Also allow actor types to pass, in order to enable proper creation of actor component types as subobjects. The actor instance will be discarded after processing.
+		if (!bCanCreate)
+		{
+			bCanCreate = ObjectClass->IsChildOf(AActor::StaticClass());
+		}
+
+		return bCanCreate;
 	}
 
 	virtual void ProcessConstructedObject(UObject* NewObject) override
 	{
 		check(NewObject);
 
-		// Add it to the new object map
-		NewObjectMap.Add(NewObject->GetFName(), Cast<UActorComponent>(NewObject));
+		TInlineComponentArray<UActorComponent*> ActorComponents;
+		if (UActorComponent* NewActorComponent = Cast<UActorComponent>(NewObject))
+		{
+			ActorComponents.Add(NewActorComponent);
+		}
+		else if (AActor* NewActor = Cast<AActor>(NewObject))
+		{
+			NewActor->GetComponents(ActorComponents);
+		}
 
-		// If this is a scene component and it has a parent
-		USceneComponent* SceneComponent = Cast<USceneComponent>(NewObject);
-		if (SceneComponent && SceneComponent->GetAttachParent())
+		for(UActorComponent* ActorComponent : ActorComponents)
 		{
-			// Add an entry to the child->parent name map
-			ParentMap.Add(NewObject->GetFName(), SceneComponent->GetAttachParent()->GetFName());
+			// Add it to the new object map
+			NewObjectMap.Add(ActorComponent->GetFName(), ActorComponent);
 
-			// Clear this so it isn't used when constructing the new SCS node
-			SceneComponent->SetupAttachment(nullptr);
+			// If this is a scene component and it has a parent
+			USceneComponent* SceneComponent = Cast<USceneComponent>(ActorComponent);
+			if (SceneComponent && SceneComponent->GetAttachParent())
+			{
+				// Add an entry to the child->parent name map
+				ParentMap.Add(ActorComponent->GetFName(), SceneComponent->GetAttachParent()->GetFName());
+
+				// Clear this so it isn't used when constructing the new SCS node
+				SceneComponent->SetupAttachment(nullptr);
+			}
 		}
 	}
 
diff --git a/Engine/Source/Editor/UnrealEd/Private/Kismet2/EnumEditorUtils.cpp b/Engine/Source/Editor/UnrealEd/Private/Kismet2/EnumEditorUtils.cpp
--- a/Engine/Source/Editor/UnrealEd/Private/Kismet2/EnumEditorUtils.cpp
+++ b/Engine/Source/Editor/UnrealEd/Private/Kismet2/EnumEditorUtils.cpp
@@ -317,29 +317,31 @@ void FEnumEditorUtils::BroadcastChanges(const UUserDefinedEnum* Enum, const TArr
 		}
 	}
 
-	for (TObjectIterator<UEdGraphPin> It(RF_Transient); It; ++It)
+	for (TObjectIterator<UEdGraphNode> It(RF_Transient); It; ++It)
 	{
-		UEdGraphPin* Pin = *It;
-		if (Pin && (Pin->PinType.PinSubCategory != UEdGraphSchema_K2::PSC_Bitmask) && (Enum == Pin->PinType.PinSubCategoryObject.Get()) && (EEdGraphPinDirection::EGPD_Input == Pin->Direction))
+		for (UEdGraphPin* Pin : It->Pins)
 		{
-			UK2Node* Node = Cast<UK2Node>(Pin->GetOuter());
-			if (FNodeValidatorHelper::IsValid(Node))
+			if (Pin && (Pin->PinType.PinSubCategory != UEdGraphSchema_K2::PSC_Bitmask) && (Enum == Pin->PinType.PinSubCategoryObject.Get()) && (EEdGraphPinDirection::EGPD_Input == Pin->Direction))
 			{
-				if (UBlueprint* Blueprint = Node->GetBlueprint())
+				UK2Node* Node = Cast<UK2Node>(Pin->GetOuter());
+				if (FNodeValidatorHelper::IsValid(Node))
 				{
-					if (INDEX_NONE == Enum->FindEnumIndex(*Pin->DefaultValue))
+					if (UBlueprint* Blueprint = Node->GetBlueprint())
 					{
-						Pin->Modify();
-						if (Blueprint->BlueprintType == BPTYPE_Interface)
-						{
-							Pin->DefaultValue = Enum->GetEnumName(0);
-						}
-						else
+						if (INDEX_NONE == Enum->FindEnumIndex(*Pin->DefaultValue))
 						{
-							Pin->DefaultValue = FEnumEditorUtilsHelper::InvalidName();
+							Pin->Modify();
+							if (Blueprint->BlueprintType == BPTYPE_Interface)
+							{
+								Pin->DefaultValue = Enum->GetEnumName(0);
+							}
+							else
+							{
+								Pin->DefaultValue = FEnumEditorUtilsHelper::InvalidName();
+							}
+							Node->PinDefaultValueChanged(Pin);
+							BlueprintsToRefresh.Add(Blueprint);
 						}
-						Node->PinDefaultValueChanged(Pin);
-						BlueprintsToRefresh.Add(Blueprint);
 					}
 				}
 			}
diff --git a/Engine/Source/Editor/UnrealEd/Private/Kismet2/Kismet2.cpp b/Engine/Source/Editor/UnrealEd/Private/Kismet2/Kismet2.cpp
--- a/Engine/Source/Editor/UnrealEd/Private/Kismet2/Kismet2.cpp
+++ b/Engine/Source/Editor/UnrealEd/Private/Kismet2/Kismet2.cpp
@@ -866,6 +866,8 @@ void FKismetEditorUtilities::CompileBlueprint(UBlueprint* BlueprintObj, bool bIs
 	{
 		BlueprintPackage->SetDirtyFlag(bStartedWithUnsavedChanges);
 	}	
+
+	UEdGraphPin::Purge();
 }
 
 /** Generates a blueprint skeleton only.  Minimal compile, no notifications will be sent, no GC, etc.  Only successful if there isn't already a skeleton generated */
@@ -919,6 +921,11 @@ void FKismetEditorUtilities::RecompileBlueprintBytecode(UBlueprint* BlueprintObj
 
 	FKismetCompilerOptions CompileOptions;
 	CompileOptions.CompileType = EKismetCompileType::BytecodeOnly;
+
+	// Determine if we want profiling data
+	IBlueprintProfilerInterface* ProfilerInterface = FModuleManager::GetModulePtr<IBlueprintProfilerInterface>("BlueprintProfiler");
+	CompileOptions.bAddInstrumentation = ProfilerInterface && ProfilerInterface->IsProfilerEnabled() && GetDefault<UEditorExperimentalSettings>()->bBlueprintPerformanceAnalysisTools;
+
 	{
 		FRecreateUberGraphFrameScope RecreateUberGraphFrameScope(BlueprintObj->GeneratedClass, true);
 		Compiler.CompileBlueprint(BlueprintObj, CompileOptions, Results, NULL, ObjLoaded);
@@ -1409,9 +1416,15 @@ UBlueprint* FKismetEditorUtilities::CreateBlueprintFromActor(const FName Bluepri
 				const auto CopyOptions = (EditorUtilities::ECopyOptions::Type)(EditorUtilities::ECopyOptions::OnlyCopyEditOrInterpProperties | EditorUtilities::ECopyOptions::PropagateChangesToArchetypeInstances);
 				EditorUtilities::CopyActorProperties(Actor, CDO, CopyOptions);
 
-				if (USceneComponent* Scene = CDO->GetRootComponent())
+				if (USceneComponent* DstSceneRoot = CDO->GetRootComponent())
 				{
-					FResetSceneComponentAfterCopy::Reset(Scene);
+					FResetSceneComponentAfterCopy::Reset(DstSceneRoot);
+
+					// Copy relative scale from source to target.
+					if (USceneComponent* SrcSceneRoot = Actor->GetRootComponent())
+					{
+						DstSceneRoot->RelativeScale3D = SrcSceneRoot->RelativeScale3D;
+					}
 				}
 			}
 
@@ -1819,6 +1832,16 @@ void FKismetEditorUtilities::BringKismetToFocusAttentionOnObject(const UObject*
 	}
 }
 
+void FKismetEditorUtilities::BringKismetToFocusAttentionOnPin(const UEdGraphPin* PinToFocusOn )
+{
+	TSharedPtr<IBlueprintEditor> BlueprintEditor = GetIBlueprintEditorForObject(PinToFocusOn->GetOwningNode(), true);
+	if (BlueprintEditor.IsValid())
+	{
+		BlueprintEditor->FocusWindow();
+		BlueprintEditor->JumpToPin(PinToFocusOn);
+	}
+}
+
 void FKismetEditorUtilities::ShowActorReferencesInLevelScript(const AActor* Actor)
 {
 	if (Actor != NULL)
diff --git a/Engine/Source/Editor/UnrealEd/Private/Kismet2/KismetDebugUtilities.cpp b/Engine/Source/Editor/UnrealEd/Private/Kismet2/KismetDebugUtilities.cpp
--- a/Engine/Source/Editor/UnrealEd/Private/Kismet2/KismetDebugUtilities.cpp
+++ b/Engine/Source/Editor/UnrealEd/Private/Kismet2/KismetDebugUtilities.cpp
@@ -968,20 +968,20 @@ bool FKismetDebugUtilities::CanWatchPin(const UBlueprint* Blueprint, const UEdGr
 
 bool FKismetDebugUtilities::IsPinBeingWatched(const UBlueprint* Blueprint, const UEdGraphPin* Pin)
 {
-	return Blueprint->PinWatches.Contains(const_cast<UEdGraphPin*>(Pin));
+	return Blueprint->WatchedPins.Contains(const_cast<UEdGraphPin*>(Pin));
 }
 
 void FKismetDebugUtilities::RemovePinWatch(UBlueprint* Blueprint, const UEdGraphPin* Pin)
 {
 	UEdGraphPin* NonConstPin = const_cast<UEdGraphPin*>(Pin);
-	Blueprint->PinWatches.Remove(NonConstPin);
+	Blueprint->WatchedPins.Remove(NonConstPin);
 	Blueprint->MarkPackageDirty();
 	Blueprint->PostEditChange();
 }
 
 void FKismetDebugUtilities::TogglePinWatch(UBlueprint* Blueprint, const UEdGraphPin* Pin)
 {
-	int32 ExistingWatchIndex = Blueprint->PinWatches.Find(const_cast<UEdGraphPin*>(Pin));
+	int32 ExistingWatchIndex = Blueprint->WatchedPins.Find(const_cast<UEdGraphPin*>(Pin));
 
 	if (ExistingWatchIndex != INDEX_NONE)
 	{
@@ -989,7 +989,7 @@ void FKismetDebugUtilities::TogglePinWatch(UBlueprint* Blueprint, const UEdGraph
 	}
 	else
 	{
-		Blueprint->PinWatches.Add(const_cast<UEdGraphPin*>(Pin));
+		Blueprint->WatchedPins.Add(const_cast<UEdGraphPin*>(Pin));
 		Blueprint->MarkPackageDirty();
 		Blueprint->PostEditChange();
 	}
@@ -997,7 +997,7 @@ void FKismetDebugUtilities::TogglePinWatch(UBlueprint* Blueprint, const UEdGraph
 
 void FKismetDebugUtilities::ClearPinWatches(UBlueprint* Blueprint)
 {
-	Blueprint->PinWatches.Empty();
+	Blueprint->WatchedPins.Empty();
 	Blueprint->MarkPackageDirty();
 	Blueprint->PostEditChange();
 }
diff --git a/Engine/Source/Editor/UnrealEd/Private/Kismet2/KismetReinstanceUtilities.cpp b/Engine/Source/Editor/UnrealEd/Private/Kismet2/KismetReinstanceUtilities.cpp
--- a/Engine/Source/Editor/UnrealEd/Private/Kismet2/KismetReinstanceUtilities.cpp
+++ b/Engine/Source/Editor/UnrealEd/Private/Kismet2/KismetReinstanceUtilities.cpp
@@ -1477,9 +1477,16 @@ void FBlueprintCompileReinstancer::ReplaceInstancesOfClass_Inner(TMap<UClass*, U
 					{
 						auto OldFlags = OldObject->GetFlags();
 
+						// If the old object was spawned from an archetype (i.e. not the CDO), we must use the new version of that archetype as the template object when constructing the new instance.
+						UObject* OldArchetype = OldObject->GetArchetype();
+						UObject* NewArchetype = OldToNewInstanceMap.FindRef(OldArchetype);
+
+						// Check that either this was an instance of the class directly, or we found a new archetype for it
+						check(OldArchetype == OldClass->GetDefaultObject() || NewArchetype);
+
 						FName OldName(OldObject->GetFName());
 						OldObject->Rename(NULL, OldObject->GetOuter(), REN_DoNotDirty | REN_DontCreateRedirectors);
-						NewUObject = NewObject<UObject>(OldObject->GetOuter(), NewClass, OldName);
+						NewUObject = NewObject<UObject>(OldObject->GetOuter(), NewClass, OldName, RF_NoFlags, NewArchetype);
 						check(NewUObject != nullptr);
 
 						auto FlagMask = RF_Public | RF_ArchetypeObject | RF_Transactional | RF_Transient | RF_TextExportTransient | RF_InheritableComponentTemplate; //TODO: what about RF_RootSet and RF_Standalone ?
diff --git a/Engine/Source/Editor/UnrealEd/Private/Kismet2/StructureEditorUtils.cpp b/Engine/Source/Editor/UnrealEd/Private/Kismet2/StructureEditorUtils.cpp
--- a/Engine/Source/Editor/UnrealEd/Private/Kismet2/StructureEditorUtils.cpp
+++ b/Engine/Source/Editor/UnrealEd/Private/Kismet2/StructureEditorUtils.cpp
@@ -25,6 +25,8 @@ FStructureEditorUtils::FStructEditorManager& FStructureEditorUtils::FStructEdito
 	return *EditorManager;
 }
 
+FStructureEditorUtils::EStructureEditorChangeInfo FStructureEditorUtils::FStructEditorManager::ActiveChange = FStructureEditorUtils::EStructureEditorChangeInfo::Unknown;
+
 //////////////////////////////////////////////////////////////////////////
 // FStructureEditorUtils
 UUserDefinedStruct* FStructureEditorUtils::CreateUserDefinedStruct(UObject* InParent, FName Name, EObjectFlags Flags)
@@ -270,7 +272,7 @@ bool FStructureEditorUtils::AddVariable(UUserDefinedStruct* Struct, const FEdGra
 		NewVar.bInvalidMember = false;
 		GetVarDesc(Struct).Add(NewVar);
 
-		OnStructureChanged(Struct);
+		OnStructureChanged(Struct, EStructureEditorChangeInfo::AddedVariable);
 		return true;
 	}
 	return false;
@@ -290,7 +292,7 @@ bool FStructureEditorUtils::RemoveVariable(UUserDefinedStruct* Struct, FGuid Var
 			GetVarDesc(Struct).RemoveAll(FFindByGuidHelper<FStructVariableDescription>(VarGuid));
 			if (OldNum != GetVarDesc(Struct).Num())
 			{
-				OnStructureChanged(Struct);
+				OnStructureChanged(Struct, EStructureEditorChangeInfo::RemovedVariable);
 				return true;
 			}
 		}
@@ -324,7 +326,7 @@ bool FStructureEditorUtils::RenameVariable(UUserDefinedStruct* Struct, FGuid Var
 				check(NULL == GetVarDesc(Struct).FindByPredicate(FFindByNameHelper<FStructVariableDescription>(NewName)))
 				VarDesc->VarName = NewName;
 			}
-			OnStructureChanged(Struct);
+			OnStructureChanged(Struct, EStructureEditorChangeInfo::RenamedVariable);
 			return true;
 		}
 	}
@@ -356,7 +358,7 @@ bool FStructureEditorUtils::ChangeVariableType(UUserDefinedStruct* Struct, FGuid
 				VarDesc->DefaultValue = FString();
 				VarDesc->SetPinType(NewType);
 
-				OnStructureChanged(Struct);
+				OnStructureChanged(Struct, EStructureEditorChangeInfo::VariableTypeChanged);
 				return true;
 			}
 		}
@@ -410,10 +412,13 @@ bool FStructureEditorUtils::ChangeVariableDefaultValue(UUserDefinedStruct* Struc
 		if (bAdvancedValidation)
 		{
 			const FScopedTransaction Transaction(LOCTEXT("ChangeVariableDefaultValue", "Change Variable Default Value"));
-			ModifyStructData(Struct);
+			
+			TGuardValue<FStructureEditorUtils::EStructureEditorChangeInfo> ActiveChangeGuard(FStructureEditorUtils::FStructEditorManager::ActiveChange, EStructureEditorChangeInfo::DefaultValueChanged);
 
+			ModifyStructData(Struct);
+			
 			VarDesc->DefaultValue = NewDefaultValue;
-			OnStructureChanged(Struct);
+			OnStructureChanged(Struct, EStructureEditorChangeInfo::DefaultValueChanged);
 			return true;
 		}
 	}
@@ -516,10 +521,12 @@ void FStructureEditorUtils::CompileStructure(UUserDefinedStruct* Struct)
 	}
 }
 
-void FStructureEditorUtils::OnStructureChanged(UUserDefinedStruct* Struct)
+void FStructureEditorUtils::OnStructureChanged(UUserDefinedStruct* Struct, EStructureEditorChangeInfo ChangeReason)
 {
 	if (Struct)
 	{
+		TGuardValue<FStructureEditorUtils::EStructureEditorChangeInfo> ActiveChangeGuard(FStructureEditorUtils::FStructEditorManager::ActiveChange, ChangeReason);
+
 		Struct->Status = EUserDefinedStructureStatus::UDSS_Dirty;
 		CompileStructure(Struct);
 		Struct->MarkPackageDirty();
@@ -683,7 +690,7 @@ bool FStructureEditorUtils::MoveVariable(UUserDefinedStruct* Struct, FGuid VarGu
 				ModifyStructData(Struct);
 
 				DescArray.Swap(Index, Index + (bMoveUp ? -1 : 1));
-				OnStructureChanged(Struct);
+				OnStructureChanged(Struct, EStructureEditorChangeInfo::MovedVariable);
 				return true;
 			}
 		}
@@ -845,7 +852,7 @@ struct FReinstanceDataTableHelper
 
 void FStructureEditorUtils::BroadcastPreChange(UUserDefinedStruct* Struct)
 {
-	FStructureEditorUtils::FStructEditorManager::Get().PreChange(Struct, EStructureEditorChangeInfo::Changed);
+	FStructureEditorUtils::FStructEditorManager::Get().PreChange(Struct, FStructureEditorUtils::FStructEditorManager::ActiveChange);
 	auto DataTables = FReinstanceDataTableHelper::GetTablesDependentOnStruct(Struct);
 	for (auto DataTable : DataTables)
 	{
@@ -860,7 +867,7 @@ void FStructureEditorUtils::BroadcastPostChange(UUserDefinedStruct* Struct)
 	{
 		DataTable->RestoreAfterStructChange();
 	}
-	FStructureEditorUtils::FStructEditorManager::Get().PostChange(Struct, EStructureEditorChangeInfo::Changed);
+	FStructureEditorUtils::FStructEditorManager::Get().PostChange(Struct, FStructureEditorUtils::FStructEditorManager::ActiveChange);
 }
 
 #undef LOCTEXT_NAMESPACE
diff --git a/Engine/Source/Editor/UnrealEd/Private/MaterialGraphNode.cpp b/Engine/Source/Editor/UnrealEd/Private/MaterialGraphNode.cpp
--- a/Engine/Source/Editor/UnrealEd/Private/MaterialGraphNode.cpp
+++ b/Engine/Source/Editor/UnrealEd/Private/MaterialGraphNode.cpp
@@ -72,13 +72,7 @@ void UMaterialGraphNode::RecreateAndLinkNode()
 		Pin->Modify();
 		Pin->BreakAllPinLinks();
 
-#if 0
-		UEdGraphNode::ReturnPinToPool(Pin);
-#else
-		Pin->Rename(NULL, GetTransientPackage(), REN_None);
-		Pin->RemoveFromRoot();
-		Pin->MarkPendingKill();
-#endif
+		UEdGraphNode::DestroyPin(Pin);
 	}
 	Pins.Empty();
 
diff --git a/Engine/Source/Editor/UnrealEd/Private/MaterialGraphNode_Base.cpp b/Engine/Source/Editor/UnrealEd/Private/MaterialGraphNode_Base.cpp
--- a/Engine/Source/Editor/UnrealEd/Private/MaterialGraphNode_Base.cpp
+++ b/Engine/Source/Editor/UnrealEd/Private/MaterialGraphNode_Base.cpp
@@ -5,6 +5,8 @@
 =============================================================================*/
 
 #include "UnrealEd.h"
+
+#include "EdGraph/EdGraphPin.h"
 #include "MaterialShared.h"
 
 /////////////////////////////////////////////////////
@@ -244,13 +246,7 @@ void UMaterialGraphNode_Base::ReconstructNode()
 		OldPin->Modify();
 		OldPin->BreakAllPinLinks();
 
-#if 0
-		UEdGraphNode::ReturnPinToPool(OldPin);
-#else
-		OldPin->Rename(NULL, GetTransientPackage(), REN_None);
-		OldPin->RemoveFromRoot();
-		OldPin->MarkPendingKill();
-#endif
+		UEdGraphNode::DestroyPin(OldPin);
 	}
 	OldPins.Empty();
 
diff --git a/Engine/Source/Editor/UnrealEd/Private/Settings/SettingsClasses.cpp b/Engine/Source/Editor/UnrealEd/Private/Settings/SettingsClasses.cpp
--- a/Engine/Source/Editor/UnrealEd/Private/Settings/SettingsClasses.cpp
+++ b/Engine/Source/Editor/UnrealEd/Private/Settings/SettingsClasses.cpp
@@ -70,6 +70,7 @@ UEditorExperimentalSettings::UEditorExperimentalSettings( const FObjectInitializ
 	, bUnifiedBlueprintEditor(true)
 	, bBlueprintableComponents(true)
 	, bBlueprintPerformanceAnalysisTools(false)
+	, bEnableBlueprintProfilerRecentSampleBias(false)
 	, BlueprintProfilerRecentSampleBias(0.2f)
 	, BlueprintProfilerEventThreshold(1.f)
 	, BlueprintProfilerExclNodeThreshold(0.2f)
@@ -111,6 +112,10 @@ void UEditorExperimentalSettings::PostEditChangeProperty( struct FPropertyChange
 			}
 		}
 	}
+	else if (Name == FName(TEXT("bEnableBlueprintProfilerRecentSampleBias")))
+	{
+		FScriptPerfData::EnableRecentSampleBias(bEnableBlueprintProfilerRecentSampleBias);
+	}
 	else if (Name == FName(TEXT("BlueprintProfilerRecentSampleBias")))
 	{
 		FScriptPerfData::SetRecentSampleBias(BlueprintProfilerRecentSampleBias);
@@ -121,7 +126,7 @@ void UEditorExperimentalSettings::PostEditChangeProperty( struct FPropertyChange
 	}
 	else if (Name == FName(TEXT("BlueprintProfilerExclNodeThreshold")))
 	{
-		FScriptPerfData::SetNodePerformanceThreshold(BlueprintProfilerExclNodeThreshold);
+		FScriptPerfData::SetExclusivePerformanceThreshold(BlueprintProfilerExclNodeThreshold);
 	}
 	else if (Name == FName(TEXT("BlueprintProfilerInclNodeThreshold")))
 	{
@@ -162,7 +167,6 @@ UEditorLoadingSavingSettings::UEditorLoadingSavingSettings( const FObjectInitial
 	bPromptBeforeAutoImporting = true;
 }
 
-
 // @todo thomass: proper settings support for source control module
 void UEditorLoadingSavingSettings::SccHackInitialize()
 {
diff --git a/Engine/Source/Editor/UnrealEd/Private/SoundCueGraphNode.cpp b/Engine/Source/Editor/UnrealEd/Private/SoundCueGraphNode.cpp
--- a/Engine/Source/Editor/UnrealEd/Private/SoundCueGraphNode.cpp
+++ b/Engine/Source/Editor/UnrealEd/Private/SoundCueGraphNode.cpp
@@ -97,7 +97,7 @@ void USoundCueGraphNode::RemoveInputPin(UEdGraphPin* InGraphPin)
 	{
 		if (InGraphPin == InputPins[InputIndex])
 		{
-			InGraphPin->BreakAllPinLinks();
+			InGraphPin->MarkPendingKill();
 			Pins.Remove(InGraphPin);
 			// also remove the SoundNode child node so ordering matches
 			SoundNode->Modify();
diff --git a/Engine/Source/Editor/UnrealEd/Private/SoundCueGraphNode_Base.cpp b/Engine/Source/Editor/UnrealEd/Private/SoundCueGraphNode_Base.cpp
--- a/Engine/Source/Editor/UnrealEd/Private/SoundCueGraphNode_Base.cpp
+++ b/Engine/Source/Editor/UnrealEd/Private/SoundCueGraphNode_Base.cpp
@@ -6,6 +6,8 @@
 
 #include "UnrealEd.h"
 
+#include "EdGraph/EdGraphPin.h"
+
 /////////////////////////////////////////////////////
 // USoundCueGraphNode_Base
 
@@ -180,13 +182,7 @@ void USoundCueGraphNode_Base::ReconstructNode()
 		OldPin->Modify();
 		OldPin->BreakAllPinLinks();
 
-#if 0
-		UEdGraphNode::ReturnPinToPool(OldPin);
-#else
-		OldPin->Rename(NULL, GetTransientPackage(), REN_None);
-		OldPin->RemoveFromRoot();
-		OldPin->MarkPendingKill();
-#endif
+		UEdGraphNode::DestroyPin(OldPin);
 	}
 	OldPins.Empty();
 }
diff --git a/Engine/Source/Editor/UnrealEd/Private/UserDefinedStructEditorData.cpp b/Engine/Source/Editor/UnrealEd/Private/UserDefinedStructEditorData.cpp
--- a/Engine/Source/Editor/UnrealEd/Private/UserDefinedStructEditorData.cpp
+++ b/Engine/Source/Editor/UnrealEd/Private/UserDefinedStructEditorData.cpp
@@ -45,6 +45,46 @@ void UUserDefinedStructEditorData::PostEditUndo()
 	FStructureEditorUtils::OnStructureChanged(GetOwnerStruct());
 }
 
+class FStructureTransactionAnnotation : public ITransactionObjectAnnotation
+{
+public:
+	FStructureTransactionAnnotation(FStructureEditorUtils::EStructureEditorChangeInfo ChangeInfo)
+		: ActiveChange(ChangeInfo)
+	{
+	}
+
+	virtual void AddReferencedObjects(FReferenceCollector& Collector) override { /** Don't need this functionality for now */ }
+
+	FStructureEditorUtils::EStructureEditorChangeInfo GetActiveChange()
+	{
+		return ActiveChange;
+	}
+
+protected:
+	FStructureEditorUtils::EStructureEditorChangeInfo ActiveChange;
+};
+
+TSharedPtr<ITransactionObjectAnnotation> UUserDefinedStructEditorData::GetTransactionAnnotation() const
+{
+	return MakeShareable(new FStructureTransactionAnnotation(FStructureEditorUtils::FStructEditorManager::ActiveChange));
+}
+
+void UUserDefinedStructEditorData::PostEditUndo(TSharedPtr<ITransactionObjectAnnotation> TransactionAnnotation)
+{
+	Super::PostEditUndo();
+	FStructureEditorUtils::EStructureEditorChangeInfo ActiveChange = FStructureEditorUtils::Unknown;
+
+	if (TransactionAnnotation.IsValid())
+	{
+		TSharedPtr<FStructureTransactionAnnotation> StructAnnotation = StaticCastSharedPtr<FStructureTransactionAnnotation>(TransactionAnnotation);
+		if (StructAnnotation.IsValid())
+		{
+			ActiveChange = StructAnnotation->GetActiveChange();
+		}
+	}
+	FStructureEditorUtils::OnStructureChanged(GetOwnerStruct(), ActiveChange);
+}
+
 void UUserDefinedStructEditorData::PostLoadSubobjects(FObjectInstancingGraph* OuterInstanceGraph)
 {
 	Super::PostLoadSubobjects(OuterInstanceGraph);
diff --git a/Engine/Source/Editor/UnrealEd/Public/EdGraphToken.h b/Engine/Source/Editor/UnrealEd/Public/EdGraphToken.h
new file mode 100644
--- /dev/null
+++ b/Engine/Source/Editor/UnrealEd/Public/EdGraphToken.h
@@ -0,0 +1,42 @@
+// Copyright 1998-2016 Epic Games, Inc. All Rights Reserved.
+
+#pragma once
+
+#include "TokenizedMessage.h"
+
+// Copyright 1998-2016 Epic Games, Inc. All Rights Reserved.
+
+#pragma once
+
+#include "TokenizedMessage.h"
+#include "WeakObjectPtr.h"
+
+/**
+ * A Message Log token that links to an elemnt (node or pin) in an EdGraph
+ */
+class FEdGraphToken : public IMessageToken
+{
+public:
+	/** Factory method, tokens can only be constructed as shared refs */
+	UNREALED_API static TSharedRef<IMessageToken> Create(const UObject* InObject, const FCompilerResultsLog* Log, UEdGraphNode*& OutSourceNode);
+	UNREALED_API static TSharedRef<IMessageToken> Create(const UEdGraphPin* InPin, const FCompilerResultsLog* Log, UEdGraphNode*& OutSourceNode);
+
+	/** Begin IMessageToken interface */
+	virtual EMessageToken::Type GetType() const override
+	{
+		return EMessageToken::EdGraph;
+	}
+
+	UNREALED_API const UEdGraphPin* GetPin() const;
+	UNREALED_API const UObject* GetGraphObject() const;
+
+private:
+	/** Private constructor */
+	FEdGraphToken( const UObject* InObject, const UEdGraphPin* InPin );
+
+	/** An object being referenced by this token, if any */
+	FWeakObjectPtr ObjectBeingReferenced;
+
+	FEdGraphPinReference PinBeingReferenced;
+};
+
diff --git a/Engine/Source/Editor/UnrealEd/Public/EdGraphUtilities.h b/Engine/Source/Editor/UnrealEd/Public/EdGraphUtilities.h
--- a/Engine/Source/Editor/UnrealEd/Public/EdGraphUtilities.h
+++ b/Engine/Source/Editor/UnrealEd/Public/EdGraphUtilities.h
@@ -161,7 +161,7 @@ struct FWeakGraphPinPtr
 	FORCEINLINE void Reset()
 	{
 		PinName = TEXT("");
-		PinObjectPtr.Reset();
+		PinReference = FEdGraphPinReference();
 		NodeObjectPtr.Reset();
 	}
 
@@ -169,7 +169,7 @@ struct FWeakGraphPinPtr
 	FORCEINLINE void operator=(const FWeakGraphPinPtr &OtherPinPtr)
 	{
 		PinName = OtherPinPtr.PinName;
-		PinObjectPtr = OtherPinPtr.PinObjectPtr;
+		PinReference = OtherPinPtr.PinReference;
 		NodeObjectPtr = OtherPinPtr.NodeObjectPtr;
 	}
 
@@ -180,7 +180,7 @@ struct FWeakGraphPinPtr
 	FORCEINLINE bool operator==(const FWeakGraphPinPtr &OtherPinPtr) const
 	{
 		return PinName.Equals(OtherPinPtr.PinName)
-			&& PinObjectPtr == OtherPinPtr.PinObjectPtr
+			&& PinReference == OtherPinPtr.PinReference
 			&& NodeObjectPtr == OtherPinPtr.NodeObjectPtr;
 	}
 
@@ -188,7 +188,7 @@ struct FWeakGraphPinPtr
 	FORCEINLINE bool operator!=(const FWeakGraphPinPtr &OtherPinPtr) const
 	{
 		return !PinName.Equals(OtherPinPtr.PinName)
-			|| PinObjectPtr != OtherPinPtr.PinObjectPtr
+			|| PinReference.Get() != OtherPinPtr.PinReference.Get()
 			|| NodeObjectPtr != OtherPinPtr.NodeObjectPtr;
 	}
 
@@ -218,7 +218,7 @@ private:
 	FString PinName;
 
 	/** Weak reference to the UEdGraphPin object */
-	TWeakObjectPtr<class UEdGraphPin> PinObjectPtr;
+	FEdGraphPinReference PinReference;
 
 	/** Weak reference to the UEdGraphNode object that owns the pin object */
 	TWeakObjectPtr<class UEdGraphNode> NodeObjectPtr;
diff --git a/Engine/Source/Editor/UnrealEd/Public/Kismet2/BlueprintEditorUtils.h b/Engine/Source/Editor/UnrealEd/Public/Kismet2/BlueprintEditorUtils.h
--- a/Engine/Source/Editor/UnrealEd/Public/Kismet2/BlueprintEditorUtils.h
+++ b/Engine/Source/Editor/UnrealEd/Public/Kismet2/BlueprintEditorUtils.h
@@ -234,8 +234,11 @@ public:
 	 * Blueprint has changed in some manner that invalidates the compiled data (link made/broken, default value changed, etc...)
 	 *  - Marks the blueprint as status unknown
 	 *  - Marks the package as dirty
+	 *
+	 * @param	Blueprint				The Blueprint to mark as modified
+	 * @param	PropertyChangedEvent	Used when marking the blueprint as modified due to a changed property (optional)
 	 */
-	static void MarkBlueprintAsModified(UBlueprint* Blueprint);
+	static void MarkBlueprintAsModified(UBlueprint* Blueprint, FPropertyChangedEvent PropertyChangedEvent = FPropertyChangedEvent(nullptr));
 
 	/** See whether or not the specified graph name / entry point name is unique */
 	static bool IsGraphNameUnique(UBlueprint* Blueprint, const FName& InName);
@@ -1411,6 +1414,7 @@ struct UNREALED_API FBlueprintDuplicationScopeFlags
 		TheSameTimelineGuid = 1 << 1,
 		// This flag is needed for C++ backend (while compiler validates graphs). The actual BPGC type is compatible with the original BPGC.
 		ValidatePinsUsingSourceClass = 1 << 2,
+		TheSameNodeGuid = 1 << 3,
 	};
 
 	static uint32 bStaticFlags;
diff --git a/Engine/Source/Editor/UnrealEd/Public/Kismet2/CompilerResultsLog.h b/Engine/Source/Editor/UnrealEd/Public/Kismet2/CompilerResultsLog.h
--- a/Engine/Source/Editor/UnrealEd/Public/Kismet2/CompilerResultsLog.h
+++ b/Engine/Source/Editor/UnrealEd/Public/Kismet2/CompilerResultsLog.h
@@ -7,6 +7,7 @@
 #include "BlueprintUtilities.h"
 #include "TokenizedMessage.h"
 #include "CompilationResult.h"
+#include "EdGraphToken.h"
 
 /** This class maps from final objects to their original source object, across cloning, autoexpansion, etc... */
 class UNREALED_API FBacktrackMap
@@ -14,16 +15,20 @@ class UNREALED_API FBacktrackMap
 protected:
 	// Maps from transient object created during compiling to original 'source code' object
 	TMap<UObject const*, UObject*> SourceBacktrackMap;
+	// Maps from transient pins created during compiling to original 'source pin' object
+	TMap<UEdGraphPin*, UEdGraphPin*> PinSourceBacktrackMap;
 public:
-	FBacktrackMap(){}
-	virtual ~FBacktrackMap(){}
-
 	/** Update the source backtrack map to note that NewObject was most closely generated/caused by the SourceObject */
 	void NotifyIntermediateObjectCreation(UObject* NewObject, UObject* SourceObject);
 
+	/** Update the pin source backtrack map to note that NewPin was most closely generated/caused by the SourcePin */
+	void NotifyIntermediatePinCreation(UEdGraphPin* NewPin, UEdGraphPin* SourcePin);
+
 	/** Returns the true source object for the passed in object */
 	UObject* FindSourceObject(UObject* PossiblyDuplicatedObject);
-	UObject const* FindSourceObject(UObject const* PossiblyDuplicatedObject);
+	UObject const* FindSourceObject(UObject const* PossiblyDuplicatedObject) const;
+	UEdGraphPin* FindSourcePin(UEdGraphPin* PossiblyDuplicatedPin);
+	UEdGraphPin const* FindSourcePin(UEdGraphPin const* PossiblyDuplicatedPin) const;
 };
 
 /** This class represents a log of compiler output lines (errors, warnings, and information notes), each of which can be a rich tokenized message */
@@ -121,20 +126,37 @@ public:
 		SourcePath = InSourcePath;
 	}
 
-	// Note: Message is not a fprintf string!  It should be preformatted, but can contain @@ to indicate object references, which are the varargs
-	void Error(const TCHAR* Message, ...);
-	void Warning(const TCHAR* Message, ...);
-	void Note(const TCHAR* Message, ...);
-	void ErrorVA(const TCHAR* Message, va_list ArgPtr);
-	void WarningVA(const TCHAR* Message, va_list ArgPtr);
-	void NoteVA(const TCHAR* Message, va_list ArgPtr);
+	// Note: @@ will re replaced by FEdGraphToken::Create
+	template<typename... Args>
+	void Error(const TCHAR* Format, Args... args)
+	{
+		++NumErrors;
+		TSharedRef<FTokenizedMessage> Line = FTokenizedMessage::Create(EMessageSeverity::Error);
+		InternalLogMessage(Format, Line, args...);
+	}
+
+	template<typename... Args>
+	void Warning(const TCHAR* Format, Args... args)
+	{
+		++NumWarnings;
+		TSharedRef<FTokenizedMessage> Line = FTokenizedMessage::Create(EMessageSeverity::Warning);
+		InternalLogMessage(Format, Line, args...);
+	}
+
+	template<typename... Args>
+	void Note(const TCHAR* Format, Args... args)
+	{
+		TSharedRef<FTokenizedMessage> Line = FTokenizedMessage::Create(EMessageSeverity::Info);
+		InternalLogMessage(Format, Line, args...);
+	}
 
 	/** Update the source backtrack map to note that NewObject was most closely generated/caused by the SourceObject */
 	void NotifyIntermediateObjectCreation(UObject* NewObject, UObject* SourceObject);
+	void NotifyIntermediatePinCreation(UEdGraphPin* NewObject, UEdGraphPin* SourceObject);
 
 	/** Returns the true source object for the passed in object */
 	UObject* FindSourceObject(UObject* PossiblyDuplicatedObject);
-	UObject const* FindSourceObject(UObject const* PossiblyDuplicatedObject);
+	UObject const* FindSourceObject(UObject const* PossiblyDuplicatedObject) const;
 
 	/** Returns the true source object for the passed in object; does type checking on the result */
 	template <typename T>
@@ -144,10 +166,13 @@ public:
 	}
 	
 	template <typename T>
-	T const* FindSourceObjectTypeChecked(UObject const* PossiblyDuplicatedObject)
+	T const* FindSourceObjectTypeChecked(UObject const* PossiblyDuplicatedObject) const
 	{
 		return CastChecked<T const>(FindSourceObject(PossiblyDuplicatedObject));
 	}
+
+	UEdGraphPin* FindSourcePin(UEdGraphPin* PossiblyDuplicatedPin);
+	const UEdGraphPin* FindSourcePin(const UEdGraphPin* PossiblyDuplicatedPin) const;
 	
 	void Append(FCompilerResultsLog const& Other);
 
@@ -167,9 +192,44 @@ protected:
 	/** Helper method to add a child event to the given parent event scope */
 	void AddChildEvent(TSharedPtr<FCompilerEvent>& ParentEventScope, TSharedRef<FCompilerEvent>& ChildEventScope);
 
-	/** Create a tokenized message record from a message containing @@ indicating where each UObject* in the ArgPtr list goes and place it in the MessageLog. */
-	void InternalLogMessage(const EMessageSeverity::Type& Severity, const TCHAR* Message, va_list ArgPtr);
-	
+	void InternalLogMessage(const TSharedRef<FTokenizedMessage>& Message, UEdGraphNode* SourceNode );
+
+	void Tokenize(const TCHAR* Text, FTokenizedMessage &OutMessage, UEdGraphNode*& OutSourceNode)
+	{
+		OutMessage.AddToken(FTextToken::Create(FText::FromString(Text)));
+	}
+
+	template<typename T, typename... Args>
+	void Tokenize(const TCHAR* Format, FTokenizedMessage &OutMessage, UEdGraphNode*& OutSourceNode, T First, Args... Rest)
+	{
+		// read to next "@@":
+		if (const TCHAR* DelimiterStr = FCString::Strstr(Format, TEXT("@@")))
+		{
+			OutMessage.AddToken(FTextToken::Create(FText::FromString(FString(DelimiterStr - Format, Format))));
+			OutMessage.AddToken(FEdGraphToken::Create(First, this, OutSourceNode));
+			const TCHAR* NextChunk = DelimiterStr + FCString::Strlen(TEXT("@@"));
+			if (*NextChunk)
+			{
+				Tokenize(NextChunk, OutMessage, OutSourceNode, Rest...);
+			}
+		}
+		else
+		{
+			Tokenize(Format, OutMessage, OutSourceNode);
+		}
+	}
+
+	template<typename... Args>
+	void InternalLogMessage(const TCHAR* Format, const TSharedRef<FTokenizedMessage>& Message, Args... args)
+	{
+		// Convention for SourceNode established by the original version of the compiler results log
+		// was to annotate the error on the first node we can find. I am preserving that behavior
+		// for this type safe, variadic version:
+		UEdGraphNode* SourceNode = nullptr;
+		Tokenize(Format, *Message, SourceNode, args...);
+		InternalLogMessage(Message, SourceNode);
+	}
+
 	/** */
 	void AnnotateNode(class UEdGraphNode* Node, TSharedRef<FTokenizedMessage> LogLine);
 
diff --git a/Engine/Source/Editor/UnrealEd/Public/Kismet2/KismetEditorUtilities.h b/Engine/Source/Editor/UnrealEd/Public/Kismet2/KismetEditorUtilities.h
--- a/Engine/Source/Editor/UnrealEd/Public/Kismet2/KismetEditorUtilities.h
+++ b/Engine/Source/Editor/UnrealEd/Public/Kismet2/KismetEditorUtilities.h
@@ -162,9 +162,12 @@ public:
 	/** Create a new Actor Blueprint and add the supplied asset to it. */
 	static UBlueprint* CreateBlueprintUsingAsset(UObject* Asset, bool bOpenInEditor);
 
-	/** Open a Kismet window, focusing on the specified object (either a pin, a node, or a graph).  Prefers existing windows, but will open a new application if required. */
+	/** Open a Kismet window, focusing on the specified object (either a node, or a graph).  Prefers existing windows, but will open a new application if required. */
 	static void BringKismetToFocusAttentionOnObject(const UObject* ObjectToFocusOn, bool bRequestRename=false);
 
+	/** Open a Kismet window, focusing on the specified pin.  Prefers existing windows, but will open a new application if required. */
+	static void BringKismetToFocusAttentionOnPin(const UEdGraphPin* PinToFocusOn );
+
 	/** Open level script kismet window and show any references to the selected actor */
 	static void ShowActorReferencesInLevelScript(const AActor* Actor);
 
@@ -270,6 +273,9 @@ public:
 
 	/** Unregisters a class from having AutoGeneratedDefaultEvent nodes or callbacks for OnBlueprintCreated */
 	static void UnregisterAutoBlueprintNodeCreation(void* InOwner);
+
+	/** Get IBlueprintEditor for given object, if it exists */
+	static TSharedPtr<class IBlueprintEditor> GetIBlueprintEditorForObject(const UObject* ObjectToFocusOn, bool bOpenEditor);
 private:
 	/** Stores whether we are already listening for kismet clicks */
 	static bool bIsListeningForClicksOnKismetLog;
@@ -283,9 +289,7 @@ private:
 	/** Mapping of classes to delegate callbacks when a Blueprint is created, occurs post Event node creation */
 	static TMultiMap<void*, FOnBlueprintCreatedData> OnBlueprintCreatedCallbacks;
 private:
-	/** Get IBlueprintEditor for given object, if it exists */
-	static TSharedPtr<class IBlueprintEditor> GetIBlueprintEditorForObject(const UObject* ObjectToFocusOn, bool bOpenEditor);
-
+	
 	/**
 	 * Attempts to decide whether a blueprint's parent class is suitable for tracking via analytics.
 	 *
diff --git a/Engine/Source/Editor/UnrealEd/Public/Kismet2/StructureEditorUtils.h b/Engine/Source/Editor/UnrealEd/Public/Kismet2/StructureEditorUtils.h
--- a/Engine/Source/Editor/UnrealEd/Public/Kismet2/StructureEditorUtils.h
+++ b/Engine/Source/Editor/UnrealEd/Public/Kismet2/StructureEditorUtils.h
@@ -9,7 +9,13 @@ class UNREALED_API FStructureEditorUtils
 public:
 	enum EStructureEditorChangeInfo
 	{
-		Changed,
+		Unknown,
+		AddedVariable,
+		RemovedVariable,
+		RenamedVariable,
+		VariableTypeChanged,
+		MovedVariable,
+		DefaultValueChanged,
 	};
 
 	class FStructEditorManager : public FListenerManager<UUserDefinedStruct, EStructureEditorChangeInfo>
@@ -21,6 +27,9 @@ public:
 		class UNREALED_API ListenerType : public InnerListenerType<FStructEditorManager>
 		{
 		};
+
+		/** The current reason why a structure is being updated */
+		UNREALED_API static EStructureEditorChangeInfo ActiveChange;
 	};
 
 	typedef FStructEditorManager::ListenerType INotifyOnStructChanged;
@@ -173,7 +182,7 @@ public:
 	static EStructureError IsStructureValid(const UScriptStruct* Struct, const UStruct* RecursionParent = NULL, FString* OutMsg = NULL);
 
 	/** called after UDS was changed by editor*/
-	static void OnStructureChanged(UUserDefinedStruct* Struct);
+	static void OnStructureChanged(UUserDefinedStruct* Struct, EStructureEditorChangeInfo ChangeReason = EStructureEditorChangeInfo::Unknown);
 
 	static void BroadcastPreChange(UUserDefinedStruct* Struct);
 	static void BroadcastPostChange(UUserDefinedStruct* Struct);
diff --git a/Engine/Source/Programs/UnrealHeaderTool/Private/CodeGenerator.cpp b/Engine/Source/Programs/UnrealHeaderTool/Private/CodeGenerator.cpp
--- a/Engine/Source/Programs/UnrealHeaderTool/Private/CodeGenerator.cpp
+++ b/Engine/Source/Programs/UnrealHeaderTool/Private/CodeGenerator.cpp
@@ -898,7 +898,7 @@ FString FNativeClassHeaderGenerator::GetOverriddenName(const UField* Item)
 	FString OverriddenName = Item->GetMetaData(TEXT("OverrideNativeName"));
 	if (!OverriddenName.IsEmpty())
 	{
-		return OverriddenName;
+		return OverriddenName.ReplaceCharWithEscapedChar();
 	}
 	return Item->GetName();
 }
@@ -1438,12 +1438,9 @@ void FNativeClassHeaderGenerator::ExportNativeGeneratedInitCode(FClass* Class, F
 		TMap<FName, FString>* MetaDataMap = UMetaData::GetMapForObject(Class);
 		{
 			FClassMetaData* ClassMetaData = GScriptHelper.FindClassData(Class);
-			if (MetaDataMap && ClassMetaData && Class)
+			if (MetaDataMap && ClassMetaData && ClassMetaData->bObjectInitializerConstructorDeclared)
 			{
-				if (!ClassMetaData->bObjectInitializerConstructorDeclared && ClassMetaData->bDefaultConstructorDeclared)
-				{
-					MetaDataMap->Add(FName(TEXT("OnlyDefaultConstructorDeclared")), FString());
-				}
+				MetaDataMap->Add(FName(TEXT("ObjectInitializerConstructorDeclared")), FString());
 			}
 		}
 
@@ -1551,11 +1548,19 @@ void FNativeClassHeaderGenerator::ExportNativeGeneratedInitCode(FClass* Class, F
 		FString OverriddenClassName = *FNativeClassHeaderGenerator::GetOverriddenName(Class);
 		const TCHAR* ClassNameCPP = NameLookupCPP.GetNameCPP(Class);
 
-		GeneratedFunctionText.Logf(TEXT("\tstatic FCompiledInDefer Z_CompiledInDefer_UClass_%s(%s, &%s::StaticClass, TEXT(\"%s\"), %s, %s, %s);\r\n"),
-			ClassNameCPP, *SingletonName, ClassNameCPP, bIsDynamic ? *OverriddenClassName : ClassNameCPP,
+		const FString InitSearchableValuesFunctionName = bIsDynamic ? Class->GetMetaData(TEXT("InitializeStaticSearchableValues")) : FString();
+		const FString InitSearchableValuesFunctionParam = InitSearchableValuesFunctionName.IsEmpty() ? FString(TEXT("nullptr")) :
+			FString::Printf(TEXT("&%s::%s"), ClassNameCPP, *InitSearchableValuesFunctionName);
+
+		GeneratedFunctionText.Logf(TEXT("\tstatic FCompiledInDefer Z_CompiledInDefer_UClass_%s(%s, &%s::StaticClass, TEXT(\"%s\"), %s, %s, %s, %s);\r\n"),
+			ClassNameCPP, 
+			*SingletonName, 
+			ClassNameCPP, 
+			bIsDynamic ? *OverriddenClassName : ClassNameCPP,
 			bIsDynamic ? TEXT("true") : TEXT("false"),
 			bIsDynamic ? *AsTEXT(FClass::GetTypePackageName(Class)) : TEXT("nullptr"),
-			bIsDynamic ? *AsTEXT(FNativeClassHeaderGenerator::GetOverriddenPathName(Class)) : TEXT("nullptr"));
+			bIsDynamic ? *AsTEXT(FNativeClassHeaderGenerator::GetOverriddenPathName(Class)) : TEXT("nullptr"),
+			*InitSearchableValuesFunctionParam);
 		
 		// Append base class' CRC at the end of the generated code, this will force update derived classes
 		// when base class changes during hot-reload.
@@ -3134,7 +3139,9 @@ void FNativeClassHeaderGenerator::ExportGeneratedEnumsInitCode(const TArray<UEnu
 			GeneratedEnumRegisterFunctionText.Logf(TEXT("\t\t\tTArray<TPair<FName, uint8>> EnumNames;\r\n"));
 			for (int32 Index = 0; Index < Enum->NumEnums(); Index++)
 			{
-				GeneratedEnumRegisterFunctionText.Logf(TEXT("\t\t\tEnumNames.Add(TPairInitializer<FName, uint8>(FName(TEXT(\"%s\")), %d));\r\n"), *Enum->GetNameByIndex(Index).ToString(), Enum->GetValueByIndex(Index));
+				const TCHAR* OverridenNameMetaDatakey = TEXT("OverrideName");
+				const FString KeyName = Enum->HasMetaData(OverridenNameMetaDatakey, Index) ? Enum->GetMetaData(OverridenNameMetaDatakey, Index) : Enum->GetNameByIndex(Index).ToString();
+				GeneratedEnumRegisterFunctionText.Logf(TEXT("\t\t\tEnumNames.Add(TPairInitializer<FName, uint8>(FName(TEXT(\"%s\")), %d));\r\n"), *KeyName, Enum->GetValueByIndex(Index));
 			}
 
 			FString EnumTypeStr;
@@ -3144,7 +3151,8 @@ void FNativeClassHeaderGenerator::ExportGeneratedEnumsInitCode(const TArray<UEnu
 			case UEnum::ECppForm::Namespaced: EnumTypeStr = TEXT("UEnum::ECppForm::Namespaced"); break;
 			case UEnum::ECppForm::EnumClass:  EnumTypeStr = TEXT("UEnum::ECppForm::EnumClass");  break;
 			}
-			GeneratedEnumRegisterFunctionText.Logf(TEXT("\t\t\tReturnEnum->SetEnums(EnumNames, %s);\r\n"), *EnumTypeStr);
+			const FString ParamAddMaxKeyIfMissing = FClass::IsDynamic(Enum) ? TEXT(", false") : TEXT("");
+			GeneratedEnumRegisterFunctionText.Logf(TEXT("\t\t\tReturnEnum->SetEnums(EnumNames, %s%s);\r\n"), *EnumTypeStr, *ParamAddMaxKeyIfMissing);
 			GeneratedEnumRegisterFunctionText.Logf(TEXT("\t\t\tReturnEnum->CppType = TEXT(\"%s\");\r\n"), *Enum->CppType);
 
 			FString Meta = GetMetaDataCodeForObject(Enum, TEXT("ReturnEnum"), TEXT("\t\t\t"));
@@ -4289,7 +4297,7 @@ void FNativeClassHeaderGenerator::ExportFunctionThunk(FUHTStringBuilder& RPCWrap
 
 		FUHTStringBuilder ReplacementText;
 		ReplacementText += TypeText;
-		
+
 		ApplyAlternatePropertyExportText(Param, ReplacementText);
 		TypeText = ReplacementText;
 
@@ -4760,6 +4768,19 @@ TArray<UFunction*> FNativeClassHeaderGenerator::ExportCallbackFunctions(FUnrealS
  */
 void FNativeClassHeaderGenerator::ApplyAlternatePropertyExportText(UProperty* Prop, FUHTStringBuilder& PropertyText)
 {
+	UArrayProperty* ArrayProperty = Cast<UArrayProperty>(Prop);
+	UByteProperty* InnerByteProperty = ArrayProperty ? Cast<UByteProperty>(ArrayProperty->Inner) : nullptr;
+	if (InnerByteProperty && InnerByteProperty->Enum && FClass::IsDynamic(InnerByteProperty->Enum))
+	{
+		const FString Original = InnerByteProperty->GetCPPType();
+		const FString RawByte = InnerByteProperty->GetCPPType(nullptr, EPropertyExportCPPFlags::CPPF_BlueprintCppBackend);
+		if (Original != RawByte)
+		{
+			PropertyText.ReplaceInline(*Original, *RawByte, ESearchCase::CaseSensitive);
+		}
+		return;
+	}
+
 	if (!bIsExportingForOffsetDeterminationOnly)
 	{
 		return;
@@ -5392,6 +5413,10 @@ void FNativeClassHeaderGenerator::ExportGeneratedCPP()
 	FString EnableDeprecationWarnings = FString(TEXT("PRAGMA_ENABLE_DEPRECATION_WARNINGS") LINE_TERMINATOR);
 	FString DisableDeprecationWarnings = FString(TEXT("PRAGMA_DISABLE_DEPRECATION_WARNINGS") LINE_TERMINATOR);
 	
+	const bool bNativizaedAssetsModule = (ModuleInfo->Name == TEXT("NativizedAssets"));
+	const FString DisableWarning4883(TEXT("#ifdef _MSC_VER") LINE_TERMINATOR TEXT("#pragma warning (push)") LINE_TERMINATOR TEXT("#pragma warning (disable : 4883)") LINE_TERMINATOR TEXT("#endif") LINE_TERMINATOR);
+	const FString EnableWarning4883(TEXT("#ifdef _MSC_VER") LINE_TERMINATOR TEXT("#pragma warning (pop)") LINE_TERMINATOR TEXT("#endif") LINE_TERMINATOR);
+
 	FString PkgName = FPackageName::GetShortName(Package);
 	FString PkgDir;
 	if (GeneratedFunctionDeclarations.Len() || CrossModuleGeneratedFunctionDeclarations.Len())
@@ -5427,7 +5452,17 @@ void FNativeClassHeaderGenerator::ExportGeneratedCPP()
 
 		FString CppPath = ModuleInfo->GeneratedCPPFilenameBase + (GeneratedFunctionBodyTextSplit.Num() > 1 ? *FString::Printf(TEXT(".%d.cpp"), FileIdx + 1) : TEXT(".cpp"));
 		const FString GeneratedLinkerFixupFunction = FString::Printf(TEXT("void EmptyLinkFunctionForGeneratedCode%d%s() {}") LINE_TERMINATOR, FileIdx + 1, *ModuleInfo->Name);
-		SaveHeaderIfChanged(*CppPath, *(GeneratedCPPPreamble + ModulePCHInclude + GeneratedCPPClassesIncludes + DisableDeprecationWarnings + GeneratedLinkerFixupFunction + FileText + GeneratedCPPEpilogue + EnableDeprecationWarnings));
+		SaveHeaderIfChanged(*CppPath, *(GeneratedCPPPreamble
+			+ ModulePCHInclude
+			+ GeneratedCPPClassesIncludes
+			+ (bNativizaedAssetsModule ? DisableWarning4883 : FString())
+			+ DisableDeprecationWarnings
+			+ GeneratedLinkerFixupFunction
+			+ FileText 
+			+ GeneratedCPPEpilogue 
+			+ EnableDeprecationWarnings
+			+ (bNativizaedAssetsModule ? EnableWarning4883 : FString())
+			));
 
 		if (GeneratedFunctionBodyTextSplit.Num() > 1)
 		{
diff --git a/Engine/Source/Programs/UnrealHeaderTool/Private/HeaderParser.cpp b/Engine/Source/Programs/UnrealHeaderTool/Private/HeaderParser.cpp
--- a/Engine/Source/Programs/UnrealHeaderTool/Private/HeaderParser.cpp
+++ b/Engine/Source/Programs/UnrealHeaderTool/Private/HeaderParser.cpp
@@ -1363,7 +1363,7 @@ UEnum* FHeaderParser::CompileEnum()
 	}
 
 	// Register the list of enum names.
-	if (!Enum->SetEnums(EnumNames, CppForm))
+	if (!Enum->SetEnums(EnumNames, CppForm, !FClass::IsDynamic(Enum)))
 	{
 		const FName MaxEnumItem      = *(Enum->GenerateEnumPrefix() + TEXT("_MAX"));
 		const int32 MaxEnumItemIndex = Enum->FindEnumIndex(MaxEnumItem);
@@ -6349,6 +6349,7 @@ struct FExposeOnSpawnValidator
 		case CPT_Name:
 		case CPT_Vector:
 		case CPT_Rotation:
+		case CPT_Interface:
 			ProperNativeType = true;
 		}
 
diff --git a/Engine/Source/Runtime/Core/Private/Internationalization/ICUText.cpp b/Engine/Source/Runtime/Core/Private/Internationalization/ICUText.cpp
--- a/Engine/Source/Runtime/Core/Private/Internationalization/ICUText.cpp
+++ b/Engine/Source/Runtime/Core/Private/Internationalization/ICUText.cpp
@@ -221,14 +221,6 @@ bool FText::FSortPredicate::operator()(const FText& A, const FText& B) const
 	return Implementation->Compare(A, B);
 }
 
-bool FText::IsLetter( const TCHAR Char )
-{
-	icu::UnicodeString PatternString = ICUUtilities::ConvertString(TEXT("[\\p{L}]"));
-	UErrorCode ICUStatus = U_ZERO_ERROR;
-	icu::UnicodeSet Uniscode(PatternString, ICUStatus);
-	return Uniscode.contains(Char) != 0;
-}
-
 bool FUnicodeChar::CodepointToString(const uint32 InCodepoint, FString& OutString)
 {
 	icu::UnicodeString CodepointString;
diff --git a/Engine/Source/Runtime/Core/Private/Internationalization/LegacyText.cpp b/Engine/Source/Runtime/Core/Private/Internationalization/LegacyText.cpp
--- a/Engine/Source/Runtime/Core/Private/Internationalization/LegacyText.cpp
+++ b/Engine/Source/Runtime/Core/Private/Internationalization/LegacyText.cpp
@@ -75,11 +75,6 @@ bool FText::FSortPredicate::operator()(const FText& A, const FText& B) const
 	return A.ToString() < B.ToString();
 }
 
-bool FText::IsLetter( const TCHAR Char )
-{
-	return (Char>='A' && Char<='Z') || (Char>='a' && Char<='z');
-}
-
 bool FUnicodeChar::CodepointToString(const uint32 InCodepoint, FString& OutString)
 {
 	return false;
diff --git a/Engine/Source/Runtime/Core/Public/Internationalization/Text.h b/Engine/Source/Runtime/Core/Public/Internationalization/Text.h
--- a/Engine/Source/Runtime/Core/Public/Internationalization/Text.h
+++ b/Engine/Source/Runtime/Core/Public/Internationalization/Text.h
@@ -308,9 +308,6 @@ public:
 	 */
 	static bool IsWhitespace( const TCHAR Char );
 
-	/** Checks to see if the given character is considered a letter */
-	static bool IsLetter( const TCHAR Char );
-
 	static void GetFormatPatternParameters(const FText& Pattern, TArray<FString>& ParameterNames);
 
 	static FText Format(FText Pattern, FFormatNamedArguments Arguments);
diff --git a/Engine/Source/Runtime/Core/Public/Logging/TokenizedMessage.h b/Engine/Source/Runtime/Core/Public/Logging/TokenizedMessage.h
--- a/Engine/Source/Runtime/Core/Public/Logging/TokenizedMessage.h
+++ b/Engine/Source/Runtime/Core/Public/Logging/TokenizedMessage.h
@@ -35,6 +35,7 @@ namespace EMessageToken
 		Text,
 		Tutorial,
 		URL,
+		EdGraph,
 	};
 }
 
@@ -572,4 +573,5 @@ private:
 
 	/** The name of the tutorial asset. */
 	FString TutorialAssetName;
-};
\ No newline at end of file
+};
+
diff --git a/Engine/Source/Runtime/Core/Public/Math/InterpCurvePoint.h b/Engine/Source/Runtime/Core/Public/Math/InterpCurvePoint.h
--- a/Engine/Source/Runtime/Core/Public/Math/InterpCurvePoint.h
+++ b/Engine/Source/Runtime/Core/Public/Math/InterpCurvePoint.h
@@ -411,3 +411,11 @@ inline void CurveFindIntervalBounds( const FInterpCurvePoint<FLinearColor>& Star
 {
 	CurveLinearColorFindIntervalBounds(Start, End, CurrentMin, CurrentMax);
 }
+
+// Native implementation of NOEXPORT FInterpCurvePoint structures
+typedef FInterpCurvePoint<float> FInterpCurvePointFloat;
+typedef FInterpCurvePoint<FVector2D> FInterpCurvePointVector2D;
+typedef FInterpCurvePoint<FVector> FInterpCurvePointVector;
+typedef FInterpCurvePoint<FQuat> FInterpCurvePointQuat;
+typedef FInterpCurvePoint<FTwoVectors> FInterpCurvePointTwoVectors;
+typedef FInterpCurvePoint<FLinearColor> FInterpCurvePointLinearColor;
\ No newline at end of file
diff --git a/Engine/Source/Runtime/Core/Public/UObject/BlueprintsObjectVersion.h b/Engine/Source/Runtime/Core/Public/UObject/BlueprintsObjectVersion.h
--- a/Engine/Source/Runtime/Core/Public/UObject/BlueprintsObjectVersion.h
+++ b/Engine/Source/Runtime/Core/Public/UObject/BlueprintsObjectVersion.h
@@ -11,6 +11,7 @@ struct CORE_API FBlueprintsObjectVersion
 		OverridenEventReferenceFixup,
 		CleanBlueprintFunctionFlags,
 		ArrayGetByRefUpgrade,
+		EdGraphPinOptimized,
 
 		// -----<new versions can be added above this line>-------------------------------------------------
 		VersionPlusOne,
diff --git a/Engine/Source/Runtime/CoreUObject/Private/Serialization/AsyncLoading.cpp b/Engine/Source/Runtime/CoreUObject/Private/Serialization/AsyncLoading.cpp
--- a/Engine/Source/Runtime/CoreUObject/Private/Serialization/AsyncLoading.cpp
+++ b/Engine/Source/Runtime/CoreUObject/Private/Serialization/AsyncLoading.cpp
@@ -2257,6 +2257,7 @@ EAsyncPackageState::Type FAsyncPackage::PostLoadDeferredObjects(double InTickSta
 	Result = (DeferredPostLoadIndex == DeferredPostLoadObjects.Num()) ? EAsyncPackageState::Complete : EAsyncPackageState::TimeOut;
 	if (Result == EAsyncPackageState::Complete)
 	{
+		TArray<UObject*> CDODefaultSubobjects;
 		// Clear async loading flags (we still want RF_Async, but EInternalObjectFlags::AsyncLoading can be cleared)
 		for (UObject* Object : DeferredFinalizeObjects)
 		{
@@ -2267,6 +2268,19 @@ EAsyncPackageState::Type FAsyncPackage::PostLoadDeferredObjects(double InTickSta
 			{
 				DynamicClass->GetDefaultObject(true);
 			}
+
+			if(Object->HasAnyFlags(RF_ClassDefaultObject))
+			{
+				Object->GetDefaultSubobjects(CDODefaultSubobjects);
+				for (UObject* SubObject : CDODefaultSubobjects)
+				{
+					if (SubObject && SubObject->HasAnyInternalFlags(EInternalObjectFlags::AsyncLoading))
+					{
+						SubObject->AtomicallyClearInternalFlags(EInternalObjectFlags::AsyncLoading);
+					}
+				}
+				CDODefaultSubobjects.Reset();
+			}
 		}
 
 		// Mark package as having been fully loaded and update load time.
