Copying //UE4/Dev-Sequencer to //UE4/Dev-Main (Source: //UE4/Dev-Sequencer @ 3048536)

diff --git a/Engine/Source/Editor/Sequencer/Private/Sequencer.h b/Engine/Source/Editor/Sequencer/Private/Sequencer.h
--- a/Engine/Source/Editor/Sequencer/Private/Sequencer.h
+++ b/Engine/Source/Editor/Sequencer/Private/Sequencer.h
@@ -98,12 +98,12 @@ public:
 	 */
 	void SetSelectionRangeStart();
 
+	/** Clear and reset the selection range. */
+	void ResetSelectionRange();
+
 	/** Select all keys that fall into the current selection range. */
 	void SelectKeysInSelectionRange();
 
-	/** Delete the time range and all the keys that fall into the selection range. */
-	void DeleteSelectionRange();
-
 public:
 
 	/**
@@ -254,8 +254,9 @@ public:
 	 * Called to delete all moviescene data from a node
 	 *
 	 * @param NodeToBeDeleted	Node with data that should be deleted
+	 * @return true if anything was deleted, otherwise false.
 	 */
-	virtual void OnRequestNodeDeleted( TSharedRef<const FSequencerDisplayNode> NodeToBeDeleted );
+	virtual bool OnRequestNodeDeleted( TSharedRef<const FSequencerDisplayNode> NodeToBeDeleted );
 
 	/** Zooms to the edges of all currently selected sections. */
 	void ZoomToSelectedSections();
@@ -316,11 +317,22 @@ public:
 	/** Extend the level editor viewport menu */
 	TSharedRef<FExtender> OnExtendLevelEditorViewMenu(const TSharedRef<FUICommandList> CommandList);
 
+	/** Extend the level viewport context menu */
+	void AddLevelViewportMenuExtender();
+	
+	/** Remove a previously registered extender for the level viewport context menu */
+	void RemoveLevelViewportMenuExtender();
+
+	/** Get the extender for the level viewport menu */
+	TSharedRef<FExtender> GetLevelViewportExtender(const TSharedRef<FUICommandList> CommandList, const TArray<AActor*> InActors);
+
+	void RecordSelectedActors();
+	
 	/** Create a menu entry we can use to toggle the transport controls */
 	void CreateTransportToggleMenuEntry(FMenuBuilder& MenuBuilder);
 
 	/** Functions to push on to the transport controls we use */
-	FReply OnPlay(bool bTogglePlay=true);
+	FReply OnPlay(bool bTogglePlay=true, float InPlayRate=1.f);
 	FReply OnRecord();
 	FReply OnStepForward();
 	FReply OnStepBackward();
@@ -430,6 +442,9 @@ public:
 	/** Moves all time data for the current scene onto a valid frame. */
 	void FixFrameTiming();
 
+	/** Exports the current scene and sequence to an fbx file. */
+	void ExportSceneAndSequence();
+
 public:
 	
 	/** Access the currently active track area edit tool */
@@ -498,8 +513,8 @@ public:
 	virtual bool IsRecordingLive() const override;
 	virtual float GetCurrentLocalTime(UMovieSceneSequence& InMovieSceneSequence) override;
 	virtual float GetGlobalTime() const override;
-	virtual void SetGlobalTime(float Time, ESnapTimeMode SnapTimeMode = ESnapTimeMode::STM_None, bool bLooped = false) override;
-	virtual void SetGlobalTimeDirectly(float Time, ESnapTimeMode SnapTimeMode = ESnapTimeMode::STM_None, bool bLooped = false) override;
+	virtual void SetGlobalTime(float Time, ESnapTimeMode SnapTimeMode = ESnapTimeMode::STM_None, bool bRestarted = false) override;
+	virtual void SetGlobalTimeDirectly(float Time, ESnapTimeMode SnapTimeMode = ESnapTimeMode::STM_None, bool bRestarted = false) override;
 	virtual void SetPerspectiveViewportPossessionEnabled(bool bEnabled) override;
 	virtual void SetPerspectiveViewportCameraCutEnabled(bool bEnabled) override;
 	virtual void RenderMovie(UMovieSceneSection* InSection) const override;
@@ -508,7 +523,10 @@ public:
 	virtual bool IsInSilentMode() const override { return SilentModeCount != 0; }
 	virtual FGuid GetHandleToObject(UObject* Object, bool bCreateHandleIfMissing = true) override;
 	virtual ISequencerObjectChangeListener& GetObjectChangeListener() override;
-	virtual void NotifyMovieSceneDataChanged() override;
+protected:
+	virtual void NotifyMovieSceneDataChangedInternal() override;
+public:
+	virtual void NotifyMovieSceneDataChanged( EMovieSceneDataChangeType DataChangeType ) override;
 	virtual void UpdateRuntimeInstances() override;
 	virtual void UpdatePlaybackRange() override;
 	virtual void AddActors(const TArray<TWeakObjectPtr<AActor> >& InActors) override;
@@ -519,8 +537,10 @@ public:
 	virtual FSequencerSelectionPreview& GetSelectionPreview() override;
 	virtual void NotifyMapChanged(UWorld* NewWorld, EMapChangeType MapChangeType) override;
 	virtual FOnGlobalTimeChanged& OnGlobalTimeChanged() override { return OnGlobalTimeChangedDelegate; }
+	virtual FOnMovieSceneDataChanged& OnMovieSceneDataChanged() override { return OnMovieSceneDataChangedDelegate; }
 	virtual FGuid CreateBinding(UObject& InObject, const FString& InName) override;
 	virtual UObject* GetPlaybackContext() const override;
+	virtual TArray<UObject*> GetEventContexts() const override;
 	virtual void GetAllKeyedProperties(UObject& Object, TSet<UProperty*>& OutProperties) override;
 	virtual FOnActorAddedToSequencer& OnActorAddedToSequencer() override;
 	virtual FOnPreSave& OnPreSave() override;
@@ -539,7 +559,7 @@ public:
 	// IMovieScenePlayer interface
 
 	virtual void GetRuntimeObjects(TSharedRef<FMovieSceneSequenceInstance> MovieSceneInstance, const FGuid& ObjectHandle, TArray<TWeakObjectPtr<UObject>>& OutObjects) const override;
-	virtual void UpdateCameraCut(UObject* CameraObject, UObject* UnlockIfCameraObject, bool bJumpCut) const override;
+	virtual void UpdateCameraCut(UObject* CameraObject, UObject* UnlockIfCameraObject, bool bJumpCut) override;
 	virtual void SetViewportSettings(const TMap<FViewportClient*, EMovieSceneViewportParams>& ViewportParamsMap) override;
 	virtual void GetViewportSettings(TMap<FViewportClient*, EMovieSceneViewportParams>& ViewportParamsMap) const override;
 	virtual EMovieScenePlayerStatus::Type GetPlaybackStatus() const override;
@@ -674,6 +694,9 @@ protected:
 	void TogglePlay();
 	void PlayForward();
 	void Rewind();
+	void ShuttleForward();
+	void ShuttleBackward();
+	void Pause();
 	void StepForward();
 	void StepBackward();
 	void StepToNextKey();
@@ -785,6 +808,9 @@ protected:
 	/** Create record transport control */
 	TSharedRef<SWidget> OnCreateTransportRecord();
 
+	/** Possess PIE viewports with the specified camera settings (a mirror of level viewport possession, but for game viewport clients) */
+	void PossessPIEViewports(UObject* CameraObject, UObject* UnlockIfCameraObject, bool bJumpCut);
+
 private:
 	/** Performs any post-tick rendering work needed when moving through scenes */
 	void PostTickRenderStateFixup();
@@ -859,6 +885,12 @@ private:
 	// @todo sequencer: Should use FTimespan or "double" for Time Cursor Position! (cascades)
 	float ScrubPosition;
 
+	/** The playback rate */
+	float PlayRate;
+
+	/** The shuttle multiplier */
+	float ShuttleMultiplier;
+
 	bool bPerspectiveViewportPossessionEnabled;
 	bool bPerspectiveViewportCameraCutEnabled;
 
@@ -892,6 +924,9 @@ private:
 	/** A delegate which is called any time the global time changes. */
 	FOnGlobalTimeChanged OnGlobalTimeChangedDelegate;
 
+	/** A delegate which is called any time the movie scene data is changed. */
+	FOnMovieSceneDataChanged OnMovieSceneDataChangedDelegate;
+
 	/** A map of all the transport controls to viewports that this sequencer has made */
 	TMap< TSharedPtr<class ILevelViewport>, TSharedPtr<class SWidget> > TransportControls;
 
@@ -914,4 +949,23 @@ private:
 
 	/** The maximum tick rate prior to playing (used for overriding delta time during playback). */
 	double OldMaxTickRate;
+
+	struct FCachedViewTarget
+	{
+		/** The player controller we're possessing */
+		TWeakObjectPtr<APlayerController> PlayerController;
+		/** The view target it was pointing at before we took over */
+		TWeakObjectPtr<AActor> ViewTarget;
+	};
+	
+	/** Cached array of view targets that were set before we possessed the player controller with a camera from sequencer */
+	TArray<FCachedViewTarget> PrePossessionViewTargets;
+
+	FDelegateHandle LevelViewportExtenderHandle;
+
+	/** Attribute used to retrieve event contexts */
+	TAttribute<TArray<UObject*>> EventContextsAttribute;
+
+	/** Event contexts retrieved from the above attribute once per frame */
+	TArray<UObject*> CachedEventContexts;
 };
diff --git a/Engine/Source/Editor/Sequencer/Private/SequencerCommands.cpp b/Engine/Source/Editor/Sequencer/Private/SequencerCommands.cpp
--- a/Engine/Source/Editor/Sequencer/Private/SequencerCommands.cpp
+++ b/Engine/Source/Editor/Sequencer/Private/SequencerCommands.cpp
@@ -10,6 +10,9 @@ void FSequencerCommands::RegisterCommands()
 	UI_COMMAND( TogglePlay, "Toggle Play", "Toggle the timeline playing", EUserInterfaceActionType::Button, FInputChord(EKeys::SpaceBar) );
 	UI_COMMAND( PlayForward, "Play Forward", "Play the timeline forward", EUserInterfaceActionType::Button, FInputChord(EKeys::Down) );
 	UI_COMMAND( Rewind, "Rewind", "Rewind the timeline", EUserInterfaceActionType::Button, FInputChord(EKeys::Up) );
+	UI_COMMAND( ShuttleBackward, "Shuttle Backward", "Shuttle backward", EUserInterfaceActionType::Button, FInputChord(EKeys::J) );
+	UI_COMMAND( ShuttleForward, "Shuttle Forward", "Shuttle forward", EUserInterfaceActionType::Button, FInputChord(EKeys::L) );
+	UI_COMMAND( Pause, "Pause", "Pause playback", EUserInterfaceActionType::Button, FInputChord(EKeys::K) );
 	UI_COMMAND( StepForward, "Step Forward", "Step the timeline forward", EUserInterfaceActionType::Button, FInputChord(EKeys::Right) );
 	UI_COMMAND( StepBackward, "Step Backward", "Step the timeline backward", EUserInterfaceActionType::Button, FInputChord(EKeys::Left) );
 	UI_COMMAND( StepToNextKey, "Step to Next Key", "Step to the next key", EUserInterfaceActionType::Button, FInputChord(EKeys::Period) );
@@ -22,7 +25,7 @@ void FSequencerCommands::RegisterCommands()
 	UI_COMMAND( ZoomInViewRange, "Zoom into the View Range", "Zoom into the view range", EUserInterfaceActionType::Button, FInputChord(EKeys::Equals) );
 	UI_COMMAND( ZoomOutViewRange, "Zoom out of the View Range", "Zoom out of the view range", EUserInterfaceActionType::Button, FInputChord(EKeys::Hyphen) );
 
-	UI_COMMAND( ToggleForceFixedFrameIntervalPlayback, "Force Fixed Frame Interval Playback.",
+	UI_COMMAND( ToggleForceFixedFrameIntervalPlayback, "Force Fixed Frame Interval Playback",
 		"Forces scene evaluation to a fixed frame interval in editor and at runtime, even if this would result in duplicated or dropped frames.",
 		EUserInterfaceActionType::ToggleButton, FInputChord() );
 	
@@ -36,8 +39,8 @@ void FSequencerCommands::RegisterCommands()
 
 	UI_COMMAND( SetSelectionRangeEnd, "Set Selection End", "Sets the end of the selection range", EUserInterfaceActionType::Button, FInputChord(EKeys::O) );
 	UI_COMMAND( SetSelectionRangeStart, "Set Selection Start", "Sets the start of the selection range", EUserInterfaceActionType::Button, FInputChord(EKeys::I) );
+	UI_COMMAND( ResetSelectionRange, "Reset Selection Range", "Reset the selection range", EUserInterfaceActionType::Button, FInputChord() );
 	UI_COMMAND( SelectKeysInSelectionRange, "Select in Selection Range", "Select all keys that fall into the selection range", EUserInterfaceActionType::Button, FInputChord() );
-	UI_COMMAND( DeleteSelectionRange, "Delete Selection Range", "Delete the time range and all the keys that fall into the selection range", EUserInterfaceActionType::Button, FInputChord() );
 
 	UI_COMMAND( SetKey, "Set Key", "Sets a key on the selected tracks", EUserInterfaceActionType::Button, FInputChord(EKeys::Enter) );
 	UI_COMMAND( SetInterpolationCubicAuto, "Set Key Auto", "Cubic interpolation - Automatic tangents", EUserInterfaceActionType::Button, FInputChord(EKeys::One));
@@ -95,6 +98,9 @@ void FSequencerCommands::RegisterCommands()
 	UI_COMMAND( FixActorReferences, "Fix Actor References", "Try to automatically fix up broken actor bindings.", EUserInterfaceActionType::Button, FInputChord() );
 
 	UI_COMMAND( FixFrameTiming, "Fix Frame Timing", "Moves all time data for this sequence to a valid frame time.", EUserInterfaceActionType::Button, FInputChord() );
+
+	UI_COMMAND( RecordSelectedActors, "Record Selected Actors", "Records the selected actors into a new sub sequence of the currently active sequence in Sequencer.", EUserInterfaceActionType::Button, FInputChord() );
+	UI_COMMAND( ExportSceneAndSequence, "Export Scene and Sequence", "Exports the current scene and sequence to an FBX file. (Shots and sub-scenes not supported)", EUserInterfaceActionType::Button, FInputChord() );
 }
 
 #undef LOCTEXT_NAMESPACE
\ No newline at end of file
diff --git a/Engine/Source/Editor/Sequencer/Private/SequencerContextMenus.cpp b/Engine/Source/Editor/Sequencer/Private/SequencerContextMenus.cpp
--- a/Engine/Source/Editor/Sequencer/Private/SequencerContextMenus.cpp
+++ b/Engine/Source/Editor/Sequencer/Private/SequencerContextMenus.cpp
@@ -231,12 +231,12 @@ void FKeyContextMenu::AddPropertiesMenu(FMenuBuilder& MenuBuilder)
 			StructureDetailsView->SetStructureData(Key.Key);
 			StructureDetailsView->GetOnFinishedChangingPropertiesDelegate().AddLambda(
 				[=](const FPropertyChangedEvent& ChangeEvent) {
+					
 					if (Key.Key->GetStruct()->IsChildOf(FMovieSceneKeyStruct::StaticStruct()))
 					{
 						((FMovieSceneKeyStruct*)Key.Key->GetStructMemory())->PropagateChanges(ChangeEvent);
 					}
-					Sequencer->GetFocusedMovieSceneSequence()->Modify();
-					Sequencer->UpdateRuntimeInstances();
+					Sequencer->NotifyMovieSceneDataChanged( EMovieSceneDataChangeType::TrackValueChanged );
 				}
 			);
 		}
@@ -674,7 +674,7 @@ void FSectionContextMenu::TrimSection(bool bTrimLeft)
 	FScopedTransaction TrimSectionTransaction(LOCTEXT("TrimSection_Transaction", "Trim Section"));
 
 	MovieSceneToolHelpers::TrimSection(Sequencer->GetSelection().GetSelectedSections(), Sequencer->GetGlobalTime(), bTrimLeft);
-	Sequencer->NotifyMovieSceneDataChanged();
+	Sequencer->NotifyMovieSceneDataChanged( EMovieSceneDataChangeType::TrackValueChanged );
 }
 
 
@@ -683,7 +683,7 @@ void FSectionContextMenu::SplitSection()
 	FScopedTransaction SplitSectionTransaction(LOCTEXT("SplitSection_Transaction", "Split Section"));
 
 	MovieSceneToolHelpers::SplitSection(Sequencer->GetSelection().GetSelectedSections(), Sequencer->GetGlobalTime());
-	Sequencer->NotifyMovieSceneDataChanged();
+	Sequencer->NotifyMovieSceneDataChanged( EMovieSceneDataChangeType::MovieSceneStructureItemAdded );
 }
 
 
@@ -732,7 +732,7 @@ void FSectionContextMenu::SetExtrapolationMode(ERichCurveExtrapolation ExtrapMod
 
 	if (bAnythingChanged)
 	{
-		Sequencer->UpdateRuntimeInstances();
+		Sequencer->NotifyMovieSceneDataChanged( EMovieSceneDataChangeType::TrackValueChanged );
 	}
 	else
 	{
@@ -787,7 +787,7 @@ void FSectionContextMenu::ToggleSectionActive()
 
 	if (bAnythingChanged)
 	{
-		Sequencer->UpdateRuntimeInstances();
+		Sequencer->NotifyMovieSceneDataChanged( EMovieSceneDataChangeType::TrackValueChanged );
 	}
 	else
 	{
@@ -829,7 +829,7 @@ void FSectionContextMenu::ToggleSectionLocked()
 
 	if (bAnythingChanged)
 	{
-		Sequencer->UpdateRuntimeInstances();
+		Sequencer->NotifyMovieSceneDataChanged( EMovieSceneDataChangeType::TrackValueChanged );
 	}
 	else
 	{
diff --git a/Engine/Source/Editor/Sequencer/Private/SequencerCurveOwner.cpp b/Engine/Source/Editor/Sequencer/Private/SequencerCurveOwner.cpp
--- a/Engine/Source/Editor/Sequencer/Private/SequencerCurveOwner.cpp
+++ b/Engine/Source/Editor/Sequencer/Private/SequencerCurveOwner.cpp
@@ -191,7 +191,7 @@ void FSequencerCurveOwner::OnCurveChanged( const TArray<FRichCurveEditInfo>& Cha
 		}
 	}
 
-	SequencerNodeTree->GetSequencer().UpdateRuntimeInstances();
+	SequencerNodeTree->GetSequencer().NotifyMovieSceneDataChanged( EMovieSceneDataChangeType::TrackValueChanged );
 }
 
 bool FSequencerCurveOwner::IsValidCurve( FRichCurveEditInfo CurveInfo )
diff --git a/Engine/Source/Editor/Sequencer/Private/SequencerEdMode.cpp b/Engine/Source/Editor/Sequencer/Private/SequencerEdMode.cpp
--- a/Engine/Source/Editor/Sequencer/Private/SequencerEdMode.cpp
+++ b/Engine/Source/Editor/Sequencer/Private/SequencerEdMode.cpp
@@ -113,7 +113,7 @@ void GetLocationAtTime(UMovieScene3DTransformSection* TransformSection, float Ke
 	TransformSection->EvalRotation(KeyTime, KeyRot);
 }
 
-void DrawTransformTrack(const FSceneView* View, FPrimitiveDrawInterface* PDI, UMovieScene3DTransformTrack* TransformTrack, const TArray<TWeakObjectPtr<UObject>>& BoundObjects)
+void DrawTransformTrack(const FSceneView* View, FPrimitiveDrawInterface* PDI, UMovieScene3DTransformTrack* TransformTrack, const TArray<TWeakObjectPtr<UObject>>& BoundObjects, const bool& bIsSelected)
 {
 	const bool bHitTesting = PDI->IsHitTesting();
 	FLinearColor TrackColor = FLinearColor::Yellow; //@todo - customizable per track
@@ -123,6 +123,16 @@ void DrawTransformTrack(const FSceneView* View, FPrimitiveDrawInterface* PDI, UM
 		UMovieSceneSection* Section = TransformTrack->GetAllSections()[SectionIndex];
 		UMovieScene3DTransformSection* TransformSection = Cast<UMovieScene3DTransformSection>(Section);
 
+		if (TransformSection->GetShow3DTrajectory() == EShow3DTrajectory::EST_Never)
+		{
+			continue;
+		}
+
+		if (TransformSection->GetShow3DTrajectory() == EShow3DTrajectory::EST_OnlyWhenSelected && !bIsSelected)
+		{
+			continue;
+		}
+
 		FRichCurve& TransXCurve = TransformSection->GetTranslationCurve(EAxis::X);
 		FRichCurve& TransYCurve = TransformSection->GetTranslationCurve(EAxis::Y);
 		FRichCurve& TransZCurve = TransformSection->GetTranslationCurve(EAxis::Z);
@@ -339,34 +349,43 @@ void FSequencerEdMode::DrawTracks3D(const FSceneView* View, FPrimitiveDrawInterf
 {
 	TSet<TSharedRef<FSequencerDisplayNode> > ObjectBindingNodes;
 
-	// Look for the object binding nodes from the nodes with selected keys or sections
-	for (auto OutlinerNode : Sequencer->GetSelection().GetNodesWithSelectedKeysOrSections())
-	{
-		TSharedRef<FSequencerDisplayNode> ObjectBindingNode = OutlinerNode;
-		if (SequencerHelpers::FindObjectBindingNode(OutlinerNode, ObjectBindingNode))
-		{
-			ObjectBindingNodes.Add(ObjectBindingNode);
-		}
-	}
+	// Map between object binding nodes and selection
+	TMap<TSharedRef<FSequencerDisplayNode>, bool > ObjectBindingNodesSelectionMap;
 
-	// And also look for the nodes that are directly selected
-	for (auto OutlinerNode : Sequencer->GetSelection().GetSelectedOutlinerNodes())
+	for (auto ObjectBinding : Sequencer->GetNodeTree()->GetObjectBindingMap() )
 	{
-		TSharedRef<FSequencerDisplayNode> ObjectBindingNode = OutlinerNode;
-		if (SequencerHelpers::FindObjectBindingNode(OutlinerNode, ObjectBindingNode))
+		if (ObjectBinding.Value.IsValid())
 		{
-			ObjectBindingNodes.Add(ObjectBindingNode);
+			TSharedRef<FSequencerObjectBindingNode> ObjectBindingNode = ObjectBinding.Value.ToSharedRef();
+
+			TSet<TSharedRef<FSequencerDisplayNode> > DescendantNodes;
+			SequencerHelpers::GetDescendantNodes(ObjectBindingNode, DescendantNodes);
+
+			bool bSelected = false;
+			for (auto DescendantNode : DescendantNodes)
+			{
+				if (Sequencer->GetSelection().IsSelected(ObjectBindingNode) || 
+					Sequencer->GetSelection().IsSelected(DescendantNode) || 
+					Sequencer->GetSelection().NodeHasSelectedKeysOrSections(DescendantNode))
+				{
+					bSelected = true;
+					break;
+				}
+			}
+
+			ObjectBindingNodesSelectionMap.Add(ObjectBindingNode, bSelected);
 		}
 	}
 
+
 	// Gather up the transform track nodes from the object binding nodes
 
-	for (auto ObjectBindingNode : ObjectBindingNodes)
+	for (auto ObjectBindingNode : ObjectBindingNodesSelectionMap)
 	{
 		TSet<TSharedRef<FSequencerDisplayNode> > AllNodes;
-		SequencerHelpers::GetDescendantNodes(ObjectBindingNode, AllNodes);
+		SequencerHelpers::GetDescendantNodes(ObjectBindingNode.Key, AllNodes);
 
-		FGuid ObjectBinding = StaticCastSharedRef<FSequencerObjectBindingNode>(ObjectBindingNode)->GetObjectBinding();
+		FGuid ObjectBinding = StaticCastSharedRef<FSequencerObjectBindingNode>(ObjectBindingNode.Key)->GetObjectBinding();
 
 		TArray<TWeakObjectPtr<UObject>> BoundObjects;
 		Sequencer->GetRuntimeObjects(Sequencer->GetFocusedMovieSceneSequenceInstance(), ObjectBinding, BoundObjects);
@@ -380,7 +399,7 @@ void FSequencerEdMode::DrawTracks3D(const FSceneView* View, FPrimitiveDrawInterf
 				UMovieScene3DTransformTrack* TransformTrack = Cast<UMovieScene3DTransformTrack>(TrackNodeTrack);
 				if (TransformTrack != nullptr)
 				{
-					DrawTransformTrack(View, PDI, TransformTrack, BoundObjects);
+					DrawTransformTrack(View, PDI, TransformTrack, BoundObjects, ObjectBindingNode.Value);
 				}
 			}
 		}
diff --git a/Engine/Source/Editor/Sequencer/Private/SequencerNodeTree.h b/Engine/Source/Editor/Sequencer/Private/SequencerNodeTree.h
--- a/Engine/Source/Editor/Sequencer/Private/SequencerNodeTree.h
+++ b/Engine/Source/Editor/Sequencer/Private/SequencerNodeTree.h
@@ -92,6 +92,11 @@ public:
 	 */
 	const TSharedPtr<FSequencerDisplayNode>& GetHoveredNode() const;
 
+	/*
+	 * Get the object binding map from guid to object binding nodes
+	 */
+	const TMap< FGuid, TSharedPtr<FSequencerObjectBindingNode> > GetObjectBindingMap() const { return ObjectBindingMap; }
+
 private:
 	/**
 	 * Finds or adds a type editor for the track
diff --git a/Engine/Source/Editor/Sequencer/Private/SequencerObjectChangeListener.cpp b/Engine/Source/Editor/Sequencer/Private/SequencerObjectChangeListener.cpp
--- a/Engine/Source/Editor/Sequencer/Private/SequencerObjectChangeListener.cpp
+++ b/Engine/Source/Editor/Sequencer/Private/SequencerObjectChangeListener.cpp
@@ -197,8 +197,8 @@ bool FSequencerObjectChangeListener::FindPropertySetter( const UClass& ObjectCla
 			}
 		}
 		
-		// Valid either if a setter function is found or a property with the Interp keyword is found. The property also needs to be editable.
-		bFound = (bFoundValidFunction || bFoundValidInterp) && bFoundEdit && !bFoundEditDefaultsOnly;
+		// Valid if there's a setter function and the property is editable. Also valid if there's an interp keyword.
+		bFound = (bFoundValidFunction && bFoundEdit && !bFoundEditDefaultsOnly) || bFoundValidInterp;
 	}
 
 	return bFound;
diff --git a/Engine/Source/Editor/Sequencer/Private/SequencerSettings.cpp b/Engine/Source/Editor/Sequencer/Private/SequencerSettings.cpp
--- a/Engine/Source/Editor/Sequencer/Private/SequencerSettings.cpp
+++ b/Engine/Source/Editor/Sequencer/Private/SequencerSettings.cpp
@@ -25,6 +25,7 @@ USequencerSettings::USequencerSettings( const FObjectInitializer& ObjectInitiali
 	CurveValueSnapInterval = 10.0f;
 	bSnapCurveValueToInterval = true;
 	bLabelBrowserVisible = false;
+	bRewindOnRecord = true;
 	ZoomPosition = ESequencerZoomPosition::SZP_CurrentTime;
 	bAutoScrollEnabled = false;
 	bShowCurveEditor = false;
@@ -38,6 +39,7 @@ USequencerSettings::USequencerSettings( const FObjectInitializer& ObjectInitiali
 	bInfiniteKeyAreas = false;
 	bShowChannelColors = false;
 	bShowViewportTransportControls = true;
+	bAllowPossessionOfPIEViewports = false;
 }
 
 EAutoKeyMode USequencerSettings::GetAutoKeyMode() const
@@ -321,6 +323,20 @@ void USequencerSettings::SetLabelBrowserVisible(bool Visible)
 	}
 }
 
+bool USequencerSettings::ShouldRewindOnRecord() const
+{
+	return bRewindOnRecord;
+}
+
+void USequencerSettings::SetRewindOnRecord(bool bInRewindOnRecord)
+{
+	if (bInRewindOnRecord != bRewindOnRecord)
+	{
+		bRewindOnRecord = bInRewindOnRecord;
+		SaveConfig();
+	}
+}
+
 ESequencerZoomPosition USequencerSettings::GetZoomPosition() const
 {
 	return ZoomPosition;
@@ -516,6 +532,20 @@ float USequencerSettings::SnapTimeToInterval( float InTimeValue ) const
 		: InTimeValue;
 }
 
+bool USequencerSettings::ShouldAllowPossessionOfPIEViewports() const
+{
+	return bAllowPossessionOfPIEViewports;
+}
+
+void USequencerSettings::SetAllowPossessionOfPIEViewports(bool bInAllowPossessionOfPIEViewports)
+{
+	if (bInAllowPossessionOfPIEViewports != bAllowPossessionOfPIEViewports)
+	{
+		bAllowPossessionOfPIEViewports = bInAllowPossessionOfPIEViewports;
+		SaveConfig();
+	}
+}
+
 USequencerSettings::FOnShowCurveEditorChanged& USequencerSettings::GetOnShowCurveEditorChanged()
 {
 	return OnShowCurveEditorChanged;
diff --git a/Engine/Source/Editor/Sequencer/Private/Tools/EditToolDragOperations.cpp b/Engine/Source/Editor/Sequencer/Private/Tools/EditToolDragOperations.cpp
--- a/Engine/Source/Editor/Sequencer/Private/Tools/EditToolDragOperations.cpp
+++ b/Engine/Source/Editor/Sequencer/Private/Tools/EditToolDragOperations.cpp
@@ -135,8 +135,7 @@ void FEditToolDragOperation::BeginTransaction( TArray< FSectionHandle >& Section
 void FEditToolDragOperation::EndTransaction()
 {
 	Transaction.Reset();
-	Sequencer.UpdateRuntimeInstances();
-	Sequencer.UpdatePlaybackRange();
+	Sequencer.NotifyMovieSceneDataChanged( EMovieSceneDataChangeType::TrackValueChanged );
 }
 
 FResizeSection::FResizeSection( FSequencer& InSequencer, TArray<FSectionHandle> InSections, bool bInDraggingByEnd )
@@ -290,7 +289,7 @@ void FResizeSection::OnDrag(const FPointerEvent& MouseEvent, FVector2D LocalMous
 		}
 	}
 
-	Sequencer.UpdateRuntimeInstances();
+	Sequencer.NotifyMovieSceneDataChanged( EMovieSceneDataChangeType::TrackValueChanged );
 }
 
 FMoveSection::FMoveSection( FSequencer& InSequencer, TArray<FSectionHandle> InSections )
@@ -494,7 +493,7 @@ void FMoveSection::OnDrag(const FPointerEvent& MouseEvent, FVector2D LocalMouseP
 		}
 	}
 
-	Sequencer.UpdateRuntimeInstances();
+	Sequencer.NotifyMovieSceneDataChanged( EMovieSceneDataChangeType::TrackValueChanged );
 }
 
 void FMoveKeys::OnBeginDrag(const FPointerEvent& MouseEvent, FVector2D LocalMousePos, const FVirtualTrackArea& VirtualTrackArea)
@@ -618,7 +617,7 @@ void FMoveKeys::OnDrag(const FPointerEvent& MouseEvent, FVector2D LocalMousePos,
 		Sequencer.SetGlobalTime(PrevNewKeyTime);
 	}
 
-	Sequencer.UpdateRuntimeInstances();
+	Sequencer.NotifyMovieSceneDataChanged( EMovieSceneDataChangeType::TrackValueChanged );
 }
 
 void FMoveKeys::OnEndDrag(const FPointerEvent& MouseEvent, FVector2D LocalMousePos, const FVirtualTrackArea& VirtualTrackArea)
diff --git a/Engine/Source/Editor/Sequencer/Public/ISequencer.h b/Engine/Source/Editor/Sequencer/Public/ISequencer.h
--- a/Engine/Source/Editor/Sequencer/Public/ISequencer.h
+++ b/Engine/Source/Editor/Sequencer/Public/ISequencer.h
@@ -67,6 +67,27 @@ enum ESnapTimeMode
 	STM_All = STM_Interval | STM_Keys
 };
 
+
+/**
+ * Defines different types of movie scene data changes. 
+ */
+enum class EMovieSceneDataChangeType
+{
+	/** Data owned by a track has been modified such as adding or removing keys, or changing their values. */
+	TrackValueChanged,
+	/** The structure of the movie scene has changed by adding folders, object bindings, tracks, or sections. */
+	MovieSceneStructureItemAdded,
+	/** The structure of the movie scene has changed by removing folders, object bindings, tracks, or sections. */
+	MovieSceneStructureItemRemoved,
+	/** The structure of the movie scene has changed by adding and removing folders, object bindings, tracks, or sections. */
+	MovieSceneStructureItemsChanged,
+	/** The active movie scene has been changed to a new movie scene. */
+	ActiveMovieSceneChanged,
+	/** It's not known what data has changed. */
+	Unknown
+};
+
+
 /**
  * Interface for sequencers.
  */
@@ -77,6 +98,7 @@ class ISequencer
 public:
 	
 	DECLARE_MULTICAST_DELEGATE(FOnGlobalTimeChanged);
+	DECLARE_MULTICAST_DELEGATE(FOnMovieSceneDataChanged);
 
 public:
 
@@ -200,13 +222,13 @@ public:
 	 *
 	 * @param Time The global time to set.
 	 * @param SnapTimeMode The type of time snapping allowed.
-	 * @param bLooped Whether the time has been looped
+	 * @param bRestarted Whether the time has been restarted from the beginning or looped.
 	 * @see GetGlobalTime
 	 */
-	virtual void SetGlobalTime(float Time, ESnapTimeMode SnapTimeMode = ESnapTimeMode::STM_None, bool bLooped = false) = 0;
+	virtual void SetGlobalTime(float Time, ESnapTimeMode SnapTimeMode = ESnapTimeMode::STM_None, bool bRestarted = false) = 0;
 
 	/** Set the global time directly, without performing any auto-scroll */
-	virtual void SetGlobalTimeDirectly(float Time, ESnapTimeMode SnapTimeMode = ESnapTimeMode::STM_None, bool bLooped = false) = 0;
+	virtual void SetGlobalTimeDirectly(float Time, ESnapTimeMode SnapTimeMode = ESnapTimeMode::STM_None, bool bRestarted = false) = 0;
 
 	/** @return The current view range */
 	virtual FAnimatedRange GetViewRange() const
@@ -302,7 +324,14 @@ public:
 
 	virtual void KeyProperty(FKeyPropertyParams KeyPropertyParams) = 0;
 
-	virtual void NotifyMovieSceneDataChanged() = 0;
+	DEPRECATED( 4.13, "NotifyMovieSceneDataChanged() is deprecated, use the version that takes EMovieSceneDataChangeType" )
+	void NotifyMovieSceneDataChanged() { NotifyMovieSceneDataChangedInternal(); };
+
+protected:
+	virtual void NotifyMovieSceneDataChangedInternal() = 0;
+
+public:
+	virtual void NotifyMovieSceneDataChanged( EMovieSceneDataChangeType DataChangeType ) = 0;
 
 	virtual void UpdateRuntimeInstances() = 0;
 
@@ -321,6 +350,9 @@ public:
 	/** Gets a multicast delegate which is executed whenever the global time changes. */
 	virtual FOnGlobalTimeChanged& OnGlobalTimeChanged() = 0;
 
+	/** Gets a multicast delegate which is executed whenever the movie scene data is changed. */
+	virtual FOnMovieSceneDataChanged& OnMovieSceneDataChanged() = 0;
+
 	/** @return a numeric type interface that will parse and display numbers as frames and times correctly */
 	virtual TSharedRef<INumericTypeInterface<float>> GetNumericTypeInterface() = 0;
 
diff --git a/Engine/Source/Editor/Sequencer/Public/ISequencerModule.h b/Engine/Source/Editor/Sequencer/Public/ISequencerModule.h
--- a/Engine/Source/Editor/Sequencer/Public/ISequencerModule.h
+++ b/Engine/Source/Editor/Sequencer/Public/ISequencerModule.h
@@ -69,6 +69,9 @@ struct FSequencerInitParams
 
 	/** Domain-specific spawn register for the movie scene */
 	TSharedPtr<IMovieSceneSpawnRegister> SpawnRegister;
+
+	/** Accessor for event contexts */
+	TAttribute<TArray<UObject*>> EventContexts;
 };
 
 
diff --git a/Engine/Source/Editor/Sequencer/Public/MovieSceneTrackEditor.h b/Engine/Source/Editor/Sequencer/Public/MovieSceneTrackEditor.h
--- a/Engine/Source/Editor/Sequencer/Public/MovieSceneTrackEditor.h
+++ b/Engine/Source/Editor/Sequencer/Public/MovieSceneTrackEditor.h
@@ -65,15 +65,6 @@ public:
 		return Sequencer.Pin()->GetCurrentLocalTime( *InMovieSceneSequence );
 	}
 
-	void NotifyMovieSceneDataChanged()
-	{
-		TSharedPtr<ISequencer> SequencerPin = Sequencer.Pin();
-		if( SequencerPin.IsValid()  )
-		{
-			SequencerPin->NotifyMovieSceneDataChanged();
-		}
-	}
-
 	void UpdatePlaybackRange()
 	{
 		TSharedPtr<ISequencer> SequencerPin = Sequencer.Pin();
@@ -104,8 +95,9 @@ public:
 
 				if( OnKeyProperty.Execute( KeyTime ) )
 				{
-					// Movie scene data has changed
-					NotifyMovieSceneDataChanged();
+					// TODO: This should pass an appropriate change type here instead of always passing structure changed since most
+					// changes will be value changes.
+					Sequencer.Pin()->NotifyMovieSceneDataChanged( EMovieSceneDataChangeType::MovieSceneStructureItemAdded );
 				}
 
 				UpdatePlaybackRange();
diff --git a/Engine/Source/Editor/Sequencer/Public/SequencerCommands.h b/Engine/Source/Editor/Sequencer/Public/SequencerCommands.h
--- a/Engine/Source/Editor/Sequencer/Public/SequencerCommands.h
+++ b/Engine/Source/Editor/Sequencer/Public/SequencerCommands.h
@@ -24,6 +24,15 @@ public:
 	/** Rewind */
 	TSharedPtr< FUICommandInfo > Rewind;
 
+	/** Shuttle forward */
+	TSharedPtr< FUICommandInfo > ShuttleForward;
+
+	/** Shuttle backward */
+	TSharedPtr< FUICommandInfo > ShuttleBackward;
+
+	/** Pause */
+	TSharedPtr< FUICommandInfo > Pause;
+
 	/** Step forward */
 	TSharedPtr< FUICommandInfo > StepForward;
 
@@ -84,12 +93,12 @@ public:
 	/** Sets the lower bound of the selection range */
 	TSharedPtr< FUICommandInfo > SetSelectionRangeStart;
 
+	/** Clear and reset the selection range */
+	TSharedPtr< FUICommandInfo > ResetSelectionRange;
+
 	/** Select all keys that fall into the selection range*/
 	TSharedPtr< FUICommandInfo > SelectKeysInSelectionRange;
 
-	/** Delete the time in the selection range */
-	TSharedPtr< FUICommandInfo > DeleteSelectionRange;
-
 	/** Sets a key at the current time for the selected actor */
 	TSharedPtr< FUICommandInfo > SetKey;
 
@@ -213,6 +222,12 @@ public:
 	/** Attempts to move all time data for this sequence on to a valid frame */
 	TSharedPtr< FUICommandInfo > FixFrameTiming;
 
+	/** Record the selected actors into a sub sequence of the currently active sequence */
+	TSharedPtr< FUICommandInfo > RecordSelectedActors;
+
+	/** Exports the current scene and sequence to fbx. */
+	TSharedPtr< FUICommandInfo > ExportSceneAndSequence;
+
 	/**
 	 * Initialize commands
 	 */
diff --git a/Engine/Source/Editor/Sequencer/Public/SequencerSettings.h b/Engine/Source/Editor/Sequencer/Public/SequencerSettings.h
--- a/Engine/Source/Editor/Sequencer/Public/SequencerSettings.h
+++ b/Engine/Source/Editor/Sequencer/Public/SequencerSettings.h
@@ -168,6 +168,11 @@ public:
 	/** Sets whether or not the label browser is visible. */
 	void SetLabelBrowserVisible(bool Visible);
 
+	/** Gets whether to jump to the start of the sequence when we start a recording or not. */
+	bool ShouldRewindOnRecord() const;
+	/** Sets whether to jump to the start of the sequence when we start a recording. */
+	void SetRewindOnRecord(bool bInRewindOnRecord);
+
 	/** Get zoom in/out position (mouse position or current time). */
 	ESequencerZoomPosition GetZoomPosition() const;
 	/** Set zoom in/out position (mouse position or current time). */
@@ -233,6 +238,11 @@ public:
 	/** Toggle whether to show transport controls in level editor viewports */
 	void SetShowViewportTransportControls(bool bVisible);
 
+	/** @return Whether to allow possession of PIE viewports */
+	bool ShouldAllowPossessionOfPIEViewports() const;
+	/** Toggle whether to allow possession of PIE viewports */
+	void SetAllowPossessionOfPIEViewports(bool bInAllowPossessionOfPIEViewports);
+
 	/** Snaps a time value in seconds to the currently selected interval. */
 	float SnapTimeToInterval(float InTimeValue) const;
 
@@ -323,6 +333,10 @@ protected:
 	UPROPERTY( config, EditAnywhere, Category=General )
 	bool bLabelBrowserVisible;
 
+	/** Defines whether to jump back to the start of the sequence when a recording is started */
+	UPROPERTY(config, EditAnywhere, Category=General)
+	bool bRewindOnRecord;
+
 	/** Whether to zoom in on the current position or the current time in the timeline. */
 	UPROPERTY( config, EditAnywhere, Category=Timeline )
 	TEnumAsByte<ESequencerZoomPosition> ZoomPosition;
@@ -375,6 +389,10 @@ protected:
 	UPROPERTY( config )
 	bool bShowViewportTransportControls;
 
+	/** When enabled, sequencer is able to possess viewports that represent PIE worlds */
+	UPROPERTY(config, EditAnywhere, Category=General)
+	bool bAllowPossessionOfPIEViewports;
+
 	FOnShowCurveEditorChanged OnShowCurveEditorChanged;
 
 	FOnTimeSnapIntervalChanged OnTimeSnapIntervalChanged;
diff --git a/Engine/Source/Editor/Sequencer/Sequencer.Build.cs b/Engine/Source/Editor/Sequencer/Sequencer.Build.cs
--- a/Engine/Source/Editor/Sequencer/Sequencer.Build.cs
+++ b/Engine/Source/Editor/Sequencer/Sequencer.Build.cs
@@ -10,7 +10,8 @@ public class Sequencer : ModuleRules
             new string[] {
                 "Editor/Sequencer/Private",
                 "Editor/Sequencer/Private/DisplayNodes",
-            }
+				"Editor/UnrealEd/Private" // TODO: Fix this, for now it's needed for the fbx exporter
+	}
         );
 
 		PrivateDependencyModuleNames.AddRange(
@@ -43,7 +44,9 @@ public class Sequencer : ModuleRules
 				"PropertyEditor",
 				"Kismet",
 				"SceneOutliner",
-				"SequenceRecorder"
+				"SequenceRecorder",
+				"MainFrame",
+				"DesktopPlatform"
 			}
 		);
 
@@ -55,9 +58,11 @@ public class Sequencer : ModuleRules
 				"WorkspaceMenuStructure",
 				"SequenceRecorder",
 				"SequenceRecorderSections",
+				"MainFrame",
 			}
 		);
 
 		CircularlyReferencedDependentModules.Add("MovieSceneTools");
+		AddEngineThirdPartyPrivateStaticDependencies(Target, "FBX");
 	}
 }
diff --git a/Engine/Source/Editor/UMGEditor/Private/WidgetBlueprintEditor.cpp b/Engine/Source/Editor/UMGEditor/Private/WidgetBlueprintEditor.cpp
--- a/Engine/Source/Editor/UMGEditor/Private/WidgetBlueprintEditor.cpp
+++ b/Engine/Source/Editor/UMGEditor/Private/WidgetBlueprintEditor.cpp
@@ -37,6 +37,7 @@ FWidgetBlueprintEditor::FWidgetBlueprintEditor()
 	, PreviewBlueprint(nullptr)
 	, bIsSimulateEnabled(false)
 	, bIsRealTime(true)
+	, bRefreshGeneratedClassAnimations(false)
 {
 	PreviewScene.GetWorld()->bBegunPlay = false;
 
@@ -58,7 +59,11 @@ FWidgetBlueprintEditor::~FWidgetBlueprintEditor()
 
 	GEditor->OnObjectsReplaced().RemoveAll(this);
 	
-	Sequencer.Reset();
+	if ( Sequencer.IsValid() )
+	{
+		Sequencer->OnMovieSceneDataChanged().RemoveAll( this );
+		Sequencer.Reset();
+	}
 
 	// Un-Register sequencer menu extenders.
 	ISequencerModule& SequencerModule = FModuleManager::Get().LoadModuleChecked<ISequencerModule>("Sequencer");
@@ -418,6 +423,19 @@ void FWidgetBlueprintEditor::Tick(float DeltaTime)
 		}
 	}
 
+	// Whenever animations change the generated class animations need to be updated since they are copied on compile.  This
+	// update is deferred to tick since some edit operations (e.g. drag/drop) cause large numbers of changes to the data.
+	if ( bRefreshGeneratedClassAnimations )
+	{
+		TArray<UWidgetAnimation*>& PreviewAnimations = Cast<UWidgetBlueprintGeneratedClass>( PreviewBlueprint->GeneratedClass )->Animations;
+		PreviewAnimations.Empty();
+		for ( UWidgetAnimation* WidgetAnimation : PreviewBlueprint->Animations )
+		{
+			PreviewAnimations.Add( DuplicateObject<UWidgetAnimation>( WidgetAnimation, PreviewBlueprint->GeneratedClass ) );
+		}
+		bRefreshGeneratedClassAnimations = false;
+	}
+
 	// Note: The weak ptr can become stale if the actor is reinstanced due to a Blueprint change, etc. In that case we 
 	//       look to see if we can find the new instance in the preview world and then update the weak ptr.
 	if ( PreviewWidgetPtr.IsStale(true) || bPreviewInvalidated )
@@ -628,6 +646,7 @@ TSharedPtr<ISequencer>& FWidgetBlueprintEditor::GetSequencer()
 		};
 
 		Sequencer = FModuleManager::LoadModuleChecked<ISequencerModule>("Sequencer").CreateSequencer(SequencerInitParams);
+		Sequencer->OnMovieSceneDataChanged().AddSP( this, &FWidgetBlueprintEditor::OnMovieSceneDataChanged );
 		ChangeViewedAnimation(*UWidgetAnimation::GetNullAnimation());
 	}
 
@@ -1018,9 +1037,14 @@ void FWidgetBlueprintEditor::AddMaterialTrack( UWidget* Widget, TArray<UProperty
 			NewTrack->SetBrushPropertyNamePath( MaterialPropertyNamePath );
 			NewTrack->SetDisplayName( FText::Format( LOCTEXT( "TrackDisplayNameFormat", "{0} Material"), MaterialPropertyDisplayName ) );
 
-			Sequencer->NotifyMovieSceneDataChanged();
+			Sequencer->NotifyMovieSceneDataChanged(EMovieSceneDataChangeType::MovieSceneStructureItemAdded );
 		}
 	}
 }
 
+void FWidgetBlueprintEditor::OnMovieSceneDataChanged()
+{
+	bRefreshGeneratedClassAnimations = true;
+}
+
 #undef LOCTEXT_NAMESPACE
diff --git a/Engine/Source/Editor/UnrealEd/Private/EditorViewportClient.cpp b/Engine/Source/Editor/UnrealEd/Private/EditorViewportClient.cpp
--- a/Engine/Source/Editor/UnrealEd/Private/EditorViewportClient.cpp
+++ b/Engine/Source/Editor/UnrealEd/Private/EditorViewportClient.cpp
@@ -4589,10 +4589,13 @@ void FEditorViewportClient::UpdateHiddenCollisionDrawing()
 			bool bCollisionMode = EngineShowFlags.Collision || EngineShowFlags.CollisionVisibility || EngineShowFlags.CollisionPawn;
 
 			// Tell engine to create proxies for hidden components, so we can still draw collision
-			World->bCreateRenderStateForHiddenComponents = bCollisionMode;
+			if (World->bCreateRenderStateForHiddenComponents != bCollisionMode)
+			{
+				World->bCreateRenderStateForHiddenComponents = bCollisionMode;
 
-			// Need to recreate scene proxies when this flag changes.
-			FGlobalComponentRecreateRenderStateContext Recreate;
+				// Need to recreate scene proxies when this flag changes.
+				FGlobalComponentRecreateRenderStateContext Recreate;
+			}
 		}
 	}
 }
diff --git a/Engine/Source/Editor/UnrealEd/Private/Fbx/FbxAnimationExport.cpp b/Engine/Source/Editor/UnrealEd/Private/Fbx/FbxAnimationExport.cpp
--- a/Engine/Source/Editor/UnrealEd/Private/Fbx/FbxAnimationExport.cpp
+++ b/Engine/Source/Editor/UnrealEd/Private/Fbx/FbxAnimationExport.cpp
@@ -169,6 +169,7 @@ void FFbxExporter::CorrectAnimTrackInterpolation( TArray<FbxNode*>& BoneNodes, F
 		for(int32 CurveIndex = 0; CurveIndex < 3; ++CurveIndex)
 		{
 			FbxAnimCurve* CurrentCurve = Curves[CurveIndex];
+			CurrentCurve->KeyModifyBegin();
 
 			float CurrentAngleOffset = 0.f;
 			for(int32 KeyIndex = 1; KeyIndex < CurrentCurve->KeyGetCount(); ++KeyIndex)
@@ -191,6 +192,8 @@ void FFbxExporter::CorrectAnimTrackInterpolation( TArray<FbxNode*>& BoneNodes, F
 
 				CurrentCurve->KeySetValue(KeyIndex, CurrentOutVal);
 			}
+
+			CurrentCurve->KeyModifyEnd();
 		}
 	}
 }
@@ -312,6 +315,8 @@ void FFbxExporter::ExportAnimSequencesAsSingle( USkeletalMesh* SkelMesh, const A
 
 }
 
+
+
 /**
  * Exports all the animation sequences part of a single Group in a Matinee sequence
  * as a single animation in the FBX document.  The animation is created by sampling the
@@ -345,14 +350,15 @@ void FFbxExporter::ExportMatineeGroup(class AMatineeActor* MatineeActor, USkelet
 	FbxSkeletonRoots.Add(SkeletalMeshComponent, SkeletonRootNode);
 	BaseNode->AddChild(SkeletonRootNode);
 
-	ExportAnimTrack(MatineeActor, SkeletalMeshComponent);
+	FMatineeAnimTrackAdapter AnimTrackAdapter(MatineeActor);
+	ExportAnimTrack(AnimTrackAdapter, SkeletalMeshComponent);
 }
 
-void FFbxExporter::ExportAnimTrack(class AMatineeActor* MatineeActor, USkeletalMeshComponent* SkeletalMeshComponent)
+void FFbxExporter::ExportAnimTrack(IAnimTrackAdapter& AnimTrackAdapter, USkeletalMeshComponent* SkeletalMeshComponent)
 {
 	static const float SamplingRate = 1.f / DEFAULT_SAMPLERATE;
 
-	float MatineeLength = MatineeActor->MatineeData->InterpLength;
+	float MatineeLength = AnimTrackAdapter.GetAnimationLength();
 	// show a status update every 1 second worth of samples
 	const float UpdateFrequency = 1.0f;
 	float NextUpdateTime = UpdateFrequency;
@@ -370,7 +376,7 @@ void FFbxExporter::ExportAnimTrack(class AMatineeActor* MatineeActor, USkeletalM
 	for(SampleTime = 0.f; SampleTime <= MatineeLength; SampleTime += SamplingRate)
 	{
 		// This will call UpdateSkelPose on the skeletal mesh component to move bones based on animations in the matinee group
-		MatineeActor->UpdateInterp( SampleTime, true );
+		AnimTrackAdapter.UpdateAnimation( SampleTime );
 
 		FbxTime ExportTime;
 		ExportTime.SetSecondDouble(SampleTime);
diff --git a/Engine/Source/Editor/UnrealEd/Private/Fbx/FbxMainExport.cpp b/Engine/Source/Editor/UnrealEd/Private/Fbx/FbxMainExport.cpp
--- a/Engine/Source/Editor/UnrealEd/Private/Fbx/FbxMainExport.cpp
+++ b/Engine/Source/Editor/UnrealEd/Private/Fbx/FbxMainExport.cpp
@@ -36,6 +36,7 @@
 #include "Engine/Light.h"
 #include "Engine/Polys.h"
 #include "Engine/StaticMeshActor.h"
+#include "Animation/SkeletalMeshActor.h"
 #include "Components/BrushComponent.h"
 #include "Components/SpotLightComponent.h"
 #include "Components/PointLightComponent.h"
@@ -43,6 +44,15 @@
 #include "Components/DirectionalLightComponent.h"
 #include "Components/ChildActorComponent.h"
 
+#include "MovieScene.h"
+#include "MovieSceneBinding.h"
+#include "IMovieScenePlayer.h"
+#include "Tracks/MovieScene3DTransformTrack.h"
+#include "Tracks/MovieSceneFloatTrack.h"
+#include "Tracks/MovieSceneSkeletalAnimationTrack.h"
+#include "Sections/MovieScene3DTransformSection.h"
+#include "Sections/MovieSceneFloatSection.h"
+
 namespace UnFbx
 {
 
@@ -362,9 +372,9 @@ void FFbxExporter::ExportLevelMesh( ULevel* InLevel, AMatineeActor* InMatineeAct
 		{
 			ExportCamera(CastChecked<ACameraActor>(Actor), InMatineeActor, false); // Just export the placement of the particle emitter.
 		}
-		else if(bIsBlueprintClass)
+		else if(Actor->IsA(ASkeletalMeshActor::StaticClass()) ||  bIsBlueprintClass)
 		{
-			// Export blueprint actors and all their components
+			// Export blueprint and skeletal actors and all their components
 			ExportActor( Actor, InMatineeActor, true );
 		}
 	}
@@ -1102,6 +1112,24 @@ FbxSurfaceMaterial* FFbxExporter::ExportMaterial(UMaterialInterface* MaterialInt
 }
 
 
+FFbxExporter::FMatineeAnimTrackAdapter::FMatineeAnimTrackAdapter( AMatineeActor* InMatineeActor )
+{
+	MatineeActor = InMatineeActor;
+}
+
+
+float FFbxExporter::FMatineeAnimTrackAdapter::GetAnimationLength() const
+{
+	return MatineeActor->MatineeData->InterpLength;
+}
+
+
+void FFbxExporter::FMatineeAnimTrackAdapter::UpdateAnimation( float Time )
+{
+	MatineeActor->UpdateInterp( Time, true );
+}
+
+
 /**
  * Exports the given Matinee sequence information into a FBX document.
  */
@@ -1153,7 +1181,8 @@ bool FFbxExporter::ExportMatinee(AMatineeActor* InMatineeActor)
 					USkeletalMeshComponent* SkeletalMeshComp = Cast<USkeletalMeshComponent>(Actor->GetComponentByClass(USkeletalMeshComponent::StaticClass()));
 					if(SkeletalMeshComp)
 					{
-						ExportAnimTrack(InMatineeActor, SkeletalMeshComp);
+						FMatineeAnimTrackAdapter AnimTrackAdapter(InMatineeActor);
+						ExportAnimTrack(AnimTrackAdapter, SkeletalMeshComp);
 					}
 				}
 			}
@@ -1170,6 +1199,90 @@ bool FFbxExporter::ExportMatinee(AMatineeActor* InMatineeActor)
 }
 
 
+FFbxExporter::FLevelSequenceAnimTrackAdapter::FLevelSequenceAnimTrackAdapter( IMovieScenePlayer* InMovieScenePlayer )
+{
+	MovieScenePlayer = InMovieScenePlayer;
+}
+
+
+float FFbxExporter::FLevelSequenceAnimTrackAdapter::GetAnimationLength() const
+{
+	return MovieScenePlayer->GetRootMovieSceneSequenceInstance()->GetTimeRange().Size<float>();
+}
+
+
+void FFbxExporter::FLevelSequenceAnimTrackAdapter::UpdateAnimation( float Time )
+{
+	EMovieSceneUpdateData UpdateData( Time, Time );
+	return MovieScenePlayer->GetRootMovieSceneSequenceInstance()->Update( UpdateData, *MovieScenePlayer );
+}
+
+
+bool FFbxExporter::ExportLevelSequence( UMovieScene* MovieScene, IMovieScenePlayer* MovieScenePlayer )
+{
+	if ( MovieScene == nullptr || MovieScenePlayer == nullptr )
+	{
+		return false;
+	}
+
+	for ( const FMovieSceneBinding& MovieSceneBinding : MovieScene->GetBindings() )
+	{
+		TArray<TWeakObjectPtr<UObject>> RuntimeObjects;
+		MovieScenePlayer->GetRuntimeObjects( MovieScenePlayer->GetRootMovieSceneSequenceInstance(), MovieSceneBinding.GetObjectGuid(), RuntimeObjects );
+
+		for ( TWeakObjectPtr<UObject> RuntimeObject : RuntimeObjects )
+		{
+			if ( RuntimeObject.IsValid() )
+			{
+				AActor* Actor = Cast<AActor>( RuntimeObject.Get() );
+				UActorComponent* Component = Cast<UActorComponent>( RuntimeObject.Get() );
+				if ( Actor == nullptr && Component != nullptr )
+				{
+					Actor = Component->GetOwner();
+				}
+
+				if ( Actor == nullptr )
+				{
+					continue;
+				}
+
+				FbxNode* FbxActor = FindActor( Actor );
+
+				// now it should export everybody
+				if ( FbxActor )
+				{
+					// Look for the tracks that we currently support
+					for ( UMovieSceneTrack* Track : MovieSceneBinding.GetTracks() )
+					{
+						if ( Track->IsA( UMovieScene3DTransformTrack::StaticClass() ) )
+						{
+							UMovieScene3DTransformTrack* TransformTrack = (UMovieScene3DTransformTrack*)Track;
+							ExportLevelSequence3DTransformTrack( *FbxActor, *TransformTrack );
+						}
+						else if ( Track->IsA( UMovieSceneFloatTrack::StaticClass() ) )
+						{
+							UMovieSceneFloatTrack* FloatTrack = (UMovieSceneFloatTrack*)Track;
+							ExportLevelSequenceFloatTrack( *FbxActor, *FloatTrack );
+						}
+						else if ( Track->IsA( UMovieSceneSkeletalAnimationTrack::StaticClass() ) )
+						{
+							USkeletalMeshComponent* SkeletalMeshComp = Cast<USkeletalMeshComponent>( Actor->GetComponentByClass( USkeletalMeshComponent::StaticClass() ) );
+							if ( SkeletalMeshComp )
+							{
+								FLevelSequenceAnimTrackAdapter AnimTrackAdapter( MovieScenePlayer );
+								ExportAnimTrack( AnimTrackAdapter, SkeletalMeshComp );
+							}
+						}
+					}
+				}
+			}
+		}
+	}
+
+	return true;
+}
+
+
 /**
  * Exports a scene node with the placement indicated by a given actor.
  * This scene node will always have two transformations: one translation vector and one Euler rotation.
@@ -1892,6 +2005,199 @@ void FFbxExporter::ExportAnimatedFloat(FbxProperty* FbxProperty, FInterpCurveFlo
 	AnimCurve->KeyModifyEnd();
 }
 
+void FFbxExporter::ExportLevelSequence3DTransformTrack( FbxNode& FbxActor, UMovieScene3DTransformTrack& TransformTrack )
+{
+	FbxAnimLayer* BaseLayer = AnimStack->GetMember<FbxAnimLayer>( 0 );
+	FbxAnimCurve* Curve;
+
+	// TODO: Support more than one section?
+	UMovieScene3DTransformSection* TransformSection = TransformTrack.GetAllSections().Num() > 0
+		? Cast<UMovieScene3DTransformSection>(TransformTrack.GetAllSections()[0])
+		: nullptr;
+
+	if ( TransformSection != nullptr )
+	{
+		// Translation;
+		FbxActor.LclTranslation.GetCurveNode( BaseLayer, true );
+		Curve = FbxActor.LclTranslation.GetCurve( BaseLayer, FBXSDK_CURVENODE_COMPONENT_X, true );
+		ExportRichCurveToFbxCurve( TransformSection->GetTranslationCurve( EAxis::X ), *Curve, TransformSection->GetRange() );
+		Curve = FbxActor.LclTranslation.GetCurve( BaseLayer, FBXSDK_CURVENODE_COMPONENT_Y, true );
+		ExportRichCurveToFbxCurve( TransformSection->GetTranslationCurve( EAxis::Y ), *Curve, TransformSection->GetRange() );
+		Curve = FbxActor.LclTranslation.GetCurve( BaseLayer, FBXSDK_CURVENODE_COMPONENT_Z, true );
+		ExportRichCurveToFbxCurve( TransformSection->GetTranslationCurve( EAxis::Z ), *Curve, TransformSection->GetRange() );
+
+		// Rotation
+		FbxActor.LclRotation.GetCurveNode( BaseLayer, true );
+		Curve = FbxActor.LclRotation.GetCurve( BaseLayer, FBXSDK_CURVENODE_COMPONENT_X, true );
+		ExportRichCurveToFbxCurve( TransformSection->GetRotationCurve( EAxis::X ), *Curve, TransformSection->GetRange() );
+		Curve = FbxActor.LclRotation.GetCurve( BaseLayer, FBXSDK_CURVENODE_COMPONENT_Y, true );
+		ExportRichCurveToFbxCurve( TransformSection->GetRotationCurve( EAxis::Y ), *Curve, TransformSection->GetRange() );
+		Curve = FbxActor.LclRotation.GetCurve( BaseLayer, FBXSDK_CURVENODE_COMPONENT_Z, true );
+		ExportRichCurveToFbxCurve( TransformSection->GetRotationCurve( EAxis::Z ), *Curve, TransformSection->GetRange() );
+
+		// Scale
+		FbxActor.LclScaling.GetCurveNode( BaseLayer, true );
+		Curve = FbxActor.LclScaling.GetCurve( BaseLayer, FBXSDK_CURVENODE_COMPONENT_X, true );
+		ExportRichCurveToFbxCurve( TransformSection->GetScaleCurve( EAxis::X ), *Curve, TransformSection->GetRange() );
+		Curve = FbxActor.LclScaling.GetCurve( BaseLayer, FBXSDK_CURVENODE_COMPONENT_Y, true );
+		ExportRichCurveToFbxCurve( TransformSection->GetScaleCurve( EAxis::Y ), *Curve, TransformSection->GetRange() );
+		Curve = FbxActor.LclScaling.GetCurve( BaseLayer, FBXSDK_CURVENODE_COMPONENT_Z, true );
+		ExportRichCurveToFbxCurve( TransformSection->GetScaleCurve( EAxis::Z ), *Curve, TransformSection->GetRange() );
+	}
+}
+
+void FFbxExporter::ExportLevelSequenceFloatTrack( FbxNode& FbxActor, UMovieSceneFloatTrack& FloatTrack )
+{
+	// TODO: Support more than one section?
+	UMovieSceneFloatSection* FloatSection = FloatTrack.GetAllSections().Num() > 0
+		? Cast<UMovieSceneFloatSection>( FloatTrack.GetAllSections()[0] )
+		: nullptr;
+
+	if ( FloatSection == nullptr || FloatSection->GetFloatCurve().GetNumKeys() == 0 )
+	{
+		return;
+	}
+
+
+	FbxNodeAttribute* FbxNodeAttr = nullptr;
+	// camera and light is appended on the fbx pivot node
+	if ( FbxActor.GetChild( 0 ) )
+	{
+		FbxNodeAttr = ( (FbxNode*)FbxActor.GetChild( 0 ) )->GetNodeAttribute();
+		if ( FbxNodeAttr == nullptr )
+		{
+			return;
+		}
+	}
+
+	FbxProperty Property;
+	FString PropertyName = FloatTrack.GetTrackName().ToString();
+	bool IsFoV = false;
+	// most properties are created as user property, only FOV of camera in FBX supports animation
+	if ( PropertyName == "Intensity" )
+	{
+		Property = FbxActor.FindProperty( "UE_Intensity", false );
+	}
+	else if ( PropertyName == "FalloffExponent" )
+	{
+		Property = FbxActor.FindProperty( "UE_FalloffExponent", false );
+	}
+	else if ( PropertyName == "AttenuationRadius" )
+	{
+		Property = FbxActor.FindProperty( "UE_Radius", false );
+	}
+	else if ( PropertyName == "FOVAngle" && FbxNodeAttr )
+	{
+		Property = ( (FbxCamera*)FbxNodeAttr )->FocalLength;
+		IsFoV = true;
+	}
+	else if ( PropertyName == "AspectRatio" )
+	{
+		Property = FbxActor.FindProperty( "UE_AspectRatio", false );
+	}
+	else if ( PropertyName == "MotionBlur_Amount" )
+	{
+		Property = FbxActor.FindProperty( "UE_MotionBlur_Amount", false );
+	}
+
+	if ( Property != 0 )
+	{
+		FRichCurve& FloatCurve = FloatSection->GetFloatCurve();
+
+		FbxAnimCurve* AnimCurve = FbxAnimCurve::Create( Scene, "" );
+		FbxAnimCurveNode* CurveNode = Property.GetCurveNode( true );
+		if ( !CurveNode )
+		{
+			return;
+		}
+
+		CurveNode->SetChannelValue<double>( 0U, FloatCurve.GetDefaultValue() );
+		CurveNode->ConnectToChannel( AnimCurve, 0U );
+
+		ExportRichCurveToFbxCurve( FloatCurve, *AnimCurve, FloatSection->GetRange(), IsFoV ? ERichCurveValueMode::Fov : ERichCurveValueMode::Default );
+	}
+}
+
+
+FbxAnimCurveDef::EInterpolationType RichCurveInterpolationToFbxInterpolation(ERichCurveInterpMode Interpolation)
+{
+	switch ( Interpolation )
+	{
+		case ERichCurveInterpMode::RCIM_Cubic:
+			return FbxAnimCurveDef::eInterpolationCubic;
+		case ERichCurveInterpMode::RCIM_Linear:
+			return FbxAnimCurveDef::eInterpolationLinear;
+		default:
+			// ERichCurveInterpMode::RCIM_Constant and ERichCurveInterpMode::RCIM_None
+			return FbxAnimCurveDef::eInterpolationConstant;
+	}
+}
+
+FbxAnimCurveDef::ETangentMode RichCurveTangentModeToFbxTangentMode( ERichCurveTangentMode TangentMode )
+{
+	switch ( TangentMode )
+	{
+		case ERichCurveTangentMode::RCTM_Break:
+			return FbxAnimCurveDef::eTangentBreak;
+		case ERichCurveTangentMode::RCTM_User:
+			return FbxAnimCurveDef::eTangentUser;
+		default:
+			// ERichCurveTangentMode::RCTM_Auto and ERichCurveTangentMode::RCTM_None
+			return FbxAnimCurveDef::eTangentAuto;
+	}
+}
+
+void FFbxExporter::ExportRichCurveToFbxCurve( FRichCurve& RichCurve, FbxAnimCurve& FbxCurve, TRange<float> InterpolationRange,  ERichCurveValueMode ValueMode )
+{
+	FbxCurve.KeyModifyBegin();
+	if ( bBakeKeys )
+	{
+		float InterpLength = InterpolationRange.Size<float>();
+		int NumKeys = InterpLength * BakeTransformsFPS;
+		for ( int KeyIndex = 0; KeyIndex <= NumKeys; KeyIndex++ )
+		{
+			float Time = InterpolationRange.GetLowerBoundValue() + ( KeyIndex * InterpLength / NumKeys );
+			float CurveValue = ValueMode == ERichCurveValueMode::Fov
+				? DefaultCamera->ComputeFocalLength( RichCurve.Eval( Time ) )
+				: RichCurve.Eval( Time );
+
+			FbxTime FbxTime;
+			FbxTime.SetSecondDouble( Time );
+			int FbxKeyIndex = FbxCurve.KeyAdd( Time );
+			FbxCurve.KeySet( FbxKeyIndex, FbxTime, CurveValue, FbxAnimCurveDef::eInterpolationConstant, FbxAnimCurveDef::eTangentAuto );
+		}
+	}
+	else
+	{
+		for ( auto It = RichCurve.GetKeyHandleIterator(); It; ++It )
+		{
+			FRichCurveKey RichCurveKey = RichCurve.GetKey( It.Key() );
+			FKeyHandle NextKeyHandle = RichCurve.GetNextKey( It.Key() );
+			float NextTangent = RichCurve.IsKeyHandleValid( NextKeyHandle )
+				? RichCurve.GetKey( NextKeyHandle ).ArriveTangent
+				: 0;
+
+			FbxTime FbxTime;
+			FbxTime.SetSecondDouble( RichCurveKey.Time );
+			int FbxKeyIndex = FbxCurve.KeyAdd( FbxTime );
+
+			//TODO: Add support for tangent weights if they get added to FRichCurve.
+			float KeyValue = ValueMode == ERichCurveValueMode::Fov
+				? DefaultCamera->ComputeFocalLength( RichCurve.Eval( RichCurveKey.Time ) )
+				: RichCurve.Eval( RichCurveKey.Time );
+			FbxCurve.KeySet(
+				FbxKeyIndex,
+				FbxTime,
+				KeyValue,
+				RichCurveInterpolationToFbxInterpolation( RichCurveKey.InterpMode ),
+				RichCurveTangentModeToFbxTangentMode( RichCurveKey.TangentMode ),
+				RichCurveKey.LeaveTangent,
+				NextTangent );
+		}
+	}
+	FbxCurve.KeyModifyEnd();
+}
+
 /**
  * Finds the given actor in the already-exported list of structures
  */
diff --git a/Engine/Source/Editor/UnrealEd/Private/FbxExporter.h b/Engine/Source/Editor/UnrealEd/Private/FbxExporter.h
--- a/Engine/Source/Editor/UnrealEd/Private/FbxExporter.h
+++ b/Engine/Source/Editor/UnrealEd/Private/FbxExporter.h
@@ -14,6 +14,11 @@ class ALight;
 class UInterpTrackInstMove;
 class AMatineeActor;
 
+class UMovieScene;
+class IMovieScenePlayer;
+class UMovieScene3DTransformTrack;
+class UMovieSceneFloatTrack;
+
 namespace UnFbx
 {
 
@@ -72,6 +77,13 @@ public:
 	 */
 	virtual bool ExportMatinee(class AMatineeActor* InMatineeActor);
 
+	/**
+	* Exports the given level sequence information into a FBX document.
+	*
+	* @return	true, if successful
+	*/
+	bool ExportLevelSequence( UMovieScene* MovieScene, IMovieScenePlayer* MovieScenePlayer );
+
 	/**
 	 * Exports all the animation sequences part of a single Group in a Matinee sequence
 	 * as a single animation in the FBX document.  The animation is created by sampling the
@@ -132,11 +144,6 @@ public:
 	 */
 	void ExportAnimSequencesAsSingle( USkeletalMesh* SkelMesh, const ASkeletalMeshActor* SkelMeshActor, const FString& ExportName, const TArray<UAnimSequence*>& AnimSeqList, const TArray<FAnimControlTrackKey>& TrackKeys );
 
-	/**
-	 * Export Anim Track of the given SkeletalMeshComponent
-	 */
-	void ExportAnimTrack(class AMatineeActor* MatineeActor, USkeletalMeshComponent* SkeletalMeshComponent);
-
 private:
 	FFbxExporter();
 
@@ -162,6 +169,45 @@ private:
 	/** Whether or not to export vertices unwelded */
 	static bool bStaticMeshExportUnWeldedVerts;
 
+	/** Adapter interface which allows ExportAnimTrack to act on both sequencer and matinee data. */
+	class IAnimTrackAdapter
+	{
+	public:
+		/** Gets the length of the animation track. */
+		virtual float GetAnimationLength() const = 0;
+		/** Updates the runtime state of the animation track to the specified time. */
+		virtual void UpdateAnimation( float Time ) = 0;
+	};
+
+	/** An anim track adapter for matinee. */
+	class FMatineeAnimTrackAdapter : public IAnimTrackAdapter
+	{
+	public:
+		FMatineeAnimTrackAdapter( AMatineeActor* InMatineeActor );
+		virtual float GetAnimationLength() const override;
+		virtual void UpdateAnimation( float Time ) override;
+
+	private:
+		AMatineeActor* MatineeActor;
+	};
+
+	/** An anim track adapter for a level sequence. */
+	class FLevelSequenceAnimTrackAdapter : public FFbxExporter::IAnimTrackAdapter
+	{
+	public:
+		FLevelSequenceAnimTrackAdapter( IMovieScenePlayer* InMovieScenePlayer );
+		virtual float GetAnimationLength() const override;
+		virtual void UpdateAnimation( float Time ) override;
+
+	private:
+		IMovieScenePlayer* MovieScenePlayer;
+	};
+
+	/**
+	* Export Anim Track of the given SkeletalMeshComponent
+	*/
+	void ExportAnimTrack( IAnimTrackAdapter& AnimTrackAdapter, USkeletalMeshComponent* SkeletalMeshComponent );
+
 	void ExportModel(UModel* Model, FbxNode* Node, const char* Name);
 	
 	/**
@@ -286,6 +332,28 @@ private:
 	
 	void ExportAnimatedFloat(FbxProperty* FbxProperty, FInterpCurveFloat* Curve, bool IsCameraFoV);
 
+	/**
+	 * Exports a level sequence 3D transform track into the FBX animation stack.
+	 */
+	void ExportLevelSequence3DTransformTrack( FbxNode& FbxActor, UMovieScene3DTransformTrack& TransformTrack );
+
+	/** 
+	 * Exports a level sequence float track into the FBX animation stack. 
+	 */
+	void ExportLevelSequenceFloatTrack( FbxNode& FbxActor, UMovieSceneFloatTrack& FloatTrack );
+
+	/** Defines value export modes for the EportRichCurveToFbxCurve method. */
+	enum class ERichCurveValueMode
+	{
+		/** Export values directly */
+		Default,
+		/** Export fov values which get processed to focal length. */
+		Fov
+	};
+
+	/** Exports an unreal rich curve to an fbx animation curve. */
+	void ExportRichCurveToFbxCurve( FRichCurve& RichCurve, FbxAnimCurve& FbxCurve, TRange<float> InterpolationRange, ERichCurveValueMode ValueMode = ERichCurveValueMode::Default );
+
 	/**
 	 * Finds the given actor in the already-exported list of structures
 	 * @return FbxNode* the FBX node created from the UE4 actor
diff --git a/Engine/Source/Editor/UnrealEd/Private/Kismet2/KismetReinstanceUtilities.cpp b/Engine/Source/Editor/UnrealEd/Private/Kismet2/KismetReinstanceUtilities.cpp
--- a/Engine/Source/Editor/UnrealEd/Private/Kismet2/KismetReinstanceUtilities.cpp
+++ b/Engine/Source/Editor/UnrealEd/Private/Kismet2/KismetReinstanceUtilities.cpp
@@ -411,14 +411,14 @@ public:
 					// cached LinkInfo data may now be invalid. This could happen in the fast path, since the original
 					// Actor instance will not be replaced in that case, and thus might still have latent actions pending.
 					World->GetLatentActionManager().RemoveActionsForObject(Actor);
-				}
 
-				Actor->ReregisterAllComponents();
-				Actor->RerunConstructionScripts();
+					Actor->ReregisterAllComponents();
+					Actor->RerunConstructionScripts();
 
-				if (SelectedObjecs.Contains(Obj))
-				{
-					GEditor->SelectActor(Actor, /*bInSelected =*/true, /*bNotify =*/true, false, true);
+					if (SelectedObjecs.Contains(Obj))
+					{
+						GEditor->SelectActor(Actor, /*bInSelected =*/true, /*bNotify =*/true, false, true);
+					}
 				}
 			}
 		}
diff --git a/Engine/Source/Editor/UnrealEd/Private/LevelEditorViewport.cpp b/Engine/Source/Editor/UnrealEd/Private/LevelEditorViewport.cpp
--- a/Engine/Source/Editor/UnrealEd/Private/LevelEditorViewport.cpp
+++ b/Engine/Source/Editor/UnrealEd/Private/LevelEditorViewport.cpp
@@ -1503,7 +1503,6 @@ FLevelEditorViewportClient::FLevelEditorViewportClient(const TSharedPtr<SLevelVi
 	, FadeAmount(0.f)	
 	, bEnableFading(false)
 	, bEnableColorScaling(false)
-	, bEditorCameraCut(false)
 	, bDrawBaseInfo(false)
 	, bDuplicateOnNextDrag( false )
 	, bDuplicateActorsInProgress( false )
@@ -1520,6 +1519,8 @@ FLevelEditorViewportClient::FLevelEditorViewportClient(const TSharedPtr<SLevelVi
 	, bWasControlledByOtherViewport(false)
 	, ActorLockedByMatinee(nullptr)
 	, ActorLockedToCamera(nullptr)
+	, bEditorCameraCut(false)
+	, bWasEditorCameraCut(false)
 {
 	// By default a level editor viewport is pointed to the editor world
 	SetReferenceToWorldContext(GEditor->GetEditorWorldContext());
@@ -2002,6 +2003,12 @@ static FMatrix GPerspViewMatrix;
 
 void FLevelEditorViewportClient::Tick(float DeltaTime)
 {
+	if (bWasEditorCameraCut && bEditorCameraCut)
+	{
+		bEditorCameraCut = false;
+	}
+	bWasEditorCameraCut = bEditorCameraCut;
+
 	FEditorViewportClient::Tick(DeltaTime);
 
 	if (!GUnrealEd->IsPivotMovedIndependently() && GCurrentLevelEditingViewportClient == this &&
diff --git a/Engine/Source/Editor/UnrealEd/Public/LevelEditorViewport.h b/Engine/Source/Editor/UnrealEd/Public/LevelEditorViewport.h
--- a/Engine/Source/Editor/UnrealEd/Public/LevelEditorViewport.h
+++ b/Engine/Source/Editor/UnrealEd/Public/LevelEditorViewport.h
@@ -161,7 +161,14 @@ public:
 
 	virtual bool OverrideHighResScreenshotCaptureRegion(FIntRect& OutCaptureRegion) override;
 
-	void SetIsCameraCut( bool bInIsCameraCut ) { bEditorCameraCut = bInIsCameraCut; }
+	void SetIsCameraCut( bool bInIsCameraCut )
+	{
+		bEditorCameraCut = bInIsCameraCut;
+		if (bEditorCameraCut)
+		{
+			bWasEditorCameraCut = false;
+		}
+	}
 
 	/** 
 	 * Initialize visibility flags
@@ -708,9 +715,6 @@ public:
 
 	bool					bEnableColorScaling;
 
-	/** If true, we switched between two different cameras. Set by matinee, used by the motion blur to invalidate this frames motion vectors */
-	bool					bEditorCameraCut;
-
 	/** Indicates whether, of not, the base attachment volume should be drawn for this viewport. */
 	bool bDrawBaseInfo;
 
@@ -775,4 +779,10 @@ private:
 
 	/** Those sound stat flags which are enabled on this viewport */
 	ESoundShowFlags::Type	SoundShowFlags;
+
+	/** If true, we switched between two different cameras. Set by matinee, used by the motion blur to invalidate this frames motion vectors */
+	bool					bEditorCameraCut;
+
+	/** Stores the previous frame's value of bEditorCameraCut in order to reset it back to false on the next frame */
+	bool					bWasEditorCameraCut;
 };
diff --git a/Engine/Source/Editor/UnrealEd/UnrealEd.Build.cs b/Engine/Source/Editor/UnrealEd/UnrealEd.Build.cs
--- a/Engine/Source/Editor/UnrealEd/UnrealEd.Build.cs
+++ b/Engine/Source/Editor/UnrealEd/UnrealEd.Build.cs
@@ -141,7 +141,9 @@ public class UnrealEd : ModuleRules
 				"Analytics",
                 "AnalyticsET",
                 "PluginWarden",
-                "PixelInspectorModule"
+                "PixelInspectorModule",
+				"MovieScene",
+				"MovieSceneTracks",
             }
 		);
 
diff --git a/Engine/Source/Runtime/CinematicCamera/Private/CineCameraComponent.cpp b/Engine/Source/Runtime/CinematicCamera/Private/CineCameraComponent.cpp
--- a/Engine/Source/Runtime/CinematicCamera/Private/CineCameraComponent.cpp
+++ b/Engine/Source/Runtime/CinematicCamera/Private/CineCameraComponent.cpp
@@ -216,7 +216,7 @@ void UCineCameraComponent::GetCameraView(float DeltaTime, FMinimalViewInfo& Desi
 	UpdateCameraLens(DeltaTime, DesiredView);
 
 #if WITH_EDITORONLY_DATA
-	if (FocusSettings.bDrawDebugFocusPlane && DebugFocusPlaneMesh)
+	if (FocusSettings.bDrawDebugFocusPlane && DebugFocusPlaneMesh && DebugFocusPlaneComponent)
 	{
 		FVector const CamDir = DesiredView.Rotation.Vector();
 		FVector const FocusPoint = DesiredView.Location + CamDir * CurrentFocusDistance;
diff --git a/Engine/Source/Runtime/CinematicCamera/Public/CineCameraComponent.h b/Engine/Source/Runtime/CinematicCamera/Public/CineCameraComponent.h
--- a/Engine/Source/Runtime/CinematicCamera/Public/CineCameraComponent.h
+++ b/Engine/Source/Runtime/CinematicCamera/Public/CineCameraComponent.h
@@ -142,11 +142,11 @@ struct FCameraFocusSettings
 
 #if WITH_EDITORONLY_DATA
 	/** True to draw a translucent plane at the current focus depth, for easy tweaking. */
-	UPROPERTY(Transient, EditAnywhere, Category = "Focus Settings")
+	UPROPERTY(Transient, BlueprintReadWrite, EditAnywhere, Category = "Focus Settings")
 	uint8 bDrawDebugFocusPlane : 1;
 
 	/** For customizing the focus plane color, in case the default doesn't show up well in your scene. */
-	UPROPERTY(EditAnywhere, Category = "Focus Settings", meta = (EditCondition = "bDrawDebugFocusPlane"))
+	UPROPERTY(EditAnywhere, BlueprintReadWrite, Category = "Focus Settings", meta = (EditCondition = "bDrawDebugFocusPlane"))
 	FColor DebugFocusPlaneColor;
 #endif 
 
diff --git a/Engine/Source/Runtime/Engine/Classes/Curves/CurveBase.h b/Engine/Source/Runtime/Engine/Classes/Curves/CurveBase.h
--- a/Engine/Source/Runtime/Engine/Classes/Curves/CurveBase.h
+++ b/Engine/Source/Runtime/Engine/Classes/Curves/CurveBase.h
@@ -538,6 +538,11 @@ public:
 	  */
 	FKeyHandle AddKey(float InTime, float InValue, const bool bUnwindRotation = false, FKeyHandle KeyHandle = FKeyHandle());
 
+	/**
+	 * Sets the keys with the supplied times and values. Expects that the times and values are already sorted.
+	 */
+	void SetKeys(const TArray<float>& InTimes, const TArray<float>& InValues);
+
 	/** Remove the specified key from the curve.*/
 	void DeleteKey(FKeyHandle KeyHandle);
 
diff --git a/Engine/Source/Runtime/Engine/Classes/Engine/Engine.h b/Engine/Source/Runtime/Engine/Classes/Engine/Engine.h
--- a/Engine/Source/Runtime/Engine/Classes/Engine/Engine.h
+++ b/Engine/Source/Runtime/Engine/Classes/Engine/Engine.h
@@ -2261,6 +2261,7 @@ public:
 
 		/** Skips copying properties with BlueprintCompilerGeneratedDefaults metadata */
 		bool bSkipCompilerGeneratedDefaults;
+		bool bNotifyObjectReplacement;
 
 		FCopyPropertiesForUnrelatedObjectsParams()
 			: bAggressiveDefaultSubobjectReplacement(false)
@@ -2269,6 +2270,7 @@ public:
 			, bCopyDeprecatedProperties(false)
 			, bPreserveRootComponent(true)
 			, bSkipCompilerGeneratedDefaults(false)
+			, bNotifyObjectReplacement(true)
 		{}
 	};
 	static void CopyPropertiesForUnrelatedObjects(UObject* OldObject, UObject* NewObject, FCopyPropertiesForUnrelatedObjectsParams Params = FCopyPropertiesForUnrelatedObjectsParams());//bool bAggressiveDefaultSubobjectReplacement = false, bool bDoDelta = true);
diff --git a/Engine/Source/Runtime/Engine/Private/CurveBase.cpp b/Engine/Source/Runtime/Engine/Private/CurveBase.cpp
--- a/Engine/Source/Runtime/Engine/Private/CurveBase.cpp
+++ b/Engine/Source/Runtime/Engine/Private/CurveBase.cpp
@@ -697,6 +697,19 @@ FKeyHandle FRichCurve::AddKey( const float InTime, const float InValue, const bo
 	return GetKeyHandle(Index);
 }
 
+void FRichCurve::SetKeys(const TArray<float>& InTimes, const TArray<float>& InValues)
+{
+	Reset();
+
+	for (int32 Index = 0; Index < InTimes.Num(); ++Index)
+	{
+		Keys.Add(FRichCurveKey(InTimes[Index], InValues[Index]));
+		KeyHandlesToIndices.Add(FKeyHandle(), Index);
+	}
+
+	AutoSetTangents();
+}
+
 
 void FRichCurve::DeleteKey(FKeyHandle InKeyHandle)
 {
diff --git a/Engine/Source/Runtime/Engine/Private/UnrealEngine.cpp b/Engine/Source/Runtime/Engine/Private/UnrealEngine.cpp
--- a/Engine/Source/Runtime/Engine/Private/UnrealEngine.cpp
+++ b/Engine/Source/Runtime/Engine/Private/UnrealEngine.cpp
@@ -11256,7 +11256,7 @@ void UEngine::CopyPropertiesForUnrelatedObjects(UObject* OldObject, UObject* New
 	}
 
 	// Now notify any tools that aren't already updated via the FArchiveReplaceObjectRef path
-	if (GEngine != nullptr)
+	if (Params.bNotifyObjectReplacement && GEngine != nullptr)
 	{
 		GEngine->NotifyToolsOfObjectReplacement(ReferenceReplacementMap);
 	}
diff --git a/Engine/Source/Runtime/LevelSequence/Private/LevelSequence.cpp b/Engine/Source/Runtime/LevelSequence/Private/LevelSequence.cpp
--- a/Engine/Source/Runtime/LevelSequence/Private/LevelSequence.cpp
+++ b/Engine/Source/Runtime/LevelSequence/Private/LevelSequence.cpp
@@ -8,6 +8,12 @@
 
 DEFINE_LOG_CATEGORY_STATIC(LogLevelSequence, Log, All);
 
+static TAutoConsoleVariable<int32> CVarFixedFrameIntervalPlayback(
+	TEXT("LevelSequence.DefaultFixedFrameIntervalPlayback"),
+	1,
+	TEXT("When non-zero, all newly created level sequences will default to fixed frame interval playback."),
+	ECVF_Default);
+
 ULevelSequence::ULevelSequence(const FObjectInitializer& ObjectInitializer)
 	: Super(ObjectInitializer)
 	, MovieScene(nullptr)
@@ -18,10 +24,32 @@ void ULevelSequence::Initialize()
 {
 	// @todo sequencer: gmp: fix me
 	MovieScene = NewObject<UMovieScene>(this, NAME_None, RF_Transactional);
-	MovieScene->SetForceFixedFrameIntervalPlayback( true );
+
+	const bool bForceFixedPlayback = CVarFixedFrameIntervalPlayback.GetValueOnGameThread() != 0;
+
+	MovieScene->SetForceFixedFrameIntervalPlayback( bForceFixedPlayback );
 	MovieScene->SetFixedFrameInterval( 1 / 30.0f );
 }
 
+UObject* ULevelSequence::MakeSpawnableTemplateFromInstance(UObject& InSourceObject, FName ObjectName)
+{
+	UObject* NewInstance = NewObject<UObject>(MovieScene, InSourceObject.GetClass(), ObjectName);
+
+	UEngine::FCopyPropertiesForUnrelatedObjectsParams CopyParams;
+	CopyParams.bNotifyObjectReplacement = false;
+	UEngine::CopyPropertiesForUnrelatedObjects(&InSourceObject, NewInstance, CopyParams);
+
+	AActor* Actor = CastChecked<AActor>(NewInstance);
+	if (Actor->GetAttachParentActor() != nullptr)
+	{
+		// We don't support spawnables and attachments right now
+		// @todo: map to attach track?
+		Actor->DetachFromActor(FDetachmentTransformRules(FAttachmentTransformRules(EAttachmentRule::KeepRelative, false), false));
+	}
+
+	return NewInstance;
+}
+
 void ULevelSequence::PostLoad()
 {
 	Super::PostLoad();
@@ -41,7 +69,7 @@ void ULevelSequence::PostLoad()
 				UObject* NewTemplate = NewObject<UObject>(MovieScene, Spawnable.GeneratedClass_DEPRECATED, TemplateName);
 				if (NewTemplate)
 				{
-					Spawnable.CopyObjectTemplate(*NewTemplate, *MovieScene);
+					Spawnable.CopyObjectTemplate(*NewTemplate, *this);
 				}
 			}
 		}
diff --git a/Engine/Source/Runtime/LevelSequence/Private/LevelSequencePlayer.cpp b/Engine/Source/Runtime/LevelSequence/Private/LevelSequencePlayer.cpp
--- a/Engine/Source/Runtime/LevelSequence/Private/LevelSequencePlayer.cpp
+++ b/Engine/Source/Runtime/LevelSequence/Private/LevelSequencePlayer.cpp
@@ -71,6 +71,7 @@ ULevelSequencePlayer::ULevelSequencePlayer(const FObjectInitializer& ObjectIniti
 	: Super(ObjectInitializer)
 	, LevelSequence(nullptr)
 	, bIsPlaying(false)
+	, bReversePlayback(false)
 	, TimeCursorPosition(0.0f)
 	, LastCursorPosition(0.0f)
 	, StartTime(0.f)
@@ -115,18 +116,32 @@ bool ULevelSequencePlayer::IsPlaying() const
 
 void ULevelSequencePlayer::Pause()
 {
-	bIsPlaying = false;
+	if (bIsPlaying)
+	{
+		bIsPlaying = false;
+
+		if (OnPause.IsBound())
+		{
+			OnPause.Broadcast();
+		}
+	}
 }
 
 void ULevelSequencePlayer::Stop()
 {
-	bIsPlaying = false;
-	TimeCursorPosition = PlaybackSettings.PlayRate < 0.f ? GetLength() : 0.f;
-	CurrentNumLoops = 0;
+	if (bIsPlaying)
+	{
+		bIsPlaying = false;
+		TimeCursorPosition = PlaybackSettings.PlayRate < 0.f ? GetLength() : 0.f;
+		CurrentNumLoops = 0;
 
-	SetTickPrerequisites(false);
+		SetTickPrerequisites(false);
 
-	// todo: Trigger an event?
+		if (OnStop.IsBound())
+		{
+			OnStop.Broadcast();
+		}
+	}
 }
 
 void GetDescendantSequences(UMovieSceneSequence* InSequence, TArray<UMovieSceneSequence*> & InSequences)
@@ -243,18 +258,47 @@ void ULevelSequencePlayer::SetTickPrerequisites(bool bAddTickPrerequisites)
 
 void ULevelSequencePlayer::Play()
 {
-	// Start playing
-	StartPlayingNextTick();
+	bReversePlayback = false;
 
-	// Update now
-	bPendingFirstUpdate = false;
-	UpdateMovieSceneInstance(TimeCursorPosition, TimeCursorPosition);
+	PlayInternal();
+}
+
+void ULevelSequencePlayer::PlayReverse() 
+{
+	bReversePlayback = true;
+
+	PlayInternal();
+}
+
+void ULevelSequencePlayer::ChangePlaybackDirection()
+{
+	bReversePlayback = !bReversePlayback;
+
+	PlayInternal();
+}
+
+void ULevelSequencePlayer::PlayInternal()
+{
+	if (!bIsPlaying)
+	{
+		// Start playing
+		StartPlayingNextTick();
+
+		// Update now
+		bPendingFirstUpdate = false;
+		UpdateMovieSceneInstance(TimeCursorPosition, TimeCursorPosition);
+
+		if (OnPlay.IsBound())
+		{
+			OnPlay.Broadcast();
+		}
+	}
 }
 
 void ULevelSequencePlayer::PlayLooping(int32 NumLoops)
 {
 	PlaybackSettings.LoopCount = NumLoops;
-	Play();
+	PlayInternal();
 }
 
 void ULevelSequencePlayer::StartPlayingNextTick()
@@ -323,7 +367,7 @@ void ULevelSequencePlayer::UpdateTimeCursorPosition(float NewPosition)
 		bPendingFirstUpdate = false;
 	}
 
-	if ((NewPosition >= Length) || (NewPosition < 0))
+	if (ShouldStopOrLoop(NewPosition))
 	{
 		// loop playback
 		if (PlaybackSettings.LoopCount < 0 || CurrentNumLoops < PlaybackSettings.LoopCount)
@@ -339,19 +383,34 @@ void ULevelSequencePlayer::UpdateTimeCursorPosition(float NewPosition)
 		}
 
 		// stop playback
-		bIsPlaying = false;
-		CurrentNumLoops = 0;
+		Stop();
 	}
 
 	LastCursorPosition = TimeCursorPosition;
 	TimeCursorPosition = NewPosition;
 }
 
+bool ULevelSequencePlayer::ShouldStopOrLoop(float NewPosition)
+{
+	bool bShouldStopOrLoop = false;
+	if (!bReversePlayback)
+	{
+		bShouldStopOrLoop = NewPosition >= GetLength();
+	}
+	else
+	{
+		bShouldStopOrLoop = NewPosition < 0.f;
+	}
+
+	return bShouldStopOrLoop;
+}
+
 /* ULevelSequencePlayer implementation
  *****************************************************************************/
 
 void ULevelSequencePlayer::Initialize(ULevelSequence* InLevelSequence, UWorld* InWorld, const FLevelSequencePlaybackSettings& Settings)
 {
+	CurrentPlayer = this;
 	LevelSequence = InLevelSequence;
 
 	World = InWorld;
@@ -381,7 +440,7 @@ void ULevelSequencePlayer::GetRuntimeObjects(TSharedRef<FMovieSceneSequenceInsta
 }
 
 
-void ULevelSequencePlayer::UpdateCameraCut(UObject* CameraObject, UObject* UnlockIfCameraObject, bool bJumpCut) const
+void ULevelSequencePlayer::UpdateCameraCut(UObject* CameraObject, UObject* UnlockIfCameraObject, bool bJumpCut)
 {
 	// skip missing player controller
 	APlayerController* PC = World->GetGameInstance()->GetFirstLocalPlayerController();
@@ -476,30 +535,36 @@ TArray<UObject*> ULevelSequencePlayer::GetEventContexts() const
 	TArray<UObject*> EventContexts;
 	if (World.IsValid())
 	{
-		if (World->GetLevelScriptActor())
-		{
-			EventContexts.Add(World->GetLevelScriptActor());
-		}
+		GetEventContexts(*World, EventContexts);
+	}
+	return EventContexts;
+}
 
-		for (ULevelStreaming* StreamingLevel : World->StreamingLevels)
+void ULevelSequencePlayer::GetEventContexts(UWorld& InWorld, TArray<UObject*>& OutContexts)
+{
+	if (InWorld.GetLevelScriptActor())
+	{
+		OutContexts.Add(InWorld.GetLevelScriptActor());
+	}
+
+	for (ULevelStreaming* StreamingLevel : InWorld.StreamingLevels)
+	{
+		if (StreamingLevel->GetLevelScriptActor())
 		{
-			if (StreamingLevel->GetLevelScriptActor())
-			{
-				EventContexts.Add(StreamingLevel->GetLevelScriptActor());
-			}
+			OutContexts.Add(StreamingLevel->GetLevelScriptActor());
 		}
 	}
-	return EventContexts;
 }
 
 void ULevelSequencePlayer::Update(const float DeltaSeconds)
 {
 	if (bIsPlaying)
 	{
-		UpdateTimeCursorPosition(TimeCursorPosition + DeltaSeconds * PlaybackSettings.PlayRate);
+		float PlayRate = bReversePlayback ? -PlaybackSettings.PlayRate : PlaybackSettings.PlayRate;
+		UpdateTimeCursorPosition(TimeCursorPosition + DeltaSeconds * PlayRate);
 		UpdateMovieSceneInstance(TimeCursorPosition, LastCursorPosition);
 	}
-	else if (!bHasCleanedUpSequence && TimeCursorPosition >= GetLength())
+	else if (!bHasCleanedUpSequence && ShouldStopOrLoop(TimeCursorPosition))
 	{
 		UpdateMovieSceneInstance(TimeCursorPosition, LastCursorPosition);
 
@@ -561,7 +626,11 @@ void ULevelSequencePlayer::TakeFrameSnapshot(FLevelSequencePlayerSnapshot& OutSn
 
 		if (ActiveShot)
 		{
-			const float ShotOffset = ActiveShot->StartOffset + ActiveShot->GetSequence()->GetMovieScene()->GetPlaybackRange().GetLowerBoundValue() - ActiveShot->PrerollTime;
+			// Assume that shots with no sequence start at 0.
+			const float ShotLowerBound = ActiveShot->GetSequence() 
+				? ActiveShot->GetSequence()->GetMovieScene()->GetPlaybackRange().GetLowerBoundValue() 
+				: 0;
+			const float ShotOffset = ActiveShot->StartOffset + ShotLowerBound - ActiveShot->PrerollTime;
 			const float ShotPosition = ShotOffset + (CurrentTime - (ActiveShot->GetStartTime() - ActiveShot->PrerollTime)) / ActiveShot->TimeScale;
 
 			OutSnapshot.CurrentShotName = ActiveShot->GetShotDisplayName();
diff --git a/Engine/Source/Runtime/LevelSequence/Public/LevelSequence.h b/Engine/Source/Runtime/LevelSequence/Public/LevelSequence.h
--- a/Engine/Source/Runtime/LevelSequence/Public/LevelSequence.h
+++ b/Engine/Source/Runtime/LevelSequence/Public/LevelSequence.h
@@ -42,9 +42,11 @@ public:
 	virtual UObject* GetParentObject(UObject* Object) const override;
 	virtual void UnbindPossessableObjects(const FGuid& ObjectId) override;
 	virtual bool AllowsSpawnableObjects() const override;
+	virtual UObject* MakeSpawnableTemplateFromInstance(UObject& InSourceObject, FName ObjectName) override;
+	
 
 	virtual void PostLoad() override;
-	
+
 	/** Bind a posessable object with an explicitly-supplied ObjectReference */
 	void BindPossessableObject(const FGuid& ObjectId, const FLevelSequenceObjectReference& ObjectReference);
 
diff --git a/Engine/Source/Runtime/LevelSequence/Public/LevelSequencePlayer.h b/Engine/Source/Runtime/LevelSequence/Public/LevelSequencePlayer.h
--- a/Engine/Source/Runtime/LevelSequence/Public/LevelSequencePlayer.h
+++ b/Engine/Source/Runtime/LevelSequence/Public/LevelSequencePlayer.h
@@ -11,6 +11,8 @@ class FMovieSceneSequenceInstance;
 class ULevel;
 class UMovieSceneBindings;
 
+DECLARE_DYNAMIC_MULTICAST_DELEGATE(FOnLevelSequencePlayerEvent);
+
 USTRUCT(BlueprintType)
 struct FLevelSequenceSnapshotSettings
 {
@@ -116,10 +118,18 @@ public:
 	UFUNCTION(BlueprintCallable, Category="Game|Cinematic", meta=(WorldContext="WorldContextObject"))
 	static ULevelSequencePlayer* CreateLevelSequencePlayer(UObject* WorldContextObject, ULevelSequence* LevelSequence, FLevelSequencePlaybackSettings Settings);
 
-	/** Start playback from the current time cursor position, using the current play rate. */
+	/** Start playback forwards from the current time cursor position, using the current play rate. */
 	UFUNCTION(BlueprintCallable, Category="Game|Cinematic")
 	void Play();
 
+	/** Reverse playback. */
+	UFUNCTION(BlueprintCallable, Category="Game|Cinematic")
+	void PlayReverse();
+
+	/** Changes the direction of playback (go in reverse if it was going forward, or vice versa) */
+	UFUNCTION(BlueprintCallable, Category="Game|Cinematic")
+	void ChangePlaybackDirection();
+
 	/**
 	 * Start playback from the current time cursor position, looping the specified number of times.
 	 * @param NumLoops - The number of loops to play. -1 indicates infinite looping.
@@ -130,7 +140,7 @@ public:
 	/** Start playback from the current time cursor position, using the current play rate. Does not update the animation until next tick. */
 	UFUNCTION(BlueprintCallable, Category="Game|Cinematic")
 	void StartPlayingNextTick();
-
+	
 	/** Pause playback. */
 	UFUNCTION(BlueprintCallable, Category="Game|Cinematic")
 	void Pause();
@@ -202,7 +212,7 @@ protected:
 
 	// IMovieScenePlayer interface
 	virtual void GetRuntimeObjects(TSharedRef<FMovieSceneSequenceInstance> MovieSceneInstance, const FGuid& ObjectHandle, TArray<TWeakObjectPtr<UObject>>& OutObjects) const override;
-	virtual void UpdateCameraCut(UObject* CameraObject, UObject* UnlockIfCameraObject, bool bJumpCut) const override;
+	virtual void UpdateCameraCut(UObject* CameraObject, UObject* UnlockIfCameraObject, bool bJumpCut) override;
 	virtual void SetViewportSettings(const TMap<FViewportClient*, EMovieSceneViewportParams>& ViewportParamsMap) override;
 	virtual void GetViewportSettings(TMap<FViewportClient*, EMovieSceneViewportParams>& ViewportParamsMap) const override;
 	virtual EMovieScenePlayerStatus::Type GetPlaybackStatus() const override;
@@ -216,6 +226,9 @@ protected:
 
 public:
 
+	/** Populate the specified array with any given event contexts for the specified world */
+	static void GetEventContexts(UWorld& InWorld, TArray<UObject*>& OutContexts);
+
 	void Update(const float DeltaSeconds);
 
 	/** Take a snapshot of the current state of this player */
@@ -223,6 +236,9 @@ public:
 
 private:
 
+	/** Internal play function. */
+	void PlayInternal();
+
 	/** Update the movie scene instance from the specified previous position, to the specified time position. */
 	void UpdateMovieSceneInstance(float CurrentPosition, float PreviousPosition);
 
@@ -232,16 +248,27 @@ private:
 	/** Add tick prerequisites so that the level sequence actor ticks before all the actors it controls */
 	void SetTickPrerequisites(bool bAddTickPrerequisites);
 	
+	/** Returns whether playback should be stopped or looped because it has gone out of the playback bounds. */
+	bool ShouldStopOrLoop(float NewPosition);
+
 private:
 
 	/** The level sequence to play. */
 	UPROPERTY(transient)
 	ULevelSequence* LevelSequence;
 
+	/** The level sequence player. */
+	UPROPERTY(transient)
+	ULevelSequencePlayer* CurrentPlayer;
+
 	/** Whether we're currently playing. If false, then sequence playback is paused or was never started. */
 	UPROPERTY()
 	bool bIsPlaying;
 
+	/** Whether we're currently playing in reverse. */
+	UPROPERTY()
+	bool bReversePlayback;
+
 	/** True where we're waiting for the first update of the sequence after calling StartPlayingNextTick. */
 	bool bPendingFirstUpdate;
 
@@ -280,7 +307,7 @@ private:
 	TSharedPtr<FLevelSequenceSpawnRegister> SpawnRegister;
 
 	/** The last view target to reset to when updating camera cuts to null */
-	mutable TWeakObjectPtr<AActor> LastViewTarget;
+	TWeakObjectPtr<AActor> LastViewTarget;
 
 protected:
 
@@ -293,6 +320,18 @@ public:
 	DECLARE_EVENT_ThreeParams( ULevelSequencePlayer, FOnLevelSequencePlayerUpdated, const ULevelSequencePlayer&, float /*current time*/, float /*previous time*/ );
 	FOnLevelSequencePlayerUpdated& OnSequenceUpdated() const { return OnLevelSequencePlayerUpdate; }
 
+	/** Event triggered when the level sequence player is played */
+	UPROPERTY(BlueprintAssignable, Category="Game|Cinematic")
+	FOnLevelSequencePlayerEvent OnPlay;
+
+	/** Event triggered when the level sequence player is stopped */
+	UPROPERTY(BlueprintAssignable, Category="Game|Cinematic")
+	FOnLevelSequencePlayerEvent OnStop;
+
+	/** Event triggered when the level sequence player is paused */
+	UPROPERTY(BlueprintAssignable, Category="Game|Cinematic")
+	FOnLevelSequencePlayerEvent OnPause;
+
 private:
 
 	/** The event that will be broadcast every time the sequence is updated */
diff --git a/Engine/Source/Runtime/MovieScene/Private/MovieScene.cpp b/Engine/Source/Runtime/MovieScene/Private/MovieScene.cpp
--- a/Engine/Source/Runtime/MovieScene/Private/MovieScene.cpp
+++ b/Engine/Source/Runtime/MovieScene/Private/MovieScene.cpp
@@ -630,6 +630,17 @@ void UMovieScene::UpgradeTimeRanges()
 void UMovieScene::PostLoad()
 {
 	UpgradeTimeRanges();
+
+	for (FMovieSceneSpawnable& Spawnable : Spawnables)
+	{
+		if (UObject* Template = Spawnable.GetObjectTemplate())
+		{
+			// Spawnables are no longer marked archetype
+			Template->ClearFlags(RF_ArchetypeObject);
+			
+			FMovieSceneSpawnable::MarkSpawnableTemplate(*Template);
+		}
+	}
 	Super::PostLoad();
 }
 
diff --git a/Engine/Source/Runtime/MovieScene/Private/MovieSceneCommonHelpers.cpp b/Engine/Source/Runtime/MovieScene/Private/MovieSceneCommonHelpers.cpp
--- a/Engine/Source/Runtime/MovieScene/Private/MovieSceneCommonHelpers.cpp
+++ b/Engine/Source/Runtime/MovieScene/Private/MovieSceneCommonHelpers.cpp
@@ -345,22 +345,26 @@ void FTrackInstancePropertyBindings::UpdateBindings( const TArray<TWeakObjectPtr
 {
 	for (auto ObjectPtr : InRuntimeObjects)
 	{
-		UObject* Object = ObjectPtr.Get();
+		UpdateBinding(ObjectPtr);
+	}
+}
 
-		if (Object != nullptr)
-		{
-			FPropertyAndFunction PropAndFunction;
-			{
-				PropAndFunction.Function = Object->FindFunction(FunctionName);
-				PropAndFunction.PropertyAddress = FindProperty(Object, PropertyPath);
-			}
+void FTrackInstancePropertyBindings::UpdateBinding(const TWeakObjectPtr<UObject>& InRuntimeObject)
+{
+	UObject* Object = InRuntimeObject.Get();
 
-			RuntimeObjectToFunctionMap.Add(Object, PropAndFunction);
+	if (Object != nullptr)
+	{
+		FPropertyAndFunction PropAndFunction;
+		{
+			PropAndFunction.Function = Object->FindFunction(FunctionName);
+			PropAndFunction.PropertyAddress = FindProperty(Object, PropertyPath);
 		}
+
+		RuntimeObjectToFunctionMap.Add(Object, PropAndFunction);
 	}
 }
 
-
 UProperty* FTrackInstancePropertyBindings::GetProperty(const UObject* Object) const
 {
 	UProperty* Prop = nullptr;
diff --git a/Engine/Source/Runtime/MovieScene/Private/MovieSceneSpawnable.cpp b/Engine/Source/Runtime/MovieScene/Private/MovieSceneSpawnable.cpp
--- a/Engine/Source/Runtime/MovieScene/Private/MovieSceneSpawnable.cpp
+++ b/Engine/Source/Runtime/MovieScene/Private/MovieSceneSpawnable.cpp
@@ -4,20 +4,45 @@
 #include "MovieSceneSpawnable.h"
 #include "MovieScene.h"
 
+struct FIsSpawnable
+{
+	FIsSpawnable() : bIsSpawnable(false) {}
+	explicit FIsSpawnable(bool bInIsSpawnable) : bIsSpawnable(bInIsSpawnable) {}
+
+	bool IsDefault() const { return !bIsSpawnable; }
+
+	bool bIsSpawnable;
+};
+
+static FUObjectAnnotationSparse<FIsSpawnable,true> SpawnablesAnnotation;
+
+bool FMovieSceneSpawnable::IsSpawnableTemplate(const UObject& InObject)
+{
+	return !SpawnablesAnnotation.GetAnnotation(&InObject).IsDefault();
+}
+
+void FMovieSceneSpawnable::MarkSpawnableTemplate(const UObject& InObject)
+{
+	SpawnablesAnnotation.AddAnnotation(&InObject, FIsSpawnable(true));
+}
 
-void FMovieSceneSpawnable::CopyObjectTemplate(UObject& InSourceObject, UMovieScene& OwnerMovieScene)
+void FMovieSceneSpawnable::CopyObjectTemplate(UObject& InSourceObject, UMovieSceneSequence& MovieSceneSequence)
 {
 	const FName ObjectName = ObjectTemplate ? ObjectTemplate->GetFName() : InSourceObject.GetFName();
 
 	if (ObjectTemplate)
 	{
+		ObjectTemplate->Rename(*MakeUniqueObjectName(MovieSceneSequence.GetMovieScene(), ObjectTemplate->GetClass(), "ExpiredSpawnable").ToString());
 		ObjectTemplate->MarkPendingKill();
 		ObjectTemplate = nullptr;
 	}
 
-	ObjectTemplate = StaticDuplicateObject(&InSourceObject, &OwnerMovieScene, ObjectName, RF_AllFlags & ~(RF_Transient));
-	ObjectTemplate->SetFlags(RF_ArchetypeObject);
+	ObjectTemplate = MovieSceneSequence.MakeSpawnableTemplateFromInstance(InSourceObject, ObjectName);
+
+	check(ObjectTemplate);
+
+	MarkSpawnableTemplate(*ObjectTemplate);
 
 	// @todo: this will mark the package as dirty whenever a spawnable is destroyed. We should diff the duplicated object, and only mark dirty where the spawnable has actually changed
-	OwnerMovieScene.MarkPackageDirty();
+	MovieSceneSequence.MarkPackageDirty();
 }
\ No newline at end of file
diff --git a/Engine/Source/Runtime/MovieScene/Public/IMovieScenePlayer.h b/Engine/Source/Runtime/MovieScene/Public/IMovieScenePlayer.h
--- a/Engine/Source/Runtime/MovieScene/Public/IMovieScenePlayer.h
+++ b/Engine/Source/Runtime/MovieScene/Public/IMovieScenePlayer.h
@@ -73,7 +73,7 @@ public:
 	 * @param UnlockIfCameraObject If this is not nullptr, release actor lock only if currently locked to this object.
 	 * @param bJumpCut Whether this is a jump cut, ie. the cut jumps from one shot to another shot
 	 */
-	virtual void UpdateCameraCut(UObject* CameraObject, UObject* UnlockIfCameraObject = nullptr, bool bJumpCut = false) const = 0;
+	virtual void UpdateCameraCut(UObject* CameraObject, UObject* UnlockIfCameraObject = nullptr, bool bJumpCut = false) = 0;
 
 	/*
 	 * Set the perspective viewport settings
diff --git a/Engine/Source/Runtime/MovieScene/Public/IMovieSceneTrackInstance.h b/Engine/Source/Runtime/MovieScene/Public/IMovieSceneTrackInstance.h
--- a/Engine/Source/Runtime/MovieScene/Public/IMovieSceneTrackInstance.h
+++ b/Engine/Source/Runtime/MovieScene/Public/IMovieSceneTrackInstance.h
@@ -20,10 +20,10 @@ struct EMovieSceneUpdateData
 	float LastPosition;
 	bool bPreroll;
 	bool bJumpCut;
-	bool bLooped;
 	/** Indicates that this update was caused by the owning movie scene stopping playback due
 	    to the active sub-scene being deactivated. */
 	bool bSubSceneDeactivate;
+	bool bUpdateCameras;
 
 	EMovieSceneUpdatePass UpdatePass;
 
@@ -33,8 +33,8 @@ struct EMovieSceneUpdateData
 		LastPosition = 0.0f;
 		bPreroll = false;
 		bJumpCut = false;
-		bLooped = false;
 		bSubSceneDeactivate = false;
+		bUpdateCameras = true;
 		UpdatePass = MSUP_PreUpdate;
 	}
 	EMovieSceneUpdateData(float InPosition, float InLastPosition)
@@ -43,8 +43,8 @@ struct EMovieSceneUpdateData
 		LastPosition = InLastPosition;
 		bPreroll = false;
 		bJumpCut = false;
-		bLooped = false;
 		bSubSceneDeactivate = false;
+		bUpdateCameras = true;
 		UpdatePass = MSUP_PreUpdate;
 	}
 };
diff --git a/Engine/Source/Runtime/MovieScene/Public/MovieSceneCommonHelpers.h b/Engine/Source/Runtime/MovieScene/Public/MovieSceneCommonHelpers.h
--- a/Engine/Source/Runtime/MovieScene/Public/MovieSceneCommonHelpers.h
+++ b/Engine/Source/Runtime/MovieScene/Public/MovieSceneCommonHelpers.h
@@ -137,6 +137,13 @@ public:
 	 */
 	void UpdateBindings( const TArray<TWeakObjectPtr<UObject>>& InRuntimeObjects );
 
+	/**
+	 * Rebuilds the property and function mappings for a single runtime object
+	 *
+	 * @param InRuntimeObject	The object to rebuild mappings for
+	 */
+	void UpdateBinding( const TWeakObjectPtr<UObject>& InRuntimeObject );
+
 	/**
 	 * Gets the UProperty that is bound to the track instance
 	 *
@@ -189,6 +196,18 @@ public:
 		}
 	}
 
+	/** @return the property path that this binding was initialized from */
+	const FString& GetPropertyPath() const
+	{
+		return PropertyPath;
+	}
+
+	/** @return the property name that this binding was initialized from */
+	const FName& GetPropertyName() const
+	{
+		return PropertyName;
+	}
+
 private:
 
 	struct FPropertyAddress
diff --git a/Engine/Source/Runtime/MovieScene/Public/MovieSceneSequence.h b/Engine/Source/Runtime/MovieScene/Public/MovieSceneSequence.h
--- a/Engine/Source/Runtime/MovieScene/Public/MovieSceneSequence.h
+++ b/Engine/Source/Runtime/MovieScene/Public/MovieSceneSequence.h
@@ -89,6 +89,15 @@ public:
 	 */
 	virtual void UnbindPossessableObjects(const FGuid& ObjectId) PURE_VIRTUAL(UMovieSceneSequence::UnbindPossessableObjects,);
 
+	/**
+	 * Create a spawnable object template from the specified source object
+	 *
+	 * @param InSourceObject The source object to copy
+	 * @param ObjectName The name of the object
+	 * @return A new object template of the specified name
+	 */
+	virtual UObject* MakeSpawnableTemplateFromInstance(UObject& InSourceObject, FName ObjectName) { return nullptr; }
+
 public:
 
 #if WITH_EDITOR
diff --git a/Engine/Source/Runtime/MovieScene/Public/MovieSceneSpawnable.h b/Engine/Source/Runtime/MovieScene/Public/MovieSceneSpawnable.h
--- a/Engine/Source/Runtime/MovieScene/Public/MovieSceneSpawnable.h
+++ b/Engine/Source/Runtime/MovieScene/Public/MovieSceneSpawnable.h
@@ -4,7 +4,7 @@
 
 #include "MovieSceneSpawnable.generated.h"
 
-class UMovieScene;
+class UMovieSceneSequence;
 
 UENUM()
 enum class ESpawnOwnership : uint8
@@ -23,7 +23,7 @@ enum class ESpawnOwnership : uint8
  * MovieSceneSpawnable describes an object that can be spawned for this MovieScene
  */
 USTRUCT()
-struct FMovieSceneSpawnable 
+struct FMovieSceneSpawnable
 {
 	GENERATED_BODY()
 
@@ -48,11 +48,24 @@ struct FMovieSceneSpawnable
 		, GeneratedClass_DEPRECATED(nullptr)
 #endif
 	{
-		ObjectTemplate->SetFlags(RF_ArchetypeObject);
+		MarkSpawnableTemplate(InObjectTemplate);
 	}
 
 public:
 
+	/**
+	 * Check if the specified object is a spawnable template
+	 * @param InObject 		The object to test
+	 * @return true if the specified object is a spawnable template, false otherwise
+	 */
+	MOVIESCENE_API static bool IsSpawnableTemplate(const UObject& InObject);
+
+	/**
+	 * Indicate that the specified object is a spawnable template object
+	 * @param InObject 		The object to mark
+	 */
+	MOVIESCENE_API static void MarkSpawnableTemplate(const UObject& InObject);
+
 	/**
 	 * Get the template object for this spawnable
 	 *
@@ -68,9 +81,9 @@ public:
 	 * Copy the specified object into this spawnable's template
 	 *
 	 * @param InSourceObject The source object to use. This object will be duplicated into the spawnable.
-	 * @param OwnerMovieScene The outer movie scene to which this spawnable belongs
+	 * @param MovieSceneSequence The movie scene sequence to which this spawnable belongs
 	 */
-	MOVIESCENE_API void CopyObjectTemplate(UObject& InSourceObject, UMovieScene& OwnerMovieScene);
+	MOVIESCENE_API void CopyObjectTemplate(UObject& InSourceObject, UMovieSceneSequence& MovieSceneSequence);
 
 	/**
 	 * Get the unique identifier of the spawnable object.
diff --git a/Engine/Source/Runtime/MovieSceneCapture/Private/AutomatedLevelSequenceCapture.cpp b/Engine/Source/Runtime/MovieSceneCapture/Private/AutomatedLevelSequenceCapture.cpp
--- a/Engine/Source/Runtime/MovieSceneCapture/Private/AutomatedLevelSequenceCapture.cpp
+++ b/Engine/Source/Runtime/MovieSceneCapture/Private/AutomatedLevelSequenceCapture.cpp
@@ -123,7 +123,7 @@ void UAutomatedLevelSequenceCapture::Initialize(TSharedPtr<FSceneViewport> InVie
 		if (Asset)
 		{
 			// Spawn a new actor
-			Actor = GWorld->SpawnActor<ALevelSequenceActor>();
+			Actor = InViewport->GetClient()->GetWorld()->SpawnActor<ALevelSequenceActor>();
 			Actor->SetSequence(Asset);
 			// Ensure it doesn't loop (-1 is indefinite)
 			Actor->PlaybackSettings.LoopCount = 0;
