Added additional logging to NormalizePackageNames to try to understand why Build DDC Mac can't find any files

diff --git a/Engine/Source/Editor/UnrealEd/Private/Commandlets/PackageUtilities.cpp b/Engine/Source/Editor/UnrealEd/Private/Commandlets/PackageUtilities.cpp
--- a/Engine/Source/Editor/UnrealEd/Private/Commandlets/PackageUtilities.cpp
+++ b/Engine/Source/Editor/UnrealEd/Private/Commandlets/PackageUtilities.cpp
@@ -78,6 +78,7 @@ bool NormalizePackageNames( TArray<FString> PackageNames, TArray<FString>& Packa
 			for ( int32 i = 0; i < Paths.Num(); i++ )
 			{
 				FString SearchWildcard = Paths[i] / PackageWildcard;
+				UE_LOG(LogPackageUtilities, Log, TEXT("Searching using wildcard: '%s'"), *SearchWildcard);
 				SearchDirectoryRecursive( SearchWildcard, PackageNames, PackagePathNames );
 			}
 		}
