VR Editor: Fixed commandlets not loading module that contained required classes for cooking engine content

diff --git a/Engine/Source/Editor/UnrealEd/Private/EditorEngine.cpp b/Engine/Source/Editor/UnrealEd/Private/EditorEngine.cpp
--- a/Engine/Source/Editor/UnrealEd/Private/EditorEngine.cpp
+++ b/Engine/Source/Editor/UnrealEd/Private/EditorEngine.cpp
@@ -778,6 +778,7 @@ void UEditorEngine::Init(IEngineLoop* InEngineLoop)
 		FModuleManager::Get().LoadModule(TEXT("HotReload"));
 
 		// Load VR Editor support
+		FModuleManager::Get().LoadModuleChecked( TEXT( "ViewportInteraction" ) );
 		FModuleManager::Get().LoadModuleChecked( TEXT( "VREditor" ) );
 	}
 
