Temp fix for crash creating a new sublevel in the editor

diff --git a/Engine/Source/Editor/UnrealEd/Private/FileHelpers.cpp b/Engine/Source/Editor/UnrealEd/Private/FileHelpers.cpp
--- a/Engine/Source/Editor/UnrealEd/Private/FileHelpers.cpp
+++ b/Engine/Source/Editor/UnrealEd/Private/FileHelpers.cpp
@@ -922,7 +922,7 @@ static bool SaveAsImplementation( UWorld* InWorld, const FString& DefaultFilenam
 		*OutSavedFilename = SaveFilename;
 	}
 
-	if (bStatus)
+	if (bStatus && InWorld == GWorld)
 	{
 		FSelectionStateOfLevel SelectionStateOfLevel;
 		GEditor->GetSelectionStateOfLevel(SelectionStateOfLevel);
