Editor - The game no longer gets priority when it comes to processing debugger commands.  This was too disruptive a change to most folks workflow.  The new strategy is to always process the debugger commands first, however if users want, the keybindings in the editor for commands like Stop, now correctly work.  So you no longer need to modify the ini file to disable or change certain debugger commands, and it allows using options like Shift+Escape to stop where as before, only keys without modifers were the possible bindable keys.

diff --git a/Engine/Source/Editor/UnrealEd/Classes/Editor/EditorEngine.h b/Engine/Source/Editor/UnrealEd/Classes/Editor/EditorEngine.h
--- a/Engine/Source/Editor/UnrealEd/Classes/Editor/EditorEngine.h
+++ b/Engine/Source/Editor/UnrealEd/Classes/Editor/EditorEngine.h
@@ -1661,12 +1661,12 @@ public:
 	/** Called when the debugger has resumed the active PIE or SIE session */
 	void PlaySessionResumed();
 
-	/** Called when the debugger has resumed the active PIE or SIE session when viewport client receives focus */
-	void PlaySessionResumedOnViewportClientFocusReceived();
-
 	/** Called when the debugger has single-stepped the active PIE or SIE session */
 	void PlaySessionSingleStepped();
 
+	/** Called when game client received input key */
+	bool ProcessDebuggerCommands(const FKey InKey, const FModifierKeysState ModifierKeyState);
+
 	/**
 	 * Kicks off a "Play From Here" request that was most likely made during a transaction
 	 */
diff --git a/Engine/Source/Editor/UnrealEd/Private/PlayLevel.cpp b/Engine/Source/Editor/UnrealEd/Private/PlayLevel.cpp
--- a/Engine/Source/Editor/UnrealEd/Private/PlayLevel.cpp
+++ b/Engine/Source/Editor/UnrealEd/Private/PlayLevel.cpp
@@ -215,8 +215,8 @@ void UEditorEngine::EndPlayMap()
 	// clean up any previous Play From Here sessions
 	if ( GameViewport != NULL && GameViewport->Viewport != NULL )
 	{
-		// Remove on focus handler binding
-		GameViewport->OnFocusReceived().Unbind();
+		// Remove debugger commands handler binding
+		GameViewport->OnGameViewportInputKey().Unbind();
 
 		// Remove close handler binding
 		GameViewport->OnCloseRequested().Unbind();
@@ -828,19 +828,6 @@ void UEditorEngine::PlaySessionPaused()
 	FEditorDelegates::PausePIE.Broadcast(bIsSimulatingInEditor);
 }
 
-void UEditorEngine::PlaySessionResumedOnViewportClientFocusReceived()
-{
-	if(GUnrealEd->PlayWorld != NULL && GUnrealEd->PlayWorld->bDebugPauseExecution == true)
-	{
-		GUnrealEd->PlayWorld->bDebugPauseExecution = false;
-
-		// Tell the application to stop ticking in this stack frame
-		FSlateApplication::Get().LeaveDebuggingMode(FKismetDebugUtilities::IsSingleStepping());
-
-		FEditorDelegates::ResumePIE.Broadcast(bIsSimulatingInEditor);
-	}
-}
-
 void UEditorEngine::PlaySessionResumed()
 {
 	FEditorDelegates::ResumePIE.Broadcast(bIsSimulatingInEditor);
@@ -851,6 +838,11 @@ void UEditorEngine::PlaySessionSingleStepped()
 	FEditorDelegates::SingleStepPIE.Broadcast(bIsSimulatingInEditor);
 }
 
+bool UEditorEngine::ProcessDebuggerCommands(const FKey InKey, const FModifierKeysState ModifierKeyState)
+{
+	return FPlayWorldCommands::GlobalPlayWorldActions->ProcessCommandBindings(InKey, ModifierKeyState, false);
+}
+
 /* fits the window position to make sure it falls within the confines of the desktop */
 void FitWindowPositionToWorkArea(FIntPoint &WinPos, FIntPoint &WinSize, const FMargin &WinPadding)
 {
@@ -3002,11 +2994,11 @@ UGameInstance* UEditorEngine::CreatePIEGameInstance(int32 InPIEInstance, bool bI
 		GameViewport->bIsPlayInEditorViewport = true;
 		PieWorldContext->GameViewport = ViewportClient;
 
+		// Add a handler for game client input key
+		ViewportClient->OnGameViewportInputKey().BindUObject(this, &UEditorEngine::ProcessDebuggerCommands);
+
 		// Add a handler for viewport close requests
 		ViewportClient->OnCloseRequested().BindUObject(this, &UEditorEngine::OnViewportCloseRequested);
-
-		// Add a handler for viewport on focus received
-		ViewportClient->OnFocusReceived().BindUObject(this, &UEditorEngine::PlaySessionResumedOnViewportClientFocusReceived);
 			
 		FSlatePlayInEditorInfo& SlatePlayInEditorSession = SlatePlayInEditorMap.Add(PieWorldContext->ContextHandle, FSlatePlayInEditorInfo());
 		SlatePlayInEditorSession.DestinationSlateViewport = RequestedDestinationSlateViewport;	// Might be invalid depending how pie was launched. Code below handles this.
diff --git a/Engine/Source/Runtime/Engine/Classes/Engine/GameViewportClient.h b/Engine/Source/Runtime/Engine/Classes/Engine/GameViewportClient.h
--- a/Engine/Source/Runtime/Engine/Classes/Engine/GameViewportClient.h
+++ b/Engine/Source/Runtime/Engine/Classes/Engine/GameViewportClient.h
@@ -476,12 +476,6 @@ public:
 	/** Make sure all navigation objects have appropriate path rendering components set.  Called when EngineShowFlags.Navigation is set. */
 	virtual void VerifyPathRenderingComponents();
 
-	/** Accessor for delegate called when a viewport received focus */
-	FOnFocusReceived& OnFocusReceived()
-	{
-		return FocusReceivedDelegate;
-	}
-
 	/** Accessor for delegate called when a screenshot is captured */
 	static FOnScreenshotCaptured& OnScreenshotCaptured()
 	{
@@ -713,6 +707,14 @@ public:
 
 	virtual FPopupMethodReply OnQueryPopupMethod() const override;
 
+#if WITH_EDITOR
+	/** Accessor for delegate called when a game viewport received input key */
+	FOnGameViewportInputKey& OnGameViewportInputKey()
+	{
+		return GameViewportInputKeyDelegate;
+	}
+#endif
+
 private:
 	/**
 	 * Set a specific stat to either enabled or disabled (returns the number of remaining enabled stats)
@@ -788,8 +790,10 @@ private:
 	 */
 	bool SetDisplayConfiguration( const FIntPoint* Dimensions, EWindowMode::Type WindowMode);
 
-	/** Delegate called when received focus */
-	FOnFocusReceived FocusReceivedDelegate;
+#if WITH_EDITOR
+	/** Delegate called when game viewport client received input key */
+	FOnGameViewportInputKey GameViewportInputKeyDelegate;
+#endif
 
 	/** Delegate called at the end of the frame when a screenshot is captured */
 	static FOnScreenshotCaptured ScreenshotCapturedDelegate;
diff --git a/Engine/Source/Runtime/Engine/Classes/Engine/GameViewportDelegates.h b/Engine/Source/Runtime/Engine/Classes/Engine/GameViewportDelegates.h
--- a/Engine/Source/Runtime/Engine/Classes/Engine/GameViewportDelegates.h
+++ b/Engine/Source/Runtime/Engine/Classes/Engine/GameViewportDelegates.h
@@ -42,4 +42,4 @@ DECLARE_MULTICAST_DELEGATE_OneParam(FOnGameViewportTick, float);
 *
 * @see UGameViewportClient
 */
-DECLARE_DELEGATE(FOnFocusReceived);
+DECLARE_DELEGATE_RetVal_TwoParams(bool, FOnGameViewportInputKey, FKey, FModifierKeysState);
diff --git a/Engine/Source/Runtime/Engine/Private/GameViewportClient.cpp b/Engine/Source/Runtime/Engine/Private/GameViewportClient.cpp
--- a/Engine/Source/Runtime/Engine/Private/GameViewportClient.cpp
+++ b/Engine/Source/Runtime/Engine/Private/GameViewportClient.cpp
@@ -358,8 +358,20 @@ bool UGameViewportClient::InputKey(FViewport* InViewport, int32 ControllerId, FK
 		GEngine->RemapGamepadControllerIdForPIE(this, ControllerId);
 	}
 
+#if WITH_EDITOR
+	// Give debugger commands a chance to process key binding
+	if (GameViewportInputKeyDelegate.IsBound())
+	{
+		if ( GameViewportInputKeyDelegate.Execute(Key, FSlateApplication::Get().GetModifierKeys()) )
+		{
+			return true;
+		}
+	}
+#endif
+
 	// route to subsystems that care
-	bool bResult = (ViewportConsole ? ViewportConsole->InputKey(ControllerId, Key, EventType, AmountDepressed, bGamepad) : false);
+	bool bResult = ( ViewportConsole ? ViewportConsole->InputKey(ControllerId, Key, EventType, AmountDepressed, bGamepad) : false );
+
 	if (!bResult)
 	{
 		ULocalPlayer* const TargetPlayer = GEngine->GetLocalPlayerFromControllerId(this, ControllerId);
@@ -1406,9 +1418,6 @@ void UGameViewportClient::ReceivedFocus(FViewport* InViewport)
 		GEngine->GetAudioDeviceManager()->SetActiveDevice(AudioDeviceHandle);
 		bHasAudioFocus = true;
 	}
-	
-	// broadcast focus received to anyone that registered an interest
-	FocusReceivedDelegate.ExecuteIfBound();
 }
 
 bool UGameViewportClient::IsFocused(FViewport* InViewport)
