Copying //UE4/Dev-Editor to //UE4/Dev-Main (Source: //UE4/Dev-Editor @ 3082391)

diff --git a/Engine/Source/Editor/UMGEditor/Private/WidgetBlueprintEditorUtils.h b/Engine/Source/Editor/UMGEditor/Private/WidgetBlueprintEditorUtils.h
--- a/Engine/Source/Editor/UMGEditor/Private/WidgetBlueprintEditorUtils.h
+++ b/Engine/Source/Editor/UMGEditor/Private/WidgetBlueprintEditorUtils.h
@@ -58,11 +58,13 @@ private:
 
 	static void WrapWidgets(TSharedRef<FWidgetBlueprintEditor> BlueprintEditor, UWidgetBlueprint* BP, TSet<FWidgetReference> Widgets, UClass* WidgetClass);
 
-	static void BuildReplaceWithMenu(FMenuBuilder& Menu, UWidgetBlueprint* BP, TSet<FWidgetReference> Widgets);
+	static void BuildReplaceWithMenu(FMenuBuilder& Menu, TSharedRef<FWidgetBlueprintEditor> BlueprintEditor, UWidgetBlueprint* BP, TSet<FWidgetReference> Widgets);
 
-	static void ReplaceWidgetWithChildren(UWidgetBlueprint* BP, FWidgetReference Widget);
+	static void ReplaceWidgetWithSelectedTemplate(TSharedRef<FWidgetBlueprintEditor> BlueprintEditor, UWidgetBlueprint* BP, FWidgetReference Widget);
 
-	static void ReplaceWidgets(UWidgetBlueprint* BP, TSet<FWidgetReference> Widgets, UClass* WidgetClass);
+	static bool CanBeReplacedWithTemplate(TSharedRef<FWidgetBlueprintEditor> BlueprintEditor, UWidgetBlueprint* BP, FWidgetReference Widget);
 
-	static bool CanBeReplaced(TSet<FWidgetReference> Widgets);
+	static void ReplaceWidgetWithChildren(TSharedRef<FWidgetBlueprintEditor> BlueprintEditor, UWidgetBlueprint* BP, FWidgetReference Widget);
+
+	static void ReplaceWidgets(TSharedRef<FWidgetBlueprintEditor> BlueprintEditor, UWidgetBlueprint* BP, TSet<FWidgetReference> Widgets, UClass* WidgetClass);
 };
diff --git a/Engine/Source/Editor/UMGEditor/Public/UMGEditorProjectSettings.h b/Engine/Source/Editor/UMGEditor/Public/UMGEditorProjectSettings.h
new file mode 100644
--- /dev/null
+++ b/Engine/Source/Editor/UMGEditor/Public/UMGEditorProjectSettings.h
@@ -0,0 +1,28 @@
+// Copyright 1998-2016 Epic Games, Inc. All Rights Reserved.
+
+#pragma once
+
+#include "UMGEditorProjectSettings.generated.h"
+
+/**
+ * Implements the settings for the UMG Editor Project Settings
+ */
+UCLASS(config=Engine, defaultconfig)
+class UMGEDITOR_API UUMGEditorProjectSettings : public UObject
+{
+	GENERATED_UCLASS_BODY()
+
+#if WITH_EDITORONLY_DATA
+	UPROPERTY(EditAnywhere, config, Category = ShowHideWidgetClasses)
+	bool bShowWidgetsFromEngineContent;
+
+	UPROPERTY(EditAnywhere, config, Category = ShowHideWidgetClasses)
+	bool bShowWidgetsFromDeveloperContent;
+
+	UPROPERTY(EditAnywhere, config, AdvancedDisplay, Category = ShowHideWidgetClasses)
+	TArray<FString> CategoriesToHide;
+
+	UPROPERTY(EditAnywhere, config, Category = ShowHideWidgetClasses, meta = (MetaClass = "Widget"))
+	TArray<FStringClassReference> WidgetClassesToHide;
+#endif
+};
diff --git a/Engine/Source/Editor/UnrealEd/Classes/Commandlets/GatherTextCommandletBase.h b/Engine/Source/Editor/UnrealEd/Classes/Commandlets/GatherTextCommandletBase.h
--- a/Engine/Source/Editor/UnrealEd/Classes/Commandlets/GatherTextCommandletBase.h
+++ b/Engine/Source/Editor/UnrealEd/Classes/Commandlets/GatherTextCommandletBase.h
@@ -3,7 +3,7 @@
 #pragma once 
 
 #include "UnrealString.h"
-#include "InternationalizationManifest.h"
+#include "LocTextHelper.h"
 #include "Commandlets/Commandlet.h"
 #include "GatherTextCommandletBase.generated.h"
 
@@ -98,159 +98,44 @@ private:
 	TArray< TArray<FString> > Data;
 };
 
-/** Helper struct for tracking duplicate localization entries */
-struct FLocConflict
-{
-public:
-
-	FLocConflict( const FString& InNamespace, const FString& InKey, const TSharedPtr<FLocMetadataObject>& InKeyMetadataObj )
-		: Namespace( InNamespace )
-		, Key( InKey )
-		, KeyMetadataObj( InKeyMetadataObj )
-	{
-	}
-
-	void Add( const FLocItem& Source, const FString& SourceLocation );
-
-	const FString Namespace;
-	const FString Key;
-	TSharedPtr<FLocMetadataObject> KeyMetadataObj;
-
-	TMultiMap<FString, FLocItem> EntriesBySourceLocation;
-};
 
-struct FConflictReportInfo
+class FGatherTextSCC
 {
-	struct FLocConflictContainerKeyFuncs : BaseKeyFuncs<TSharedRef<FLocConflict>, FString, true>
-	{
-		static FORCEINLINE const FString& GetSetKey(const TPair<FString, TSharedRef<FLocConflict>>& Element)
-		{
-			return Element.Key;
-		}
-		static FORCEINLINE bool Matches(const FString& A, const FString& B)
-		{
-			return A.Equals(B, ESearchCase::CaseSensitive);
-		}
-		static FORCEINLINE uint32 GetKeyHash(const FString& Key)
-		{
-			return FCrc::StrCrc32<TCHAR>(*Key);
-		}
-	};
-	typedef TMultiMap<FString, TSharedRef<FLocConflict>, FDefaultSetAllocator, FLocConflictContainerKeyFuncs> TLocConflictContainer;
-
-	/**
-	 * Return the singleton instance of the report info.
-	 * @return The singleton instance.
-	 */
-	static FConflictReportInfo& GetInstance();
-
-	/**
-	 * Deletes the singleton instance of the report info.
-	 */
-	static void DeleteInstance();
-
-	/**
-	 * Finds a conflict entry based on namespace and key.
-	 * @param Namespace - The namespace of the entry.
-	 * @param Key - The key/identifier of the entry.
-	 * @return A pointer to the conflict entry or null if not found.
-	 */
-	TSharedPtr< FLocConflict > FindEntryByKey( const FString& Namespace, const FString& Key, const TSharedPtr<FLocMetadataObject> KeyMetadata );
-
-	/**
-	 * Gets the iterator that can be used to traverse all the conflict entries.
-	 *
-	 * @return A const iterator.
-	 */
-	TLocConflictContainer::TConstIterator GetEntryConstIterator() const
-	{
-		return EntriesByKey.CreateConstIterator();
-	}
-	
-	/**
-	 * Adds a conflict entry.
-	 * @param Namespace - The namespace of the entry.
-	 * @param Key - The key/identifier of the entry.
-	 * @param KeyMetadata - Entry Metadata keys.
-	 * @param Source - The source info for the conflict.
-	 * @param SourceLocation - The source location of the conflict.
-	 */
-	void AddConflict(const FString& Namespace, const FString& Key, const TSharedPtr<FLocMetadataObject> KeyMetadata, const FLocItem& Source, const FString& SourceLocation );
-
-	/**
-	 * Return the report in string format.
-	 *
-	 * @return String representing the report.
-	 */
-	FString ToString();
+public:
+	FGatherTextSCC();
 
-private:
-	FConflictReportInfo() {}
-	FConflictReportInfo( const FConflictReportInfo& );
-	FConflictReportInfo& operator=( const FConflictReportInfo& );
+	~FGatherTextSCC();
 
+	bool CheckOutFile(const FString& InFile, FText& OutError);
+	bool CheckinFiles( const FText& InChangeDescription, FText& OutError );
+	bool CleanUp( FText& OutError );
+	bool RevertFile(const FString& InFile, FText& OutError);
+	bool IsReady(FText& OutError);
 
 private:
-	static TSharedPtr<FConflictReportInfo> StaticConflictInstance;
-	TLocConflictContainer EntriesByKey;
+	TArray<FString> CheckedOutFiles;
 };
 
 
-class FManifestInfo
+class FLocFileSCCNotifies : public ILocFileNotifies
 {
 public:
- 	FManifestInfo( )
+	FLocFileSCCNotifies(const TSharedPtr<FGatherTextSCC>& InSourceControlInfo)
+		: SourceControlInfo(InSourceControlInfo)
 	{
-		Manifest = MakeShareable( new FInternationalizationManifest() );
 	}
 
-	// Given a list of manifest file paths, this will load and add the files to the dependency list. ApplyManifestDependencies must be explicitly called to remove dependency entries from the manifest.
-	bool AddManifestDependencies( const TArray< FString >& InManifestFiles );
-
-	// Given a manifest file path will load and add that file to the dependency list. ApplyManifestDependencies must be explicitly called to remove dependency entries from the manifest.
-	bool AddManifestDependency( const FString& InManifestFile );
-
-	// Return the first entry in the dependencies that matches the passed in namespace and context.
-	TSharedPtr< FManifestEntry > FindDependencyEntryByContext(  const FString& Namespace, const FManifestContext& Context, FString& OutFileName );
+	/** Virtual destructor */
+	virtual ~FLocFileSCCNotifies() {}
 
-	// Return the first entry in the dependencies that matches the passed in namespace and source.
-	TSharedPtr< FManifestEntry > FindDependencyEntryBySource( const FString& Namespace, const FLocItem& Source, FString& OutFileName );
-
-	// Returns the number of manifest dependencies
-	int32 NumDependencies() { return ManifestDependencies.Num(); }
-
-	// Strips the current manifest of any entries that appear in the dependencies.  This is not automatically called when dependency files are added.
-	void ApplyManifestDependencies();
-
-	// Adds an entry to the manifest if it is not in one of the manifest dependencies
-	bool AddEntry( const FString& EntryDescription, const FString& Namespace, const FLocItem& Source, const FManifestContext& Context );
-
-	TSharedPtr< class FInternationalizationManifest > GetManifest() { return Manifest; }
+	//~ ILocFileNotifies interface
+	virtual void PreFileRead(const FString& InFilename) override {}
+	virtual void PostFileRead(const FString& InFilename) override {}
+	virtual void PreFileWrite(const FString& InFilename) override;
+	virtual void PostFileWrite(const FString& InFilename) override;
 
 private:
-	TArray < FString > ManifestDependenciesFilePaths;
-	TArray < TSharedPtr< class FInternationalizationManifest > >  ManifestDependencies;
-
-	TSharedPtr< class FInternationalizationManifest > Manifest;
-
-};
-
-
-class FGatherTextSCC
-{
-public:
-	FGatherTextSCC();
-
-	~FGatherTextSCC();
-
-	bool CheckOutFile(const FString& InFile, FText& OutError);
-	bool CheckinFiles( const FText& InChangeDescription, FText& OutError );
-	bool CleanUp( FText& OutError );
-	bool RevertFile(const FString& InFile, FText& OutError);
-	bool IsReady(FText& OutError);
-
-private:
-	TArray<FString> CheckedOutFiles;
+	TSharedPtr<FGatherTextSCC> SourceControlInfo;
 };
 
 
@@ -329,13 +214,7 @@ class UGatherTextCommandletBase : public UCommandlet
 
 
 public:
-	virtual void Initialize( const TSharedRef< FManifestInfo >& InManifestInfo, const TSharedPtr< FGatherTextSCC >& InSourceControlInfo );
-
-	static TSharedPtr<FJsonObject> ReadJSONTextFile( const FString& InFilePath ) ;
-	static bool WriteJSONToTextFile( TSharedPtr<FJsonObject> Output, const FString& Filename, TSharedPtr<FGatherTextSCC> SourceControl );
-
-	// Used to replace strings found in localization text that could cause automated builds to fail if seen in the log output
-	static FString MungeLogOutput( const FString& InString );
+	virtual void Initialize( const TSharedPtr< FLocTextHelper >& InGatherManifestHelper, const TSharedPtr< FGatherTextSCC >& InSourceControlInfo );
 
 	// Wrappers for extracting config values
 	bool GetBoolFromConfig( const TCHAR* Section, const TCHAR* Key, bool& OutValue, const FString& Filename );
@@ -345,7 +224,7 @@ public:
 	int32 GetPathArrayFromConfig( const TCHAR* Section, const TCHAR* Key, TArray<FString>& OutArr, const FString& Filename );
 
 protected:
-	TSharedPtr< FManifestInfo > ManifestInfo;
+	TSharedPtr< FLocTextHelper > GatherManifestHelper;
 
 	TSharedPtr< FGatherTextSCC > SourceControlInfo;
 
diff --git a/Engine/Source/Editor/UnrealEd/Classes/Commandlets/GatherTextFromSourceCommandlet.h b/Engine/Source/Editor/UnrealEd/Classes/Commandlets/GatherTextFromSourceCommandlet.h
--- a/Engine/Source/Editor/UnrealEd/Classes/Commandlets/GatherTextFromSourceCommandlet.h
+++ b/Engine/Source/Editor/UnrealEd/Classes/Commandlets/GatherTextFromSourceCommandlet.h
@@ -54,7 +54,7 @@ private:
 		bool ShouldGatherFromEditorOnlyData;
 
 		//Destination location of the parsed FLocTextEntrys
-		TSharedPtr< FManifestInfo > ManifestInfo;
+		TSharedPtr< FLocTextHelper > GatherManifestHelper;
 
 		FSourceFileParseContext()
 			: Filename()
@@ -69,7 +69,7 @@ private:
 			, WithinNamespaceDefine(false)
 			, WithinStartingLine()
 			, ShouldGatherFromEditorOnlyData(false)
-			, ManifestInfo()
+			, GatherManifestHelper()
 			, MacroBlockStack()
 			, CachedMacroBlockState()
 		{
diff --git a/Engine/Source/Editor/UnrealEd/Classes/Commandlets/GenerateGatherArchiveCommandlet.h b/Engine/Source/Editor/UnrealEd/Classes/Commandlets/GenerateGatherArchiveCommandlet.h
--- a/Engine/Source/Editor/UnrealEd/Classes/Commandlets/GenerateGatherArchiveCommandlet.h
+++ b/Engine/Source/Editor/UnrealEd/Classes/Commandlets/GenerateGatherArchiveCommandlet.h
@@ -20,12 +20,5 @@ public:
 	//~ Begin UCommandlet Interface
 	virtual int32 Main( const FString& Params ) override;
 	//~ End UCommandlet Interface
-	
-	bool WriteArchiveToFile( TSharedRef< FJsonObject> ArchiveJSONObject, const FString& OutputFilePath );
-	void BuildArchiveFromManifest( TSharedRef< const FInternationalizationManifest > InManifest, TSharedRef< FInternationalizationArchive > Archive, const FString& SourceCulture, const FString& TargetCulture  );
-	void AppendArchiveData( TSharedRef< const FInternationalizationArchive > InArchiveToAppend, TSharedRef< FInternationalizationArchive > ArchiveCombined );
-	void ConditionTranslation( FLocItem& LocItem );
-	void ConditionSource( FLocItem& LocItem );
-
 #endif
 };
diff --git a/Engine/Source/Editor/UnrealEd/Classes/Commandlets/GenerateGatherManifestCommandlet.h b/Engine/Source/Editor/UnrealEd/Classes/Commandlets/GenerateGatherManifestCommandlet.h
--- a/Engine/Source/Editor/UnrealEd/Classes/Commandlets/GenerateGatherManifestCommandlet.h
+++ b/Engine/Source/Editor/UnrealEd/Classes/Commandlets/GenerateGatherManifestCommandlet.h
@@ -15,8 +15,6 @@ class UGenerateGatherManifestCommandlet : public UGatherTextCommandletBase
 public:
 	//~ Begin UCommandlet Interface
 	virtual int32 Main( const FString& Params ) override;
-
-	bool WriteManifestToFile( const TSharedPtr<FInternationalizationManifest>& InManifest, const FString& OutputFilePath );
 	//~ End UCommandlet Interface
 #endif
 };
diff --git a/Engine/Source/Editor/UnrealEd/Classes/Commandlets/GenerateTextLocalizationReportCommandlet.h b/Engine/Source/Editor/UnrealEd/Classes/Commandlets/GenerateTextLocalizationReportCommandlet.h
--- a/Engine/Source/Editor/UnrealEd/Classes/Commandlets/GenerateTextLocalizationReportCommandlet.h
+++ b/Engine/Source/Editor/UnrealEd/Classes/Commandlets/GenerateTextLocalizationReportCommandlet.h
@@ -11,22 +11,12 @@ class UGenerateTextLocalizationReportCommandlet : public UGatherTextCommandletBa
 {
 	GENERATED_UCLASS_BODY()
 
-
-
 public:
 	//~ Begin UCommandlet Interface
 	virtual int32 Main(const FString& Params) override;
 	//~ End UCommandlet Interface
 
 private:
-
-	/**
-	 * Calculates the number of words found in a given string.
-	 * @param Text - The string to examine.
-	 * @return Number of words found.
-	 */
-	int32 CountWords( const FString& Text ) const;
-
 	/**
 	 * Collects word count info and writes out a report.
 	 * @param SourcePath - Source path to localization files.
diff --git a/Engine/Source/Editor/UnrealEd/Classes/Commandlets/ImportDialogueScriptCommandlet.h b/Engine/Source/Editor/UnrealEd/Classes/Commandlets/ImportDialogueScriptCommandlet.h
--- a/Engine/Source/Editor/UnrealEd/Classes/Commandlets/ImportDialogueScriptCommandlet.h
+++ b/Engine/Source/Editor/UnrealEd/Classes/Commandlets/ImportDialogueScriptCommandlet.h
@@ -23,8 +23,5 @@ public:
 	//~ End UCommandlet Interface
 
 private:
-	bool ImportDialogueScriptForCulture(const FString& InDialogueScriptFileName, const FString& InCultureArchiveFileName, const FString& InCultureName, const bool bIsNativeCulture);
-
-	TSharedPtr<FInternationalizationManifest> InternationalizationManifest;
-	TSharedPtr<FInternationalizationArchive> NativeArchive;
+	bool ImportDialogueScriptForCulture(FLocTextHelper& InLocTextHelper, const FString& InDialogueScriptFileName, const FString& InCultureName, const bool bIsNativeCulture);
 };
diff --git a/Engine/Source/Editor/UnrealEd/Classes/Commandlets/ImportLocalizedDialogueCommandlet.h b/Engine/Source/Editor/UnrealEd/Classes/Commandlets/ImportLocalizedDialogueCommandlet.h
--- a/Engine/Source/Editor/UnrealEd/Classes/Commandlets/ImportLocalizedDialogueCommandlet.h
+++ b/Engine/Source/Editor/UnrealEd/Classes/Commandlets/ImportLocalizedDialogueCommandlet.h
@@ -35,7 +35,7 @@ private:
 		bool bIsNativeCulture;
 	};
 
-	bool ImportDialogueForCulture(UDialogueWave* const DialogueWave, const FString& DialogueWaveSubPath, const FCultureImportInfo& InCultureImportInfo, const bool bImportAsSource);
+	bool ImportDialogueForCulture(FLocTextHelper& InLocTextHelper, UDialogueWave* const DialogueWave, const FString& DialogueWaveSubPath, const FCultureImportInfo& InCultureImportInfo, const bool bImportAsSource);
 
 	USoundWave* ConditionalImportSoundWave(const FString& InSoundWavePackageName, const FString& InSoundWaveAssetName, const FString& InWavFilename) const;
 	USoundWave* ImportSoundWave(const FString& InSoundWavePackageName, const FString& InSoundWaveAssetName, const FString& InWavFilename) const;
@@ -43,7 +43,5 @@ private:
 	UPROPERTY()
 	USoundFactory* SoundWaveFactory;
 
-	TSharedPtr<FInternationalizationManifest> InternationalizationManifest;
-	TSharedPtr<FInternationalizationArchive> NativeArchive;
 	TSet<FName> AssetsToKeep;
 };
diff --git a/Engine/Source/Editor/UnrealEd/Classes/Factories/BlueprintFunctionLibraryFactory.h b/Engine/Source/Editor/UnrealEd/Classes/Factories/BlueprintFunctionLibraryFactory.h
--- a/Engine/Source/Editor/UnrealEd/Classes/Factories/BlueprintFunctionLibraryFactory.h
+++ b/Engine/Source/Editor/UnrealEd/Classes/Factories/BlueprintFunctionLibraryFactory.h
@@ -7,7 +7,7 @@
 #pragma once
 #include "BlueprintFunctionLibraryFactory.generated.h"
 
-UCLASS(hidecategories=Object, collapsecategories)
+UCLASS(MinimalAPI, hidecategories=Object, collapsecategories)
 class UBlueprintFunctionLibraryFactory : public UBlueprintFactory
 {
 	GENERATED_UCLASS_BODY()
@@ -23,6 +23,3 @@ class UBlueprintFunctionLibraryFactory : public UBlueprintFactory
 	virtual FString GetDefaultNewAssetName() const override;
 	// End of UFactory interface
 };
-
-
-
diff --git a/Engine/Source/Editor/UnrealEd/Classes/Factories/BlueprintMacroFactory.h b/Engine/Source/Editor/UnrealEd/Classes/Factories/BlueprintMacroFactory.h
--- a/Engine/Source/Editor/UnrealEd/Classes/Factories/BlueprintMacroFactory.h
+++ b/Engine/Source/Editor/UnrealEd/Classes/Factories/BlueprintMacroFactory.h
@@ -7,7 +7,7 @@
 #pragma once
 #include "BlueprintMacroFactory.generated.h"
 
-UCLASS(hidecategories=Object, collapsecategories)
+UCLASS(MinimalAPI, hidecategories=Object, collapsecategories)
 class UBlueprintMacroFactory : public UBlueprintFactory
 {
 	GENERATED_UCLASS_BODY()
@@ -27,6 +27,3 @@ protected:
 	virtual bool IsMacroFactory() const override { return true; }
 	// End of UBlueprintFactory interface
 };
-
-
-
diff --git a/Engine/Source/Editor/UnrealEd/Classes/Factories/CurveFactory.h b/Engine/Source/Editor/UnrealEd/Classes/Factories/CurveFactory.h
--- a/Engine/Source/Editor/UnrealEd/Classes/Factories/CurveFactory.h
+++ b/Engine/Source/Editor/UnrealEd/Classes/Factories/CurveFactory.h
@@ -8,7 +8,7 @@
  * Factory that creates curve assets, prompting to pick the kind of curve at creation time
  */
 
-UCLASS()
+UCLASS(MinimalAPI)
 class UCurveFactory : public UFactory
 {
 	GENERATED_UCLASS_BODY()
@@ -27,7 +27,7 @@ class UCurveFactory : public UFactory
  * Factory that creates float curve assets
  */
 
-UCLASS()
+UCLASS(MinimalAPI)
 class UCurveFloatFactory : public UCurveFactory
 {
 	GENERATED_UCLASS_BODY()
@@ -42,7 +42,7 @@ class UCurveFloatFactory : public UCurveFactory
  * Factory that creates linear color curve assets
  */
 
-UCLASS()
+UCLASS(MinimalAPI)
 class UCurveLinearColorFactory : public UCurveFactory
 {
 	GENERATED_UCLASS_BODY()
@@ -57,7 +57,7 @@ class UCurveLinearColorFactory : public UCurveFactory
  * Factory that creates vector curve assets
  */
 
-UCLASS()
+UCLASS(MinimalAPI)
 class UCurveVectorFactory : public UCurveFactory
 {
 	GENERATED_UCLASS_BODY()
diff --git a/Engine/Source/Editor/UnrealEd/Classes/Factories/DialogueVoiceFactory.h b/Engine/Source/Editor/UnrealEd/Classes/Factories/DialogueVoiceFactory.h
--- a/Engine/Source/Editor/UnrealEd/Classes/Factories/DialogueVoiceFactory.h
+++ b/Engine/Source/Editor/UnrealEd/Classes/Factories/DialogueVoiceFactory.h
@@ -7,7 +7,7 @@
 #pragma once
 #include "DialogueVoiceFactory.generated.h"
 
-UCLASS(hidecategories=Object)
+UCLASS(MinimalAPI, hidecategories=Object)
 class UDialogueVoiceFactory : public UFactory
 {
 	GENERATED_UCLASS_BODY()
@@ -16,6 +16,3 @@ class UDialogueVoiceFactory : public UFactory
 	virtual UObject* FactoryCreateNew(UClass* Class,UObject* InParent,FName Name,EObjectFlags Flags,UObject* Context,FFeedbackContext* Warn) override;
 	//~ Begin UFactory Interface	
 };
-
-
-
diff --git a/Engine/Source/Editor/UnrealEd/Classes/Factories/FbxSkeletalMeshImportData.h b/Engine/Source/Editor/UnrealEd/Classes/Factories/FbxSkeletalMeshImportData.h
--- a/Engine/Source/Editor/UnrealEd/Classes/Factories/FbxSkeletalMeshImportData.h
+++ b/Engine/Source/Editor/UnrealEd/Classes/Factories/FbxSkeletalMeshImportData.h
@@ -6,7 +6,7 @@
 /**
  * Import data and options used when importing a static mesh from fbx
  */
-UCLASS()
+UCLASS(MinimalAPI)
 class UFbxSkeletalMeshImportData : public UFbxMeshImportData
 {
 	GENERATED_UCLASS_BODY()
diff --git a/Engine/Source/Editor/UnrealEd/Classes/Factories/ForceFeedbackEffectFactory.h b/Engine/Source/Editor/UnrealEd/Classes/Factories/ForceFeedbackEffectFactory.h
--- a/Engine/Source/Editor/UnrealEd/Classes/Factories/ForceFeedbackEffectFactory.h
+++ b/Engine/Source/Editor/UnrealEd/Classes/Factories/ForceFeedbackEffectFactory.h
@@ -3,7 +3,7 @@
 #pragma once
 #include "ForceFeedbackEffectFactory.generated.h"
 
-UCLASS(hidecategories=Object)
+UCLASS(MinimalAPI, hidecategories=Object)
 class UForceFeedbackEffectFactory : public UFactory
 {
 	GENERATED_UCLASS_BODY()
@@ -13,6 +13,3 @@ class UForceFeedbackEffectFactory : public UFactory
 	virtual UObject* FactoryCreateNew(UClass* Class,UObject* InParent,FName Name,EObjectFlags Flags,UObject* Context,FFeedbackContext* Warn) override;
 	//~ Begin UFactory Interface	
 };
-
-
-
diff --git a/Engine/Source/Editor/UnrealEd/Classes/Factories/MaterialFunctionFactoryNew.h b/Engine/Source/Editor/UnrealEd/Classes/Factories/MaterialFunctionFactoryNew.h
--- a/Engine/Source/Editor/UnrealEd/Classes/Factories/MaterialFunctionFactoryNew.h
+++ b/Engine/Source/Editor/UnrealEd/Classes/Factories/MaterialFunctionFactoryNew.h
@@ -7,7 +7,7 @@
 #pragma once
 #include "MaterialFunctionFactoryNew.generated.h"
 
-UCLASS(hidecategories=Object, collapsecategories)
+UCLASS(MinimalAPI, hidecategories=Object, collapsecategories)
 class UMaterialFunctionFactoryNew : public UFactory
 {
 	GENERATED_UCLASS_BODY()
diff --git a/Engine/Source/Editor/UnrealEd/Classes/Factories/MaterialParameterCollectionFactoryNew.h b/Engine/Source/Editor/UnrealEd/Classes/Factories/MaterialParameterCollectionFactoryNew.h
--- a/Engine/Source/Editor/UnrealEd/Classes/Factories/MaterialParameterCollectionFactoryNew.h
+++ b/Engine/Source/Editor/UnrealEd/Classes/Factories/MaterialParameterCollectionFactoryNew.h
@@ -7,7 +7,7 @@
 #pragma once
 #include "MaterialParameterCollectionFactoryNew.generated.h"
 
-UCLASS(hidecategories=Object, collapsecategories)
+UCLASS(MinimalAPI, hidecategories=Object, collapsecategories)
 class UMaterialParameterCollectionFactoryNew : public UFactory
 {
 	GENERATED_UCLASS_BODY()
diff --git a/Engine/Source/Editor/UnrealEd/Classes/Factories/ObjectLibraryFactory.h b/Engine/Source/Editor/UnrealEd/Classes/Factories/ObjectLibraryFactory.h
--- a/Engine/Source/Editor/UnrealEd/Classes/Factories/ObjectLibraryFactory.h
+++ b/Engine/Source/Editor/UnrealEd/Classes/Factories/ObjectLibraryFactory.h
@@ -3,7 +3,7 @@
 #pragma once
 #include "ObjectLibraryFactory.generated.h"
 
-UCLASS(hidecategories=Object)
+UCLASS(MinimalAPI, hidecategories=Object)
 class UObjectLibraryFactory : public UFactory
 {
 	GENERATED_UCLASS_BODY()
diff --git a/Engine/Source/Editor/UnrealEd/Classes/Factories/PhysicalMaterialFactoryNew.h b/Engine/Source/Editor/UnrealEd/Classes/Factories/PhysicalMaterialFactoryNew.h
--- a/Engine/Source/Editor/UnrealEd/Classes/Factories/PhysicalMaterialFactoryNew.h
+++ b/Engine/Source/Editor/UnrealEd/Classes/Factories/PhysicalMaterialFactoryNew.h
@@ -7,7 +7,7 @@
 #pragma once
 #include "PhysicalMaterialFactoryNew.generated.h"
 
-UCLASS(HideCategories=Object)
+UCLASS(MinimalAPI, HideCategories=Object)
 class UPhysicalMaterialFactoryNew : public UFactory
 {
 	GENERATED_UCLASS_BODY()
@@ -20,6 +20,3 @@ class UPhysicalMaterialFactoryNew : public UFactory
 	virtual UObject* FactoryCreateNew(UClass* Class,UObject* InParent,FName Name,EObjectFlags Flags,UObject* Context,FFeedbackContext* Warn) override;
 	//~ Begin UFactory Interface	
 };
-
-
-
diff --git a/Engine/Source/Editor/UnrealEd/Classes/Factories/ReimportFbxSkeletalMeshFactory.h b/Engine/Source/Editor/UnrealEd/Classes/Factories/ReimportFbxSkeletalMeshFactory.h
--- a/Engine/Source/Editor/UnrealEd/Classes/Factories/ReimportFbxSkeletalMeshFactory.h
+++ b/Engine/Source/Editor/UnrealEd/Classes/Factories/ReimportFbxSkeletalMeshFactory.h
@@ -7,7 +7,7 @@
 #pragma once
 #include "ReimportFbxSkeletalMeshFactory.generated.h"
 
-UCLASS(collapsecategories)
+UCLASS(MinimalAPI, collapsecategories)
 class UReimportFbxSkeletalMeshFactory : public UFbxFactory, public FReimportHandler
 {
 	GENERATED_UCLASS_BODY()
@@ -24,6 +24,3 @@ class UReimportFbxSkeletalMeshFactory : public UFbxFactory, public FReimportHand
 	virtual bool FactoryCanImport(const FString& Filename) override;
 	//~ End UFactory Interface
 };
-
-
-
diff --git a/Engine/Source/Editor/UnrealEd/Classes/Factories/ReimportFbxStaticMeshFactory.h b/Engine/Source/Editor/UnrealEd/Classes/Factories/ReimportFbxStaticMeshFactory.h
--- a/Engine/Source/Editor/UnrealEd/Classes/Factories/ReimportFbxStaticMeshFactory.h
+++ b/Engine/Source/Editor/UnrealEd/Classes/Factories/ReimportFbxStaticMeshFactory.h
@@ -7,7 +7,7 @@
 #pragma once
 #include "ReimportFbxStaticMeshFactory.generated.h"
 
-UCLASS(collapsecategories)
+UCLASS(MinimalAPI, collapsecategories)
 class UReimportFbxStaticMeshFactory : public UFbxFactory, public FReimportHandler
 {
 	GENERATED_UCLASS_BODY()
@@ -24,6 +24,3 @@ class UReimportFbxStaticMeshFactory : public UFbxFactory, public FReimportHandle
 	virtual bool FactoryCanImport(const FString& Filename) override;
 	//~ End UFactory Interface
 };
-
-
-
diff --git a/Engine/Source/Editor/UnrealEd/Classes/Factories/ReverbEffectFactory.h b/Engine/Source/Editor/UnrealEd/Classes/Factories/ReverbEffectFactory.h
--- a/Engine/Source/Editor/UnrealEd/Classes/Factories/ReverbEffectFactory.h
+++ b/Engine/Source/Editor/UnrealEd/Classes/Factories/ReverbEffectFactory.h
@@ -7,16 +7,12 @@
 #pragma once
 #include "ReverbEffectFactory.generated.h"
 
-UCLASS(hidecategories=Object)
+UCLASS(MinimalAPI, hidecategories=Object)
 class UReverbEffectFactory : public UFactory
 {
 	GENERATED_UCLASS_BODY()
 
-
 	//~ Begin UFactory Interface
 	virtual UObject* FactoryCreateNew(UClass* Class,UObject* InParent,FName Name,EObjectFlags Flags,UObject* Context,FFeedbackContext* Warn) override;
 	//~ Begin UFactory Interface	
 };
-
-
-
diff --git a/Engine/Source/Editor/UnrealEd/Classes/Factories/TextureRenderTargetCubeFactoryNew.h b/Engine/Source/Editor/UnrealEd/Classes/Factories/TextureRenderTargetCubeFactoryNew.h
--- a/Engine/Source/Editor/UnrealEd/Classes/Factories/TextureRenderTargetCubeFactoryNew.h
+++ b/Engine/Source/Editor/UnrealEd/Classes/Factories/TextureRenderTargetCubeFactoryNew.h
@@ -7,7 +7,7 @@
 #pragma once
 #include "TextureRenderTargetCubeFactoryNew.generated.h"
 
-UCLASS(hidecategories=(Object, Texture))
+UCLASS(MinimalAPI, hidecategories=(Object, Texture))
 class UTextureRenderTargetCubeFactoryNew : public UFactory
 {
 	GENERATED_UCLASS_BODY()
@@ -25,6 +25,3 @@ class UTextureRenderTargetCubeFactoryNew : public UFactory
 	virtual UObject* FactoryCreateNew(UClass* Class,UObject* InParent,FName Name,EObjectFlags Flags,UObject* Context,FFeedbackContext* Warn) override;
 	//~ Begin UFactory Interface	
 };
-
-
-
diff --git a/Engine/Source/Editor/UnrealEd/Classes/Factories/TextureRenderTargetFactoryNew.h b/Engine/Source/Editor/UnrealEd/Classes/Factories/TextureRenderTargetFactoryNew.h
--- a/Engine/Source/Editor/UnrealEd/Classes/Factories/TextureRenderTargetFactoryNew.h
+++ b/Engine/Source/Editor/UnrealEd/Classes/Factories/TextureRenderTargetFactoryNew.h
@@ -7,7 +7,7 @@
 #pragma once
 #include "TextureRenderTargetFactoryNew.generated.h"
 
-UCLASS(hidecategories=(Object, Texture))
+UCLASS(MinimalAPI, hidecategories=(Object, Texture))
 class UTextureRenderTargetFactoryNew : public UFactory
 {
 	GENERATED_UCLASS_BODY()
diff --git a/Engine/Source/Editor/UnrealEd/Classes/Factories/TouchInterfaceFactory.h b/Engine/Source/Editor/UnrealEd/Classes/Factories/TouchInterfaceFactory.h
--- a/Engine/Source/Editor/UnrealEd/Classes/Factories/TouchInterfaceFactory.h
+++ b/Engine/Source/Editor/UnrealEd/Classes/Factories/TouchInterfaceFactory.h
@@ -4,7 +4,7 @@
 #include "TouchInterfaceFactory.generated.h"
 
 UCLASS(hidecategories=Object)
-class UTouchInterfaceFactory : public UFactory
+class UNREALED_API UTouchInterfaceFactory : public UFactory
 {
 	GENERATED_UCLASS_BODY()
 
diff --git a/Engine/Source/Editor/UnrealEd/Classes/MaterialGraph/MaterialGraphSchema.h b/Engine/Source/Editor/UnrealEd/Classes/MaterialGraph/MaterialGraphSchema.h
--- a/Engine/Source/Editor/UnrealEd/Classes/MaterialGraph/MaterialGraphSchema.h
+++ b/Engine/Source/Editor/UnrealEd/Classes/MaterialGraph/MaterialGraphSchema.h
@@ -155,6 +155,14 @@ class UMaterialGraphSchema : public UEdGraphSchema
 	UPROPERTY()
 	FLinearColor AlphaPinColor;
 
+	/**
+	 *  Add all linked to nodes to this pin to selection
+	 *  
+	 *	@param	Graph			CurrentGraph
+	 *	@param	InGraphPin		Pin clicked on
+	 */
+	void SelectAllInputNodes(UEdGraph* Graph, UEdGraphPin* InGraphPin);
+
 	/**
 	 * Get menu for breaking links to specific nodes
 	 *
diff --git a/Engine/Source/Editor/UnrealEd/Classes/Settings/EditorLoadingSavingSettings.h b/Engine/Source/Editor/UnrealEd/Classes/Settings/EditorLoadingSavingSettings.h
--- a/Engine/Source/Editor/UnrealEd/Classes/Settings/EditorLoadingSavingSettings.h
+++ b/Engine/Source/Editor/UnrealEd/Classes/Settings/EditorLoadingSavingSettings.h
@@ -173,8 +173,8 @@ public:
 	UPROPERTY(EditAnywhere, config, Category=SourceControl, meta=(DisplayName="Use Global Settings"))
 	uint32 bSCCUseGlobalSettings:1;
 
-	/** Specifies the file path to the tool to be used for diff'ing text files */
-	UPROPERTY(EditAnywhere, config, Category=SourceControl, meta=(DisplayName="Tool for diff'ing text", FilePathFilter = "exe"))
+	/** Specifies the file path to the tool to be used for diffing text files */
+	UPROPERTY(EditAnywhere, config, Category=SourceControl, meta=(DisplayName="Tool for diffing text", FilePathFilter = "exe"))
 	FFilePath TextDiffToolPath;
 
 public:
diff --git a/Engine/Source/Editor/UnrealEd/Private/BSPOps.cpp b/Engine/Source/Editor/UnrealEd/Private/BSPOps.cpp
--- a/Engine/Source/Editor/UnrealEd/Private/BSPOps.cpp
+++ b/Engine/Source/Editor/UnrealEd/Private/BSPOps.cpp
@@ -439,7 +439,7 @@ void FBSPOps::SplitPolyList
 	// Add the splitter poly to the Bsp with either a new BspSurf or an existing one.
 	if( RebuildSimplePolys )
 	{
-		SplitPoly->iLink = Model->Surfs.Num();
+		SplitPoly->iLinkSurf = Model->Surfs.Num();
 	}
 
 	int32 iOurNode	= bspAddNode(Model,iParent,NodePlace,0,SplitPoly);
@@ -470,7 +470,7 @@ void FBSPOps::SplitPolyList
 			case SP_Coplanar:
 	            if( RebuildSimplePolys )
 				{
-					EdPoly->iLink = Model->Surfs.Num()-1;
+					EdPoly->iLinkSurf = Model->Surfs.Num()-1;
 				}
 				iPlaneNode = bspAddNode( Model, iPlaneNode, NODE_Plane, 0, EdPoly );
 				break;
@@ -1071,7 +1071,7 @@ int32	FBSPOps::bspAddNode( UModel* Model, int32 iParent, ENodePlace NodePlace, u
 		}
 	}
 	FBspSurf* Surf = NULL;
-	if( EdPoly->iLink == Model->Surfs.Num() )
+	if( EdPoly->iLinkSurf == Model->Surfs.Num() )
 	{
 		int32 NewIndex = Model->Surfs.AddZeroed();
 		Surf = &Model->Surfs[NewIndex];
@@ -1102,9 +1102,9 @@ int32	FBSPOps::bspAddNode( UModel* Model, int32 iParent, ENodePlace NodePlace, u
 	}
 	else
 	{
-		check(EdPoly->iLink!=INDEX_NONE);
-		check(EdPoly->iLink<Model->Surfs.Num());
-		Surf = &Model->Surfs[EdPoly->iLink];
+		check(EdPoly->iLinkSurf!=INDEX_NONE);
+		check(EdPoly->iLinkSurf<Model->Surfs.Num());
+		Surf = &Model->Surfs[EdPoly->iLinkSurf];
 	}
 
 	// Set NodeFlags.
@@ -1148,7 +1148,7 @@ int32	FBSPOps::bspAddNode( UModel* Model, int32 iParent, ENodePlace NodePlace, u
 			Parent = &Model->Nodes[iParent];
 
 		// Set node properties.
-		Node.iSurf       	 = EdPoly->iLink;
+		Node.iSurf       	 = EdPoly->iLinkSurf;
 		Node.NodeFlags   	 = NodeFlags;
 		Node.iCollisionBound = INDEX_NONE;
 		Node.Plane           = FPlane( EdPoly->Vertices[0], EdPoly->Normal );
diff --git a/Engine/Source/Editor/UnrealEd/Private/Commandlets/ExportDialogueScriptCommandlet.cpp b/Engine/Source/Editor/UnrealEd/Private/Commandlets/ExportDialogueScriptCommandlet.cpp
--- a/Engine/Source/Editor/UnrealEd/Private/Commandlets/ExportDialogueScriptCommandlet.cpp
+++ b/Engine/Source/Editor/UnrealEd/Private/Commandlets/ExportDialogueScriptCommandlet.cpp
@@ -5,10 +5,6 @@
 #include "AssetRegistryModule.h"
 #include "Sound/DialogueWave.h"
 #include "Sound/DialogueVoice.h"
-#include "Internationalization/InternationalizationManifest.h"
-#include "Internationalization/InternationalizationArchive.h"
-#include "JsonInternationalizationManifestSerializer.h"
-#include "JsonInternationalizationArchiveSerializer.h"
 
 DEFINE_LOG_CATEGORY_STATIC(LogExportDialogueScriptCommandlet, Log, All);
 
@@ -150,57 +146,21 @@ int32 UExportDialogueScriptCommandlet::Main(const FString& Params)
 		return -1;
 	}
 
-	// Prepare the manifest
-	TSharedPtr<FInternationalizationManifest> InternationalizationManifest;
+	// We may only have a single culture if using this setting
+	if (!bUseCultureDirectory && CulturesToGenerate.Num() > 1)
 	{
-		const FString ManifestFileName = SourcePath / ManifestName;
-		if (!FPaths::FileExists(ManifestFileName))
-		{
-			UE_LOG(LogExportDialogueScriptCommandlet, Error, TEXT("Failed to find manifest '%s'."), *ManifestFileName);
-			return -1;
-		}
-
-		const TSharedPtr<FJsonObject> ManifestJsonObject = ReadJSONTextFile(ManifestFileName);
-		if (!ManifestJsonObject.IsValid())
-		{
-			UE_LOG(LogExportDialogueScriptCommandlet, Error, TEXT("Failed to parse manifest '%s'."), *ManifestFileName);
-			return -1;
-		}
-
-		FJsonInternationalizationManifestSerializer ManifestSerializer;
-		InternationalizationManifest = MakeShareable(new FInternationalizationManifest());
-		if (!ManifestSerializer.DeserializeManifest(ManifestJsonObject.ToSharedRef(), InternationalizationManifest.ToSharedRef()))
-		{
-			UE_LOG(LogExportDialogueScriptCommandlet, Error, TEXT("Failed to deserialize manifest '%s'."), *ManifestFileName);
-			return -1;
-		}
+		UE_LOG(LogExportDialogueScriptCommandlet, Error, TEXT("bUseCultureDirectory may only be used with a single culture."));
+		return false;
 	}
 
-	// Prepare the native archive
-	TSharedPtr<FInternationalizationArchive> NativeArchive;
+	// Load the manifest and all archives
+	FLocTextHelper LocTextHelper(SourcePath, ManifestName, ArchiveName, NativeCulture, CulturesToGenerate, MakeShareable(new FLocFileSCCNotifies(SourceControlInfo)));
 	{
-		const FString NativeCulturePath = SourcePath / NativeCulture;
-
-		const FString NativeArchiveFileName = NativeCulturePath / ArchiveName;
-		if (!FPaths::FileExists(NativeArchiveFileName))
+		FText LoadError;
+		if (!LocTextHelper.LoadAll(ELocTextHelperLoadFlags::LoadOrCreate, &LoadError))
 		{
-			UE_LOG(LogExportDialogueScriptCommandlet, Error, TEXT("Failed to find archive '%s'."), *NativeArchiveFileName);
-			return -1;
-		}
-
-		const TSharedPtr<FJsonObject> ArchiveJsonObject = ReadJSONTextFile(NativeArchiveFileName);
-		if (!ArchiveJsonObject.IsValid())
-		{
-			UE_LOG(LogExportDialogueScriptCommandlet, Error, TEXT("Failed to parse archive '%s'."), *NativeArchiveFileName);
-			return -1;
-		}
-
-		FJsonInternationalizationArchiveSerializer ArchiveSerializer;
-		NativeArchive = MakeShareable(new FInternationalizationArchive());
-		if (!ArchiveSerializer.DeserializeArchive(ArchiveJsonObject.ToSharedRef(), NativeArchive.ToSharedRef()))
-		{
-			UE_LOG(LogExportDialogueScriptCommandlet, Error, TEXT("Failed to deserialize archive '%s'."), *NativeArchiveFileName);
-			return -1;
+			UE_LOG(LogExportDialogueScriptCommandlet, Error, TEXT("%s"), *LoadError.ToString());
+			return false;
 		}
 	}
 
@@ -225,37 +185,6 @@ int32 UExportDialogueScriptCommandlet::Main(const FString& Params)
 		const FString CultureSourcePath = SourcePath / CultureName;
 		const FString CultureDestinationPath = DestinationPath / (bUseCultureDirectory ? CultureName : TEXT(""));
 
-		// Prepare the culture archive
-		TSharedPtr<FInternationalizationArchive> CultureArchive;
-		if (bIsNativeCulture)
-		{
-			CultureArchive = NativeArchive;
-		}
-		else
-		{
-			const FString ArchiveFileName = CultureSourcePath / ArchiveName;
-			if (!FPaths::FileExists(ArchiveFileName))
-			{
-				UE_LOG(LogExportDialogueScriptCommandlet, Error, TEXT("Failed to find archive '%s'."), *ArchiveFileName);
-				continue;
-			}
-
-			const TSharedPtr<FJsonObject> ArchiveJsonObject = ReadJSONTextFile(ArchiveFileName);
-			if (!ArchiveJsonObject.IsValid())
-			{
-				UE_LOG(LogExportDialogueScriptCommandlet, Error, TEXT("Failed to parse archive '%s'."), *ArchiveFileName);
-				continue;
-			}
-
-			FJsonInternationalizationArchiveSerializer ArchiveSerializer;
-			CultureArchive = MakeShareable(new FInternationalizationArchive());
-			if (!ArchiveSerializer.DeserializeArchive(ArchiveJsonObject.ToSharedRef(), CultureArchive.ToSharedRef()))
-			{
-				UE_LOG(LogExportDialogueScriptCommandlet, Error, TEXT("Failed to deserialize archive '%s'."), *ArchiveFileName);
-				continue;
-			}
-		}
-
 		TArray<TSharedPtr<FDialogueScriptEntry>> ExportedDialogueLines;
 		for (const FAssetData& AssetData : AssetDataArrayForDialogueWaves)
 		{
@@ -287,7 +216,7 @@ int32 UExportDialogueScriptCommandlet::Main(const FString& Params)
 				const FString ContextLocalizationKey = DialogueWave->GetContextLocalizationKey(ContextMapping);
 
 				// Check that this entry exists in the manifest file, as we want to skip over dialogue that we aren't gathering
-				TSharedPtr<FManifestEntry> ContextManifestEntry = InternationalizationManifest->FindEntryByKey(FDialogueConstants::DialogueNamespace, ContextLocalizationKey, &DialogueWave->SpokenText);
+				TSharedPtr<FManifestEntry> ContextManifestEntry = LocTextHelper.FindSourceText(FDialogueConstants::DialogueNamespace, ContextLocalizationKey, &DialogueWave->SpokenText);
 				if (!ContextManifestEntry.IsValid())
 				{
 					UE_LOG(LogExportDialogueScriptCommandlet, Log, TEXT("No internationalization manifest entry was found for context '%s' in culture '%s'. This context will be skipped."), *ContextLocalizationKey, *CultureName);
@@ -296,40 +225,20 @@ int32 UExportDialogueScriptCommandlet::Main(const FString& Params)
 
 				// Find the correct entry for our context
 				const FManifestContext* ContextManifestEntryContext = ContextManifestEntry->FindContextByKey(ContextLocalizationKey);
-				check(ContextManifestEntryContext); // This should never fail as we pass in the key to FindEntryByKey
+				check(ContextManifestEntryContext); // This should never fail as we pass in the key to FindSourceText
 
-				// We might have a native translation for this text, in which case we need to export the translation for that text, rather than for the source text
-				FString SourceText = DialogueWave->SpokenText;
-				{
-					TSharedPtr<FArchiveEntry> NativeArchiveEntry = NativeArchive->FindEntryBySource(FDialogueConstants::DialogueNamespace, SourceText, ContextManifestEntryContext->KeyMetadataObj);
-					if (NativeArchiveEntry.IsValid())
-					{
-						SourceText = NativeArchiveEntry->Translation.Text;
-					}
-				}
+				// Get the localized text to export
+				FLocItem ExportedSource;
+				FLocItem ExportedTranslation;
+				LocTextHelper.GetExportText(CultureName, FDialogueConstants::DialogueNamespace, ContextManifestEntryContext->Key, ContextManifestEntryContext->KeyMetadataObj, ELocTextExportSourceMethod::NativeText, ContextManifestEntry->Source, ExportedSource, ExportedTranslation);
 
-				// Find the archive entry so we can export the correct text
-				FString TranslatedText;
-				if (bIsNativeCulture)
-				{
-					TranslatedText = SourceText;
-				}
-				else
-				{
-					TSharedPtr<FArchiveEntry> ArchiveEntry = CultureArchive->FindEntryBySource(FDialogueConstants::DialogueNamespace, SourceText, ContextManifestEntryContext->KeyMetadataObj);
-					if (ArchiveEntry.IsValid())
-					{
-						TranslatedText = ArchiveEntry->Translation.Text;
-					}
-				}
-
-				if (TranslatedText.IsEmpty())
+				if (ExportedTranslation.Text.IsEmpty())
 				{
 					UE_LOG(LogExportDialogueScriptCommandlet, Log, TEXT("Empty translation found for context '%s' in culture '%s'. This context will be skipped."), *ContextLocalizationKey, *CultureName);
 					continue;
 				}
 
-				const auto CollapsedDialogueContextKey = FCollapsedDialogueContextKey(DialogueWave, &ContextMapping, TranslatedText);
+				const auto CollapsedDialogueContextKey = FCollapsedDialogueContextKey(DialogueWave, &ContextMapping, ExportedTranslation.Text);
 				TArray<const FDialogueContextMapping*>* MergedContextsPtr = CollapsedDialogueContexts.Find(CollapsedDialogueContextKey);
 				if (MergedContextsPtr)
 				{
@@ -342,27 +251,9 @@ int32 UExportDialogueScriptCommandlet::Main(const FString& Params)
 			}
 
 			// Get the localized voice actor direction
-			FString LocalizedVoiceActorDirection = DialogueWave->VoiceActorDirection;
-			{
-				// We might have a native translation for this text, in which case we need to export the translation for that text, rather than for the source text
-				{
-					TSharedPtr<FArchiveEntry> NativeArchiveEntry = NativeArchive->FindEntryBySource(FDialogueConstants::DialogueNotesNamespace, LocalizedVoiceActorDirection, nullptr);
-					if (NativeArchiveEntry.IsValid())
-					{
-						LocalizedVoiceActorDirection = NativeArchiveEntry->Translation.Text;
-					}
-				}
-
-				// Find the archive entry so we can export the correct text
-				if (!bIsNativeCulture)
-				{
-					TSharedPtr<FArchiveEntry> ArchiveEntry = CultureArchive->FindEntryBySource(FDialogueConstants::DialogueNotesNamespace, LocalizedVoiceActorDirection, nullptr);
-					if (ArchiveEntry.IsValid())
-					{
-						LocalizedVoiceActorDirection = ArchiveEntry->Translation.Text;
-					}
-				}
-			}
+			FLocItem ExportedVoiceActorDirectionSource;
+			FLocItem ExportedVoiceActorDirectionTranslation;
+			LocTextHelper.GetExportText(CultureName, FDialogueConstants::DialogueNamespace, DialogueWave->LocalizationGUID.ToString() + FDialogueConstants::ActingDirectionKeySuffix, nullptr, ELocTextExportSourceMethod::NativeText, FLocItem(DialogueWave->VoiceActorDirection), ExportedVoiceActorDirectionSource, ExportedVoiceActorDirectionTranslation);
 
 			// Get the localized version of the dialogue wave for the current culture
 			UDialogueWave* LocalizedDialogueWave = nullptr;
@@ -380,7 +271,7 @@ int32 UExportDialogueScriptCommandlet::Main(const FString& Params)
 			for (const auto& CollapsedDialogueContextPair : CollapsedDialogueContexts)
 			{
 				TSharedRef<FDialogueScriptEntry> ExportedDialogueLine = MakeShareable(new FDialogueScriptEntry());
-				PopulateDialogueScriptEntry(DialogueWave, LocalizedDialogueWave, *CollapsedDialogueContextPair.Key.Context, CollapsedDialogueContextPair.Value, CollapsedDialogueContextPair.Key.LocalizedSpokenText, LocalizedVoiceActorDirection, *ExportedDialogueLine);
+				PopulateDialogueScriptEntry(DialogueWave, LocalizedDialogueWave, *CollapsedDialogueContextPair.Key.Context, CollapsedDialogueContextPair.Value, CollapsedDialogueContextPair.Key.LocalizedSpokenText, ExportedVoiceActorDirectionTranslation.Text, *ExportedDialogueLine);
 				ExportedDialogueLines.Add(ExportedDialogueLine);
 			}
 		}
diff --git a/Engine/Source/Editor/UnrealEd/Private/Commandlets/GatherTextCommandlet.cpp b/Engine/Source/Editor/UnrealEd/Private/Commandlets/GatherTextCommandlet.cpp
--- a/Engine/Source/Editor/UnrealEd/Private/Commandlets/GatherTextCommandlet.cpp
+++ b/Engine/Source/Editor/UnrealEd/Private/Commandlets/GatherTextCommandlet.cpp
@@ -73,7 +73,6 @@ int32 UGatherTextCommandlet::Main( const FString& Params )
 
 	UE_LOG(LogGatherTextCommandlet, Log,TEXT("Beginning GatherText Commandlet."));
 
-	TSharedRef< FManifestInfo > CommandletManifestInfo = MakeShareable( new FManifestInfo() );
 	TSharedPtr< FGatherTextSCC > CommandletSourceControlInfo = nullptr;
 
 	if( bEnableSourceControl )
@@ -88,6 +87,10 @@ int32 UGatherTextCommandlet::Main( const FString& Params )
 		}
 	}
 
+	// Basic helper that can be used only to gather a new manifest for writing
+	TSharedRef<FLocTextHelper> CommandletGatherManifestHelper = MakeShareable(new FLocTextHelper(MakeShareable(new FLocFileSCCNotifies(CommandletSourceControlInfo))));
+	CommandletGatherManifestHelper->LoadManifest(ELocTextHelperLoadFlags::Create);
+
 	int32 NumSteps = (ConfigFile->Find("CommonSettings") != NULL) ? ConfigFile->Num() - 1 :  ConfigFile->Num();
 
 	//Execute each step defined in the config file.
@@ -113,7 +116,7 @@ int32 UGatherTextCommandlet::Main( const FString& Params )
 		UGatherTextCommandletBase* Commandlet = NewObject<UGatherTextCommandletBase>(GetTransientPackage(), CommandletClass);
 		check(Commandlet);
 		Commandlet->AddToRoot();
-		Commandlet->Initialize( CommandletManifestInfo, CommandletSourceControlInfo );
+		Commandlet->Initialize( CommandletGatherManifestHelper, CommandletSourceControlInfo );
 
 		// Execute the commandlet.
 		double CommandletExecutionStartTime = FPlatformTime::Seconds();
diff --git a/Engine/Source/Editor/UnrealEd/Private/Commandlets/GatherTextCommandletBase.cpp b/Engine/Source/Editor/UnrealEd/Private/Commandlets/GatherTextCommandletBase.cpp
--- a/Engine/Source/Editor/UnrealEd/Private/Commandlets/GatherTextCommandletBase.cpp
+++ b/Engine/Source/Editor/UnrealEd/Private/Commandlets/GatherTextCommandletBase.cpp
@@ -12,9 +12,6 @@
 
 DEFINE_LOG_CATEGORY_STATIC(LogGatherTextCommandletBase, Log, All);
 
-TSharedPtr<FConflictReportInfo> FConflictReportInfo::StaticConflictInstance;
-
-
 //////////////////////////////////////////////////////////////////////////
 //UGatherTextCommandletBase
 
@@ -23,9 +20,9 @@ UGatherTextCommandletBase::UGatherTextCommandletBase(const FObjectInitializer& O
 {
 }
 
-void UGatherTextCommandletBase::Initialize( const TSharedRef< FManifestInfo >& InManifestInfo, const TSharedPtr< FGatherTextSCC >& InSourceControlInfo )
+void UGatherTextCommandletBase::Initialize( const TSharedPtr< FLocTextHelper >& InGatherManifestHelper, const TSharedPtr< FGatherTextSCC >& InSourceControlInfo )
 {
-	ManifestInfo = InManifestInfo;
+	GatherManifestHelper = InGatherManifestHelper;
 	SourceControlInfo = InSourceControlInfo;
 }
 
@@ -34,147 +31,6 @@ void UGatherTextCommandletBase::CreateCustomEngine(const FString& Params)
 	GEngine = GEditor = NULL;//Force a basic default engine. 
 }
 
-TSharedPtr<FJsonObject> UGatherTextCommandletBase::ReadJSONTextFile(const FString& InFilePath)
-{
-	//read in file as string
-	FString FileContents;
-	if ( !FFileHelper::LoadFileToString(FileContents, *InFilePath) )
-	{
-		UE_LOG(LogGatherTextCommandletBase, Error,TEXT("Failed to load file %s."), *InFilePath);
-		return NULL;
-	}
-
-	//parse as JSON
-	TSharedPtr<FJsonObject> JSONObject;
-
-	TSharedRef< TJsonReader<> > Reader = TJsonReaderFactory<>::Create( FileContents );
-
-	if( !FJsonSerializer::Deserialize( Reader, JSONObject ) || !JSONObject.IsValid())
-	{
-		UE_LOG(LogGatherTextCommandletBase, Error,TEXT("Invalid JSON in file %s."), *InFilePath);
-		UE_LOG(LogGatherTextCommandletBase, Error, TEXT("JSON Error: %s."), *Reader->GetErrorMessage());
-		return NULL;
-	}
-
-	return JSONObject;
-}
-
-bool UGatherTextCommandletBase::WriteJSONToTextFile(TSharedPtr<FJsonObject> Output, const FString& Filename, TSharedPtr<FGatherTextSCC> SourceControl)
-{
-	bool WasSuccessful = true;
-
-	// If the user specified a reference file - write the entries read from code to a ref file
-	if ( !Filename.IsEmpty() && Output.IsValid() )
-	{
-		const bool DidFileExist = FPaths::FileExists(Filename);
-		if (DidFileExist)
-		{
-			if( SourceControl.IsValid() )
-			{
-				FText SCCErrorText;
-				if (!SourceControl->CheckOutFile(Filename, SCCErrorText))
-				{
-					UE_LOG(LogGatherTextCommandletBase, Error, TEXT("Check out of file %s failed: %s"), *Filename, *SCCErrorText.ToString());
-				}
-			}
-		}
-
-		if( WasSuccessful )
-		{
-			//Print the JSON data out to the ref file.
-			FString OutputString;
-			TSharedRef< TJsonWriter<> > Writer = TJsonWriterFactory<>::Create( &OutputString );
-			FJsonSerializer::Serialize( Output.ToSharedRef(), Writer );
-
-			if (!FFileHelper::SaveStringToFile(OutputString, *Filename, FFileHelper::EEncodingOptions::ForceUnicode))
-			{
-				UE_LOG(LogGatherTextCommandletBase, Error, TEXT("Failed to write localization entries to file %s"), *Filename);
-				WasSuccessful = false;
-			}
-		}
-
-		if (!DidFileExist)
-		{
-			// Checkout on a new file will cause it to be added
-			if( SourceControl.IsValid() )
-			{
-				FText SCCErrorText;
-				if (!SourceControl->CheckOutFile(Filename, SCCErrorText))
-				{
-					UE_LOG(LogGatherTextCommandletBase, Error, TEXT("Check out of file %s failed: %s"), *Filename, *SCCErrorText.ToString());
-				}
-			}
-		}
-	}
-	else
-	{
-		WasSuccessful = false;
-	}
-
-	return WasSuccessful;
-}
-
-FString UGatherTextCommandletBase::MungeLogOutput( const FString& InString )
-{
-	if(InString.IsEmpty())
-	{
-		return InString;
-	}
-
-	FString ResultStr = InString.ReplaceCharWithEscapedChar();
-	if(!GIsBuildMachine)
-	{
-		return ResultStr;
-	}
-
-	static TArray<FString> ErrorStrs;
-	if(ErrorStrs.Num() == 0)
-	{
-		ErrorStrs.Add(FString(TEXT("Error")));
-		ErrorStrs.Add(FString(TEXT("Failed")));
-		ErrorStrs.Add(FString(TEXT("[BEROR]"))); 
-		ErrorStrs.Add(FString(TEXT("Utility finished with exit code: -1"))); 
-		ErrorStrs.Add(FString(TEXT("is not recognized as an internal or external command"))); 
-		ErrorStrs.Add(FString(TEXT("Could not open solution: "))); 
-		ErrorStrs.Add(FString(TEXT("Parameter format not correct"))); 
-		ErrorStrs.Add(FString(TEXT("Another build is already started on this computer."))); 
-		ErrorStrs.Add(FString(TEXT("Sorry but the link was not completed because memory was exhausted."))); 
-		ErrorStrs.Add(FString(TEXT("simply rerunning the compiler might fix this problem"))); 
-		ErrorStrs.Add(FString(TEXT("No connection could be made because the target machine actively refused"))); 
-		ErrorStrs.Add(FString(TEXT("Internal Linker Exception:"))); 
-		ErrorStrs.Add(FString(TEXT(": warning LNK4019: corrupt string table"))); 
-		ErrorStrs.Add(FString(TEXT("Proxy could not update its cache"))); 
-		ErrorStrs.Add(FString(TEXT("You have not agreed to the Xcode license agreements"))); 
-		ErrorStrs.Add(FString(TEXT("Connection to build service terminated"))); 
-		ErrorStrs.Add(FString(TEXT("cannot execute binary file"))); 
-		ErrorStrs.Add(FString(TEXT("Invalid solution configuration"))); 
-		ErrorStrs.Add(FString(TEXT("is from a previous version of this application and must be converted in order to build"))); 
-		ErrorStrs.Add(FString(TEXT("This computer has not been authenticated for your account using Steam Guard"))); 
-		ErrorStrs.Add(FString(TEXT("invalid name for SPA section"))); 
-		ErrorStrs.Add(FString(TEXT(": Invalid file name, "))); 
-		ErrorStrs.Add(FString(TEXT("The specified PFX file do not exist. Aborting"))); 
-		ErrorStrs.Add(FString(TEXT("binary is not found. Aborting"))); 
-		ErrorStrs.Add(FString(TEXT("Input file not found: "))); 
-		ErrorStrs.Add(FString(TEXT("An exception occurred during merging:"))); 
-		ErrorStrs.Add(FString(TEXT("Install the 'Microsoft Windows SDK for Windows 7 and .NET Framework 3.5 SP1'"))); 
-		ErrorStrs.Add(FString(TEXT("is less than package's new version 0x"))); 
-		ErrorStrs.Add(FString(TEXT("current engine version is older than version the package was originally saved with")));
-		ErrorStrs.Add(FString(TEXT("exceeds maximum length")));
-		ErrorStrs.Add(FString(TEXT("can't edit exclusive file already opened")));
-	}
-
-	for( int32 ErrStrIdx = 0; ErrStrIdx < ErrorStrs.Num(); ErrStrIdx++ )
-	{
-		FString& FindStr = ErrorStrs[ErrStrIdx];
-		FString ReplaceStr = FString::Printf(TEXT("%s %s"), *FindStr.Left(1), *FindStr.RightChop(1));
-
-		ResultStr.ReplaceInline(*FindStr, *ReplaceStr);
-	}
-
-	return ResultStr;
-}
-
-
 bool UGatherTextCommandletBase::GetBoolFromConfig( const TCHAR* Section, const TCHAR* Key, bool& OutValue, const FString& Filename )
 {
 	bool bSuccess = GConfig->GetBool( Section, Key, OutValue, Filename );
@@ -246,205 +102,6 @@ int32 UGatherTextCommandletBase::GetPathArrayFromConfig( const TCHAR* Section, c
 	return count;
 }
 
-bool FManifestInfo::AddManifestDependency( const FString& InManifestFilePath )
-{
-	if( ManifestDependenciesFilePaths.Contains( InManifestFilePath ) )
-	{
-		return true;
-	}
-
-	bool bSuccess = false;
-
-	if( FPaths::FileExists( InManifestFilePath ) )
-	{
-		FJsonInternationalizationManifestSerializer ManifestSerializer;
-		TSharedPtr< FJsonObject > LoadedManifestDependencyJsonObject = UGatherTextCommandletBase::ReadJSONTextFile( InManifestFilePath );
-
-		TSharedRef< FInternationalizationManifest > LoadedManifestDependency = MakeShareable( new FInternationalizationManifest );
-		bSuccess = ManifestSerializer.DeserializeManifest( LoadedManifestDependencyJsonObject.ToSharedRef(), LoadedManifestDependency );
-		
-		if( bSuccess )
-		{
-			ManifestDependencies.Add( LoadedManifestDependency );
-			ManifestDependenciesFilePaths.Add( InManifestFilePath );
-		}
-		else
-		{
-			UE_LOG(LogGatherTextCommandletBase, Warning, TEXT("Could not load manifest dependency file %s"), *InManifestFilePath);
-		}
-	}
-	else
-	{
-		bSuccess = false;
-		UE_LOG(LogGatherTextCommandletBase, Warning, TEXT("Could not find manifest dependency file %s"), *InManifestFilePath);
-	}
-
-	return bSuccess;
-}
-
-bool FManifestInfo::AddManifestDependencies( const TArray< FString >& InManifestFilePaths )
-{
-	bool bSuccess = true;
-	for( int32 i = 0; i < InManifestFilePaths.Num(); i++ )
-	{
-		FString FilePath = InManifestFilePaths[i];
-		bSuccess &= AddManifestDependency( FilePath );
-	}
-	return bSuccess;
-}
-
-
-TSharedPtr< FManifestEntry > FManifestInfo::FindDependencyEntryByContext( const FString& Namespace, const FManifestContext& Context, FString& OutFileName )
-{
-	TSharedPtr<FManifestEntry> DependencyEntry = NULL;
-	OutFileName = TEXT("");
-
-	for( int32 Idx = 0; Idx < ManifestDependencies.Num(); Idx++ )
-	{
-		DependencyEntry = ManifestDependencies[Idx]->FindEntryByContext( Namespace, Context );
-		if( DependencyEntry.IsValid() )
-		{
-			OutFileName = ManifestDependenciesFilePaths[Idx];
-			break;
-		}
-	}
-	return DependencyEntry;
-}
-
-TSharedPtr< FManifestEntry > FManifestInfo::FindDependencyEntryBySource( const FString& Namespace, const FLocItem& Source, FString& OutFileName )
-{
-	TSharedPtr<FManifestEntry> DependencyEntry = NULL;
-	OutFileName = TEXT("");
-
-	for( int32 Idx = 0; Idx < ManifestDependencies.Num(); Idx++ )
-	{
-		DependencyEntry = ManifestDependencies[Idx]->FindEntryBySource( Namespace, Source );
-		if( DependencyEntry.IsValid() )
-		{
-			OutFileName = ManifestDependenciesFilePaths[Idx];
-			break;
-		}
-	}
-	return DependencyEntry;
-}
-
-void FManifestInfo::ApplyManifestDependencies()
-{
-	if( ManifestDependencies.Num() > 0 )
-	{
-		TSharedRef< FInternationalizationManifest > NewManifest =  MakeShareable( new FInternationalizationManifest );
-		// We'll generate a new manifest by only including items that are not in the dependencies
-		for( TManifestEntryByContextIdContainer::TConstIterator It( Manifest->GetEntriesByContextIdIterator() ); It; ++It )
-		{
-			const TSharedRef<FManifestEntry> ManifestEntry = It.Value();
-
-			for(  auto ContextIt = ManifestEntry->Contexts.CreateConstIterator(); ContextIt; ++ContextIt )
-			{
-				FString DependencyFileName;
-
-				const TSharedPtr<FManifestEntry> DependencyEntry = FindDependencyEntryByContext( ManifestEntry->Namespace, *ContextIt, DependencyFileName );
-				
-				if( DependencyEntry.IsValid() )
-				{
-					if( !(DependencyEntry->Source.IsExactMatch( ManifestEntry->Source )) )
-					{
-						// There is a dependency manifest entry that has the same namespace and keys as our main manifest entry but the source text differs.
-						FString Message = UGatherTextCommandletBase::MungeLogOutput( FString::Printf(TEXT("Found previously entered localized string [%s] %s %s=\"%s\" %s. It was previously \"%s\" %s in dependency manifest %s."),
-							*ManifestEntry->Namespace,
-							*ContextIt->Key,
-							*FJsonInternationalizationMetaDataSerializer::MetadataToString( ContextIt->KeyMetadataObj ),
-							*ManifestEntry->Source.Text,
-							*FJsonInternationalizationMetaDataSerializer::MetadataToString(ManifestEntry->Source.MetadataObj),
-							*DependencyEntry->Source.Text,
-							*FJsonInternationalizationMetaDataSerializer::MetadataToString(DependencyEntry->Source.MetadataObj),
-							*DependencyFileName));
-						UE_LOG(LogGatherTextCommandletBase, Warning, TEXT("%s"), *Message);
-
-						FConflictReportInfo::GetInstance().AddConflict( ManifestEntry->Namespace, ContextIt->Key, ContextIt->KeyMetadataObj, ManifestEntry->Source, *ContextIt->SourceLocation );
-
-						FManifestContext* ConflictingContext = DependencyEntry->FindContext( ContextIt->Key, ContextIt->KeyMetadataObj );
-						FString DependencyEntryFullSrcLoc = ( !DependencyFileName.IsEmpty() ) ? DependencyFileName : ConflictingContext->SourceLocation;
-						
-						FConflictReportInfo::GetInstance().AddConflict( ManifestEntry->Namespace, ContextIt->Key, ContextIt->KeyMetadataObj, DependencyEntry->Source, DependencyEntryFullSrcLoc );
-
-					}
-				}
-				else
-				{
-					// Since we did not find any entries in the dependencies list that match, we'll add to the new manifest
-					bool bAddSuccessful = NewManifest->AddSource( ManifestEntry->Namespace, ManifestEntry->Source, *ContextIt );
-					if(!bAddSuccessful)
-					{
-						UE_LOG(LogGatherTextCommandletBase, Error, TEXT("Could not process localized string: %s [%s] %s=\"%s\" %s."),
-							*ManifestEntry->Namespace,
-							*ContextIt->Key,
-							*ManifestEntry->Source.Text,
-							*FJsonInternationalizationMetaDataSerializer::MetadataToString( ManifestEntry->Source.MetadataObj ));
-					}
-				}
-			}
-		}
-
-		Manifest = NewManifest;
-	}
-}
-
-bool FManifestInfo::AddEntry( const FString& EntryDescription, const FString& Namespace, const FLocItem& Source, const FManifestContext& Context )
-{
-	bool bAddSuccessful = false;
-	// Check if the entry already exists in the manifest or one of the manifest dependencies
-	FString ExistingEntryFileName;
-	TSharedPtr< FManifestEntry > ExistingEntry = Manifest->FindEntryByContext( Namespace, Context );
-	if( !ExistingEntry.IsValid() )
-	{
-		ExistingEntry = FindDependencyEntryByContext( Namespace, Context, ExistingEntryFileName );
-	}
-
-	if( ExistingEntry.IsValid() )
-	{
-		if( Source.IsExactMatch( ExistingEntry->Source ) )
-		{
-			bAddSuccessful = true;
-		}
-		else
-		{
-			// Grab the source location of the conflicting context
-			FManifestContext* ConflictingContext = ExistingEntry->FindContext( Context.Key, Context.KeyMetadataObj );
-			FString ExistingEntrySourceLocation = ( !ExistingEntryFileName.IsEmpty() ) ?  ExistingEntryFileName : ConflictingContext->SourceLocation;
-
-			FString Message = UGatherTextCommandletBase::MungeLogOutput( FString::Printf(TEXT("Previously entered localized string: %s [%s] %s %s=\"%s\" %s. It was previously \"%s\" %s in %s." ),
-				*EntryDescription,
-				*Namespace,
-				*Context.Key,
-				*FJsonInternationalizationMetaDataSerializer::MetadataToString( Context.KeyMetadataObj ),
-				*Source.Text,
-				*FJsonInternationalizationMetaDataSerializer::MetadataToString( Source.MetadataObj ),
-				*ExistingEntry->Source.Text,
-				*FJsonInternationalizationMetaDataSerializer::MetadataToString( ExistingEntry->Source.MetadataObj ),
-				*ExistingEntrySourceLocation));
-
-			UE_LOG(LogGatherTextCommandletBase, Warning, TEXT("%s"), *Message);
-
-			FConflictReportInfo::GetInstance().AddConflict(Namespace, Context.Key, Context.KeyMetadataObj, Source, Context.SourceLocation );
-			FConflictReportInfo::GetInstance().AddConflict(Namespace, Context.Key, Context.KeyMetadataObj, ExistingEntry->Source, ExistingEntrySourceLocation);
-		}
-	}
-	else
-	{
-		bAddSuccessful = Manifest->AddSource( Namespace, Source, Context );
-		if( !bAddSuccessful )
-		{
-			UE_LOG(LogGatherTextCommandletBase, Error, TEXT("Could not process localized string: %s [%s] %s=\"%s\" %s."),
-				*EntryDescription,
-				*Namespace,
-				*Context.Key,
-				*Source.Text,
-				*FJsonInternationalizationMetaDataSerializer::MetadataToString( Source.MetadataObj ));
-		}
-	}
-	return bAddSuccessful;
-}
-
 
 const TCHAR* FWordCountReportData::EntryDelimiter = TEXT(",");
 const TCHAR* FWordCountReportData::NewLineDelimiter = TEXT("\n");
@@ -625,109 +282,6 @@ int32 FWordCountReportData::AddColumn(const FString* ColumnHeading )
 	return GetColumnCount() - 1;
 }
 
-void FLocConflict::Add( const FLocItem& Source, const FString& SourceLocation )
-{
-	EntriesBySourceLocation.AddUnique( SourceLocation, Source );
-}
-
-FConflictReportInfo& FConflictReportInfo::GetInstance()
-{
-	if( !StaticConflictInstance.IsValid() )
-	{
-		StaticConflictInstance = MakeShareable( new FConflictReportInfo() );
-	}
-	return *StaticConflictInstance;
-}
-
-void FConflictReportInfo::DeleteInstance()
-{
-	StaticConflictInstance.Reset();
-}
-
-TSharedPtr< FLocConflict > FConflictReportInfo::FindEntryByKey( const FString& Namespace, const FString& Key, const TSharedPtr<FLocMetadataObject> KeyMetadata )
-{
-	TArray< TSharedRef< FLocConflict > > MatchingEntries;
-	EntriesByKey.MultiFind( Key, MatchingEntries );
-
-	for( int Idx = 0; Idx < MatchingEntries.Num(); Idx++ )
-	{
-		if( MatchingEntries[Idx]->Namespace == Namespace )
-		{
-			if( KeyMetadata.IsValid() != MatchingEntries[Idx]->KeyMetadataObj.IsValid() )
-			{
-				continue;
-			}
-			else if ( (!KeyMetadata.IsValid() && !MatchingEntries[Idx]->KeyMetadataObj.IsValid()) ||
-				*(KeyMetadata) == *(MatchingEntries[Idx]->KeyMetadataObj) )
-			{
-				return MatchingEntries[Idx];
-			}
-		}
-	}
-
-	return NULL;
-}
-
-void FConflictReportInfo::AddConflict(const FString& Namespace, const FString& Key, const TSharedPtr<FLocMetadataObject> KeyMetadata, const FLocItem& Source, const FString& SourceLocation )
-{
-	TSharedPtr< FLocConflict > ExistingEntry = FindEntryByKey( Namespace, Key, KeyMetadata );
-	if( !ExistingEntry.IsValid() )
-	{
-		TSharedRef< FLocConflict > NewEntry = MakeShareable( new FLocConflict( Namespace, Key, KeyMetadata ) );
-		EntriesByKey.Add( Key, NewEntry );
-		ExistingEntry = NewEntry;
-	}
-	ExistingEntry->Add( Source, SourceLocation.ReplaceCharWithEscapedChar() );
-}
-
-FString FConflictReportInfo::ToString()
-{
-	FString Report;
-	for( auto ConflictIter = GetEntryConstIterator(); ConflictIter; ++ConflictIter )
-	{
-		const TSharedRef<FLocConflict>& Conflict = ConflictIter.Value();
-		const FString& Namespace = Conflict->Namespace;
-		const FString& Key = Conflict->Key;
-
-		bool bAddToReport = false;
-		TArray<FLocItem> SourceList; 
-		Conflict->EntriesBySourceLocation.GenerateValueArray(SourceList);
-		if(SourceList.Num() >= 2)
-		{
-			for(int32 i = 0; i < SourceList.Num() - 1 && !bAddToReport; ++i)
-			{
-				for(int32 j = i+1; j < SourceList.Num() && !bAddToReport; ++j)
-				{
-					if( !(SourceList[i] == SourceList[j]) )
-					{
-						bAddToReport = true;
-					}
-				}
-			}
-		}
-		
-		if(bAddToReport)
-		{
-			FString KeyMetadataString = FJsonInternationalizationMetaDataSerializer::MetadataToString( Conflict->KeyMetadataObj );
-			Report += FString::Printf(TEXT("%s - %s %s\n"), *Namespace, *Key, *KeyMetadataString );
-						
-			for(auto EntryIter = Conflict->EntriesBySourceLocation.CreateConstIterator(); EntryIter; ++EntryIter)
-			{
-				const FString& SourceLocation = EntryIter.Key();
-				FString ProcessedSourceLocation = FPaths::ConvertRelativePathToFull(SourceLocation);
-				ProcessedSourceLocation.ReplaceInline( TEXT("\\"), TEXT("/"));
-				ProcessedSourceLocation.ReplaceInline( *FPaths::RootDir(), TEXT("/"));
-
-				const FString& SourceText = EntryIter.Value().Text.ReplaceCharWithEscapedChar();
-
-				FString SourceMetadataString =  FJsonInternationalizationMetaDataSerializer::MetadataToString( EntryIter.Value().MetadataObj );
-				Report += FString::Printf(TEXT("\t%s - \"%s\" %s\n"), *ProcessedSourceLocation, *SourceText, *SourceMetadataString);
-			}
-			Report += TEXT("\n");
-		}
-	}
-	return Report;
-}
 
 FGatherTextSCC::FGatherTextSCC()
 {
@@ -980,6 +534,32 @@ bool FGatherTextSCC::RevertFile( const FString& InFile, FText& OutError )
 	return bSuccessfullyReverted;
 }
 
+void FLocFileSCCNotifies::PreFileWrite(const FString& InFilename)
+{
+	if (SourceControlInfo.IsValid() && FPaths::FileExists(InFilename))
+	{
+		// File already exists, so check it out before writing to it
+		FText ErrorMsg;
+		if (!SourceControlInfo->CheckOutFile(InFilename, ErrorMsg))
+		{
+			UE_LOG(LogGatherTextCommandletBase, Error, TEXT("Failed to check out file '%s'. %s"), *InFilename, *ErrorMsg.ToString());
+		}
+	}
+}
+
+void FLocFileSCCNotifies::PostFileWrite(const FString& InFilename)
+{
+	if (SourceControlInfo.IsValid())
+	{
+		// If the file didn't exist before then this will add it, otherwise it will do nothing
+		FText ErrorMsg;
+		if (!SourceControlInfo->CheckOutFile(InFilename, ErrorMsg))
+		{
+			UE_LOG(LogGatherTextCommandletBase, Error, TEXT("Failed to check out file '%s'. %s"), *InFilename, *ErrorMsg.ToString());
+		}
+	}
+}
+
 bool FLocalizedAssetSCCUtil::SaveAssetWithSCC(const TSharedPtr<FGatherTextSCC>& InSourceControlInfo, UObject* InAsset)
 {
 	UPackage* const AssetPackage = InAsset->GetOutermost();
diff --git a/Engine/Source/Editor/UnrealEd/Private/Commandlets/GatherTextFromAssetsCommandlet.cpp b/Engine/Source/Editor/UnrealEd/Private/Commandlets/GatherTextFromAssetsCommandlet.cpp
--- a/Engine/Source/Editor/UnrealEd/Private/Commandlets/GatherTextFromAssetsCommandlet.cpp
+++ b/Engine/Source/Editor/UnrealEd/Private/Commandlets/GatherTextFromAssetsCommandlet.cpp
@@ -157,7 +157,7 @@ void UGatherTextFromAssetsCommandlet::ProcessGatherableTextDataArray(const FStri
 
 				FLocItem Source(GatherableTextData.SourceData.SourceString);
 
-				ManifestInfo->AddEntry(TextSourceSiteContext.SiteDescription, GatherableTextData.NamespaceName, Source, Context);
+				GatherManifestHelper->AddSourceText(GatherableTextData.NamespaceName, Source, Context, &TextSourceSiteContext.SiteDescription);
 			}
 		}
 	}
@@ -328,10 +328,14 @@ int32 UGatherTextFromAssetsCommandlet::Main(const FString& Params)
 	TArray<FString> ManifestDependenciesList;
 	GetPathArrayFromConfig(*SectionName, TEXT("ManifestDependencies"), ManifestDependenciesList, GatherTextConfigPath);
 
-	if( !ManifestInfo->AddManifestDependencies( ManifestDependenciesList ) )
+	for (const FString& ManifestDependency : ManifestDependenciesList)
 	{
-		UE_LOG(LogGatherTextFromAssetsCommandlet, Error, TEXT("The GatherTextFromAssets commandlet couldn't find all the specified manifest dependencies."));
-		return -1;
+		FText OutError;
+		if (!GatherManifestHelper->AddDependency(ManifestDependency, &OutError))
+		{
+			UE_LOG(LogGatherTextFromAssetsCommandlet, Error, TEXT("The GatherTextFromAssets commandlet couldn't load the specified manifest dependency: '%'. %s"), *ManifestDependency, *OutError.ToString());
+			return -1;
+		}
 	}
 
 	//The main array of files to work from.
diff --git a/Engine/Source/Editor/UnrealEd/Private/Commandlets/GatherTextFromMetadataCommandlet.cpp b/Engine/Source/Editor/UnrealEd/Private/Commandlets/GatherTextFromMetadataCommandlet.cpp
--- a/Engine/Source/Editor/UnrealEd/Private/Commandlets/GatherTextFromMetadataCommandlet.cpp
+++ b/Engine/Source/Editor/UnrealEd/Private/Commandlets/GatherTextFromMetadataCommandlet.cpp
@@ -117,11 +117,14 @@ int32 UGatherTextFromMetaDataCommandlet::Main( const FString& Params )
 	TArray<FString> ManifestDependenciesList;
 	GetPathArrayFromConfig(*SectionName, TEXT("ManifestDependencies"), ManifestDependenciesList, GatherTextConfigPath);
 
-
-	if( !ManifestInfo->AddManifestDependencies( ManifestDependenciesList ) )
+	for (const FString& ManifestDependency : ManifestDependenciesList)
 	{
-		UE_LOG(LogGatherTextFromMetaDataCommandlet, Error, TEXT("The GatherTextFromMetaData commandlet couldn't find all the specified manifest dependencies."));
-		return -1;
+		FText OutError;
+		if (!GatherManifestHelper->AddDependency(ManifestDependency, &OutError))
+		{
+			UE_LOG(LogGatherTextFromMetaDataCommandlet, Error, TEXT("The GatherTextFromMetaData commandlet couldn't load the specified manifest dependency: '%'. %s"), *ManifestDependency, *OutError.ToString());
+			return -1;
+		}
 	}
 
 	return 0;
@@ -182,7 +185,7 @@ void UGatherTextFromMetaDataCommandlet::GatherTextFromUObject(UField* const Fiel
 					FManifestContext Context;
 					Context.Key = FText::Format(Arguments.OutputKeys[i], PatternArguments).ToString();
 					Context.SourceLocation = FString::Printf(TEXT("From metadata for key %s of member %s in %s"), *Arguments.InputKeys[i], *Field->GetName(), *Field->GetFullGroupName(true));
-					ManifestInfo->AddEntry(TEXT("EntryDescription"), Namespace, LocItem, Context);
+					GatherManifestHelper->AddSourceText(Namespace, LocItem, Context);
 				}
 			}
 		}
@@ -218,7 +221,7 @@ void UGatherTextFromMetaDataCommandlet::GatherTextFromUObject(UField* const Fiel
 							FManifestContext Context;
 							Context.Key = FText::Format(Arguments.OutputKeys[j], PatternArguments).ToString();
 							Context.SourceLocation = FString::Printf(TEXT("From metadata for key %s of enum value %s of enum %s in %s"), *Arguments.InputKeys[j], *Enum->GetEnumName(i), *Enum->GetName(), *Enum->GetFullGroupName(true));
-							ManifestInfo->AddEntry(TEXT("EntryDescription"), Namespace, LocItem, Context);
+							GatherManifestHelper->AddSourceText(Namespace, LocItem, Context);
 						}
 					}
 				}
diff --git a/Engine/Source/Editor/UnrealEd/Private/Commandlets/GatherTextFromSourceCommandlet.cpp b/Engine/Source/Editor/UnrealEd/Private/Commandlets/GatherTextFromSourceCommandlet.cpp
--- a/Engine/Source/Editor/UnrealEd/Private/Commandlets/GatherTextFromSourceCommandlet.cpp
+++ b/Engine/Source/Editor/UnrealEd/Private/Commandlets/GatherTextFromSourceCommandlet.cpp
@@ -202,11 +202,14 @@ int32 UGatherTextFromSourceCommandlet::Main( const FString& Params )
 	TArray<FString> ManifestDependenciesList;
 	GetPathArrayFromConfig(*SectionName, TEXT("ManifestDependencies"), ManifestDependenciesList, GatherTextConfigPath);
 	
-
-	if( !ManifestInfo->AddManifestDependencies( ManifestDependenciesList ) )
+	for (const FString& ManifestDependency : ManifestDependenciesList)
 	{
-		UE_LOG(LogGatherTextFromSourceCommandlet, Error, TEXT("The GatherTextFromSource commandlet couldn't find all the specified manifest dependencies."));
-		return -1;
+		FText OutError;
+		if (!GatherManifestHelper->AddDependency(ManifestDependency, &OutError))
+		{
+			UE_LOG(LogGatherTextFromSourceCommandlet, Error, TEXT("The GatherTextFromSource commandlet couldn't load the specified manifest dependency: '%'. %s"), *ManifestDependency, *OutError.ToString());
+			return -1;
+		}
 	}
 
 	// Get the loc macros and their syntax
@@ -243,7 +246,7 @@ int32 UGatherTextFromSourceCommandlet::Main( const FString& Params )
 
 	// Init a parse context to track the state of the file parsing 
 	FSourceFileParseContext ParseCtxt;
-	ParseCtxt.ManifestInfo = ManifestInfo;
+	ParseCtxt.GatherManifestHelper = GatherManifestHelper;
 
 	// Get whether we should gather editor-only data. Typically only useful for the localization of UE4 itself.
 	if (!GetBoolFromConfig(*SectionName, TEXT("ShouldGatherFromEditorOnlyData"), ParseCtxt.ShouldGatherFromEditorOnlyData, GatherTextConfigPath))
@@ -510,14 +513,14 @@ FString UGatherTextFromSourceCommandlet::RemoveStringFromTextMacro(const FString
 		int32 OpenQuoteIdx = TextMacro.Find(TEXT("\""), ESearchCase::CaseSensitive);
 		if (0 > OpenQuoteIdx || TextMacro.Len() - 1 == OpenQuoteIdx)
 		{
-			UE_LOG(LogGatherTextFromSourceCommandlet, Warning, TEXT("Missing quotes in %s"), *MungeLogOutput(IdentForLogging));
+			UE_LOG(LogGatherTextFromSourceCommandlet, Warning, TEXT("Missing quotes in %s"), *FLocTextHelper::SanitizeLogOutput(IdentForLogging));
 		}
 		else
 		{
 			int32 CloseQuoteIdx = TextMacro.Find(TEXT("\""), ESearchCase::CaseSensitive, ESearchDir::FromStart, OpenQuoteIdx+1);
 			if (0 > CloseQuoteIdx)
 			{
-				UE_LOG(LogGatherTextFromSourceCommandlet, Warning, TEXT("Missing quotes in %s"), *MungeLogOutput(IdentForLogging));
+				UE_LOG(LogGatherTextFromSourceCommandlet, Warning, TEXT("Missing quotes in %s"), *FLocTextHelper::SanitizeLogOutput(IdentForLogging));
 			}
 			else
 			{
@@ -860,7 +863,7 @@ bool UGatherTextFromSourceCommandlet::FSourceFileParseContext::AddManifestText(
 			LineNumber, 
 			*LineText);
 		FLocItem Source( SourceText.ReplaceEscapedCharWithChar() );
-		return ManifestInfo->AddEntry(EntryDescription, InNamespace, Source, Context);
+		return GatherManifestHelper->AddSourceText(InNamespace, Source, Context, &EntryDescription);
 	}
 
 	return false;
@@ -1118,7 +1121,7 @@ bool UGatherTextFromSourceCommandlet::FMacroDescriptor::ParseArgsFromMacro(const
 	int32 OpenBracketIdx = RemainingText.Find(TEXT("("));
 	if (0 > OpenBracketIdx)
 	{
-		UE_LOG(LogGatherTextFromSourceCommandlet, Warning, TEXT("Missing bracket '(' in %s macro in %s(%d):%s"), *GetToken(), *Context.Filename, Context.LineNumber, *MungeLogOutput(Context.LineText));
+		UE_LOG(LogGatherTextFromSourceCommandlet, Warning, TEXT("Missing bracket '(' in %s macro in %s(%d):%s"), *GetToken(), *Context.Filename, Context.LineNumber, *FLocTextHelper::SanitizeLogOutput(Context.LineText));
 		//Dont assume this is an error. It's more likely trying to parse something it shouldn't be.
 		return false;
 	}
@@ -1175,7 +1178,7 @@ bool UGatherTextFromSourceCommandlet::FMacroDescriptor::ParseArgsFromMacro(const
 
 				if (0 > BracketStack)
 				{
-					UE_LOG(LogGatherTextFromSourceCommandlet, Warning, TEXT("Unexpected bracket ')' in %s macro in %s(%d):%s"), *GetToken(), *Context.Filename, Context.LineNumber, *MungeLogOutput(Context.LineText));
+					UE_LOG(LogGatherTextFromSourceCommandlet, Warning, TEXT("Unexpected bracket ')' in %s macro in %s(%d):%s"), *GetToken(), *Context.Filename, Context.LineNumber, *FLocTextHelper::SanitizeLogOutput(Context.LineText));
 					return false;
 				}
 			}
@@ -1230,7 +1233,7 @@ void UGatherTextFromSourceCommandlet::FCommandMacroDescriptor::TryParse(const FS
 		{
 			if (Arguments.Num() != 5)
 			{
-				UE_LOG(LogGatherTextFromSourceCommandlet, Warning, TEXT("Too many arguments in command %s macro in %s(%d):%s"), *GetToken(), *Context.Filename, Context.LineNumber, *MungeLogOutput(Context.LineText));
+				UE_LOG(LogGatherTextFromSourceCommandlet, Warning, TEXT("Too many arguments in command %s macro in %s(%d):%s"), *GetToken(), *Context.Filename, Context.LineNumber, *FLocTextHelper::SanitizeLogOutput(Context.LineText));
 			}
 			else
 			{
@@ -1297,7 +1300,7 @@ void UGatherTextFromSourceCommandlet::FStringMacroDescriptor::TryParse(const FSt
 
 			if (NumArgs != Arguments.Num())
 			{
-				UE_LOG(LogGatherTextFromSourceCommandlet, Warning, TEXT("Too many arguments in %s macro in %s(%d):%s"), *GetToken(), *Context.Filename, Context.LineNumber, *MungeLogOutput(Context.LineText));
+				UE_LOG(LogGatherTextFromSourceCommandlet, Warning, TEXT("Too many arguments in %s macro in %s(%d):%s"), *GetToken(), *Context.Filename, Context.LineNumber, *FLocTextHelper::SanitizeLogOutput(Context.LineText));
 			}
 			else
 			{
@@ -1313,7 +1316,7 @@ void UGatherTextFromSourceCommandlet::FStringMacroDescriptor::TryParse(const FSt
 					FString ArgText = ArgArray[ArgIdx].Trim();
 
 					bool HasQuotes;
-					FString MacroDesc = FString::Printf(TEXT("argument %d of %d in localization macro %s %s(%d):%s"), ArgIdx+1, Arguments.Num(), *GetToken(), *Context.Filename, Context.LineNumber, *MungeLogOutput(Context.LineText));
+					FString MacroDesc = FString::Printf(TEXT("argument %d of %d in localization macro %s %s(%d):%s"), ArgIdx+1, Arguments.Num(), *GetToken(), *Context.Filename, Context.LineNumber, *FLocTextHelper::SanitizeLogOutput(Context.LineText));
 					if (!PrepareArgument(ArgText, Arg.IsAutoText, MacroDesc, HasQuotes))
 					{
 						ArgParseError = true;
