Copying //UE4/Dev-Editor to //UE4/Dev-Main (Source: //UE4/Dev-Editor @ 3082391)

diff --git a/Engine/Source/Editor/UnrealEd/Private/Commandlets/GenerateGatherArchiveCommandlet.cpp b/Engine/Source/Editor/UnrealEd/Private/Commandlets/GenerateGatherArchiveCommandlet.cpp
--- a/Engine/Source/Editor/UnrealEd/Private/Commandlets/GenerateGatherArchiveCommandlet.cpp
+++ b/Engine/Source/Editor/UnrealEd/Private/Commandlets/GenerateGatherArchiveCommandlet.cpp
@@ -1,17 +1,120 @@
 // Copyright 1998-2016 Epic Games, Inc. All Rights Reserved.
 
 #include "UnrealEd.h"
-#include "ISourceControlModule.h"
-#include "Json.h"
-#include "Internationalization/InternationalizationArchive.h"
-#include "Internationalization/InternationalizationManifest.h"
-#include "Internationalization/InternationalizationMetadata.h"
-#include "JsonInternationalizationArchiveSerializer.h"
-#include "JsonInternationalizationManifestSerializer.h"
-
 
 DEFINE_LOG_CATEGORY_STATIC(LogGenerateArchiveCommandlet, Log, All);
 
+void ConditionTranslationMetadata(TSharedRef<FLocMetadataValue> MetadataValue)
+{
+	switch (MetadataValue->Type)
+	{
+	case ELocMetadataType::String:
+	{
+		TSharedPtr<FLocMetadataValue> MetadataValuePtr = MetadataValue;
+		TSharedPtr<FLocMetadataValueString> MetadataString = StaticCastSharedPtr<FLocMetadataValueString>(MetadataValuePtr);
+		MetadataString->SetString(TEXT(""));
+	}
+	break;
+
+	case ELocMetadataType::Array:
+	{
+		TArray< TSharedPtr< FLocMetadataValue > > MetadataArray = MetadataValue->AsArray();
+		for (auto ArrayIter = MetadataArray.CreateIterator(); ArrayIter; ++ArrayIter)
+		{
+			TSharedPtr<FLocMetadataValue>& Item = *ArrayIter;
+			if (Item.IsValid())
+			{
+				ConditionTranslationMetadata(Item.ToSharedRef());
+			}
+		}
+	}
+	break;
+
+	case ELocMetadataType::Object:
+	{
+		TSharedPtr< FLocMetadataObject > MetadataObject = MetadataValue->AsObject();
+		for (auto ValueIter = MetadataObject->Values.CreateConstIterator(); ValueIter; ++ValueIter)
+		{
+			const FString Name = (*ValueIter).Key;
+			TSharedPtr< FLocMetadataValue > Value = (*ValueIter).Value;
+			if (Value.IsValid())
+			{
+				if (Value->Type == ELocMetadataType::String)
+				{
+					MetadataObject->SetStringField(Name, TEXT(""));
+				}
+				else
+				{
+					ConditionTranslationMetadata(Value.ToSharedRef());
+				}
+			}
+		}
+	}
+
+	default:
+		break;
+	}
+}
+
+void ConditionTranslation(FLocItem& LocItem)
+{
+	// We clear out the translation text because this should only be entered by translators
+	LocItem.Text = TEXT("");
+
+	// The translation might have metadata so we want to clear all the values of any string metadata
+	if (LocItem.MetadataObj.IsValid())
+	{
+		ConditionTranslationMetadata(MakeShareable(new FLocMetadataValueObject(LocItem.MetadataObj)));
+	}
+}
+
+void ConditionSourceMetadata(TSharedRef<FLocMetadataValue> MetadataValue)
+{
+	switch (MetadataValue->Type)
+	{
+	case ELocMetadataType::Object:
+	{
+		TSharedPtr< FLocMetadataObject > MetadataObject = MetadataValue->AsObject();
+
+		TArray< FString > NamesToBeReplaced;
+
+		// We go through all the source metadata entries and look for any that have names prefixed with a '*'.  If we
+		//  find any, we will replace them with a String type that contains an empty value
+		for (auto Iter = MetadataObject->Values.CreateConstIterator(); Iter; ++Iter)
+		{
+			const FString& Name = (*Iter).Key;
+			TSharedPtr< FLocMetadataValue > Value = (*Iter).Value;
+
+			if (Name.StartsWith(FLocMetadataObject::COMPARISON_MODIFIER_PREFIX))
+			{
+				NamesToBeReplaced.Add(Name);
+			}
+			else
+			{
+				ConditionSourceMetadata(Value.ToSharedRef());
+			}
+		}
+
+		for (auto Iter = NamesToBeReplaced.CreateConstIterator(); Iter; ++Iter)
+		{
+			MetadataObject->RemoveField(*Iter);
+			MetadataObject->SetStringField(*Iter, TEXT(""));
+		}
+	}
+
+	default:
+		break;
+	}
+}
+
+void ConditionSource(FLocItem& LocItem)
+{
+	if (LocItem.MetadataObj.IsValid())
+	{
+		ConditionSourceMetadata(MakeShareable(new FLocMetadataValueObject(LocItem.MetadataObj)));
+	}
+}
+
 /**
  *	UGenerateGatherArchiveCommandlet
  */
@@ -22,8 +125,6 @@ UGenerateGatherArchiveCommandlet::UGenerateGatherArchiveCommandlet(const FObject
 
 int32 UGenerateGatherArchiveCommandlet::Main( const FString& Params )
 {
-	FInternationalization& I18N = FInternationalization::Get();
-
 	// Parse command line - we're interested in the param vals
 	TArray<FString> Tokens;
 	TArray<FString> Switches;
@@ -66,36 +167,20 @@ int32 UGenerateGatherArchiveCommandlet::Main( const FString& Params )
 		return -1;
 	}
 
-	// Get native culture.
-	FString NativeCulture;
+	// Get archive name.
+	FString ArchiveName;
+	if (!(GetStringFromConfig(*SectionName, TEXT("ArchiveName"), ArchiveName, GatherTextConfigPath)))
 	{
-		const bool WasNativeCultureSpecified = GetStringFromConfig( *SectionName, TEXT("NativeCulture"), NativeCulture, GatherTextConfigPath );
-
-		// Get source culture (DEPRECATED).
-		FString SourceCulture;
-		const bool WasSourceCultureSpecified = GetStringFromConfig( *SectionName, TEXT("SourceCulture"), SourceCulture, GatherTextConfigPath );
-		if (WasSourceCultureSpecified)
-		{
-			UE_LOG(LogGenerateArchiveCommandlet, Warning, TEXT("SourceCulture detected in section %s. SourceCulture is deprecated, please use NativeCulture."), *SectionName);
-		}
-
-		// Use SourceCulture if NativeCulture isn't specified.
-		if (!WasNativeCultureSpecified && WasSourceCultureSpecified)
-		{
-			NativeCulture = SourceCulture;
-		}
+		UE_LOG(LogGenerateArchiveCommandlet, Error, TEXT("No archive name specified."));
+		return -1;
 	}
 
-	if (!NativeCulture.IsEmpty())
-	{
-		if (!I18N.GetCulture(NativeCulture).IsValid())
-		{
-			UE_LOG(LogGenerateArchiveCommandlet, Verbose, TEXT("Specified native culture is not a valid runtime culture, but may be a valid base language: %s"), *(NativeCulture) );
-		}
-	}
-	else
+	// Get native culture.
+	FString NativeCulture;
+	if (!GetStringFromConfig(*SectionName, TEXT("NativeCulture"), NativeCulture, GatherTextConfigPath))
 	{
-		UE_LOG(LogGenerateArchiveCommandlet, Warning, TEXT("No native culture specified. If the actual native culture is specified for generation, its archive entries may lack default translations."));
+		UE_LOG(LogGenerateArchiveCommandlet, Error, TEXT("No native culture specified."));
+		return -1;
 	}
 
 	// Get cultures to generate.
@@ -108,14 +193,6 @@ int32 UGenerateGatherArchiveCommandlet::Main( const FString& Params )
 		return -1;
 	}
 
-	for(int32 i = 0; i < CulturesToGenerate.Num(); ++i)
-	{
-		if( I18N.GetCulture( CulturesToGenerate[i] ).IsValid() )
-		{
-			UE_LOG(LogGenerateArchiveCommandlet, Verbose, TEXT("Specified culture is not a valid runtime culture, but may be a valid base language: %s"), *(CulturesToGenerate[i]) );
-		}
-	}
-
 	// Get destination path.
 	FString DestinationPath;
 	if( !GetPathFromConfig( *SectionName, TEXT("DestinationPath"), DestinationPath, GatherTextConfigPath ) )
@@ -124,396 +201,87 @@ int32 UGenerateGatherArchiveCommandlet::Main( const FString& Params )
 		return -1;
 	}
 
-	// Get archive name.
-	FString ArchiveName;
-	if( !( GetStringFromConfig(* SectionName, TEXT("ArchiveName"), ArchiveName, GatherTextConfigPath ) ) )
+	// Load the manifest and all archives
+	FLocTextHelper LocTextHelper(DestinationPath, ManifestName, ArchiveName, NativeCulture, CulturesToGenerate, MakeShareable(new FLocFileSCCNotifies(SourceControlInfo)));
 	{
-		UE_LOG(LogGenerateArchiveCommandlet, Error, TEXT("No archive name specified."));
-		return -1;
-	}
-
-	// Get bPurgeOldEmptyEntries option.
-	bool ShouldPurgeOldEmptyEntries;
-	if ( !GetBoolFromConfig( *SectionName, TEXT("bPurgeOldEmptyEntries"), ShouldPurgeOldEmptyEntries, GatherTextConfigPath) )
-	{
-		ShouldPurgeOldEmptyEntries = false;
-	}
-
-	// Prepare the manifest
-	TSharedRef<FInternationalizationManifest> InternationalizationManifest = MakeShareable(new FInternationalizationManifest());
-	{
-		const FString ManifestFileName = DestinationPath / ManifestName;
-		if (!FPaths::FileExists(ManifestFileName))
-		{
-			UE_LOG(LogGenerateArchiveCommandlet, Error, TEXT("Failed to find manifest '%s'."), *ManifestFileName);
-			return -1;
-		}
-
-		const TSharedPtr<FJsonObject> ManifestJsonObject = ReadJSONTextFile(ManifestFileName);
-		if (!ManifestJsonObject.IsValid())
-		{
-			UE_LOG(LogGenerateArchiveCommandlet, Error, TEXT("Failed to parse manifest '%s'."), *ManifestFileName);
-			return -1;
-		}
-
-		FJsonInternationalizationManifestSerializer ManifestSerializer;
-		if (!ManifestSerializer.DeserializeManifest(ManifestJsonObject.ToSharedRef(), InternationalizationManifest))
+		FText LoadError;
+		if (!LocTextHelper.LoadAll(ELocTextHelperLoadFlags::LoadOrCreate, &LoadError))
 		{
-			UE_LOG(LogGenerateArchiveCommandlet, Error, TEXT("Failed to deserialize manifest '%s'."), *ManifestFileName);
-			return -1;
+			UE_LOG(LogGenerateArchiveCommandlet, Error, TEXT("%s"), *LoadError.ToString());
+			return false;
 		}
 	}
 
-	// Load the native archive info
-	TSharedPtr<FInternationalizationArchive> NativeArchive;
+	// We need to make sure the native culture is processed first
+	if (CulturesToGenerate.RemoveSingle(NativeCulture) > 0)
 	{
-		const FString NativeCulturePath = DestinationPath / NativeCulture;
-
-		// The native archive may not exist if this is the first time we've run a gather, so this shouldn't fail the commandlet
-		const FString NativeArchiveFileName = NativeCulturePath / ArchiveName;
-		if (FPaths::FileExists(NativeArchiveFileName))
-		{
-			const TSharedPtr<FJsonObject> ArchiveJsonObject = ReadJSONTextFile(NativeArchiveFileName);
-			if (!ArchiveJsonObject.IsValid())
-			{
-				UE_LOG(LogGenerateArchiveCommandlet, Error, TEXT("Failed to read archive file %s."), *NativeArchiveFileName);
-				return -1;
-			}
-
-			FJsonInternationalizationArchiveSerializer ArchiveSerializer;
-			NativeArchive = MakeShareable(new FInternationalizationArchive());
-			if (!ArchiveSerializer.DeserializeArchive(ArchiveJsonObject.ToSharedRef(), NativeArchive.ToSharedRef()))
-			{
-				UE_LOG(LogGenerateArchiveCommandlet, Error, TEXT("Failed to deserialize archive '%s'."), *NativeArchiveFileName);
-				return -1;
-			}
-		}
+		CulturesToGenerate.Insert(NativeCulture, 0);
 	}
 
 	for (const FString& CultureName : CulturesToGenerate)
 	{
-		TSharedRef< FInternationalizationArchive > InternationalizationArchive = MakeShareable( new FInternationalizationArchive );
-		BuildArchiveFromManifest( InternationalizationManifest, InternationalizationArchive, NativeCulture, CultureName );
-
-		const FString CulturePath = DestinationPath / CultureName;
-		FJsonInternationalizationArchiveSerializer ArchiveSerializer;
-		TSharedRef< FInternationalizationArchive > OutputInternationalizationArchive = MakeShareable( new FInternationalizationArchive );
-
-		// Read in any existing archive for this culture.
-		FString ExistingArchiveFileName = CulturePath / ArchiveName;
-		TSharedPtr< FJsonObject > ExistingArchiveJsonObject = NULL;
-			
-		if( FPaths::FileExists(ExistingArchiveFileName) )
+		// Add any missing manifest entries to the archive for this culture
+		const bool bIsNativeCulture = CultureName == NativeCulture;
+		LocTextHelper.EnumerateSourceTexts([&LocTextHelper, &CultureName, &NativeCulture, &bIsNativeCulture](TSharedRef<FManifestEntry> InManifestEntry) -> bool
 		{
-			ExistingArchiveJsonObject = ReadJSONTextFile( ExistingArchiveFileName );
-			
-			// Some of the existing archives were saved out with an "Unnamed" namespace for the root instead of the empty string.  We try to fix that here.
-			if( ExistingArchiveJsonObject->HasField( FJsonInternationalizationArchiveSerializer::TAG_NAMESPACE ) )
-			{
-				FString RootNamespace = ExistingArchiveJsonObject->GetStringField( FJsonInternationalizationArchiveSerializer::TAG_NAMESPACE );
-				if( RootNamespace == TEXT("Unnamed") )
-				{
-					ExistingArchiveJsonObject->RemoveField( FJsonInternationalizationArchiveSerializer::TAG_NAMESPACE );
-					ExistingArchiveJsonObject->SetStringField( FJsonInternationalizationArchiveSerializer::TAG_NAMESPACE, TEXT("") );
-				}
-			}
-
-			struct Local
+			for (const FManifestContext& Context : InManifestEntry->Contexts)
 			{
-				// Purges this JSONObject of an entries with no translated text and purges empty namespaces.
-				// Returns true if the object was modified, false if not.
-				static bool PurgeNamespaceOfEmptyEntries(const TSharedPtr<FJsonObject>& JSONObject)
+				if (!Context.bIsOptional)
 				{
-					bool ModifiedChildrenArray = false;
-					if( JSONObject->HasField( FJsonInternationalizationArchiveSerializer::TAG_CHILDREN ) )
+					TSharedPtr<FArchiveEntry> ArchiveEntry = LocTextHelper.FindTranslation(CultureName, InManifestEntry->Namespace, Context.Key, Context.KeyMetadataObj);
+					if (ArchiveEntry.IsValid())
 					{
-						TArray<TSharedPtr<FJsonValue>> ChildrenArray = JSONObject->GetArrayField(FJsonInternationalizationArchiveSerializer::TAG_CHILDREN);
-						for( int32 ChildIndex = ChildrenArray.Num() - 1; ChildIndex >= 0; --ChildIndex )
+						// We only update existing entries for the native culture, as this lets us preserve stale translations in foreign archives so we can decide later whether we actually want to use them
+						if (bIsNativeCulture && !ArchiveEntry->Source.IsExactMatch(InManifestEntry->Source))
 						{
-							TSharedPtr<FJsonObject> Child = ChildrenArray[ ChildIndex ]->AsObject();
-							TSharedPtr<FJsonObject> TranslationObject = Child->GetObjectField(FJsonInternationalizationArchiveSerializer::TAG_TRANSLATION);
-
-							const FString& TranslatedText = TranslationObject->GetStringField(FJsonInternationalizationArchiveSerializer::TAG_TRANSLATION_TEXT);
-
-							if(TranslatedText.IsEmpty())
-							{
-								ChildrenArray.RemoveAt( ChildIndex );
-								Child = NULL;
-								ModifiedChildrenArray = true;
-							}
-						}
-						if(ModifiedChildrenArray)
-						{
-							JSONObject->RemoveField(FJsonInternationalizationArchiveSerializer::TAG_CHILDREN);
-							if(ChildrenArray.Num())
-							{
-								JSONObject->SetArrayField(FJsonInternationalizationArchiveSerializer::TAG_CHILDREN, ChildrenArray);
-							}
+							LocTextHelper.UpdateTranslation(CultureName, InManifestEntry->Namespace, Context.Key, Context.KeyMetadataObj, InManifestEntry->Source, InManifestEntry->Source);
 						}
 					}
-
-					bool ModifiedSubnamespaceArray = false;
-					if( JSONObject->HasField( FJsonInternationalizationArchiveSerializer::TAG_SUBNAMESPACES ) )
+					else
 					{
-						TArray<TSharedPtr<FJsonValue>> SubnamespaceArray = JSONObject->GetArrayField(FJsonInternationalizationArchiveSerializer::TAG_SUBNAMESPACES);
-
-						for( int32 Index = SubnamespaceArray.Num() - 1; Index >= 0; --Index )
+						// Get the source we should use for the new entry
+						FLocItem ArchiveSource = InManifestEntry->Source;
+						if (!bIsNativeCulture)
 						{
-							TSharedPtr<FJsonObject> Subnamespace = SubnamespaceArray[ Index ]->AsObject();
-							ModifiedSubnamespaceArray = PurgeNamespaceOfEmptyEntries(Subnamespace);
-
-							bool HasChildren = Subnamespace->HasField( FJsonInternationalizationArchiveSerializer::TAG_CHILDREN );
-							bool HasSubnamespaces = Subnamespace->HasField( FJsonInternationalizationArchiveSerializer::TAG_SUBNAMESPACES );
-
-							if(!HasChildren && !HasSubnamespaces)
+							TSharedPtr<FArchiveEntry> NativeArchiveEntry = LocTextHelper.FindTranslation(NativeCulture, InManifestEntry->Namespace, Context.Key, Context.KeyMetadataObj);
+							if (NativeArchiveEntry.IsValid() && !NativeArchiveEntry->Source.IsExactMatch(NativeArchiveEntry->Translation))
 							{
-								SubnamespaceArray.RemoveAt( Index );
-								Subnamespace = NULL;
-								ModifiedSubnamespaceArray = true;
+								ArchiveSource = NativeArchiveEntry->Source;
 							}
 						}
-						if(ModifiedSubnamespaceArray)
-						{
-							JSONObject->RemoveField(FJsonInternationalizationArchiveSerializer::TAG_SUBNAMESPACES);
-							if(SubnamespaceArray.Num())
-							{
-								JSONObject->SetArrayField(FJsonInternationalizationArchiveSerializer::TAG_SUBNAMESPACES, SubnamespaceArray);
-							}
-						}
-					}
 
-					return ModifiedChildrenArray || ModifiedSubnamespaceArray;
-				}
-			};
+						// We want to process the source before adding it to the archive
+						ConditionSource(ArchiveSource);
 
-			if(ShouldPurgeOldEmptyEntries)
-			{
-				// Remove entries lacking translations from pre-existing archive.
-				// If they are absent in the source manifest, we save on not translating non-existent text.
-				// If they are present in the source manifest, then the newly generated entries will contain the empty text again.
-				Local::PurgeNamespaceOfEmptyEntries(ExistingArchiveJsonObject);
-			}
-
-			ArchiveSerializer.DeserializeArchive( ExistingArchiveJsonObject.ToSharedRef(), OutputInternationalizationArchive );
-
-			if (NativeArchive.IsValid() && CultureName != NativeCulture)
-			{
-				for(TManifestEntryByContextIdContainer::TConstIterator i = InternationalizationManifest->GetEntriesByContextIdIterator(); i; ++i)
-				{
-					// Gather relevant info from manifest entry.
-					const TSharedRef<FManifestEntry>& ManifestEntry = i.Value();
-					const FString& Namespace = ManifestEntry->Namespace;
-					const FLocItem& Source = ManifestEntry->Source;
-					const FString& SourceString = Source.Text;
-					const FString UnescapedSourceString = SourceString;
-					const uint32 SourceStringHash = FCrc::StrCrc32(*UnescapedSourceString);
-					const FString Key = i.Key();
-
-					for (const auto& Context : ManifestEntry->Contexts)
-					{
-						const TSharedPtr<FArchiveEntry> NativeArchiveEntry = NativeArchive->FindEntryBySource(Namespace, Source, Context.KeyMetadataObj);
-						if (NativeArchiveEntry.IsValid() && !NativeArchiveEntry->Source.IsExactMatch(NativeArchiveEntry->Translation))
+						FLocItem ArchiveTranslation = ArchiveSource;
+						if (!bIsNativeCulture)
 						{
-							FLocItem NewSource = NativeArchiveEntry->Translation;
-							ConditionSource(NewSource);
-							FLocItem NewTranslation = NativeArchiveEntry->Translation;
-							ConditionTranslation(NewTranslation);
-							OutputInternationalizationArchive->AddEntry(Namespace, NewSource, NewTranslation, Context.KeyMetadataObj, Context.bIsOptional);
-							break;
+							// We want to process the translation before adding it to the archive
+							// We skip this for native entries so that the source text is also added as the translation
+							ConditionTranslation(ArchiveTranslation);
 						}
-					}
-				}
-			}
-		}
-
-		if (InternationalizationArchive->GetFormatVersion() < FInternationalizationArchive::EFormatVersion::Latest)
-		{
-			UE_LOG( LogGenerateArchiveCommandlet, Error,TEXT("Archive version is out of date. Repair the archives or manually set the version to %d."), static_cast<int32>(FInternationalizationArchive::EFormatVersion::Latest));
-			return -1;
-		}
-
-		// Combine the generated gather archive with the contents of the archive structure we will write out.
-		AppendArchiveData( InternationalizationArchive, OutputInternationalizationArchive );
-		InternationalizationArchive->SetFormatVersion(FInternationalizationArchive::EFormatVersion::Latest);
-
-		TSharedRef< FJsonObject > OutputArchiveJsonObj = MakeShareable( new FJsonObject );
-		ArchiveSerializer.SerializeArchive( OutputInternationalizationArchive, OutputArchiveJsonObj );
-
-		const FString ArchivePath = FPaths::ConvertRelativePathToFull(DestinationPath) / CultureName / ArchiveName;
-		if( !WriteArchiveToFile( OutputArchiveJsonObj, ArchivePath ) )
-		{
-			UE_LOG( LogGenerateArchiveCommandlet, Error,TEXT("Failed to write archive to %s."), *ArchivePath );				
-			return -1;
-		}
-	}
-	
-	return 0;
-}
-
-bool UGenerateGatherArchiveCommandlet::WriteArchiveToFile(TSharedRef< FJsonObject > ArchiveJSONObject, const FString& OutputFilePath)
-{
-	UE_LOG(LogGenerateArchiveCommandlet, Log, TEXT("Writing archive to %s."), *OutputFilePath);
-
-	if( !WriteJSONToTextFile(ArchiveJSONObject, OutputFilePath, SourceControlInfo ) )
-	{
-		return false;
-	}
-
-	return true;
-}
-
-
-void UGenerateGatherArchiveCommandlet::BuildArchiveFromManifest( TSharedRef< const FInternationalizationManifest > InManifest, TSharedRef< FInternationalizationArchive > Archive, const FString& SourceCulture, const FString& TargetCulture )
-{
-	for(TManifestEntryByContextIdContainer::TConstIterator It( InManifest->GetEntriesByContextIdIterator() ); It; ++It)
-	{
-		const TSharedRef<FManifestEntry> UnstructuredManifestEntry = It.Value();
-		const FString& Namespace = UnstructuredManifestEntry->Namespace;
-		const FLocItem& Source = UnstructuredManifestEntry->Source; 
-
-		for( auto ContextIter = UnstructuredManifestEntry->Contexts.CreateConstIterator(); ContextIter; ++ContextIter )
-		{
-			const FManifestContext& Context = *ContextIter;
-
-			// We only add the non-optional entries
-			if( !Context.bIsOptional )
-			{
-				FLocItem Translation = Source;
-
-				if( SourceCulture != TargetCulture)
-				{
-					// We want to process the translation before adding it to the archive
-					ConditionTranslation( Translation );
-				}
-
-				// We also condition the source object
-				FLocItem ConditionedSource = Source;
-				ConditionSource( ConditionedSource );
-
-				Archive->AddEntry(Namespace, ConditionedSource, Translation, Context.KeyMetadataObj, Context.bIsOptional );
-					
-			}
-		}
-	}
-}
-
-void UGenerateGatherArchiveCommandlet::AppendArchiveData( TSharedRef< const FInternationalizationArchive > InArchiveToAppend, TSharedRef< FInternationalizationArchive > ArchiveCombined )
-{
-	for(TArchiveEntryContainer::TConstIterator It( InArchiveToAppend->GetEntryIterator() ); It; ++It)
-	{
-		const TSharedRef<FArchiveEntry> EntryToAppend = It.Value();
-		ArchiveCombined->AddEntry( EntryToAppend );
-	}
-}
-
-void ConditionTranslationMetadata( TSharedRef<FLocMetadataValue> MetadataValue )
-{
-	switch( MetadataValue->Type )
-	{
-	case ELocMetadataType::String:
-		{
-			TSharedPtr<FLocMetadataValue> MetadataValuePtr = MetadataValue;
-			TSharedPtr<FLocMetadataValueString> MetadataString = StaticCastSharedPtr<FLocMetadataValueString>( MetadataValuePtr );
-			MetadataString->SetString( TEXT("") );
-		}
-		break;
 
-	case ELocMetadataType::Array:
-		{
-			TArray< TSharedPtr< FLocMetadataValue > > MetadataArray = MetadataValue->AsArray();
-			for( auto ArrayIter = MetadataArray.CreateIterator(); ArrayIter; ++ArrayIter )
-			{
-				TSharedPtr<FLocMetadataValue>& Item = *ArrayIter;
-				if( Item.IsValid() )
-				{
-					ConditionTranslationMetadata( Item.ToSharedRef() );
-				}
-			}
-		}
-		break;
-
-	case ELocMetadataType::Object:
-		{
-			TSharedPtr< FLocMetadataObject > MetadataObject = MetadataValue->AsObject();
-			for( auto ValueIter = MetadataObject->Values.CreateConstIterator(); ValueIter; ++ValueIter )
-			{
-				const FString Name = (*ValueIter).Key;
-				TSharedPtr< FLocMetadataValue > Value = (*ValueIter).Value;
-				if( Value.IsValid() )
-				{
-					if( Value->Type == ELocMetadataType::String )
-					{
-						MetadataObject->SetStringField( Name, TEXT("") );
-					}
-					else
-					{
-						ConditionTranslationMetadata( Value.ToSharedRef() );
+						LocTextHelper.AddTranslation(CultureName, InManifestEntry->Namespace, Context.Key, Context.KeyMetadataObj, ArchiveSource, ArchiveTranslation, Context.bIsOptional);
 					}
 				}
 			}
-		}
 
-	default:
-		break;
-	}
-}
+			return true; // continue enumeration
+		}, true);
 
-void UGenerateGatherArchiveCommandlet::ConditionTranslation( FLocItem& LocItem )
-{
-	// We clear out the translation text because this should only be entered by translators
-	LocItem.Text = TEXT("");
-
-	// The translation might have metadata so we want to clear all the values of any string metadata
-	if( LocItem.MetadataObj.IsValid() )
-	{
-		ConditionTranslationMetadata( MakeShareable( new FLocMetadataValueObject( LocItem.MetadataObj ) ) );
-	}
-	
-}
+		// Trim any dead entries out of the archive
+		LocTextHelper.TrimArchive(CultureName);
 
-void ConditionSourceMetadata( TSharedRef<FLocMetadataValue> MetadataValue )
-{
-	switch( MetadataValue->Type )
-	{
-	case ELocMetadataType::Object:
+		// Save the new archive
 		{
-			TSharedPtr< FLocMetadataObject > MetadataObject = MetadataValue->AsObject();
-
-			TArray< FString > NamesToBeReplaced;
-
-			// We go through all the source metadata entries and look for any that have names prefixed with a '*'.  If we
-			//  find any, we will replace them with a String type that contains an empty value
-			for( auto Iter = MetadataObject->Values.CreateConstIterator(); Iter; ++Iter )
-			{
-				const FString& Name = (*Iter).Key;
-				TSharedPtr< FLocMetadataValue > Value = (*Iter).Value;
-
-				if( Name.StartsWith( FLocMetadataObject::COMPARISON_MODIFIER_PREFIX ) )
-				{
-					NamesToBeReplaced.Add(Name);
-				}
-				else
-				{
-					ConditionSourceMetadata( Value.ToSharedRef() );
-				}
-			}
-
-			for( auto Iter = NamesToBeReplaced.CreateConstIterator(); Iter; ++Iter )
+			FText SaveError;
+			if (!LocTextHelper.SaveForeignArchive(CultureName, &SaveError))
 			{
-				MetadataObject->RemoveField( *Iter );
-				MetadataObject->SetStringField( *Iter, TEXT("") );
+				UE_LOG(LogGenerateArchiveCommandlet, Error, TEXT("%s"), *SaveError.ToString());
+				return -1;
 			}
 		}
-
-	default:
-		break;
-	}
-}
-
-void UGenerateGatherArchiveCommandlet::ConditionSource( FLocItem& LocItem )
-{
-	if( LocItem.MetadataObj.IsValid() )
-	{
-		ConditionSourceMetadata( MakeShareable( new FLocMetadataValueObject( LocItem.MetadataObj ) ) );
 	}
+	
+	return 0;
 }
diff --git a/Engine/Source/Editor/UnrealEd/Private/Commandlets/GenerateGatherManifestCommandlet.cpp b/Engine/Source/Editor/UnrealEd/Private/Commandlets/GenerateGatherManifestCommandlet.cpp
--- a/Engine/Source/Editor/UnrealEd/Private/Commandlets/GenerateGatherManifestCommandlet.cpp
+++ b/Engine/Source/Editor/UnrealEd/Private/Commandlets/GenerateGatherManifestCommandlet.cpp
@@ -1,10 +1,6 @@
 // Copyright 1998-2016 Epic Games, Inc. All Rights Reserved.
 
 #include "UnrealEd.h"
-#include "ISourceControlModule.h"
-#include "Internationalization/InternationalizationManifest.h"
-#include "Json.h"
-#include "JsonInternationalizationManifestSerializer.h"
 
 DEFINE_LOG_CATEGORY_STATIC(LogGenerateManifestCommandlet, Log, All);
 
@@ -72,43 +68,26 @@ int32 UGenerateGatherManifestCommandlet::Main( const FString& Params )
 	TArray<FString> ManifestDependenciesList;
 	GetPathArrayFromConfig(*SectionName, TEXT("ManifestDependencies"), ManifestDependenciesList, GatherTextConfigPath);
 
-	if( ManifestDependenciesList.Num() > 0 )
+	for (const FString& ManifestDependency : ManifestDependenciesList)
 	{
-		if( !ManifestInfo->AddManifestDependencies( ManifestDependenciesList ) )
+		FText OutError;
+		if (!GatherManifestHelper->AddDependency(ManifestDependency, &OutError))
 		{
-			UE_LOG(LogGenerateManifestCommandlet, Error, TEXT("The GenerateGatherManifest commandlet couldn't find all the specified manifest dependencies."));
+			UE_LOG(LogGenerateManifestCommandlet, Error, TEXT("The GenerateGatherManifest commandlet couldn't load the specified manifest dependency: '%'. %s"), *ManifestDependency, *OutError.ToString());
 			return -1;
 		}
-	
-		ManifestInfo->ApplyManifestDependencies();
 	}
+
+	// Trim the manifest to remove any entries that came from a dependency
+	GatherManifestHelper->TrimManifest();
 	
 	const FString ManifestPath = FPaths::ConvertRelativePathToFull(DestinationPath) / ManifestName;
-	if( !WriteManifestToFile( ManifestInfo->GetManifest(), ManifestPath ) )
+	FText ManifestSaveError;
+	if (!GatherManifestHelper->SaveManifest(ManifestPath, &ManifestSaveError))
 	{
-		UE_LOG( LogGenerateManifestCommandlet, Error,TEXT("Failed to write manifest to %s."), *ManifestPath );				
+		UE_LOG(LogGenerateManifestCommandlet, Error,TEXT("Failed to write manifest to %s. %s."), *ManifestPath, *ManifestSaveError.ToString());
 		return -1;
 	}
-	return 0;
-}
 
-bool UGenerateGatherManifestCommandlet::WriteManifestToFile( const TSharedPtr<FInternationalizationManifest>& InManifest, const FString& OutputFilePath )
-{
-	UE_LOG(LogGenerateManifestCommandlet, Log, TEXT("Writing manifest to %s."), *OutputFilePath);
-
-	// We can not continue if the provided manifest is not valid
-	if( !InManifest.IsValid() )
-	{
-		return false;
-	}
-	FJsonInternationalizationManifestSerializer ManifestSerializer;
-	TSharedRef<FJsonObject> JsonManifestObj = MakeShareable( new FJsonObject );
-	
-	bool bSuccess = ManifestSerializer.SerializeManifest( InManifest.ToSharedRef(), JsonManifestObj );
-	
-	if( bSuccess )
-	{
-		bSuccess = WriteJSONToTextFile( JsonManifestObj, OutputFilePath, SourceControlInfo );
-	}
-	return bSuccess;	
+	return 0;
 }
diff --git a/Engine/Source/Editor/UnrealEd/Private/Commandlets/GenerateTextLocalizationReportCommandlet.cpp b/Engine/Source/Editor/UnrealEd/Private/Commandlets/GenerateTextLocalizationReportCommandlet.cpp
--- a/Engine/Source/Editor/UnrealEd/Private/Commandlets/GenerateTextLocalizationReportCommandlet.cpp
+++ b/Engine/Source/Editor/UnrealEd/Private/Commandlets/GenerateTextLocalizationReportCommandlet.cpp
@@ -1,11 +1,6 @@
 // Copyright 1998-2016 Epic Games, Inc. All Rights Reserved.
 
 #include "UnrealEd.h"
-#include "Json.h"
-#include "InternationalizationManifest.h"
-#include "InternationalizationArchive.h"
-#include "JsonInternationalizationManifestSerializer.h"
-#include "JsonInternationalizationArchiveSerializer.h"
 #include "BreakIterator.h"
 
 DEFINE_LOG_CATEGORY_STATIC(LogGenerateTextLocalizationReportCommandlet, Log, All);
@@ -15,29 +10,6 @@ UGenerateTextLocalizationReportCommandlet::UGenerateTextLocalizationReportComman
 {
 }
 
-int32 UGenerateTextLocalizationReportCommandlet::CountWords( const FString& Text ) const
-{
-	int32 NumOfWords = 0;
-	// Calculate # of words
-
-	TSharedRef<IBreakIterator> LineBreakIterator = FBreakIterator::CreateLineBreakIterator();
-	LineBreakIterator->SetString( Text );
-
-	int32 PreviousBreak = 0;
-	int32 CurrentBreak = 0;
-
-	while( ( CurrentBreak = LineBreakIterator->MoveToNext() ) != INDEX_NONE )
-	{
-		if ( CurrentBreak > PreviousBreak )
-		{
-			++NumOfWords;
-		}
-		PreviousBreak = CurrentBreak;
-	}
-
-	return NumOfWords;
-}
-
 int32 UGenerateTextLocalizationReportCommandlet::Main(const FString& Params)
 {
 	// Parse command line - we're interested in the param vals
@@ -106,7 +78,6 @@ int32 UGenerateTextLocalizationReportCommandlet::Main(const FString& Params)
 	GetBoolFromConfig( *SectionName, TEXT("bWordCountReport"), bWordCountReport, GatherTextConfigPath );
 	GetBoolFromConfig( *SectionName, TEXT("bConflictReport"), bConflictReport, GatherTextConfigPath );
 
-
 	if( bWordCountReport )
 	{
 		if( !ProcessWordCountReport( SourcePath, DestinationPath ) )
@@ -169,7 +140,7 @@ bool UGenerateTextLocalizationReportCommandlet::ProcessWordCountReport(const FSt
 	}
 
 	FString ReportStr;
-	FString ReportFilePath = (DestinationPath / WordCountReportName);
+	const FString ReportFilePath = (DestinationPath / WordCountReportName);
 
 	if ( FPaths::FileExists( ReportFilePath ) )
 	{
@@ -187,28 +158,13 @@ bool UGenerateTextLocalizationReportCommandlet::ProcessWordCountReport(const FSt
 		return false;
 	}
 
-	// Prepare the manifest
-	TSharedRef<FInternationalizationManifest> InternationalizationManifest = MakeShareable(new FInternationalizationManifest());
+	// Load the manifest and all archives
+	FLocTextHelper LocTextHelper(SourcePath, ManifestName, ArchiveName, FString(), CulturesToGenerate, MakeShareable(new FLocFileSCCNotifies(SourceControlInfo)));
 	{
-		const FString ManifestFileName = SourcePath / ManifestName;
-		if (!FPaths::FileExists(ManifestFileName))
-		{
-			UE_LOG(LogGenerateTextLocalizationReportCommandlet, Error, TEXT("Failed to find manifest '%s'."), *ManifestFileName);
-			return false;
-		}
-
-		const TSharedPtr<FJsonObject> ManifestJsonObject = ReadJSONTextFile(ManifestFileName);
-		if (!ManifestJsonObject.IsValid())
+		FText LoadError;
+		if (!LocTextHelper.LoadAll(ELocTextHelperLoadFlags::LoadOrCreate, &LoadError))
 		{
-			UE_LOG(LogGenerateTextLocalizationReportCommandlet, Error, TEXT("Failed to parse manifest '%s'."), *ManifestFileName);
-			return false;
-		}
-
-		FJsonInternationalizationManifestSerializer ManifestSerializer;
-		InternationalizationManifest = MakeShareable(new FInternationalizationManifest());
-		if (!ManifestSerializer.DeserializeManifest(ManifestJsonObject.ToSharedRef(), InternationalizationManifest))
-		{
-			UE_LOG(LogGenerateTextLocalizationReportCommandlet, Error, TEXT("Failed to deserialize manifest '%s'."), *ManifestFileName);
+			UE_LOG(LogGenerateTextLocalizationReportCommandlet, Error, TEXT("%s"), *LoadError.ToString());
 			return false;
 		}
 	}
@@ -226,9 +182,6 @@ bool UGenerateTextLocalizationReportCommandlet::ProcessWordCountReport(const FSt
 	int32 NewRowIndex = ReportData.AddRow();
 	ReportData.SetEntry(NewRowIndex, FWordCountReportData::ColHeadingDateTime, TimeStamp);
 
-	int32 TotalNumOfWords = 0;
-
-
 	struct FCaseSensitiveStringKeyFuncs : BaseKeyFuncs<FString, FString, false>
 	{
 		static FORCEINLINE const FString& GetSetKey(const FString& Element)
@@ -245,36 +198,55 @@ bool UGenerateTextLocalizationReportCommandlet::ProcessWordCountReport(const FSt
 		}
 	};
 
-	TSet< FString, FCaseSensitiveStringKeyFuncs > CountedText;
-	for(TManifestEntryBySourceTextContainer::TConstIterator i = InternationalizationManifest->GetEntriesBySourceTextIterator(); i; ++i)
+	auto CountWords = [](const FString& InTextToCount) -> int32
 	{
-		// Gather relevant info from manifest entry.
-		const TSharedRef<FManifestEntry>& ManifestEntry = i.Value();
-		const FLocItem& Source = ManifestEntry->Source;
+		int32 NumOfWords = 0;
+		// Calculate # of words
+
+		TSharedRef<IBreakIterator> LineBreakIterator = FBreakIterator::CreateLineBreakIterator();
+		LineBreakIterator->SetString(InTextToCount);
+
+		int32 PreviousBreak = 0;
+		int32 CurrentBreak = 0;
 
-		bool HasNonOptional = false;
-		for( auto ContextIter = ManifestEntry->Contexts.CreateConstIterator(); ContextIter; ++ContextIter )
+		while ((CurrentBreak = LineBreakIterator->MoveToNext()) != INDEX_NONE)
 		{
-			if( !(*ContextIter).bIsOptional )
+			if (CurrentBreak > PreviousBreak)
 			{
-				HasNonOptional = true;
-				break;
+				++NumOfWords;
 			}
+			PreviousBreak = CurrentBreak;
 		}
 
-		if( HasNonOptional )
+		return NumOfWords;
+	};
+
+	int32 TotalNumOfWords = 0;
+	TSet< FString, FCaseSensitiveStringKeyFuncs > CountedText;
+
+	LocTextHelper.EnumerateSourceTexts([&TotalNumOfWords, &CountedText, &CountWords](TSharedRef<FManifestEntry> InManifestEntry) -> bool
+	{
+		const int32 NumWords = CountWords(InManifestEntry->Source.Text);
+
+		// Gather relevant info from each manifest entry
+		for (const FManifestContext& Context : InManifestEntry->Contexts)
 		{
-			const FString CountedTextId = Source.Text + FString(TEXT("::")) + ManifestEntry->Namespace;
-			if ( !CountedText.Contains( CountedTextId ) )
+			if (!Context.bIsOptional)
 			{
-				TotalNumOfWords += CountWords( Source.Text );
+				const FString CountedTextId = FString::Printf(TEXT("%s::%s::%s"), *InManifestEntry->Source.Text, *InManifestEntry->Namespace, *Context.Key);
+				if (!CountedText.Contains(CountedTextId))
+				{
+					TotalNumOfWords += NumWords;
 
-				bool IsAlreadySet = false;
-				CountedText.Add( CountedTextId, &IsAlreadySet );
-				check( !IsAlreadySet );
+					bool IsAlreadySet = false;
+					CountedText.Add(CountedTextId, &IsAlreadySet);
+					check(!IsAlreadySet);
+				}
 			}
 		}
-	}
+
+		return true; // continue enumeration
+	}, true);
 
 	ReportData.SetEntry(NewRowIndex, FWordCountReportData::ColHeadingWordCount, FString::FromInt( TotalNumOfWords ) );
 	
@@ -282,147 +254,79 @@ bool UGenerateTextLocalizationReportCommandlet::ProcessWordCountReport(const FSt
 	for (const FString& CultureName : CulturesToGenerate)
 	{
 		uint32 TranslatedWords = 0;
-		FString CulturePath = SourcePath / CultureName;
-		CountedText.Empty();
-
-		TSharedRef<FInternationalizationArchive> InternationalizationArchive = MakeShareable(new FInternationalizationArchive());
-		{
-			const FString ArchiveFilePath = FPaths::ConvertRelativePathToFull(CulturePath / ArchiveName);
-			TSharedPtr<FJsonObject> ArchiveJSONObject = ReadJSONTextFile(ArchiveFilePath);
-			if( !(ArchiveJSONObject.IsValid()) )
-			{
-				UE_LOG(LogGenerateTextLocalizationReportCommandlet, Error, TEXT("No archive found at %s."), *ArchiveFilePath);
-				continue;
-			}
-			
-			FJsonInternationalizationArchiveSerializer ArchiveSerializer;
-			ArchiveSerializer.DeserializeArchive(ArchiveJSONObject.ToSharedRef(), InternationalizationArchive);
-		}
+		CountedText.Reset();
 
 		// Finds all the manifest entries in the archive and adds the source text word count to the running total if there is a valid translation.
-		for (TManifestEntryBySourceTextContainer::TConstIterator i = InternationalizationManifest->GetEntriesBySourceTextIterator(); i; ++i)
+		LocTextHelper.EnumerateSourceTexts([&LocTextHelper, &CultureName, &TranslatedWords, &CountedText, &CountWords](TSharedRef<FManifestEntry> InManifestEntry) -> bool
 		{
-			// Gather relevant info from manifest entry.
-			const TSharedRef<FManifestEntry>& ManifestEntry = i.Value();
-			const FString& Namespace = ManifestEntry->Namespace;
-			const FLocItem& Source = ManifestEntry->Source;
+			const int32 NumWords = CountWords(InManifestEntry->Source.Text);
 
-			if (ManifestEntry->Contexts.Num() > 0)
+			// Gather relevant info from each manifest entry
+			for (const FManifestContext& Context : InManifestEntry->Contexts)
 			{
-				TSharedPtr<FArchiveEntry> ArchiveEntry = InternationalizationArchive->FindEntryBySource(Namespace, Source, ManifestEntry->Contexts[0].KeyMetadataObj);
-				if (ArchiveEntry.IsValid() && !ArchiveEntry->Translation.Text.IsEmpty())
+				if (!Context.bIsOptional)
 				{
-					bool HasNonOptional = false;
-					for (auto ContextIter = ManifestEntry->Contexts.CreateConstIterator(); ContextIter; ++ContextIter)
+					TSharedPtr<FArchiveEntry> ArchiveEntry = LocTextHelper.FindTranslation(CultureName, InManifestEntry->Namespace, Context.Key, Context.KeyMetadataObj);
+					if (ArchiveEntry.IsValid() && ArchiveEntry->Source.IsExactMatch(InManifestEntry->Source) && !ArchiveEntry->Translation.Text.IsEmpty())
 					{
-						if (!(*ContextIter).bIsOptional)
+						const FString CountedTextId = FString::Printf(TEXT("%s::%s::%s"), *InManifestEntry->Source.Text, *InManifestEntry->Namespace, *Context.Key);
+						if (!CountedText.Contains(CountedTextId))
 						{
-							HasNonOptional = true;
-							break;
-						}
-					}
-
-					const FString CountedTextId = Source.Text + FString(TEXT("::")) + ManifestEntry->Namespace;
-					if (!CountedText.Contains(CountedTextId))
-					{
-						TranslatedWords += CountWords(Source.Text);
+							TranslatedWords += NumWords;
 
-						bool IsAlreadySet = false;
-						CountedText.Add(CountedTextId, &IsAlreadySet);
-						check(!IsAlreadySet);
+							bool IsAlreadySet = false;
+							CountedText.Add(CountedTextId, &IsAlreadySet);
+							check(!IsAlreadySet);
+						}
 					}
 				}
 			}
-		}
+			
+			return true; // continue enumeration
+		}, true);
+
 		ReportData.SetEntry(NewRowIndex, CultureName, FString::FromInt(TranslatedWords));
 	}
 
-	const bool DidFileExist = FPaths::FileExists(ReportFilePath);
-	if (DidFileExist)
+	const bool bReportFileSaved = FLocalizedAssetSCCUtil::SaveFileWithSCC(SourceControlInfo, ReportFilePath, [&ReportData](const FString& InSaveFileName) -> bool
 	{
-		if( SourceControlInfo.IsValid() )
-		{
-			FText SCCErrorText;
-			if (!SourceControlInfo->CheckOutFile(ReportFilePath, SCCErrorText))
-			{
-				UE_LOG(LogGenerateTextLocalizationReportCommandlet, Warning, TEXT("Check out of file %s failed: %s"), *ReportFilePath, *SCCErrorText.ToString());
-			}
-		}
-	}
+		return FFileHelper::SaveStringToFile(ReportData.ToCSV(), *InSaveFileName);
+	});
 
-	if ( !FFileHelper::SaveStringToFile( ReportData.ToCSV(), *ReportFilePath ) )
+	if (!bReportFileSaved)
 	{
 		UE_LOG(LogGenerateTextLocalizationReportCommandlet, Error, TEXT("Unable to save report at %s."), *ReportFilePath);
 		return false;
 	}
 
-	if (!DidFileExist)
-	{
-		// Checkout on a new file will cause it to be added
-		if( SourceControlInfo.IsValid() )
-		{
-			FText SCCErrorText;
-			if (!SourceControlInfo->CheckOutFile(ReportFilePath, SCCErrorText))
-			{
-				UE_LOG(LogGenerateTextLocalizationReportCommandlet, Warning, TEXT("Check out of file %s failed: %s"), *ReportFilePath, *SCCErrorText.ToString());
-			}
-		}
-	}
-
 	return true;
 }
 
 bool UGenerateTextLocalizationReportCommandlet::ProcessConflictReport(const FString& DestinationPath)
 {
-	FString ConflictReportName;
-	FString TimeStamp;
-
-
-	TimeStamp = (CmdlineTimeStamp.IsEmpty()) ? *FDateTime::Now().ToString() : CmdlineTimeStamp;
+	const FString TimeStamp = (CmdlineTimeStamp.IsEmpty()) ? *FDateTime::Now().ToString() : CmdlineTimeStamp;
 
 	// Get report name.
-	if( !( GetStringFromConfig( *SectionName, TEXT("ConflictReportName"), ConflictReportName, GatherTextConfigPath ) ) )
+	FString ConflictReportName;
+	if (!GetStringFromConfig(*SectionName, TEXT("ConflictReportName"), ConflictReportName, GatherTextConfigPath))
 	{
 		UE_LOG(LogGenerateTextLocalizationReportCommandlet, Error, TEXT("No conflict report name specified."));
 		return false;
 	}
 
-	FString ReportStr;
-	FString ReportFilePath = (DestinationPath / ConflictReportName);
-
-	ReportStr = FConflictReportInfo::GetInstance().ToString();
+	const FString ReportStr = GatherManifestHelper->GetConflictReport();
+	const FString ReportFilePath = (DestinationPath / ConflictReportName);
 
-	const bool DidFileExist = FPaths::FileExists(ReportFilePath);
-	if (DidFileExist)
+	const bool bReportFileSaved = FLocalizedAssetSCCUtil::SaveFileWithSCC(SourceControlInfo, ReportFilePath, [&ReportStr](const FString& InSaveFileName) -> bool
 	{
-		if( SourceControlInfo.IsValid() )
-		{
-			FText SCCErrorText;
-			if (!SourceControlInfo->CheckOutFile(ReportFilePath, SCCErrorText))
-			{
-				UE_LOG(LogGenerateTextLocalizationReportCommandlet, Warning, TEXT("Check out of file %s failed: %s"), *ReportFilePath, *SCCErrorText.ToString());
-			}
-		}
-	}
+		return FFileHelper::SaveStringToFile(ReportStr, *InSaveFileName);
+	});
 
-	if ( !FFileHelper::SaveStringToFile( ReportStr, *ReportFilePath ) )
+	if (!bReportFileSaved)
 	{
 		UE_LOG(LogGenerateTextLocalizationReportCommandlet, Error, TEXT("Unable to save report at %s."), *ReportFilePath);
 		return false;
 	}
 
-	if (!DidFileExist)
-	{
-		// Checkout on a new file will cause it to be added
-		if( SourceControlInfo.IsValid() )
-		{
-			FText SCCErrorText;
-			if (!SourceControlInfo->CheckOutFile(ReportFilePath, SCCErrorText))
-			{
-				UE_LOG(LogGenerateTextLocalizationReportCommandlet, Warning, TEXT("Check out of file %s failed: %s"), *ReportFilePath, *SCCErrorText.ToString());
-			}
-		}
-	}
-
 	return true;
 }
diff --git a/Engine/Source/Editor/UnrealEd/Private/Commandlets/GenerateTextLocalizationResourceCommandlet.cpp b/Engine/Source/Editor/UnrealEd/Private/Commandlets/GenerateTextLocalizationResourceCommandlet.cpp
--- a/Engine/Source/Editor/UnrealEd/Private/Commandlets/GenerateTextLocalizationResourceCommandlet.cpp
+++ b/Engine/Source/Editor/UnrealEd/Private/Commandlets/GenerateTextLocalizationResourceCommandlet.cpp
@@ -2,11 +2,6 @@
 
 #include "UnrealEd.h"
 #include "TextLocalizationResourceGenerator.h"
-#include "Json.h"
-#include "InternationalizationManifest.h"
-#include "InternationalizationArchive.h"
-#include "JsonInternationalizationManifestSerializer.h"
-#include "JsonInternationalizationArchiveSerializer.h"
 
 DEFINE_LOG_CATEGORY_STATIC(LogGenerateTextLocalizationResourceCommandlet, Log, All);
 
@@ -67,6 +62,14 @@ int32 UGenerateTextLocalizationResourceCommandlet::Main(const FString& Params)
 		return -1;
 	}
 
+	// Get archive name.
+	FString ArchiveName;
+	if (!GetStringFromConfig(*SectionName, TEXT("ArchiveName"), ArchiveName, GatherTextConfigPath))
+	{
+		UE_LOG(LogGenerateTextLocalizationResourceCommandlet, Error, TEXT("No archive name specified."));
+		return -1;
+	}
+
 	// Get cultures to generate.
 	FString NativeCultureName;
 	if( !( GetStringFromConfig( *SectionName, TEXT("NativeCulture"), NativeCultureName, GatherTextConfigPath ) ) )
@@ -109,68 +112,52 @@ int32 UGenerateTextLocalizationResourceCommandlet::Main(const FString& Params)
 		return -1;
 	}
 
-	// Read the manifest file from the source path.
-	FString ManifestFilePath = (SourcePath / ManifestName);
-	ManifestFilePath = FPaths::ConvertRelativePathToFull(ManifestFilePath);
-	TSharedPtr<FJsonObject> ManifestJSONObject = ReadJSONTextFile(ManifestFilePath);
-	if( !(ManifestJSONObject.IsValid()) )
+	// Get whether to skip the source check.
+	bool bSkipSourceCheck = false;
+	if (!GetBoolFromConfig(*SectionName, TEXT("bSkipSourceCheck"), bSkipSourceCheck, GatherTextConfigPath))
 	{
-		UE_LOG(LogGenerateTextLocalizationResourceCommandlet, Error, TEXT("No manifest found at %s."), *ManifestFilePath);
-		return -1;
+		bSkipSourceCheck = false;
 	}
-	TSharedRef<FInternationalizationManifest> InternationalizationManifest = MakeShareable( new FInternationalizationManifest );
+
+	// Load the manifest and all archives
+	FLocTextHelper LocTextHelper(SourcePath, ManifestName, ArchiveName, NativeCultureName, CulturesToGenerate, MakeShareable(new FLocFileSCCNotifies(SourceControlInfo)));
 	{
-		FJsonInternationalizationManifestSerializer ManifestSerializer;
-		ManifestSerializer.DeserializeManifest(ManifestJSONObject.ToSharedRef(), InternationalizationManifest);
+		FText LoadError;
+		if (!LocTextHelper.LoadAll(ELocTextHelperLoadFlags::LoadOrCreate, &LoadError))
+		{
+			UE_LOG(LogGenerateTextLocalizationResourceCommandlet, Error, TEXT("%s"), *LoadError.ToString());
+			return false;
+		}
 	}
 
 	// For each culture:
-	for(int32 Culture = 0; Culture < CulturesToGenerate.Num(); Culture++)
+	for (const FString& CultureName : CulturesToGenerate)
 	{
-		const FString CultureName = *(CulturesToGenerate[Culture]);
-		const FString CulturePath = SourcePath / CultureName;
-
 		// Write resource.
 		const FString TextLocalizationResourcePath = DestinationPath / CultureName / ResourceName;
 
-		const bool DidFileExist = FPaths::FileExists(TextLocalizationResourcePath);
-		if (DidFileExist)
+		const bool bLocResFileSaved = FLocalizedAssetSCCUtil::SaveFileWithSCC(SourceControlInfo, TextLocalizationResourcePath, [&LocTextHelper, &CultureName, &bSkipSourceCheck](const FString& InSaveFileName) -> bool
 		{
-			if( SourceControlInfo.IsValid() )
+			bool bSaved = false;
+
+			TAutoPtr<FArchive> TextLocalizationResourceArchive(IFileManager::Get().CreateFileWriter(*InSaveFileName));
+			if (TextLocalizationResourceArchive)
 			{
-				FText SCCErrorText;
-				if (!SourceControlInfo->CheckOutFile(TextLocalizationResourcePath, SCCErrorText))
+				bSaved = FTextLocalizationResourceGenerator::Generate(LocTextHelper, CultureName, bSkipSourceCheck, *TextLocalizationResourceArchive);
+				if (!bSaved)
 				{
-					UE_LOG(LogGenerateTextLocalizationResourceCommandlet, Error, TEXT("Check out of file %s failed: %s"), *TextLocalizationResourcePath, *SCCErrorText.ToString());
-					return -1;
+					IFileManager::Get().Delete(*InSaveFileName);
 				}
+				TextLocalizationResourceArchive->Close();
 			}
-		}
-
-		TAutoPtr<FArchive> TextLocalizationResourceArchive(IFileManager::Get().CreateFileWriter(*TextLocalizationResourcePath));
-		if (TextLocalizationResourceArchive)
-		{
-			FJsonInternationalizationArchiveSerializer ArchiveSerializer;
 
-			if( !(FTextLocalizationResourceGenerator::Generate(SourcePath, InternationalizationManifest, NativeCultureName, CultureName, TextLocalizationResourceArchive, ArchiveSerializer)) )
-			{
-				IFileManager::Get().Delete( *TextLocalizationResourcePath );
-			}
-			TextLocalizationResourceArchive->Close();
-		}
+			return bSaved;
+		});
 
-		if (!DidFileExist)
+		if (!bLocResFileSaved)
 		{
-			// Checkout on a new file will cause it to be added
-			if( SourceControlInfo.IsValid() )
-			{
-				FText SCCErrorText;
-				if (!SourceControlInfo->CheckOutFile(TextLocalizationResourcePath, SCCErrorText))
-				{
-					UE_LOG(LogGenerateTextLocalizationResourceCommandlet, Error, TEXT("Check out of file %s failed: %s"), *TextLocalizationResourcePath, *SCCErrorText.ToString());
-					return -1;
-				}
-			}
+			UE_LOG(LogGenerateTextLocalizationResourceCommandlet, Error, TEXT("Could not write file %s"), *TextLocalizationResourcePath);
+			return false;
 		}
 	}
 
diff --git a/Engine/Source/Editor/UnrealEd/Private/Commandlets/ImportDialogueScriptCommandlet.cpp b/Engine/Source/Editor/UnrealEd/Private/Commandlets/ImportDialogueScriptCommandlet.cpp
--- a/Engine/Source/Editor/UnrealEd/Private/Commandlets/ImportDialogueScriptCommandlet.cpp
+++ b/Engine/Source/Editor/UnrealEd/Private/Commandlets/ImportDialogueScriptCommandlet.cpp
@@ -4,10 +4,6 @@
 #include "Commandlets/ImportDialogueScriptCommandlet.h"
 #include "CsvParser.h"
 #include "Sound/DialogueWave.h"
-#include "Internationalization/InternationalizationManifest.h"
-#include "Internationalization/InternationalizationArchive.h"
-#include "JsonInternationalizationManifestSerializer.h"
-#include "JsonInternationalizationArchiveSerializer.h"
 
 DEFINE_LOG_CATEGORY_STATIC(LogImportDialogueScriptCommandlet, Log, All);
 
@@ -111,55 +107,21 @@ int32 UImportDialogueScriptCommandlet::Main(const FString& Params)
 		return -1;
 	}
 
-	// Prepare the manifest
+	// We may only have a single culture if using this setting
+	if (!bUseCultureDirectory && CulturesToGenerate.Num() > 1)
 	{
-		const FString ManifestFileName = DestinationPath / ManifestName;
-		if (!FPaths::FileExists(ManifestFileName))
-		{
-			UE_LOG(LogImportDialogueScriptCommandlet, Error, TEXT("Failed to find manifest '%s'."), *ManifestFileName);
-			return -1;
-		}
-
-		const TSharedPtr<FJsonObject> ManifestJsonObject = ReadJSONTextFile(ManifestFileName);
-		if (!ManifestJsonObject.IsValid())
-		{
-			UE_LOG(LogImportDialogueScriptCommandlet, Error, TEXT("Failed to parse manifest '%s'."), *ManifestFileName);
-			return -1;
-		}
-
-		FJsonInternationalizationManifestSerializer ManifestSerializer;
-		InternationalizationManifest = MakeShareable(new FInternationalizationManifest());
-		if (!ManifestSerializer.DeserializeManifest(ManifestJsonObject.ToSharedRef(), InternationalizationManifest.ToSharedRef()))
-		{
-			UE_LOG(LogImportDialogueScriptCommandlet, Error, TEXT("Failed to deserialize manifest '%s'."), *ManifestFileName);
-			return -1;
-		}
+		UE_LOG(LogImportDialogueScriptCommandlet, Error, TEXT("bUseCultureDirectory may only be used with a single culture."));
+		return false;
 	}
 
-	// Prepare the native archive
+	// Load the manifest and all archives
+	FLocTextHelper LocTextHelper(DestinationPath, ManifestName, ArchiveName, NativeCulture, CulturesToGenerate, MakeShareable(new FLocFileSCCNotifies(SourceControlInfo)));
 	{
-		const FString NativeCulturePath = DestinationPath / NativeCulture;
-
-		const FString NativeArchiveFileName = NativeCulturePath / ArchiveName;
-		if (!FPaths::FileExists(NativeArchiveFileName))
+		FText LoadError;
+		if (!LocTextHelper.LoadAll(ELocTextHelperLoadFlags::LoadOrCreate, &LoadError))
 		{
-			UE_LOG(LogImportDialogueScriptCommandlet, Error, TEXT("Failed to find archive '%s'."), *NativeArchiveFileName);
-			return -1;
-		}
-
-		const TSharedPtr<FJsonObject> ArchiveJsonObject = ReadJSONTextFile(NativeArchiveFileName);
-		if (!ArchiveJsonObject.IsValid())
-		{
-			UE_LOG(LogImportDialogueScriptCommandlet, Error, TEXT("Failed to parse archive '%s'."), *NativeArchiveFileName);
-			return -1;
-		}
-
-		FJsonInternationalizationArchiveSerializer ArchiveSerializer;
-		NativeArchive = MakeShareable(new FInternationalizationArchive());
-		if (!ArchiveSerializer.DeserializeArchive(ArchiveJsonObject.ToSharedRef(), NativeArchive.ToSharedRef()))
-		{
-			UE_LOG(LogImportDialogueScriptCommandlet, Error, TEXT("Failed to deserialize archive '%s'."), *NativeArchiveFileName);
-			return -1;
+			UE_LOG(LogImportDialogueScriptCommandlet, Error, TEXT("%s"), *LoadError.ToString());
+			return false;
 		}
 	}
 
@@ -167,7 +129,7 @@ int32 UImportDialogueScriptCommandlet::Main(const FString& Params)
 	{
 		const FString CultureSourcePath = SourcePath / (bUseCultureDirectory ? NativeCulture : TEXT(""));
 		const FString CultureDestinationPath = DestinationPath / NativeCulture;
-		ImportDialogueScriptForCulture(CultureSourcePath / DialogueScriptName, CultureDestinationPath / ArchiveName, NativeCulture, true);
+		ImportDialogueScriptForCulture(LocTextHelper, CultureSourcePath / DialogueScriptName, NativeCulture, true);
 	}
 
 	// Import any remaining cultures
@@ -181,13 +143,13 @@ int32 UImportDialogueScriptCommandlet::Main(const FString& Params)
 
 		const FString CultureSourcePath = SourcePath / (bUseCultureDirectory ? CultureName : TEXT(""));
 		const FString CultureDestinationPath = DestinationPath / CultureName;
-		ImportDialogueScriptForCulture(CultureSourcePath / DialogueScriptName, CultureDestinationPath / ArchiveName, CultureName, false);
+		ImportDialogueScriptForCulture(LocTextHelper, CultureSourcePath / DialogueScriptName, CultureName, false);
 	}
 
 	return 0;
 }
 
-bool UImportDialogueScriptCommandlet::ImportDialogueScriptForCulture(const FString& InDialogueScriptFileName, const FString& InCultureArchiveFileName, const FString& InCultureName, const bool bIsNativeCulture)
+bool UImportDialogueScriptCommandlet::ImportDialogueScriptForCulture(FLocTextHelper& InLocTextHelper, const FString& InDialogueScriptFileName, const FString& InCultureName, const bool bIsNativeCulture)
 {
 	// Load dialogue script file contents to string
 	FString DialogScriptFileContents;
@@ -250,36 +212,6 @@ bool UImportDialogueScriptCommandlet::ImportDialogueScriptForCulture(const FStri
 		return false;
 	}
 
-	// Prepare the culture archive
-	TSharedPtr<FInternationalizationArchive> CultureArchive;
-	if (bIsNativeCulture)
-	{
-		CultureArchive = NativeArchive;
-	}
-	else
-	{
-		if (!FPaths::FileExists(InCultureArchiveFileName))
-		{
-			UE_LOG(LogImportDialogueScriptCommandlet, Error, TEXT("Failed to find archive '%s'."), *InCultureArchiveFileName);
-			return false;
-		}
-
-		const TSharedPtr<FJsonObject> ArchiveJsonObject = ReadJSONTextFile(InCultureArchiveFileName);
-		if (!ArchiveJsonObject.IsValid())
-		{
-			UE_LOG(LogImportDialogueScriptCommandlet, Error, TEXT("Failed to parse archive '%s'."), *InCultureArchiveFileName);
-			return false;
-		}
-
-		FJsonInternationalizationArchiveSerializer ArchiveSerializer;
-		CultureArchive = MakeShareable(new FInternationalizationArchive());
-		if (!ArchiveSerializer.DeserializeArchive(ArchiveJsonObject.ToSharedRef(), CultureArchive.ToSharedRef()))
-		{
-			UE_LOG(LogImportDialogueScriptCommandlet, Error, TEXT("Failed to deserialize archive '%s'."), *InCultureArchiveFileName);
-			return false;
-		}
-	}
-
 	bool bHasUpdatedArchive = false;
 
 	// Parse each row of the CSV data
@@ -312,7 +244,7 @@ bool UImportDialogueScriptCommandlet::ImportDialogueScriptForCulture(const FStri
 		for (const FString& ContextLocalizationKey : ParsedScriptEntry.LocalizationKeys)
 		{
 			// Find the manifest entry so that we can find the corresponding archive entry
-			TSharedPtr<FManifestEntry> ContextManifestEntry = InternationalizationManifest->FindEntryByKey(FDialogueConstants::DialogueNamespace, ContextLocalizationKey);
+			TSharedPtr<FManifestEntry> ContextManifestEntry = InLocTextHelper.FindSourceText(FDialogueConstants::DialogueNamespace, ContextLocalizationKey);
 			if (!ContextManifestEntry.IsValid())
 			{
 				UE_LOG(LogImportDialogueScriptCommandlet, Log, TEXT("No internationalization manifest entry was found for context '%s' in culture '%s'. This context will be skipped."), *ContextLocalizationKey, *InCultureName);
@@ -321,32 +253,17 @@ bool UImportDialogueScriptCommandlet::ImportDialogueScriptForCulture(const FStri
 
 			// Find the correct entry for our context
 			const FManifestContext* ContextManifestEntryContext = ContextManifestEntry->FindContextByKey(ContextLocalizationKey);
-			check(ContextManifestEntryContext); // This should never fail as we pass in the key to FindEntryByKey
+			check(ContextManifestEntryContext); // This should never fail as we pass in the key to FindSourceText
 
-			// Find the correct source text (we might have a native translation that we should update instead of the source)
-			FString SourceText = ContextManifestEntry->Source.Text;
-			if (!bIsNativeCulture)
-			{
-				TSharedPtr<FArchiveEntry> NativeArchiveEntry = NativeArchive->FindEntryBySource(FDialogueConstants::DialogueNamespace, SourceText, ContextManifestEntryContext->KeyMetadataObj);
-				if (NativeArchiveEntry.IsValid())
-				{
-					SourceText = NativeArchiveEntry->Translation.Text;
-				}
-			}
+			// Get the text we would have exported
+			FLocItem ExportedSource;
+			FLocItem ExportedTranslation;
+			InLocTextHelper.GetExportText(InCultureName, FDialogueConstants::DialogueNamespace, ContextManifestEntryContext->Key, ContextManifestEntryContext->KeyMetadataObj, ELocTextExportSourceMethod::NativeText, ContextManifestEntry->Source, ExportedSource, ExportedTranslation);
 
-			// Update (or add) the entry in the archive
-			TSharedPtr<FArchiveEntry> ArchiveEntry = CultureArchive->FindEntryBySource(FDialogueConstants::DialogueNamespace, SourceText, ContextManifestEntryContext->KeyMetadataObj);
-			if (ArchiveEntry.IsValid())
+			// Attempt to import the new text (if required)
+			if (!ExportedTranslation.Text.Equals(ParsedScriptEntry.SpokenDialogue, ESearchCase::CaseSensitive))
 			{
-				if (!ArchiveEntry->Translation.Text.Equals(ParsedScriptEntry.SpokenDialogue, ESearchCase::CaseSensitive))
-				{
-					bHasUpdatedArchive = true;
-					ArchiveEntry->Translation.Text = ParsedScriptEntry.SpokenDialogue;
-				}
-			}
-			else
-			{
-				if (CultureArchive->AddEntry(FDialogueConstants::DialogueNamespace, SourceText, ParsedScriptEntry.SpokenDialogue, ContextManifestEntryContext->KeyMetadataObj, false))
+				if (InLocTextHelper.ImportTranslation(InCultureName, FDialogueConstants::DialogueNamespace, ContextManifestEntryContext->Key, ContextManifestEntryContext->KeyMetadataObj, ExportedSource, FLocItem(ParsedScriptEntry.SpokenDialogue), ContextManifestEntryContext->bIsOptional))
 				{
 					bHasUpdatedArchive = true;
 				}
@@ -357,13 +274,10 @@ bool UImportDialogueScriptCommandlet::ImportDialogueScriptForCulture(const FStri
 	// Write out the updated archive file
 	if (bHasUpdatedArchive)
 	{
-		TSharedRef<FJsonObject> ArchiveJsonObj = MakeShareable(new FJsonObject());
-		FJsonInternationalizationArchiveSerializer ArchiveSerializer;
-		ArchiveSerializer.SerializeArchive(CultureArchive.ToSharedRef(), ArchiveJsonObj);
-
-		if (!WriteJSONToTextFile(ArchiveJsonObj, InCultureArchiveFileName, SourceControlInfo))
+		FText SaveError;
+		if (!InLocTextHelper.SaveForeignArchive(InCultureName, &SaveError))
 		{
-			UE_LOG(LogImportDialogueScriptCommandlet, Error, TEXT("Failed to write archive file for culture '%s' to '%s'."), *InCultureName, *InCultureArchiveFileName);
+			UE_LOG(LogImportDialogueScriptCommandlet, Error, TEXT("%s"), *SaveError.ToString());
 			return false;
 		}
 	}
diff --git a/Engine/Source/Editor/UnrealEd/Private/Commandlets/ImportLocalizedDialogueCommandlet.cpp b/Engine/Source/Editor/UnrealEd/Private/Commandlets/ImportLocalizedDialogueCommandlet.cpp
--- a/Engine/Source/Editor/UnrealEd/Private/Commandlets/ImportLocalizedDialogueCommandlet.cpp
+++ b/Engine/Source/Editor/UnrealEd/Private/Commandlets/ImportLocalizedDialogueCommandlet.cpp
@@ -6,10 +6,6 @@
 #include "TargetPlatform.h"
 #include "Sound/DialogueWave.h"
 #include "Sound/SoundWave.h"
-#include "Internationalization/InternationalizationManifest.h"
-#include "Internationalization/InternationalizationArchive.h"
-#include "JsonInternationalizationManifestSerializer.h"
-#include "JsonInternationalizationArchiveSerializer.h"
 
 DEFINE_LOG_CATEGORY_STATIC(LogImportLocalizedDialogueCommandlet, Log, All);
 
@@ -133,58 +129,17 @@ int32 UImportLocalizedDialogueCommandlet::Main(const FString& Params)
 		return -1;
 	}
 
-	// Prepare the manifest
+	// Load the manifest and all archives
+	FLocTextHelper LocTextHelper(SourcePath, ManifestName, ArchiveName, NativeCulture, CulturesToGenerate, MakeShareable(new FLocFileSCCNotifies(SourceControlInfo)));
 	{
-		const FString ManifestFileName = SourcePath / ManifestName;
-		if (!FPaths::FileExists(ManifestFileName))
+		FText LoadError;
+		if (!LocTextHelper.LoadAll(ELocTextHelperLoadFlags::LoadOrCreate, &LoadError))
 		{
-			UE_LOG(LogImportLocalizedDialogueCommandlet, Error, TEXT("Failed to find manifest '%s'."), *ManifestFileName);
-			return -1;
-		}
-
-		const TSharedPtr<FJsonObject> ManifestJsonObject = ReadJSONTextFile(ManifestFileName);
-		if (!ManifestJsonObject.IsValid())
-		{
-			UE_LOG(LogImportLocalizedDialogueCommandlet, Error, TEXT("Failed to parse manifest '%s'."), *ManifestFileName);
-			return -1;
-		}
-
-		FJsonInternationalizationManifestSerializer ManifestSerializer;
-		InternationalizationManifest = MakeShareable(new FInternationalizationManifest());
-		if (!ManifestSerializer.DeserializeManifest(ManifestJsonObject.ToSharedRef(), InternationalizationManifest.ToSharedRef()))
-		{
-			UE_LOG(LogImportLocalizedDialogueCommandlet, Error, TEXT("Failed to deserialize manifest '%s'."), *ManifestFileName);
-			return -1;
-		}
-	}
-
-	// Prepare the native archive
-	{
-		const FString NativeCulturePath = SourcePath / NativeCulture;
-
-		const FString NativeArchiveFileName = NativeCulturePath / ArchiveName;
-		if (!FPaths::FileExists(NativeArchiveFileName))
-		{
-			UE_LOG(LogImportLocalizedDialogueCommandlet, Error, TEXT("Failed to find archive '%s'."), *NativeArchiveFileName);
-			return -1;
-		}
-
-		const TSharedPtr<FJsonObject> ArchiveJsonObject = ReadJSONTextFile(NativeArchiveFileName);
-		if (!ArchiveJsonObject.IsValid())
-		{
-			UE_LOG(LogImportLocalizedDialogueCommandlet, Error, TEXT("Failed to parse archive '%s'."), *NativeArchiveFileName);
-			return -1;
-		}
-
-		FJsonInternationalizationArchiveSerializer ArchiveSerializer;
-		NativeArchive = MakeShareable(new FInternationalizationArchive());
-		if (!ArchiveSerializer.DeserializeArchive(ArchiveJsonObject.ToSharedRef(), NativeArchive.ToSharedRef()))
-		{
-			UE_LOG(LogImportLocalizedDialogueCommandlet, Error, TEXT("Failed to deserialize archive '%s'."), *NativeArchiveFileName);
-			return -1;
+			UE_LOG(LogImportLocalizedDialogueCommandlet, Error, TEXT("%s"), *LoadError.ToString());
+			return false;
 		}
 	}
-
+	
 	const FString RootAssetPath = FApp::HasGameName() ? TEXT("/Game") : TEXT("/Engine");
 	const FString RootContentDir = FApp::HasGameName() ? FPaths::GameContentDir() : FPaths::EngineContentDir();
 
@@ -272,7 +227,7 @@ int32 UImportLocalizedDialogueCommandlet::Main(const FString& Params)
 			SourceCultureImportInfo.LocalizedRootPackagePath = RootAssetPath;
 			SourceCultureImportInfo.LocalizedImportedDialoguePackagePath = SourceCultureImportInfo.LocalizedRootPackagePath / ImportedDialogueFolder;
 
-			ImportDialogueForCulture(DialogueWave, DialogueWaveSubPath, SourceCultureImportInfo, /*bImportAsSource*/true);
+			ImportDialogueForCulture(LocTextHelper, DialogueWave, DialogueWaveSubPath, SourceCultureImportInfo, /*bImportAsSource*/true);
 		}
 
 		// Iterate over each context looking for new audio to import
@@ -286,7 +241,7 @@ int32 UImportLocalizedDialogueCommandlet::Main(const FString& Params)
 				continue;
 			}
 
-			ImportDialogueForCulture(DialogueWave, DialogueWaveSubPath, CultureImportInfo, /*bImportAsSource*/false);
+			ImportDialogueForCulture(LocTextHelper, DialogueWave, DialogueWaveSubPath, CultureImportInfo, /*bImportAsSource*/false);
 		}
 	}
 
@@ -313,41 +268,11 @@ int32 UImportLocalizedDialogueCommandlet::Main(const FString& Params)
 	return 0;
 }
 
-bool UImportLocalizedDialogueCommandlet::ImportDialogueForCulture(UDialogueWave* const DialogueWave, const FString& DialogueWaveSubPath, const FCultureImportInfo& InCultureImportInfo, const bool bImportAsSource)
+bool UImportLocalizedDialogueCommandlet::ImportDialogueForCulture(FLocTextHelper& InLocTextHelper, UDialogueWave* const DialogueWave, const FString& DialogueWaveSubPath, const FCultureImportInfo& InCultureImportInfo, const bool bImportAsSource)
 {
 	UDialogueWave* LocalizedDialogueWave = nullptr;
 	FString LocalizedDialogueWaveFileName;
 
-	// Prepare the culture archive
-	TSharedPtr<FInternationalizationArchive> CultureArchive;
-	if (InCultureImportInfo.bIsNativeCulture)
-	{
-		CultureArchive = NativeArchive;
-	}
-	else
-	{
-		if (!FPaths::FileExists(InCultureImportInfo.ArchiveFileName))
-		{
-			UE_LOG(LogImportLocalizedDialogueCommandlet, Error, TEXT("Failed to find archive '%s'."), *InCultureImportInfo.ArchiveFileName);
-			return false;
-		}
-
-		const TSharedPtr<FJsonObject> ArchiveJsonObject = ReadJSONTextFile(InCultureImportInfo.ArchiveFileName);
-		if (!ArchiveJsonObject.IsValid())
-		{
-			UE_LOG(LogImportLocalizedDialogueCommandlet, Error, TEXT("Failed to parse archive '%s'."), *InCultureImportInfo.ArchiveFileName);
-			return false;
-		}
-
-		FJsonInternationalizationArchiveSerializer ArchiveSerializer;
-		CultureArchive = MakeShareable(new FInternationalizationArchive());
-		if (!ArchiveSerializer.DeserializeArchive(ArchiveJsonObject.ToSharedRef(), CultureArchive.ToSharedRef()))
-		{
-			UE_LOG(LogImportLocalizedDialogueCommandlet, Error, TEXT("Failed to deserialize archive '%s'."), *InCultureImportInfo.ArchiveFileName);
-			return false;
-		}
-	}
-
 	if (bImportAsSource)
 	{
 		LocalizedDialogueWave = DialogueWave;
@@ -387,7 +312,7 @@ bool UImportLocalizedDialogueCommandlet::ImportDialogueForCulture(UDialogueWave*
 		const FString ContextLocalizationKey = LocalizedDialogueWave->GetContextLocalizationKey(ContextMapping);
 
 		// Check that this entry exists in the manifest file, as we want to skip over dialogue that we aren't gathering
-		TSharedPtr<FManifestEntry> ContextManifestEntry = InternationalizationManifest->FindEntryByKey(FDialogueConstants::DialogueNamespace, ContextLocalizationKey, &LocalizedDialogueWave->SpokenText);
+		TSharedPtr<FManifestEntry> ContextManifestEntry = InLocTextHelper.FindSourceText(FDialogueConstants::DialogueNamespace, ContextLocalizationKey, &LocalizedDialogueWave->SpokenText);
 		if (!ContextManifestEntry.IsValid())
 		{
 			// We're skipping this context entry due to our manifest, but we don't want the sound wave it's using to be deleted
@@ -433,7 +358,7 @@ bool UImportLocalizedDialogueCommandlet::ImportDialogueForCulture(UDialogueWave*
 			const FString ContextLocalizationKey = LocalizedDialogueWave->GetContextLocalizationKey(*ContextMappingPtr);
 
 			// Find the manifest entry for our context
-			TSharedPtr<FManifestEntry> ContextManifestEntry = InternationalizationManifest->FindEntryByKey(FDialogueConstants::DialogueNamespace, ContextLocalizationKey, &LocalizedDialogueWave->SpokenText);
+			TSharedPtr<FManifestEntry> ContextManifestEntry = InLocTextHelper.FindSourceText(FDialogueConstants::DialogueNamespace, ContextLocalizationKey, &LocalizedDialogueWave->SpokenText);
 			if (!ContextManifestEntry.IsValid())
 			{
 				return FString();
@@ -441,34 +366,14 @@ bool UImportLocalizedDialogueCommandlet::ImportDialogueForCulture(UDialogueWave*
 
 			// Find the correct entry for our context
 			const FManifestContext* ContextManifestEntryContext = ContextManifestEntry->FindContextByKey(ContextLocalizationKey);
-			check(ContextManifestEntryContext); // This should never fail as we pass in the key to FindEntryByKey
+			check(ContextManifestEntryContext); // This should never fail as we pass in the key to FindSourceText
 
-			// We might have a native translation for this text, in which case we need to use the translation for that text, rather than for the source text
-			FString SourceText = LocalizedDialogueWave->SpokenText;
-			{
-				TSharedPtr<FArchiveEntry> NativeArchiveEntry = NativeArchive->FindEntryBySource(FDialogueConstants::DialogueNamespace, SourceText, ContextManifestEntryContext->KeyMetadataObj);
-				if (NativeArchiveEntry.IsValid())
-				{
-					SourceText = NativeArchiveEntry->Translation.Text;
-				}
-			}
-
-			// Find the archive entry so we can use the correct text
-			FString TranslatedText;
-			if (InCultureImportInfo.bIsNativeCulture)
-			{
-				TranslatedText = SourceText;
-			}
-			else
-			{
-				TSharedPtr<FArchiveEntry> ArchiveEntry = CultureArchive->FindEntryBySource(FDialogueConstants::DialogueNamespace, SourceText, ContextManifestEntryContext->KeyMetadataObj);
-				if (ArchiveEntry.IsValid())
-				{
-					TranslatedText = ArchiveEntry->Translation.Text;
-				}
-			}
+			// Get the localized text to export
+			FLocItem ExportedSource;
+			FLocItem ExportedTranslation;
+			InLocTextHelper.GetExportText(InCultureImportInfo.Name, FDialogueConstants::DialogueNamespace, ContextManifestEntryContext->Key, ContextManifestEntryContext->KeyMetadataObj, ELocTextExportSourceMethod::NativeText, ContextManifestEntry->Source, ExportedSource, ExportedTranslation);
 
-			return TranslatedText;
+			return ExportedTranslation.Text;
 		};
 
 		// Find the correct localized dialogue for this context
diff --git a/Engine/Source/Editor/UnrealEd/Private/Commandlets/InternationalizationConditioningCommandlet.cpp b/Engine/Source/Editor/UnrealEd/Private/Commandlets/InternationalizationConditioningCommandlet.cpp
--- a/Engine/Source/Editor/UnrealEd/Private/Commandlets/InternationalizationConditioningCommandlet.cpp
+++ b/Engine/Source/Editor/UnrealEd/Private/Commandlets/InternationalizationConditioningCommandlet.cpp
@@ -337,17 +337,12 @@ bool UInternationalizationConditioningCommandlet::ProcessManifest( const FString
 
 	// First we want to see if there is an existing manifest.  If so we will load it up and add our entries there
 	TSharedRef< FInternationalizationManifest > InternationalizationManifest = MakeShareable( new FInternationalizationManifest );
-	FJsonInternationalizationManifestSerializer ManifestSerializer;
 
 	FString ExistingManifestFileName = DestinationPath / ManifestName;
 
 	if( FPaths::FileExists(ExistingManifestFileName) )
 	{
-		TSharedPtr< FJsonObject > ExistingManifestJsonObject = ReadJSONTextFile( ExistingManifestFileName );
-		if( ExistingManifestJsonObject.IsValid() )
-		{
-			ManifestSerializer.DeserializeManifest( ExistingManifestJsonObject.ToSharedRef(), InternationalizationManifest );
-		}
+		FJsonInternationalizationManifestSerializer::DeserializeManifestFromFile( ExistingManifestFileName, InternationalizationManifest );
 	}
 
 	// Now we add our properties to the manifest. 
@@ -368,14 +363,19 @@ bool UInternationalizationConditioningCommandlet::ProcessManifest( const FString
 		}
 	}
 
-	TSharedRef<FJsonObject> FinalManifestJsonObj = MakeShareable( new FJsonObject );
-	ManifestSerializer.SerializeManifest( InternationalizationManifest, FinalManifestJsonObj );
-
 	FString DestinationManifestFileName = DestinationPath / ManifestName;
-	WriteJSONToTextFile( FinalManifestJsonObj, DestinationManifestFileName, SourceControlInfo );
+	const bool bDidWriteManifest = FLocalizedAssetSCCUtil::SaveFileWithSCC(SourceControlInfo, DestinationManifestFileName, [&InternationalizationManifest](const FString& InSaveFileName) -> bool
+	{
+		return FJsonInternationalizationManifestSerializer::SerializeManifestToFile(InternationalizationManifest, InSaveFileName);
+	});
 
-	LocPairs.Empty();
+	if (!bDidWriteManifest)
+	{
+		UE_LOG(LogInternationalizationConditioningCommandlet, Error, TEXT("Could not save file %s"), *DestinationManifestFileName);
+		return false;
+	}
 
+	LocPairs.Empty();
 
 	return true;
 }
@@ -449,7 +449,6 @@ bool UInternationalizationConditioningCommandlet::ProcessArchive( const FString&
 		}
 
 		TSharedRef< FInternationalizationArchive > InternationalizationArchive = MakeShareable( new FInternationalizationArchive );
-		FJsonInternationalizationArchiveSerializer ArchiveSerializer;
 
 		const FString DestinationArchiveFileName = DestinationPath / TargetSubfolder / ArchiveName;
 
@@ -460,11 +459,7 @@ bool UInternationalizationConditioningCommandlet::ProcessArchive( const FString&
 
 			if( FPaths::FileExists(ExistingArchiveFileName) )
 			{
-				TSharedPtr< FJsonObject > ExistingArchiveJsonObject = ReadJSONTextFile( ExistingArchiveFileName );
-				if( ExistingArchiveJsonObject.IsValid() )
-				{
-					ArchiveSerializer.DeserializeArchive( ExistingArchiveJsonObject.ToSharedRef(), InternationalizationArchive );
-				}
+				FJsonInternationalizationArchiveSerializer::DeserializeArchiveFromFile( ExistingArchiveFileName, InternationalizationArchive, nullptr, nullptr );
 			}
 		}
 
@@ -472,13 +467,14 @@ bool UInternationalizationConditioningCommandlet::ProcessArchive( const FString&
 		{
 			FLocalizationFileEntry& Prop = ArchiveProperties[PropIndex];
 			FString NewNamespace = Prop.Namespace;
+			FString NewKey = Prop.Key;
 
 			FLocItem Source( Prop.SourceText );
 			FLocItem Translation( Prop.TranslatedText );
 
-			if( !InternationalizationArchive->AddEntry( NewNamespace, Source, Translation, NULL, false ) )
+			if( !InternationalizationArchive->AddEntry( NewNamespace, NewKey, Source, Translation, NULL, false ) )
 			{
-				TSharedPtr<FArchiveEntry> ExistingConflictEntry = InternationalizationArchive->FindEntryBySource( NewNamespace, Source, NULL );
+				TSharedPtr<FArchiveEntry> ExistingConflictEntry = InternationalizationArchive->FindEntryByKey( NewNamespace, NewKey, NULL );
 
 				if( !ExistingConflictEntry.IsValid() )
 				{
@@ -502,10 +498,12 @@ bool UInternationalizationConditioningCommandlet::ProcessArchive( const FString&
 			}
 		}
 
-		TSharedRef<FJsonObject> FinalArchiveJsonObj = MakeShareable( new FJsonObject );
-		ArchiveSerializer.SerializeArchive( InternationalizationArchive, FinalArchiveJsonObj );
+		const bool bDidWriteArchive = FLocalizedAssetSCCUtil::SaveFileWithSCC(SourceControlInfo, DestinationArchiveFileName, [&InternationalizationArchive](const FString& InSaveFileName) -> bool
+		{
+			return FJsonInternationalizationArchiveSerializer::SerializeArchiveToFile(InternationalizationArchive, InSaveFileName);
+		});
 
-		if( !WriteJSONToTextFile( FinalArchiveJsonObj, DestinationArchiveFileName, SourceControlInfo ) )
+		if (!bDidWriteArchive)
 		{
 			UE_LOG(LogInternationalizationConditioningCommandlet, Error, TEXT("Could not save file %s"), *DestinationArchiveFileName);
 			return false;
