Copying //UE4/Dev-Editor to //UE4/Dev-Main (Source: //UE4/Dev-Editor @ 3082391)

diff --git a/Engine/Source/Editor/UnrealEd/Private/Commandlets/InternationalizationExportCommandlet.cpp b/Engine/Source/Editor/UnrealEd/Private/Commandlets/InternationalizationExportCommandlet.cpp
--- a/Engine/Source/Editor/UnrealEd/Private/Commandlets/InternationalizationExportCommandlet.cpp
+++ b/Engine/Source/Editor/UnrealEd/Private/Commandlets/InternationalizationExportCommandlet.cpp
@@ -2,8 +2,6 @@
 
 #include "UnrealEd.h"
 #include "Culture.h"
-#include "JsonInternationalizationManifestSerializer.h"
-#include "JsonInternationalizationArchiveSerializer.h"
 #include "PortableObjectFormatDOM.h"
 
 DEFINE_LOG_CATEGORY_STATIC(LogInternationalizationExportCommandlet, Log, All);
@@ -44,10 +42,10 @@ namespace
 			return Result;
 		};
 
-		FString EscapedNamespace = EscapeMsgCtxtParticle(Namespace);
-		FString EscapedKey = EscapeMsgCtxtParticle(Key);
+		const FString EscapedNamespace = EscapeMsgCtxtParticle(Namespace);
+		const FString EscapedKey = EscapeMsgCtxtParticle(Key);
 
-		return KeyMetaData.IsValid() ? FString::Printf(TEXT("%s,%s"), *EscapedNamespace, *EscapedKey) : EscapedNamespace;
+		return FString::Printf(TEXT("%s,%s"), *EscapedNamespace, *EscapedKey);
 	}
 
 	void ParsePOMsgCtxtForIdentity(const FString& MsgCtxt, FString& OutNamespace, FString& OutKey)
@@ -294,288 +292,154 @@ bool UInternationalizationExportCommandlet::DoExport( const FString& SourcePath,
 		bUseCultureDirectory = true;
 	}
 
+	// We may only have a single culture if using this setting
+	if (!bUseCultureDirectory && CulturesToGenerate.Num() > 1)
+	{
+		UE_LOG( LogInternationalizationExportCommandlet, Error, TEXT("bUseCultureDirectory may only be used with a single culture.") );
+		return false;
+	}
+
 	bool ShouldAddSourceLocationsAsComments = true;
 	GetBoolFromConfig(*SectionName, TEXT("ShouldAddSourceLocationsAsComments"), ShouldAddSourceLocationsAsComments, ConfigPath);
 
-	// Prepare the manifest
-	TSharedRef<FInternationalizationManifest> InternationalizationManifest = MakeShareable(new FInternationalizationManifest());
+	// Load the manifest and all archives
+	FLocTextHelper LocTextHelper(SourcePath, ManifestName, ArchiveName, NativeCultureName, CulturesToGenerate, MakeShareable(new FLocFileSCCNotifies(SourceControlInfo)));
 	{
-		const FString ManifestFileName = SourcePath / ManifestName;
-		if (!FPaths::FileExists(ManifestFileName))
-		{
-			UE_LOG(LogInternationalizationExportCommandlet, Error, TEXT("Failed to find manifest '%s'."), *ManifestFileName);
-			return false;
-		}
-
-		const TSharedPtr<FJsonObject> ManifestJsonObject = ReadJSONTextFile(ManifestFileName);
-		if (!ManifestJsonObject.IsValid())
+		FText LoadError;
+		if (!LocTextHelper.LoadAll(ELocTextHelperLoadFlags::LoadOrCreate, &LoadError))
 		{
-			UE_LOG(LogInternationalizationExportCommandlet, Error, TEXT("Failed to parse manifest '%s'."), *ManifestFileName);
-			return false;
-		}
-
-		FJsonInternationalizationManifestSerializer ManifestSerializer;
-		InternationalizationManifest = MakeShareable(new FInternationalizationManifest());
-		if (!ManifestSerializer.DeserializeManifest(ManifestJsonObject.ToSharedRef(), InternationalizationManifest))
-		{
-			UE_LOG(LogInternationalizationExportCommandlet, Error, TEXT("Failed to deserialize manifest '%s'."), *ManifestFileName);
+			UE_LOG(LogInternationalizationExportCommandlet, Error, TEXT("%s"), *LoadError.ToString());
 			return false;
 		}
 	}
 
-	// Load the native archive info
-	TSharedRef<FInternationalizationArchive> NativeArchive = MakeShareable(new FInternationalizationArchive());
+	// Process the desired cultures
+	for (const FString& CultureName : CulturesToGenerate)
 	{
-		const FString NativeCulturePath = SourcePath / NativeCultureName;
+		FPortableObjectFormatDOM NewPortableObject;
 
-		const FString NativeArchiveFileName = NativeCulturePath / ArchiveName;
-		if (!FPaths::FileExists(NativeArchiveFileName))
+		FString LocLang;
+		if( !NewPortableObject.SetLanguage( CultureName ) )
 		{
-			UE_LOG(LogInternationalizationExportCommandlet, Error, TEXT("Failed to find archive '%s'."), *NativeArchiveFileName);
-			return false;
-		}
-
-		const TSharedPtr<FJsonObject> ArchiveJsonObject = ReadJSONTextFile(NativeArchiveFileName);
-		if (!ArchiveJsonObject.IsValid())
-		{
-			UE_LOG(LogInternationalizationExportCommandlet, Error, TEXT("Failed to read archive file '%s'."), *NativeArchiveFileName);
-			return false;
+			UE_LOG( LogInternationalizationExportCommandlet, Error, TEXT("Skipping export of loc language %s because it is not recognized."), *LocLang );
+			continue;
 		}
 
-		FJsonInternationalizationArchiveSerializer ArchiveSerializer;
-		if (!ArchiveSerializer.DeserializeArchive(ArchiveJsonObject.ToSharedRef(), NativeArchive))
-		{
-			UE_LOG(LogInternationalizationExportCommandlet, Error, TEXT("Failed to deserialize archive '%s'."), *NativeArchiveFileName);
-			return false;
-		}
-	}
+		NewPortableObject.SetProjectName( FPaths::GetBaseFilename( ManifestName ) );
+		NewPortableObject.CreateNewHeader();
 
-	// Process the desired cultures
-	for(int32 Culture = 0; Culture < CulturesToGenerate.Num(); Culture++)
-	{
-		// Load the archive
-		const FString CultureName = CulturesToGenerate[Culture];
-		const FString CulturePath = SourcePath / CultureName;
-
-		FString ArchiveFileName = CulturePath / ArchiveName;
-		if( FPaths::FileExists(ArchiveFileName) )
+		// Add each manifest entry to the PO file
+		LocTextHelper.EnumerateSourceTexts([&LocTextHelper, &CultureName, &NewPortableObject](TSharedRef<FManifestEntry> InManifestEntry) -> bool
 		{
-			TSharedPtr< FJsonObject > ArchiveJsonObject = ReadJSONTextFile(ArchiveFileName);
+			// For each context, we may need to create a different or even multiple PO entries.
+			for (const FManifestContext& Context : InManifestEntry->Contexts)
+			{
+				TSharedRef<FPortableObjectEntry> PoEntry = MakeShareable(new FPortableObjectEntry());
 
-			FJsonInternationalizationArchiveSerializer ArchiveSerializer;
-			TSharedRef< FInternationalizationArchive > InternationalizationArchive = MakeShareable( new FInternationalizationArchive );
-			ArchiveSerializer.DeserializeArchive( ArchiveJsonObject.ToSharedRef(), InternationalizationArchive );
+				// Find the correct translation based upon the native source text
+				FLocItem ExportedSource;
+				FLocItem ExportedTranslation;
+				LocTextHelper.GetExportText(CultureName, InManifestEntry->Namespace, Context.Key, Context.KeyMetadataObj, ELocTextExportSourceMethod::NativeText, InManifestEntry->Source, ExportedSource, ExportedTranslation);
 
-			{
-				FPortableObjectFormatDOM NewPortableObject;
+				PoEntry->MsgId = ConditionArchiveStrForPo(ExportedSource.Text);
+				PoEntry->MsgCtxt = ConditionIdentityForPOMsgCtxt(InManifestEntry->Namespace, Context.Key, Context.KeyMetadataObj);
+				PoEntry->MsgStr.Add(ConditionArchiveStrForPo(ExportedTranslation.Text));
 
-				FString LocLang;
-				if( !NewPortableObject.SetLanguage( CultureName ) )
-				{
-					UE_LOG( LogInternationalizationExportCommandlet, Error, TEXT("Skipping export of loc language %s because it is not recognized."), *LocLang );
-					continue;
-				}
+				//@TODO: We support additional metadata entries that can be translated.  How do those fit in the PO file format?  Ex: isMature
+				const FString PORefString = ConvertSrcLocationToPORef(Context.SourceLocation);
+				PoEntry->AddReference(PORefString); // Source location.
 
-				NewPortableObject.SetProjectName( FPaths::GetBaseFilename( ManifestName ) );
-				NewPortableObject.CreateNewHeader();
+				PoEntry->AddExtractedComment(FString::Printf(TEXT("Key:\t%s"), *Context.Key)); // "Notes from Programmer" in the form of the Key.
+				PoEntry->AddExtractedComment(FString::Printf(TEXT("SourceLocation:\t%s"), *PORefString)); // "Notes from Programmer" in the form of the Source Location, since this comes in handy too and OneSky doesn't properly show references, only comments.
 
+				TArray<FString> InfoMetaDataStrings;
+				if (Context.InfoMetadataObj.IsValid())
 				{
-					for(TManifestEntryBySourceTextContainer::TConstIterator ManifestIter = InternationalizationManifest->GetEntriesBySourceTextIterator(); ManifestIter; ++ManifestIter)
+					for (auto InfoMetaDataPair : Context.InfoMetadataObj->Values)
 					{
-						// Gather relevant info from manifest entry.
-						const TSharedRef<FManifestEntry>& ManifestEntry = ManifestIter.Value();
-						const FString& Namespace = ManifestEntry->Namespace;
-						const FLocItem& Source = ManifestEntry->Source;
-
-						// For each context, we may need to create a different or even multiple PO entries.
-						for( auto ContextIter = ManifestEntry->Contexts.CreateConstIterator(); ContextIter; ++ContextIter )
-						{
-							const FManifestContext& Context = *ContextIter;
-
-							bool bHasAddedPOEntry = false;
-
-							// If we're exporting for something other than the native culture, we'll need to create PO entries for archive entries based on the native archive's translation.
-							if (CultureName != NativeCultureName)
-							{
-								// Find the native archive entry which matches the exact same namespace, source, and key metadata, if it exists.
-								TSharedPtr<FArchiveEntry> NativeArchiveEntry = NativeArchive->FindEntryBySource( Namespace, Source, Context.KeyMetadataObj );
-								if (NativeArchiveEntry.IsValid())
-								{
-									// Only need to create this PO entry if the native archive entry's translation differs from its source, in which case we need to find the our translation of the native translation.
-									if (!NativeArchiveEntry->Source.IsExactMatch(NativeArchiveEntry->Translation))
-									{
-										const TSharedPtr<FArchiveEntry> ArchiveEntry = InternationalizationArchive->FindEntryBySource( Namespace, NativeArchiveEntry->Translation, NativeArchiveEntry->KeyMetadataObj );
-										if (ArchiveEntry.IsValid())
-										{
-											const FString ConditionedArchiveSource = ConditionArchiveStrForPo(ArchiveEntry->Source.Text);
-											const FString ConditionedArchiveTranslation = ConditionArchiveStrForPo(ArchiveEntry->Translation.Text);
-
-											TSharedRef<FPortableObjectEntry> PoEntry = MakeShareable( new FPortableObjectEntry );
-											//@TODO: We support additional metadata entries that can be translated.  How do those fit in the PO file format?  Ex: isMature
-											PoEntry->MsgId = ConditionedArchiveSource;
-											PoEntry->MsgCtxt = ConditionIdentityForPOMsgCtxt(Namespace, Context.Key, Context.KeyMetadataObj);
-											PoEntry->MsgStr.Add( ConditionedArchiveTranslation );
-
-											const FString PORefString = ConvertSrcLocationToPORef( Context.SourceLocation );
-											PoEntry->AddReference( PORefString ); // Source location.
-
-											PoEntry->AddExtractedComment( FString::Printf(TEXT("Key:\t%s"), *Context.Key) ); // "Notes from Programmer" in the form of the Key.
-											PoEntry->AddExtractedComment( FString::Printf(TEXT("SourceLocation:\t%s"), *PORefString) ); // "Notes from Programmer" in the form of the Source Location, since this comes in handy too and OneSky doesn't properly show references, only comments.
-											TArray<FString> InfoMetaDataStrings;
-											if (Context.InfoMetadataObj.IsValid())
-											{
-												for (auto InfoMetaDataPair : Context.InfoMetadataObj->Values)
-												{
-													const FString KeyName = InfoMetaDataPair.Key;
-													const TSharedPtr<FLocMetadataValue> Value = InfoMetaDataPair.Value;
-													InfoMetaDataStrings.Add(FString::Printf(TEXT("InfoMetaData:\t\"%s\" : \"%s\""), *KeyName, *Value->ToString()));
-												}
-											}
-											if (InfoMetaDataStrings.Num())
-											{
-												PoEntry->AddExtractedComments(InfoMetaDataStrings);
-											}
-
-											bHasAddedPOEntry = true;
-											NewPortableObject.AddEntry( PoEntry );
-										}
-									}
-								}
-							}
-
-							// Create the typical PO entry from the archive entry which matches the exact same namespace, source, and key metadata, if it exists.
-							if (!bHasAddedPOEntry)
-							{
-								const TSharedPtr<FArchiveEntry> ArchiveEntry = InternationalizationArchive->FindEntryBySource( Namespace, Source, Context.KeyMetadataObj );
-								if( ArchiveEntry.IsValid() )
-								{
-									const FString ConditionedArchiveSource = ConditionArchiveStrForPo(ArchiveEntry->Source.Text);
-									const FString ConditionedArchiveTranslation = ConditionArchiveStrForPo(ArchiveEntry->Translation.Text);
-
-									TSharedRef<FPortableObjectEntry> PoEntry = MakeShareable( new FPortableObjectEntry );
-									//@TODO: We support additional metadata entries that can be translated.  How do those fit in the PO file format?  Ex: isMature
-									PoEntry->MsgId = ConditionedArchiveSource;
-									PoEntry->MsgCtxt = ConditionIdentityForPOMsgCtxt(Namespace, Context.Key, Context.KeyMetadataObj);
-									PoEntry->MsgStr.Add( ConditionedArchiveTranslation );
-
-									const FString PORefString = ConvertSrcLocationToPORef( Context.SourceLocation );
-									PoEntry->AddReference(PORefString); // Source location.
-
-									PoEntry->AddExtractedComment( GetConditionedKeyForExtractedComment(Context.Key) ); // "Notes from Programmer" in the form of the Key.
-
-									if (ShouldAddSourceLocationsAsComments)
-									{
-										PoEntry->AddExtractedComment(GetConditionedReferenceForExtractedComment(PORefString)); // "Notes from Programmer" in the form of the Source Location, since this comes in handy too and OneSky doesn't properly show references, only comments.
-									}
-
-									TArray<FString> InfoMetaDataStrings;
-									if (Context.InfoMetadataObj.IsValid())
-									{
-										for (auto InfoMetaDataPair : Context.InfoMetadataObj->Values)
-										{
-											const FString KeyName = InfoMetaDataPair.Key;
-											const TSharedPtr<FLocMetadataValue> Value = InfoMetaDataPair.Value;
-											InfoMetaDataStrings.Add(GetConditionedInfoMetaDataForExtractedComment(KeyName, Value->ToString()));
-										}
-									}
-									if (InfoMetaDataStrings.Num())
-									{
-										PoEntry->AddExtractedComments(InfoMetaDataStrings);
-									}
-
-									bHasAddedPOEntry = true;
-									NewPortableObject.AddEntry( PoEntry );
-								}
-							}
-						}
+						const FString KeyName = InfoMetaDataPair.Key;
+						const TSharedPtr<FLocMetadataValue> Value = InfoMetaDataPair.Value;
+						InfoMetaDataStrings.Add(FString::Printf(TEXT("InfoMetaData:\t\"%s\" : \"%s\""), *KeyName, *Value->ToString()));
 					}
 				}
-
-				// Write out the Portable Object to .po file.
+				if (InfoMetaDataStrings.Num())
 				{
-					FString OutputFileName;
-					if (bUseCultureDirectory)
-					{
-						OutputFileName = DestinationPath / CultureName / Filename;
-					}
-					else
-					{
-						OutputFileName = DestinationPath / Filename;
-					}
-
-					// Persist comments if requested.
-					if (ShouldPersistComments)
-					{
-						// Preserve comments from the specified file now, if they haven't already been.
-						if (!HasPreservedComments)
-						{
-							FPortableObjectFormatDOM ExistingPortableObject;
-							const bool HasLoadedPOFile = LoadPOFile(OutputFileName, ExistingPortableObject);
-							if (!HasLoadedPOFile)
-							{
-								return false;
-							}
-
-							PreserveExtractedCommentsForPersistence(ExistingPortableObject);
-						}
+					PoEntry->AddExtractedComments(InfoMetaDataStrings);
+				}
 
-						// Persist the comments into the new portable object we're going to be saving.
-						for (const auto& Pair : POEntryToCommentMap)
-						{
-							const TSharedPtr<FPortableObjectEntry> FoundEntry = NewPortableObject.FindEntry(Pair.Key.MsgId, Pair.Key.MsgIdPlural, Pair.Key.MsgCtxt);
-							if (FoundEntry.IsValid())
-							{
-								FoundEntry->AddExtractedComments(Pair.Value);
-							}
-						}
-					}
+				NewPortableObject.AddEntry(PoEntry);
+			}
 
-					NewPortableObject.SortEntries();
+			return true; // continue enumeration
+		}, true);
 
-					const bool DidFileExist = FPaths::FileExists(OutputFileName);
-					if (DidFileExist)
-					{
-						if( SourceControlInfo.IsValid() )
-						{
-							FText SCCErrorText;
-							if (!SourceControlInfo->CheckOutFile(OutputFileName, SCCErrorText))
-							{
-								UE_LOG(LogInternationalizationExportCommandlet, Error, TEXT("Check out of file %s failed: %s"), *OutputFileName, *SCCErrorText.ToString());
-								return false;
-							}
-						}
-					}
+		// Write out the Portable Object to .po file.
+		{
+			FString OutputFileName;
+			if (bUseCultureDirectory)
+			{
+				OutputFileName = DestinationPath / CultureName / Filename;
+			}
+			else
+			{
+				OutputFileName = DestinationPath / Filename;
+			}
 
-					//@TODO We force UTF8 at the moment but we want this to be based on the format found in the header info.
-					const FString OutputString = NewPortableObject.ToString();
-					if (!FFileHelper::SaveStringToFile(OutputString, *OutputFileName, FFileHelper::EEncodingOptions::ForceUTF8))
+			// Persist comments if requested.
+			if (ShouldPersistComments)
+			{
+				// Preserve comments from the specified file now, if they haven't already been.
+				if (!HasPreservedComments)
+				{
+					FPortableObjectFormatDOM ExistingPortableObject;
+					const bool HasLoadedPOFile = LoadPOFile(OutputFileName, ExistingPortableObject);
+					if (HasLoadedPOFile)
 					{
-						UE_LOG( LogInternationalizationExportCommandlet, Error, TEXT("Could not write file %s"), *OutputFileName );
-						return false;
+						PreserveExtractedCommentsForPersistence(ExistingPortableObject);
 					}
+				}
 
-					if (!DidFileExist)
+				// Persist the comments into the new portable object we're going to be saving.
+				for (const auto& Pair : POEntryToCommentMap)
+				{
+					const TSharedPtr<FPortableObjectEntry> FoundEntry = NewPortableObject.FindEntry(Pair.Key.MsgId, Pair.Key.MsgIdPlural, Pair.Key.MsgCtxt);
+					if (FoundEntry.IsValid())
 					{
-						// Checkout on a new file will cause it to be added
-						if( SourceControlInfo.IsValid() )
-						{
-							FText SCCErrorText;
-							if (!SourceControlInfo->CheckOutFile(OutputFileName, SCCErrorText))
-							{
-								UE_LOG(LogInternationalizationExportCommandlet, Error, TEXT("Check out of file %s failed: %s"), *OutputFileName, *SCCErrorText.ToString());
-								return false;
-							}
-						}
+						FoundEntry->AddExtractedComments(Pair.Value);
 					}
 				}
 			}
+
+			NewPortableObject.SortEntries();
+
+			const bool bPOFileSaved = FLocalizedAssetSCCUtil::SaveFileWithSCC(SourceControlInfo, OutputFileName, [&NewPortableObject](const FString& InSaveFileName) -> bool
+			{
+				//@TODO We force UTF8 at the moment but we want this to be based on the format found in the header info.
+				const FString OutputString = NewPortableObject.ToString();
+				return FFileHelper::SaveStringToFile(OutputString, *InSaveFileName, FFileHelper::EEncodingOptions::ForceUTF8);
+			});
+
+			if (!bPOFileSaved)
+			{
+				UE_LOG(LogInternationalizationExportCommandlet, Error, TEXT("Could not write file %s"), *OutputFileName);
+				return false;
+			}
 		}
 	}
+
 	return true;
 }
 
 bool UInternationalizationExportCommandlet::DoImport(const FString& SourcePath, const FString& DestinationPath, const FString& Filename)
 {
+	// Get native culture.
+	FString NativeCultureName;
+	if (!GetStringFromConfig(*SectionName, TEXT("NativeCulture"), NativeCultureName, ConfigPath))
+	{
+		UE_LOG(LogInternationalizationExportCommandlet, Error, TEXT("No native culture specified."));
+		return false;
+	}
+
 	// Get manifest name.
 	FString ManifestName;
 	if( !GetStringFromConfig( *SectionName, TEXT("ManifestName"), ManifestName, ConfigPath ) )
@@ -599,12 +463,29 @@ bool UInternationalizationExportCommandlet::DoImport(const FString& SourcePath,
 		bUseCultureDirectory = true;
 	}
 
+	// We may only have a single culture if using this setting
+	if (!bUseCultureDirectory && CulturesToGenerate.Num() > 1)
+	{
+		UE_LOG(LogInternationalizationExportCommandlet, Error, TEXT("bUseCultureDirectory may only be used with a single culture."));
+		return false;
+	}
+
+	// Load the manifest and all archives
+	FLocTextHelper LocTextHelper(DestinationPath, ManifestName, ArchiveName, NativeCultureName, CulturesToGenerate, MakeShareable(new FLocFileSCCNotifies(SourceControlInfo)));
+	{
+		FText LoadError;
+		if (!LocTextHelper.LoadAll(ELocTextHelperLoadFlags::LoadOrCreate, &LoadError))
+		{
+			UE_LOG(LogInternationalizationExportCommandlet, Error, TEXT("%s"), *LoadError.ToString());
+			return false;
+		}
+	}
+
 	// Process the desired cultures
-	for(int32 Culture = 0; Culture < CulturesToGenerate.Num(); Culture++)
+	for (const FString& CultureName : CulturesToGenerate)
 	{
 		// Load the Portable Object file if found
-		const FString CultureName = CulturesToGenerate[Culture];
-		FString POFilePath = "";
+		FString POFilePath;
 		if (bUseCultureDirectory)
 		{
 			POFilePath = SourcePath / CultureName / Filename;
@@ -626,111 +507,86 @@ bool UInternationalizationExportCommandlet::DoImport(const FString& SourcePath,
 			PreserveExtractedCommentsForPersistence(PortableObject);
 		}
 
-		if (PortableObject.GetProjectName() != ManifestName.Replace(TEXT(".manifest"), TEXT("")))
-		{
-			UE_LOG(LogInternationalizationExportCommandlet, Log, TEXT("The project name (%s) in the file (%s) did not match the target manifest project (%s)."), *POFilePath, *PortableObject.GetProjectName(), *ManifestName.Replace(TEXT(".manifest"), TEXT("")));
-		}
-
-		const FString ManifestFileName = DestinationPath / ManifestName;
-
-		TSharedPtr< FJsonObject > ManifestJsonObject = NULL;
-		ManifestJsonObject = ReadJSONTextFile( ManifestFileName );
-
-		FJsonInternationalizationManifestSerializer ManifestSerializer;
-		TSharedRef< FInternationalizationManifest > InternationalizationManifest = MakeShareable( new FInternationalizationManifest );
-		ManifestSerializer.DeserializeManifest( ManifestJsonObject.ToSharedRef(), InternationalizationManifest );
-
-		if( !FPaths::FileExists(ManifestFileName) )
-		{
-			UE_LOG( LogInternationalizationExportCommandlet, Error, TEXT("Failed to find manifest %s."), *ManifestFileName);
-			continue;
-		}
-
-		const FString DestinationCulturePath = DestinationPath / CultureName;
-		const FString ArchiveFileName = DestinationCulturePath / ArchiveName;
-
-		if( !FPaths::FileExists(ArchiveFileName) )
-		{
-			UE_LOG( LogInternationalizationExportCommandlet, Error, TEXT("Failed to find destination archive %s."), *ArchiveFileName);
-			continue;
-		}
-
-		TSharedPtr< FJsonObject > ArchiveJsonObject = ReadJSONTextFile(ArchiveFileName);
-
-		FJsonInternationalizationArchiveSerializer ArchiveSerializer;
-		TSharedRef< FInternationalizationArchive > InternationalizationArchive = MakeShareable( new FInternationalizationArchive );
-		ArchiveSerializer.DeserializeArchive( ArchiveJsonObject.ToSharedRef(), InternationalizationArchive );
-
 		bool bModifiedArchive = false;
 		{
-			for( auto EntryIter = PortableObject.GetEntriesIterator(); EntryIter; ++EntryIter )
+			for (auto EntryIter = PortableObject.GetEntriesIterator(); EntryIter; ++EntryIter)
 			{
 				auto POEntry = *EntryIter;
-				if( POEntry->MsgId.IsEmpty() || POEntry->MsgStr.Num() == 0 || POEntry->MsgStr[0].Trim().IsEmpty() )
+				if (POEntry->MsgId.IsEmpty() || POEntry->MsgStr.Num() == 0 || POEntry->MsgStr[0].Trim().IsEmpty())
 				{
 					// We ignore the header entry or entries with no translation.
 					continue;
 				}
 
 				// Some warning messages for data we don't process at the moment
-				if( !POEntry->MsgIdPlural.IsEmpty() || POEntry->MsgStr.Num() > 1 )
+				if (!POEntry->MsgIdPlural.IsEmpty() || POEntry->MsgStr.Num() > 1)
 				{
-					UE_LOG( LogInternationalizationExportCommandlet, Error, TEXT("Portable Object entry has plural form we did not process.  File: %s  MsgCtxt: %s  MsgId: %s"), *POFilePath, *POEntry->MsgCtxt, *POEntry->MsgId );
+					UE_LOG(LogInternationalizationExportCommandlet, Error, TEXT("Portable Object entry has plural form we did not process.  File: %s  MsgCtxt: %s  MsgId: %s"), *POFilePath, *POEntry->MsgCtxt, *POEntry->MsgId);
 				}
 
-				FString Key;
 				FString Namespace;
-				ParsePOMsgCtxtForIdentity(POEntry->MsgCtxt, Namespace, Key);
+				TArray<FString> Keys;
 				const FString& SourceText = ConditionPoStringForArchive(POEntry->MsgId);
 				const FString& Translation = ConditionPoStringForArchive(POEntry->MsgStr[0]);
 
-				TSharedPtr<FLocMetadataObject> KeyMetaDataObject;
-				// Get key metadata from the manifest, using the namespace and key.
-				if (!Key.IsEmpty())
 				{
-					// Find manifest entry by namespace
-					for (auto ManifestEntryIterator = InternationalizationManifest->GetEntriesByContextIdIterator(); ManifestEntryIterator; ++ManifestEntryIterator)
+					FString Key;
+					ParsePOMsgCtxtForIdentity(POEntry->MsgCtxt, Namespace, Key);
+
+					if (Key.IsEmpty())
 					{
-						const FString& ManifestEntryNamespace = ManifestEntryIterator->Key;
-						const TSharedRef<FManifestEntry>& ManifestEntry = ManifestEntryIterator->Value;
-						if (ManifestEntry->Namespace == Namespace)
-						{
-							FManifestContext* const MatchingContext = ManifestEntry->Contexts.FindByPredicate([&](FManifestContext& Context) -> bool
-								{
-									return Context.Key == Key;
-								});
-							if (MatchingContext)
-							{
-								KeyMetaDataObject = MatchingContext->KeyMetadataObj;
-							}
-						}
+						// Legacy non-keyed PO entry - need to look up and perform the import using the keys inferred from the manifest
+						LocTextHelper.FindKeysForLegacyTranslation(CultureName, Namespace, SourceText, nullptr, Keys);
+					}
+					else
+					{
+						Keys.Add(MoveTemp(Key));
 					}
 				}
 
-				//@TODO: Take into account optional entries and entries that differ by keymetadata.  Ex. Each optional entry needs a unique msgCtxt
-				const TSharedPtr< FArchiveEntry > FoundEntry = InternationalizationArchive->FindEntryBySource( Namespace, SourceText, KeyMetaDataObject );
-				if( !FoundEntry.IsValid() )
+				if (Keys.Num() == 0)
 				{
-					UE_LOG(LogInternationalizationExportCommandlet, Log, TEXT("Could not find corresponding archive entry for PO entry.  File: %s  MsgCtxt: %s  MsgId: %s"), *POFilePath, *POEntry->MsgCtxt, *POEntry->MsgId );
+					UE_LOG(LogInternationalizationExportCommandlet, Warning, TEXT("Could not import a PO entry as no key was specified, nor could a key be inferred from the manifest.  File: %s  MsgCtxt: %s  MsgId: %s"), *POFilePath, *POEntry->MsgCtxt, *POEntry->MsgId);
 					continue;
 				}
-
-				if( FoundEntry->Translation != Translation )
+				
+				for (const FString& Key : Keys)
 				{
-					FoundEntry->Translation = Translation;
-					bModifiedArchive = true;
+					// Get key metadata from the manifest, using the namespace and key.
+					const FManifestContext* ItemContext = nullptr;
+					{
+						// Find manifest entry by namespace and key
+						TSharedPtr<FManifestEntry> ManifestEntry = LocTextHelper.FindSourceText(Namespace, Key);
+						if (ManifestEntry.IsValid())
+						{
+							ItemContext = ManifestEntry->FindContextByKey(Key);
+						}
+					}
+
+					//@TODO: Take into account optional entries and entries that differ by keymetadata.  Ex. Each optional entry needs a unique msgCtxt
+
+					// Attempt to import the new text (if required)
+					const TSharedPtr<FArchiveEntry> FoundEntry = LocTextHelper.FindTranslation(CultureName, Namespace, Key, ItemContext ? ItemContext->KeyMetadataObj : nullptr);
+					if (!FoundEntry.IsValid() || !FoundEntry->Source.Text.Equals(SourceText, ESearchCase::CaseSensitive) || !FoundEntry->Translation.Text.Equals(Translation, ESearchCase::CaseSensitive))
+					{
+						if (LocTextHelper.ImportTranslation(CultureName, Namespace, Key, ItemContext ? ItemContext->KeyMetadataObj : nullptr, FLocItem(SourceText), FLocItem(Translation), ItemContext && ItemContext->bIsOptional))
+						{
+							bModifiedArchive = true;
+						}
+					}
 				}
 			}
 		}
 
-		if( bModifiedArchive )
+		if (bModifiedArchive)
 		{
-			TSharedRef<FJsonObject> FinalArchiveJsonObj = MakeShareable( new FJsonObject );
-			ArchiveSerializer.SerializeArchive( InternationalizationArchive, FinalArchiveJsonObj );
+			// Trim any dead entries out of the archive
+			LocTextHelper.TrimArchive(CultureName);
 
-			if( !WriteJSONToTextFile(FinalArchiveJsonObj, ArchiveFileName, SourceControlInfo ) )
+			FText SaveError;
+			if (!LocTextHelper.SaveForeignArchive(CultureName, &SaveError))
 			{
-				UE_LOG( LogInternationalizationExportCommandlet, Error, TEXT("Failed to write archive to %s."), *ArchiveFileName );				
+				UE_LOG(LogInternationalizationExportCommandlet, Error, TEXT("%s"), *SaveError.ToString());
 				return false;
 			}
 		}
@@ -850,7 +706,7 @@ bool UInternationalizationExportCommandlet::LoadPOFile(const FString& POFilePath
 {
 	if (!FPaths::FileExists(POFilePath))
 	{
-		UE_LOG(LogInternationalizationExportCommandlet, Warning, TEXT("Could not find file %s"), *POFilePath);
+		UE_LOG(LogInternationalizationExportCommandlet, Log, TEXT("Could not find file %s"), *POFilePath);
 		return false;
 	}
 
diff --git a/Engine/Source/Editor/UnrealEd/Private/Commandlets/StabilizeLocalizationKeys.cpp b/Engine/Source/Editor/UnrealEd/Private/Commandlets/StabilizeLocalizationKeys.cpp
--- a/Engine/Source/Editor/UnrealEd/Private/Commandlets/StabilizeLocalizationKeys.cpp
+++ b/Engine/Source/Editor/UnrealEd/Private/Commandlets/StabilizeLocalizationKeys.cpp
@@ -316,15 +316,10 @@ int32 UStabilizeLocalizationKeysCommandlet::Main(const FString& Params)
 
 			for (const FString& ManifestFilename : ManifestFilenames)
 			{
-				const TSharedPtr<FJsonObject> ManifestJsonObject = UGatherTextCommandletBase::ReadJSONTextFile(ManifestFilename);
-				if (ManifestJsonObject.IsValid())
+				TSharedRef<FInternationalizationManifest> InternationalizationManifest = MakeShareable(new FInternationalizationManifest());
+				if (FJsonInternationalizationManifestSerializer::DeserializeManifestFromFile(ManifestFilename, InternationalizationManifest))
 				{
-					FJsonInternationalizationManifestSerializer ManifestSerializer;
-					TSharedRef<FInternationalizationManifest> InternationalizationManifest = MakeShareable(new FInternationalizationManifest());
-					if (ManifestSerializer.DeserializeManifest(ManifestJsonObject.ToSharedRef(), InternationalizationManifest))
-					{
-						Manifests.Add(InternationalizationManifest);
-					}
+					Manifests.Add(InternationalizationManifest);
 				}
 			}
 
@@ -336,23 +331,18 @@ int32 UStabilizeLocalizationKeysCommandlet::Main(const FString& Params)
 
 			for (const FString& ArchiveFilename : ArchiveFilenames)
 			{
-				TSharedPtr<FJsonObject> ArchiveJsonObject = UGatherTextCommandletBase::ReadJSONTextFile(ArchiveFilename);
-				if (ArchiveJsonObject.IsValid())
+				TSharedRef<FInternationalizationArchive> InternationalizationArchive = MakeShareable(new FInternationalizationArchive());
+				if (FJsonInternationalizationArchiveSerializer::DeserializeArchiveFromFile(ArchiveFilename, InternationalizationArchive, nullptr, nullptr))
 				{
-					FJsonInternationalizationArchiveSerializer ArchiveSerializer;
-					TSharedRef<FInternationalizationArchive> InternationalizationArchive = MakeShareable(new FInternationalizationArchive());
-					if (ArchiveSerializer.DeserializeArchive(ArchiveJsonObject.ToSharedRef(), InternationalizationArchive))
+					const FString ArchivePath = FPaths::GetPath(ArchiveFilename);
+					const bool bIsNativeArchive = ArchivePath.EndsWith(NativeCulture);
+					if (bIsNativeArchive)
 					{
-						const FString ArchivePath = FPaths::GetPath(ArchiveFilename);
-						const bool bIsNativeArchive = ArchivePath.EndsWith(NativeCulture);
-						if (bIsNativeArchive)
-						{
-							NativeLocArchives.Add(FLocArchiveInfo(ArchiveFilename, InternationalizationArchive));
-						}
-						else
-						{
-							ForeignLocArchives.Add(FLocArchiveInfo(ArchiveFilename, InternationalizationArchive));
-						}
+						NativeLocArchives.Add(FLocArchiveInfo(ArchiveFilename, InternationalizationArchive));
+					}
+					else
+					{
+						ForeignLocArchives.Add(FLocArchiveInfo(ArchiveFilename, InternationalizationArchive));
 					}
 				}
 			}
@@ -361,44 +351,15 @@ int32 UStabilizeLocalizationKeysCommandlet::Main(const FString& Params)
 		// Update archives to preserve the translations from the old keys
 		for (const auto& IdPair : TextKeyMap)
 		{
-			TArray<FString> SourceStrings;
-
-			// Find all the source strings for this entry
-			for (const TSharedRef<FInternationalizationManifest>& Manifest : Manifests)
-			{
-				TSharedPtr<FManifestEntry> FoundManifestEntry = Manifest->FindEntryByKey(IdPair.Key.GetNamespace(), IdPair.Key.GetKey());
-				if (FoundManifestEntry.IsValid())
-				{
-					SourceStrings.AddUnique(FoundManifestEntry->Source.Text);
-				}
-			}
-
-			// First check to see if we have a native translation for this entry
-			const int32 NumManifestSourceStrings = SourceStrings.Num();
-			for (FLocArchiveInfo& LocArchive : NativeLocArchives)
-			{
-				// Use non-ranged iteration since we may add to the array while iterating
-				for (int32 SourceStringIndex = 0; SourceStringIndex < NumManifestSourceStrings; ++SourceStringIndex)
-				{
-					const FString& SourceString = SourceStrings[SourceStringIndex];
-					TSharedPtr<FArchiveEntry> FoundArchiveEntry = LocArchive.Archive->FindEntryBySource(IdPair.Key.GetNamespace(), FLocItem(SourceString), nullptr);
-					if (FoundArchiveEntry.IsValid() && !FoundArchiveEntry->Translation.Text.IsEmpty())
-					{
-						SourceStrings.AddUnique(FoundArchiveEntry->Translation.Text);
-					}
-				}
-			}
-
-			// Then see if we have any foreign translations, and if so, add an extra entry using the new namespace
 			for (FLocArchiveInfo& LocArchive : ForeignLocArchives)
 			{
-				for (const FString& SourceString : SourceStrings)
+				TSharedPtr<FArchiveEntry> FoundArchiveEntry = LocArchive.Archive->FindEntryByKey(IdPair.Key.GetNamespace(), IdPair.Key.GetKey(), nullptr);
+				if (FoundArchiveEntry.IsValid() && !FoundArchiveEntry->Translation.Text.IsEmpty())
 				{
-					TSharedPtr<FArchiveEntry> FoundArchiveEntry = LocArchive.Archive->FindEntryBySource(IdPair.Key.GetNamespace(), FLocItem(SourceString), nullptr);
-					if (FoundArchiveEntry.IsValid() && !FoundArchiveEntry->Translation.Text.IsEmpty())
+					LocArchive.bHasArchiveChanged = true;
+					if (!LocArchive.Archive->SetTranslation(IdPair.Value.GetNamespace(), IdPair.Value.GetKey(), FoundArchiveEntry->Source, FoundArchiveEntry->Translation, FoundArchiveEntry->KeyMetadataObj))
 					{
-						LocArchive.bHasArchiveChanged = true;
-						LocArchive.Archive->AddEntry(IdPair.Value.GetNamespace(), FoundArchiveEntry->Source, FoundArchiveEntry->Translation, FoundArchiveEntry->KeyMetadataObj, FoundArchiveEntry->bIsOptional);
+						LocArchive.Archive->AddEntry(IdPair.Value.GetNamespace(), IdPair.Value.GetKey(), FoundArchiveEntry->Source, FoundArchiveEntry->Translation, FoundArchiveEntry->KeyMetadataObj, FoundArchiveEntry->bIsOptional);
 					}
 				}
 			}
@@ -409,11 +370,12 @@ int32 UStabilizeLocalizationKeysCommandlet::Main(const FString& Params)
 		{
 			if (LocArchive.bHasArchiveChanged)
 			{
-				TSharedRef<FJsonObject> FinalArchiveJsonObj = MakeShareable(new FJsonObject());
-				FJsonInternationalizationArchiveSerializer ArchiveSerializer;
-				ArchiveSerializer.SerializeArchive(LocArchive.Archive, FinalArchiveJsonObj);
+				const bool bDidWriteArchive = FLocalizedAssetSCCUtil::SaveFileWithSCC(SourceControlInfo, LocArchive.Filename, [&LocArchive](const FString& InSaveFileName) -> bool
+				{
+					return FJsonInternationalizationArchiveSerializer::SerializeArchiveToFile(LocArchive.Archive, InSaveFileName);
+				});
 
-				if (!UGatherTextCommandletBase::WriteJSONToTextFile(FinalArchiveJsonObj, LocArchive.Filename, SourceControlInfo))
+				if (!bDidWriteArchive)
 				{
 					UE_LOG(LogStabilizeLocalizationKeys, Error, TEXT("Failed to write archive to %s."), *LocArchive.Filename);
 					return false;
diff --git a/Engine/Source/Editor/UnrealEd/Private/ComponentVisualizer.cpp b/Engine/Source/Editor/UnrealEd/Private/ComponentVisualizer.cpp
--- a/Engine/Source/Editor/UnrealEd/Private/ComponentVisualizer.cpp
+++ b/Engine/Source/Editor/UnrealEd/Private/ComponentVisualizer.cpp
@@ -2,6 +2,7 @@
 
 #include "UnrealEd.h"
 #include "ComponentVisualizer.h"
+#include "ActorEditorUtils.h"
 
 IMPLEMENT_HIT_PROXY(HComponentVisProxy, HHitProxy);
 
@@ -79,3 +80,129 @@ UActorComponent* FComponentVisualizer::GetComponentFromPropertyName(const AActor
 
 	return ResultComp;
 }
+
+void FComponentVisualizer::NotifyPropertyModified(UActorComponent* Component, UProperty* Property)
+{
+	TArray<UProperty*> Properties;
+	Properties.Add(Property);
+	NotifyPropertiesModified(Component, Properties);
+}
+
+void FComponentVisualizer::NotifyPropertiesModified(UActorComponent* Component, const TArray<UProperty*>& Properties)
+{
+	if (Component == nullptr)
+	{
+		return;
+	}
+
+	for (UProperty* Property : Properties)
+	{
+		FPropertyChangedEvent PropertyChangedEvent(Property);
+		Component->PostEditChangeProperty(PropertyChangedEvent);
+	}
+
+	AActor* Owner = Component->GetOwner();
+
+	if (Owner && FActorEditorUtils::IsAPreviewOrInactiveActor(Owner))
+	{
+		// The component belongs to the preview actor in the BP editor, so we need to propagate the property change to the archetype.
+		// Before this, we exploit the fact that the archetype and the preview actor have the old and new value of the property, respectively.
+		// So we can go through all archetype instances, and if they hold the (old) archetype value, update it to the new value.
+
+		// Get archetype
+		UActorComponent* Archetype = Cast<UActorComponent>(Component->GetArchetype());
+		check(Archetype);
+
+		// Get all archetype instances (the preview actor passed in should be amongst them)
+		TArray<UObject*> ArchetypeInstances;
+		Archetype->GetArchetypeInstances(ArchetypeInstances);
+		check(ArchetypeInstances.Contains(Component));
+
+		// Identify which of the modified properties are at their default values in the archetype instances,
+		// and thus need the new value to be propagated to them
+		struct FInstanceDefaultProperties
+		{
+			UActorComponent* ArchetypeInstance;
+			TArray<UProperty*, TInlineAllocator<8>> Properties;
+		};
+
+		TArray<FInstanceDefaultProperties> InstanceDefaultProperties;
+		InstanceDefaultProperties.Reserve(ArchetypeInstances.Num());
+
+		for (UObject* ArchetypeInstance : ArchetypeInstances)
+		{
+			UActorComponent* InstanceComp = Cast<UActorComponent>(ArchetypeInstance);
+			if (InstanceComp != Component)
+			{
+				FInstanceDefaultProperties Entry;
+				for (UProperty* Property : Properties)
+				{
+					uint8* ArchetypePtr = Property->ContainerPtrToValuePtr<uint8>(Archetype);
+					uint8* InstancePtr = Property->ContainerPtrToValuePtr<uint8>(InstanceComp);
+					if (Property->Identical(ArchetypePtr, InstancePtr))
+					{
+						Entry.Properties.Add(Property);
+					}
+				}
+
+				if (Entry.Properties.Num() > 0)
+				{
+					Entry.ArchetypeInstance = InstanceComp;
+					InstanceDefaultProperties.Add(MoveTemp(Entry));
+				}
+			}
+		}
+
+		// Propagate all modified properties to the archetype
+		Archetype->SetFlags(RF_Transactional);
+		Archetype->Modify();
+
+		if (Archetype->GetOwner())
+		{
+			Archetype->GetOwner()->Modify();
+		}
+
+		for (UProperty* Property : Properties)
+		{
+			uint8* ArchetypePtr = Property->ContainerPtrToValuePtr<uint8>(Archetype);
+			uint8* PreviewPtr = Property->ContainerPtrToValuePtr<uint8>(Component);
+			Property->CopyCompleteValue(ArchetypePtr, PreviewPtr);
+
+			FPropertyChangedEvent PropertyChangedEvent(Property);
+			Archetype->PostEditChangeProperty(PropertyChangedEvent);
+		}
+
+		// Apply changes to each archetype instance
+		for (const auto& Instance : InstanceDefaultProperties)
+		{
+			Instance.ArchetypeInstance->SetFlags(RF_Transactional);
+			Instance.ArchetypeInstance->Modify();
+
+			AActor* InstanceOwner = Instance.ArchetypeInstance->GetOwner();
+
+			if (InstanceOwner)
+			{
+				InstanceOwner->Modify();
+			}
+
+			for (UProperty* Property : Instance.Properties)
+			{
+				uint8* InstancePtr = Property->ContainerPtrToValuePtr<uint8>(Instance.ArchetypeInstance);
+				uint8* PreviewPtr = Property->ContainerPtrToValuePtr<uint8>(Component);
+				Property->CopyCompleteValue(InstancePtr, PreviewPtr);
+
+				FPropertyChangedEvent PropertyChangedEvent(Property);
+				Instance.ArchetypeInstance->PostEditChangeProperty(PropertyChangedEvent);
+			}
+
+			// Rerun construction script on instance
+			if (InstanceOwner)
+			{
+				InstanceOwner->PostEditMove(false);
+			}
+		}
+	}
+
+	// Rerun construction script on preview actor
+	Owner->PostEditMove(false);
+}
diff --git a/Engine/Source/Editor/UnrealEd/Private/ComponentVisualizerManager.cpp b/Engine/Source/Editor/UnrealEd/Private/ComponentVisualizerManager.cpp
--- a/Engine/Source/Editor/UnrealEd/Private/ComponentVisualizerManager.cpp
+++ b/Engine/Source/Editor/UnrealEd/Private/ComponentVisualizerManager.cpp
@@ -3,10 +3,10 @@
 #include "UnrealEd.h"
 #include "ComponentVisualizerManager.h"
 #include "ILevelEditor.h"
-#include "ILevelViewport.h"
+#include "SLevelViewport.h"
 
 /** Handle a click on the specified editor viewport client */
-bool FComponentVisualizerManager::HandleClick(FLevelEditorViewportClient* InViewportClient, HHitProxy* HitProxy, const FViewportClick& Click)
+bool FComponentVisualizerManager::HandleClick(FEditorViewportClient* InViewportClient, HHitProxy* HitProxy, const FViewportClick& Click)
 {
 	bool bHandled = HandleProxyForComponentVis(InViewportClient, HitProxy, Click);
 	if (bHandled && Click.GetKey() == EKeys::RightMouseButton)
@@ -14,12 +14,11 @@ bool FComponentVisualizerManager::HandleClick(FLevelEditorViewportClient* InView
 		TSharedPtr<SWidget> MenuWidget = GenerateContextMenuForComponentVis();
 		if (MenuWidget.IsValid())
 		{
-			auto ParentLevelEditor = InViewportClient->ParentLevelEditor.Pin();
-			if (ParentLevelEditor.IsValid() && ParentLevelEditor->GetActiveViewportInterface().IsValid() )
+			TSharedPtr<SEditorViewport> ViewportWidget = InViewportClient->GetEditorViewportWidget();
+			if (ViewportWidget.IsValid())
 			{
-				TSharedPtr<SWidget> MenuParent = ParentLevelEditor->GetActiveViewportInterface()->GetViewportWidget().Pin();
 				FSlateApplication::Get().PushMenu(
-					MenuParent.ToSharedRef(),
+					ViewportWidget.ToSharedRef(),
 					FWidgetPath(),
 					MenuWidget.ToSharedRef(),
 					FSlateApplication::Get().GetCursorPos(),
@@ -33,9 +32,9 @@ bool FComponentVisualizerManager::HandleClick(FLevelEditorViewportClient* InView
 	return false;
 }
 
-bool FComponentVisualizerManager::HandleProxyForComponentVis(FLevelEditorViewportClient* InViewportClient, HHitProxy* HitProxy, const FViewportClick& Click)
+bool FComponentVisualizerManager::HandleProxyForComponentVis(FEditorViewportClient* InViewportClient, HHitProxy* HitProxy, const FViewportClick& Click)
 {
-	if (HitProxy->IsA(HComponentVisProxy::StaticGetType()))
+	if (HitProxy && HitProxy->IsA(HComponentVisProxy::StaticGetType()))
 	{
 		HComponentVisProxy* VisProxy = (HComponentVisProxy*)HitProxy;
 		const UActorComponent* ClickedComponent = VisProxy->Component.Get();
@@ -147,3 +146,17 @@ TSharedPtr<SWidget> FComponentVisualizerManager::GenerateContextMenuForComponent
 
 	return TSharedPtr<SWidget>();
 }
+
+
+bool FComponentVisualizerManager::IsActive() const
+{
+	TSharedPtr<FComponentVisualizer> EditedVisualizer = EditedVisualizerPtr.Pin();
+	return EditedVisualizer.IsValid();
+}
+
+
+bool FComponentVisualizerManager::IsVisualizingArchetype() const
+{
+	TSharedPtr<FComponentVisualizer> EditedVisualizer = EditedVisualizerPtr.Pin();
+	return EditedVisualizer.IsValid() && EditedVisualizer->IsVisualizingArchetype();
+}
diff --git a/Engine/Source/Editor/UnrealEd/Private/EditorActor.cpp b/Engine/Source/Editor/UnrealEd/Private/EditorActor.cpp
--- a/Engine/Source/Editor/UnrealEd/Private/EditorActor.cpp
+++ b/Engine/Source/Editor/UnrealEd/Private/EditorActor.cpp
@@ -781,9 +781,9 @@ bool UUnrealEdEngine::edactDeleteSelected( UWorld* InWorld, bool bVerifyDeletion
 				int32 Result = FMessageDialog::Open( MessageType, ConfirmDelete );
 				if ( Result == EAppReturnType::YesAll )
 				{
-					bRequestedDeleteAllByLevel = bReferencedByLevelScript;
-					bRequestedDeleteAllByActor = bReferencedByActor;
-					bRequestedDeleteAllByLODActor = bReferencedByLODActor;
+					bRequestedDeleteAllByLevel |= bReferencedByLevelScript;
+					bRequestedDeleteAllByActor |= bReferencedByActor;
+					bRequestedDeleteAllByLODActor |= bReferencedByLODActor;
 				}
 				else if ( Result == EAppReturnType::NoAll )
 				{
diff --git a/Engine/Source/Editor/UnrealEd/Private/EditorBsp.cpp b/Engine/Source/Editor/UnrealEd/Private/EditorBsp.cpp
--- a/Engine/Source/Editor/UnrealEd/Private/EditorBsp.cpp
+++ b/Engine/Source/Editor/UnrealEd/Private/EditorBsp.cpp
@@ -316,7 +316,7 @@ void UEditorEngine::bspMergeCoplanars( UModel* Model, bool RemapLinks, bool Merg
 	{
 		for( int32 i=0; i<Model->Polys->Element.Num(); i++ )
 		{
-			if( Model->Polys->Element[i].iLink != INDEX_NONE )
+			if (Model->Polys->Element[i].iLink != INDEX_NONE)
 			{
 				CA_SUPPRESS(6001); // warning C6001: Using uninitialized memory '*Remap'.
 				Model->Polys->Element[i].iLink = Remap[Model->Polys->Element[i].iLink];
@@ -921,7 +921,7 @@ int UEditorEngine::bspNodeToFPoly
 	EdPoly->Normal      = Model->Vectors[Poly.vNormal];
 
 	EdPoly->PolyFlags 	= Poly.PolyFlags & ~(PF_EdCut | PF_EdProcessed | PF_Selected | PF_Memorized);
-	EdPoly->iLink     	= Node.iSurf;
+	EdPoly->iLinkSurf   = Node.iSurf;
 	EdPoly->Material    = Poly.Material;
 
 	EdPoly->Actor    	= Poly.Actor;
@@ -1185,8 +1185,10 @@ int UEditorEngine::bspBrushCSG
 		DestEdPoly.PolyFlags = (DestEdPoly.PolyFlags | PolyFlags) & ~NotPolyFlags;
 
 		// Set its internal link.
-		if( DestEdPoly.iLink==INDEX_NONE )
+		if (DestEdPoly.iLink == INDEX_NONE)
+		{
 			DestEdPoly.iLink = i;
+		}
 
 		// Transform it.
 		DestEdPoly.Scale( Scale );
@@ -1216,6 +1218,7 @@ int UEditorEngine::bspBrushCSG
 		{
          	FPoly EdPoly = TempModel->Polys->Element[i];
 			GModel = Brush;
+			// TODO: iLink / iLinkSurf in EdPoly / TempModel->Polys->Element[i] ?
 			BspFilterFPoly( CSGOper==CSG_Intersect ? IntersectBrushWithWorldFunc : DeIntersectBrushWithWorldFunc, Model,  &EdPoly );
 		}
 		NumPolysFromBrush = Brush->Polys->Element.Num();
@@ -1237,11 +1240,11 @@ int UEditorEngine::bspBrushCSG
 			{
 				const int32 NewSurfaceIndex = Model->Surfs.Num();
 				SurfaceIndexRemap.Add(EdPoly.iLink, NewSurfaceIndex);
-				EdPoly.iLink = TempModel->Polys->Element[i].iLink = NewSurfaceIndex;
+				EdPoly.iLinkSurf = TempModel->Polys->Element[i].iLinkSurf = NewSurfaceIndex;
 			}
 			else
 			{
-				EdPoly.iLink = TempModel->Polys->Element[i].iLink = *SurfaceIndexPtr;
+				EdPoly.iLinkSurf = TempModel->Polys->Element[i].iLinkSurf = *SurfaceIndexPtr;
 			}
 
 			// Filter brush through the world.
diff --git a/Engine/Source/Editor/UnrealEd/Private/EditorCommandLineUtils.cpp b/Engine/Source/Editor/UnrealEd/Private/EditorCommandLineUtils.cpp
--- a/Engine/Source/Editor/UnrealEd/Private/EditorCommandLineUtils.cpp
+++ b/Engine/Source/Editor/UnrealEd/Private/EditorCommandLineUtils.cpp
@@ -579,7 +579,7 @@ static void EditorCommandLineUtilsImpl::RunAssetDiffCommand(TSharedPtr<SWindow>
 	{
 		if (bWantsMerge)
 		{
-			// unlike with diff'ing, for merging we rely on asset editors for
+			// unlike with diffing, for merging we rely on asset editors for
 			// merging, and those windows get childed to the main window (so it
 			// needs to be visible)
 			//
diff --git a/Engine/Source/Editor/UnrealEd/Private/EditorEngine.cpp b/Engine/Source/Editor/UnrealEd/Private/EditorEngine.cpp
--- a/Engine/Source/Editor/UnrealEd/Private/EditorEngine.cpp
+++ b/Engine/Source/Editor/UnrealEd/Private/EditorEngine.cpp
@@ -713,7 +713,8 @@ void UEditorEngine::Init(IEngineLoop* InEngineLoop)
 			TEXT("TreeMap"),
 			TEXT("SizeMap"),
 			TEXT("MergeActors"),
-			TEXT("NiagaraEditor")
+			TEXT("NiagaraEditor"),
+			TEXT("EditorAutomation"),
 		};
 
 		FScopedSlowTask ModuleSlowTask(ARRAY_COUNT(ModuleNames));
@@ -6728,7 +6729,8 @@ void UEditorEngine::AutomationLoadMap(const FString& MapName, FString* OutError)
 	//should really be wait until the map is properly loaded....in PIE or gameplay....
 	if (bNeedPieStart)
 	{
-		ADD_LATENT_AUTOMATION_COMMAND(FWaitLatentCommand(10.f));
+		//TODO NICKD We need a better way to determine when to start the map.
+		//ADD_LATENT_AUTOMATION_COMMAND(FWaitLatentCommand(10.f));
 	}
 #endif
 	return;
diff --git a/Engine/Source/Editor/UnrealEd/Private/EditorSelectUtils.cpp b/Engine/Source/Editor/UnrealEd/Private/EditorSelectUtils.cpp
--- a/Engine/Source/Editor/UnrealEd/Private/EditorSelectUtils.cpp
+++ b/Engine/Source/Editor/UnrealEd/Private/EditorSelectUtils.cpp
@@ -406,8 +406,14 @@ void UUnrealEdEngine::NoteSelectionChange()
 void UUnrealEdEngine::SelectGroup(AGroupActor* InGroupActor, bool bForceSelection/*=false*/, bool bInSelected/*=true*/, bool bNotify/*=true*/)
 {
 	USelection* SelectedActors = GetSelectedActors();
-	SelectedActors->BeginBatchSelectOperation();
-	SelectedActors->Modify();
+	bool bStartedBatchSelect = false;
+	if(!SelectedActors->IsBatchSelecting())
+	{
+		bStartedBatchSelect = true;
+		// These will have already been called when batch selecting
+		SelectedActors->BeginBatchSelectOperation();
+		SelectedActors->Modify();
+	}
 
 	static bool bIteratingGroups = false;
 
@@ -434,7 +440,10 @@ void UUnrealEdEngine::SelectGroup(AGroupActor* InGroupActor, bool bForceSelectio
 			}
 		}
 
-		SelectedActors->EndBatchSelectOperation(bNotify);
+		if(bStartedBatchSelect)
+		{
+			SelectedActors->EndBatchSelectOperation(bNotify);
+		}
 		if (bNotify)
 		{
 			NoteSelectionChange();
diff --git a/Engine/Source/Editor/UnrealEd/Private/EditorViewportClient.cpp b/Engine/Source/Editor/UnrealEd/Private/EditorViewportClient.cpp
--- a/Engine/Source/Editor/UnrealEd/Private/EditorViewportClient.cpp
+++ b/Engine/Source/Editor/UnrealEd/Private/EditorViewportClient.cpp
@@ -231,6 +231,28 @@ int32 FEditorViewportClient::GetCameraSpeedSetting() const
 
 float const FEditorViewportClient::SafePadding = 0.075f;
 
+static int32 ViewOptionIndex = 0;
+static TArray<ELevelViewportType> ViewOptions;
+
+void InitViewOptionsArray()
+{
+	ViewOptions.Empty();
+
+	ELevelViewportType Front = ELevelViewportType::LVT_OrthoXZ;
+	ELevelViewportType Back = ELevelViewportType::LVT_OrthoNegativeXZ;
+	ELevelViewportType Top = ELevelViewportType::LVT_OrthoXY;
+	ELevelViewportType Bottom = ELevelViewportType::LVT_OrthoNegativeXY;
+	ELevelViewportType Left = ELevelViewportType::LVT_OrthoYZ;
+	ELevelViewportType Right = ELevelViewportType::LVT_OrthoNegativeYZ;
+
+	ViewOptions.Add(Front);
+	ViewOptions.Add(Back);
+	ViewOptions.Add(Top);
+	ViewOptions.Add(Bottom);
+	ViewOptions.Add(Left);
+	ViewOptions.Add(Right);
+}
+
 FEditorViewportClient::FEditorViewportClient(FEditorModeTools* InModeTools, FPreviewScene* InPreviewScene, const TWeakPtr<SEditorViewport>& InEditorViewportWidget)
 	: bAllowCinematicPreview(false)
 	, CameraSpeedSetting(4)
@@ -304,6 +326,7 @@ FEditorViewportClient::FEditorViewportClient(FEditorModeTools* InModeTools, FPre
 	, bInGameViewMode(false)
 	, bShouldInvalidateViewportWidget(false)
 {
+	InitViewOptionsArray();
 	if (ModeTools == nullptr)
 	{
 		ModeTools = new FAssetEditorModeManager();
@@ -1426,6 +1449,33 @@ void FEditorViewportClient::SetViewportType( ELevelViewportType InViewportType )
 	Invalidate();
 }
 
+void FEditorViewportClient::RotateViewportType()
+{	
+	ViewportType = ViewOptions[ViewOptionIndex];
+
+	// Changing the type may also change the active view mode; re-apply that now
+	ApplyViewMode(GetViewMode(), IsPerspective(), EngineShowFlags);
+
+	// We might have changed to an orthographic viewport; if so, update any viewport links
+	UpdateLinkedOrthoViewports(true);
+
+	Invalidate();
+
+	if (ViewOptionIndex == 5)
+	{
+		ViewOptionIndex = 0;
+	}
+	else
+	{
+		ViewOptionIndex++;
+	}
+}
+
+bool FEditorViewportClient::IsActiveViewportTypeInRotation() const
+{
+	return GetViewportType() == ViewOptions[ViewOptionIndex];
+}
+
 bool FEditorViewportClient::IsActiveViewportType(ELevelViewportType InViewportType) const
 {
 	return GetViewportType() == InViewportType;
@@ -1589,7 +1639,6 @@ void FEditorViewportClient::UpdateCameraMovement( float DeltaTime )
 		// (considering there will be further hitching trying to get back to where you were)
 		EditorMovementDeltaUpperBound = .15f;
 #endif
-
 		// Check whether the camera is being moved by the mouse or keyboard
 		bool bHasMovement = bIsTracking;
 
@@ -1619,7 +1668,6 @@ void FEditorViewportClient::UpdateCameraMovement( float DeltaTime )
 			NewViewEuler,
 			NewViewFOV );
 
-
 		// We'll zero out rotation velocity modifier after updating the simulation since these actions
 		// are always momentary -- that is, when the user mouse looks some number of pixels,
 		// we increment the impulse value right there
@@ -2351,13 +2399,6 @@ void FEditorViewportClient::StopTracking()
 {
 	if( bIsTracking )
 	{
-		//cache the axis of any widget we might be moving
-		EAxisList::Type DraggingAxis = EAxisList::None;
-		if( Widget != NULL )
-		{
-			DraggingAxis = Widget->GetCurrentAxis();
-		}
-
 		MouseDeltaTracker->EndTracking( this );
 
 		Widget->SetCurrentAxis( EAxisList::None );
diff --git a/Engine/Source/Editor/UnrealEd/Private/EditorViewportCommands.cpp b/Engine/Source/Editor/UnrealEd/Private/EditorViewportCommands.cpp
--- a/Engine/Source/Editor/UnrealEd/Private/EditorViewportCommands.cpp
+++ b/Engine/Source/Editor/UnrealEd/Private/EditorViewportCommands.cpp
@@ -15,6 +15,7 @@ void FEditorViewportCommands::RegisterCommands()
 	UI_COMMAND( Bottom, "Bottom", "Switches the viewport to bottom view", EUserInterfaceActionType::RadioButton, FInputChord( EModifierKey::Alt | EModifierKey::Shift, EKeys::J ) );
 	UI_COMMAND( Left, "Left", "Switches the viewport to left view", EUserInterfaceActionType::RadioButton, FInputChord( EModifierKey::Alt, EKeys::K ) );
 	UI_COMMAND( Right, "Right", "Switches the viewport to right view", EUserInterfaceActionType::RadioButton, FInputChord( EModifierKey::Alt | EModifierKey::Shift, EKeys::K ) );
+	UI_COMMAND( Next, "Next", "Rotate through each view options", EUserInterfaceActionType::RadioButton, FInputChord( EModifierKey::Control | EModifierKey::Shift, EKeys::SpaceBar ) );
 
 	UI_COMMAND( WireframeMode, "Brush Wireframe View Mode", "Renders the scene in brush wireframe", EUserInterfaceActionType::RadioButton, FInputChord( EModifierKey::Alt, EKeys::Two ) );
 	UI_COMMAND( UnlitMode, "Unlit View Mode", "Renders the scene with no lights", EUserInterfaceActionType::RadioButton, FInputChord( EModifierKey::Alt, EKeys::Three ) );
diff --git a/Engine/Source/Editor/UnrealEd/Private/Fbx/FbxMainImport.cpp b/Engine/Source/Editor/UnrealEd/Private/Fbx/FbxMainImport.cpp
--- a/Engine/Source/Editor/UnrealEd/Private/Fbx/FbxMainImport.cpp
+++ b/Engine/Source/Editor/UnrealEd/Private/Fbx/FbxMainImport.cpp
@@ -811,6 +811,9 @@ bool FFbxImporter::OpenFile(FString Filename, bool bParseStatistics, bool bForSc
 		}
 	}
 
+	//Cache the current file hash
+	Md5Hash = FMD5Hash::HashFile(*Filename);
+
 	CurPhase = FILEOPENED;
 	// Destroy the importer
 	//Importer->Destroy();
@@ -2116,9 +2119,9 @@ void FFbxImporter::CheckSmoothingInfo(FbxMesh* FbxMesh)
 		bFirstMesh = false;	 // don't check again
 		
 		FbxLayer* LayerSmoothing = FbxMesh->GetLayer(0, FbxLayerElement::eSmoothing);
-		if (!LayerSmoothing)
+		if (!LayerSmoothing && !GIsAutomationTesting)
 		{
-			AddTokenizedErrorMessage(FTokenizedMessage::Create(EMessageSeverity::Warning, LOCTEXT("Prompt_NoSmoothgroupForFBXScene", "Warning: No smoothing group information was found in this FBX scene.  Please make sure to enable the 'Export Smoothing Groups' option in the FBX Exporter plug-in before exporting the file.  Even for tools that don't support smoothing groups, the FBX Exporter will generate appropriate smoothing data at export-time so that correct vertex normals can be inferred while importing.")), FFbxErrors::Generic_Mesh_NoSmoothingGroup);
+			AddTokenizedErrorMessage(FTokenizedMessage::Create(EMessageSeverity::Warning, LOCTEXT("Prompt_NoSmoothgroupForFBXScene", "No smoothing group information was found in this FBX scene.  Please make sure to enable the 'Export Smoothing Groups' option in the FBX Exporter plug-in before exporting the file.  Even for tools that don't support smoothing groups, the FBX Exporter will generate appropriate smoothing data at export-time so that correct vertex normals can be inferred while importing.")), FFbxErrors::Generic_Mesh_NoSmoothingGroup);
 		}
 	}
 }
diff --git a/Engine/Source/Editor/UnrealEd/Private/Fbx/FbxMaterialImport.cpp b/Engine/Source/Editor/UnrealEd/Private/Fbx/FbxMaterialImport.cpp
--- a/Engine/Source/Editor/UnrealEd/Private/Fbx/FbxMaterialImport.cpp
+++ b/Engine/Source/Editor/UnrealEd/Private/Fbx/FbxMaterialImport.cpp
@@ -43,7 +43,7 @@ UTexture* UnFbx::FFbxImporter::ImportTexture( FbxFileTexture* FbxTexture, bool b
 	// First check if the asset already exists.
 	{
 		FString ObjectPath = BasePackageName + TEXT(".") + TextureName;
-		ExistingTexture = LoadObject<UTexture>(NULL, *ObjectPath);
+		ExistingTexture = LoadObject<UTexture>(NULL, *ObjectPath, nullptr, LOAD_Quiet | LOAD_NoWarn);
 	}
 
 
@@ -444,7 +444,7 @@ void UnFbx::FFbxImporter::CreateUnrealMaterial(FbxSurfaceMaterial& FbxMaterial,
 	}
 	else
 	{
-		UMaterialInterface* FoundMaterial = LoadObject<UMaterialInterface>(NULL, *ObjectPath.ToString(), NULL, LOAD_Quiet);
+		UMaterialInterface* FoundMaterial = LoadObject<UMaterialInterface>(NULL, *ObjectPath.ToString(), NULL, LOAD_Quiet | LOAD_NoWarn);
 		// do not override existing materials
 		if (FoundMaterial)
 		{
diff --git a/Engine/Source/Editor/UnrealEd/Private/Fbx/FbxSkeletalMeshImport.cpp b/Engine/Source/Editor/UnrealEd/Private/Fbx/FbxSkeletalMeshImport.cpp
--- a/Engine/Source/Editor/UnrealEd/Private/Fbx/FbxSkeletalMeshImport.cpp
+++ b/Engine/Source/Editor/UnrealEd/Private/Fbx/FbxSkeletalMeshImport.cpp
@@ -725,7 +725,8 @@ bool UnFbx::FFbxImporter::RetrievePoseFromBindPose(const TArray<FbxNode*>& NodeA
 						else
 						{
 							FString ErrorString = Status.GetErrorString();
-							UE_LOG(LogFbx, Warning, TEXT("Not valid bind pose for Pose (%s) - Node %s : %s"), *PoseName, *FString(Current->GetName()), *ErrorString);
+							if (!GIsAutomationTesting)
+								UE_LOG(LogFbx, Warning, TEXT("Not valid bind pose for Pose (%s) - Node %s : %s"), *PoseName, *FString(Current->GetName()), *ErrorString);
 						}
 					}
 				}
@@ -795,7 +796,8 @@ bool UnFbx::FFbxImporter::ImportBone(TArray<FbxNode*>& NodeArray, FSkeletalMeshI
 		// get bind pose
 		if(RetrievePoseFromBindPose(NodeArray, PoseArray) == false)
 		{
-			UE_LOG(LogFbx, Warning, TEXT("Getting valid bind pose failed. Try to recreate bind pose"));
+			if (!GIsAutomationTesting)
+				UE_LOG(LogFbx, Warning, TEXT("Getting valid bind pose failed. Try to recreate bind pose"));
 			// if failed, delete bind pose, and retry.
 			const int32 PoseCount = Scene->GetPoseCount();
 			for(int32 PoseIndex = PoseCount-1; PoseIndex >= 0; --PoseIndex)
@@ -813,11 +815,13 @@ bool UnFbx::FFbxImporter::ImportBone(TArray<FbxNode*>& NodeArray, FSkeletalMeshI
 			SdkManager->CreateMissingBindPoses(Scene);
 			if ( RetrievePoseFromBindPose(NodeArray, PoseArray) == false)
 			{
-				UE_LOG(LogFbx, Warning, TEXT("Recreating bind pose failed."));
+				if (!GIsAutomationTesting)
+					UE_LOG(LogFbx, Warning, TEXT("Recreating bind pose failed."));
 			}
 			else
 			{
-				UE_LOG(LogFbx, Warning, TEXT("Recreating bind pose succeeded."));
+				if (!GIsAutomationTesting)
+					UE_LOG(LogFbx, Warning, TEXT("Recreating bind pose succeeded."));
 			}
 		}
 
@@ -1296,8 +1300,9 @@ USkeletalMesh* UnFbx::FFbxImporter::ImportSkeletalMesh(UObject* InParent, TArray
 		SkeletalMesh->Rename(NULL, GetTransientPackage());
 		return nullptr;
 	}
-
-	UE_LOG(LogFbx, Warning, TEXT("Bones digested - %i  Depth of hierarchy - %i"), SkeletalMesh->RefSkeleton.GetNum(), SkeletalDepth);
+	
+	if (!GIsAutomationTesting)
+		UE_LOG(LogFbx, Warning, TEXT("Bones digested - %i  Depth of hierarchy - %i"), SkeletalMesh->RefSkeleton.GetNum(), SkeletalDepth);
 
 	// process bone influences from import data
 	ProcessImportMeshInfluences(*SkelMeshImportDataPtr);
@@ -1378,7 +1383,7 @@ USkeletalMesh* UnFbx::FFbxImporter::ImportSkeletalMesh(UObject* InParent, TArray
 
 		// Store the current file path and timestamp for re-import purposes
 		UFbxSkeletalMeshImportData* ImportData = UFbxSkeletalMeshImportData::GetImportDataForSkeletalMesh(SkeletalMesh, TemplateImportData);
-		SkeletalMesh->AssetImportData->Update(UFactory::GetCurrentFilename());
+		SkeletalMesh->AssetImportData->Update(UFactory::GetCurrentFilename(), &Md5Hash);
 
 		SkeletalMesh->CalculateInvRefMatrices();
 		SkeletalMesh->PostEditChange();
@@ -1515,7 +1520,7 @@ UObject* UnFbx::FFbxImporter::CreateAssetOfClass(UClass* AssetClass, FString Par
 	FString 	ParentPath = FString::Printf(TEXT("%s/%s"), *FPackageName::GetLongPackagePath(*ParentPackageName), *ObjectName);
 	UObject* 	Parent = CreatePackage(NULL, *ParentPath);
 	// See if an object with this name exists
-	UObject* Object = LoadObject<UObject>(Parent, *ObjectName, NULL, LOAD_None, NULL);
+	UObject* Object = LoadObject<UObject>(Parent, *ObjectName, NULL, LOAD_NoWarn | LOAD_Quiet, NULL);
 
 	// if object with same name but different class exists, warn user
 	if ((Object != NULL) && (Object->GetClass() != AssetClass))
diff --git a/Engine/Source/Editor/UnrealEd/Private/Fbx/FbxStaticMeshImport.cpp b/Engine/Source/Editor/UnrealEd/Private/Fbx/FbxStaticMeshImport.cpp
--- a/Engine/Source/Editor/UnrealEd/Private/Fbx/FbxStaticMeshImport.cpp
+++ b/Engine/Source/Editor/UnrealEd/Private/Fbx/FbxStaticMeshImport.cpp
@@ -328,7 +328,8 @@ bool UnFbx::FFbxImporter::BuildStaticMeshFromGeometry(FbxNode* Node, UStaticMesh
 
 	if (!Mesh->IsTriangleMesh())
 	{
-		UE_LOG(LogFbx, Warning, TEXT("Triangulating static mesh %s"), UTF8_TO_TCHAR(Node->GetName()));
+		if(!GIsAutomationTesting)
+			UE_LOG(LogFbx, Warning, TEXT("Triangulating static mesh %s"), UTF8_TO_TCHAR(Node->GetName()));
 
 		const bool bReplace = true;
 		FbxNodeAttribute* ConvertedNode = GeometryConverter->Triangulate(Mesh, bReplace);
@@ -1399,7 +1400,7 @@ UStaticMesh* UnFbx::FFbxImporter::ImportStaticMeshAsSingle(UObject* InParent, TA
 		}
 
 		UFbxStaticMeshImportData* ImportData = UFbxStaticMeshImportData::GetImportDataForStaticMesh(StaticMesh, TemplateImportData);
-		ImportData->Update(UFactory::GetCurrentFilename());
+		ImportData->Update(UFactory::GetCurrentFilename(), &Md5Hash);
 
 		// @todo This overrides restored values currently but we need to be able to import over the existing settings if the user chose to do so.
 		SrcModel.BuildSettings.bRemoveDegenerates = ImportOptions->bRemoveDegenerates;
diff --git a/Engine/Source/Editor/UnrealEd/Private/FbxImporter.h b/Engine/Source/Editor/UnrealEd/Private/FbxImporter.h
--- a/Engine/Source/Editor/UnrealEd/Private/FbxImporter.h
+++ b/Engine/Source/Editor/UnrealEd/Private/FbxImporter.h
@@ -726,6 +726,9 @@ public:
 	FbxScene* Scene;
 	FBXImportOptions* ImportOptions;
 
+	//We cache the hash of the file when we open the file. This is to avoid calculating the hash many time when importing many asset in one fbx file.
+	FMD5Hash Md5Hash;
+
 protected:
 	enum IMPORTPHASE
 	{
diff --git a/Engine/Source/Editor/UnrealEd/Private/FileHelpers.cpp b/Engine/Source/Editor/UnrealEd/Private/FileHelpers.cpp
--- a/Engine/Source/Editor/UnrealEd/Private/FileHelpers.cpp
+++ b/Engine/Source/Editor/UnrealEd/Private/FileHelpers.cpp
@@ -250,7 +250,7 @@ void FEditorFileUtils::RegisterLevelFilename(UObject* Object, const FString& New
 
 ///////////////////////////////////////////////////////////////////////////////
 
-static FString GetFilename(const FName& PackageName)
+FString FEditorFileUtils::GetFilename(const FName& PackageName)
 {
 	// First see if it is an in-memory package that already has an associated filename
 	const FString PackageNameString = PackageName.ToString();
@@ -282,7 +282,7 @@ static FString GetFilename(const FName& PackageName)
 	return *Result;
 }
 
-static FString GetFilename(UObject* LevelObject)
+FString FEditorFileUtils::GetFilename(UObject* LevelObject)
 {
 	return GetFilename( LevelObject->GetOutermost()->GetFName() );
 }
@@ -2232,6 +2232,14 @@ void FEditorFileUtils::LoadMap(const FString& InFilename, bool LoadAsTemplate, b
 	// If there are any old mirrored brushes in the map with inverted polys, fix them here
 	GUnrealEd->FixAnyInvertedBrushes(World);
 
+	// Rebuild BSP if the loading process flagged it as not up-to-date
+	TArray< TWeakObjectPtr< ULevel > > LevelsToRebuild;
+	ABrush::NeedsRebuild(&LevelsToRebuild);
+	if (LevelsToRebuild.Num() > 0)
+	{
+		GUnrealEd->RebuildAlteredBSP();
+	}
+
 	// Fire delegate when a new map is opened, with name of map
 	FEditorDelegates::OnMapOpened.Broadcast(InFilename, LoadAsTemplate);
 }
diff --git a/Engine/Source/Editor/UnrealEd/Private/Kismet2/DebuggerCommands.cpp b/Engine/Source/Editor/UnrealEd/Private/Kismet2/DebuggerCommands.cpp
--- a/Engine/Source/Editor/UnrealEd/Private/Kismet2/DebuggerCommands.cpp
+++ b/Engine/Source/Editor/UnrealEd/Private/Kismet2/DebuggerCommands.cpp
@@ -1229,7 +1229,7 @@ void RecordLastExecutedPlayMode()
 			PlayModeString = TEXT("<UNKNOWN>");
 		}
 
-		FEngineAnalytics::GetProvider().RecordEvent(TEXT("Editor.Usage.PIE"), TEXT("PlayLocation"), PlayModeString);
+		FEngineAnalytics::GetProvider().RecordEvent(TEXT("Editor.Usage.PIE"), TEXT("PlayLocation"), PlayLocationString);
 		FEngineAnalytics::GetProvider().RecordEvent(TEXT("Editor.Usage.PIE"), TEXT("PlayMode"), PlayModeString);
 	}	
 }
diff --git a/Engine/Source/Editor/UnrealEd/Private/Kismet2/Kismet2.cpp b/Engine/Source/Editor/UnrealEd/Private/Kismet2/Kismet2.cpp
--- a/Engine/Source/Editor/UnrealEd/Private/Kismet2/Kismet2.cpp
+++ b/Engine/Source/Editor/UnrealEd/Private/Kismet2/Kismet2.cpp
@@ -492,16 +492,41 @@ UBlueprint* FKismetEditorUtilities::CreateBlueprint(UClass* ParentClass, UObject
 			TArray<FName> AutoSpawnedEventNames;
 			int32 NodePositionY = 0;
 
-			// Spawn any defined auto generated default events for the class
-			for (TMultiMap<void*, FDefaultEventNodeData>::TIterator DataIt(AutoGeneratedDefaultEventsMap); DataIt; ++DataIt)
+			// Spawn any defined auto generated default events for the class.  Only do this for the most senior class specified, so
+			// that subclasses may have an entirely different set of default nodes if they wish.
+			UClass* DefaultNodesClass = NewBP->GeneratedClass;
+			while ( DefaultNodesClass )
 			{
-				FDefaultEventNodeData Data = DataIt.Value();
-				if (NewBP->GeneratedClass->IsChildOf(Data.TargetClass))
+				bool bFoundDefaultNodes = false;
+				for ( TMultiMap<void*, FDefaultEventNodeData>::TIterator DataIt(AutoGeneratedDefaultEventsMap); DataIt; ++DataIt )
+				{
+					FDefaultEventNodeData Data = DataIt.Value();
+					if ( DefaultNodesClass == Data.TargetClass )
+					{
+						bFoundDefaultNodes = true;
+						FKismetEditorUtilities::AddDefaultEventNode(NewBP, NewBP->UbergraphPages[0], Data.EventName, Data.TargetClass, NodePositionY);
+					}
+				}
+
+				if ( bFoundDefaultNodes )
 				{
-					FKismetEditorUtilities::AddDefaultEventNode(NewBP, NewBP->UbergraphPages[0], Data.EventName, Data.TargetClass, NodePositionY);
+					break;
 				}
+
+				DefaultNodesClass = DefaultNodesClass->GetSuperClass();
 			}
 
+
+			//// Spawn any defined auto generated default events for the class
+			//for (TMultiMap<void*, FDefaultEventNodeData>::TIterator DataIt(AutoGeneratedDefaultEventsMap); DataIt; ++DataIt)
+			//{
+			//	FDefaultEventNodeData Data = DataIt.Value();
+			//	if (NewBP->GeneratedClass->IsChildOf(Data.TargetClass))
+			//	{
+			//		FKismetEditorUtilities::AddDefaultEventNode(NewBP, NewBP->UbergraphPages[0], Data.EventName, Data.TargetClass, NodePositionY);
+			//	}
+			//}
+
 			// Give anyone who wants to do more advanced BP modification post-creation a chance to do so.
 			// Anim Blueprints, for example, adds a non-event node to the main ubergraph.
 			for (TMultiMap<void*, FKismetEditorUtilities::FOnBlueprintCreatedData>::TIterator DataIt(OnBlueprintCreatedCallbacks); DataIt; ++DataIt)
@@ -1731,6 +1756,14 @@ UBlueprint* FKismetEditorUtilities::CreateBlueprintUsingAsset(UObject* Asset, bo
 	return NULL;
 }
 
+void FKismetEditorUtilities::AddToSelection(const class UEdGraph* Graph, UEdGraphNode* InNode)
+{
+	TSharedPtr<class IBlueprintEditor> BlueprintEditor = GetIBlueprintEditorForObject(Graph, false);
+	if (BlueprintEditor.IsValid())
+	{
+		BlueprintEditor->AddToSelection(InNode);
+	}
+}
 
 TSharedPtr<class IBlueprintEditor> FKismetEditorUtilities::GetIBlueprintEditorForObject( const UObject* ObjectToFocusOn, bool bOpenEditor )
 {
diff --git a/Engine/Source/Editor/UnrealEd/Private/LevelEditorViewport.cpp b/Engine/Source/Editor/UnrealEd/Private/LevelEditorViewport.cpp
--- a/Engine/Source/Editor/UnrealEd/Private/LevelEditorViewport.cpp
+++ b/Engine/Source/Editor/UnrealEd/Private/LevelEditorViewport.cpp
@@ -39,6 +39,7 @@
 #include "DragTool_BoxSelect.h"
 #include "DragTool_FrustumSelect.h"
 #include "DragTool_Measure.h"
+#include "DragTool_ViewportChange.h"
 #include "LevelEditorActions.h"
 #include "BrushBuilderDragDropOp.h"
 #include "AssetRegistryModule.h"
@@ -1668,6 +1669,11 @@ ELevelViewportType FLevelEditorViewportClient::GetViewportType() const
 	}
 }
 
+void FLevelEditorViewportClient::SetViewportTypeFromTool( ELevelViewportType InViewportType )
+{
+	SetViewportType(InViewportType);
+}
+
 void FLevelEditorViewportClient::SetViewportType( ELevelViewportType InViewportType )
 {
 	if (InViewportType != LVT_Perspective)
@@ -1679,6 +1685,14 @@ void FLevelEditorViewportClient::SetViewportType( ELevelViewportType InViewportT
 	FEditorViewportClient::SetViewportType(InViewportType);
 }
 
+void FLevelEditorViewportClient::RotateViewportType()
+{
+	SetActorLock(nullptr);
+	UpdateViewForLockedActor();
+
+	FEditorViewportClient::RotateViewportType();
+}
+
 void FLevelEditorViewportClient::OverridePostProcessSettings( FSceneView& View )
 {
 	const UCameraComponent* CameraComponent = GetCameraComponentForView();
@@ -1841,6 +1855,11 @@ void FLevelEditorViewportClient::ProcessClick(FSceneView& View, HHitProxy* HitPr
 	static FName ProcessClickTrace = FName(TEXT("ProcessClickTrace"));
 
 	const FViewportClick Click(&View,this,Key,Event,HitX,HitY);
+	if (Click.GetKey() == EKeys::MiddleMouseButton && !Click.IsAltDown() && !Click.IsShiftDown())
+	{
+		ClickHandlers::ClickViewport(this, Click);
+		return;
+	}
 	if (!ModeTools->HandleClick(this, HitProxy,Click))
 	{
 		if (HitProxy == NULL)
@@ -1918,7 +1937,7 @@ void FLevelEditorViewportClient::ProcessClick(FSceneView& View, HHitProxy* HitPr
 			}
 
 			// We clicked an actor, allow the pivot to reposition itself.
-//			GUnrealEd->SetPivotMovedIndependently(false);
+			// GUnrealEd->SetPivotMovedIndependently(false);
 		}
 		else if (HitProxy->IsA(HInstancedStaticMeshInstance::StaticGetType()))
 		{
@@ -2368,6 +2387,9 @@ TSharedPtr<FDragTool> FLevelEditorViewportClient::MakeDragTool( EDragTool::Type
 	case EDragTool::Measure:
 		DragTool = MakeShareable( new FDragTool_Measure(this) );
 		break;
+	case EDragTool::ViewportChange:
+		DragTool = MakeShareable( new FDragTool_ViewportChange(this) );
+		break;
 	};
 
 	return DragTool;
@@ -2408,7 +2430,7 @@ void FLevelEditorViewportClient::SetCurrentViewport()
 
 void FLevelEditorViewportClient::SetLastKeyViewport()
 {
-	// Store a reference to the last viewport that received a keypress.
+	// Store a reference to the last viewport that received a key press.
 	GLastKeyLevelEditingViewportClient = this;
 
 	if (GCurrentLevelEditingViewportClient != this)
@@ -2966,6 +2988,16 @@ void FLevelEditorViewportClient::RedrawAllViewportsIntoThisScene()
 	GEditor->RedrawLevelEditingViewports();
 }
 
+FWidget::EWidgetMode FLevelEditorViewportClient::GetWidgetMode() const
+{
+	if (GUnrealEd->ComponentVisManager.IsActive() && GUnrealEd->ComponentVisManager.IsVisualizingArchetype())
+	{
+		return FWidget::WM_None;
+	}
+
+	return FEditorViewportClient::GetWidgetMode();
+}
+
 FVector FLevelEditorViewportClient::GetWidgetLocation() const
 {
 	FVector ComponentVisWidgetLocation;
@@ -3964,32 +3996,33 @@ void FLevelEditorViewportClient::DrawBrushDetails(const FSceneView* View, FPrimi
 		const float TextureSizeX = VertexTexture->GetSizeX() * 0.170f;
 		const float TextureSizeY = VertexTexture->GetSizeY() * 0.170f;
 
-		for (FSelectionIterator It(*ModeTools->GetSelectedActors()); It; ++It)
+		USelection* Selection = ModeTools->GetSelectedActors();
+		if(Selection->IsClassSelected(ABrush::StaticClass()))
 		{
-			AActor* SelectedActor = static_cast<AActor*>(*It);
-			checkSlow(SelectedActor->IsA(AActor::StaticClass()));
-
-			ABrush* Brush = Cast< ABrush >(SelectedActor);
-			if (Brush && Brush->Brush && !FActorEditorUtils::IsABuilderBrush(Brush))
+			for (FSelectionIterator It(*ModeTools->GetSelectedActors()); It; ++It)
 			{
-				for (int32 p = 0; p < Brush->Brush->Polys->Element.Num(); ++p)
+				ABrush* Brush = Cast<ABrush>(*It);
+				if(Brush && Brush->Brush && !FActorEditorUtils::IsABuilderBrush(Brush))
 				{
-					FTransform BrushTransform = Brush->ActorToWorld();
-
-					FPoly* poly = &Brush->Brush->Polys->Element[p];
-					for (int32 VertexIndex = 0; VertexIndex < poly->Vertices.Num(); ++VertexIndex)
+					for(int32 p = 0; p < Brush->Brush->Polys->Element.Num(); ++p)
 					{
-						const FVector& PolyVertex = poly->Vertices[VertexIndex];
-						const FVector WorldLocation = BrushTransform.TransformPosition(PolyVertex);
+						FTransform BrushTransform = Brush->ActorToWorld();
 
-						const float Scale = View->WorldToScreen(WorldLocation).W * (4.0f / View->ViewRect.Width() / View->ViewMatrices.ProjMatrix.M[0][0]);
+						FPoly* poly = &Brush->Brush->Polys->Element[p];
+						for(int32 VertexIndex = 0; VertexIndex < poly->Vertices.Num(); ++VertexIndex)
+						{
+							const FVector& PolyVertex = poly->Vertices[VertexIndex];
+							const FVector WorldLocation = BrushTransform.TransformPosition(PolyVertex);
 
-						const FColor Color(Brush->GetWireColor());
-						PDI->SetHitProxy(new HBSPBrushVert(Brush, &poly->Vertices[VertexIndex]));
+							const float Scale = View->WorldToScreen(WorldLocation).W * (4.0f / View->ViewRect.Width() / View->ViewMatrices.ProjMatrix.M[0][0]);
 
-						PDI->DrawSprite(WorldLocation, TextureSizeX * Scale, TextureSizeY * Scale, VertexTexture->Resource, Color, SDPG_World, 0.0f, 0.0f, 0.0f, 0.0f, SE_BLEND_Masked);
+							const FColor Color(Brush->GetWireColor());
+							PDI->SetHitProxy(new HBSPBrushVert(Brush, &poly->Vertices[VertexIndex]));
 
-						PDI->SetHitProxy(NULL);
+							PDI->DrawSprite(WorldLocation, TextureSizeX * Scale, TextureSizeY * Scale, VertexTexture->Resource, Color, SDPG_World, 0.0f, 0.0f, 0.0f, 0.0f, SE_BLEND_Masked);
+
+							PDI->SetHitProxy(NULL);
+						}
 					}
 				}
 			}
diff --git a/Engine/Source/Editor/UnrealEd/Private/LevelViewportClickHandlers.cpp b/Engine/Source/Editor/UnrealEd/Private/LevelViewportClickHandlers.cpp
--- a/Engine/Source/Editor/UnrealEd/Private/LevelViewportClickHandlers.cpp
+++ b/Engine/Source/Editor/UnrealEd/Private/LevelViewportClickHandlers.cpp
@@ -33,6 +33,14 @@ namespace ClickHandlers
 		}
 	}	
 
+	static void PrivateSummonViewportMenu( FLevelEditorViewportClient* ViewportClient )
+	{
+		if (ViewportClient->ParentLevelEditor.IsValid())
+		{
+			ViewportClient->ParentLevelEditor.Pin()->SummonLevelViewportViewOptionMenu(LVT_Perspective);
+		}
+	}
+
 	/**
  	 * Creates an actor of the specified type, trying first to find an actor factory,
 	 * falling back to "ACTOR ADD" exec and SpawnActor if no factory is found.
@@ -91,6 +99,21 @@ namespace ClickHandlers
 		}
 	}
 
+	bool ClickViewport(FLevelEditorViewportClient* ViewportClient, const FViewportClick& Click)
+	{
+		if (Click.GetKey() == EKeys::MiddleMouseButton && Click.GetEvent() == EInputEvent::IE_DoubleClick)
+		{
+			ViewportClient->SetViewportTypeFromTool(LVT_Perspective);
+			return true;
+		}
+		else if (Click.GetKey() == EKeys::MiddleMouseButton && Click.IsControlDown())
+		{
+			PrivateSummonViewportMenu(ViewportClient);
+			return true;
+		}
+		return false;
+	}
+
 	bool ClickActor(FLevelEditorViewportClient* ViewportClient,AActor* Actor,const FViewportClick& Click,bool bAllowSelectionChange)
 	{
 		// Pivot snapping
diff --git a/Engine/Source/Editor/UnrealEd/Private/LevelViewportClickHandlers.h b/Engine/Source/Editor/UnrealEd/Private/LevelViewportClickHandlers.h
--- a/Engine/Source/Editor/UnrealEd/Private/LevelViewportClickHandlers.h
+++ b/Engine/Source/Editor/UnrealEd/Private/LevelViewportClickHandlers.h
@@ -35,6 +35,8 @@ struct HLevelSocketProxy : public HHitProxy
 
 namespace ClickHandlers
 {
+	bool ClickViewport(FLevelEditorViewportClient* ViewportClient, const FViewportClick& Click);
+
 	bool ClickActor(FLevelEditorViewportClient* ViewportClient,AActor* Actor,const FViewportClick& Click,bool bAllowSelectionChange);
 
 	bool ClickComponent(FLevelEditorViewportClient* ViewportClient, HActor* ActorHitProxy, const FViewportClick& Click);
diff --git a/Engine/Source/Editor/UnrealEd/Private/MaterialGraphSchema.cpp b/Engine/Source/Editor/UnrealEd/Private/MaterialGraphSchema.cpp
--- a/Engine/Source/Editor/UnrealEd/Private/MaterialGraphSchema.cpp
+++ b/Engine/Source/Editor/UnrealEd/Private/MaterialGraphSchema.cpp
@@ -165,6 +165,35 @@ UMaterialGraphSchema::UMaterialGraphSchema(const FObjectInitializer& ObjectIniti
 	AlphaPinColor = FLinearColor(0.5f, 0.5f, 0.5f);
 }
 
+void UMaterialGraphSchema::SelectAllInputNodes(UEdGraph* Graph, UEdGraphPin* InGraphPin)
+{
+	TArray<UEdGraphPin*> AllPins = InGraphPin->LinkedTo;
+
+	if (AllPins.Num() == 0)
+	{
+		return;
+	}
+
+	for (UEdGraphPin* Pin : AllPins)
+	{
+		UMaterialGraphNode* MaterialNode = Cast<UMaterialGraphNode>(Pin->GetOwningNode());
+		FMaterialEditorUtilities::AddToSelection(Graph, MaterialNode->MaterialExpression);
+
+		TArray<UEdGraphPin*> LinkedPins = Pin->GetOwningNode()->GetAllPins();
+		for (UEdGraphPin* InputPin : LinkedPins)
+		{
+			if (InputPin->Direction == EEdGraphPinDirection::EGPD_Output)
+			{
+				continue;
+			}
+			else
+			{
+				SelectAllInputNodes(Graph, InputPin);
+			}
+		}
+	}
+}
+
 void UMaterialGraphSchema::GetBreakLinkToSubMenuActions( class FMenuBuilder& MenuBuilder, UEdGraphPin* InGraphPin )
 {
 	// Make sure we have a unique name for every entry in the list
@@ -187,7 +216,7 @@ void UMaterialGraphSchema::GetBreakLinkToSubMenuActions( class FMenuBuilder& Men
 			Title = FText::Format( LOCTEXT("BreakDescPin", "{NodeTitle} ({PinName})"), Args );
 		}
 
-		uint32 &Count = LinkTitleCount.FindOrAdd( TitleString );
+		uint32 &Count = LinkTitleCount.FindOrAdd(TitleString);
 
 		FText Description;
 		FFormatNamedArguments Args;
@@ -313,7 +342,7 @@ void UMaterialGraphSchema::GetGraphContextActions(FGraphContextMenuBuilder& Cont
 {
 	const UMaterialGraph* MaterialGraph = CastChecked<UMaterialGraph>(ContextMenuBuilder.CurrentGraph);
 
-	// Run thru all nodes and add any menu items they want to add
+	// Run through all nodes and add any menu items they want to add
 	Super::GetGraphContextActions(ContextMenuBuilder);
 
 	// Get the Context Actions from Material Editor Module
@@ -344,6 +373,13 @@ void UMaterialGraphSchema::GetContextMenuActions(const UEdGraph* CurrentGraph, c
 			// Only display the 'Break Link' option if there is a link to break!
 			if (InGraphPin->LinkedTo.Num() > 0)
 			{
+				MenuBuilder->AddMenuEntry(
+					LOCTEXT("SelectLinkedNodes", "Select Linked Nodes"),
+					LOCTEXT("SelectLinkedNodesTooltip", "Adds all input Nodes linked to this Pin to selection"),
+					FSlateIcon(),
+					FUIAction(FExecuteAction::CreateUObject((UMaterialGraphSchema*const)this, &UMaterialGraphSchema::SelectAllInputNodes, const_cast<UEdGraph*>(CurrentGraph), const_cast<UEdGraphPin*>(InGraphPin)))
+					);
+
 				MenuBuilder->AddMenuEntry(FGraphEditorCommands::Get().BreakPinLinks);
 
 				// add sub menu for break link to
diff --git a/Engine/Source/Editor/UnrealEd/Private/MouseDeltaTracker.cpp b/Engine/Source/Editor/UnrealEd/Private/MouseDeltaTracker.cpp
--- a/Engine/Source/Editor/UnrealEd/Private/MouseDeltaTracker.cpp
+++ b/Engine/Source/Editor/UnrealEd/Private/MouseDeltaTracker.cpp
@@ -185,7 +185,6 @@ void FMouseDeltaTracker::DetermineCurrentAxis(FEditorViewportClient* InViewportC
 	}
 }
 
-
 /**
  * Begin tracking at the specified location for the specified viewport.
  */
@@ -324,7 +323,7 @@ bool FMouseDeltaTracker::EndTracking(FEditorViewportClient* InViewportClient)
 		}
 	}
 
-	// Do not fade snapping indicators over time if this viewport is not realtime
+	// Do not fade snapping indicators over time if this viewport is not real time
 	bool bClearImmediately = !InViewportClient->IsRealtime();
 	FSnappingUtils::ClearSnappingHelpers( bClearImmediately );
 	return true;
@@ -350,22 +349,29 @@ void FMouseDeltaTracker::ConditionalBeginUsingDragTool( FEditorViewportClient* I
 		{
 
 			// Create a drag tool.
-			if (InViewportClient->IsOrtho())
+			if (!(bAltDown + bShiftDown) && bControlDown && MiddleMouseButtonDown && !LeftMouseButtonDown && !RightMouseButtonDown)
 			{
-				if (LeftMouseButtonDown)
-				{
-					DragTool = InViewportClient->MakeDragTool(EDragTool::BoxSelect);
-				}
-				else if (!(bControlDown + bAltDown + bShiftDown) && MiddleMouseButtonDown)
-				{
-					DragTool = InViewportClient->MakeDragTool(EDragTool::Measure);
-				}
+				DragTool = InViewportClient->MakeDragTool(EDragTool::ViewportChange);
 			}
 			else
 			{
-				if (LeftMouseButtonDown && bControlDown && bAltDown)
+				if (InViewportClient->IsOrtho())
+				{
+					if (LeftMouseButtonDown)
+					{
+						DragTool = InViewportClient->MakeDragTool(EDragTool::BoxSelect);
+					}
+					else if (!(bControlDown + bAltDown + bShiftDown) && MiddleMouseButtonDown)
+					{
+						DragTool = InViewportClient->MakeDragTool(EDragTool::Measure);
+					}
+				}
+				else
 				{
-					DragTool = InViewportClient->MakeDragTool(EDragTool::FrustumSelect);
+					if (LeftMouseButtonDown && bControlDown && bAltDown)
+					{
+						DragTool = InViewportClient->MakeDragTool(EDragTool::FrustumSelect);
+					}
 				}
 			}
 
@@ -553,7 +559,6 @@ void FMouseDeltaTracker::AddDelta(FEditorViewportClient* InViewportClient, FKey
 				break;
 		}
 	}
-
 }
 
 /**
@@ -582,7 +587,6 @@ const FVector FMouseDeltaTracker::GetDeltaSnapped() const
 	return SnappedDelta;
 }
 
-
 /**
 * Returns the absolute delta since dragging started.
 */
diff --git a/Engine/Source/Editor/UnrealEd/Private/ObjectTools.cpp b/Engine/Source/Editor/UnrealEd/Private/ObjectTools.cpp
--- a/Engine/Source/Editor/UnrealEd/Private/ObjectTools.cpp
+++ b/Engine/Source/Editor/UnrealEd/Private/ObjectTools.cpp
@@ -1422,7 +1422,7 @@ namespace ObjectTools
 	{
 		TArray<UPackage*> PackagesToDelete = PotentialPackagesToDelete;
 		TArray<FString> PackageFilesToDelete;
-		TArray<FSourceControlStatePtr> PackageSCCStates;
+		TArray<TSharedRef<ISourceControlState, ESPMode::ThreadSafe> > PackageSCCStates;
 		ISourceControlProvider& SourceControlProvider = ISourceControlModule::Get().GetProvider();
 
 		GWarn->BeginSlowTask( NSLOCTEXT("ObjectTools", "OldPackageCleanupSlowTask", "Cleaning Up Old Assets"), true );
@@ -1472,13 +1472,15 @@ namespace ObjectTools
 
 				PackageFilesToDelete.Add(PackageFilename);
 				Cast<UPackage>(Package)->SetDirtyFlag(false);
-				if ( ISourceControlModule::Get().IsEnabled() )
-				{
-					PackageSCCStates.Add( SourceControlProvider.GetState(PackageFilename, EStateCacheUsage::ForceUpdate) );
-				}
 			}
 		}
 
+		// Get the current source control states of all the package files we're deleting at once.
+		if (ISourceControlModule::Get().IsEnabled())
+		{
+			SourceControlProvider.GetState(PackageFilesToDelete, PackageSCCStates, EStateCacheUsage::ForceUpdate);
+		}
+
 		GWarn->EndSlowTask();
 
 		if (GUnrealEd)
@@ -1497,24 +1499,27 @@ namespace ObjectTools
 		// Now delete all packages that have become empty
 		bool bMakeWritable = false;
 		bool bSilent = false;
+		TArray<FString> SCCFilesToRevert;
+		TArray<FString> SCCFilesToDelete;
+
 		for ( int32 PackageFileIdx = 0; PackageFileIdx < PackageFilesToDelete.Num(); ++PackageFileIdx )
 		{
 			const FString& PackageFilename = PackageFilesToDelete[PackageFileIdx];
 			if ( ISourceControlModule::Get().IsEnabled() )
 			{
-				const FSourceControlStatePtr SourceControlState = PackageSCCStates[PackageFileIdx];
-				const bool bInDepot = (SourceControlState.IsValid() && SourceControlState->IsSourceControlled());
+				const FSourceControlStateRef SourceControlState = PackageSCCStates[PackageFileIdx];
+				const bool bInDepot = SourceControlState->IsSourceControlled();
 				if ( bInDepot )
 				{
 					// The file is managed by source control. Open it for delete.
-					TArray<FString> DeleteFilenames;
-					DeleteFilenames.Add(FPaths::ConvertRelativePathToFull(PackageFilename));
+					FString FullPackageFilename = FPaths::ConvertRelativePathToFull(PackageFilename);
 
 					// Revert the file if it is checked out
 					const bool bIsAdded = SourceControlState->IsAdded();
 					if ( SourceControlState->IsCheckedOut() || bIsAdded || SourceControlState->IsDeleted() )
 					{
-						SourceControlProvider.Execute(ISourceControlOperation::Create<FRevert>(), DeleteFilenames);
+						// Batch the revert operation so that we only make one request to the source control module.
+						SCCFilesToRevert.Add(FullPackageFilename);
 					}
 
 					if ( bIsAdded )
@@ -1524,10 +1529,14 @@ namespace ObjectTools
 					}
 					else
 					{
-						// Open the file for delete
-						if ( SourceControlProvider.Execute(ISourceControlOperation::Create<FDelete>(), DeleteFilenames) == ECommandResult::Failed )
+						// Batch this file for deletion so that we only send one deletion request to the source control module.
+						if (SourceControlState->CanDelete())
 						{
-							UE_LOG(LogObjectTools, Warning, TEXT("SCC failed to open '%s' for delete while saving an empty package."), *PackageFilename);
+							SCCFilesToDelete.Add(FullPackageFilename);
+						}
+						else
+						{
+							UE_LOG(LogObjectTools, Warning, TEXT("SCC failed to open '%s' for deletion."), *PackageFilename);
 						}
 					}
 				}
@@ -1566,6 +1575,23 @@ namespace ObjectTools
 				}
 			}
 		}
+
+		// Handle all source control revert and delete operations as a batched operation.
+		if (ISourceControlModule::Get().IsEnabled())
+		{
+			if (SCCFilesToRevert.Num() > 0)
+			{
+				SourceControlProvider.Execute(ISourceControlOperation::Create<FRevert>(), SCCFilesToRevert);
+			}
+
+			if (SCCFilesToDelete.Num() > 0)
+			{
+				if (SourceControlProvider.Execute(ISourceControlOperation::Create<FDelete>(), SCCFilesToDelete) == ECommandResult::Failed)
+				{
+					UE_LOG(LogObjectTools, Warning, TEXT("SCC failed to open the selected files for deletion."));
+				}
+			}
+		}
 	}
 
 	int32 DeleteAssets( const TArray<FAssetData>& AssetsToDelete, bool bShowConfirmation )
diff --git a/Engine/Source/Editor/UnrealEd/Private/PlayLevel.cpp b/Engine/Source/Editor/UnrealEd/Private/PlayLevel.cpp
--- a/Engine/Source/Editor/UnrealEd/Private/PlayLevel.cpp
+++ b/Engine/Source/Editor/UnrealEd/Private/PlayLevel.cpp
@@ -60,7 +60,6 @@
 #include "GameDelegates.h"
 #include "GeneralProjectSettings.h"
 #include "OnlineEngineInterface.h"
-#include "KismetDebugUtilities.h"
 #include "DebuggerCommands.h"
 
 #include "AudioThread.h"
@@ -1970,10 +1969,9 @@ void UEditorEngine::PlayUsingLauncher()
 
 			// const TArray<FString>& CookedMaps = ChainState.Profile->GetCookedMaps();
 			TArray<FString> CookDirectories;
-			TArray<FString> CookCultures;
 			TArray<FString> IniMapSections;
 
-			StartCookByTheBookInEditor(TargetPlatforms, CookedMaps, CookDirectories, CookCultures, IniMapSections );
+			StartCookByTheBookInEditor(TargetPlatforms, CookedMaps, CookDirectories, GetDefault<UProjectPackagingSettings>()->CulturesToStage, IniMapSections );
 
 			FIsCookFinishedDelegate &CookerFinishedDelegate = LauncherProfile->OnIsCookFinished();
 
@@ -2508,7 +2506,13 @@ void UEditorEngine::SpawnIntraProcessPIEWorlds(bool bAnyBlueprintErrors, bool bS
 			PlayInSettings->SetPlayNetMode(EPlayNetMode::PIE_Standalone);
 		}
 
-		GetMultipleInstancePositions(SettingsIndex++, NextX, NextY);
+		// For legacy reasons, single player PIE uses ULevelEditorPlaySettings::NewWindowPosition as its window position.
+		// Multiple PIE uses the ULevelEditorPlaySettings::MultipleInstancePositions array, starting with index 1.
+		// If this is a single player PIE, with dedicated server, don't set NewWindowPosition from the MultipleInstancePositions array - leave it as is.
+		if (PlayNumberOfClients > 1)
+		{
+			GetMultipleInstancePositions(SettingsIndex++, NextX, NextY);
+		}
 
 		UGameInstance* const ClientGameInstance = CreatePIEGameInstance(PIEInstance, bInSimulateInEditor, bAnyBlueprintErrors, bStartInSpectatorMode, false, PIEStartTime);
 		if (ClientGameInstance)
@@ -3019,7 +3023,6 @@ UGameInstance* UEditorEngine::CreatePIEGameInstance(int32 InPIEInstance, bool bI
 
 		// Add a handler for viewport close requests
 		ViewportClient->OnCloseRequested().BindUObject(this, &UEditorEngine::OnViewportCloseRequested);
-			
 		FSlatePlayInEditorInfo& SlatePlayInEditorSession = SlatePlayInEditorMap.Add(PieWorldContext->ContextHandle, FSlatePlayInEditorInfo());
 		SlatePlayInEditorSession.DestinationSlateViewport = RequestedDestinationSlateViewport;	// Might be invalid depending how pie was launched. Code below handles this.
 		RequestedDestinationSlateViewport = NULL;
@@ -3053,12 +3056,14 @@ UGameInstance* UEditorEngine::CreatePIEGameInstance(int32 InPIEInstance, bool bI
 				//SlatePlayInEditorSession.SlatePlayInEditorWindowViewport = MakeShareable<FSceneViewport>((FSceneViewport*)LevelViewportRef->GetActiveViewport());
 			}
 			else
-			{		
+			{
+				const int32 PlayNumberOfClients = [&PlayInSettings] { int32 NumberOfClients(0); return (PlayInSettings->GetPlayNumberOfClients(NumberOfClients) ? NumberOfClients : 0); }();
+
 				// Create the top level pie window and add it to Slate
 				uint32 NewWindowHeight = PlayInSettings->NewWindowHeight;
 				uint32 NewWindowWidth = PlayInSettings->NewWindowWidth;
 				FIntPoint NewWindowPosition = PlayInSettings->NewWindowPosition;
-				bool CenterNewWindow = PlayInSettings->CenterNewWindow && (PlayNetMode == PIE_Standalone);
+				bool CenterNewWindow = PlayInSettings->CenterNewWindow && (PlayNumberOfClients == 1);
 
 				// Setup size for PIE window
 				if ((NewWindowWidth <= 0) || (NewWindowHeight <= 0))
@@ -3211,7 +3216,7 @@ UGameInstance* UEditorEngine::CreatePIEGameInstance(int32 InPIEInstance, bool bI
 				
 					const bool CanPlayNetDedicated = [&PlayInSettings]{ bool PlayNetDedicated(false); return (PlayInSettings->GetPlayNetDedicated(PlayNetDedicated) && PlayNetDedicated); }();
 					PieWindow->SetOnWindowClosed(FOnWindowClosed::CreateStatic(&FLocal::OnPIEWindowClosed, TWeakPtr<SViewport>(PieViewportWidget), 
-						PieWorldContext->PIEInstance - (CanPlayNetDedicated ? 1 : 0)));
+						(PlayNumberOfClients == 1) ? 0 : PieWorldContext->PIEInstance - (CanPlayNetDedicated ? 1 : 0)));
 				}
 
 				// Create a new viewport that the viewport widget will use to render the game
diff --git a/Engine/Source/Editor/UnrealEd/Private/SEditorViewport.cpp b/Engine/Source/Editor/UnrealEd/Private/SEditorViewport.cpp
--- a/Engine/Source/Editor/UnrealEd/Private/SEditorViewport.cpp
+++ b/Engine/Source/Editor/UnrealEd/Private/SEditorViewport.cpp
@@ -214,6 +214,12 @@ void SEditorViewport::BindCommands()
 		FCanExecuteAction(),
 		FIsActionChecked::CreateSP(ClientRef, &FEditorViewportClient::IsActiveViewportType, LVT_OrthoNegativeXY));
 
+	CommandListRef.MapAction(
+		Commands.Next,
+		FExecuteAction::CreateSP(ClientRef, &FEditorViewportClient::RotateViewportType),
+		FCanExecuteAction(),
+		FIsActionChecked::CreateSP(ClientRef, &FEditorViewportClient::IsActiveViewportTypeInRotation));
+
 	CommandListRef.MapAction(
 		Commands.ScreenCapture,
 		FExecuteAction::CreateSP( this, &SEditorViewport::OnScreenCapture ),
diff --git a/Engine/Source/Editor/UnrealEd/Private/SkeletalMeshEdit.cpp b/Engine/Source/Editor/UnrealEd/Private/SkeletalMeshEdit.cpp
--- a/Engine/Source/Editor/UnrealEd/Private/SkeletalMeshEdit.cpp
+++ b/Engine/Source/Editor/UnrealEd/Private/SkeletalMeshEdit.cpp
@@ -77,7 +77,7 @@ UAnimSequence * UEditorEngine::ImportFbxAnimation( USkeleton* Skeleton, UObject*
 			{
 				// since to know full path, reimport will need to do same
 				UFbxAnimSequenceImportData* ImportData = UFbxAnimSequenceImportData::GetImportDataForAnimSequence(NewAnimation, TemplateImportData);
-				ImportData->Update(UFactory::GetCurrentFilename());
+				ImportData->Update(UFactory::GetCurrentFilename(), &(FFbxImporter->Md5Hash));
 			}
 		}
 	}
@@ -523,7 +523,7 @@ UAnimSequence * UnFbx::FFbxImporter::ImportAnimations(USkeleton* Skeleton, UObje
 
 		FString 	ParentPath = FString::Printf(TEXT("%s/%s"), *FPackageName::GetLongPackagePath(*Outer->GetName()), *SequenceName);
 		UObject* 	ParentPackage = CreatePackage(NULL, *ParentPath);
-		UObject* Object = LoadObject<UObject>(ParentPackage, *SequenceName, NULL, LOAD_None, NULL);
+		UObject* Object = LoadObject<UObject>(ParentPackage, *SequenceName, NULL, (LOAD_Quiet | LOAD_NoWarn), NULL);
 		UAnimSequence * DestSeq = Cast<UAnimSequence>(Object);
 		// if object with same name exists, warn user
 		if (Object && !DestSeq)
@@ -549,7 +549,7 @@ UAnimSequence * UnFbx::FFbxImporter::ImportAnimations(USkeleton* Skeleton, UObje
 
 		// since to know full path, reimport will need to do same
 		UFbxAnimSequenceImportData* ImportData = UFbxAnimSequenceImportData::GetImportDataForAnimSequence(DestSeq, TemplateImportData);
-		ImportData->Update(UFactory::GetCurrentFilename());
+		ImportData->Update(UFactory::GetCurrentFilename(), &Md5Hash);
 
 		ImportAnimation(Skeleton, DestSeq, Name, SortedLinks, NodeArray, CurAnimStack, ResampleRate, AnimTimeSpan);
 
