Copying //UE4/Dev-Editor to //UE4/Dev-Main (Source: //UE4/Dev-Editor @ 3082391)

diff --git a/Engine/Source/Editor/UnrealEd/Private/Tests/AutomationEditorCommon.cpp b/Engine/Source/Editor/UnrealEd/Private/Tests/AutomationEditorCommon.cpp
--- a/Engine/Source/Editor/UnrealEd/Private/Tests/AutomationEditorCommon.cpp
+++ b/Engine/Source/Editor/UnrealEd/Private/Tests/AutomationEditorCommon.cpp
@@ -66,436 +66,666 @@
 
 DEFINE_LOG_CATEGORY_STATIC(LogAutomationEditorCommon, Log, All);
 
-namespace AutomationEditorCommonUtils
+UWorld* FAutomationEditorCommonUtils::CreateNewMap()
 {
+	// Change out of Matinee when opening new map, so we avoid editing data in the old one.
+	if ( GLevelEditorModeTools().IsModeActive(FBuiltinEditorModes::EM_InterpEdit) )
+	{
+		GLevelEditorModeTools().DeactivateMode(FBuiltinEditorModes::EM_InterpEdit);
+	}
 
-	/**
-	* Imports an object using a given factory
-	*
-	* @param ImportFactory - The factory to use to import the object
-	* @param ObjectName - The name of the object to create
-	* @param PackagePath - The full path of the package file to create
-	* @param ImportPath - The path to the object to import
-	*/
-	UObject* ImportAssetUsingFactory(UFactory* ImportFactory, const FString& ObjectName, const FString& PackagePath, const FString& ImportPath)
+	// Also change out of Landscape mode to ensure all references are cleared.
+	if ( GLevelEditorModeTools().IsModeActive(FBuiltinEditorModes::EM_Landscape) )
 	{
-		UObject* ImportedAsset = NULL;
+		GLevelEditorModeTools().DeactivateMode(FBuiltinEditorModes::EM_Landscape);
+	}
 
-		UPackage* Pkg = CreatePackage(NULL, *PackagePath);
-		if (Pkg)
-		{
-			// Make sure the destination package is loaded
-			Pkg->FullyLoad();
+	// Also change out of Foliage mode to ensure all references are cleared.
+	if ( GLevelEditorModeTools().IsModeActive(FBuiltinEditorModes::EM_Foliage) )
+	{
+		GLevelEditorModeTools().DeactivateMode(FBuiltinEditorModes::EM_Foliage);
+	}
 
-			UClass* ImportAssetType = ImportFactory->ResolveSupportedClass();
-			bool bDummy = false;
+	// Change out of mesh paint mode when opening a new map.
+	if ( GLevelEditorModeTools().IsModeActive(FBuiltinEditorModes::EM_MeshPaint) )
+	{
+		GLevelEditorModeTools().DeactivateMode(FBuiltinEditorModes::EM_MeshPaint);
+	}
 
-			//If we are a texture factory suppress some warning dialog that we don't want
-			if (ImportFactory->IsA(UTextureFactory::StaticClass()))
-			{
-				UTextureFactory::SuppressImportOverwriteDialog();
-			}
+	return GEditor->NewMap();
+}
 
-			bool OutCanceled = false;
-			ImportedAsset = ImportFactory->ImportObject(ImportAssetType, Pkg, FName(*ObjectName), RF_Public | RF_Standalone, ImportPath, nullptr, OutCanceled);
+/**
+* Imports an object using a given factory
+*
+* @param ImportFactory - The factory to use to import the object
+* @param ObjectName - The name of the object to create
+* @param PackagePath - The full path of the package file to create
+* @param ImportPath - The path to the object to import
+*/
+UObject* FAutomationEditorCommonUtils::ImportAssetUsingFactory(UFactory* ImportFactory, const FString& ObjectName, const FString& PackagePath, const FString& ImportPath)
+{
+	UObject* ImportedAsset = NULL;
 
-			if (ImportedAsset != nullptr)
-			{
-				UE_LOG(LogAutomationEditorCommon, Display, TEXT("Imported %s"), *ImportPath);
-			}
-			else if (OutCanceled)
-			{
-				UE_LOG(LogAutomationEditorCommon, Display, TEXT("Canceled import of %s"), *ImportPath);
-			}
-			else
-			{
-				UE_LOG(LogAutomationEditorCommon, Error, TEXT("Failed to import asset using factory %s!"), *ImportFactory->GetName());
-			}
+	UPackage* Pkg = CreatePackage(NULL, *PackagePath);
+	if (Pkg)
+	{
+		// Make sure the destination package is loaded
+		Pkg->FullyLoad();
+
+		UClass* ImportAssetType = ImportFactory->ResolveSupportedClass();
+		bool bDummy = false;
+
+		//If we are a texture factory suppress some warning dialog that we don't want
+		if (ImportFactory->IsA(UTextureFactory::StaticClass()))
+		{
+			UTextureFactory::SuppressImportOverwriteDialog();
+		}
+
+		bool OutCanceled = false;
+		ImportedAsset = ImportFactory->ImportObject(ImportAssetType, Pkg, FName(*ObjectName), RF_Public | RF_Standalone, ImportPath, nullptr, OutCanceled);
+
+		if (ImportedAsset != nullptr)
+		{
+			UE_LOG(LogAutomationEditorCommon, Display, TEXT("Imported %s"), *ImportPath);
+		}
+		else if (OutCanceled)
+		{
+			UE_LOG(LogAutomationEditorCommon, Display, TEXT("Canceled import of %s"), *ImportPath);
 		}
 		else
 		{
-			UE_LOG(LogAutomationEditorCommon, Error, TEXT("Failed to create a package!"));
+			UE_LOG(LogAutomationEditorCommon, Error, TEXT("Failed to import asset using factory %s!"), *ImportFactory->GetName());
+		}
+	}
+	else
+	{
+		UE_LOG(LogAutomationEditorCommon, Error, TEXT("Failed to create a package!"));
+	}
+
+	return ImportedAsset;
+}
+
+/**
+* Nulls out references to a given object
+*
+* @param InObject - Object to null references to
+*/
+void FAutomationEditorCommonUtils::NullReferencesToObject(UObject* InObject)
+{
+	TArray<UObject*> ReplaceableObjects;
+	TMap<UObject*, UObject*> ReplacementMap;
+	ReplacementMap.Add(InObject, NULL);
+	ReplacementMap.GenerateKeyArray(ReplaceableObjects);
+
+	// Find all the properties (and their corresponding objects) that refer to any of the objects to be replaced
+	TMap< UObject*, TArray<UProperty*> > ReferencingPropertiesMap;
+	for (FObjectIterator ObjIter; ObjIter; ++ObjIter)
+	{
+		UObject* CurObject = *ObjIter;
+
+		// Find the referencers of the objects to be replaced
+		FFindReferencersArchive FindRefsArchive(CurObject, ReplaceableObjects);
+
+		// Inform the object referencing any of the objects to be replaced about the properties that are being forcefully
+		// changed, and store both the object doing the referencing as well as the properties that were changed in a map (so that
+		// we can correctly call PostEditChange later)
+		TMap<UObject*, int32> CurNumReferencesMap;
+		TMultiMap<UObject*, UProperty*> CurReferencingPropertiesMMap;
+		if (FindRefsArchive.GetReferenceCounts(CurNumReferencesMap, CurReferencingPropertiesMMap) > 0)
+		{
+			TArray<UProperty*> CurReferencedProperties;
+			CurReferencingPropertiesMMap.GenerateValueArray(CurReferencedProperties);
+			ReferencingPropertiesMap.Add(CurObject, CurReferencedProperties);
+			for (TArray<UProperty*>::TConstIterator RefPropIter(CurReferencedProperties); RefPropIter; ++RefPropIter)
+			{
+				CurObject->PreEditChange(*RefPropIter);
+			}
 		}
 
-		return ImportedAsset;
 	}
 
-	/**
-	* Nulls out references to a given object
-	*
-	* @param InObject - Object to null references to
-	*/
-	void NullReferencesToObject(UObject* InObject)
+	// Iterate over the map of referencing objects/changed properties, forcefully replacing the references and then
+	// alerting the referencing objects the change has completed via PostEditChange
+	int32 NumObjsReplaced = 0;
+	for (TMap< UObject*, TArray<UProperty*> >::TConstIterator MapIter(ReferencingPropertiesMap); MapIter; ++MapIter)
 	{
-		TArray<UObject*> ReplaceableObjects;
-		TMap<UObject*, UObject*> ReplacementMap;
-		ReplacementMap.Add(InObject, NULL);
-		ReplacementMap.GenerateKeyArray(ReplaceableObjects);
+		++NumObjsReplaced;
 
-		// Find all the properties (and their corresponding objects) that refer to any of the objects to be replaced
-		TMap< UObject*, TArray<UProperty*> > ReferencingPropertiesMap;
-		for (FObjectIterator ObjIter; ObjIter; ++ObjIter)
+		UObject* CurReplaceObj = MapIter.Key();
+		const TArray<UProperty*>& RefPropArray = MapIter.Value();
+
+		FArchiveReplaceObjectRef<UObject> ReplaceAr(CurReplaceObj, ReplacementMap, false, true, false);
+
+		for (TArray<UProperty*>::TConstIterator RefPropIter(RefPropArray); RefPropIter; ++RefPropIter)
 		{
-			UObject* CurObject = *ObjIter;
+			FPropertyChangedEvent PropertyEvent(*RefPropIter);
+			CurReplaceObj->PostEditChangeProperty(PropertyEvent);
+		}
 
-			// Find the referencers of the objects to be replaced
-			FFindReferencersArchive FindRefsArchive(CurObject, ReplaceableObjects);
+		if (!CurReplaceObj->HasAnyFlags(RF_Transient) && CurReplaceObj->GetOutermost() != GetTransientPackage())
+		{
+			if (!CurReplaceObj->RootPackageHasAnyFlags(PKG_CompiledIn))
+			{
+				CurReplaceObj->MarkPackageDirty();
+			}
+		}
+	}
+}
 
-			// Inform the object referencing any of the objects to be replaced about the properties that are being forcefully
-			// changed, and store both the object doing the referencing as well as the properties that were changed in a map (so that
-			// we can correctly call PostEditChange later)
-			TMap<UObject*, int32> CurNumReferencesMap;
-			TMultiMap<UObject*, UProperty*> CurReferencingPropertiesMMap;
-			if (FindRefsArchive.GetReferenceCounts(CurNumReferencesMap, CurReferencingPropertiesMMap) > 0)
+/**
+* gets a factory class based off an asset file extension
+*
+* @param AssetExtension - The file extension to use to find a supporting UFactory
+*/
+UClass* FAutomationEditorCommonUtils::GetFactoryClassForType(const FString& AssetExtension)
+{
+	// First instantiate one factory for each file extension encountered that supports the extension
+	for (TObjectIterator<UClass> ClassIt; ClassIt; ++ClassIt)
+	{
+		if ((*ClassIt)->IsChildOf(UFactory::StaticClass()) && !((*ClassIt)->HasAnyClassFlags(CLASS_Abstract)))
+		{
+			UFactory* Factory = Cast<UFactory>((*ClassIt)->GetDefaultObject());
+			if (Factory->bEditorImport)
 			{
-				TArray<UProperty*> CurReferencedProperties;
-				CurReferencingPropertiesMMap.GenerateValueArray(CurReferencedProperties);
-				ReferencingPropertiesMap.Add(CurObject, CurReferencedProperties);
-				for (TArray<UProperty*>::TConstIterator RefPropIter(CurReferencedProperties); RefPropIter; ++RefPropIter)
+				TArray<FString> FactoryExtensions;
+				Factory->GetSupportedFileExtensions(FactoryExtensions);
+
+				// Case insensitive string compare with supported formats of this factory
+				if (FactoryExtensions.Contains(AssetExtension))
 				{
-					CurObject->PreEditChange(*RefPropIter);
+					return *ClassIt;
 				}
 			}
-
 		}
+	}
 
-		// Iterate over the map of referencing objects/changed properties, forcefully replacing the references and then
-		// alerting the referencing objects the change has completed via PostEditChange
-		int32 NumObjsReplaced = 0;
-		for (TMap< UObject*, TArray<UProperty*> >::TConstIterator MapIter(ReferencingPropertiesMap); MapIter; ++MapIter)
-		{
-			++NumObjsReplaced;
+	return NULL;
+}
 
-			UObject* CurReplaceObj = MapIter.Key();
-			const TArray<UProperty*>& RefPropArray = MapIter.Value();
+/**
+* Applies settings to an object by finding UProperties by name and calling ImportText
+*
+* @param InObject - The object to search for matching properties
+* @param PropertyChain - The list UProperty names recursively to search through
+* @param Value - The value to import on the found property
+*/
+void FAutomationEditorCommonUtils::ApplyCustomFactorySetting(UObject* InObject, TArray<FString>& PropertyChain, const FString& Value)
+{
+	const FString PropertyName = PropertyChain[0];
+	PropertyChain.RemoveAt(0);
+
+	UProperty* TargetProperty = FindField<UProperty>(InObject->GetClass(), *PropertyName);
+	if (TargetProperty)
+	{
+		if (PropertyChain.Num() == 0)
+		{
+			TargetProperty->ImportText(*Value, TargetProperty->ContainerPtrToValuePtr<uint8>(InObject), 0, InObject);
+		}
+		else
+		{
+			UStructProperty* StructProperty = Cast<UStructProperty>(TargetProperty);
+			UObjectProperty* ObjectProperty = Cast<UObjectProperty>(TargetProperty);
 
-			FArchiveReplaceObjectRef<UObject> ReplaceAr(CurReplaceObj, ReplacementMap, false, true, false);
+			UObject* SubObject = NULL;
+			bool bValidPropertyType = true;
 
-			for (TArray<UProperty*>::TConstIterator RefPropIter(RefPropArray); RefPropIter; ++RefPropIter)
+			if (StructProperty)
 			{
-				FPropertyChangedEvent PropertyEvent(*RefPropIter);
-				CurReplaceObj->PostEditChangeProperty(PropertyEvent);
+				SubObject = StructProperty->Struct;
+			}
+			else if (ObjectProperty)
+			{
+				SubObject = ObjectProperty->GetObjectPropertyValue(ObjectProperty->ContainerPtrToValuePtr<UObject>(InObject));
+			}
+			else
+			{
+				//Unknown nested object type
+				bValidPropertyType = false;
+				UE_LOG(LogAutomationEditorCommon, Error, TEXT("ERROR: Unknown nested object type for property: %s"), *PropertyName);
 			}
 
-			if (!CurReplaceObj->HasAnyFlags(RF_Transient) && CurReplaceObj->GetOutermost() != GetTransientPackage())
+			if (SubObject)
 			{
-				if (!CurReplaceObj->RootPackageHasAnyFlags(PKG_CompiledIn))
-				{
-					CurReplaceObj->MarkPackageDirty();
-				}
+				ApplyCustomFactorySetting(SubObject, PropertyChain, Value);
+			}
+			else if (bValidPropertyType)
+			{
+				UE_LOG(LogAutomationEditorCommon, Error, TEXT("Error accessing null property: %s"), *PropertyName);
 			}
 		}
 	}
+	else
+	{
+		UE_LOG(LogAutomationEditorCommon, Error, TEXT("ERROR: Could not find factory property: %s"), *PropertyName);
+	}
+}
 
-	/**
-	* gets a factory class based off an asset file extension
-	*
-	* @param AssetExtension - The file extension to use to find a supporting UFactory
-	*/
-	UClass* GetFactoryClassForType(const FString& AssetExtension)
+/**
+* Applies the custom factory settings
+*
+* @param InFactory - The factory to apply custom settings to
+* @param FactorySettings - An array of custom settings to apply to the factory
+*/
+void FAutomationEditorCommonUtils::ApplyCustomFactorySettings(UFactory* InFactory, const TArray<FImportFactorySettingValues>& FactorySettings)
+{
+	bool bCallConfigureProperties = true;
+
+	for (int32 i = 0; i < FactorySettings.Num(); ++i)
 	{
-		// First instantiate one factory for each file extension encountered that supports the extension
-		for (TObjectIterator<UClass> ClassIt; ClassIt; ++ClassIt)
+		if (FactorySettings[i].SettingName.Len() > 0 && FactorySettings[i].Value.Len() > 0)
 		{
-			if ((*ClassIt)->IsChildOf(UFactory::StaticClass()) && !((*ClassIt)->HasAnyClassFlags(CLASS_Abstract)))
+			//Check if we are setting an FBX import type override.  If we are, we don't want to call ConfigureProperties because that enables bDetectImportTypeOnImport
+			if (FactorySettings[i].SettingName.Contains(TEXT("MeshTypeToImport")))
 			{
-				UFactory* Factory = Cast<UFactory>((*ClassIt)->GetDefaultObject());
-				if (Factory->bEditorImport)
-				{
-					TArray<FString> FactoryExtensions;
-					Factory->GetSupportedFileExtensions(FactoryExtensions);
-
-					// Case insensitive string compare with supported formats of this factory
-					if (FactoryExtensions.Contains(AssetExtension))
-					{
-						return *ClassIt;
-					}
-				}
+				bCallConfigureProperties = false;
 			}
+
+			TArray<FString> PropertyChain;
+			FactorySettings[i].SettingName.ParseIntoArray(PropertyChain, TEXT("."), false);
+			ApplyCustomFactorySetting(InFactory, PropertyChain, FactorySettings[i].Value);
 		}
+	}
 
-		return NULL;
+	if (bCallConfigureProperties)
+	{
+		InFactory->ConfigureProperties();
 	}
+}
 
-	/**
-	* Applies settings to an object by finding UProperties by name and calling ImportText
-	*
-	* @param InObject - The object to search for matching properties
-	* @param PropertyChain - The list UProperty names recursively to search through
-	* @param Value - The value to import on the found property
-	*/
-	void ApplyCustomFactorySetting(UObject* InObject, TArray<FString>& PropertyChain, const FString& Value)
+/**
+* Writes a number to a text file.
+* @param InTestName is the folder that has the same name as the test. (For Example: "Performance").
+* @param InItemBeingTested is the name for the thing that is being tested. (For Example: "MapName").
+* @param InFileName is the name of the file with an extension
+* @param InNumberToBeWritten is the float number that is expected to be written to the file.
+* @param Delimiter is the delimiter to be used. TEXT(",")
+*/
+void FAutomationEditorCommonUtils::WriteToTextFile(const FString& InTestName, const FString& InTestItem, const FString& InFileName, const float& InEntry, const FString& Delimiter)
+{
+	//Performance file locations and setups.
+	FString FileSaveLocation = FPaths::Combine(*FPaths::AutomationLogDir(), *InTestName, *InTestItem, *InFileName);
+
+	if (FPaths::FileExists(FileSaveLocation))
 	{
-		const FString PropertyName = PropertyChain[0];
-		PropertyChain.RemoveAt(0);
+		//The text files existing content.
+		FString TextFileContents;
 
-		UProperty* TargetProperty = FindField<UProperty>(InObject->GetClass(), *PropertyName);
-		if (TargetProperty)
-		{
-			if (PropertyChain.Num() == 0)
-			{
-				TargetProperty->ImportText(*Value, TargetProperty->ContainerPtrToValuePtr<uint8>(InObject), 0, InObject);
-			}
-			else
-			{
-				UStructProperty* StructProperty = Cast<UStructProperty>(TargetProperty);
-				UObjectProperty* ObjectProperty = Cast<UObjectProperty>(TargetProperty);
+		//Write to the text file the combined contents from the text file with the number to write.
+		FFileHelper::LoadFileToString(TextFileContents, *FileSaveLocation);
+		FString FileSetup = TextFileContents + Delimiter + FString::SanitizeFloat(InEntry);
+		FFileHelper::SaveStringToFile(FileSetup, *FileSaveLocation);
+		return;
+	}
 
-				UObject* SubObject = NULL;
-				bool bValidPropertyType = true;
+	FFileHelper::SaveStringToFile(FString::SanitizeFloat(InEntry), *FileSaveLocation);
+}
 
-				if (StructProperty)
-				{
-					SubObject = StructProperty->Struct;
-				}
-				else if (ObjectProperty)
-				{
-					SubObject = ObjectProperty->GetObjectPropertyValue(ObjectProperty->ContainerPtrToValuePtr<UObject>(InObject));
-				}
-				else
-				{
-					//Unknown nested object type
-					bValidPropertyType = false;
-					UE_LOG(LogAutomationEditorCommon, Error, TEXT("ERROR: Unknown nested object type for property: %s"), *PropertyName);
-				}
+/**
+* Returns the sum of the numbers available in an array of float.
+* @param InFloatArray is the name of the array intended to be used.
+* @param bisAveragedInstead will return the average of the available numbers instead of the sum.
+*/
+float FAutomationEditorCommonUtils::TotalFromFloatArray(const TArray<float>& InFloatArray, bool bIsAveragedInstead)
+{
+	//Total Value holds the sum of all the numbers available in the array.
+	float TotalValue = 0;
 
-				if (SubObject)
-				{
-					ApplyCustomFactorySetting(SubObject, PropertyChain, Value);
-				}
-				else if (bValidPropertyType)
-				{
-					UE_LOG(LogAutomationEditorCommon, Error, TEXT("Error accessing null property: %s"), *PropertyName);
-				}
-			}
-		}
-		else
+	//Get the sum of the array.
+	for (int32 I = 0; I < InFloatArray.Num(); ++I)
+	{
+		TotalValue += InFloatArray[I];
+	}
+
+	//If bAverageInstead equals true then only the average is returned.
+	if (bIsAveragedInstead)
+	{
+		UE_LOG(LogEditorAutomationTests, VeryVerbose, TEXT("Average value of the Array is %f"), (TotalValue / InFloatArray.Num()));
+		return (TotalValue / InFloatArray.Num());
+	}
+
+	UE_LOG(LogEditorAutomationTests, VeryVerbose, TEXT("Total Value of the Array is %f"), TotalValue);
+	return TotalValue;
+}
+
+/**
+* Returns the largest value from an array of float numbers.
+* @param InFloatArray is the name of the array intended to be used.
+*/
+float FAutomationEditorCommonUtils::LargestValueInFloatArray(const TArray<float>& InFloatArray)
+{
+	//Total Value holds the sum of all the numbers available in the array.
+	float LargestValue = 0;
+
+	//Find the largest value
+	for (int32 I = 0; I < InFloatArray.Num(); ++I)
+	{
+		if (LargestValue < InFloatArray[I])
 		{
-			UE_LOG(LogAutomationEditorCommon, Error, TEXT("ERROR: Could not find factory property: %s"), *PropertyName);
+			LargestValue = InFloatArray[I];
 		}
 	}
+	UE_LOG(LogEditorAutomationTests, VeryVerbose, TEXT("The Largest value of the array is %f"), LargestValue);
+	return LargestValue;
+}
 
-	/**
-	* Applies the custom factory settings
-	*
-	* @param InFactory - The factory to apply custom settings to
-	* @param FactorySettings - An array of custom settings to apply to the factory
-	*/
-	void ApplyCustomFactorySettings(UFactory* InFactory, const TArray<FImportFactorySettingValues>& FactorySettings)
+/**
+* Returns the contents of a text file as an array of FString.
+* @param InFileLocation -  is the location of the file.
+* @param OutArray - The name of the array where the 
+*/
+void FAutomationEditorCommonUtils::CreateArrayFromFile(const FString& InFileLocation, TArray<FString>& OutArray)
+{
+	FString RawData;
+
+	if (FPaths::FileExists(*InFileLocation))
+	{
+		UE_LOG(LogEditorAutomationTests, VeryVerbose, TEXT("Loading and parsing the data from '%s' into an array."), *InFileLocation);
+		FFileHelper::LoadFileToString(RawData, *InFileLocation);
+		RawData.ParseIntoArray(OutArray, TEXT(","), false);
+	}
+
+	UE_LOG(LogEditorAutomationTests, Warning, TEXT("Unable to create an array.  '%s' does not exist."), *InFileLocation);
+	RawData = TEXT("0");
+	OutArray.Add(RawData);
+}
+
+/**
+* Returns true if the archive/file can be written to otherwise false..
+* @param InFilePath - is the location of the file.
+* @param InArchiveName - is the name of the archive to be used.
+*/
+bool FAutomationEditorCommonUtils::IsArchiveWriteable(const FString& InFilePath, const FArchive* InArchiveName)
+{
+	if (!InArchiveName)
 	{
-		bool bCallConfigureProperties = true;
+		UE_LOG(LogEditorAutomationTests, Error, TEXT("Failed to write to the csv file: %s"), *FPaths::ConvertRelativePathToFull(InFilePath));
+		return false;
+	}
+	return true;
+}
+
+
+void FAutomationEditorCommonUtils::GetLaunchOnDeviceID(FString& OutDeviceID, const FString& InMapName)
+{
+	UAutomationTestSettings const* AutomationTestSettings = GetDefault<UAutomationTestSettings>();
+	check(AutomationTestSettings);
 
-		for (int32 i = 0; i < FactorySettings.Num(); ++i)
+	OutDeviceID = "None";
+
+	FString LaunchOnDeviceId;
+	for (auto LaunchIter = AutomationTestSettings->LaunchOnSettings.CreateConstIterator(); LaunchIter; LaunchIter++)
+	{
+		FString LaunchOnSettings = LaunchIter->DeviceID;
+		FString LaunchOnMap = FPaths::GetBaseFilename(LaunchIter->LaunchOnTestmap.FilePath);
+		if (LaunchOnMap.Equals(InMapName))
 		{
-			if (FactorySettings[i].SettingName.Len() > 0 && FactorySettings[i].Value.Len() > 0)
+			// shared devices section
+			TSharedPtr<ITargetDeviceServicesModule> TargetDeviceServicesModule = StaticCastSharedPtr<ITargetDeviceServicesModule>(FModuleManager::Get().LoadModule(TEXT("TargetDeviceServices")));
+			// for each platform...
+			TArray<ITargetDeviceProxyPtr> DeviceProxies;
+			TargetDeviceServicesModule->GetDeviceProxyManager()->GetProxies(FName(*LaunchOnSettings), true, DeviceProxies);
+			// for each proxy...
+			for (auto DeviceProxyIt = DeviceProxies.CreateIterator(); DeviceProxyIt; ++DeviceProxyIt)
 			{
-				//Check if we are setting an FBX import type override.  If we are, we don't want to call ConfigureProperties because that enables bDetectImportTypeOnImport
-				if (FactorySettings[i].SettingName.Contains(TEXT("MeshTypeToImport")))
+				ITargetDeviceProxyPtr DeviceProxy = *DeviceProxyIt;
+				if (DeviceProxy->IsConnected())
 				{
-					bCallConfigureProperties = false;
+					OutDeviceID = DeviceProxy->GetTargetDeviceId((FName)*LaunchOnSettings);
+					break;
 				}
-
-				TArray<FString> PropertyChain;
-				FactorySettings[i].SettingName.ParseIntoArray(PropertyChain, TEXT("."), false);
-				ApplyCustomFactorySetting(InFactory, PropertyChain, FactorySettings[i].Value);
 			}
 		}
+	}
+}
 
-		if (bCallConfigureProperties)
+void FAutomationEditorCommonUtils::GetLaunchOnDeviceID(FString& OutDeviceID, const FString& InMapName, const FString& InDeviceName)
+{
+	UAutomationTestSettings const* AutomationTestSettings = GetDefault<UAutomationTestSettings>();
+	check(AutomationTestSettings);
+
+	//Output device name will default to "None".
+	OutDeviceID = "None";
+
+	// shared devices section
+	TSharedPtr<ITargetDeviceServicesModule> TargetDeviceServicesModule = StaticCastSharedPtr<ITargetDeviceServicesModule>(FModuleManager::Get().LoadModule(TEXT("TargetDeviceServices")));
+	// for each platform...
+	TArray<ITargetDeviceProxyPtr> DeviceProxies;
+	TargetDeviceServicesModule->GetDeviceProxyManager()->GetProxies(FName(*InDeviceName), true, DeviceProxies);
+	// for each proxy...
+	for (auto DeviceProxyIt = DeviceProxies.CreateIterator(); DeviceProxyIt; ++DeviceProxyIt)
+	{
+		ITargetDeviceProxyPtr DeviceProxy = *DeviceProxyIt;
+		if (DeviceProxy->IsConnected())
 		{
-			InFactory->ConfigureProperties();
+			OutDeviceID = DeviceProxy->GetTargetDeviceId((FName)*InDeviceName);
+			break;
 		}
 	}
+}
 
-	/**
-	* Writes a number to a text file.
-	* @param InTestName is the folder that has the same name as the test. (For Example: "Performance").
-	* @param InItemBeingTested is the name for the thing that is being tested. (For Example: "MapName").
-	* @param InFileName is the name of the file with an extension
-	* @param InNumberToBeWritten is the float number that is expected to be written to the file.
-	* @param Delimiter is the delimiter to be used. TEXT(",")
-	*/
-	void WriteToTextFile(const FString& InTestName, const FString& InTestItem, const FString& InFileName, const float& InEntry, const FString& Delimiter)
+bool FAutomationEditorCommonUtils::SetOrthoViewportView(const FVector& ViewLocation, const FRotator& ViewRotation)
+{
+	for (int32 i = 0; i < GEditor->LevelViewportClients.Num(); i++)
 	{
-		//Performance file locations and setups.
-		FString FileSaveLocation = FPaths::Combine(*FPaths::AutomationLogDir(), *InTestName, *InTestItem, *InFileName);
-
-		if (FPaths::FileExists(FileSaveLocation))
+		FLevelEditorViewportClient* ViewportClient = GEditor->LevelViewportClients[i];
+		if (!ViewportClient->IsOrtho())
 		{
-			//The text files existing content.
-			FString TextFileContents;
-
-			//Write to the text file the combined contents from the text file with the number to write.
-			FFileHelper::LoadFileToString(TextFileContents, *FileSaveLocation);
-			FString FileSetup = TextFileContents + Delimiter + FString::SanitizeFloat(InEntry);
-			FFileHelper::SaveStringToFile(FileSetup, *FileSaveLocation);
-			return;
+			ViewportClient->SetViewLocation(ViewLocation);
+			ViewportClient->SetViewRotation(ViewRotation);
+			return true;
 		}
+	}
 
-		FFileHelper::SaveStringToFile(FString::SanitizeFloat(InEntry), *FileSaveLocation);
+	UE_LOG(LogEditorAutomationTests, Log, TEXT("An ortho viewport was not found.  May affect the test results."));
+	return false;
+}
+
+//////////////////////////////////////////////////////////////////////
+//Asset Path Commands
+
+/**
+* Converts a package path to an asset path
+*
+* @param PackagePath - The package path to convert
+*/
+FString FAutomationEditorCommonUtils::ConvertPackagePathToAssetPath(const FString& PackagePath)
+{
+	const FString Filename = FPaths::ConvertRelativePathToFull(PackagePath);
+	FString EngineFileName = Filename;
+	FString GameFileName = Filename;
+	if (FPaths::MakePathRelativeTo(EngineFileName, *FPaths::EngineContentDir()) && !EngineFileName.Contains(TEXT("../")))
+	{
+		const FString ShortName = FPaths::GetBaseFilename(EngineFileName);
+		const FString PathName = FPaths::GetPath(EngineFileName);
+		const FString AssetName = FString::Printf(TEXT("/Engine/%s/%s.%s"), *PathName, *ShortName, *ShortName);
+		return AssetName;
 	}
+	else if (FPaths::MakePathRelativeTo(GameFileName, *FPaths::GameContentDir()) && !GameFileName.Contains(TEXT("../")))
+	{
+		const FString ShortName = FPaths::GetBaseFilename(GameFileName);
+		const FString PathName = FPaths::GetPath(GameFileName);
+		const FString AssetName = FString::Printf(TEXT("/Game/%s/%s.%s"), *PathName, *ShortName, *ShortName);
+		return AssetName;
+	}
+	else
+	{
+		UE_LOG(LogAutomationEditorCommon, Error, TEXT("PackagePath (%s) is invalid for the current project"), *PackagePath);
+		return TEXT("");
+	}
+}
 
-	/**
-	* Returns the sum of the numbers available in an array of float.
-	* @param InFloatArray is the name of the array intended to be used.
-	* @param bisAveragedInstead will return the average of the available numbers instead of the sum.
-	*/
-	float TotalFromFloatArray(const TArray<float>& InFloatArray, bool bIsAveragedInstead)
+/**
+* Gets the asset data from a package path
+*
+* @param PackagePath - The package path used to look up the asset data
+*/
+FAssetData FAutomationEditorCommonUtils::GetAssetDataFromPackagePath(const FString& PackagePath)
+{
+	FString AssetPath = FAutomationEditorCommonUtils::ConvertPackagePathToAssetPath(PackagePath);
+	if (AssetPath.Len() > 0)
 	{
-		//Total Value holds the sum of all the numbers available in the array.
-		float TotalValue = 0;
+		IAssetRegistry& AssetRegistry = FModuleManager::LoadModuleChecked<FAssetRegistryModule>(TEXT("AssetRegistry")).Get();
+		return AssetRegistry.GetAssetByObjectPath(*AssetPath);
+	}
 
-		//Get the sum of the array.
-		for (int32 I = 0; I < InFloatArray.Num(); ++I)
-		{
-			TotalValue += InFloatArray[I];
-		}
+	return FAssetData();
+}
 
-		//If bAverageInstead equals true then only the average is returned.
-		if (bIsAveragedInstead)
+//////////////////////////////////////////////////////////////////////
+//Find Asset Commands
+
+/**
+* Generates a list of assets from the ENGINE and the GAME by a specific type.
+* This is to be used by the GetTest() function.
+*/
+void FAutomationEditorCommonUtils::CollectTestsByClass(UClass * Class, TArray<FString>& OutBeautifiedNames, TArray <FString>& OutTestCommands)
+{
+	FAssetRegistryModule& AssetRegistryModule = FModuleManager::LoadModuleChecked<FAssetRegistryModule>(TEXT("AssetRegistry"));
+	TArray<FAssetData> ObjectList;
+	AssetRegistryModule.Get().GetAssetsByClass(Class->GetFName(), ObjectList);
+
+	for (TObjectIterator<UClass> AllClassesIt; AllClassesIt; ++AllClassesIt)
+	{
+		UClass* ClassList = *AllClassesIt;
+		FName ClassName = ClassList->GetFName();
+	}
+
+	for (auto ObjIter = ObjectList.CreateConstIterator(); ObjIter; ++ObjIter)
+	{
+		const FAssetData& Asset = *ObjIter;
+		FString Filename = Asset.ObjectPath.ToString();
+		//convert to full paths
+		Filename = FPackageName::LongPackageNameToFilename(Filename);
+		if (FAutomationTestFramework::GetInstance().ShouldTestContent(Filename))
 		{
-			UE_LOG(LogEditorAutomationTests, VeryVerbose, TEXT("Average value of the Array is %f"), (TotalValue / InFloatArray.Num()));
-			return (TotalValue / InFloatArray.Num());
+			FString BeautifiedFilename = Asset.AssetName.ToString();
+			OutBeautifiedNames.Add(BeautifiedFilename);
+			OutTestCommands.Add(Asset.ObjectPath.ToString());
 		}
-
-		UE_LOG(LogEditorAutomationTests, VeryVerbose, TEXT("Total Value of the Array is %f"), TotalValue);
-		return TotalValue;
 	}
+}
+
+/**
+* Generates a list of assets from the GAME by a specific type.
+* This is to be used by the GetTest() function.
+*/
+void FAutomationEditorCommonUtils::CollectGameContentTestsByClass(UClass * Class, bool bRecursiveClass, TArray<FString>& OutBeautifiedNames, TArray <FString>& OutTestCommands)
+{
+	//Setting the Asset Registry
+	FAssetRegistryModule& AssetRegistryModule = FModuleManager::LoadModuleChecked<FAssetRegistryModule>(TEXT("AssetRegistry"));
+
+	//Variable setups
+	TArray<FAssetData> ObjectList;
+	FARFilter AssetFilter;
+
+	//Generating the list of assets.
+	//This list is being filtered by the game folder and class type.  The results are placed into the ObjectList variable.
+	AssetFilter.ClassNames.Add(Class->GetFName());
+
+	//removed path as a filter as it causes two large lists to be sorted.  Filtering on "game" directory on iteration
+	//AssetFilter.PackagePaths.Add("/Game");
+	AssetFilter.bRecursiveClasses = bRecursiveClass;
+	AssetFilter.bRecursivePaths = true;
+	AssetRegistryModule.Get().GetAssets(AssetFilter, ObjectList);
 
-	/**
-	* Returns the largest value from an array of float numbers.
-	* @param InFloatArray is the name of the array intended to be used.
-	*/
-	float LargestValueInFloatArray(const TArray<float>& InFloatArray)
+	//Loop through the list of assets, make their path full and a string, then add them to the test.
+	for (auto ObjIter = ObjectList.CreateConstIterator(); ObjIter; ++ObjIter)
 	{
-		//Total Value holds the sum of all the numbers available in the array.
-		float LargestValue = 0;
+		const FAssetData& Asset = *ObjIter;
+		FString Filename = Asset.ObjectPath.ToString();
 
-		//Find the largest value
-		for (int32 I = 0; I < InFloatArray.Num(); ++I)
+		if (Filename.StartsWith("/Game"))
 		{
-			if (LargestValue < InFloatArray[I])
+			//convert to full paths
+			Filename = FPackageName::LongPackageNameToFilename(Filename);
+			if (FAutomationTestFramework::GetInstance().ShouldTestContent(Filename))
 			{
-				LargestValue = InFloatArray[I];
+				FString BeautifiedFilename = Asset.AssetName.ToString();
+				OutBeautifiedNames.Add(BeautifiedFilename);
+				OutTestCommands.Add(Asset.ObjectPath.ToString());
 			}
 		}
-		UE_LOG(LogEditorAutomationTests, VeryVerbose, TEXT("The Largest value of the array is %f"), LargestValue);
-		return LargestValue;
 	}
+}
 
-	/**
-	* Returns the contents of a text file as an array of FString.
-	* @param InFileLocation -  is the location of the file.
-	* @param OutArray - The name of the array where the 
-	*/
-	void CreateArrayFromFile(const FString& InFileLocation, TArray<FString>& OutArray)
-	{
-		FString RawData;
+void FAutomationEditorCommonUtils::LoadMap(const FString& MapName)
+{
+	bool bLoadAsTemplate = false;
+	bool bShowProgress = false;
+	FEditorFileUtils::LoadMap(MapName, bLoadAsTemplate, bShowProgress);
+}
 
-		if (FPaths::FileExists(*InFileLocation))
-		{
-			UE_LOG(LogEditorAutomationTests, VeryVerbose, TEXT("Loading and parsing the data from '%s' into an array."), *InFileLocation);
-			FFileHelper::LoadFileToString(RawData, *InFileLocation);
-			RawData.ParseIntoArray(OutArray, TEXT(","), false);
-		}
+void FAutomationEditorCommonUtils::RunPIE()
+{
+	bool bInSimulateInEditor = true;
+	//once in the editor
+	ADD_LATENT_AUTOMATION_COMMAND(FStartPIECommand(true));
+	ADD_LATENT_AUTOMATION_COMMAND(FWaitLatentCommand(3.0f));
+	ADD_LATENT_AUTOMATION_COMMAND(FEndPlayMapCommand());
+
+	//wait between tests
+	ADD_LATENT_AUTOMATION_COMMAND(FWaitLatentCommand(1.0f));
+
+	//once not in the editor
+	ADD_LATENT_AUTOMATION_COMMAND(FStartPIECommand(false));
+	ADD_LATENT_AUTOMATION_COMMAND(FWaitLatentCommand(3.0f));
+	ADD_LATENT_AUTOMATION_COMMAND(FEndPlayMapCommand());
+}
 
-		UE_LOG(LogEditorAutomationTests, Warning, TEXT("Unable to create an array.  '%s' does not exist."), *InFileLocation);
-		RawData = TEXT("0");
-		OutArray.Add(RawData);
-	}
+/**
+* Generates a list of assets from the GAME by a specific type.
+* This is to be used by the GetTest() function.
+*/
+void FAutomationEditorCommonUtils::CollectGameContentTests(TArray<FString>& OutBeautifiedNames, TArray <FString>& OutTestCommands)
+{
+	//Setting the Asset Registry
+	FAssetRegistryModule& AssetRegistryModule = FModuleManager::LoadModuleChecked<FAssetRegistryModule>(TEXT("AssetRegistry"));
+
+	//Variable setups
+	TArray<FAssetData> ObjectList;
+	FARFilter AssetFilter;
+	
+	//removed path as a filter as it causes two large lists to be sorted.  Filtering on "game" directory on iteration
+	AssetFilter.PackagePaths.Add("/Game");
+	AssetFilter.bRecursiveClasses = true;
+	AssetFilter.bRecursivePaths = true;
+	AssetRegistryModule.Get().GetAssets(AssetFilter, ObjectList);
 
-	/**
-	* Returns true if the archive/file can be written to otherwise false..
-	* @param InFilePath - is the location of the file.
-	* @param InArchiveName - is the name of the archive to be used.
-	*/
-	bool IsArchiveWriteable(const FString& InFilePath, const FArchive* InArchiveName)
+	//Loop through the list of assets, make their path full and a string, then add them to the test.
+	for (auto ObjIter = ObjectList.CreateConstIterator(); ObjIter; ++ObjIter)
 	{
-		if (!InArchiveName)
+		const FAssetData& Asset = *ObjIter;
+		if (Asset.GetClass() == nullptr)
 		{
-			UE_LOG(LogEditorAutomationTests, Error, TEXT("Failed to write to the csv file: %s"), *FPaths::ConvertRelativePathToFull(InFilePath));
-			return false;
+			// a nullptr class is bad !
+			UE_LOG(LogAutomationEditorCommon, Warning, TEXT("GetClass for %s (%s) returned nullptr. Asset ignored"), *Asset.AssetName.ToString(), *Asset.ObjectPath.ToString());
 		}
-		return true;
-	}
-
-
-	void GetLaunchOnDeviceID(FString& OutDeviceID, const FString& InMapName)
-	{
-		UAutomationTestSettings const* AutomationTestSettings = GetDefault<UAutomationTestSettings>();
-		check(AutomationTestSettings);
-
-		OutDeviceID = "None";
-
-		FString LaunchOnDeviceId;
-		for (auto LaunchIter = AutomationTestSettings->LaunchOnSettings.CreateConstIterator(); LaunchIter; LaunchIter++)
+		else 
 		{
-			FString LaunchOnSettings = LaunchIter->DeviceID;
-			FString LaunchOnMap = FPaths::GetBaseFilename(LaunchIter->LaunchOnTestmap.FilePath);
-			if (LaunchOnMap.Equals(InMapName))
+			FString Filename = Asset.ObjectPath.ToString();
+
+			if (Filename.StartsWith("/Game"))
 			{
-				// shared devices section
-				TSharedPtr<ITargetDeviceServicesModule> TargetDeviceServicesModule = StaticCastSharedPtr<ITargetDeviceServicesModule>(FModuleManager::Get().LoadModule(TEXT("TargetDeviceServices")));
-				// for each platform...
-				TArray<ITargetDeviceProxyPtr> DeviceProxies;
-				TargetDeviceServicesModule->GetDeviceProxyManager()->GetProxies(FName(*LaunchOnSettings), true, DeviceProxies);
-				// for each proxy...
-				for (auto DeviceProxyIt = DeviceProxies.CreateIterator(); DeviceProxyIt; ++DeviceProxyIt)
+				//convert to full paths
+				Filename = FPackageName::LongPackageNameToFilename(Filename);
+				if (FAutomationTestFramework::GetInstance().ShouldTestContent(Filename))
 				{
-					ITargetDeviceProxyPtr DeviceProxy = *DeviceProxyIt;
-					if (DeviceProxy->IsConnected())
-					{
-						OutDeviceID = DeviceProxy->GetTargetDeviceId((FName)*LaunchOnSettings);
-						break;
-					}
+					FString BeautifiedFilename = FString::Printf(TEXT("%s.%s"), *Asset.GetClass()->GetFName().ToString(), *Asset.AssetName.ToString());
+					OutBeautifiedNames.Add(BeautifiedFilename);
+					OutTestCommands.Add(Asset.ObjectPath.ToString());
 				}
 			}
 		}
 	}
+}
+
 
-	void GetLaunchOnDeviceID(FString& OutDeviceID, const FString& InMapName, const FString& InDeviceName)
-	{
-		UAutomationTestSettings const* AutomationTestSettings = GetDefault<UAutomationTestSettings>();
-		check(AutomationTestSettings);
 
-		//Output device name will default to "None".
-		OutDeviceID = "None";
 
-		// shared devices section
-		TSharedPtr<ITargetDeviceServicesModule> TargetDeviceServicesModule = StaticCastSharedPtr<ITargetDeviceServicesModule>(FModuleManager::Get().LoadModule(TEXT("TargetDeviceServices")));
-		// for each platform...
-		TArray<ITargetDeviceProxyPtr> DeviceProxies;
-		TargetDeviceServicesModule->GetDeviceProxyManager()->GetProxies(FName(*InDeviceName), true, DeviceProxies);
-		// for each proxy...
-		for (auto DeviceProxyIt = DeviceProxies.CreateIterator(); DeviceProxyIt; ++DeviceProxyIt)
-		{
-			ITargetDeviceProxyPtr DeviceProxy = *DeviceProxyIt;
-			if (DeviceProxy->IsConnected())
-			{
-				OutDeviceID = DeviceProxy->GetTargetDeviceId((FName)*InDeviceName);
-				break;
-			}
-		}
-	}
 
-	bool SetOrthoViewportView(const FVector& ViewLocation, const FRotator& ViewRotation)
-	{
-		for (int32 i = 0; i < GEditor->LevelViewportClients.Num(); i++)
-		{
-			FLevelEditorViewportClient* ViewportClient = GEditor->LevelViewportClients[i];
-			if (!ViewportClient->IsOrtho())
-			{
-				ViewportClient->SetViewLocation(ViewLocation);
-				ViewportClient->SetViewRotation(ViewRotation);
-				return true;
-			}
-		}
 
-		UE_LOG(LogEditorAutomationTests, Log, TEXT("An ortho viewport was not found.  May affect the test results."));
-		return false;
-	}
 
-}
 
 ///////////////////////////////////////////////////////////////////////
 // Common Latent commands
@@ -504,7 +734,7 @@ namespace AutomationEditorCommonUtils
 //If bUndo is true then the undo action will occur otherwise a redo will happen.
 bool FUndoRedoCommand::Update()
 {
-	if (bUndo == true)
+	if ( bUndo == true )
 	{
 		//Undo
 		GEditor->UndoTransaction();
@@ -524,11 +754,11 @@ bool FUndoRedoCommand::Update()
 bool FOpenEditorForAssetCommand::Update()
 {
 	UObject* Object = StaticLoadObject(UObject::StaticClass(), NULL, *AssetName);
-	if (Object)
+	if ( Object )
 	{
 		FAssetEditorManager::Get().OpenEditorForAsset(Object);
 		//This checks to see if the asset sub editor is loaded.
-		if (FAssetEditorManager::Get().FindEditorForAsset(Object, true) != NULL)
+		if ( FAssetEditorManager::Get().FindEditorForAsset(Object, true) != NULL )
 		{
 			UE_LOG(LogEditorAutomationTests, Log, TEXT("Verified asset editor for: %s."), *AssetName);
 			UE_LOG(LogEditorAutomationTests, Display, TEXT("The editor successffully loaded for: %s."), *AssetName);
@@ -550,7 +780,7 @@ bool FCloseAllAssetEditorsCommand::Update()
 	FAssetEditorManager::Get().CloseAllAssetEditors();
 
 	//Get all assets currently being tracked with open editors and make sure they are not still opened.
-	if(FAssetEditorManager::Get().GetAllEditedAssets().Num() >= 1)
+	if ( FAssetEditorManager::Get().GetAllEditedAssets().Num() >= 1 )
 	{
 		UE_LOG(LogEditorAutomationTests, Warning, TEXT("Not all of the editors were closed."));
 		return true;
@@ -594,7 +824,7 @@ bool FEditorLoadMap::Update()
 	double MapLoadStartTime = FPlatformTime::Seconds();
 
 	//Load the map
-	FEditorAutomationTestUtilities::LoadMap(MapName);
+	FAutomationEditorCommonUtils::LoadMap(MapName);
 
 	//This is the time it took to load the map in the editor.
 	double MapLoadTime = FPlatformTime::Seconds() - MapLoadStartTime;
@@ -606,10 +836,10 @@ bool FEditorLoadMap::Update()
 	UE_LOG(LogEditorAutomationTests, Log, TEXT("%s has been loaded."), *ShortMapName);
 
 	//Log out to a text file the time it takes to load the map.
-	AutomationEditorCommonUtils::WriteToTextFile(TEXT("Performance"), LoadedMapName, TEXT("RAWMapLoadTime.txt"), MapLoadTime, TEXT(","));
-	
+	FAutomationEditorCommonUtils::WriteToTextFile(TEXT("Performance"), LoadedMapName, TEXT("RAWMapLoadTime.txt"), MapLoadTime, TEXT(","));
+
 	UE_LOG(LogEditorAutomationTests, Display, TEXT("%s took %.3f to load."), *LoadedMapName, MapLoadTime);
-	
+
 	return true;
 }
 
@@ -636,13 +866,13 @@ bool FChangeViewportToFirstAvailableBookmarkCommand::Update()
 	UE_LOG(LogEditorAutomationTests, Log, TEXT("Attempting to change the editor viewports view to the first set bookmark."));
 
 	//Move the perspective viewport view to show the test.
-	for (int32 i = 0; i < GEditor->LevelViewportClients.Num(); i++)
+	for ( int32 i = 0; i < GEditor->LevelViewportClients.Num(); i++ )
 	{
 		ViewportClient = GEditor->LevelViewportClients[i];
 
-		for (ViewportIndex = 0; ViewportIndex <= AWorldSettings::MAX_BOOKMARK_NUMBER; ViewportIndex++)
+		for ( ViewportIndex = 0; ViewportIndex <= AWorldSettings::MAX_BOOKMARK_NUMBER; ViewportIndex++ )
 		{
-			if (EditorModeTools.CheckBookmark(ViewportIndex, ViewportClient))
+			if ( EditorModeTools.CheckBookmark(ViewportIndex, ViewportClient) )
 			{
 				UE_LOG(LogEditorAutomationTests, VeryVerbose, TEXT("Changing a viewport view to the set bookmark %i"), ViewportIndex);
 				EditorModeTools.JumpToBookmark(ViewportIndex, true, ViewportClient);
@@ -667,7 +897,7 @@ bool FAddStaticMeshCommand::Update()
 
 	UE_LOG(LogEditorAutomationTests, Log, TEXT("Static Mesh cube has been added to 0, 0, 0."))
 
-	return true;
+		return true;
 }
 
 /**
@@ -676,13 +906,13 @@ bool FAddStaticMeshCommand::Update()
 bool FBuildLightingCommand::Update()
 {
 	//If we are running with -NullRHI then we have to skip this step.
-	if (GUsingNullRHI)
+	if ( GUsingNullRHI )
 	{
 		UE_LOG(LogEditorAutomationTests, Log, TEXT("SKIPPED Build Lighting Step.  You're currently running with -NullRHI."));
 		return true;
 	}
 
-	if (GUnrealEd->WarnIfLightingBuildIsCurrentlyRunning())
+	if ( GUnrealEd->WarnIfLightingBuildIsCurrentlyRunning() )
 	{
 		UE_LOG(LogEditorAutomationTests, Warning, TEXT("Lighting is already being built."));
 		return true;
@@ -714,7 +944,7 @@ bool FBuildLightingCommand::Update()
 
 bool FSaveLevelCommand::Update()
 {
-	if (!GUnrealEd->IsLightingBuildCurrentlyExporting() && !GUnrealEd->IsLightingBuildCurrentlyRunning())
+	if ( !GUnrealEd->IsLightingBuildCurrentlyExporting() && !GUnrealEd->IsLightingBuildCurrentlyRunning() )
 	{
 		UWorld* World = GEditor->GetEditorWorldContext().World();
 		ULevel* Level = World->GetCurrentLevel();
@@ -735,15 +965,15 @@ bool FLaunchOnCommand::Update()
 
 bool FWaitToFinishCookByTheBookCommand::Update()
 {
-	if (!GUnrealEd->CookServer->IsCookByTheBookRunning())
+	if ( !GUnrealEd->CookServer->IsCookByTheBookRunning() )
 	{
-		if (GUnrealEd->IsCookByTheBookInEditorFinished())
+		if ( GUnrealEd->IsCookByTheBookInEditorFinished() )
 		{
 			UE_LOG(LogEditorAutomationTests, Log, TEXT("The cook by the book operation has finished."));
 		}
 		return true;
 	}
-	else if ((FPlatformTime::Seconds() - StartTime) == COOK_TIMEOUT)
+	else if ( ( FPlatformTime::Seconds() - StartTime ) == COOK_TIMEOUT )
 	{
 		GUnrealEd->CancelCookByTheBookInEditor();
 		UE_LOG(LogEditorAutomationTests, Error, TEXT("It has been an hour or more since the cook has started."));
@@ -755,7 +985,7 @@ bool FWaitToFinishCookByTheBookCommand::Update()
 bool FDeleteDirCommand::Update()
 {
 	FString FullFolderPath = FPaths::ConvertRelativePathToFull(*InFolderLocation);
-	if (IFileManager::Get().DirectoryExists(*FullFolderPath))
+	if ( IFileManager::Get().DirectoryExists(*FullFolderPath) )
 	{
 		IFileManager::Get().DeleteDirectory(*FullFolderPath, false, true);
 	}
@@ -764,190 +994,15 @@ bool FDeleteDirCommand::Update()
 
 bool FWaitToFinishBuildDeployCommand::Update()
 {
-	if (GEditor->LauncherWorker->GetStatus() == ELauncherWorkerStatus::Completed)
+	if ( GEditor->LauncherWorker->GetStatus() == ELauncherWorkerStatus::Completed )
 	{
 		UE_LOG(LogEditorAutomationTests, Log, TEXT("The build game and deploy operation has finished."));
 		return true;
 	}
-	else if (GEditor->LauncherWorker->GetStatus() == ELauncherWorkerStatus::Canceled || GEditor->LauncherWorker->GetStatus() == ELauncherWorkerStatus::Canceling)
+	else if ( GEditor->LauncherWorker->GetStatus() == ELauncherWorkerStatus::Canceled || GEditor->LauncherWorker->GetStatus() == ELauncherWorkerStatus::Canceling )
 	{
 		UE_LOG(LogEditorAutomationTests, Warning, TEXT("The build was canceled."));
 		return true;
 	}
 	return false;
-}
-
-//////////////////////////////////////////////////////////////////////
-//Asset Path Commands
-
-/**
-* Converts a package path to an asset path
-*
-* @param PackagePath - The package path to convert
-*/
-FString FEditorAutomationTestUtilities::ConvertPackagePathToAssetPath(const FString& PackagePath)
-{
-	const FString Filename = FPaths::ConvertRelativePathToFull(PackagePath);
-	FString EngineFileName = Filename;
-	FString GameFileName = Filename;
-	if (FPaths::MakePathRelativeTo(EngineFileName, *FPaths::EngineContentDir()) && !EngineFileName.Contains(TEXT("../")))
-	{
-		const FString ShortName = FPaths::GetBaseFilename(EngineFileName);
-		const FString PathName = FPaths::GetPath(EngineFileName);
-		const FString AssetName = FString::Printf(TEXT("/Engine/%s/%s.%s"), *PathName, *ShortName, *ShortName);
-		return AssetName;
-	}
-	else if (FPaths::MakePathRelativeTo(GameFileName, *FPaths::GameContentDir()) && !GameFileName.Contains(TEXT("../")))
-	{
-		const FString ShortName = FPaths::GetBaseFilename(GameFileName);
-		const FString PathName = FPaths::GetPath(GameFileName);
-		const FString AssetName = FString::Printf(TEXT("/Game/%s/%s.%s"), *PathName, *ShortName, *ShortName);
-		return AssetName;
-	}
-	else
-	{
-		UE_LOG(LogAutomationEditorCommon, Error, TEXT("PackagePath (%s) is invalid for the current project"), *PackagePath);
-		return TEXT("");
-	}
-}
-
-/**
-* Gets the asset data from a package path
-*
-* @param PackagePath - The package path used to look up the asset data
-*/
-FAssetData FEditorAutomationTestUtilities::GetAssetDataFromPackagePath(const FString& PackagePath)
-{
-	FString AssetPath = FEditorAutomationTestUtilities::ConvertPackagePathToAssetPath(PackagePath);
-	if (AssetPath.Len() > 0)
-	{
-		IAssetRegistry& AssetRegistry = FModuleManager::LoadModuleChecked<FAssetRegistryModule>(TEXT("AssetRegistry")).Get();
-		return AssetRegistry.GetAssetByObjectPath(*AssetPath);
-	}
-
-	return FAssetData();
-}
-
-//////////////////////////////////////////////////////////////////////
-//Find Asset Commands
-
-/**
-* Generates a list of assets from the ENGINE and the GAME by a specific type.
-* This is to be used by the GetTest() function.
-*/
-void FEditorAutomationTestUtilities::CollectTestsByClass(UClass * Class, TArray<FString>& OutBeautifiedNames, TArray <FString>& OutTestCommands)
-{
-	FAssetRegistryModule& AssetRegistryModule = FModuleManager::LoadModuleChecked<FAssetRegistryModule>(TEXT("AssetRegistry"));
-	TArray<FAssetData> ObjectList;
-	AssetRegistryModule.Get().GetAssetsByClass(Class->GetFName(), ObjectList);
-
-	for (TObjectIterator<UClass> AllClassesIt; AllClassesIt; ++AllClassesIt)
-	{
-		UClass* ClassList = *AllClassesIt;
-		FName ClassName = ClassList->GetFName();
-	}
-
-	for (auto ObjIter = ObjectList.CreateConstIterator(); ObjIter; ++ObjIter)
-	{
-		const FAssetData& Asset = *ObjIter;
-		FString Filename = Asset.ObjectPath.ToString();
-		//convert to full paths
-		Filename = FPackageName::LongPackageNameToFilename(Filename);
-		if (FAutomationTestFramework::GetInstance().ShouldTestContent(Filename))
-		{
-			FString BeautifiedFilename = Asset.AssetName.ToString();
-			OutBeautifiedNames.Add(BeautifiedFilename);
-			OutTestCommands.Add(Asset.ObjectPath.ToString());
-		}
-	}
-}
-
-/**
-* Generates a list of assets from the GAME by a specific type.
-* This is to be used by the GetTest() function.
-*/
-void FEditorAutomationTestUtilities::CollectGameContentTestsByClass(UClass * Class, bool bRecursiveClass, TArray<FString>& OutBeautifiedNames, TArray <FString>& OutTestCommands)
-{
-	//Setting the Asset Registry
-	FAssetRegistryModule& AssetRegistryModule = FModuleManager::LoadModuleChecked<FAssetRegistryModule>(TEXT("AssetRegistry"));
-
-	//Variable setups
-	TArray<FAssetData> ObjectList;
-	FARFilter AssetFilter;
-
-	//Generating the list of assets.
-	//This list is being filtered by the game folder and class type.  The results are placed into the ObjectList variable.
-	AssetFilter.ClassNames.Add(Class->GetFName());
-
-	//removed path as a filter as it causes two large lists to be sorted.  Filtering on "game" directory on iteration
-	//AssetFilter.PackagePaths.Add("/Game");
-	AssetFilter.bRecursiveClasses = bRecursiveClass;
-	AssetFilter.bRecursivePaths = true;
-	AssetRegistryModule.Get().GetAssets(AssetFilter, ObjectList);
-
-	//Loop through the list of assets, make their path full and a string, then add them to the test.
-	for (auto ObjIter = ObjectList.CreateConstIterator(); ObjIter; ++ObjIter)
-	{
-		const FAssetData& Asset = *ObjIter;
-		FString Filename = Asset.ObjectPath.ToString();
-
-		if (Filename.StartsWith("/Game"))
-		{
-			//convert to full paths
-			Filename = FPackageName::LongPackageNameToFilename(Filename);
-			if (FAutomationTestFramework::GetInstance().ShouldTestContent(Filename))
-			{
-				FString BeautifiedFilename = Asset.AssetName.ToString();
-				OutBeautifiedNames.Add(BeautifiedFilename);
-				OutTestCommands.Add(Asset.ObjectPath.ToString());
-			}
-		}
-	}
-}
-
-/**
-* Generates a list of assets from the GAME by a specific type.
-* This is to be used by the GetTest() function.
-*/
-void FEditorAutomationTestUtilities::CollectGameContentTests(TArray<FString>& OutBeautifiedNames, TArray <FString>& OutTestCommands)
-{
-	//Setting the Asset Registry
-	FAssetRegistryModule& AssetRegistryModule = FModuleManager::LoadModuleChecked<FAssetRegistryModule>(TEXT("AssetRegistry"));
-
-	//Variable setups
-	TArray<FAssetData> ObjectList;
-	FARFilter AssetFilter;
-	
-	//removed path as a filter as it causes two large lists to be sorted.  Filtering on "game" directory on iteration
-	AssetFilter.PackagePaths.Add("/Game");
-	AssetFilter.bRecursiveClasses = true;
-	AssetFilter.bRecursivePaths = true;
-	AssetRegistryModule.Get().GetAssets(AssetFilter, ObjectList);
-
-	//Loop through the list of assets, make their path full and a string, then add them to the test.
-	for (auto ObjIter = ObjectList.CreateConstIterator(); ObjIter; ++ObjIter)
-	{
-		const FAssetData& Asset = *ObjIter;
-		if (Asset.GetClass() == nullptr)
-		{
-			// a nullptr class is bad !
-			UE_LOG(LogAutomationEditorCommon, Warning, TEXT("GetClass for %s (%s) returned nullptr. Asset ignored"), *Asset.AssetName.ToString(), *Asset.ObjectPath.ToString());
-		}
-		else 
-		{
-			FString Filename = Asset.ObjectPath.ToString();
-
-			if (Filename.StartsWith("/Game"))
-			{
-				//convert to full paths
-				Filename = FPackageName::LongPackageNameToFilename(Filename);
-				if (FAutomationTestFramework::GetInstance().ShouldTestContent(Filename))
-				{
-					FString BeautifiedFilename = FString::Printf(TEXT("%s.%s"), *Asset.GetClass()->GetFName().ToString(), *Asset.AssetName.ToString());
-					OutBeautifiedNames.Add(BeautifiedFilename);
-					OutTestCommands.Add(Asset.ObjectPath.ToString());
-				}
-			}
-		}
-	}
 }
\ No newline at end of file
diff --git a/Engine/Source/Editor/UnrealEd/Private/UnrealEdEngine.cpp b/Engine/Source/Editor/UnrealEd/Private/UnrealEdEngine.cpp
--- a/Engine/Source/Editor/UnrealEd/Private/UnrealEdEngine.cpp
+++ b/Engine/Source/Editor/UnrealEd/Private/UnrealEdEngine.cpp
@@ -1258,7 +1258,6 @@ void UUnrealEdEngine::FixAnyInvertedBrushes(UWorld* World)
 		}
 	}
 
-	bool bAnyStaticBrushesFixed = false;
 	if (Brushes.Num() > 0)
 	{
 		for (ABrush* Brush : Brushes)
@@ -1275,7 +1274,7 @@ void UUnrealEdEngine::FixAnyInvertedBrushes(UWorld* World)
 			if (Brush->IsStaticBrush())
 			{
 				// Static brushes require a full BSP rebuild
-				bAnyStaticBrushesFixed = true;
+				ABrush::SetNeedRebuild(Brush->GetLevel());
 			}
 			else
 			{
@@ -1286,11 +1285,6 @@ void UUnrealEdEngine::FixAnyInvertedBrushes(UWorld* World)
 
 			Brush->MarkPackageDirty();
 		}
-
-		if (bAnyStaticBrushesFixed)
-		{
-			RebuildAlteredBSP();
-		}
 	}
 }
 
diff --git a/Engine/Source/Editor/UnrealEd/Private/UnrealEdSrv.cpp b/Engine/Source/Editor/UnrealEd/Private/UnrealEdSrv.cpp
--- a/Engine/Source/Editor/UnrealEd/Private/UnrealEdSrv.cpp
+++ b/Engine/Source/Editor/UnrealEd/Private/UnrealEdSrv.cpp
@@ -35,8 +35,6 @@
 #include "LandscapeComponent.h"
 #include "LandscapeHeightfieldCollisionComponent.h"
 #include "ComponentReregisterContext.h"
-#include "JsonInternationalizationArchiveSerializer.h"
-#include "JsonInternationalizationManifestSerializer.h"
 #include "Engine/DestructibleMesh.h"
 #include "NotificationManager.h"
 #include "SNotificationList.h"
diff --git a/Engine/Source/Editor/UnrealEd/Public/AssetSelection.h b/Engine/Source/Editor/UnrealEd/Public/AssetSelection.h
--- a/Engine/Source/Editor/UnrealEd/Public/AssetSelection.h
+++ b/Engine/Source/Editor/UnrealEd/Public/AssetSelection.h
@@ -208,7 +208,7 @@ namespace AssetSelectionUtils
 	UNREALED_API bool IsBuilderBrushSelected();
 };
 
-class FActorFactoryAssetProxy
+class UNREALED_API FActorFactoryAssetProxy
 {
 public:
 
@@ -237,7 +237,7 @@ public:
 	 * @param	OutMenuItems				receives the list of menu items to use for populating an actor factory menu.
 	 * @param	ExcludeStandAloneFactories	if true, only factories that can create actors with the currently selected assets will be added
 	 */
-	UNREALED_API static void GenerateActorFactoryMenuItems( const FAssetData& AssetData, TArray<FMenuItem>* OutMenuItems, bool ExcludeStandAloneFactories );
+	static void GenerateActorFactoryMenuItems( const FAssetData& AssetData, TArray<FMenuItem>* OutMenuItems, bool ExcludeStandAloneFactories );
 
 	/**
 	 * Find the appropriate actor factory for an asset by type.
@@ -248,7 +248,7 @@ public:
 	 *
 	 * @return	the factory that is responsible for creating actors for the specified asset type.
 	 */
-	UNREALED_API static UActorFactory* GetFactoryForAsset( const FAssetData& DropData, bool bRequireValidObject=false );
+	static UActorFactory* GetFactoryForAsset( const FAssetData& DropData, bool bRequireValidObject=false );
 
 	/**
 	 * Find the appropriate actor factory for an asset.
@@ -257,7 +257,7 @@ public:
 	 *
 	 * @return	The factory that is responsible for creating actors for the specified asset
 	 */
-	UNREALED_API static UActorFactory* GetFactoryForAssetObject( UObject* AssetObj );
+	static UActorFactory* GetFactoryForAssetObject( UObject* AssetObj );
 
 	/**
 	 * Places an actor instance using the factory appropriate for the type of asset
diff --git a/Engine/Source/Editor/UnrealEd/Public/ComponentVisualizer.h b/Engine/Source/Editor/UnrealEd/Public/ComponentVisualizer.h
--- a/Engine/Source/Editor/UnrealEd/Public/ComponentVisualizer.h
+++ b/Engine/Source/Editor/UnrealEd/Public/ComponentVisualizer.h
@@ -39,7 +39,7 @@ public:
 	/** Draw HUD on viewport for the supplied component */
 	virtual void DrawVisualizationHUD(const UActorComponent* Component, const FViewport* Viewport, const FSceneView* View, FCanvas* Canvas) {}
 	/** */
-	virtual bool VisProxyHandleClick(FLevelEditorViewportClient* InViewportClient, HComponentVisProxy* VisProxy, const FViewportClick& Click) { return false; }
+	virtual bool VisProxyHandleClick(FEditorViewportClient* InViewportClient, HComponentVisProxy* VisProxy, const FViewportClick& Click) { return false; }
 	/** */
 	virtual void EndEditing() {}
 	/** */
@@ -52,6 +52,8 @@ public:
 	virtual bool HandleInputKey(FEditorViewportClient* ViewportClient,FViewport* Viewport,FKey Key,EInputEvent Event) { return false; }
 	/** */
 	virtual TSharedPtr<SWidget> GenerateContextMenu() const { return TSharedPtr<SWidget>(); }
+	/** */
+	virtual bool IsVisualizingArchetype() const { return false; }
 
 	struct FPropertyNameAndIndex
 	{
@@ -79,8 +81,15 @@ public:
 
 	/** Find the name of the property that points to this component */
 	static FPropertyNameAndIndex GetComponentPropertyName(const UActorComponent* Component);
+
 	/** Get a component pointer from the property name */
 	static UActorComponent* GetComponentFromPropertyName(const AActor* CompOwner, const FPropertyNameAndIndex& Property);
+
+	/** Notify that a component property has been modified */
+	static void NotifyPropertyModified(UActorComponent* Component, UProperty* Property);
+
+	/** Notify that many component properties have been modified */
+	static void NotifyPropertiesModified(UActorComponent* Component, const TArray<UProperty*>& Properties);
 };
 
 struct FCachedComponentVisualizer
diff --git a/Engine/Source/Editor/UnrealEd/Public/ComponentVisualizerManager.h b/Engine/Source/Editor/UnrealEd/Public/ComponentVisualizerManager.h
--- a/Engine/Source/Editor/UnrealEd/Public/ComponentVisualizerManager.h
+++ b/Engine/Source/Editor/UnrealEd/Public/ComponentVisualizerManager.h
@@ -15,24 +15,35 @@ public:
 
 
 	/** Activate a component visualizer given a clicked proxy */
-	bool HandleProxyForComponentVis(FLevelEditorViewportClient* InViewportClient, HHitProxy *HitProxy, const FViewportClick &Click);
+	bool HandleProxyForComponentVis(FEditorViewportClient* InViewportClient, HHitProxy *HitProxy, const FViewportClick &Click);
+
 	/** Clear active component visualizer */
 	void ClearActiveComponentVis();
 
 	/** Handle a click on the specified level editor viewport client */
-	bool HandleClick(FLevelEditorViewportClient* InViewportClient, HHitProxy *HitProxy, const FViewportClick &Click);
+	bool HandleClick(FEditorViewportClient* InViewportClient, HHitProxy *HitProxy, const FViewportClick &Click);
+
 	/** Pass key input to active visualizer */
 	bool HandleInputKey(FEditorViewportClient* InViewportClient, FViewport* Viewport, FKey Key, EInputEvent Event) const;
+
 	/** Pass delta input to active visualizer */
 	bool HandleInputDelta(FEditorViewportClient* InViewportClient, FViewport* InViewport, FVector& InDrag, FRotator& InRot, FVector& InScale) const;
+
 	/** Get widget location from active visualizer */
 	bool GetWidgetLocation(const FEditorViewportClient* InViewportClient, FVector& OutLocation) const;
+
 	/** Get custom widget coordinate system from active visualizer */
 	bool GetCustomInputCoordinateSystem(const FEditorViewportClient* InViewportClient, FMatrix& OutMatrix) const;
 
 	/** Generate context menu for the component visualizer */
 	TSharedPtr<SWidget> GenerateContextMenuForComponentVis() const;
 
+	/** Returns whether there is currently an active visualizer */
+	bool IsActive() const;
+
+	/** Returns whether the component being visualized is an archetype or not */
+	bool IsVisualizingArchetype() const;
+
 private:
 	/** Currently 'active' visualizer that we should pass input to etc */
 	TWeakPtr<class FComponentVisualizer> EditedVisualizerPtr;
diff --git a/Engine/Source/Editor/UnrealEd/Public/EditorViewportClient.h b/Engine/Source/Editor/UnrealEd/Public/EditorViewportClient.h
--- a/Engine/Source/Editor/UnrealEd/Public/EditorViewportClient.h
+++ b/Engine/Source/Editor/UnrealEd/Public/EditorViewportClient.h
@@ -37,7 +37,8 @@ namespace EDragTool
 	{
 		BoxSelect,
 		FrustumSelect,
-		Measure
+		Measure,
+		ViewportChange
 	};
 }
 
@@ -812,6 +813,16 @@ public:
 	 */
 	virtual void SetViewportType( ELevelViewportType InViewportType );
 	
+	/**
+	 * Rotate through viewport view options
+	 */
+	virtual void RotateViewportType();
+
+	/**
+	* @return If the viewport option in the array is the active viewport type
+	*/
+	bool IsActiveViewportTypeInRotation() const;
+
 	/**
 	 * @return If InViewportType is the active viewport type 
 	 */
diff --git a/Engine/Source/Editor/UnrealEd/Public/EditorViewportCommands.h b/Engine/Source/Editor/UnrealEd/Public/EditorViewportCommands.h
--- a/Engine/Source/Editor/UnrealEd/Public/EditorViewportCommands.h
+++ b/Engine/Source/Editor/UnrealEd/Public/EditorViewportCommands.h
@@ -41,6 +41,9 @@ public:
 	/** Changes the viewport to back view */
 	TSharedPtr< FUICommandInfo > Back;
 
+	/** Rotate through viewport view options */
+	TSharedPtr< FUICommandInfo > Next;
+
 	/** Changes the viewport to wireframe */
 	TSharedPtr< FUICommandInfo > WireframeMode;
 
diff --git a/Engine/Source/Editor/UnrealEd/Public/FileHelpers.h b/Engine/Source/Editor/UnrealEd/Public/FileHelpers.h
--- a/Engine/Source/Editor/UnrealEd/Public/FileHelpers.h
+++ b/Engine/Source/Editor/UnrealEd/Public/FileHelpers.h
@@ -431,7 +431,14 @@ public:
 	 * @param	ObjectPath		The path to the package to test
 	 * @return					The package name from the string
 	 */
-	UNREALED_API static FString ExtractPackageName(const FString& ObjectPath);
+	UNREALED_API static FString ExtractPackageName(const FString& ObjectPath);\
+
+	////////////////////////////////////////////////////////////////////////////
+	// File
+
+	UNREALED_API static FString GetFilename(const FName& PackageName);
+
+	UNREALED_API static FString GetFilename(UObject* LevelObject);
 
 private:
 
diff --git a/Engine/Source/Editor/UnrealEd/Public/Kismet2/KismetEditorUtilities.h b/Engine/Source/Editor/UnrealEd/Public/Kismet2/KismetEditorUtilities.h
--- a/Engine/Source/Editor/UnrealEd/Public/Kismet2/KismetEditorUtilities.h
+++ b/Engine/Source/Editor/UnrealEd/Public/Kismet2/KismetEditorUtilities.h
@@ -276,6 +276,9 @@ public:
 	/** Unregisters a class from having AutoGeneratedDefaultEvent nodes or callbacks for OnBlueprintCreated */
 	static void UnregisterAutoBlueprintNodeCreation(void* InOwner);
 
+	/** Add InNode to selection of editor */
+	static void AddToSelection(const class UEdGraph* Graph, UEdGraphNode* InNode);
+
 	/** Get IBlueprintEditor for given object, if it exists */
 	static TSharedPtr<class IBlueprintEditor> GetIBlueprintEditorForObject(const UObject* ObjectToFocusOn, bool bOpenEditor);
 private:
diff --git a/Engine/Source/Editor/UnrealEd/Public/LevelEditorViewport.h b/Engine/Source/Editor/UnrealEd/Public/LevelEditorViewport.h
--- a/Engine/Source/Editor/UnrealEd/Public/LevelEditorViewport.h
+++ b/Engine/Source/Editor/UnrealEd/Public/LevelEditorViewport.h
@@ -147,6 +147,7 @@ public:
 	virtual void TrackingStarted( const struct FInputEventState& InInputState, bool bIsDraggingWidget, bool bNudge ) override;
 	virtual void TrackingStopped() override;
 	virtual void AbortTracking() override;
+	virtual FWidget::EWidgetMode GetWidgetMode() const override;
 	virtual FVector GetWidgetLocation() const override;
 	virtual FMatrix GetWidgetCoordSystem() const override;
 	virtual void SetupViewForRendering( FSceneViewFamily& ViewFamily, FSceneView& View ) override;
@@ -533,6 +534,11 @@ public:
 
 	void UpdateHoveredObjects( const TSet<FViewportHoverTarget>& NewHoveredObjects );
 
+	/**
+	 * Calling SetViewportType from Dragtool_ViewportChange
+	 */
+	void SetViewportTypeFromTool(ELevelViewportType InViewportType);
+
 	/**
 	 * Static: Attempts to place the specified object in the level, returning one or more newly-created actors if successful.
 	 * IMPORTANT: The placed actor's location must be first set using GEditor->ClickLocation and GEditor->ClickPlane.
@@ -583,6 +589,7 @@ protected:
 	virtual void UpdateLinkedOrthoViewports( bool bInvalidate = false ) override;
 	virtual ELevelViewportType GetViewportType() const override;
 	virtual void SetViewportType( ELevelViewportType InViewportType ) override;
+	virtual void RotateViewportType() override;
 	virtual void OverridePostProcessSettings( FSceneView& View ) override;
 	virtual void PerspectiveCameraMoved() override;
 	virtual bool ShouldLockPitch() const override;
diff --git a/Engine/Source/Editor/UnrealEd/Public/PackageHelperFunctions.h b/Engine/Source/Editor/UnrealEd/Public/PackageHelperFunctions.h
--- a/Engine/Source/Editor/UnrealEd/Public/PackageHelperFunctions.h
+++ b/Engine/Source/Editor/UnrealEd/Public/PackageHelperFunctions.h
@@ -47,7 +47,7 @@ void SearchDirectoryRecursive( const FString& SearchPathMask, TArray<FString>& o
  *
  * @return	true if packages were found successfully, false otherwise.
  */
-bool NormalizePackageNames( TArray<FString> PackageNames, TArray<FString>& PackagePathNames, const FString& PackageWildcard=FString(TEXT("*.*")), uint8 PackageFilter=NORMALIZE_DefaultFlags );
+bool UNREALED_API NormalizePackageNames( TArray<FString> PackageNames, TArray<FString>& PackagePathNames, const FString& PackageWildcard=FString(TEXT("*.*")), uint8 PackageFilter=NORMALIZE_DefaultFlags );
 
 
 /** 
@@ -63,7 +63,7 @@ bool NormalizePackageNames( TArray<FString> PackageNames, TArray<FString>& Packa
  *
  * @return true if successful
  */
-bool SavePackageHelper(UPackage* Package, FString Filename,  EObjectFlags KeepObjectFlags = RF_Standalone, FOutputDevice* ErrorDevice=GWarn, FLinkerLoad* LinkerToConformAgainst=NULL, ESaveFlags SaveFlags = SAVE_None);
+bool UNREALED_API SavePackageHelper(UPackage* Package, FString Filename,  EObjectFlags KeepObjectFlags = RF_Standalone, FOutputDevice* ErrorDevice=GWarn, FLinkerLoad* LinkerToConformAgainst=NULL, ESaveFlags SaveFlags = SAVE_None);
 
 
 /**
diff --git a/Engine/Source/Editor/UnrealEd/Public/Tests/AutomationEditorCommon.h b/Engine/Source/Editor/UnrealEd/Public/Tests/AutomationEditorCommon.h
--- a/Engine/Source/Editor/UnrealEd/Public/Tests/AutomationEditorCommon.h
+++ b/Engine/Source/Editor/UnrealEd/Public/Tests/AutomationEditorCommon.h
@@ -7,41 +7,18 @@
 #include "AutomationCommon.h"
 #include "Tests/AutomationTestSettings.h"
 
-namespace AutomationEditorCommonUtils
+//////////////////////////////////////////////////////////////////////////
+// FAutomationEditorCommonUtils
+class UNREALED_API FAutomationEditorCommonUtils
 {
+public:
+	
 	/**
 	* Creates a new map for editing.  Also clears editor tools that could cause issues when changing maps
 	*
 	* @return - The UWorld for the new map
 	*/
-	static UWorld* CreateNewMap()
-	{
-		// Change out of Matinee when opening new map, so we avoid editing data in the old one.
-		if (GLevelEditorModeTools().IsModeActive(FBuiltinEditorModes::EM_InterpEdit))
-		{
-			GLevelEditorModeTools().DeactivateMode(FBuiltinEditorModes::EM_InterpEdit);
-		}
-
-		// Also change out of Landscape mode to ensure all references are cleared.
-		if (GLevelEditorModeTools().IsModeActive(FBuiltinEditorModes::EM_Landscape))
-		{
-			GLevelEditorModeTools().DeactivateMode(FBuiltinEditorModes::EM_Landscape);
-		}
-
-		// Also change out of Foliage mode to ensure all references are cleared.
-		if (GLevelEditorModeTools().IsModeActive(FBuiltinEditorModes::EM_Foliage))
-		{
-			GLevelEditorModeTools().DeactivateMode(FBuiltinEditorModes::EM_Foliage);
-		}
-
-		// Change out of mesh paint mode when opening a new map.
-		if (GLevelEditorModeTools().IsModeActive(FBuiltinEditorModes::EM_MeshPaint))
-		{
-			GLevelEditorModeTools().DeactivateMode(FBuiltinEditorModes::EM_MeshPaint);
-		}
-
-		return GEditor->NewMap();
-	}
+	static UWorld* CreateNewMap();
 
 	/**
 	* Imports an object using a given factory
@@ -51,21 +28,21 @@ namespace AutomationEditorCommonUtils
 	* @param PackagePath - The full path of the package file to create
 	* @param ImportPath - The path to the object to import
 	*/
-	UObject* ImportAssetUsingFactory(UFactory* ImportFactory, const FString& ObjectName, const FString& PackageName, const FString& ImportPath);
+	static UObject* ImportAssetUsingFactory(UFactory* ImportFactory, const FString& ObjectName, const FString& PackageName, const FString& ImportPath);
 
 	/**
 	* Nulls out references to a given object
 	*
 	* @param InObject - Object to null references to
 	*/
-	UNREALED_API void NullReferencesToObject(UObject* InObject);
+	static void NullReferencesToObject(UObject* InObject);
 
 	/**
 	* gets a factory class based off an asset file extension
 	*
 	* @param AssetExtension - The file extension to use to find a supporting UFactory
 	*/
-	UClass* GetFactoryClassForType(const FString& AssetExtension);
+	static UClass* GetFactoryClassForType(const FString& AssetExtension);
 
 	/**
 	* Applies settings to an object by finding UProperties by name and calling ImportText
@@ -74,7 +51,7 @@ namespace AutomationEditorCommonUtils
 	* @param PropertyChain - The list UProperty names recursively to search through
 	* @param Value - The value to import on the found property
 	*/
-	void ApplyCustomFactorySetting(UObject* InObject, TArray<FString>& PropertyChain, const FString& Value);
+	static void ApplyCustomFactorySetting(UObject* InObject, TArray<FString>& PropertyChain, const FString& Value);
 
 	/**
 	* Applies the custom factory settings
@@ -82,7 +59,7 @@ namespace AutomationEditorCommonUtils
 	* @param InFactory - The factory to apply custom settings to
 	* @param FactorySettings - An array of custom settings to apply to the factory
 	*/
-	void ApplyCustomFactorySettings(UFactory* InFactory, const TArray<FImportFactorySettingValues>& FactorySettings);
+	static void ApplyCustomFactorySettings(UFactory* InFactory, const TArray<FImportFactorySettingValues>& FactorySettings);
 
 	/**
 	* Writes a number to a text file.
@@ -93,7 +70,7 @@ namespace AutomationEditorCommonUtils
 	* @param InNumberToBeWritten - is the float number that is expected to be written to the file.
 	* @param Delimiter - is the delimiter to be used. TEXT(",")
 	*/
-	void WriteToTextFile(const FString& InTestName, const FString& InTestItem, const FString& InFileName, const float& InEntry, const FString& Delimiter);
+	static void WriteToTextFile(const FString& InTestName, const FString& InTestItem, const FString& InFileName, const float& InEntry, const FString& Delimiter);
 
 	/**
 	* Returns the sum of the numbers available in an array of float.
@@ -101,14 +78,14 @@ namespace AutomationEditorCommonUtils
 	* @param InFloatArray - is the name of the array intended to be used.
 	* @param bisAveragedInstead - will return the average of the available numbers instead of the sum.
 	*/
-	float TotalFromFloatArray(const TArray<float>& InFloatArray, bool bisAveragedInstead);
+	static float TotalFromFloatArray(const TArray<float>& InFloatArray, bool bisAveragedInstead);
 
 	/**
 	* Returns the largest value from an array of float numbers.
 
 	* @param InFloatArray - is the name of the array intended to be used.
 	*/
-	float LargestValueInFloatArray(const TArray<float>& InFloatArray);
+	static float LargestValueInFloatArray(const TArray<float>& InFloatArray);
 
 	/**
 	* Returns the contents of a text file as an array of FString.
@@ -116,7 +93,7 @@ namespace AutomationEditorCommonUtils
 	* @param InFileLocation - is the location of the file.
 	* @param OutArray - The name of the array that will store the data.
 	*/
-	void CreateArrayFromFile(const FString& InFileLocation, TArray<FString>& OutArray);
+	static void CreateArrayFromFile(const FString& InFileLocation, TArray<FString>& OutArray);
 
 	/**
 	* Returns true if the archive/file can be written to otherwise false.
@@ -124,8 +101,8 @@ namespace AutomationEditorCommonUtils
 	* @param InFilePath - is the location of the file.
 	* @param InArchiveName - is the name of the archive to be used.
 	*/
-	bool IsArchiveWriteable(const FString& InFilePath, const FArchive* InArchiveName);
-		
+	static bool IsArchiveWriteable(const FString& InFilePath, const FArchive* InArchiveName);
+
 	/**
 	* Returns the first DeviceID 'Platform@Device'
 	* Searches the automation preferences for the platform to use.
@@ -133,7 +110,7 @@ namespace AutomationEditorCommonUtils
 	* @param OutDeviceID - The variable that will hold the device ID.
 	* @param InMapName - The map name to check against in the automation preferences.
 	*/
-	void GetLaunchOnDeviceID(FString& OutDeviceID, const FString& InMapName);
+	static void GetLaunchOnDeviceID(FString& OutDeviceID, const FString& InMapName);
 
 	/**
 	* Returns the DeviceID 'Platform@Device'
@@ -141,9 +118,9 @@ namespace AutomationEditorCommonUtils
 
 	* @param OutDeviceID - The variable that will hold the device ID.
 	* @param InMapName - The map name to check against in the automation preferences.
-	* @param InDeviceName - Device Name 
+	* @param InDeviceName - Device Name
 	*/
-	void GetLaunchOnDeviceID(FString& OutDeviceID, const FString& InMapName, const FString& InDeviceName);
+	static void GetLaunchOnDeviceID(FString& OutDeviceID, const FString& InMapName, const FString& InDeviceName);
 
 	/**
 	* Sets the first found ortho viewport camera to the desired location and rotation.
@@ -151,8 +128,52 @@ namespace AutomationEditorCommonUtils
 	* @param ViewLocation - Desired location for the viewport view.
 	* @param ViewRotation - Desired rotation of the viewport view.
 	*/
-	bool SetOrthoViewportView(const FVector& ViewLocation, const FRotator& ViewRotation);
-}
+	static bool SetOrthoViewportView(const FVector& ViewLocation, const FRotator& ViewRotation);
+
+	/**
+	* Converts a package path to an asset path
+	*
+	* @param PackagePath - The package path to convert
+	*/
+	static FString ConvertPackagePathToAssetPath(const FString& PackagePath);
+
+	/**
+	* Gets the asset data from a package path
+	*
+	* @param PackagePath - The package path used to look up the asset data
+	*/
+	static FAssetData GetAssetDataFromPackagePath(const FString& PackagePath);
+
+	/**
+	* Loads the map specified by an automation test
+	*
+	* @param MapName - Map to load
+	*/
+	static void LoadMap(const FString& MapName);
+
+	/**
+	* Run PIE
+	*/
+	static void RunPIE();
+	
+	/**
+	* Generates a list of assets from the ENGINE and the GAME by a specific type.
+	* This is to be used by the GetTest() function.
+	*/
+	static void CollectTestsByClass(UClass * Class, TArray<FString>& OutBeautifiedNames, TArray <FString>& OutTestCommands);
+	
+	/**
+	* Generates a list of assets from the GAME by a specific type.
+	* This is to be used by the GetTest() function.
+	*/
+	static void CollectGameContentTestsByClass(UClass * Class, bool bRecursiveClass, TArray<FString>& OutBeautifiedNames, TArray <FString>& OutTestCommands);
+		
+	/**
+	* Generates a list of assets from the GAME by a specific type.
+	* This is to be used by the GetTest() function.
+	*/
+	static void CollectGameContentTests(TArray<FString>& OutBeautifiedNames, TArray <FString>& OutTestCommands);
+};
 
 
 //////////////////////////////////////////////////////////////////////////
@@ -162,148 +183,74 @@ namespace AutomationEditorCommonUtils
 * Creates a latent command which the user can either undo or redo an action.
 * True will trigger Undo, False will trigger a redo
 */
-DEFINE_LATENT_AUTOMATION_COMMAND_ONE_PARAMETER(FUndoRedoCommand, bool, bUndo);
+DEFINE_EXPORTED_LATENT_AUTOMATION_COMMAND_ONE_PARAMETER(UNREALED_API, FUndoRedoCommand, bool, bUndo);
 
 /**
 * Open editor for a particular asset
 */
-DEFINE_LATENT_AUTOMATION_COMMAND_ONE_PARAMETER(FOpenEditorForAssetCommand, FString, AssetName);
+DEFINE_EXPORTED_LATENT_AUTOMATION_COMMAND_ONE_PARAMETER(UNREALED_API, FOpenEditorForAssetCommand, FString, AssetName);
 
 /**
 * Close all asset editors
 */
-DEFINE_LATENT_AUTOMATION_COMMAND(FCloseAllAssetEditorsCommand);
+DEFINE_EXPORTED_LATENT_AUTOMATION_COMMAND(UNREALED_API, FCloseAllAssetEditorsCommand);
 
 /**
 * Start PIE session
 */
-DEFINE_LATENT_AUTOMATION_COMMAND_ONE_PARAMETER(FStartPIECommand, bool, bSimulateInEditor);
+DEFINE_EXPORTED_LATENT_AUTOMATION_COMMAND_ONE_PARAMETER(UNREALED_API, FStartPIECommand, bool, bSimulateInEditor);
 
 /**
 * End PlayMap session
 */
-DEFINE_LATENT_AUTOMATION_COMMAND(FEndPlayMapCommand);
+DEFINE_EXPORTED_LATENT_AUTOMATION_COMMAND(UNREALED_API, FEndPlayMapCommand);
 
 /**
 * Loads a map
 */
-DEFINE_LATENT_AUTOMATION_COMMAND_ONE_PARAMETER(FEditorLoadMap, FString, MapName);
+DEFINE_EXPORTED_LATENT_AUTOMATION_COMMAND_ONE_PARAMETER(UNREALED_API, FEditorLoadMap, FString, MapName);
 
 /**
 * Waits for shaders to finish compiling before moving on to the next thing.
 */
-DEFINE_LATENT_AUTOMATION_COMMAND(FWaitForShadersToFinishCompiling);
+DEFINE_EXPORTED_LATENT_AUTOMATION_COMMAND(UNREALED_API, FWaitForShadersToFinishCompiling);
 
 /**
 * Latent command that changes the editor viewport to the first available bookmarked view.
 */
-DEFINE_LATENT_AUTOMATION_COMMAND(FChangeViewportToFirstAvailableBookmarkCommand);
+DEFINE_EXPORTED_LATENT_AUTOMATION_COMMAND(UNREALED_API, FChangeViewportToFirstAvailableBookmarkCommand);
 
 /**
 * Latent command that adds a static mesh to the worlds origin.
 */
-DEFINE_LATENT_AUTOMATION_COMMAND(FAddStaticMeshCommand);
+DEFINE_EXPORTED_LATENT_AUTOMATION_COMMAND(UNREALED_API, FAddStaticMeshCommand);
 
 /**
 * Latent command that builds the lighting for the current level.
 */
-DEFINE_LATENT_AUTOMATION_COMMAND(FBuildLightingCommand);
+DEFINE_EXPORTED_LATENT_AUTOMATION_COMMAND(UNREALED_API, FBuildLightingCommand);
 
 /**
 * Latent command that saves a copy of the level to a transient folder.
 */
-DEFINE_LATENT_AUTOMATION_COMMAND_ONE_PARAMETER(FSaveLevelCommand, FString, MapName);
+DEFINE_EXPORTED_LATENT_AUTOMATION_COMMAND_ONE_PARAMETER(UNREALED_API, FSaveLevelCommand, FString, MapName);
 
 /**
 * Triggers a launch on using a specified device ID.
 */
-DEFINE_LATENT_AUTOMATION_COMMAND_ONE_PARAMETER(FLaunchOnCommand, FString, InLauncherDeviceID);
+DEFINE_EXPORTED_LATENT_AUTOMATION_COMMAND_ONE_PARAMETER(UNREALED_API, FLaunchOnCommand, FString, InLauncherDeviceID);
 
 /**
 * Wait for a cook by the book to finish.  Will time out after 3600 seconds (1 hr).
 */
-DEFINE_LATENT_AUTOMATION_COMMAND(FWaitToFinishCookByTheBookCommand);
+DEFINE_EXPORTED_LATENT_AUTOMATION_COMMAND(UNREALED_API, FWaitToFinishCookByTheBookCommand);
 
 /**
 * Wait for a build and deploy to finish before moving on.
 */
-DEFINE_LATENT_AUTOMATION_COMMAND(FWaitToFinishBuildDeployCommand);
+DEFINE_EXPORTED_LATENT_AUTOMATION_COMMAND(UNREALED_API, FWaitToFinishBuildDeployCommand);
 
 /**
 * Latent command to delete a directory.
 */
-DEFINE_LATENT_AUTOMATION_COMMAND_ONE_PARAMETER(FDeleteDirCommand, FString, InFolderLocation);
-
-
-//////////////////////////////////////////////////////////////////////////
-// FEditorAutomationTestUtilities
-class FEditorAutomationTestUtilities
-{
-public:
-	
-	/**
-	* Converts a package path to an asset path
-	*
-	* @param PackagePath - The package path to convert
-	*/
-	UNREALED_API static FString ConvertPackagePathToAssetPath(const FString& PackagePath);
-
-	/**
-	* Gets the asset data from a package path
-	*
-	* @param PackagePath - The package path used to look up the asset data
-	*/
-	UNREALED_API static FAssetData GetAssetDataFromPackagePath(const FString& PackagePath);
-
-	/**
-	* Loads the map specified by an automation test
-	*
-	* @param MapName - Map to load
-	*/
-	static void LoadMap(const FString& MapName)
-	{
-		bool bLoadAsTemplate = false;
-		bool bShowProgress = false;
-		FEditorFileUtils::LoadMap(MapName, bLoadAsTemplate, bShowProgress);
-	}
-
-	/**
-	* Run PIE
-	*/
-	static void RunPIE()
-	{
-		bool bInSimulateInEditor = true;
-		//once in the editor
-		ADD_LATENT_AUTOMATION_COMMAND(FStartPIECommand(true));
-		ADD_LATENT_AUTOMATION_COMMAND(FWaitLatentCommand(3.0f));
-		ADD_LATENT_AUTOMATION_COMMAND(FEndPlayMapCommand());
-
-		//wait between tests
-		ADD_LATENT_AUTOMATION_COMMAND(FWaitLatentCommand(1.0f));
-
-		//once not in the editor
-		ADD_LATENT_AUTOMATION_COMMAND(FStartPIECommand(false));
-		ADD_LATENT_AUTOMATION_COMMAND(FWaitLatentCommand(3.0f));
-		ADD_LATENT_AUTOMATION_COMMAND(FEndPlayMapCommand());
-	}
-	
-	/**
-	* Generates a list of assets from the ENGINE and the GAME by a specific type.
-	* This is to be used by the GetTest() function.
-	*/
-	static void CollectTestsByClass(UClass * Class, TArray<FString>& OutBeautifiedNames, TArray <FString>& OutTestCommands);
-	
-	/**
-	* Generates a list of assets from the GAME by a specific type.
-	* This is to be used by the GetTest() function.
-	*/
-	static void CollectGameContentTestsByClass(UClass * Class, bool bRecursiveClass, TArray<FString>& OutBeautifiedNames, TArray <FString>& OutTestCommands);
-		
-	/**
-	* Generates a list of assets from the GAME by a specific type.
-	* This is to be used by the GetTest() function.
-	*/
-	static void CollectGameContentTests(TArray<FString>& OutBeautifiedNames, TArray <FString>& OutTestCommands);
-};
-
-
+DEFINE_EXPORTED_LATENT_AUTOMATION_COMMAND_ONE_PARAMETER(UNREALED_API, FDeleteDirCommand, FString, InFolderLocation);
diff --git a/Engine/Source/Editor/UnrealEd/Public/UnrealEd.h b/Engine/Source/Editor/UnrealEd/Public/UnrealEd.h
--- a/Engine/Source/Editor/UnrealEd/Public/UnrealEd.h
+++ b/Engine/Source/Editor/UnrealEd/Public/UnrealEd.h
@@ -214,7 +214,6 @@
 #include "Classes/Commandlets/GenerateTextLocalizationResourceCommandlet.h"
 #include "Classes/Commandlets/InternationalizationConditioningCommandlet.h"
 #include "Classes/Commandlets/InternationalizationExportCommandlet.h"
-#include "Classes/Commandlets/RepairLocalizationDataCommandlet.h"
 #include "Classes/Commandlets/GenerateBlueprintAPICommandlet.h"
 #include "Classes/Commandlets/GenerateDistillFileSetsCommandlet.h"
 #include "Classes/Editor/GroupActor.h"
diff --git a/Engine/Source/Editor/UnrealEd/UnrealEd.Build.cs b/Engine/Source/Editor/UnrealEd/UnrealEd.Build.cs
--- a/Engine/Source/Editor/UnrealEd/UnrealEd.Build.cs
+++ b/Engine/Source/Editor/UnrealEd/UnrealEd.Build.cs
@@ -86,8 +86,10 @@ public class UnrealEd : ModuleRules
 				"BlueprintGraph",
                 "Http",
 				"UnrealAudio",
-                "FunctionalTesting"
-            }
+                "FunctionalTesting",
+				"AutomationController",
+				"Internationalization",
+			}
 		);
 
 		PrivateDependencyModuleNames.AddRange(
@@ -102,7 +104,6 @@ public class UnrealEd : ModuleRules
 				"EngineSettings",
 				"InputCore",
 				"InputBindingEditor",
-				"Internationalization",
 				"LauncherServices",
 				"MaterialEditor",
 				"MessageLog",
@@ -215,8 +216,9 @@ public class UnrealEd : ModuleRules
                 "PortalServices",
                 "GeometryCacheEd",
                 "BlueprintNativeCodeGen",
-				"VREditor"
-			}
+				"VREditor",
+                "EditorAutomation",
+            }
 		);
 
 		if (Target.Platform == UnrealTargetPlatform.Win64 || Target.Platform == UnrealTargetPlatform.Win32 || Target.Platform == UnrealTargetPlatform.Mac)
diff --git a/Engine/Source/Editor/WorkspaceMenuStructure/Private/WorkspaceMenuStructureModule.cpp b/Engine/Source/Editor/WorkspaceMenuStructure/Private/WorkspaceMenuStructureModule.cpp
--- a/Engine/Source/Editor/WorkspaceMenuStructure/Private/WorkspaceMenuStructureModule.cpp
+++ b/Engine/Source/Editor/WorkspaceMenuStructure/Private/WorkspaceMenuStructureModule.cpp
@@ -61,6 +61,11 @@ public:
 		return DeveloperToolsMiscCategory.ToSharedRef();
 	}
 
+	virtual TSharedRef<FWorkspaceItem> GetAutomationCategory() const override
+	{
+		return AutomationCategory.ToSharedRef();
+	}
+
 	virtual TSharedRef<FWorkspaceItem> GetAutomationToolsCategory() const override
 	{
 		return AutomationToolsCategory.ToSharedRef();
@@ -91,7 +96,8 @@ public:
 		DeveloperToolsLogCategory = DeveloperToolsCategory->AddGroup(LOCTEXT("WorkspaceMenu_DeveloperToolsLogCategory", "Log"), FSlateIcon(), true);
 		DeveloperToolsMiscCategory = DeveloperToolsCategory->AddGroup(LOCTEXT("WorkspaceMenu_DeveloperToolsMiscCategory", "Miscellaneous"), FSlateIcon(), true);
 
-		AutomationToolsCategory = ToolsCategory->AddGroup(LOCTEXT("WorkspaceMenu_AutomationToolsCategory", "Automation Tools"), FSlateIcon(FEditorStyle::GetStyleSetName(), "AutomationTools.MenuIcon"));
+		AutomationCategory = ToolsCategory->AddGroup(LOCTEXT("WorkspaceMenu_AutomationCategory", "Automation"), FSlateIcon(FEditorStyle::GetStyleSetName(), "AutomationTools.MenuIcon"));
+		AutomationToolsCategory = AutomationCategory->AddGroup(LOCTEXT("WorkspaceMenu_AutomationToolsCategory", "Automation Tools"), FSlateIcon(), true);
 	}
 
 public:
@@ -120,6 +126,8 @@ private:
 	TSharedPtr<FWorkspaceItem> DeveloperToolsDebugCategory;
 	TSharedPtr<FWorkspaceItem> DeveloperToolsLogCategory;
 	TSharedPtr<FWorkspaceItem> DeveloperToolsMiscCategory;
+	
+	TSharedPtr<FWorkspaceItem> AutomationCategory;
 	TSharedPtr<FWorkspaceItem> AutomationToolsCategory;
 	
 	TSharedPtr<FWorkspaceItem> EditOptions;
diff --git a/Engine/Source/Editor/WorkspaceMenuStructure/Public/WorkspaceMenuStructure.h b/Engine/Source/Editor/WorkspaceMenuStructure/Public/WorkspaceMenuStructure.h
--- a/Engine/Source/Editor/WorkspaceMenuStructure/Public/WorkspaceMenuStructure.h
+++ b/Engine/Source/Editor/WorkspaceMenuStructure/Public/WorkspaceMenuStructure.h
@@ -16,6 +16,7 @@
  *  +-Tools Tabs
  *    |-Tab
  *    |-...
+ *    \ Automation >
  *    \ Developer Tools >
  */
 
@@ -51,6 +52,9 @@ public:
 	/** See diagram above */
 	virtual TSharedRef<FWorkspaceItem> GetDeveloperToolsMiscCategory() const = 0;
 
+	/** See diagram above */
+	virtual TSharedRef<FWorkspaceItem> GetAutomationCategory() const = 0;
+
 	/** See diagram above */
 	virtual TSharedRef<FWorkspaceItem> GetAutomationToolsCategory() const = 0;
 
diff --git a/Engine/Source/Programs/AutomationTool/Localization/LocalizationProvider.cs b/Engine/Source/Programs/AutomationTool/Localization/LocalizationProvider.cs
--- a/Engine/Source/Programs/AutomationTool/Localization/LocalizationProvider.cs
+++ b/Engine/Source/Programs/AutomationTool/Localization/LocalizationProvider.cs
@@ -43,7 +43,7 @@ public struct ProjectStepInfo
 	/** The name of this localization step */
 	public string Name;
 
-	/** Path to this steps localization config file (relative to the root working directory for the commandlet) */
+	/** Absolute path to this steps localization config file */
 	public string LocalizationConfigFile;
 };
 
@@ -67,12 +67,14 @@ public abstract class LocalizationProvider
 	public struct LocalizationProviderArgs
 	{
 		public string RootWorkingDirectory;
+		public string RemoteFilenamePrefix;
 		public CommandUtils CommandUtils;
 	};
 
 	public LocalizationProvider(LocalizationProviderArgs InArgs)
 	{
 		RootWorkingDirectory = InArgs.RootWorkingDirectory;
+		RemoteFilenamePrefix = InArgs.RemoteFilenamePrefix;
 		CommandUtils = InArgs.CommandUtils;
 
 		LocalizationBranchName = CommandUtils.ParseParamValue("LocalizationBranch");
@@ -154,6 +156,7 @@ public abstract class LocalizationProvider
 
 	protected string RootWorkingDirectory;
 	protected string LocalizationBranchName;
+	protected string RemoteFilenamePrefix;
 	protected bool bUploadAllCultures;
 	protected CommandUtils CommandUtils;
 
diff --git a/Engine/Source/Programs/AutomationTool/Scripts/CopyBuildToStagingDirectory.Automation.cs b/Engine/Source/Programs/AutomationTool/Scripts/CopyBuildToStagingDirectory.Automation.cs
--- a/Engine/Source/Programs/AutomationTool/Scripts/CopyBuildToStagingDirectory.Automation.cs
+++ b/Engine/Source/Programs/AutomationTool/Scripts/CopyBuildToStagingDirectory.Automation.cs
@@ -426,6 +426,46 @@ public partial class Project : CommandUtils
 				}
 			}
 
+			// Stage all plugin localization targets
+			{
+				ProjectDescriptor Project = ProjectDescriptor.FromFile(SC.RawProjectPath.FullName);
+
+				List<PluginInfo> AvailablePlugins = Plugins.ReadAvailablePlugins(new DirectoryReference(CombinePaths(SC.LocalRoot, "Engine")), new FileReference(CombinePaths(SC.ProjectRoot, Params.ShortProjectName + ".uproject")));
+				foreach (var Plugin in AvailablePlugins)
+				{
+					if (!UProjectInfo.IsPluginEnabledForProject(Plugin, Project, SC.StageTargetPlatform.PlatformType, TargetRules.TargetType.Game) &&
+						UProjectInfo.IsPluginEnabledForProject(Plugin, Project, SC.StageTargetPlatform.PlatformType, TargetRules.TargetType.Client))
+					{
+						// skip editor plugins
+						continue;
+					}
+
+					if (Plugin.Descriptor.LocalizationTargets == null || Plugin.Descriptor.LocalizationTargets.Length == 0)
+					{
+						// skip plugins with no localization targets
+						continue;
+					}
+
+					foreach (var LocalizationTarget in Plugin.Descriptor.LocalizationTargets)
+					{
+						if (LocalizationTarget.LoadingPolicy != LocalizationTargetDescriptorLoadingPolicy.Always && LocalizationTarget.LoadingPolicy != LocalizationTargetDescriptorLoadingPolicy.Game)
+						{
+							// skip targets not loaded by the game
+							continue;
+						}
+
+						var PluginLocTargetDirectory = CombinePaths(Plugin.Directory.FullName, "Content", "Localization", LocalizationTarget.Name);
+						if (DirectoryExists(PluginLocTargetDirectory))
+						{
+							foreach (string Culture in CulturesToStage)
+							{
+								StageLocalizationDataForCulture(SC, Culture, PluginLocTargetDirectory, null, !Params.UsePak(SC.StageTargetPlatform));
+							}
+						}
+					}
+				}
+			}
+
             // Stage any additional UFS and NonUFS paths specified in the project ini files; these dirs are relative to the game content directory
             if (PlatformGameConfig != null)
             {
