Fixed the PO exporter writing out invalid escape sequences

diff --git a/Engine/Source/Editor/UnrealEd/Private/Commandlets/InternationalizationExportCommandlet.cpp b/Engine/Source/Editor/UnrealEd/Private/Commandlets/InternationalizationExportCommandlet.cpp
--- a/Engine/Source/Editor/UnrealEd/Private/Commandlets/InternationalizationExportCommandlet.cpp
+++ b/Engine/Source/Editor/UnrealEd/Private/Commandlets/InternationalizationExportCommandlet.cpp
@@ -26,6 +26,21 @@ uint32 GetTypeHash(const FPortableObjectEntryIdentity& ID)
 
 namespace
 {
+	/**
+	 * Declarations
+	 */
+	FString ConditionIdentityForPOMsgCtxt(const FString& Namespace, const FString& Key, const TSharedPtr<FLocMetadataObject>& KeyMetaData);
+	void ParsePOMsgCtxtForIdentity(const FString& MsgCtxt, FString& OutNamespace, FString& OutKey);
+	FString ConditionArchiveStrForPo(const FString& InStr);
+	FString ConditionPoStringForArchive(const FString& InStr);
+	FString ConvertSrcLocationToPORef(const FString& InSrcLocation);
+	FString GetConditionedKeyForExtractedComment(const FString& Key);
+	FString GetConditionedReferenceForExtractedComment(const FString& PORefString);
+	FString GetConditionedInfoMetaDataForExtractedComment(const FString& KeyName, const FString& ValueString);
+
+	/**
+	 * Definitions
+	 */
 	FString ConditionIdentityForPOMsgCtxt(const FString& Namespace, const FString& Key, const TSharedPtr<FLocMetadataObject>& KeyMetaData)
 	{
 		const auto& EscapeMsgCtxtParticle = [](const FString& InStr) -> FString
@@ -45,11 +60,14 @@ namespace
 		const FString EscapedNamespace = EscapeMsgCtxtParticle(Namespace);
 		const FString EscapedKey = EscapeMsgCtxtParticle(Key);
 
-		return FString::Printf(TEXT("%s,%s"), *EscapedNamespace, *EscapedKey);
+		const FString MsgCtxt = FString::Printf(TEXT("%s,%s"), *EscapedNamespace, *EscapedKey);
+		return ConditionArchiveStrForPo(MsgCtxt);
 	}
 
 	void ParsePOMsgCtxtForIdentity(const FString& MsgCtxt, FString& OutNamespace, FString& OutKey)
 	{
+		const FString ConditionedMsgCtxt = ConditionPoStringForArchive(MsgCtxt);
+
 		static const int32 OutputBufferCount = 2;
 		FString* OutputBuffers[OutputBufferCount] = { &OutNamespace, &OutKey };
 		int32 OutputBufferIndex = 0;
@@ -89,12 +107,12 @@ namespace
 			}
 		};
 
-		for (const TCHAR C : MsgCtxt)
+		for (const TCHAR C : ConditionedMsgCtxt)
 		{
 			// If we're out of buffers, break out. The particle list is longer than expected.
 			if (OutputBufferIndex >= OutputBufferCount)
 			{
-				UE_LOG( LogInternationalizationExportCommandlet, Warning, TEXT("msgctxt found in PO has too many parts: %s"), *MsgCtxt );
+				UE_LOG( LogInternationalizationExportCommandlet, Warning, TEXT("msgctxt found in PO has too many parts: %s"), *ConditionedMsgCtxt );
 				break;
 			}
 
