Backing out CL 3093644 (from Dev-BP) as it ended up causing CIS cook errors.

diff --git a/Engine/Source/Editor/BlueprintGraph/Classes/K2Node_FunctionTerminator.h b/Engine/Source/Editor/BlueprintGraph/Classes/K2Node_FunctionTerminator.h
--- a/Engine/Source/Editor/BlueprintGraph/Classes/K2Node_FunctionTerminator.h
+++ b/Engine/Source/Editor/BlueprintGraph/Classes/K2Node_FunctionTerminator.h
@@ -39,7 +39,5 @@ class UK2Node_FunctionTerminator : public UK2Node_EditablePinBase
 
 	/** Promotes the node from being a part of an interface override to a full function that allows for parameter and result pin additions */
 	virtual void PromoteFromInterfaceOverride(bool bIsPrimaryTerminator = true);
-
-	static UClass* FindOriginalFunctionOwner(UClass* CurrentOwner, FName FunctionName);
 };
 
diff --git a/Engine/Source/Editor/BlueprintGraph/Private/K2Node_FunctionEntry.cpp b/Engine/Source/Editor/BlueprintGraph/Private/K2Node_FunctionEntry.cpp
--- a/Engine/Source/Editor/BlueprintGraph/Private/K2Node_FunctionEntry.cpp
+++ b/Engine/Source/Editor/BlueprintGraph/Private/K2Node_FunctionEntry.cpp
@@ -217,12 +217,7 @@ void UK2Node_FunctionEntry::AllocateDefaultPins()
 	const UEdGraphSchema_K2* K2Schema = GetDefault<UEdGraphSchema_K2>();
 	CreatePin(EGPD_Output, K2Schema->PC_Exec, TEXT(""), NULL, false, false, K2Schema->PN_Then);
 
-	UFunction* Function = nullptr;
-	if (UClass* CurrentSignatureClass = *SignatureClass)
-	{
-		UClass* OriginalSignatureOwner = FindOriginalFunctionOwner(CurrentSignatureClass, SignatureName);
-		Function = FindField<UFunction>(OriginalSignatureOwner, SignatureName);
-	}
+	UFunction* Function = FindField<UFunction>(SignatureClass, SignatureName);
 
 	if (Function == nullptr)
 	{
diff --git a/Engine/Source/Editor/BlueprintGraph/Private/K2Node_FunctionResult.cpp b/Engine/Source/Editor/BlueprintGraph/Private/K2Node_FunctionResult.cpp
--- a/Engine/Source/Editor/BlueprintGraph/Private/K2Node_FunctionResult.cpp
+++ b/Engine/Source/Editor/BlueprintGraph/Private/K2Node_FunctionResult.cpp
@@ -126,14 +126,8 @@ void UK2Node_FunctionResult::AllocateDefaultPins()
 	const UEdGraphSchema_K2* K2Schema = GetDefault<UEdGraphSchema_K2>();
 	CreatePin(EGPD_Input, K2Schema->PC_Exec, TEXT(""), NULL, false, false, K2Schema->PN_Execute);
 
-	UFunction* Function = nullptr;
-	if (UClass* CurrentSignatureClass = *SignatureClass)
-	{
-		UClass* OriginalSignatureOwner = FindOriginalFunctionOwner(CurrentSignatureClass, SignatureName);
-		Function = FindField<UFunction>(OriginalSignatureOwner, SignatureName);
-	}
-
-	if (Function != nullptr)
+	UFunction* Function = FindField<UFunction>(SignatureClass, SignatureName);
+	if (Function != NULL)
 	{
 		CreatePinsForFunctionEntryExit(Function, /*bIsFunctionEntry=*/ false);
 	}
diff --git a/Engine/Source/Editor/BlueprintGraph/Private/K2Node_FunctionTerminator.cpp b/Engine/Source/Editor/BlueprintGraph/Private/K2Node_FunctionTerminator.cpp
--- a/Engine/Source/Editor/BlueprintGraph/Private/K2Node_FunctionTerminator.cpp
+++ b/Engine/Source/Editor/BlueprintGraph/Private/K2Node_FunctionTerminator.cpp
@@ -50,28 +50,6 @@ bool UK2Node_FunctionTerminator::CanCreateUserDefinedPin(const FEdGraphPinType&
 	return bIsNodeEditable;
 }
 
-UClass* UK2Node_FunctionTerminator::FindOriginalFunctionOwner(UClass* CurrentOwner, FName FunctionName)
-{
-	check(CurrentOwner);
-	for (auto& Inter : CurrentOwner->Interfaces)
-	{
-		if (Inter.Class->FindFunctionByName(FunctionName) != nullptr)
-		{
-			return FindOriginalFunctionOwner(Inter.Class, FunctionName);
-		}
-	}
-
-	for (UClass* SearchClass = CurrentOwner->GetSuperClass(); SearchClass; SearchClass = SearchClass->GetSuperClass())
-	{
-		if (SearchClass->FindFunctionByName(FunctionName) != nullptr)
-		{
-			return FindOriginalFunctionOwner(SearchClass, FunctionName);
-		}
-	}
-
-	return CurrentOwner;
-}
-
 bool UK2Node_FunctionTerminator::HasExternalDependencies(TArray<class UStruct*>* OptionalOutput) const
 {
 	const UBlueprint* SourceBlueprint = GetBlueprint();
@@ -83,16 +61,6 @@ bool UK2Node_FunctionTerminator::HasExternalDependencies(TArray<class UStruct*>*
 		OptionalOutput->AddUnique(SourceClass);
 	}
 
-	if (SourceClass)
-	{
-		UClass* OriginalSignatureOwner = FindOriginalFunctionOwner(SourceClass, SignatureName);
-		bResult |= (OriginalSignatureOwner != nullptr) && (OriginalSignatureOwner->ClassGeneratedBy != SourceBlueprint);
-		if (bResult && (OriginalSignatureOwner != SourceClass) && OptionalOutput)
-		{
-			OptionalOutput->AddUnique(OriginalSignatureOwner);
-		}
-	}
-
 	// All structures, that are required for the BP compilation, should be gathered
 	for (auto Pin : Pins)
 	{
diff --git a/Engine/Source/Editor/Kismet/Private/BlueprintEditor.cpp b/Engine/Source/Editor/Kismet/Private/BlueprintEditor.cpp
--- a/Engine/Source/Editor/Kismet/Private/BlueprintEditor.cpp
+++ b/Engine/Source/Editor/Kismet/Private/BlueprintEditor.cpp
@@ -3153,60 +3153,28 @@ void FBlueprintEditor::Compile()
 		// Decide if we want an instrumented compile.
 		const bool bProfilerAvailable = IsProfilerAvailable();
 		bool bAddInstrumentation = false;
-		if (bProfilerAvailable)
-		{
-			if (GetDefault<UEditorExperimentalSettings>()->bBlueprintPerformanceAnalysisTools)
-			{
-				IBlueprintProfilerInterface& ProfilerModule = FModuleManager::LoadModuleChecked<IBlueprintProfilerInterface>("BlueprintProfiler");
-				bAddInstrumentation = ProfilerModule.IsProfilerEnabled();
-			}
+	    if (bProfilerAvailable)
+	    {
+		    if (GetDefault<UEditorExperimentalSettings>()->bBlueprintPerformanceAnalysisTools)
+		    {
+			    IBlueprintProfilerInterface& ProfilerModule = FModuleManager::LoadModuleChecked<IBlueprintProfilerInterface>("BlueprintProfiler");
+			    bAddInstrumentation = ProfilerModule.IsProfilerEnabled();
+		    }
 		}
 
-		TSharedRef<FCompilerResultsLog> LogResults(new FCompilerResultsLog());
-		LogResults->BeginEvent(TEXT("Compile"));
-		LogResults->bLogDetailedResults = GetDefault<UBlueprintEditorSettings>()->bShowDetailedCompileResults;
-		LogResults->EventDisplayThresholdMs = GetDefault<UBlueprintEditorSettings>()->CompileEventDisplayThresholdMs;
-
-		{
-			struct FParentCompiler
-			{
-				static void CompileParent(UClass* InParentClass, TWeakPtr<FKismetEditorUtilities::FCustomCompilationSettingsMap> InCustomCompilationSettingsMap)
-				{
-					UBlueprint* BP = UBlueprint::GetBlueprintFromClass(InParentClass);
-					if (BP)
-					{
-						CompileParent(BP->ParentClass, InCustomCompilationSettingsMap);
-
-						if (BP->IsPossiblyDirty())
-						{
-							FKismetEditorUtilities::CompileBlueprint(BP, false, false, false, false, InCustomCompilationSettingsMap);
-						}
-					}
-				}
-			};
-
-			FKismetEditorUtilities::FCustomCompilationSetting CustomCompilationSetting;
-			CustomCompilationSetting.bAddInstrumentation = bAddInstrumentation;
-			CustomCompilationSetting.bSaveIntermediateBuildProducts = bSaveIntermediateBuildProducts;
-			CustomCompilationSetting.LogResults = TWeakPtr<FCompilerResultsLog>(LogResults);
-
-			TSharedRef<FKismetEditorUtilities::FCustomCompilationSettingsMap> CustomCompilationSettingsMap = MakeShareable(new FKismetEditorUtilities::FCustomCompilationSettingsMap());
-			CustomCompilationSettingsMap->AddSettings(BlueprintObj, CustomCompilationSetting);
-
-			FParentCompiler::CompileParent(BlueprintObj->ParentClass, CustomCompilationSettingsMap);
-			if (!CustomCompilationSettingsMap->WasBPCompiledWithCustomSettings(BlueprintObj))
-			{
-				FKismetEditorUtilities::CompileBlueprint(BlueprintObj, false, false, false, false, CustomCompilationSettingsMap);
-			}
-		}
+		FCompilerResultsLog LogResults;
+		LogResults.BeginEvent(TEXT("Compile"));
+		LogResults.bLogDetailedResults = GetDefault<UBlueprintEditorSettings>()->bShowDetailedCompileResults;
+		LogResults.EventDisplayThresholdMs = GetDefault<UBlueprintEditorSettings>()->CompileEventDisplayThresholdMs;
+		FKismetEditorUtilities::CompileBlueprint(BlueprintObj, false, false, bSaveIntermediateBuildProducts, &LogResults, false, false, bAddInstrumentation);
 
-		bool bForceMessageDisplay = ((LogResults->NumWarnings > 0) || (LogResults->NumErrors > 0)) && !BlueprintObj->bIsRegeneratingOnLoad;
-		DumpMessagesToCompilerLog(LogResults->Messages, bForceMessageDisplay);
+		bool bForceMessageDisplay = ((LogResults.NumWarnings > 0) || (LogResults.NumErrors > 0)) && !BlueprintObj->bIsRegeneratingOnLoad;
+		DumpMessagesToCompilerLog(LogResults.Messages, bForceMessageDisplay);
 
 		UBlueprintEditorSettings const* BpEditorSettings = GetDefault<UBlueprintEditorSettings>();
-		if ((LogResults->NumErrors > 0) && BpEditorSettings->bJumpToNodeErrors)
+		if ((LogResults.NumErrors > 0) && BpEditorSettings->bJumpToNodeErrors)
 		{
-			if (UEdGraphNode* NodeWithError = BlueprintEditorImpl::FindNodeWithError(*LogResults))
+			if (UEdGraphNode* NodeWithError = BlueprintEditorImpl::FindNodeWithError(LogResults))
 			{
 				JumpToNode(NodeWithError, /*bRequestRename =*/false);
 			}
@@ -3219,7 +3187,7 @@ void FBlueprintEditor::Compile()
 
 		AppendExtraCompilerResults(CompilerResultsListing);
 
-		LogResults->EndEvent();
+		LogResults.EndEvent();
 
 	    // Update the blueprint instrumenation state and show the profiler window if required
 	    if (bProfilerAvailable)
@@ -3242,7 +3210,7 @@ void FBlueprintEditor::Compile()
 
 		// send record when player clicks compile and send the result
 		// this will make sure how the users activity is
-		AnalyticsTrackCompileEvent(BlueprintObj, LogResults->NumErrors, LogResults->NumWarnings);
+		AnalyticsTrackCompileEvent(BlueprintObj, LogResults.NumErrors, LogResults.NumWarnings);
 	}
 }
 
diff --git a/Engine/Source/Editor/UnrealEd/Private/Kismet2/BlueprintEditorUtils.cpp b/Engine/Source/Editor/UnrealEd/Private/Kismet2/BlueprintEditorUtils.cpp
--- a/Engine/Source/Editor/UnrealEd/Private/Kismet2/BlueprintEditorUtils.cpp
+++ b/Engine/Source/Editor/UnrealEd/Private/Kismet2/BlueprintEditorUtils.cpp
@@ -445,17 +445,16 @@ FUCSComponentId::FUCSComponentId(const UK2Node_AddComponent* UCSNode)
 //////////////////////////////////////
 // FBlueprintEditorUtils
 
-bool FBlueprintEditorUtils::RefreshAllNodes(UBlueprint* Blueprint, bool bMarkChildrenAsStructurallyModified)
+void FBlueprintEditorUtils::RefreshAllNodes(UBlueprint* Blueprint)
 {
 	if (!Blueprint || !Blueprint->HasAllFlags(RF_LoadCompleted))
 	{
 		UE_LOG(LogBlueprint, Warning, 
 			TEXT("RefreshAllNodes called on incompletly loaded blueprint '%s'"), 
 			Blueprint ? *Blueprint->GetFullName() : TEXT("NULL"));
-		return false;
+		return;
 	}
 
-	bool bWasBPMarkAsStructurallyModified = false;
 	TArray<UK2Node*> AllNodes;
 	FBlueprintEditorUtils::GetAllNodesOfClass(Blueprint, AllNodes);
 
@@ -481,8 +480,7 @@ bool FBlueprintEditorUtils::RefreshAllNodes(UBlueprint* Blueprint, bool bMarkChi
 			// Ignore this for macros
 			if (!bIsMacro)
 			{
-				FBlueprintEditorUtils::MarkBlueprintAsStructurallyModified(Blueprint, bMarkChildrenAsStructurallyModified);
-				bWasBPMarkAsStructurallyModified = true;
+				FBlueprintEditorUtils::MarkBlueprintAsStructurallyModified(Blueprint);
 			}
 			bLastChangesStructure = bCurrentChangesStructure;
 		}
@@ -495,11 +493,8 @@ bool FBlueprintEditorUtils::RefreshAllNodes(UBlueprint* Blueprint, bool bMarkChi
 	// If all nodes change structure, catch that case and recompile now
 	if( bLastChangesStructure )
 	{
-		FBlueprintEditorUtils::MarkBlueprintAsStructurallyModified(Blueprint, bMarkChildrenAsStructurallyModified);
-		bWasBPMarkAsStructurallyModified = true;
+		FBlueprintEditorUtils::MarkBlueprintAsStructurallyModified(Blueprint);
 	}
-
-	return bWasBPMarkAsStructurallyModified;
 }
 
 
@@ -1566,7 +1561,7 @@ UClass* FBlueprintEditorUtils::RegenerateBlueprintClass(UBlueprint* Blueprint, U
 			FBlueprintEditorUtils::ReconstructAllNodes(Blueprint);
 
 			// Compile the actual blueprint
-			FKismetEditorUtilities::CompileBlueprint(Blueprint, true, false, bSkeletonUpToDate);
+			FKismetEditorUtilities::CompileBlueprint(Blueprint, true, false, false, nullptr, bSkeletonUpToDate);
 		}
 		else if( bIsMacro )
 		{
@@ -2194,7 +2189,7 @@ void FBlueprintEditorUtils::UpdateDelegatesInBlueprint(UBlueprint* Blueprint)
 }
 
 // Blueprint has materially changed.  Recompile the skeleton, notify observers, and mark the package as dirty.
-void FBlueprintEditorUtils::MarkBlueprintAsStructurallyModified(UBlueprint* Blueprint, bool bMarkChildrenAsStructurallyModified)
+void FBlueprintEditorUtils::MarkBlueprintAsStructurallyModified(UBlueprint* Blueprint)
 {
 	FSecondsCounterScope Timer(BlueprintCompileAndLoadTimerData);
 	
@@ -2286,10 +2281,9 @@ void FBlueprintEditorUtils::MarkBlueprintAsStructurallyModified(UBlueprint* Blue
 			Blueprint->Status = BS_Dirty;
 		}
 		UpdateDelegatesInBlueprint(Blueprint);
-		if (bMarkChildrenAsStructurallyModified)
-		{
-			FRefreshHelper::SkeletalRecompileChildren(ChildrenOfClass, Blueprint->bIsRegeneratingOnLoad);
-		}
+
+		FRefreshHelper::SkeletalRecompileChildren(ChildrenOfClass, Blueprint->bIsRegeneratingOnLoad);
+
 		{
 			BP_SCOPED_COMPILER_EVENT_STAT(EKismetCompilerStats_NotifyBlueprintChanged);
 
@@ -2974,15 +2968,13 @@ void FBlueprintEditorUtils::GatherDependencies(const UBlueprint* InBlueprint, TS
 	InBlueprint->GatherDependencies(Dependencies);
 
 	FGatherDependenciesHelper::ProcessHierarchy(InBlueprint->ParentClass, Dependencies);
-	for (const UBlueprint* BPIter = InBlueprint; BPIter; BPIter = UBlueprint::GetBlueprintFromClass(BPIter->ParentClass))
+
+	for (const auto& InterfaceDesc : InBlueprint->ImplementedInterfaces)
 	{
-		for (const auto& InterfaceDesc : BPIter->ImplementedInterfaces)
+		UBlueprint* InterfaceBP = InterfaceDesc.Interface ? Cast<UBlueprint>(InterfaceDesc.Interface->ClassGeneratedBy) : nullptr;
+		if (InterfaceBP)
 		{
-			UBlueprint* InterfaceBP = InterfaceDesc.Interface ? Cast<UBlueprint>(InterfaceDesc.Interface->ClassGeneratedBy) : nullptr;
-			if (InterfaceBP)
-			{
-				Dependencies.Add(InterfaceBP);
-			}
+			Dependencies.Add(InterfaceBP);
 		}
 	}
 
diff --git a/Engine/Source/Editor/UnrealEd/Private/Kismet2/Kismet2.cpp b/Engine/Source/Editor/UnrealEd/Private/Kismet2/Kismet2.cpp
--- a/Engine/Source/Editor/UnrealEd/Private/Kismet2/Kismet2.cpp
+++ b/Engine/Source/Editor/UnrealEd/Private/Kismet2/Kismet2.cpp
@@ -706,51 +706,11 @@ bool FKismetEditorUtilities::IsReferencedByUndoBuffer(UBlueprint* Blueprint)
 	return (TotalReferenceCount > NonUndoReferenceCount);
 }
 
-void FKismetEditorUtilities::FCustomCompilationSettingsMap::AddSettings(TWeakObjectPtr<UBlueprint> BP, FKismetEditorUtilities::FCustomCompilationSetting Setting)
-{
-	FDataPerBP Data;
-	Data.Setting = Setting;
-	CustomSettingsMap.Add(BP, Data);
-}
-
-bool FKismetEditorUtilities::FCustomCompilationSettingsMap::UseSettings(TWeakObjectPtr<UBlueprint> BP, bool bMarkAsCompiled, FKismetEditorUtilities::FCustomCompilationSetting& OutSetting)
-{
-	FDataPerBP* Data = CustomSettingsMap.Find(BP);
-	if (Data)
-	{
-		if (bMarkAsCompiled)
-		{
-			ensure(!Data->bWasBlueprintCompiledUsingTheSettings);
-			Data->bWasBlueprintCompiledUsingTheSettings = false;
-		}
-		OutSetting = Data->Setting;
-	}
-	return Data != nullptr;
-}
-
-bool FKismetEditorUtilities::FCustomCompilationSettingsMap::WasBPCompiledWithCustomSettings(TWeakObjectPtr<UBlueprint> BP)
-{
-	FDataPerBP* Data = CustomSettingsMap.Find(BP);
-	ensure(Data);
-	return Data && Data->bWasBlueprintCompiledUsingTheSettings;
-}
-
-void FKismetEditorUtilities::CompileBlueprint(UBlueprint* BlueprintObj
-	, bool bIsRegeneratingOnLoad, bool bSkipGarbageCollection, bool bSkeletonUpToDate, bool bBatchCompile
-	, TWeakPtr<FCustomCompilationSettingsMap> CustomCompilationSettingsMap)
+void FKismetEditorUtilities::CompileBlueprint(UBlueprint* BlueprintObj, bool bIsRegeneratingOnLoad, bool bSkipGarbageCollection, bool bSaveIntermediateProducts, FCompilerResultsLog* pResults, bool bSkeletonUpToDate, bool bBatchCompile, bool bAddInstrumentation)
 {
 	FSecondsCounterScope Timer(BlueprintCompileAndLoadTimerData); 
 	BP_SCOPED_COMPILER_EVENT_STAT(EKismetCompilerStats_CompileBlueprint);
 
-	FCustomCompilationSetting CustomCompilationSetting;
-	{
-		TSharedPtr<FCustomCompilationSettingsMap> CustomCompilationSettingsMapPinned = CustomCompilationSettingsMap.Pin();
-		if (CustomCompilationSettingsMapPinned.IsValid())
-		{
-			CustomCompilationSettingsMapPinned->UseSettings(BlueprintObj, true, CustomCompilationSetting);
-		}
-	}
-
 	// Wipe the PreCompile log, any generated messages are now irrelevant
 	BlueprintObj->PreCompileLog.Reset();
 
@@ -796,8 +756,7 @@ void FKismetEditorUtilities::CompileBlueprint(UBlueprint* BlueprintObj
 
 	// Compile
 	FCompilerResultsLog LocalResults;
-	TSharedPtr<FCompilerResultsLog> LogResultsPinned = CustomCompilationSetting.LogResults.Pin();
-	FCompilerResultsLog& Results = LogResultsPinned.IsValid() ? *LogResultsPinned : LocalResults;
+	FCompilerResultsLog& Results = (pResults != NULL) ? *pResults : LocalResults;
 
 	// Monitoring UE-20486, the OldClass->ClassGeneratedBy is NULL or otherwise not a UBlueprint.
 	if (OldClass && (OldClass->ClassGeneratedBy != BlueprintObj))
@@ -806,16 +765,15 @@ void FKismetEditorUtilities::CompileBlueprint(UBlueprint* BlueprintObj
 		OldClass->ClassGeneratedBy = BlueprintObj;
 	}
 	auto ReinstanceHelper = FBlueprintCompileReinstancer::Create(OldClass);
-	ReinstanceHelper->CustomCompilationSettingsMap = CustomCompilationSettingsMap;
 
 	// If enabled, suppress errors/warnings in the log if we're recompiling on load on a build machine
 	static const FBoolConfigValueHelper IgnoreCompileOnLoadErrorsOnBuildMachine(TEXT("Kismet"), TEXT("bIgnoreCompileOnLoadErrorsOnBuildMachine"), GEngineIni);
 	Results.bLogInfoOnly = BlueprintObj->bIsRegeneratingOnLoad && GIsBuildMachine && IgnoreCompileOnLoadErrorsOnBuildMachine;
 
 	FKismetCompilerOptions CompileOptions;
-	CompileOptions.bSaveIntermediateProducts = CustomCompilationSetting.bSaveIntermediateBuildProducts;
+	CompileOptions.bSaveIntermediateProducts = bSaveIntermediateProducts;
 	CompileOptions.bRegenerateSkelton = !bSkeletonUpToDate;
-	CompileOptions.bAddInstrumentation = CustomCompilationSetting.bAddInstrumentation;
+	CompileOptions.bAddInstrumentation = bAddInstrumentation;
 	Compiler.CompileBlueprint(BlueprintObj, CompileOptions, Results, ReinstanceHelper);
 
 	FBlueprintEditorUtils::UpdateDelegatesInBlueprint(BlueprintObj);
@@ -901,37 +859,6 @@ void FKismetEditorUtilities::CompileBlueprint(UBlueprint* BlueprintObj
 		TArray<UBlueprint*> DependentBPs;
 		FBlueprintEditorUtils::GetDependentBlueprints(BlueprintObj, DependentBPs);
 
-		struct FBlueprintSorter
-		{
-			static int32 GetSortValue(const UBlueprint* InBlueprint)
-			{
-				int32 Result = 0;
-				if (InBlueprint)
-				{
-					for (UClass* SuperClass = InBlueprint->ParentClass; SuperClass && !SuperClass->HasAnyClassFlags(CLASS_Native); SuperClass = SuperClass->GetSuperClass())
-					{
-						Result++;
-					}
-
-					if (InBlueprint->BlueprintType == BPTYPE_Interface)
-					{
-						Result -= 1024;
-					}
-				}
-				return Result;
-			}
-
-			bool operator()(const UBlueprint& A, const UBlueprint& B) const
-			{
-				//Smaller is first
-				return GetSortValue(&A) < GetSortValue(&B);
-			}
-		};
-		if (bIsInterface)
-		{
-			DependentBPs.Sort<FBlueprintSorter>(FBlueprintSorter());
-		}
-
 		// refresh each dependent blueprint
 		for (UBlueprint* Dependent : DependentBPs)
 		{
@@ -941,11 +868,7 @@ void FKismetEditorUtilities::CompileBlueprint(UBlueprint* BlueprintObj
 			{
 				bool bPreviousRegenValue = Dependent->bIsRegeneratingOnLoad;
 				Dependent->bIsRegeneratingOnLoad = Dependent->bIsRegeneratingOnLoad || BlueprintObj->bIsRegeneratingOnLoad;
-				if (!FBlueprintEditorUtils::RefreshAllNodes(Dependent, false))
-				{
-					FBlueprintEditorUtils::MarkBlueprintAsStructurallyModified(Dependent, false);
-				}
-
+				FBlueprintEditorUtils::RefreshAllNodes(Dependent);
 				Dependent->bIsRegeneratingOnLoad = bPreviousRegenValue;
 			}
 			else if(!BlueprintObj->bIsRegeneratingOnLoad)
diff --git a/Engine/Source/Editor/UnrealEd/Private/Kismet2/KismetReinstanceUtilities.cpp b/Engine/Source/Editor/UnrealEd/Private/Kismet2/KismetReinstanceUtilities.cpp
--- a/Engine/Source/Editor/UnrealEd/Private/Kismet2/KismetReinstanceUtilities.cpp
+++ b/Engine/Source/Editor/UnrealEd/Private/Kismet2/KismetReinstanceUtilities.cpp
@@ -554,8 +554,7 @@ void FBlueprintCompileReinstancer::CompileChildren()
 			// was updated); however, if this is a native class (like when hot-
 			// reloading), then we want to make sure to update the skel as well
 			const bool bSkeletonUpToDate = !ClassToReinstance->HasAnyClassFlags(CLASS_Native);
-			const bool bBatchCompile = false;
-			FKismetEditorUtilities::CompileBlueprint(BP, false, bSkipGarbageCollection, bSkeletonUpToDate, bBatchCompile, CustomCompilationSettingsMap);
+			FKismetEditorUtilities::CompileBlueprint(BP, false, bSkipGarbageCollection, false, nullptr, bSkeletonUpToDate, false);
 		}
 		else if (IsReinstancingSkeleton())
 		{
@@ -706,8 +705,7 @@ void FBlueprintCompileReinstancer::ReinstanceObjects(bool bForceAlwaysReinstance
 						const bool bSkipGC = true;
 						// Full compiles first recompile all skeleton classes, so they are up-to-date
 						const bool bSkeletonUpToDate = true;
-						const bool bBatchCompile = true;
-						FKismetEditorUtilities::CompileBlueprint(BP, false, bSkipGC, bSkeletonUpToDate, bBatchCompile, CustomCompilationSettingsMap);
+						FKismetEditorUtilities::CompileBlueprint(BP, false, bSkipGC, false, nullptr, bSkeletonUpToDate, true);
 						CompiledBlueprints.Add(BP);
 					}
 
diff --git a/Engine/Source/Editor/UnrealEd/Private/PlayLevel.cpp b/Engine/Source/Editor/UnrealEd/Private/PlayLevel.cpp
--- a/Engine/Source/Editor/UnrealEd/Private/PlayLevel.cpp
+++ b/Engine/Source/Editor/UnrealEd/Private/PlayLevel.cpp
@@ -1741,8 +1741,11 @@ struct FInternalPlayLevelUtils
 				FKismetEditorUtilities::CompileBlueprint(Blueprint,
 					/*bIsRegeneratingOnLoad =*/false,
 					/*bSkipGarbageCollection =*/true,
+					/*bSaveIntermediateProducts =*/false,
+					/*pResults =*/nullptr,
 					/*bSkeletonUpToDate =*/false,
-					/*bBatchCompile =*/false);
+					/*bBatchCompile =*/false,
+					/*bAddInstrumentation =*/false);
 
 				// Check for errors after compiling
 				for (auto CompiledIt = CompiledBlueprints.CreateIterator(); CompiledIt; ++CompiledIt)
diff --git a/Engine/Source/Editor/UnrealEd/Public/Kismet2/BlueprintEditorUtils.h b/Engine/Source/Editor/UnrealEd/Public/Kismet2/BlueprintEditorUtils.h
--- a/Engine/Source/Editor/UnrealEd/Public/Kismet2/BlueprintEditorUtils.h
+++ b/Engine/Source/Editor/UnrealEd/Public/Kismet2/BlueprintEditorUtils.h
@@ -115,9 +115,9 @@ class UNREALED_API FBlueprintEditorUtils
 public:
 
 	/**
-	 * Schedules and refreshes all nodes in the blueprint, making sure that nodes that affect function signatures get regenerated first. Returns true if BP was marked as structurally modified.
+	 * Schedules and refreshes all nodes in the blueprint, making sure that nodes that affect function signatures get regenerated first
 	 */
-	static bool RefreshAllNodes(UBlueprint* Blueprint, bool bMarkChildrenAsStructurallyModified = true);
+	static void RefreshAllNodes(UBlueprint* Blueprint);
 
 	/**
 	 * Reconstructs all nodes in the blueprint, node reconstruction order determined by FCompareNodePriority.
@@ -228,7 +228,7 @@ public:
 	 *  - Notifies any observers.
 	 *  - Marks the package as dirty.
 	 */
-	static void MarkBlueprintAsStructurallyModified(UBlueprint* Blueprint, bool bMarkChildrenAsStructurallyModified = true);
+	static void MarkBlueprintAsStructurallyModified(UBlueprint* Blueprint);
 
 	/**
 	 * Blueprint has changed in some manner that invalidates the compiled data (link made/broken, default value changed, etc...)
diff --git a/Engine/Source/Editor/UnrealEd/Public/Kismet2/KismetEditorUtilities.h b/Engine/Source/Editor/UnrealEd/Public/Kismet2/KismetEditorUtilities.h
--- a/Engine/Source/Editor/UnrealEd/Public/Kismet2/KismetEditorUtilities.h
+++ b/Engine/Source/Editor/UnrealEd/Public/Kismet2/KismetEditorUtilities.h
@@ -96,47 +96,8 @@ public:
 	/** Create the correct event graphs for this blueprint */
 	static void CreateDefaultEventGraphs(UBlueprint* Blueprint);
 
-	/** While reinstancing many different BPs may be compiled, some of them require custom settings. This structure stores them. */
-	struct FCustomCompilationSetting
-	{
-		TWeakPtr<FCompilerResultsLog> LogResults;
-		bool bSaveIntermediateBuildProducts;
-		bool bAddInstrumentation;
-
-		FCustomCompilationSetting()
-			: bSaveIntermediateBuildProducts(false)
-			, bAddInstrumentation(false)
-		{}
-	};
-
-	struct UNREALED_API FCustomCompilationSettingsMap : TSharedFromThis<FCustomCompilationSettingsMap>
-	{
-	private:
-		struct FDataPerBP
-		{
-			FCustomCompilationSetting Setting;
-			bool bWasBlueprintCompiledUsingTheSettings;
-
-			FDataPerBP() : bWasBlueprintCompiledUsingTheSettings(false) {}
-		};
-
-		TMap<TWeakObjectPtr<UBlueprint>, FDataPerBP> CustomSettingsMap;
-	public:
-
-		void AddSettings(TWeakObjectPtr<UBlueprint> BP, FCustomCompilationSetting Setting);
-
-		bool UseSettings(TWeakObjectPtr<UBlueprint> BP, bool bMarkAsCompiled, FCustomCompilationSetting& OutSetting);
-
-		bool WasBPCompiledWithCustomSettings(TWeakObjectPtr<UBlueprint> BP);
-	};
-
 	/** Tries to compile a blueprint, updating any actors in the editor who are using the old class, etc... */
-	static void CompileBlueprint(UBlueprint* BlueprintObj
-		, bool bIsRegeneratingOnLoad = false
-		, bool bSkipGarbageCollection = false
-		, bool bSkeletonUpToDate = false
-		, bool bBatchCompile = false
-		, TWeakPtr<FCustomCompilationSettingsMap> CustomCompilationSettingsMap = {});
+	static void CompileBlueprint(UBlueprint* BlueprintObj, bool bIsRegeneratingOnLoad = false, bool bSkipGarbageCollection = false, bool bSaveIntermediateProducts = false, class FCompilerResultsLog* pResults = nullptr, bool bSkeletonUpToDate = false, bool bBatchCompile = false, bool bAddInstrumentation = false);
 
 	/** Generates a blueprint skeleton only.  Minimal compile, no notifications will be sent, no GC, etc.  Only successful if there isn't already a skeleton generated */
 	static bool GenerateBlueprintSkeleton(UBlueprint* BlueprintObj, bool bForceRegeneration = false);
diff --git a/Engine/Source/Editor/UnrealEd/Public/Kismet2/KismetReinstanceUtilities.h b/Engine/Source/Editor/UnrealEd/Public/Kismet2/KismetReinstanceUtilities.h
--- a/Engine/Source/Editor/UnrealEd/Public/Kismet2/KismetReinstanceUtilities.h
+++ b/Engine/Source/Editor/UnrealEd/Public/Kismet2/KismetReinstanceUtilities.h
@@ -7,7 +7,6 @@
 #pragma once
 
 #include "Engine.h"
-#include "KismetEditorUtilities.h"
 
 DECLARE_STATS_GROUP(TEXT("Kismet Reinstancer"), STATGROUP_KismetReinstancer, STATCAT_Advanced);
 
@@ -93,8 +92,6 @@ protected:
 	bool bAllowResaveAtTheEndIfRequested;
 
 public:
-	TWeakPtr<FKismetEditorUtilities::FCustomCompilationSettingsMap> CustomCompilationSettingsMap;
-
 	// FSerializableObject interface
 	virtual void AddReferencedObjects(FReferenceCollector& Collector) override;
 	// End of FSerializableObject interface
