CIS fixes.

diff --git a/Engine/Source/Editor/UnrealEd/Private/Fbx/FbxSkeletalMeshImport.cpp b/Engine/Source/Editor/UnrealEd/Private/Fbx/FbxSkeletalMeshImport.cpp
--- a/Engine/Source/Editor/UnrealEd/Private/Fbx/FbxSkeletalMeshImport.cpp
+++ b/Engine/Source/Editor/UnrealEd/Private/Fbx/FbxSkeletalMeshImport.cpp
@@ -2502,7 +2502,7 @@ bool UnFbx::FFbxImporter::FillSkelMeshImporterFromFbx( FSkeletalMeshImportData&
 			// When Maya plug-in exports rigid binding, it will generate "CompensationCluster" for each ancestor links.
 			// FBX writes these "CompensationCluster" out. The CompensationCluster also has weight 1 for vertices.
 			// Unreal importer should skip these clusters.
-			if(!Cluster || FCStringAnsi::Strcmp(Cluster->GetUserDataID(), "Maya_ClusterHint") == 0 && FCStringAnsi::Strcmp(Cluster->GetUserData(), "CompensationCluster")  == 0)
+			if(!Cluster || (FCStringAnsi::Strcmp(Cluster->GetUserDataID(), "Maya_ClusterHint") == 0 && FCStringAnsi::Strcmp(Cluster->GetUserData(), "CompensationCluster")  == 0))
 			{
 				continue;
 			}
diff --git a/Engine/Source/Runtime/Engine/Private/Animation/AnimationRecordingSettings.cpp b/Engine/Source/Runtime/Engine/Private/Animation/AnimationRecordingSettings.cpp
--- a/Engine/Source/Runtime/Engine/Private/Animation/AnimationRecordingSettings.cpp
+++ b/Engine/Source/Runtime/Engine/Private/Animation/AnimationRecordingSettings.cpp
@@ -2,6 +2,7 @@
 
 #include "EnginePrivate.h"
 #include "Animation/AnimationRecordingSettings.h"
+#include "Animation/AnimSequenceBase.h"
 
 /** 30Hz default sample rate */
 const float FAnimationRecordingSettings::DefaultSampleRate = DEFAULT_SAMPLERATE;
