VR Editor: Fix QAGame crash when closing Editor after using VREditor #jira UE-36104

diff --git a/Engine/Source/Editor/UnrealEd/Private/EditorWorldManager.cpp b/Engine/Source/Editor/UnrealEd/Private/EditorWorldManager.cpp
--- a/Engine/Source/Editor/UnrealEd/Private/EditorWorldManager.cpp
+++ b/Engine/Source/Editor/UnrealEd/Private/EditorWorldManager.cpp
@@ -22,19 +22,8 @@ FEditorWorldWrapper::FEditorWorldWrapper(FWorldContext& InWorldContext):
 
 FEditorWorldWrapper::~FEditorWorldWrapper()
 {
-	if(VREditorMode != nullptr && !VREditorMode->IsPendingKill() )
-	{
-		VREditorMode->Shutdown();
-		VREditorMode->MarkPendingKill();
-		VREditorMode = nullptr;
-	}
-
-	if(WorldInteraction != nullptr && !WorldInteraction->IsPendingKill())
-	{
-		WorldInteraction->Shutdown();
-		WorldInteraction->MarkPendingKill();
-		WorldInteraction = nullptr;
-	}
+	WorldInteraction = nullptr;
+	VREditorMode = nullptr;
 }
 
 UWorld* FEditorWorldWrapper::GetWorld() const
