Copying //UE4/Dev-Framework to //UE4/Dev-Main (Source: //UE4/Dev-Framework @ 3136612)

diff --git a/Engine/Source/Editor/UnrealEd/Private/Kismet2/BlueprintEditorUtils.cpp b/Engine/Source/Editor/UnrealEd/Private/Kismet2/BlueprintEditorUtils.cpp
--- a/Engine/Source/Editor/UnrealEd/Private/Kismet2/BlueprintEditorUtils.cpp
+++ b/Engine/Source/Editor/UnrealEd/Private/Kismet2/BlueprintEditorUtils.cpp
@@ -91,23 +91,26 @@ struct FCompareNodePriority
  */
 static UClass* FindInheritedInterface(UClass* const Class, FBPInterfaceDescription const& Interface)
 {
-	UClass* ClassWithInterface = NULL;
+	UClass* ClassWithInterface = nullptr;
 
-	if (Class != NULL)
+	if (Class != nullptr)
 	{
 		UClass* const ParentClass = Class->GetSuperClass();
 		// search depth first so that we may find the highest parent in the chain that implements this interface
 		ClassWithInterface = FindInheritedInterface(ParentClass, Interface);
 
-		for (auto InterfaceIt(Class->Interfaces.CreateConstIterator()); InterfaceIt && (ClassWithInterface == NULL); ++InterfaceIt)
+		if (ClassWithInterface == nullptr)
+		{
+			for (const FImplementedInterface& ImplementedInterface : Class->Interfaces)
 		{
-			if (InterfaceIt->Class == Interface.Interface)
+				if (ImplementedInterface.Class == Interface.Interface)
 			{
 				ClassWithInterface = Class;
 				break;
 			}
 		}
 	}
+	}
 
 	return ClassWithInterface;
 }
@@ -123,8 +126,8 @@ static UClass* FindInheritedInterface(UClass* const Class, FBPInterfaceDescripti
  */
 static bool IsInterfaceImplementedByParent(FBPInterfaceDescription const& Interface, UBlueprint const* const Blueprint)
 {
-	check(Blueprint != NULL);
-	return (FindInheritedInterface(Blueprint->ParentClass, Interface) != NULL);
+	check(Blueprint != nullptr);
+	return (FindInheritedInterface(Blueprint->ParentClass, Interface) != nullptr);
 }
 
 /**
@@ -165,15 +168,15 @@ static void ReplaceNode(UK2Node* OldNode, UK2Node* NewNode)
  */
 static void PromoteInterfaceImplementationToOverride(FBPInterfaceDescription const& Interface, UBlueprint* const BlueprintObj)
 {
-	check(BlueprintObj != NULL);
+	check(BlueprintObj != nullptr);
 	// find the parent whose interface we're overriding 
 	UClass* ParentClass = FindInheritedInterface(BlueprintObj->ParentClass, Interface);
 
-	if (ParentClass != NULL)
+	if (ParentClass != nullptr)
 	{
 		for (UEdGraph* InterfaceGraph : Interface.Graphs)
 		{
-			check(InterfaceGraph != NULL);
+			check(InterfaceGraph != nullptr);
 
 			// The graph can be deleted now that it is a simple function override
 			InterfaceGraph->bAllowDeletion = true;
@@ -325,10 +328,10 @@ void FBasePinChangeHelper::Broadcast(UBlueprint* InBlueprint, UK2Node_EditablePi
 
 				const bool bNameMatches = (CallSite->FunctionReference.GetMemberName() == FuncName);
 				const UClass* MemberParentClass = CallSite->FunctionReference.GetMemberParentClass(CallSite->GetBlueprintClassFromNode());
-				const bool bClassMatchesEasy = (MemberParentClass != NULL) && (MemberParentClass->IsChildOf(SignatureClass) || MemberParentClass->IsChildOf(InBlueprint->GeneratedClass));
-				const bool bClassMatchesHard = (CallSite->FunctionReference.IsSelfContext()) && (SignatureClass == NULL) 
+				const bool bClassMatchesEasy = (MemberParentClass != nullptr) && (MemberParentClass->IsChildOf(SignatureClass) || MemberParentClass->IsChildOf(InBlueprint->GeneratedClass));
+				const bool bClassMatchesHard = (CallSite->FunctionReference.IsSelfContext()) && (SignatureClass == nullptr) 
 					&& (CallSiteBlueprint == InBlueprint || CallSiteBlueprint->SkeletonGeneratedClass->IsChildOf(InBlueprint->SkeletonGeneratedClass));
-				const bool bValidSchema = CallSite->GetSchema() != NULL;
+				const bool bValidSchema = CallSite->GetSchema() != nullptr;
 
 				if (bNameMatches && bValidSchema && (bClassMatchesEasy || bClassMatchesHard))
 				{
@@ -365,12 +368,12 @@ void FBasePinChangeHelper::Broadcast(UBlueprint* InBlueprint, UK2Node_EditablePi
 
 void FParamsChangedHelper::EditCompositeTunnelNode(UK2Node_Tunnel* TunnelNode)
 {
-	if (TunnelNode->InputSinkNode != NULL)
+	if (TunnelNode->InputSinkNode != nullptr)
 	{
 		TunnelNode->InputSinkNode->ReconstructNode();
 	}
 
-	if (TunnelNode->OutputSourceNode != NULL)
+	if (TunnelNode->OutputSourceNode != nullptr)
 	{
 		TunnelNode->OutputSourceNode->ReconstructNode();
 	}
@@ -389,7 +392,7 @@ void FParamsChangedHelper::EditCallSite(UK2Node_CallFunction* CallSite, UBluepri
 {
 	CallSite->Modify();
 	CallSite->ReconstructNode();
-	if (Blueprint != NULL)
+	if (Blueprint != nullptr)
 	{
 		ModifiedBlueprints.Add(Blueprint);
 	}
@@ -399,21 +402,20 @@ void FParamsChangedHelper::EditDelegates(UK2Node_BaseMCDelegate* CallSite, UBlue
 {
 	CallSite->Modify();
 	CallSite->ReconstructNode();
-	if (auto AssignNode = Cast<UK2Node_AddDelegate>(CallSite))
+	if (UK2Node_AddDelegate* AssignNode = Cast<UK2Node_AddDelegate>(CallSite))
 	{
-		if (auto DelegateInPin = AssignNode->GetDelegatePin())
+		if (UEdGraphPin* DelegateInPin = AssignNode->GetDelegatePin())
 		{
-			for(auto DelegateOutPinIt = DelegateInPin->LinkedTo.CreateIterator(); DelegateOutPinIt; ++DelegateOutPinIt)
+			for (UEdGraphPin* DelegateOutPin : DelegateInPin->LinkedTo)
 			{
-				UEdGraphPin* DelegateOutPin = *DelegateOutPinIt;
-				if (auto CustomEventNode = (DelegateOutPin ? Cast<UK2Node_CustomEvent>(DelegateOutPin->GetOwningNode()) : NULL))
+				if (UK2Node_CustomEvent* CustomEventNode = (DelegateOutPin ? Cast<UK2Node_CustomEvent>(DelegateOutPin->GetOwningNode()) : nullptr))
 				{
 					CustomEventNode->ReconstructNode();
 				}
 			}
 		}
 	}
-	if (Blueprint != NULL)
+	if (Blueprint != nullptr)
 	{
 		ModifiedBlueprints.Add(Blueprint);
 	}
@@ -421,8 +423,8 @@ void FParamsChangedHelper::EditDelegates(UK2Node_BaseMCDelegate* CallSite, UBlue
 
 void FParamsChangedHelper::EditCreateDelegates(UK2Node_CreateDelegate* CallSite)
 {
-	UBlueprint* Blueprint = NULL;
-	UEdGraph* Graph = NULL;
+	UBlueprint* Blueprint = nullptr;
+	UEdGraph* Graph = nullptr;
 	CallSite->HandleAnyChange(Graph, Blueprint);
 	if(Blueprint)
 	{
@@ -543,9 +545,8 @@ void FBlueprintEditorUtils::RefreshExternalBlueprintDependencyNodes(UBlueprint*
 
 	if (!RefreshOnlyChild)
 	{
-		for (auto NodeIt = AllNodes.CreateIterator(); NodeIt; ++NodeIt)
+		for (UK2Node* Node : AllNodes)
 		{
-			UK2Node* Node = *NodeIt;
 			if (Node->HasExternalDependencies())
 			{
 				//@todo:  Do we really need per-schema refreshing?
@@ -556,9 +557,8 @@ void FBlueprintEditorUtils::RefreshExternalBlueprintDependencyNodes(UBlueprint*
 	}
 	else
 	{
-		for (auto NodeIt = AllNodes.CreateIterator(); NodeIt; ++NodeIt)
+		for (UK2Node* Node : AllNodes)
 		{
-			UK2Node* Node = *NodeIt;
 			TArray<UStruct*> Dependencies;
 			if (Node->HasExternalDependencies(&Dependencies))
 			{
@@ -598,9 +598,8 @@ void FBlueprintEditorUtils::RefreshGraphNodes(const UEdGraph* Graph)
 	TArray<UK2Node*> AllNodes;
 	Graph->GetNodesOfClass(AllNodes);
 
-	for( auto NodeIt = AllNodes.CreateIterator(); NodeIt; ++NodeIt )
+	for (UK2Node* Node : AllNodes)
 	{
-		UK2Node* Node = *NodeIt;
 		const UEdGraphSchema* Schema = Node->GetGraph()->GetSchema();
 		Schema->ReconstructNode(*Node, true);
 	}
@@ -619,7 +618,7 @@ void FBlueprintEditorUtils::PreloadMembers(UObject* InObject)
 		UObject* CurrentObject = *it;
 		if( CurrentObject->HasAnyFlags(RF_NeedLoad) )
 		{
-			auto Linker = CurrentObject->GetLinker();
+			FLinkerLoad* Linker = CurrentObject->GetLinker();
 			if (Linker)
 			{
 				Linker->Preload(CurrentObject);
@@ -633,7 +632,7 @@ void FBlueprintEditorUtils::PreloadConstructionScript(UBlueprint* Blueprint)
 {
 	if ( Blueprint )
 	{
-		FLinkerLoad* TargetLinker = Blueprint->SimpleConstructionScript ? Blueprint->SimpleConstructionScript->GetLinker() : NULL;
+		FLinkerLoad* TargetLinker = Blueprint->SimpleConstructionScript ? Blueprint->SimpleConstructionScript->GetLinker() : nullptr;
 		if (TargetLinker)
 		{
 			TargetLinker->Preload(Blueprint->SimpleConstructionScript);
@@ -665,15 +664,15 @@ void FBlueprintEditorUtils::PreloadConstructionScript(UBlueprint* Blueprint)
 
 void FBlueprintEditorUtils::PatchNewCDOIntoLinker(UObject* CDO, FLinkerLoad* Linker, int32 ExportIndex, TArray<UObject*>& ObjLoaded)
 {
-	if( (CDO != NULL) && (Linker != NULL) && (ExportIndex != INDEX_NONE) )
+	if( (CDO != nullptr) && (Linker != nullptr) && (ExportIndex != INDEX_NONE) )
 	{
 		// Get rid of the old thing that was in its place
 		UObject* OldCDO = Linker->ExportMap[ExportIndex].Object;
-		if( OldCDO != NULL )
+		if( OldCDO != nullptr )
 		{
 			EObjectFlags OldObjectFlags = OldCDO->GetFlags();
 			OldCDO->ClearFlags(RF_NeedLoad|RF_NeedPostLoad);
-			OldCDO->SetLinker(NULL, INDEX_NONE);
+			OldCDO->SetLinker(nullptr, INDEX_NONE);
 			
 			// Copy flags from the old CDO.
 			CDO->SetFlags(OldObjectFlags);
@@ -729,7 +728,7 @@ void FBlueprintEditorUtils::GetAllGraphNames(const UBlueprint* Blueprint, TArray
 	for (int32 StackIndex = ParentBPStack.Num() - 1; StackIndex >= 0; --StackIndex)
 	{
 		UBlueprint* ParentBP = ParentBPStack[StackIndex];
-		check(ParentBP != NULL);
+		check(ParentBP != nullptr);
 
 		for(int32 FunctionIndex = 0; FunctionIndex < ParentBP->FunctionGraphs.Num(); ++FunctionIndex)
 		{
@@ -754,7 +753,7 @@ void FBlueprintEditorUtils::GetCompilerRelevantNodeLinks(UEdGraphPin* FromPin, F
 				if(FromPin)
 				{
 					// Recursively check each link for a compiler-relevant node that will "pass through" this node at compile time
-					for(auto LinkedPin : FromPin->LinkedTo)
+					for (UEdGraphPin* LinkedPin : FromPin->LinkedTo)
 					{
 						GetCompilerRelevantNodeLinks(LinkedPin, OutNodeLinks);
 					}
@@ -837,7 +836,7 @@ public:
 
 	void UpdateReferences()
 	{
-		if( TargetBlueprint != NULL && (TargetBlueprint->BlueprintType != BPTYPE_MacroLibrary) )
+		if( TargetBlueprint != nullptr && (TargetBlueprint->BlueprintType != BPTYPE_MacroLibrary) )
 		{
 			if( ensureMsgf(TargetBlueprint->SkeletonGeneratedClass, TEXT("Blueprint %s is missing its skeleton generated class - known possible for assets on revision 1 older than 2088505"), *TargetBlueprint->GetName() ) )
 			{
@@ -849,9 +848,9 @@ public:
 			TArray<UObject*> SubObjs;
 			GetObjectsWithOuter(TargetBlueprint, SubObjs, true);
 
-			for( auto ObjIt = SubObjs.CreateIterator(); ObjIt; ++ObjIt )
+			for (UObject* SubObj : SubObjs)
 			{
-				(*ObjIt)->Serialize(*this);
+				SubObj->Serialize(*this);
 			}
 
 			TargetBlueprint->bLegacyNeedToPurgeSkelRefs = false;
@@ -870,7 +869,7 @@ protected:
 	FArchive& operator<<( UObject*& Object )
 	{
 		// Check if this is a reference to an object existing in the transient package, and if so, NULL it.
-		if (Object != NULL )
+		if (Object != nullptr )
 		{
 			if( UClass* RefClass = Cast<UClass>(Object) )
 			{
@@ -902,7 +901,7 @@ struct FRegenerationHelper
 {
 	static bool ForcedLoad(UObject* Obj)
 	{
-		auto Linker = Obj->GetLinker();
+		FLinkerLoad* Linker = Obj->GetLinker();
 		if (Linker && !Obj->HasAnyFlags(RF_LoadCompleted))
 		{
 			Obj->SetFlags(RF_NeedLoad);
@@ -945,7 +944,7 @@ struct FRegenerationHelper
 		bool bChanged = false;
 		if (Struct->HasAnyFlags(RF_NeedLoad))
 		{
-			if (auto Linker = Struct->GetLinker())
+			if (FLinkerLoad* Linker = Struct->GetLinker())
 			{
 				Linker->Preload(Struct);
 				bChanged = true;
@@ -969,7 +968,7 @@ struct FRegenerationHelper
 
 	static UBlueprint* GetGeneratingBlueprint(const UObject* Obj)
 	{
-		const UBlueprintGeneratedClass* BPGC = NULL;
+		const UBlueprintGeneratedClass* BPGC = nullptr;
 		while (!BPGC && Obj)
 		{
 			BPGC = Cast<const UBlueprintGeneratedClass>(Obj);
@@ -1002,13 +1001,13 @@ struct FRegenerationHelper
 
 	static void PreloadMacroSources(TSet<UBlueprint*>& MacroSources)
 	{
-		for (auto BP : MacroSources)
+		for (UBlueprint* BP : MacroSources)
 		{
 			if (!BP->bHasBeenRegenerated)
 			{
 				if (BP->HasAnyFlags(RF_NeedLoad))
 				{
-					if (auto Linker = BP->GetLinker())
+					if (FLinkerLoad* Linker = BP->GetLinker())
 					{
 						Linker->Preload(BP);
 					}
@@ -1041,7 +1040,7 @@ struct FRegenerationHelper
 		for (FBPInterfaceDescription const& InterfaceDesc : Blueprint->ImplementedInterfaces)
 		{
 			UClass* InterfaceClass = InterfaceDesc.Interface;
-			UBlueprint* InterfaceBlueprint = InterfaceClass ? Cast<UBlueprint>(InterfaceClass->ClassGeneratedBy) : NULL;
+			UBlueprint* InterfaceBlueprint = InterfaceClass ? Cast<UBlueprint>(InterfaceClass->ClassGeneratedBy) : nullptr;
 			if (InterfaceBlueprint)
 			{
 				ForcedLoadMembers(InterfaceBlueprint);
@@ -1061,11 +1060,11 @@ struct FRegenerationHelper
 		TSet<UStruct*> Dependencies;
 		ProcessHierarchy(Blueprint->ParentClass, Dependencies);
 		
-		for (const auto& NewVar : Blueprint->NewVariables)
+		for (const FBPVariableDescription& NewVar : Blueprint->NewVariables)
 		{
 			if (UObject* TypeObject = NewVar.VarType.PinSubCategoryObject.Get())
 			{
-				auto Linker = TypeObject->GetLinker();
+				FLinkerLoad* Linker = TypeObject->GetLinker();
 				if (Linker && TypeObject->HasAnyFlags(RF_NeedLoad))
 				{
 					Linker->Preload(TypeObject);
@@ -1074,7 +1073,7 @@ struct FRegenerationHelper
 
 			if (UClass* TypeClass = NewVar.VarType.PinSubCategoryMemberReference.GetMemberParentClass())
 			{
-				auto Linker = TypeClass->GetLinker();
+				FLinkerLoad* Linker = TypeClass->GetLinker();
 				if (Linker && TypeClass->HasAnyFlags(RF_NeedLoad))
 				{
 					Linker->Preload(TypeClass);
@@ -1085,7 +1084,7 @@ struct FRegenerationHelper
 		TSet<UBlueprint*> MacroSources;
 		TArray<UEdGraph*> Graphs;
 		Blueprint->GetAllGraphs(Graphs);
-		for (auto Graph : Graphs)
+		for (UEdGraph* Graph : Graphs)
 		{
 			if (Graph && !FBlueprintEditorUtils::IsGraphIntermediate(Graph))
 			{
@@ -1093,19 +1092,19 @@ struct FRegenerationHelper
 
 				TArray<UK2Node*> Nodes;
 				Graph->GetNodesOfClass(Nodes);
-				for (auto Node : Nodes)
+				for (UK2Node* Node : Nodes)
 				{
 					if (Node)
 					{
 						TArray<UStruct*> LocalDependentStructures;
 						if (Node->HasExternalDependencies(&LocalDependentStructures))
 						{
-							for (auto Struct : LocalDependentStructures)
+							for (UStruct* Struct : LocalDependentStructures)
 							{
 								ProcessHierarchy(Struct, Dependencies);
 							}
 
-							if (auto MacroNode = Cast<UK2Node_MacroInstance>(Node))
+							if (UK2Node_MacroInstance* MacroNode = Cast<UK2Node_MacroInstance>(Node))
 							{
 								if (UBlueprint* MacroSource = MacroNode->GetSourceBlueprint())
 								{
@@ -1117,7 +1116,7 @@ struct FRegenerationHelper
 							// the external BP class (and thus the array property itself) will have been loaded/processed via the above
 							// ProcessHierarchy() call. However, the array's 'Inner' property may not have been preloaded as part of that path.
 							// Thus, we handle that here in order to ensure that all 'Inner' fields are also valid before class regeneration.
-							if (auto VariableNode = Cast<UK2Node_Variable>(Node))
+							if (UK2Node_Variable* VariableNode = Cast<UK2Node_Variable>(Node))
 							{
 								UArrayProperty* ArrayProperty = Cast<UArrayProperty>(VariableNode->VariableReference.ResolveMember<UProperty>(Node->GetBlueprintClassFromNode()));
 								if (ArrayProperty != nullptr && ArrayProperty->Inner != nullptr && ArrayProperty->Inner->HasAnyFlags(RF_NeedLoad|RF_WasLoaded))
@@ -1127,13 +1126,13 @@ struct FRegenerationHelper
 							}
 						}
 
-						auto FunctionEntry = Cast<UK2Node_FunctionEntry>(Node);
+						UK2Node_FunctionEntry* FunctionEntry = Cast<UK2Node_FunctionEntry>(Node);
 						if (FunctionEntry && !bIsDelegateSignatureGraph)
 						{
 							const FName FunctionName = (FunctionEntry->CustomGeneratedFunctionName != NAME_None) 
 								? FunctionEntry->CustomGeneratedFunctionName 
 								: FunctionEntry->SignatureName;
-							UFunction* ParentFunction = Blueprint->ParentClass ? Blueprint->ParentClass->FindFunctionByName(FunctionName) : NULL;
+							UFunction* ParentFunction = Blueprint->ParentClass ? Blueprint->ParentClass->FindFunctionByName(FunctionName) : nullptr;
 							if (ParentFunction && (Schema->FN_UserConstructionScript != FunctionName))
 							{
 								ProcessHierarchy(ParentFunction, Dependencies);
@@ -1141,9 +1140,9 @@ struct FRegenerationHelper
 						}
 
 						// load Enums
-						for (auto Pin : Node->Pins)
+						for (UEdGraphPin* Pin : Node->Pins)
 						{
-							auto SubCategoryObject = Pin ? Pin->PinType.PinSubCategoryObject.Get() : NULL;
+							UObject* SubCategoryObject = Pin ? Pin->PinType.PinSubCategoryObject.Get() : nullptr;
 							if (SubCategoryObject && SubCategoryObject->IsA<UEnum>())
 							{
 								ForcedLoad(SubCategoryObject);
@@ -1191,7 +1190,7 @@ struct FEditoronlyBlueprintHelper
 	{
 		check(Blueprint);
 		//any blueprint variable
-		for (const auto& VarDesc : Blueprint->NewVariables)
+		for (const FBPVariableDescription& VarDesc : Blueprint->NewVariables)
 		{
 			if (IsUnwantedType(VarDesc.VarType))
 			{
@@ -1207,13 +1206,13 @@ struct FEditoronlyBlueprintHelper
 		//anything on graph
 		TArray<UEdGraph*> Graphs;
 		Blueprint->GetAllGraphs(Graphs);
-		for (const auto Graph : Graphs)
+		for (const UEdGraph* Graph : Graphs)
 		{
 			TArray<UK2Node*> Nodes;
 			Graph->GetNodesOfClass(Nodes);
-			for (const auto Node : Nodes)
+			for (const UK2Node* Node : Nodes)
 			{
-				for (const auto Pin : Node->Pins)
+				for (const UEdGraphPin* Pin : Node->Pins)
 				{
 					if (Pin)
 					{
@@ -1259,7 +1258,7 @@ struct FEditoronlyBlueprintHelper
 
 		const UEdGraphSchema_K2* Schema = GetDefault<UEdGraphSchema_K2>();
 
-		for (auto UserPinInfo : Node->UserDefinedPins)
+		for (const TSharedPtr<FUserPinInfo>& UserPinInfo : Node->UserDefinedPins)
 		{
 			if (UserPinInfo.IsValid() && IsUnwantedType(UserPinInfo->PinType))
 			{
@@ -1270,7 +1269,7 @@ struct FEditoronlyBlueprintHelper
 			}
 		}
 
-		for (auto Pin : Node->Pins)
+		for (UEdGraphPin* Pin : Node->Pins)
 		{
 			if (IsUnwantedType(Pin->PinType))
 			{
@@ -1279,7 +1278,7 @@ struct FEditoronlyBlueprintHelper
 				if (Pin->DefaultObject)
 				{
 					UBlueprint* DefaultBlueprint = Cast<UBlueprint>(Pin->DefaultObject);
-					Pin->DefaultObject = DefaultBlueprint ? *DefaultBlueprint->GeneratedClass : NULL;
+					Pin->DefaultObject = DefaultBlueprint ? *DefaultBlueprint->GeneratedClass : nullptr;
 				}
 			}
 
@@ -1290,7 +1289,7 @@ struct FEditoronlyBlueprintHelper
 	{
 		check(Node);
 
-		auto Pin = Node->GetVariablePin();
+		UEdGraphPin* Pin = Node->GetVariablePin();
 		if (Pin && IsUnwantedType(Pin->PinType))
 		{
 			ChangePinType(Pin->PinType);
@@ -1298,7 +1297,7 @@ struct FEditoronlyBlueprintHelper
 			if (Pin->DefaultObject)
 			{
 				UBlueprint* DefaultBlueprint = Cast<UBlueprint>(Pin->DefaultObject);
-				Pin->DefaultObject = DefaultBlueprint ? *DefaultBlueprint->GeneratedClass : NULL;
+				Pin->DefaultObject = DefaultBlueprint ? *DefaultBlueprint->GeneratedClass : nullptr;
 			}
 		}
 
@@ -1312,7 +1311,7 @@ struct FEditoronlyBlueprintHelper
 	{
 		if (Node)
 		{
-			for (auto Pin : Node->Pins)
+			for (UEdGraphPin* Pin : Node->Pins)
 			{
 				if (Pin && IsUnwantedDefaultObject(Pin->DefaultObject))
 				{
@@ -1337,18 +1336,18 @@ struct FEditoronlyBlueprintHelper
 			//any native function call with blueprint parameter
 			TArray<UEdGraph*> Graphs;
 			Blueprint->GetAllGraphs(Graphs);
-			for (const auto Graph : Graphs)
+			for (UEdGraph* Graph : Graphs)
 			{
 				Schema->BackwardCompatibilityNodeConversion(Graph, false);
 			}
 
 			//any blueprint pin (function/event/macro parameter
-			for (const auto Graph : Graphs)
+			for (const UEdGraph* Graph : Graphs)
 			{
 				{
 					TArray<UK2Node_EditablePinBase*> EditablePinNodes;
 					Graph->GetNodesOfClass(EditablePinNodes);
-					for (const auto Node : EditablePinNodes)
+					for (UK2Node_EditablePinBase* Node : EditablePinNodes)
 					{
 						HandleEditablePinNode(Node);
 					}
@@ -1357,7 +1356,7 @@ struct FEditoronlyBlueprintHelper
 				{
 					TArray<UK2Node_TemporaryVariable*> TempVariablesNodes;
 					Graph->GetNodesOfClass(TempVariablesNodes);
-					for (const auto Node : TempVariablesNodes)
+					for (UK2Node_TemporaryVariable* Node : TempVariablesNodes)
 					{
 						HandleTemporaryVariableNode(Node);
 					}
@@ -1365,7 +1364,7 @@ struct FEditoronlyBlueprintHelper
 			}
 
 			// change variables
-			for (auto& VarDesc : Blueprint->NewVariables)
+			for (FBPVariableDescription& VarDesc : Blueprint->NewVariables)
 			{
 				if (IsUnwantedType(VarDesc.VarType))
 				{
@@ -1373,11 +1372,11 @@ struct FEditoronlyBlueprintHelper
 				}
 			}
 
-			for (const auto Graph : Graphs)
+			for (const UEdGraph* Graph : Graphs)
 			{
 				TArray<UK2Node*> Nodes;
 				Graph->GetNodesOfClass(Nodes);
-				for (const auto Node : Nodes)
+				for (UK2Node* Node : Nodes)
 				{
 					HandleDefaultObjects(Node);
 				}
@@ -1483,7 +1482,7 @@ UClass* FBlueprintEditorUtils::RegenerateBlueprintClass(UBlueprint* Blueprint, U
 		Blueprint->GetBlueprintCDONames(GeneratedName, SkeletonName);
 		int32 OldSkelLinkerIdx = INDEX_NONE;
 		int32 OldGenLinkerIdx = INDEX_NONE;
-		auto OldLinker = Blueprint->GetLinker();
+		FLinkerLoad* OldLinker = Blueprint->GetLinker();
 		for( int32 i = 0; i < OldLinker->ExportMap.Num(); i++ )
 		{
 			FObjectExport& ThisExport = OldLinker->ExportMap[i];
@@ -1532,7 +1531,7 @@ UClass* FBlueprintEditorUtils::RegenerateBlueprintClass(UBlueprint* Blueprint, U
 		const bool bDataOnlyClassThatMustBeRecompiled = !bHasCode && !bIsMacro
 			&& (!ClassToRegenerate || (Blueprint->ParentClass != ClassToRegenerate->GetSuperClass()));
 
-		auto BPGClassToRegenerate = Cast<UBlueprintGeneratedClass>(ClassToRegenerate);
+		UBlueprintGeneratedClass* BPGClassToRegenerate = Cast<UBlueprintGeneratedClass>(ClassToRegenerate);
 		const bool bHasPendingUberGraphFrame = BPGClassToRegenerate
 			&& (BPGClassToRegenerate->UberGraphFramePointerProperty || BPGClassToRegenerate->UberGraphFunction);
 
@@ -1590,18 +1589,18 @@ UClass* FBlueprintEditorUtils::RegenerateBlueprintClass(UBlueprint* Blueprint, U
 		}
 		else
 		{
-			if (Blueprint->IsGeneratedClassAuthoritative() && (Blueprint->GeneratedClass != NULL))
+			if (Blueprint->IsGeneratedClassAuthoritative() && (Blueprint->GeneratedClass != nullptr))
 			{
 				RemoveStaleFunctions(Cast<UBlueprintGeneratedClass>(Blueprint->GeneratedClass), Blueprint);
 
-				check(PreviousCDO != NULL);
-				check(Blueprint->SkeletonGeneratedClass != NULL);
+				check(PreviousCDO != nullptr);
+				check(Blueprint->SkeletonGeneratedClass != nullptr);
 
 				// We now know we're a data-only blueprint on the outer pass (generate class is valid), where generated class is authoritative
 				// If the PreviousCDO is to the skeleton, then it will corrupt data when copied over the AuthoriativeClass later on in this function
 				if (PreviousCDO == Blueprint->SkeletonGeneratedClass->GetDefaultObject())
 				{
-					check(Blueprint->PRIVATE_InnermostPreviousCDO == NULL);
+					check(Blueprint->PRIVATE_InnermostPreviousCDO == nullptr);
 					Blueprint->PRIVATE_InnermostPreviousCDO = Blueprint->GeneratedClass->GetDefaultObject();
 				}
 			}
@@ -1635,7 +1634,7 @@ UClass* FBlueprintEditorUtils::RegenerateBlueprintClass(UBlueprint* Blueprint, U
 
 		// Success or failure, there's no point in trying to recompile this class again when other objects reference it
 		// redo data only blueprints later, when we actually have a generated class
-		Blueprint->bHasBeenRegenerated = !FBlueprintEditorUtils::IsDataOnlyBlueprint(Blueprint) || Blueprint->GeneratedClass != NULL; 
+		Blueprint->bHasBeenRegenerated = !FBlueprintEditorUtils::IsDataOnlyBlueprint(Blueprint) || Blueprint->GeneratedClass != nullptr; 
 
 		Blueprint->bIsRegeneratingOnLoad = false;
 
@@ -1667,7 +1666,7 @@ UClass* FBlueprintEditorUtils::RegenerateBlueprintClass(UBlueprint* Blueprint, U
 
 				// If this is the top of the compile-on-load stack for this object, copy the old CDO properties to the newly created one unless they are the same
 				UClass* AuthoritativeClass = (Blueprint->IsGeneratedClassAuthoritative() ? Blueprint->GeneratedClass : Blueprint->SkeletonGeneratedClass);
-				if (AuthoritativeClass != NULL && PreviousCDO != AuthoritativeClass->GetDefaultObject())
+				if (AuthoritativeClass != nullptr && PreviousCDO != AuthoritativeClass->GetDefaultObject())
 				{
 					TGuardValue<bool> GuardTemplateNameFlag(GCompilingBlueprint, true);
 
@@ -1700,7 +1699,7 @@ UClass* FBlueprintEditorUtils::RegenerateBlueprintClass(UBlueprint* Blueprint, U
 					}
 				}
 
-				Blueprint->PRIVATE_InnermostPreviousCDO = NULL;
+				Blueprint->PRIVATE_InnermostPreviousCDO = nullptr;
 			}
 			else
 			{
@@ -1709,7 +1708,7 @@ UClass* FBlueprintEditorUtils::RegenerateBlueprintClass(UBlueprint* Blueprint, U
 			}
 
 			// If this is the top of the compile-on-load stack for this object, copy the old CDO properties to the newly created one
-			if (!Blueprint->IsGeneratedClassAuthoritative() && Blueprint->GeneratedClass != NULL)
+			if (!Blueprint->IsGeneratedClassAuthoritative() && Blueprint->GeneratedClass != nullptr)
 			{
 				TGuardValue<bool> GuardTemplateNameFlag(GCompilingBlueprint, true);
 
@@ -1763,7 +1762,7 @@ UClass* FBlueprintEditorUtils::RegenerateBlueprintClass(UBlueprint* Blueprint, U
 		Package->SetDirtyFlag(bIsPackageDirty);
 	}
 
-	return bRegenerated ? Blueprint->GeneratedClass : NULL;
+	return bRegenerated ? Blueprint->GeneratedClass : nullptr;
 }
 
 
@@ -1790,7 +1789,7 @@ void FBlueprintEditorUtils::RecreateClassMetaData(UBlueprint* Blueprint, UClass*
 		Class->RemoveMetaData(FBlueprintMetadata::MD_AllowableBlueprintVariableType);
 	}
 
-	if (ensure(ParentClass != NULL))
+	if (ensure(ParentClass != nullptr))
 	{
 		if (!ParentClass->HasMetaData(FBlueprintMetadata::MD_IgnoreCategoryKeywordsInSubclasses))
 		{
@@ -1847,7 +1846,7 @@ void FBlueprintEditorUtils::RecreateClassMetaData(UBlueprint* Blueprint, UClass*
 
 	for (FString HideCategory : Blueprint->HideCategories)
 	{
-		auto& CharArray = HideCategory.GetCharArray();
+		TArray<TCHAR>& CharArray = HideCategory.GetCharArray();
 
 		int32 SpaceIndex = CharArray.Find(TEXT(' '));
 		while (SpaceIndex != INDEX_NONE)
@@ -1946,7 +1945,7 @@ void FBlueprintEditorUtils::PropagateParentBlueprintDefaults(UClass* ClassToProp
 	// Starting from the least derived BP class, copy the properties into the new CDO
 	for(int32 i = ParentBP.Num() - 1; i > 0; i--)
 	{
-		checkf(ParentBP[i]->GeneratedClass != NULL, TEXT("Parent classes for class %s have not yet been generated.  Compile-on-load must be processed for the parent class first."), *ClassToPropagate->GetName());
+		checkf(ParentBP[i]->GeneratedClass != nullptr, TEXT("Parent classes for class %s have not yet been generated.  Compile-on-load must be processed for the parent class first."), *ClassToPropagate->GetName());
 		UObject* LayerCDO = ParentBP[i]->GeneratedClass->GetDefaultObject();
 
 		UEditorEngine::FCopyPropertiesForUnrelatedObjectsParams CopyDetails;
@@ -1966,13 +1965,13 @@ void FBlueprintEditorUtils::PostDuplicateBlueprint(UBlueprint* Blueprint, bool b
 	// Only recompile after duplication if this isn't PIE
 	if (!bDuplicateForPIE)
 	{
-		check(Blueprint->GeneratedClass != NULL);
+		check(Blueprint->GeneratedClass != nullptr);
 		{
 			// Grab the old CDO, which contains the class defaults
 			UClass* OldBPGCAsClass = Blueprint->GeneratedClass;
 			UBlueprintGeneratedClass* OldBPGC = (UBlueprintGeneratedClass*)(OldBPGCAsClass);
 			UObject* OldCDO = OldBPGC->GetDefaultObject();
-			check(OldCDO != NULL);
+			check(OldCDO != nullptr);
 
 			if (FBlueprintDuplicationScopeFlags::HasAnyFlag(FBlueprintDuplicationScopeFlags::ValidatePinsUsingSourceClass))
 			{
@@ -1981,10 +1980,10 @@ void FBlueprintEditorUtils::PostDuplicateBlueprint(UBlueprint* Blueprint, bool b
 
 			// Grab the old class templates, which needs to be moved to the new class
 			USimpleConstructionScript* SCSRootNode = Blueprint->SimpleConstructionScript;
-			Blueprint->SimpleConstructionScript = NULL;
+			Blueprint->SimpleConstructionScript = nullptr;
 
 			UInheritableComponentHandler* InheritableComponentHandler = Blueprint->InheritableComponentHandler;
-			Blueprint->InheritableComponentHandler = NULL;
+			Blueprint->InheritableComponentHandler = nullptr;
 
 			TArray<UActorComponent*> Templates = Blueprint->ComponentTemplates;
 			Blueprint->ComponentTemplates.Empty();
@@ -2050,18 +2049,16 @@ void FBlueprintEditorUtils::PostDuplicateBlueprint(UBlueprint* Blueprint, bool b
 				}
 			}
 
-			for(auto CompIt = Templates.CreateIterator(); CompIt; ++CompIt)
+			for (UActorComponent* OldComponent : Templates)
 			{
-				UActorComponent* OldComponent = *CompIt;
 				UActorComponent* NewComponent = CastChecked<UActorComponent>(StaticDuplicateObject(OldComponent, NewBPGC, OldComponent->GetFName()));
 
 				NewBPGC->ComponentTemplates.Add(NewComponent);
 				OldToNewMap.Add(OldComponent, NewComponent);
 			}
 
-			for(auto TimelineIt = Timelines.CreateIterator(); TimelineIt; ++TimelineIt)
+			for (UTimelineTemplate* OldTimeline : Timelines)
 			{
-				UTimelineTemplate* OldTimeline = *TimelineIt;
 				UTimelineTemplate* NewTimeline = CastChecked<UTimelineTemplate>(StaticDuplicateObject(OldTimeline, NewBPGC, OldTimeline->GetFName()));
 
 				if (FBlueprintDuplicationScopeFlags::HasAnyFlag(FBlueprintDuplicationScopeFlags::TheSameTimelineGuid))
@@ -2094,7 +2091,7 @@ void FBlueprintEditorUtils::PostDuplicateBlueprint(UBlueprint* Blueprint, bool b
 			// Give all nodes a new Guid
 			TArray< UEdGraphNode* > AllGraphNodes;
 			GetAllNodesOfClass(Blueprint, AllGraphNodes);
-			for(auto& Node : AllGraphNodes)
+			for(UEdGraphNode* Node : AllGraphNodes)
 			{
 				if (!FBlueprintDuplicationScopeFlags::HasAnyFlag(FBlueprintDuplicationScopeFlags::TheSameNodeGuid))
 				{
@@ -2135,10 +2132,10 @@ void FBlueprintEditorUtils::PostDuplicateBlueprint(UBlueprint* Blueprint, bool b
 			FArchiveReplaceObjectRef<UObject> ReplaceTemplateRefs(NewBPGC, OldToNewMap, /*bNullPrivateRefs=*/ false, /*bIgnoreOuterRef=*/ false, /*bIgnoreArchetypeRef=*/ false);
 
 			// Now propagate the values from the old CDO to the new one
-			check(Blueprint->SkeletonGeneratedClass != NULL);
+			check(Blueprint->SkeletonGeneratedClass != nullptr);
 
 			UObject* NewCDO = Blueprint->GeneratedClass->GetDefaultObject();
-			check(NewCDO != NULL);
+			check(NewCDO != nullptr);
 			UEditorEngine::CopyPropertiesForUnrelatedObjects(OldCDO, NewCDO);
 		}
 
@@ -2166,23 +2163,22 @@ void FBlueprintEditorUtils::UpdateDelegatesInBlueprint(UBlueprint* Blueprint)
 	check(Blueprint);
 	TArray<UEdGraph*> Graphs;
 	Blueprint->GetAllGraphs(Graphs);
-	for (auto GraphIt = Graphs.CreateIterator(); GraphIt; ++GraphIt)
+	for (UEdGraph* Graph : Graphs)
 	{
-		UEdGraph* Graph = (*GraphIt);
 		if(!IsGraphIntermediate(Graph))
 		{
 			TArray<UK2Node_CreateDelegate*> CreateDelegateNodes;
 			Graph->GetNodesOfClass(CreateDelegateNodes);
-			for (auto NodeIt = CreateDelegateNodes.CreateIterator(); NodeIt; ++NodeIt)
+			for (UK2Node_CreateDelegate* DelegateNode: CreateDelegateNodes)
 			{
-				(*NodeIt)->HandleAnyChangeWithoutNotifying();
+				DelegateNode->HandleAnyChangeWithoutNotifying();
 			}
 
 			TArray<UK2Node_Event*> EventNodes;
 			Graph->GetNodesOfClass(EventNodes);
-			for (auto NodeIt = EventNodes.CreateIterator(); NodeIt; ++NodeIt)
+			for (UK2Node_Event* EventNode : EventNodes)
 			{
-				(*NodeIt)->UpdateDelegatePin();
+				EventNode->UpdateDelegatePin();
 			}
 		}
 	}
@@ -2199,7 +2195,7 @@ void FBlueprintEditorUtils::MarkBlueprintAsStructurallyModified(UBlueprint* Blue
 		{
 			FSecondsCounterScope SkeletalRecompileTimer(BlueprintCompileAndLoadTimerData);
 			
-			for (auto SkelClass : SkelClassesToRecompile)
+			for (UClass* SkelClass : SkelClassesToRecompile)
 			{
 				if (SkelClass->HasAnyClassFlags(CLASS_NewerVersionExists))
 				{
@@ -2254,7 +2250,7 @@ void FBlueprintEditorUtils::MarkBlueprintAsStructurallyModified(UBlueprint* Blue
 		Blueprint->GetAllGraphs(AllGraphs);
 		for (int32 i = 0; i < AllGraphs.Num(); i++)
 		{
-			auto EntryNode = GetEntryNode(AllGraphs[i]);
+			UK2Node_EditablePinBase* EntryNode = GetEntryNode(AllGraphs[i]);
 			if(UK2Node_Tunnel* TunnelNode = ExactCast<UK2Node_Tunnel>(EntryNode))
 			{
 				// Remove data marking graphs as latent, this will be re-cache'd as needed
@@ -2305,7 +2301,7 @@ void FBlueprintEditorUtils::MarkBlueprintAsModified(UBlueprint* Blueprint, FProp
 		Blueprint->GetAllGraphs(AllGraphs);
 		for (int32 i = 0; i < AllGraphs.Num(); i++)
 		{
-			auto EntryNode = GetEntryNode(AllGraphs[i]);
+			UK2Node_EditablePinBase* EntryNode = GetEntryNode(AllGraphs[i]);
 			if(UK2Node_Tunnel* TunnelNode = ExactCast<UK2Node_Tunnel>(EntryNode))
 			{
 				// Remove data marking graphs as latent, this will be re-cache'd as needed
@@ -2348,7 +2344,7 @@ bool FBlueprintEditorUtils::ShouldRegenerateBlueprint(UBlueprint* Blueprint)
 // Helper function to get the blueprint that ultimately owns a node.
 UBlueprint* FBlueprintEditorUtils::FindBlueprintForNode(const UEdGraphNode* Node)
 {
-	auto Graph = Node ? Cast<UEdGraph>(Node->GetOuter()) : NULL;
+	UEdGraph* Graph = Node ? Cast<UEdGraph>(Node->GetOuter()) : nullptr;
 	return FindBlueprintForGraph(Graph);
 }
 
@@ -2362,7 +2358,7 @@ UBlueprint* FBlueprintEditorUtils::FindBlueprintForNodeChecked(const UEdGraphNod
 // Helper function to get the blueprint that ultimately owns a graph.
 UBlueprint* FBlueprintEditorUtils::FindBlueprintForGraph(const UEdGraph* Graph)
 {
-	for (UObject* TestOuter = Graph ? Graph->GetOuter() : NULL; TestOuter; TestOuter = TestOuter->GetOuter())
+	for (UObject* TestOuter = Graph ? Graph->GetOuter() : nullptr; TestOuter; TestOuter = TestOuter->GetOuter())
 	{
 		if (UBlueprint* Result = Cast<UBlueprint>(TestOuter))
 		{
@@ -2370,7 +2366,7 @@ UBlueprint* FBlueprintEditorUtils::FindBlueprintForGraph(const UEdGraph* Graph)
 		}
 	}
 
-	return NULL;
+	return nullptr;
 }
 
 // Helper function to get the blueprint that ultimately owns a graph.  Cannot fail.
@@ -2415,7 +2411,7 @@ bool FBlueprintEditorUtils::IsGraphNameUnique(UBlueprint* Blueprint, const FName
 
 UEdGraph* FBlueprintEditorUtils::CreateNewGraph(UObject* ParentScope, const FName& GraphName, TSubclassOf<class UEdGraph> GraphClass, TSubclassOf<class UEdGraphSchema> SchemaClass)
 {
-	UEdGraph* NewGraph = NULL;
+	UEdGraph* NewGraph = nullptr;
 	bool bRename = false;
 
 	// Ensure this name isn't already being used for a graph
@@ -2427,7 +2423,7 @@ UEdGraph* FBlueprintEditorUtils::CreateNewGraph(UObject* ParentScope, const FNam
 		// Rename the old graph out of the way; but we have already failed at this point
 		if (ExistingGraph)
 		{
-			ExistingGraph->Rename(NULL, ExistingGraph->GetOuter(), REN_DoNotDirty | REN_ForceNoResetLoaders);
+			ExistingGraph->Rename(nullptr, ExistingGraph->GetOuter(), REN_DoNotDirty | REN_ForceNoResetLoaders);
 		}
 
 		// Construct new graph with the supplied name
@@ -2463,9 +2459,8 @@ UFunction* FBlueprintEditorUtils::FindFunctionInImplementedInterfaces(const UBlu
 		}
 
 		// Now loop through the interface classes and try and find the function
-		for (auto It = InterfaceClasses.CreateConstIterator(); It; It++)
+		for (UClass* SearchClass : InterfaceClasses)
 		{
-			UClass* SearchClass = (*It);
 			if( SearchClass )
 			{
 				do
@@ -2484,15 +2479,15 @@ UFunction* FBlueprintEditorUtils::FindFunctionInImplementedInterfaces(const UBlu
 		}
 	}
 
-	return NULL;
+	return nullptr;
 }
 
 void FBlueprintEditorUtils::FindImplementedInterfaces(const UBlueprint* Blueprint, bool bGetAllInterfaces, TArray<UClass*>& ImplementedInterfaces)
 {
 	// First get the ones this blueprint implemented
-	for (auto It = Blueprint->ImplementedInterfaces.CreateConstIterator(); It; It++)
+	for (const FBPInterfaceDescription& ImplementedInterface : Blueprint->ImplementedInterfaces)
 	{
-		ImplementedInterfaces.AddUnique((*It).Interface);
+		ImplementedInterfaces.AddUnique(ImplementedInterface.Interface);
 	}
 
 	if (bGetAllInterfaces)
@@ -2501,9 +2496,9 @@ void FBlueprintEditorUtils::FindImplementedInterfaces(const UBlueprint* Blueprin
 		UClass* BlueprintParent =  Blueprint->ParentClass;
 		while (BlueprintParent)
 		{
-			for (auto It = BlueprintParent->Interfaces.CreateConstIterator(); It; It++)
+			for (const FImplementedInterface& ImplementedInterface : BlueprintParent->Interfaces)
 			{
-				ImplementedInterfaces.AddUnique((*It).Class);
+				ImplementedInterfaces.AddUnique(ImplementedInterface.Class);
 			}
 			BlueprintParent = BlueprintParent->GetSuperClass();
 		}
@@ -2518,7 +2513,7 @@ void FBlueprintEditorUtils::AddMacroGraph( UBlueprint* Blueprint, class UEdGraph
 
 	Schema->CreateDefaultNodesForGraph(*Graph);
 
-	if (K2Schema != NULL)
+	if (K2Schema != nullptr)
 	{
 		K2Schema->CreateMacroGraphTerminators(*Graph, SignatureFromClass);
 
@@ -2543,7 +2538,7 @@ void FBlueprintEditorUtils::AddMacroGraph( UBlueprint* Blueprint, class UEdGraph
 void FBlueprintEditorUtils::AddInterfaceGraph(UBlueprint* Blueprint, class UEdGraph* Graph, UClass* InterfaceClass)
 {
 	const UEdGraphSchema_K2* K2Schema = Cast<const UEdGraphSchema_K2>(Graph->GetSchema());
-	if (K2Schema != NULL)
+	if (K2Schema != nullptr)
 	{
 		K2Schema->CreateFunctionGraphTerminators(*Graph, InterfaceClass);
 	}
@@ -2620,9 +2615,8 @@ void FBlueprintEditorUtils::RemoveGraph(UBlueprint* Blueprint, class UEdGraph* G
 				//removes all macro nodes using this macro graph
 				TArray<UK2Node_MacroInstance*> MacroNodes;
 				FBlueprintEditorUtils::GetAllNodesOfClass(Blueprint, MacroNodes);
-				for(auto It(MacroNodes.CreateIterator());It;++It)
+				for (UK2Node_MacroInstance* Node : MacroNodes)
 				{
-					UK2Node_MacroInstance* Node = *It;
 					if(Node->GetMacroGraph() == GraphToRemove)
 					{
 						FBlueprintEditorUtils::RemoveNode(Blueprint, Node);
@@ -2630,9 +2624,8 @@ void FBlueprintEditorUtils::RemoveGraph(UBlueprint* Blueprint, class UEdGraph* G
 				}
 			}
 
-			for( auto InterfaceIt = Blueprint->ImplementedInterfaces.CreateIterator(); InterfaceIt; ++InterfaceIt )
+			for (FBPInterfaceDescription& CurrInterface : Blueprint->ImplementedInterfaces)
 			{
-				FBPInterfaceDescription& CurrInterface = *InterfaceIt;
 				CurrInterface.Graphs.Remove( GraphToRemove );
 			}
 		}
@@ -2654,9 +2647,8 @@ void FBlueprintEditorUtils::RemoveGraph(UBlueprint* Blueprint, class UEdGraph* G
 	// Remove timelines held in the graph
 	TArray<UK2Node_Timeline*> AllTimelineNodes;
 	GraphToRemove->GetNodesOfClass<UK2Node_Timeline>(AllTimelineNodes);
-	for (auto TimelineIt = AllTimelineNodes.CreateIterator(); TimelineIt; ++TimelineIt)
+	for (UK2Node_Timeline* TimelineNode : AllTimelineNodes)
 	{
-		UK2Node_Timeline* TimelineNode = *TimelineIt;
 		TimelineNode->DestroyNode();
 	}
 
@@ -2665,9 +2657,8 @@ void FBlueprintEditorUtils::RemoveGraph(UBlueprint* Blueprint, class UEdGraph* G
 	GraphToRemove->GetNodesOfClass<UK2Node_Composite>(AllCompositeNodes);
 
 	const bool bDontRecompile = true;
-	for (auto CompIt = AllCompositeNodes.CreateIterator(); CompIt; ++CompIt)
+	for (UK2Node_Composite* CompNode : AllCompositeNodes)
 	{
-		UK2Node_Composite* CompNode = *CompIt;
 		if (CompNode->BoundGraph && Local::IsASubGraph(CompNode->BoundGraph))
 		{
 			FBlueprintEditorUtils::RemoveGraph(Blueprint, CompNode->BoundGraph, EGraphRemoveFlags::None);
@@ -2704,7 +2695,7 @@ void FBlueprintEditorUtils::RemoveGraph(UBlueprint* Blueprint, class UEdGraph* G
 
 	GraphToRemove->GetSchema()->HandleGraphBeingDeleted(*GraphToRemove);
 
-	GraphToRemove->Rename(NULL, Blueprint->GetOuter(), REN_DoNotDirty | REN_DontCreateRedirectors);
+	GraphToRemove->Rename(nullptr, Blueprint->GetOuter(), REN_DoNotDirty | REN_DontCreateRedirectors);
 	GraphToRemove->ClearFlags(RF_Standalone | RF_Public);
 	GraphToRemove->RemoveFromRoot();
 
@@ -2736,9 +2727,8 @@ void FBlueprintEditorUtils::RenameGraph(UEdGraph* Graph, const FString& NewNameS
 		Graph->Rename(*NewNameStr, Graph->GetOuter(), (Blueprint->bIsRegeneratingOnLoad ? REN_ForceNoResetLoaders : 0) | REN_DontCreateRedirectors);
 
 		// Clean function entry & result nodes if they exist
-		for (auto NodeIt = Graph->Nodes.CreateIterator(); NodeIt; ++NodeIt)
+		for (UEdGraphNode* Node : Graph->Nodes)
 		{
-			UEdGraphNode* Node = *NodeIt;
 			if (UK2Node_FunctionEntry* EntryNode = Cast<UK2Node_FunctionEntry>(Node))
 			{
 				if (EntryNode->SignatureName == OldGraphName)
@@ -2766,10 +2756,8 @@ void FBlueprintEditorUtils::RenameGraph(UEdGraph* Graph, const FString& NewNameS
 		TArray<UK2Node_CallFunction*> AllFunctionCalls;
 		FBlueprintEditorUtils::GetAllNodesOfClass<UK2Node_CallFunction>(Blueprint, AllFunctionCalls);
 	
-		for (auto FuncIt = AllFunctionCalls.CreateIterator(); FuncIt; ++FuncIt)
+		for (UK2Node_CallFunction* FunctionNode : AllFunctionCalls)
 		{
-			UK2Node_CallFunction* FunctionNode = *FuncIt;
-
 			if (FunctionNode->FunctionReference.IsSelfContext() && (FunctionNode->FunctionReference.GetMemberName() == OldGraphName))
 			{
 				FunctionNode->Modify();
@@ -2799,7 +2787,7 @@ void FBlueprintEditorUtils::RemoveNode(UBlueprint* Blueprint, UEdGraphNode* Node
 {
 	check(Node);
 
-	const UEdGraphSchema* Schema = NULL;
+	const UEdGraphSchema* Schema = nullptr;
 
 	// Ensure we mark parent graph modified
 	if (UEdGraph* GraphObj = Node->GetGraph())
@@ -2808,7 +2796,7 @@ void FBlueprintEditorUtils::RemoveNode(UBlueprint* Blueprint, UEdGraphNode* Node
 		Schema = GraphObj->GetSchema();
 	}
 
-	if (Blueprint != NULL)
+	if (Blueprint != nullptr)
 	{
 		// Remove any breakpoints set on the node
 		if (UBreakpoint* Breakpoint = FKismetDebugUtilities::FindBreakpointForNode(Blueprint, Node))
@@ -2833,7 +2821,7 @@ void FBlueprintEditorUtils::RemoveNode(UBlueprint* Blueprint, UEdGraphNode* Node
 
 	Node->DestroyNode(); 
 
-	if (!bDontRecompile && (Blueprint != NULL))
+	if (!bDontRecompile && (Blueprint != nullptr))
 	{
 		FBlueprintEditorUtils::MarkBlueprintAsStructurallyModified(Blueprint);
 	}
@@ -2868,7 +2856,7 @@ bool FBlueprintEditorUtils::IsGraphReadOnly(UEdGraph* InGraph)
 		if (!bGraphReadOnly)
 		{
 			const UBlueprint* BlueprintForGraph = FBlueprintEditorUtils::FindBlueprintForGraph(InGraph);
-			bool const bIsInterface = ((BlueprintForGraph != NULL) && (BlueprintForGraph->BlueprintType == BPTYPE_Interface));
+			bool const bIsInterface = ((BlueprintForGraph != nullptr) && (BlueprintForGraph->BlueprintType == BPTYPE_Interface));
 			bool const bIsDelegate  = FBlueprintEditorUtils::IsDelegateSignatureGraph(InGraph);
 			bool const bIsMathExpression = FBlueprintEditorUtils::IsMathExpressionGraph(InGraph);
 
@@ -2895,7 +2883,7 @@ UK2Node_Event* FBlueprintEditorUtils::FindOverrideForFunction(const UBlueprint*
 		}
 	}
 
-	return NULL;
+	return nullptr;
 }
 
 UK2Node_Event* FBlueprintEditorUtils::FindCustomEventNode(const UBlueprint* Blueprint, FName const CustomName)
@@ -2926,7 +2914,7 @@ void FBlueprintEditorUtils::GatherDependencies(const UBlueprint* InBlueprint, TS
 	{
 		static UBlueprint* GetGeneratingBlueprint(const UObject* Obj)
 		{
-			const UBlueprintGeneratedClass* BPGC = NULL;
+			const UBlueprintGeneratedClass* BPGC = nullptr;
 			while (!BPGC && Obj)
 			{
 				BPGC = Cast<const UBlueprintGeneratedClass>(Obj);
@@ -2969,7 +2957,7 @@ void FBlueprintEditorUtils::GatherDependencies(const UBlueprint* InBlueprint, TS
 
 	FGatherDependenciesHelper::ProcessHierarchy(InBlueprint->ParentClass, Dependencies);
 
-	for (const auto& InterfaceDesc : InBlueprint->ImplementedInterfaces)
+	for (const FBPInterfaceDescription& InterfaceDesc : InBlueprint->ImplementedInterfaces)
 	{
 		UBlueprint* InterfaceBP = InterfaceDesc.Interface ? Cast<UBlueprint>(InterfaceDesc.Interface->ClassGeneratedBy) : nullptr;
 		if (InterfaceBP)
@@ -2980,20 +2968,20 @@ void FBlueprintEditorUtils::GatherDependencies(const UBlueprint* InBlueprint, TS
 
 	TArray<UEdGraph*> Graphs;
 	InBlueprint->GetAllGraphs(Graphs);
-	for (auto Graph : Graphs)
+	for (UEdGraph* Graph : Graphs)
 	{
 		if (Graph && !FBlueprintEditorUtils::IsGraphIntermediate(Graph))
 		{
 			TArray<UK2Node*> Nodes;
 			Graph->GetNodesOfClass(Nodes);
-			for (auto Node : Nodes)
+			for (UK2Node* Node : Nodes)
 			{
 				TArray<UStruct*> LocalDependentStructures;
 				if (Node && Node->HasExternalDependencies(&LocalDependentStructures))
 				{
-					for (auto Struct : LocalDependentStructures)
+					for (UStruct* Struct : LocalDependentStructures)
 					{
-						if (auto UDS = Cast<UUserDefinedStruct>(Struct))
+						if (UUserDefinedStruct* UDS = Cast<UUserDefinedStruct>(Struct))
 						{
 							OutUDSDependencies.Add(UDS);
 						}
@@ -3025,10 +3013,10 @@ void FBlueprintEditorUtils::GetDependentBlueprints(UBlueprint* Blueprint, TArray
 	bool const bIncludeDerivedClasses = true;
 	GetObjectsOfClass(UBlueprint::StaticClass(), AllBlueprints, bIncludeDerivedClasses );
 
-	for (auto ObjIt = AllBlueprints.CreateIterator(); ObjIt; ++ObjIt)
+	for (UObject* Obj : AllBlueprints)
 	{
 		// we know the class is correct so a fast cast is ok here
-		UBlueprint* TestBP = (UBlueprint*) *ObjIt;
+		UBlueprint* TestBP = static_cast<UBlueprint*>(Obj);
 		if (TestBP && !TestBP->IsPendingKill())
 		{
 			EnsureCachedDependenciesUpToDate(TestBP);
@@ -3093,7 +3081,7 @@ bool FBlueprintEditorUtils::IsDataOnlyBlueprint(const UBlueprint* Blueprint)
 	}
 	
 	// No extra functions, other than the user construction script(only AActor and subclasses of AActor have)
-	auto DefaultFunctionNum = (Blueprint->ParentClass && Blueprint->ParentClass->IsChildOf(AActor::StaticClass())) ? 1 : 0;
+	const int32 DefaultFunctionNum = (Blueprint->ParentClass && Blueprint->ParentClass->IsChildOf(AActor::StaticClass())) ? 1 : 0;
 	if ((Blueprint->FunctionGraphs.Num() > DefaultFunctionNum) || (Blueprint->MacroGraphs.Num() > 0))
 	{
 		return false;
@@ -3123,7 +3111,7 @@ bool FBlueprintEditorUtils::IsDataOnlyBlueprint(const UBlueprint* Blueprint)
 	}
 
 	// Make sure there's nothing in the user construction script, other than an entry node
-	UEdGraph* UserConstructionScript = (Blueprint->FunctionGraphs.Num() == 1) ? Blueprint->FunctionGraphs[0] : NULL;
+	UEdGraph* UserConstructionScript = (Blueprint->FunctionGraphs.Num() == 1) ? Blueprint->FunctionGraphs[0] : nullptr;
 	if (UserConstructionScript && Blueprint->ParentClass)
 	{
 		//Call parent construction script may be added automatically
@@ -3143,7 +3131,7 @@ bool FBlueprintEditorUtils::IsDataOnlyBlueprint(const UBlueprint* Blueprint)
 			else
 			{
 				// Just make sure the nodes are trivial, if they aren't then we're not data only:
-				for (auto Node : UserConstructionScript->Nodes)
+				for (UEdGraphNode* Node : UserConstructionScript->Nodes)
 				{
 					if (!Cast<UK2Node_FunctionEntry>(Node) &&
 						!Cast<UK2Node_CallParentFunction>(Node))
@@ -3162,7 +3150,7 @@ bool FBlueprintEditorUtils::IsDataOnlyBlueprint(const UBlueprint* Blueprint)
 	}
 
 	// EventGraph is empty
-	UEdGraph* EventGraph = (Blueprint->UbergraphPages.Num() == 1) ? Blueprint->UbergraphPages[0] : NULL;
+	UEdGraph* EventGraph = (Blueprint->UbergraphPages.Num() == 1) ? Blueprint->UbergraphPages[0] : nullptr;
 	if( EventGraph && EventGraph->Nodes.Num() > 0 )
 	{
 		for(UEdGraphNode* GraphNode : EventGraph->Nodes)
@@ -3213,7 +3201,7 @@ bool FBlueprintEditorUtils::IsDelegateSignatureGraph(const UEdGraph* Graph)
 	{
 		if(const UBlueprint* Blueprint = FindBlueprintForGraph(Graph))
 		{
-			return (NULL != Blueprint->DelegateSignatureGraphs.FindByKey(Graph));
+			return (nullptr != Blueprint->DelegateSignatureGraphs.FindByKey(Graph));
 		}
 	}
 	return false;
@@ -3277,31 +3265,29 @@ bool FBlueprintEditorUtils::CanClassGenerateEvents(const UClass* InClass)
 UEdGraph* FBlueprintEditorUtils::FindUserConstructionScript(const UBlueprint* Blueprint)
 {
 	const UEdGraphSchema_K2* Schema = GetDefault<UEdGraphSchema_K2>();
-	for( auto GraphIt = Blueprint->FunctionGraphs.CreateConstIterator(); GraphIt; ++GraphIt )
+	for (UEdGraph* CurrentGraph : Blueprint->FunctionGraphs)
 	{
-		UEdGraph* CurrentGraph = *GraphIt;
 		if( CurrentGraph->GetFName() == Schema->FN_UserConstructionScript )
 		{
 			return CurrentGraph;
 		}
 	}
 
-	return NULL;
+	return nullptr;
 }
 
 UEdGraph* FBlueprintEditorUtils::FindEventGraph(const UBlueprint* Blueprint)
 {
 	const UEdGraphSchema_K2* Schema = GetDefault<UEdGraphSchema_K2>();
-	for( auto GraphIt = Blueprint->UbergraphPages.CreateConstIterator(); GraphIt; ++GraphIt )
+	for (UEdGraph* CurrentGraph : Blueprint->UbergraphPages)
 	{
-		UEdGraph* CurrentGraph = *GraphIt;
 		if( CurrentGraph->GetFName() == Schema->GN_EventGraph )
 		{
 			return CurrentGraph;
 		}
 	}
 
-	return NULL;
+	return nullptr;
 }
 
 bool FBlueprintEditorUtils::IsEventGraph(const UEdGraph* InGraph)
@@ -3327,8 +3313,8 @@ bool FBlueprintEditorUtils::IsTunnelInstanceNode(const UEdGraphNode* InGraphNode
 
 bool FBlueprintEditorUtils::DoesBlueprintDeriveFrom(const UBlueprint* Blueprint, UClass* TestClass)
 {
-	check(Blueprint->SkeletonGeneratedClass != NULL);
-	return	TestClass != NULL && 
+	check(Blueprint->SkeletonGeneratedClass != nullptr);
+	return	TestClass != nullptr && 
 		Blueprint->SkeletonGeneratedClass->IsChildOf(TestClass);
 }
 
@@ -3518,11 +3504,11 @@ int32 FBlueprintEditorUtils::FindTimelineIndex(const UBlueprint* Blueprint, cons
 
 void FBlueprintEditorUtils::GetSCSVariableNameList(const UBlueprint* Blueprint, TArray<FName>& VariableNames)
 {
-	if(Blueprint != NULL && Blueprint->SimpleConstructionScript != NULL)
+	if(Blueprint != nullptr && Blueprint->SimpleConstructionScript != nullptr)
 	{
 		for (USCS_Node* SCS_Node : Blueprint->SimpleConstructionScript->GetAllNodes())
 		{
-			if (SCS_Node != NULL)
+			if (SCS_Node != nullptr)
 			{
 				const FName VariableName = SCS_Node->GetVariableName();
 				if(VariableName != NAME_None)
@@ -3536,18 +3522,18 @@ void FBlueprintEditorUtils::GetSCSVariableNameList(const UBlueprint* Blueprint,
 
 void FBlueprintEditorUtils::GetImplementingBlueprintsFunctionNameList(const UBlueprint* Blueprint, TArray<FName>& FunctionNames)
 {
-	if(Blueprint != NULL && FBlueprintEditorUtils::IsInterfaceBlueprint(Blueprint))
+	if(Blueprint != nullptr && FBlueprintEditorUtils::IsInterfaceBlueprint(Blueprint))
 	{
 		for(TObjectIterator<UBlueprint> BlueprintIt; BlueprintIt; ++BlueprintIt)
 		{
 			const UBlueprint* ChildBlueprint = *BlueprintIt;
-			if(ChildBlueprint != NULL)
+			if(ChildBlueprint != nullptr)
 			{
 				for (int32 InterfaceIndex = 0; InterfaceIndex < ChildBlueprint->ImplementedInterfaces.Num(); InterfaceIndex++)
 				{
 					const FBPInterfaceDescription& CurrentInterface = ChildBlueprint->ImplementedInterfaces[InterfaceIndex];
 					const UBlueprint* BlueprintInterfaceClass = UBlueprint::GetBlueprintFromClass(CurrentInterface.Interface);
-					if(BlueprintInterfaceClass != NULL && BlueprintInterfaceClass == Blueprint)
+					if(BlueprintInterfaceClass != nullptr && BlueprintInterfaceClass == Blueprint)
 					{
 						FBlueprintEditorUtils::GetAllGraphNames(ChildBlueprint, FunctionNames);
 					}
@@ -3581,7 +3567,7 @@ void FBlueprintEditorUtils::SetBlueprintOnlyEditableFlag(UBlueprint* Blueprint,
 
 	if (bNewBlueprintOnly)
 	{
-		FBlueprintEditorUtils::RemoveBlueprintVariableMetaData(Blueprint, VarName, NULL, FEdMode::MD_MakeEditWidget);
+		FBlueprintEditorUtils::RemoveBlueprintVariableMetaData(Blueprint, VarName, nullptr, FEdMode::MD_MakeEditWidget);
 	}
 
 	if (VarIndex != INDEX_NONE)
@@ -3691,7 +3677,7 @@ struct FMetaDataDependencyHelper
 					Graph->GetNodesOfClass(LocalSpawnNodes);
 					for (UK2Node_SpawnActorFromClass* Node : LocalSpawnNodes)
 					{
-						UClass* ClassToSpawn = Node ? Node->GetClassToSpawn() : NULL;
+						UClass* ClassToSpawn = Node ? Node->GetClassToSpawn() : nullptr;
 						if (ClassToSpawn && ClassToSpawn->IsChildOf(Blueprint->GeneratedClass))
 						{
 							Schema->ReconstructNode(*Node, true);
@@ -3899,7 +3885,7 @@ void FBlueprintEditorUtils::SetBlueprintVariableCategory(UBlueprint* Blueprint,
 	if (UProperty* TargetProperty = FindField<UProperty>(SkeletonGeneratedClass, VarName))
 	{
 		UClass* OuterClass = CastChecked<UClass>(TargetProperty->GetOuter());
-		const bool bIsNativeVar = (OuterClass->ClassGeneratedBy == NULL);
+		const bool bIsNativeVar = (OuterClass->ClassGeneratedBy == nullptr);
 
 		// If the category does not change, we will not recompile the Blueprint
 		bool bIsCategoryChanged = false;
@@ -3965,7 +3951,7 @@ FText FBlueprintEditorUtils::GetBlueprintVariableCategory(UBlueprint* Blueprint,
 	FText CategoryName;
 	UClass* SkeletonGeneratedClass = Blueprint->SkeletonGeneratedClass;
 	UProperty* TargetProperty = FindField<UProperty>(SkeletonGeneratedClass, VarName);
-	if(TargetProperty != NULL)
+	if(TargetProperty != nullptr)
 	{
 		CategoryName = FObjectEditorUtils::GetCategoryText(TargetProperty);
 	}
@@ -3978,7 +3964,7 @@ FText FBlueprintEditorUtils::GetBlueprintVariableCategory(UBlueprint* Blueprint,
 		}
 	}
 
-	if(CategoryName.IsEmpty() && Blueprint->SimpleConstructionScript != NULL)
+	if(CategoryName.IsEmpty() && Blueprint->SimpleConstructionScript != nullptr)
 	{
 		// Look for the variable in the SCS (in case the Blueprint has not been compiled yet)
 		const int32 SCS_NodeIndex = FBlueprintEditorUtils::FindSCS_Node(Blueprint, VarName);
@@ -3998,7 +3984,7 @@ uint64* FBlueprintEditorUtils::GetBlueprintVariablePropertyFlags(UBlueprint* Blu
 	{
 		return &Blueprint->NewVariables[VarIndex].PropertyFlags;
 	}
-	return NULL;
+	return nullptr;
 }
 
 FName FBlueprintEditorUtils::GetBlueprintVariableRepNotifyFunc(UBlueprint* Blueprint, const FName& VarName)
@@ -4047,7 +4033,7 @@ void FBlueprintEditorUtils::GetDelegateNameList(const UBlueprint* Blueprint, TAr
 
 UEdGraph* FBlueprintEditorUtils::GetDelegateSignatureGraphByName(UBlueprint* Blueprint, FName FunctionName)
 {
-	if ((NULL != Blueprint) && (FunctionName != NAME_None))
+	if ((nullptr != Blueprint) && (FunctionName != NAME_None))
 	{
 		for (int32 It = 0; It < Blueprint->DelegateSignatureGraphs.Num(); It++)
 		{
@@ -4060,15 +4046,15 @@ UEdGraph* FBlueprintEditorUtils::GetDelegateSignatureGraphByName(UBlueprint* Blu
 			}
 		}
 	}
-	return NULL;
+	return nullptr;
 }
 
 // Gets a list of pins that should hidden for a given function
 void FBlueprintEditorUtils::GetHiddenPinsForFunction(UEdGraph const* Graph, UFunction const* Function, TSet<FString>& HiddenPins, TSet<FString>* OutInternalPins)
 {
-	check(Function != NULL);
+	check(Function != nullptr);
 	TMap<FName, FString>* MetaData = UMetaData::GetMapForObject(Function);	
-	if (MetaData != NULL)
+	if (MetaData != nullptr)
 	{
 		for (TMap<FName, FString>::TConstIterator It(*MetaData); It; ++It)
 		{
@@ -4125,7 +4111,7 @@ void FBlueprintEditorUtils::GetHiddenPinsForFunction(UEdGraph const* Graph, UFun
 
 bool FBlueprintEditorUtils::IsPinTypeValid(const FEdGraphPinType& Type)
 {
-	if (const auto UDStruct = Cast<const UUserDefinedStruct>(Type.PinSubCategoryObject.Get()))
+	if (const UUserDefinedStruct* UDStruct = Cast<const UUserDefinedStruct>(Type.PinSubCategoryObject.Get()))
 	{
 		if (EUserDefinedStructureStatus::UDSS_UpToDate != UDStruct->Status)
 		{
@@ -4139,8 +4125,8 @@ bool FBlueprintEditorUtils::IsPinTypeValid(const FEdGraphPinType& Type)
 void FBlueprintEditorUtils::GetClassVariableList(const UBlueprint* Blueprint, TArray<FName>& VisibleVariables, bool bIncludePrivateVars) 
 {
 	// Existing variables in the parent class and above
-	check(!Blueprint->bHasBeenRegenerated || Blueprint->bIsRegeneratingOnLoad || (Blueprint->SkeletonGeneratedClass != NULL));
-	if (Blueprint->SkeletonGeneratedClass != NULL)
+	check(!Blueprint->bHasBeenRegenerated || Blueprint->bIsRegeneratingOnLoad || (Blueprint->SkeletonGeneratedClass != nullptr));
+	if (Blueprint->SkeletonGeneratedClass != nullptr)
 	{
 		for (TFieldIterator<UProperty> PropertyIt(Blueprint->SkeletonGeneratedClass, EFieldIteratorFlags::IncludeSuper); PropertyIt; ++PropertyIt)
 		{
@@ -4160,7 +4146,7 @@ void FBlueprintEditorUtils::GetClassVariableList(const UBlueprint* Blueprint, TA
 			for (int32 StackIndex = ParentBPStack.Num() - 1; StackIndex >= 0; --StackIndex)
 			{
 				UBlueprint* ParentBP = ParentBPStack[StackIndex];
-				check(ParentBP != NULL);
+				check(ParentBP != nullptr);
 
 				GetSCSVariableNameList(ParentBP, VisibleVariables);
 
@@ -4169,7 +4155,7 @@ void FBlueprintEditorUtils::GetClassVariableList(const UBlueprint* Blueprint, TA
 					VisibleVariables.AddUnique(ParentBP->NewVariables[VariableIndex].VarName);
 				}
 
-				for (auto Timeline : ParentBP->Timelines)
+				for (UTimelineTemplate* Timeline : ParentBP->Timelines)
 				{
 					if (Timeline)
 					{
@@ -4393,9 +4379,9 @@ void FBlueprintEditorUtils::RenameComponentMemberVariable(UBlueprint* Blueprint,
 			const FComponentKey Key(Node);
 			TArray<UBlueprint*> Dependents;
 			GetDependentBlueprints(Blueprint, Dependents);
-			for (auto DepBP : Dependents)
+			for (UBlueprint* DepBP : Dependents)
 			{
-				auto InheritableComponentHandler = DepBP ? DepBP->GetInheritableComponentHandler(false) : nullptr;
+				UInheritableComponentHandler* InheritableComponentHandler = DepBP ? DepBP->GetInheritableComponentHandler(false) : nullptr;
 				if (InheritableComponentHandler && InheritableComponentHandler->GetOverridenComponentTemplate(Key))
 				{
 					InheritableComponentHandler->Modify();
@@ -4477,9 +4463,8 @@ void FBlueprintEditorUtils::RenameMemberVariable(UBlueprint* Blueprint, const FN
 					// this code should not be necessary, because the GUID remains valid, but let it be for backward compatibility.
 					TArray<UK2Node_BaseMCDelegate*> NodeUsingDelegate;
 					FBlueprintEditorUtils::GetAllNodesOfClass<UK2Node_BaseMCDelegate>(Blueprint, NodeUsingDelegate);
-					for (auto FuncIt = NodeUsingDelegate.CreateIterator(); FuncIt; ++FuncIt)
+					for (UK2Node_BaseMCDelegate* FunctionNode : NodeUsingDelegate)
 					{
-						UK2Node_BaseMCDelegate* FunctionNode = *FuncIt;
 						if (FunctionNode->DelegateReference.IsSelfContext() && (FunctionNode->DelegateReference.GetMemberName() == OldName))
 						{
 							FunctionNode->Modify();
@@ -4556,7 +4541,7 @@ void FBlueprintEditorUtils::GetLoadedChildBlueprints(UBlueprint* InBlueprint, TA
 	for(TObjectIterator<UBlueprint> BlueprintIt; BlueprintIt; ++BlueprintIt)
 	{
 		UBlueprint* ChildBP = *BlueprintIt;
-		if(ChildBP != NULL && ChildBP->ParentClass != NULL)
+		if(ChildBP != nullptr && ChildBP->ParentClass != nullptr)
 		{
 			TArray<UBlueprint*> ParentBPArray;
 			// Get the parent hierarchy
@@ -4565,10 +4550,10 @@ void FBlueprintEditorUtils::GetLoadedChildBlueprints(UBlueprint* InBlueprint, TA
 			// Also get any BP interfaces we use
 			TArray<UClass*> ImplementedInterfaces;
 			FindImplementedInterfaces(ChildBP, true, ImplementedInterfaces);
-			for(auto InterfaceIt(ImplementedInterfaces.CreateConstIterator()); InterfaceIt; InterfaceIt++)
+			for (UClass* ImplementedInterface : ImplementedInterfaces)
 			{
-				UBlueprint* BlueprintInterfaceClass = UBlueprint::GetBlueprintFromClass(*InterfaceIt);
-				if(BlueprintInterfaceClass != NULL)
+				UBlueprint* BlueprintInterfaceClass = UBlueprint::GetBlueprintFromClass(ImplementedInterface);
+				if(BlueprintInterfaceClass != nullptr)
 				{
 					ParentBPArray.Add(BlueprintInterfaceClass);
 				}
@@ -4631,7 +4616,7 @@ void FBlueprintEditorUtils::ChangeMemberVariableType(UBlueprint* Blueprint, cons
 					if(NewPinType.PinSubCategoryObject.IsValid())
 					{
 						const UClass* ClassObject = Cast<UClass>(NewPinType.PinSubCategoryObject.Get());
-						check(ClassObject != NULL);
+						check(ClassObject != nullptr);
 
 						if (ClassObject->IsChildOf(AActor::StaticClass()))
 						{
@@ -4722,7 +4707,7 @@ FName FBlueprintEditorUtils::DuplicateVariable(UBlueprint* InBlueprint, const US
 			NewVar = DuplicateVariableDescription(InBlueprint, Variable);
 
 			// We need to manually pull the DefaultValue from the UProperty to set it
-			void* OldPropertyAddr = NULL;
+			void* OldPropertyAddr = nullptr;
 
 			//Grab property of blueprint's current CDO
 			UClass* GeneratedClass = InBlueprint->GeneratedClass;
@@ -4737,7 +4722,7 @@ FName FBlueprintEditorUtils::DuplicateVariable(UBlueprint* InBlueprint, const US
 				{
 					// if there is a property for variable, it means the original default value was already copied, so it can be safely overridden
 					Variable.DefaultValue.Empty();
-					TargetProperty->ExportTextItem(NewVar.DefaultValue, OldPropertyAddr, OldPropertyAddr, NULL, PPF_None);
+					TargetProperty->ExportTextItem(NewVar.DefaultValue, OldPropertyAddr, OldPropertyAddr, nullptr, PPF_None);
 				}
 			}
 
@@ -4748,7 +4733,7 @@ FName FBlueprintEditorUtils::DuplicateVariable(UBlueprint* InBlueprint, const US
 		{
 			// It's probably a local variable
 
-			UK2Node_FunctionEntry* FunctionEntry = NULL;
+			UK2Node_FunctionEntry* FunctionEntry = nullptr;
 			FBPVariableDescription* LocalVariable = FBlueprintEditorUtils::FindLocalVariable(InBlueprint, InScope, InVariableToDuplicate, &FunctionEntry);
 
 			if (LocalVariable)
@@ -4793,7 +4778,7 @@ FBPVariableDescription FBlueprintEditorUtils::DuplicateVariableDescription(UBlue
 
 bool FBlueprintEditorUtils::AddLocalVariable(UBlueprint* Blueprint, UEdGraph* InTargetGraph, const FName InNewVarName, const FEdGraphPinType& InNewVarType, const FString& DefaultValue/*= FString()*/)
 {
-	if(InTargetGraph != NULL && InTargetGraph->GetSchema()->GetGraphType(InTargetGraph) == GT_Function)
+	if(InTargetGraph != nullptr && InTargetGraph->GetSchema()->GetGraphType(InTargetGraph) == GT_Function)
 	{
 		// Ensure we have the top level graph for the function, in-case we are in a child graph
 		UEdGraph* TargetGraph = FBlueprintEditorUtils::GetTopLevelGraph(InTargetGraph);
@@ -4872,9 +4857,9 @@ FFunctionFromNodeHelper::FFunctionFromNodeHelper(const UObject* Obj) : Function(
 
 UFunction* FFunctionFromNodeHelper::FunctionFromNode(const UK2Node* Node)
 {
-	UFunction* Function = NULL;
-	UBlueprint* Blueprint = Node ? Node->GetBlueprint() : NULL;
-	const UClass* SearchScope = Blueprint ? Blueprint->SkeletonGeneratedClass : NULL;
+	UFunction* Function = nullptr;
+	UBlueprint* Blueprint = Node ? Node->GetBlueprint() : nullptr;
+	const UClass* SearchScope = Blueprint ? Blueprint->SkeletonGeneratedClass : nullptr;
 	if (SearchScope)
 	{
 		if (const UK2Node_FunctionResult* ResultNode = Cast<UK2Node_FunctionResult>(Node))
@@ -4904,15 +4889,15 @@ UFunction* FFunctionFromNodeHelper::FunctionFromNode(const UK2Node* Node)
 
 UEdGraph* FBlueprintEditorUtils::FindScopeGraph(const UBlueprint* InBlueprint, const UStruct* InScope)
 {
-	UEdGraph* ScopeGraph = NULL;
+	UEdGraph* ScopeGraph = nullptr;
 
 	TArray<UEdGraph*> AllGraphs;
 	InBlueprint->GetAllGraphs(AllGraphs);
 
-	for(const auto Graph : AllGraphs)
+	for (UEdGraph* Graph : AllGraphs)
 	{
-		check(Graph != NULL);
-		if(Graph->GetName() == InScope->GetName())
+		check(Graph != nullptr);
+		if(Graph->GetFName() == InScope->GetFName())
 		{
 			// This graph should always be a function graph
 			check(Graph->GetSchema()->GetGraphType(Graph) == GT_Function);
@@ -4930,8 +4915,8 @@ void FBlueprintEditorUtils::RenameLocalVariable(UBlueprint* InBlueprint, const U
 	{
 		UK2Node_FunctionEntry* FunctionEntry = nullptr;
 		FBPVariableDescription* LocalVariable = FindLocalVariable(InBlueprint, InScope, InOldName, &FunctionEntry);
-		const auto OldProperty = FindField<const UProperty>(InScope, InOldName);
-		const auto ExistingProperty = FindField<const UProperty>(InScope, InNewName);
+		const UProperty* OldProperty = FindField<const UProperty>(InScope, InOldName);
+		const UProperty* ExistingProperty = FindField<const UProperty>(InScope, InNewName);
 		const bool bHasExistingProperty = ExistingProperty && ExistingProperty != OldProperty;
 		if (bHasExistingProperty)
 		{
@@ -4969,7 +4954,7 @@ FBPVariableDescription* FBlueprintEditorUtils::FindLocalVariable(UBlueprint* InB
 
 FBPVariableDescription* FBlueprintEditorUtils::FindLocalVariable(const UBlueprint* InBlueprint, const UEdGraph* InScopeGraph, const FName InVariableName, class UK2Node_FunctionEntry** OutFunctionEntry)
 {
-	FBPVariableDescription* ReturnVariable = NULL;
+	FBPVariableDescription* ReturnVariable = nullptr;
 	if(DoesSupportLocalVariables(InScopeGraph))
 	{
 		UEdGraph* FunctionGraph = GetTopLevelGraph(InScopeGraph);
@@ -5060,7 +5045,7 @@ void FBlueprintEditorUtils::ChangeLocalVariableType(UBlueprint* InBlueprint, con
 		FString ActionCategory;
 		const UEdGraphSchema_K2* K2Schema = GetDefault<UEdGraphSchema_K2>();
 
-		UK2Node_FunctionEntry* FunctionEntry = NULL;
+		UK2Node_FunctionEntry* FunctionEntry = nullptr;
 		FBPVariableDescription* VariablePtr = FindLocalVariable(InBlueprint, InScope, InVariableName, &FunctionEntry);
 
 		if(VariablePtr)
@@ -5097,7 +5082,7 @@ void FBlueprintEditorUtils::ChangeLocalVariableType(UBlueprint* InBlueprint, con
 					if(NewPinType.PinSubCategoryObject.IsValid())
 					{
 						const UClass* ClassObject = Cast<UClass>(NewPinType.PinSubCategoryObject.Get());
-						check(ClassObject != NULL);
+						check(ClassObject != nullptr);
 
 						if (ClassObject->IsChildOf(AActor::StaticClass()))
 						{
@@ -5144,16 +5129,15 @@ void FBlueprintEditorUtils::ReplaceVariableReferences(UBlueprint* Blueprint, con
 	TArray<UBlueprint*> Dependents;
 	GetDependentBlueprints(Blueprint, Dependents);
 
-	for (auto DependentIt = Dependents.CreateIterator(); DependentIt; ++DependentIt)
+	for (UBlueprint* DependentBp : Dependents)
 	{
-		UBlueprint* DependentBp = *DependentIt;
 		::RenameVariableReferences(DependentBp, Blueprint->GeneratedClass, OldName, NewName);
 	}
 }
 
 void FBlueprintEditorUtils::ReplaceVariableReferences(UBlueprint* Blueprint, const UProperty* OldVariable, const UProperty* NewVariable)
 {
-	check((OldVariable != NULL) && (NewVariable != NULL));
+	check((OldVariable != nullptr) && (NewVariable != nullptr));
 	ReplaceVariableReferences(Blueprint, OldVariable->GetFName(), NewVariable->GetFName());
 }
 
@@ -5229,13 +5213,13 @@ bool FBlueprintEditorUtils::ValidateAllMemberVariables(UBlueprint* InBlueprint,
 
 bool FBlueprintEditorUtils::ValidateAllComponentMemberVariables(UBlueprint* InBlueprint, UBlueprint* InParentBlueprint, const FName& InVariableName)
 {
-	if(InBlueprint->SimpleConstructionScript != NULL)
+	if(InBlueprint->SimpleConstructionScript != nullptr)
 	{
 		TArray<USCS_Node*> ChildSCSNodes = InBlueprint->SimpleConstructionScript->GetAllNodes();
 		for(int32 NodeIndex = 0; NodeIndex < ChildSCSNodes.Num(); ++NodeIndex)
 		{
 			USCS_Node* SCS_Node = ChildSCSNodes[NodeIndex];
-			if(SCS_Node != NULL && SCS_Node->GetVariableName() == InVariableName)
+			if(SCS_Node != nullptr && SCS_Node->GetVariableName() == InVariableName)
 			{
 				FName NewChildName = FBlueprintEditorUtils::FindUniqueKismetName(InBlueprint, InVariableName.ToString());
 
@@ -5314,7 +5298,7 @@ void FBlueprintEditorUtils::ValidateBlueprintChildVariables(UBlueprint* InBluepr
 	for(TObjectIterator<UBlueprint> BlueprintIt; BlueprintIt; ++BlueprintIt)
 	{
 		UBlueprint* ChildBP = *BlueprintIt;
-		if(ChildBP != NULL && ChildBP->ParentClass != NULL)
+		if(ChildBP != nullptr && ChildBP->ParentClass != nullptr)
 		{
 			TArray<UBlueprint*> ParentBPArray;
 			// Get the parent hierarchy
@@ -5323,10 +5307,10 @@ void FBlueprintEditorUtils::ValidateBlueprintChildVariables(UBlueprint* InBluepr
 			// Also get any BP interfaces we use
 			TArray<UClass*> ImplementedInterfaces;
 			FindImplementedInterfaces(ChildBP, true, ImplementedInterfaces);
-			for(auto InterfaceIt(ImplementedInterfaces.CreateConstIterator()); InterfaceIt; InterfaceIt++)
+			for (UClass* ImplementedInterface : ImplementedInterfaces)
 			{
-				UBlueprint* BlueprintInterfaceClass = UBlueprint::GetBlueprintFromClass(*InterfaceIt);
-				if(BlueprintInterfaceClass != NULL)
+				UBlueprint* BlueprintInterfaceClass = UBlueprint::GetBlueprintFromClass(ImplementedInterface);
+				if(BlueprintInterfaceClass != nullptr)
 				{
 					ParentBPArray.Add(BlueprintInterfaceClass);
 				}
@@ -5436,7 +5420,7 @@ void FBlueprintEditorUtils::ExportPropertyToKismetDefaultValue(UEdGraphPin* Targ
 	const int32 PortFlags = 0;
 
 	TargetPin->DefaultValue.Empty();
-	SourceProperty->ExportTextItem(TargetPin->DefaultValue, SourceAddress, NULL, OwnerObject, PortFlags);
+	SourceProperty->ExportTextItem(TargetPin->DefaultValue, SourceAddress, nullptr, OwnerObject, PortFlags);
 }
 
 //////////////////////////////////////////////////////////////////////////
@@ -5448,16 +5432,15 @@ FGuid FBlueprintEditorUtils::FindInterfaceFunctionGuid(const UFunction* Function
 	check(Function);
 	check(InterfaceClass);
 	const UBlueprint* InterfaceBlueprint = Cast<UBlueprint>(InterfaceClass->ClassGeneratedBy);
-	if(InterfaceBlueprint != NULL)
+	if(InterfaceBlueprint != nullptr)
 	{
 		// find the graph for this function
 		TArray<UEdGraph*> InterfaceGraphs;
 		InterfaceBlueprint->GetAllGraphs(InterfaceGraphs);
 
-		for (auto InterfaceGraphIt(InterfaceGraphs.CreateConstIterator()); InterfaceGraphIt; InterfaceGraphIt++)
+		for (const UEdGraph* InterfaceGraph : InterfaceGraphs)
 		{
-			const UEdGraph* InterfaceGraph = *InterfaceGraphIt;
-			if(InterfaceGraph != NULL && InterfaceGraph->GetFName() == Function->GetFName())
+			if(InterfaceGraph != nullptr && InterfaceGraph->GetFName() == Function->GetFName())
 			{
 				return InterfaceGraph->GraphGuid;
 			}
@@ -5505,7 +5488,7 @@ bool FBlueprintEditorUtils::ImplementNewInterface(UBlueprint* Blueprint, const F
 		{
 			FName FunctionName = Function->GetFName();
 			UEdGraph* FuncGraph = FindObject<UEdGraph>( Blueprint, *(FunctionName.ToString()) );
-			if (FuncGraph != NULL)
+			if (FuncGraph != nullptr)
 			{
 				bAllFunctionsAdded = false;
 
@@ -5614,9 +5597,8 @@ void FBlueprintEditorUtils::RemoveInterface(UBlueprint* Blueprint, const FName&
 					UK2Node_CustomEvent* NewEvent = UK2Node_CustomEvent::CreateFromFunction(PreviousNodePos, EventNode->GetGraph(), PreviousNodeName, PreviousSignatureFunction, false);
 
 					// Move the pin links from the old pin to the new pin to preserve connections
-					for(auto PinIt = EventNode->Pins.CreateIterator(); PinIt; ++PinIt)
+					for (UEdGraphPin* CurrentPin : EventNode->Pins)
 					{
-						UEdGraphPin* CurrentPin = *PinIt;
 						UEdGraphPin* TargetPin = NewEvent->FindPinChecked(CurrentPin->PinName);
 						const UEdGraphSchema_K2* Schema = GetDefault<UEdGraphSchema_K2>();
 						Schema->MovePinLinks(*CurrentPin, *TargetPin);
@@ -5720,14 +5702,14 @@ void FBlueprintEditorUtils::PurgeNullGraphs(UBlueprint* Blueprint)
 // Makes sure that calls to parent functions are valid, and removes them if not
 void FBlueprintEditorUtils::ConformCallsToParentFunctions(UBlueprint* Blueprint)
 {
-	check(NULL != Blueprint);
+	check(nullptr != Blueprint);
 
 	TArray<UEdGraph*> AllGraphs;
 	Blueprint->GetAllGraphs(AllGraphs);
 	for(int GraphIndex = 0; GraphIndex < AllGraphs.Num(); ++GraphIndex)
 	{
 		UEdGraph* CurrentGraph = AllGraphs[GraphIndex];
-		check(NULL != CurrentGraph);
+		check(nullptr != CurrentGraph);
 
 		// Make sure the graph is loaded
 		if(!CurrentGraph->HasAnyFlags(RF_NeedLoad|RF_NeedPostLoad))
@@ -5739,14 +5721,14 @@ void FBlueprintEditorUtils::ConformCallsToParentFunctions(UBlueprint* Blueprint)
 			for(int CallFunctionNodeIndex = 0; CallFunctionNodeIndex < CallFunctionNodes.Num(); ++CallFunctionNodeIndex)
 			{
 				UK2Node_CallParentFunction* CallFunctionNode = CallFunctionNodes[CallFunctionNodeIndex];
-				check(NULL != CallFunctionNode);
+				check(nullptr != CallFunctionNode);
 
 				// Make sure the node has already been loaded
 				if(!CallFunctionNode->HasAnyFlags(RF_NeedLoad|RF_NeedPostLoad))
 				{
 					// Attempt to locate the function within the parent class
 					const UFunction* TargetFunction = CallFunctionNode->GetTargetFunction();
-					TargetFunction = TargetFunction && Blueprint->ParentClass ? Blueprint->ParentClass->FindFunctionByName(TargetFunction->GetFName()) : NULL;
+					TargetFunction = TargetFunction && Blueprint->ParentClass ? Blueprint->ParentClass->FindFunctionByName(TargetFunction->GetFName()) : nullptr;
 					if(TargetFunction)
 					{
 						// If the function signature does not match the parent class
@@ -5771,7 +5753,7 @@ void FBlueprintEditorUtils::ConformCallsToParentFunctions(UBlueprint* Blueprint)
 						for(int PinIndex = 0; PinIndex < CallFunctionNode->Pins.Num(); ++PinIndex)
 						{
 							UEdGraphPin* InputPin = CallFunctionNode->Pins[PinIndex];
-							check(NULL != InputPin);
+							check(nullptr != InputPin);
 
 							// If this is an input exec pin
 							const UEdGraphSchema_K2* K2Schema = GetDefault<UEdGraphSchema_K2>();
@@ -5781,7 +5763,7 @@ void FBlueprintEditorUtils::ConformCallsToParentFunctions(UBlueprint* Blueprint)
 								for(int InputLinkedToPinIndex = 0; InputLinkedToPinIndex < InputPin->LinkedTo.Num(); ++InputLinkedToPinIndex)
 								{
 									UEdGraphPin* InputLinkedToPin = InputPin->LinkedTo[InputLinkedToPinIndex];
-									check(NULL != InputLinkedToPin);
+									check(nullptr != InputLinkedToPin);
 
 									// Break the existing link to the node we're about to remove
 									InputLinkedToPin->BreakLinkTo(InputPin);
@@ -5790,7 +5772,7 @@ void FBlueprintEditorUtils::ConformCallsToParentFunctions(UBlueprint* Blueprint)
 									for(int OutputLinkedToPinIndex = 0; OutputLinkedToPinIndex < OutputPin->LinkedTo.Num(); ++OutputLinkedToPinIndex)
 									{
 										UEdGraphPin* OutputLinkedToPin = OutputPin->LinkedTo[OutputLinkedToPinIndex];
-										check(NULL != OutputLinkedToPin);
+										check(nullptr != OutputLinkedToPin);
 
 										// Make sure the output connection isn't linked to the node we're about to remove
 										if(OutputLinkedToPin->LinkedTo.Contains(OutputPin))
@@ -5894,7 +5876,7 @@ namespace
 			else
 			{
 				TArray<FName> DummyExtraNameList;
-				UEdGraphNode* CustomEventNode = CurrentGraph->GetSchema()->CreateSubstituteNode(EventNode, CurrentGraph, NULL, DummyExtraNameList);
+				UEdGraphNode* CustomEventNode = CurrentGraph->GetSchema()->CreateSubstituteNode(EventNode, CurrentGraph, nullptr, DummyExtraNameList);
 				if (ensure(CustomEventNode))
 				{
 					// Destroy the old event node (this will also break all pin links and remove it from the graph)
@@ -5915,7 +5897,7 @@ namespace
 // Makes sure that all events we handle exist, and replace with custom events if not
 void FBlueprintEditorUtils::ConformImplementedEvents(UBlueprint* Blueprint)
 {
-	check(NULL != Blueprint);
+	check(nullptr != Blueprint);
 
 	// Collect all implemented interface classes
 	TArray<UClass*> ImplementedInterfaceClasses;
@@ -5926,7 +5908,7 @@ void FBlueprintEditorUtils::ConformImplementedEvents(UBlueprint* Blueprint)
 	for(int GraphIndex = 0; GraphIndex < AllGraphs.Num(); ++GraphIndex)
 	{
 		UEdGraph* CurrentGraph = AllGraphs[GraphIndex];
-		check(NULL != CurrentGraph);
+		check(nullptr != CurrentGraph);
 
 		// Make sure the graph is loaded
 		if(!CurrentGraph->HasAnyFlags(RF_NeedLoad|RF_NeedPostLoad))
@@ -5938,7 +5920,7 @@ void FBlueprintEditorUtils::ConformImplementedEvents(UBlueprint* Blueprint)
 			for(int EventNodeIndex = 0; EventNodeIndex < EventNodes.Num(); ++EventNodeIndex)
 			{
 				UK2Node_Event* EventNode = EventNodes[EventNodeIndex];
-				check(NULL != EventNode);
+				check(nullptr != EventNode);
 
 				// If the event is loaded and is not a custom event
 				if(!EventNode->HasAnyFlags(RF_NeedLoad|RF_NeedPostLoad) && EventNode->bOverrideFunction)
@@ -5983,16 +5965,14 @@ static void ConformInterfaceByGUID(const UBlueprint* Blueprint, const FBPInterfa
 	TArray<UEdGraph*> BlueprintGraphs;
 	Blueprint->GetAllGraphs(BlueprintGraphs);
 		
-	for (auto BlueprintGraphIt(BlueprintGraphs.CreateConstIterator()); BlueprintGraphIt; BlueprintGraphIt++)
+	for (UEdGraph* BlueprintGraph : BlueprintGraphs)
 	{
-		UEdGraph* BlueprintGraph = *BlueprintGraphIt;
-		if(BlueprintGraph != NULL && BlueprintGraph->InterfaceGuid.IsValid())
+		if(BlueprintGraph != nullptr && BlueprintGraph->InterfaceGuid.IsValid())
 		{
 			// valid interface Guid found, so fixup name if it is different
-			for (auto InterfaceGraphIt(InterfaceGraphs.CreateConstIterator()); InterfaceGraphIt; InterfaceGraphIt++)
+			for (const UEdGraph* InterfaceGraph : InterfaceGraphs)
 			{
-				const UEdGraph* InterfaceGraph = *InterfaceGraphIt;
-				if(InterfaceGraph != NULL && InterfaceGraph->GraphGuid == BlueprintGraph->InterfaceGuid && InterfaceGraph->GetFName() != BlueprintGraph->GetFName())
+				if(InterfaceGraph != nullptr && InterfaceGraph->GraphGuid == BlueprintGraph->InterfaceGuid && InterfaceGraph->GetFName() != BlueprintGraph->GetFName())
 				{
 					FBlueprintEditorUtils::RenameGraph(BlueprintGraph, InterfaceGraph->GetFName().ToString());
 					FBlueprintEditorUtils::RefreshGraphNodes(BlueprintGraph);
@@ -6043,8 +6023,8 @@ static void ConformInterfaceByName(UBlueprint* Blueprint, FBPInterfaceDescriptio
 
 			UEdGraph* EventGraph = EventNode->GetGraph();
 			// we've already implemented this interface function as an event (which we need to replace)
-			UK2Node_CustomEvent* CustomEventNode = Cast<UK2Node_CustomEvent>(EventGraph->GetSchema()->CreateSubstituteNode(EventNode, EventGraph, NULL, ExtraNameList));
-			if (CustomEventNode == NULL)
+			UK2Node_CustomEvent* CustomEventNode = Cast<UK2Node_CustomEvent>(EventGraph->GetSchema()->CreateSubstituteNode(EventNode, EventGraph, nullptr, ExtraNameList));
+			if (CustomEventNode == nullptr)
 			{
 				continue;
 			}
@@ -6136,7 +6116,7 @@ static void ConformInterfaceByName(UBlueprint* Blueprint, FBPInterfaceDescriptio
 // Makes sure that all graphs for all interfaces we implement exist, and add if not
 void FBlueprintEditorUtils::ConformImplementedInterfaces(UBlueprint* Blueprint)
 {
-	check(NULL != Blueprint);
+	check(nullptr != Blueprint);
 	FString ErrorStr;
 
 	// Collect all variables names in current blueprint 
@@ -6145,9 +6125,9 @@ void FBlueprintEditorUtils::ConformImplementedInterfaces(UBlueprint* Blueprint)
 	{
 		VariableNamesUsedInBlueprint.Add(VariablesIter->GetFName());
 	}
-	for (auto VariablesIter = Blueprint->NewVariables.CreateConstIterator(); VariablesIter; ++VariablesIter)
+	for (const FBPVariableDescription& NewVariable : Blueprint->NewVariables)
 	{
-		VariableNamesUsedInBlueprint.AddUnique(VariablesIter->VarName);
+		VariableNamesUsedInBlueprint.AddUnique(NewVariable.VarName);
 	}
 
 	TArray<UEdGraph*> AllGraphs;
@@ -6216,9 +6196,8 @@ void FBlueprintEditorUtils::UpdateOutOfDateAnimBlueprints(UBlueprint* InBlueprin
 		// Ensure all transition graphs have the correct schema
 		TArray<UAnimStateTransitionNode*> TransitionNodes;
 		GetAllNodesOfClass<UAnimStateTransitionNode>(AnimBlueprint, /*out*/ TransitionNodes);
-		for (auto NodeIt = TransitionNodes.CreateConstIterator(); NodeIt; ++NodeIt)
+		for (UAnimStateTransitionNode* Node : TransitionNodes)
 		{
-			UAnimStateTransitionNode* Node = *NodeIt;
 			UEdGraph* TestGraph = Node->BoundGraph;
 			if (TestGraph->Schema == UAnimationGraphSchema::StaticClass())
 			{
@@ -6227,10 +6206,10 @@ void FBlueprintEditorUtils::UpdateOutOfDateAnimBlueprints(UBlueprint* InBlueprin
 		}
 
 		// Handle a reparented anim blueprint that either needs or no longer needs an anim graph
-		if (UAnimBlueprint::FindRootAnimBlueprint(AnimBlueprint) == NULL)
+		if (UAnimBlueprint::FindRootAnimBlueprint(AnimBlueprint) == nullptr)
 		{
 			// Add an anim graph if not present
-			if (FindObject<UEdGraph>(AnimBlueprint, *(K2Schema->GN_AnimGraph.ToString())) == NULL)
+			if (FindObject<UEdGraph>(AnimBlueprint, *(K2Schema->GN_AnimGraph.ToString())) == nullptr)
 			{
 				UEdGraph* NewGraph = FBlueprintEditorUtils::CreateNewGraph(AnimBlueprint, K2Schema->GN_AnimGraph, UAnimationGraph::StaticClass(), UAnimationGraphSchema::StaticClass());
 				FBlueprintEditorUtils::AddDomainSpecificGraph(AnimBlueprint, NewGraph);
@@ -6244,7 +6223,7 @@ void FBlueprintEditorUtils::UpdateOutOfDateAnimBlueprints(UBlueprint* InBlueprin
 			for (int32 i = 0; i < AnimBlueprint->FunctionGraphs.Num(); ++i)
 			{
 				UEdGraph* FuncGraph = AnimBlueprint->FunctionGraphs[i];
-				if ((FuncGraph != NULL) && (FuncGraph->GetFName() == K2Schema->GN_AnimGraph))
+				if ((FuncGraph != nullptr) && (FuncGraph->GetFName() == K2Schema->GN_AnimGraph))
 				{
 					UE_LOG(LogBlueprint, Log, TEXT("!!! Removing AnimGraph from %s, because it has a parent anim blueprint that defines the AnimGraph"), *AnimBlueprint->GetPathName());
 					AnimBlueprint->FunctionGraphs.RemoveAt(i);
@@ -6257,27 +6236,25 @@ void FBlueprintEditorUtils::UpdateOutOfDateAnimBlueprints(UBlueprint* InBlueprin
 
 void FBlueprintEditorUtils::UpdateOutOfDateCompositeNodes(UBlueprint* Blueprint)
 {
-	for( auto It = Blueprint->UbergraphPages.CreateIterator(); It; ++It )
+	for (UEdGraph* UbergraphPage : Blueprint->UbergraphPages)
 	{
-		UpdateOutOfDateCompositeWithOuter(Blueprint, *It);
+		UpdateOutOfDateCompositeWithOuter(Blueprint, UbergraphPage);
 	}
-	for( auto It = Blueprint->FunctionGraphs.CreateIterator(); It; ++It )
+	for (UEdGraph* FunctionGraph : Blueprint->FunctionGraphs)
 	{
-		UpdateOutOfDateCompositeWithOuter(Blueprint, *It);
+		UpdateOutOfDateCompositeWithOuter(Blueprint, FunctionGraph);
 	}
 }
 
 void FBlueprintEditorUtils::UpdateOutOfDateCompositeWithOuter(UBlueprint* Blueprint, UEdGraph* OuterGraph )
 {
-	check(OuterGraph != NULL);
+	check(OuterGraph != nullptr);
 	check(FindBlueprintForGraphChecked(OuterGraph) == Blueprint);
 
-	for (auto It = OuterGraph->Nodes.CreateIterator();It;++It)
+	for (UEdGraphNode* Node : OuterGraph->Nodes)
 	{
-		UEdGraphNode* Node = *It;
-		
 		//Is this node of a type that has a BoundGraph to update
-		UEdGraph* BoundGraph = NULL;
+		UEdGraph* BoundGraph = nullptr;
 		if (UK2Node_Composite* Composite = Cast<UK2Node_Composite>(Node))
 		{
 			BoundGraph = Composite->BoundGraph;
@@ -6309,9 +6286,9 @@ void FBlueprintEditorUtils::UpdateOutOfDateCompositeWithOuter(UBlueprint* Bluepr
 		}
 	}
 
-	for (auto It = OuterGraph->SubGraphs.CreateIterator();It;++It)
+	for (UEdGraph* SubGraph : OuterGraph->SubGraphs)
 	{
-		UpdateOutOfDateCompositeWithOuter(Blueprint,*It);
+		UpdateOutOfDateCompositeWithOuter(Blueprint, SubGraph);
 	}
 }
 
@@ -6323,9 +6300,8 @@ void FBlueprintEditorUtils::UpdateComponentTemplates(UBlueprint* Blueprint)
 	TArray<UK2Node_AddComponent*> AllComponents;
 	FBlueprintEditorUtils::GetAllNodesOfClass(Blueprint, AllComponents);
 
-	for( auto CompIt = AllComponents.CreateIterator(); CompIt; ++CompIt )
+	for (UK2Node_AddComponent* ComponentNode : AllComponents)
 	{
-		UK2Node_AddComponent* ComponentNode = (*CompIt);
 		check(ComponentNode);
 
 		UActorComponent* ActorComp = ComponentNode->GetTemplateFromNode();
@@ -6366,12 +6342,12 @@ void FBlueprintEditorUtils::UpdateRootComponentReference(UBlueprint* Blueprint)
 			// If both CDOs are valid, check for a valid scene root component
 			USceneComponent* ParentSceneRootComponent = ParentActorCDO->GetRootComponent();
 			USceneComponent* BlueprintSceneRootComponent = BlueprintActorCDO->GetRootComponent();
-			if((ParentSceneRootComponent == NULL && BlueprintSceneRootComponent != NULL)
-				|| (ParentSceneRootComponent != NULL && BlueprintSceneRootComponent == NULL)
-				|| (ParentSceneRootComponent != NULL && BlueprintSceneRootComponent != NULL && ParentSceneRootComponent->GetFName() != BlueprintSceneRootComponent->GetFName()))
+			if((ParentSceneRootComponent == nullptr && BlueprintSceneRootComponent != nullptr)
+				|| (ParentSceneRootComponent != nullptr && BlueprintSceneRootComponent == nullptr)
+				|| (ParentSceneRootComponent != nullptr && BlueprintSceneRootComponent != nullptr && ParentSceneRootComponent->GetFName() != BlueprintSceneRootComponent->GetFName()))
 			{
 				// If the parent CDO has a valid scene root component
-				if(ParentSceneRootComponent != NULL)
+				if(ParentSceneRootComponent != nullptr)
 				{
 					// Search for a scene component with the same name in the Blueprint CDO's Components list
 					TInlineComponentArray<USceneComponent*> SceneComponents;
@@ -6387,10 +6363,10 @@ void FBlueprintEditorUtils::UpdateRootComponentReference(UBlueprint* Blueprint)
 						}
 					}
 				}
-				else if(BlueprintSceneRootComponent != NULL)
+				else if(BlueprintSceneRootComponent != nullptr)
 				{
 					// The parent CDO does not have a valid scene root, so NULL out the Blueprint CDO reference to match
-					BlueprintActorCDO->SetRootComponent(NULL);
+					BlueprintActorCDO->SetRootComponent(nullptr);
 				}
 			}
 		}
@@ -6499,9 +6475,8 @@ void FBlueprintEditorUtils::UpdateTransactionalFlags(UBlueprint* Blueprint)
 	TArray<UK2Node*> AllNodes;
 	FBlueprintEditorUtils::GetAllNodesOfClass(Blueprint, AllNodes);
 
-	for( auto NodeIt = AllNodes.CreateIterator(); NodeIt; ++NodeIt )
+	for (UK2Node* K2Node : AllNodes)
 	{
-		UK2Node* K2Node = (*NodeIt);
 		check(K2Node);
 
 		if (!K2Node->HasAnyFlags(RF_Transactional))
@@ -6517,10 +6492,10 @@ void FBlueprintEditorUtils::UpdateStalePinWatches( UBlueprint* Blueprint )
 	TSet<UEdGraphPin*> AllPins;
 
 	// Find all unique pins being watched
-	for (auto PinIt = Blueprint->WatchedPins.CreateIterator(); PinIt; ++PinIt)
+	for (const FEdGraphPinReference& PinRef : Blueprint->WatchedPins)
 	{
-		UEdGraphPin* Pin (PinIt->Get());
-		if (Pin == NULL)
+		UEdGraphPin* Pin = PinRef.Get();
+		if (Pin == nullptr)
 		{
 			continue;
 		}
@@ -6528,7 +6503,7 @@ void FBlueprintEditorUtils::UpdateStalePinWatches( UBlueprint* Blueprint )
 		UEdGraphNode* OwningNode = Pin->GetOwningNode();
 		// during node reconstruction, dead pins get moved to the transient 
 		// package (so just in case this blueprint got saved with dead pin watches)
-		if (OwningNode == NULL)
+		if (OwningNode == nullptr)
 		{
 			continue;
 		}
@@ -6545,16 +6520,16 @@ void FBlueprintEditorUtils::UpdateStalePinWatches( UBlueprint* Blueprint )
 	if (Blueprint->WatchedPins.Num() != AllPins.Num())
 	{
 		Blueprint->WatchedPins.Empty();
-		for (auto PinIt = AllPins.CreateIterator(); PinIt; ++PinIt)
+		for (UEdGraphPin* Pin : AllPins)
 		{
-			Blueprint->WatchedPins.Add(*PinIt);
+			Blueprint->WatchedPins.Add(Pin);
 		}
 
 		Blueprint->Status = BS_Dirty;
 	}
 }
 
-FName FBlueprintEditorUtils::FindUniqueKismetName(const UBlueprint* InBlueprint, const FString& InBaseName, UStruct* InScope/* = NULL*/)
+FName FBlueprintEditorUtils::FindUniqueKismetName(const UBlueprint* InBlueprint, const FString& InBaseName, UStruct* InScope/* = nullptr*/)
 {
 	int32 Count = 0;
 	FString KismetName;
@@ -6596,20 +6571,20 @@ UTimelineTemplate* FBlueprintEditorUtils::AddNewTimeline(UBlueprint* Blueprint,
 	// Early out if we don't support timelines in this class
 	if( !FBlueprintEditorUtils::DoesSupportTimelines(Blueprint) )
 	{
-		return NULL;
+		return nullptr;
 	}
 
 	// First look to see if we already have a timeline with that name
 	UTimelineTemplate* Timeline = Blueprint->FindTimelineTemplateByVariableName(TimelineVarName);
-	if (Timeline != NULL)
+	if (Timeline != nullptr)
 	{
 		UE_LOG(LogBlueprint, Log, TEXT("AddNewTimeline: Blueprint '%s' already contains a timeline called '%s'"), *Blueprint->GetPathName(), *TimelineVarName.ToString());
-		return NULL;
+		return nullptr;
 	}
 	else
 	{
 		Blueprint->Modify();
-		check(NULL != Blueprint->GeneratedClass);
+		check(nullptr != Blueprint->GeneratedClass);
 		// Construct new graph with the supplied name
 		const FName TimelineTemplateName = *UTimelineTemplate::TimelineVariableNameToTemplateName(TimelineVarName);
 		Timeline = NewObject<UTimelineTemplate>(Blueprint->GeneratedClass, TimelineTemplateName, RF_Transactional);
@@ -6655,7 +6630,7 @@ UK2Node_Timeline* FBlueprintEditorUtils::FindNodeForTimeline(UBlueprint* Bluepri
 		}
 	}
 
-	return NULL; // no node found
+	return nullptr; // no node found
 }
 
 bool FBlueprintEditorUtils::RenameTimeline(UBlueprint* Blueprint, const FName OldNameRef, const FName NewName)
@@ -6709,7 +6684,7 @@ bool FBlueprintEditorUtils::RenameTimeline(UBlueprint* Blueprint, const FName Ol
 						else
 						{
 							//TODO:
-							UClass* ParentClass = TestNode->VariableReference.GetMemberParentClass((UClass*)NULL);
+							UClass* ParentClass = TestNode->VariableReference.GetMemberParentClass((UClass*)nullptr);
 							TestNode->VariableReference.SetExternalMember(NewName, ParentClass);
 						}
 						TestPin->Modify();
@@ -6720,8 +6695,8 @@ bool FBlueprintEditorUtils::RenameTimeline(UBlueprint* Blueprint, const FName Ol
 
 			Blueprint->Timelines.Remove(Template);
 			
-			UObject* ExistingObject = StaticFindObject(NULL, Template->GetOuter(), *NewTemplateName, true);
-			if (ExistingObject != Template && ExistingObject != NULL)
+			UObject* ExistingObject = StaticFindObject(nullptr, Template->GetOuter(), *NewTemplateName, true);
+			if (ExistingObject != Template && ExistingObject != nullptr)
 			{
 				ExistingObject->Rename(*MakeUniqueObjectName(ExistingObject->GetOuter(), ExistingObject->GetClass(), ExistingObject->GetFName()).ToString());
 			}
@@ -6759,7 +6734,7 @@ int32 FBlueprintEditorUtils::FindNumReferencesToActorFromLevelScript(ULevelScrip
 		for( TArray<UK2Node*>::TConstIterator NodeIt(GraphNodes); NodeIt; ++NodeIt )
 		{
 			const UK2Node* CurrentNode = *NodeIt;
-			if( CurrentNode != NULL && CurrentNode->GetReferencedLevelActor() == InActor )
+			if( CurrentNode != nullptr && CurrentNode->GetReferencedLevelActor() == InActor )
 			{
 				RefCount++;
 			}
@@ -6818,7 +6793,7 @@ void  FBlueprintEditorUtils::ModifyActorReferencedGraphNodes(ULevelScriptBluepri
 		for( TArray<UK2Node*>::TIterator NodeIt(GraphNodes); NodeIt; ++NodeIt )
 		{
 			UK2Node* CurrentNode = *NodeIt;
-			if( CurrentNode != NULL && CurrentNode->GetReferencedLevelActor() == InActor )
+			if( CurrentNode != nullptr && CurrentNode->GetReferencedLevelActor() == InActor )
 			{
 				CurrentNode->Modify();
 			}
@@ -6933,7 +6908,7 @@ bool FBlueprintEditorUtils::FixLevelScriptActorBindings(ALevelScriptActor* Level
 		{
 			const UK2Node_MatineeController* MatController = *NodeIt;
 
-			if(MatController->MatineeActor != NULL)
+			if(MatController->MatineeActor != nullptr)
 			{
 				MatController->MatineeActor->MatineeControllerName = MatController->GetFName();
 			}
@@ -6988,12 +6963,12 @@ bool FBlueprintEditorUtils::KismetDiagnosticExec(const TCHAR* InStream, FOutputD
 
 
 			// Look for a containing scope (either a blueprint or a class)
-			UObject* OuterScope = NULL;
-			for (UObject* TestOuter = TestObj; (TestOuter != NULL) && (OuterScope == NULL); TestOuter = TestOuter->GetOuter())
+			UObject* OuterScope = nullptr;
+			for (UObject* TestOuter = TestObj; (TestOuter != nullptr) && (OuterScope == nullptr); TestOuter = TestOuter->GetOuter())
 			{
 				if (UClass* OuterClass = Cast<UClass>(TestOuter))
 				{
-					if (OuterClass->ClassGeneratedBy != NULL)
+					if (OuterClass->ClassGeneratedBy != nullptr)
 					{
 						OuterScope = OuterClass;
 					}
@@ -7004,17 +6979,15 @@ bool FBlueprintEditorUtils::KismetDiagnosticExec(const TCHAR* InStream, FOutputD
 				}
 			}
 
-			if ((OuterScope != NULL) && !OuterScope->IsIn(TransientPackage))
+			if ((OuterScope != nullptr) && !OuterScope->IsIn(TransientPackage))
 			{
 				// Find all references
 				TArray<UObject*> ReferencedObjects;
-				FReferenceFinder ObjectReferenceCollector(ReferencedObjects, NULL, false, false, false, false);
+				FReferenceFinder ObjectReferenceCollector(ReferencedObjects, nullptr, false, false, false, false);
 				ObjectReferenceCollector.FindReferences(TestObj);
 
-				for (auto TouchedIt = ReferencedObjects.CreateConstIterator(); TouchedIt; ++TouchedIt)
+				for (UObject* ReferencedObj : ReferencedObjects)
 				{
-					UObject* ReferencedObj = *TouchedIt;
-
 					// Ignore any references inside the outer blueprint or class; they're intrinsically safe
 					if (!ReferencedObj->IsIn(OuterScope))
 					{
@@ -7037,13 +7010,12 @@ bool FBlueprintEditorUtils::KismetDiagnosticExec(const TCHAR* InStream, FOutputD
 
 		// Report all the bad outers as text dumps so the exact property can be identified
 		Ar.Logf(TEXT("Summary of assets containing objects that have bad references"));
-		for (auto BadIt = ObjectsContainingBadness.CreateConstIterator(); BadIt; ++BadIt)
+		for (UObject* BadObj : ObjectsContainingBadness)
 		{
-			UObject* BadObj = *BadIt;
 			Ar.Logf(TEXT("\n\nObject %s referenced private objects outside of it's container asset inappropriately"), *BadObj->GetFullName());
 
 			UBlueprint* Blueprint = Cast<UBlueprint>(BadObj);
-			if (Blueprint == NULL)
+			if (Blueprint == nullptr)
 			{
 				if (UClass* Class = Cast<UClass>(BadObj))
 				{
@@ -7071,7 +7043,7 @@ bool FBlueprintEditorUtils::KismetDiagnosticExec(const TCHAR* InStream, FOutputD
 				UnMarkAllObjects(EObjectMark(OBJECTMARK_TagExp | OBJECTMARK_TagImp));
 				FStringOutputDevice Archive;
 				const FExportObjectInnerContext Context;
-				UExporter::ExportToOutputDevice(&Context, BadObj, NULL, Archive, TEXT("copy"), 0, /*PPF_IncludeTransient*/ /*PPF_ExportsNotFullyQualified|*/PPF_Copy, false, NULL);
+				UExporter::ExportToOutputDevice(&Context, BadObj, nullptr, Archive, TEXT("copy"), 0, /*PPF_IncludeTransient*/ /*PPF_ExportsNotFullyQualified|*/PPF_Copy, false, nullptr);
 				FString ExportedText = Archive;
 
 				Ar.Logf(TEXT("%s"), *ExportedText);
@@ -7079,10 +7051,8 @@ bool FBlueprintEditorUtils::KismetDiagnosticExec(const TCHAR* InStream, FOutputD
 		}
 
 		// Report the contents of the bad packages
-		for (auto BadIt = BadPackages.CreateConstIterator(); BadIt; ++BadIt)
+		for (UPackage* BadPackage : BadPackages)
 		{
-			UPackage* BadPackage = *BadIt;
-			
 			Ar.Logf(TEXT("\nBad package %s contains:"), *BadPackage->GetName());
 			for (FObjectIterator ObjIt; ObjIt; ++ObjIt)
 			{
@@ -7097,7 +7067,7 @@ bool FBlueprintEditorUtils::KismetDiagnosticExec(const TCHAR* InStream, FOutputD
 	}
 	else if (FParse::Command(&Str, TEXT("ListPackageContents")))
 	{
-		if (UPackage* Package = FindPackage(NULL, Str))
+		if (UPackage* Package = FindPackage(nullptr, Str))
 		{
 			FBlueprintEditorUtils::ListPackageContents(Package, Ar);
 		}
@@ -7124,7 +7094,7 @@ bool FBlueprintEditorUtils::KismetDiagnosticExec(const TCHAR* InStream, FOutputD
 		for( TObjectIterator<UClass> it; it; ++it )
 		{
 			UClass* CurrClass = *it;
-			if( CurrClass->ClassGeneratedBy != NULL && CurrClass->GetOutermost() != GetTransientPackage() )
+			if( CurrClass->ClassGeneratedBy != nullptr && CurrClass->GetOutermost() != GetTransientPackage() )
 			{
 				if( UBlueprint* GeneratingBP = Cast<UBlueprint>(CurrClass->ClassGeneratedBy) )
 				{
@@ -7219,7 +7189,7 @@ TSharedRef<SWidget> FBlueprintEditorUtils::ConstructBlueprintParentClassPicker(
 	TArray<UClass*> BlueprintClasses;
 	for( auto BlueprintIter = Blueprints.CreateConstIterator(); (!bIsActor && !bIsAnimBlueprint) && BlueprintIter; ++BlueprintIter )
 	{
-		const auto Blueprint = *BlueprintIter;
+		const UBlueprint* Blueprint = *BlueprintIter;
 		bIsActor |= Blueprint->ParentClass->IsChildOf( AActor::StaticClass() );
 		bIsAnimBlueprint |= Blueprint->IsA(UAnimBlueprint::StaticClass());
 		bIsLevelScriptActor |= Blueprint->ParentClass->IsChildOf( ALevelScriptActor::StaticClass() );
@@ -7242,9 +7212,9 @@ TSharedRef<SWidget> FBlueprintEditorUtils::ConstructBlueprintParentClassPicker(
 	Filter->DisallowedChildrenOfClasses.Add( UInterface::StaticClass() );
 
 	// never allow parenting to children of itself
-	for( auto ClassIt = BlueprintClasses.CreateIterator(); ClassIt; ++ClassIt )
+	for (UClass* BPClass : BlueprintClasses)
 	{
-		Filter->DisallowedChildrenOfClasses.Add(*ClassIt);
+		Filter->DisallowedChildrenOfClasses.Add(BPClass);
 	}
 
 	for ( UBlueprint* Blueprint : Blueprints )
@@ -7285,10 +7255,8 @@ TSharedRef<SWidget> FBlueprintEditorUtils::ConstructBlueprintParentClassPicker(
 	}
 
 
-	for( auto BlueprintIter = Blueprints.CreateConstIterator(); BlueprintIter; ++BlueprintIter )
+	for( const UBlueprint* Blueprint : Blueprints)
 	{
-		const auto Blueprint = *BlueprintIter;
-
 		// don't allow making me my own parent!
 		Filter->DisallowedClasses.Add(Blueprint->GeneratedClass);
 	}
@@ -7373,9 +7341,8 @@ public:
 TSharedRef<SWidget> FBlueprintEditorUtils::ConstructBlueprintInterfaceClassPicker( const TArray< UBlueprint* >& Blueprints, const FOnClassPicked& OnPicked)
 {
 	TArray<UClass*> BlueprintClasses;
-	for( auto BlueprintIter = Blueprints.CreateConstIterator(); BlueprintIter; ++BlueprintIter )
+	for (const UBlueprint* Blueprint : Blueprints)
 	{
-		const auto Blueprint = *BlueprintIter;
 		BlueprintClasses.Add(Blueprint->GeneratedClass);
 	}
 
@@ -7387,16 +7354,14 @@ TSharedRef<SWidget> FBlueprintEditorUtils::ConstructBlueprintInterfaceClassPicke
 	Options.ClassFilter = Filter;
 	Options.ViewerTitleString = LOCTEXT("ImplementInterfaceBlueprint", "Implement Interface");
 
-	for( auto BlueprintIter = Blueprints.CreateConstIterator(); BlueprintIter; ++BlueprintIter )
+	for (const UBlueprint* Blueprint : Blueprints)
 	{
-		const auto Blueprint = *BlueprintIter;
-
 		// don't allow making me my own parent!
 		Filter->DisallowedClasses.Add(Blueprint->GeneratedClass);
 
 		UClass const* const ParentClass = Blueprint->ParentClass;
 		// see if the parent class has any prohibited interfaces
-		if ((ParentClass != NULL) && ParentClass->HasMetaData(FBlueprintMetadata::MD_ProhibitedInterfaces))
+		if ((ParentClass != nullptr) && ParentClass->HasMetaData(FBlueprintMetadata::MD_ProhibitedInterfaces))
 		{
 			FString const& ProhibitedList = Blueprint->ParentClass->GetMetaData(FBlueprintMetadata::MD_ProhibitedInterfaces);
 
@@ -7425,9 +7390,9 @@ TSharedRef<SWidget> FBlueprintEditorUtils::ConstructBlueprintInterfaceClassPicke
 	}
 
 	// never allow parenting to children of itself
-	for( auto ClassIt = BlueprintClasses.CreateIterator(); ClassIt; ++ClassIt )
+	for (UClass*  BPClass : BlueprintClasses)
 	{
-		Filter->DisallowedChildrenOfClasses.Add(*ClassIt);
+		Filter->DisallowedChildrenOfClasses.Add(BPClass);
 	}
 
 	return FModuleManager::LoadModuleChecked<FClassViewerModule>("ClassViewer").CreateClassViewer(Options, OnPicked);
@@ -7437,14 +7402,13 @@ void FBlueprintEditorUtils::UpdateOldPureFunctions( UBlueprint* Blueprint )
 {
 	if(UBlueprintGeneratedClass* GeneratedClass = Cast<UBlueprintGeneratedClass>(Blueprint->SkeletonGeneratedClass))
 	{
-		for(auto Function : TFieldRange<UFunction>(GeneratedClass, EFieldIteratorFlags::ExcludeSuper))
+		for (UFunction* Function : TFieldRange<UFunction>(GeneratedClass, EFieldIteratorFlags::ExcludeSuper))
 		{
 			if(Function && ((Function->FunctionFlags & FUNC_BlueprintPure) > 0))
 			{
-				for(auto FuncGraphsIt(Blueprint->FunctionGraphs.CreateIterator());FuncGraphsIt;++FuncGraphsIt)
+				for (UEdGraph* FunctionGraph : Blueprint->FunctionGraphs)
 				{
-					UEdGraph* FunctionGraph = *FuncGraphsIt;
-					if(FunctionGraph->GetName() == Function->GetName())
+					if (FunctionGraph->GetFName() == Function->GetFName())
 					{
 						TArray<UK2Node_FunctionEntry*> EntryNodes;
 						FunctionGraph->GetNodesOfClass(EntryNodes);
@@ -7475,10 +7439,10 @@ void FBlueprintEditorUtils::PostEditChangeBlueprintActors(UBlueprint* Blueprint,
 		TArray<UObject*> MatchingBlueprintObjects;
 		GetObjectsOfClass(Blueprint->GeneratedClass, MatchingBlueprintObjects, bIncludeDerivedClasses, RF_ClassDefaultObject, EInternalObjectFlags::PendingKill);
 
-		for (auto ObjIt : MatchingBlueprintObjects)
+		for (UObject* MatchingObj : MatchingBlueprintObjects)
 		{
 			// We know the class was derived from AActor because we checked the Blueprint->GeneratedClass.
-			AActor* Actor = static_cast<AActor*>(ObjIt);
+			AActor* Actor = static_cast<AActor*>(MatchingObj);
 			Actor->PostEditChange();
 
 			// Broadcast edit notification if necessary so that the level editor's detail panel is refreshed
@@ -7498,10 +7462,10 @@ void FBlueprintEditorUtils::PostEditChangeBlueprintActors(UBlueprint* Blueprint,
 
 	// Let the blueprint thumbnail renderer know that a blueprint has been modified so it knows to reinstance components for visualization
 	FThumbnailRenderingInfo* RenderInfo = GUnrealEd->GetThumbnailManager()->GetRenderingInfo( Blueprint );
-	if ( RenderInfo != NULL )
+	if ( RenderInfo != nullptr )
 	{
 		UBlueprintThumbnailRenderer* BlueprintThumbnailRenderer = Cast<UBlueprintThumbnailRenderer>(RenderInfo->Renderer);
-		if ( BlueprintThumbnailRenderer != NULL )
+		if ( BlueprintThumbnailRenderer != nullptr )
 		{
 			BlueprintThumbnailRenderer->BlueprintChanged(Blueprint);
 		}
@@ -7535,9 +7499,9 @@ void FBlueprintEditorUtils::FindAndSetDebuggableBlueprintInstances()
 	for(int32 i=0; i<EditedAssets.Num(); i++)
 	{
 		UBlueprint* Blueprint = Cast<UBlueprint>( EditedAssets[i] );
-		if( Blueprint != NULL )
+		if( Blueprint != nullptr )
 		{
-			if( Blueprint->GetObjectBeingDebugged() == NULL )
+			if( Blueprint->GetObjectBeingDebugged() == nullptr )
 			{
 				BlueprintsNeedingInstancesToDebug.FindOrAdd( Blueprint );
 			}			
@@ -7559,8 +7523,8 @@ void FBlueprintEditorUtils::FindAndSetDebuggableBlueprintInstances()
 		{	
 			UBlueprint* EachBlueprint = ObjIt.Key();
 			bool bFoundItemToDebug = false;
-			AActor* SimilarInstanceSelected = NULL;
-			AActor* SimilarInstanceUnselected = NULL;
+			AActor* SimilarInstanceSelected = nullptr;
+			AActor* SimilarInstanceUnselected = nullptr;
 
 			// First check selected objects.
 			if( Selected->Num() != 0 )
@@ -7568,8 +7532,8 @@ void FBlueprintEditorUtils::FindAndSetDebuggableBlueprintInstances()
 				for (int32 iSelected = 0; iSelected < Selected->Num() ; iSelected++)
 				{
 					AActor* ObjectAsActor = Cast<AActor>( Selected->GetSelectedObject( iSelected ) );
-					UWorld* ActorWorld = ObjectAsActor ? ObjectAsActor->GetWorld() : NULL;
-					if ((ActorWorld != NULL) && (ActorWorld->WorldType != EWorldType::Preview) && (ActorWorld->WorldType != EWorldType::Inactive))
+					UWorld* ActorWorld = ObjectAsActor ? ObjectAsActor->GetWorld() : nullptr;
+					if ((ActorWorld != nullptr) && (ActorWorld->WorldType != EWorldType::Preview) && (ActorWorld->WorldType != EWorldType::Inactive))
 					{
 						if( IsObjectADebugCandidate(ObjectAsActor, EachBlueprint, true/*bInDisallowDerivedBlueprints*/ ) == true )
 						{
@@ -7578,7 +7542,7 @@ void FBlueprintEditorUtils::FindAndSetDebuggableBlueprintInstances()
 							BlueprintsToRefresh.Add( EachBlueprint );
 							break;
 						}
-						else if( SimilarInstanceSelected == NULL)
+						else if( SimilarInstanceSelected == nullptr)
 						{
 							// If we haven't found a similar selected instance already check for one now
 							if( IsObjectADebugCandidate(ObjectAsActor, EachBlueprint, false/*bInDisallowDerivedBlueprints*/ ) == true )
@@ -7595,7 +7559,7 @@ void FBlueprintEditorUtils::FindAndSetDebuggableBlueprintInstances()
 				for (TObjectIterator<UObject> It; It; ++It)
 				{
 					AActor* ObjectAsActor = Cast<AActor>( *It );
-					UWorld* ActorWorld = ObjectAsActor ? ObjectAsActor->GetWorld() : NULL;
+					UWorld* ActorWorld = ObjectAsActor ? ObjectAsActor->GetWorld() : nullptr;
 					if( ActorWorld && ( ActorWorld->WorldType != EWorldType::Preview) && ActorWorld->WorldType != EWorldType::Inactive )
 					{
 						if( IsObjectADebugCandidate(ObjectAsActor, EachBlueprint, true/*bInDisallowDerivedBlueprints*/ ) == true )
@@ -7605,7 +7569,7 @@ void FBlueprintEditorUtils::FindAndSetDebuggableBlueprintInstances()
 							BlueprintsToRefresh.Add( EachBlueprint );
 							break;
 						}
-						else if( SimilarInstanceUnselected == NULL)
+						else if( SimilarInstanceUnselected == nullptr)
 						{
 							// If we haven't found a similar unselected instance already check for one now
 							if( IsObjectADebugCandidate(ObjectAsActor, EachBlueprint, false/*bInDisallowDerivedBlueprints*/ ) == true )
@@ -7620,9 +7584,9 @@ void FBlueprintEditorUtils::FindAndSetDebuggableBlueprintInstances()
 			// If we didn't find and exact type match, but we did find a related type use that.
 			if( bFoundItemToDebug == false )
 			{
-				if( ( SimilarInstanceSelected != NULL ) || ( SimilarInstanceUnselected != NULL ) )
+				if( ( SimilarInstanceSelected != nullptr ) || ( SimilarInstanceUnselected != nullptr ) )
 				{
-					EachBlueprint->SetObjectBeingDebugged( SimilarInstanceSelected != NULL ? SimilarInstanceSelected : SimilarInstanceUnselected );
+					EachBlueprint->SetObjectBeingDebugged( SimilarInstanceSelected != nullptr ? SimilarInstanceSelected : SimilarInstanceUnselected );
 					BlueprintsToRefresh.Add( EachBlueprint );
 				}
 			}
@@ -7722,7 +7686,7 @@ bool FBlueprintEditorUtils::PropertyValueFromString(const UProperty* Property, c
 		else if( Property->IsA(UTextProperty::StaticClass()) )
 		{
 			FStringOutputDevice ImportError;
-			const auto EndOfParsedBuff = Property->ImportText(*Value, Property->ContainerPtrToValuePtr<uint8>(DefaultObject), 0, nullptr, &ImportError);
+			const TCHAR* EndOfParsedBuff = Property->ImportText(*Value, Property->ContainerPtrToValuePtr<uint8>(DefaultObject), 0, nullptr, &ImportError);
 			bParseSucceeded = EndOfParsedBuff && ImportError.IsEmpty();
 		}
 		else
@@ -7733,7 +7697,7 @@ bool FBlueprintEditorUtils::PropertyValueFromString(const UProperty* Property, c
 				: *Value;
 
 			FStringOutputDevice ImportError;
-			const auto EndOfParsedBuff = Property->ImportText(*Value, Property->ContainerPtrToValuePtr<uint8>(DefaultObject), 0, nullptr, &ImportError);
+			const TCHAR* EndOfParsedBuff = Property->ImportText(*Value, Property->ContainerPtrToValuePtr<uint8>(DefaultObject), 0, nullptr, &ImportError);
 			bParseSucceeded = EndOfParsedBuff && ImportError.IsEmpty();
 		}
 	}
@@ -7782,7 +7746,7 @@ bool FBlueprintEditorUtils::PropertyValueFromString(const UProperty* Property, c
 			bParseSucceeded = FStructureEditorUtils::Fill_MakeStructureDefaultValue(Cast<const UUserDefinedStruct>(Struct), StructData);
 
 			FStringOutputDevice ImportError;
-			const auto EndOfParsedBuff = StructProperty->ImportText(Value.IsEmpty() ? TEXT("()") : *Value, StructData, 0, nullptr, &ImportError);
+			const TCHAR* EndOfParsedBuff = StructProperty->ImportText(Value.IsEmpty() ? TEXT("()") : *Value, StructData, 0, nullptr, &ImportError);
 			bParseSucceeded &= EndOfParsedBuff && ImportError.IsEmpty();
 		}
 	}
@@ -7901,9 +7865,9 @@ bool FBlueprintEditorUtils::CheckIfNodeConnectsToSelection(UEdGraphNode* InNode,
 bool FBlueprintEditorUtils::CheckIfSelectionIsCycling(const TSet<UEdGraphNode*>& InSelectionSet, FCompilerResultsLog& InMessageLog)
 {
 	const UEdGraphSchema_K2* Schema = GetDefault<UEdGraphSchema_K2>();
-	for (auto NodeIt = InSelectionSet.CreateConstIterator(); NodeIt; ++NodeIt)
+	for (UEdGraphNode* Node : InSelectionSet)
 	{
-		if (UEdGraphNode* Node = Cast<UEdGraphNode>(*NodeIt))
+		if (Node)
 		{
 			for (int32 PinIndex = 0; PinIndex < Node->Pins.Num(); ++PinIndex)
 			{
@@ -7943,7 +7907,7 @@ bool FBlueprintEditorUtils::IsPaletteActionReadOnly(TSharedPtr<FEdGraphSchemaAct
 	{
 		FEdGraphSchemaAction_K2Graph* GraphAction = (FEdGraphSchemaAction_K2Graph*)ActionIn.Get();
 		// No graph is evidence of an overridable function, don't let the user modify it
-		if(GraphAction->EdGraph == NULL)
+		if(GraphAction->EdGraph == nullptr)
 		{
 			bIsReadOnly = true;
 		}
@@ -7960,7 +7924,7 @@ bool FBlueprintEditorUtils::IsPaletteActionReadOnly(TSharedPtr<FEdGraphSchemaAct
 				{
 					// Check if the function is an override
 					UFunction* OverrideFunc = FindField<UFunction>(BlueprintObj->ParentClass, GraphAction->FuncName);
-					if ( OverrideFunc != NULL )
+					if ( OverrideFunc != nullptr )
 					{
 						bIsReadOnly = true;
 					}
@@ -8018,7 +7982,7 @@ bool FBlueprintEditorUtils::IsPaletteActionReadOnly(TSharedPtr<FEdGraphSchemaAct
 		FEdGraphSchemaAction_K2Event* EventAction = (FEdGraphSchemaAction_K2Event*)ActionIn.Get();
 		UK2Node* AssociatedNode = EventAction->NodeTemplate;
 
-		bIsReadOnly = (AssociatedNode == NULL) || (!AssociatedNode->bCanRenameNode);	
+		bIsReadOnly = (AssociatedNode == nullptr) || (!AssociatedNode->bCanRenameNode);	
 	}
 	else if (ActionIn->GetTypeId() == FEdGraphSchemaAction_K2InputAction::StaticGetTypeId())
 	{
@@ -8033,7 +7997,7 @@ struct FUberGraphHelper
 {
 	static void GetAll(const UBlueprint* Blueprint, TArray<UEdGraph*>& OutGraphs)
 	{
-		for (const auto UberGraph : Blueprint->UbergraphPages)
+		for (UEdGraph* UberGraph : Blueprint->UbergraphPages)
 		{
 			OutGraphs.Add(UberGraph);
 			UberGraph->GetAllChildrenGraphs(OutGraphs);
@@ -8049,7 +8013,7 @@ FName FBlueprintEditorUtils::GetFunctionNameFromClassByGuid(const UClass* InClas
 	for (int32 BPIndex = 0; BPIndex < Blueprints.Num(); ++BPIndex)
 	{
 		UBlueprint* Blueprint = Blueprints[BPIndex];
-		for (auto FunctionGraph : Blueprint->FunctionGraphs)
+		for (UEdGraph* FunctionGraph : Blueprint->FunctionGraphs)
 		{
 			if (FunctionGraph && FunctionGraph->GraphGuid == FunctionGuid)
 			{
@@ -8057,7 +8021,7 @@ FName FBlueprintEditorUtils::GetFunctionNameFromClassByGuid(const UClass* InClas
 			}
 		}
 
-		for (auto FunctionGraph : Blueprint->DelegateSignatureGraphs)
+		for (UEdGraph* FunctionGraph : Blueprint->DelegateSignatureGraphs)
 		{
 			if (FunctionGraph && FunctionGraph->GraphGuid == FunctionGuid)
 			{
@@ -8069,11 +8033,11 @@ FName FBlueprintEditorUtils::GetFunctionNameFromClassByGuid(const UClass* InClas
 		//FUNCTIONS BASED ON CUSTOM EVENTS:
 		TArray<UEdGraph*> UberGraphs;
 		FUberGraphHelper::GetAll(Blueprint, UberGraphs);
-		for (const auto UberGraph : UberGraphs)
+		for (const UEdGraph* UberGraph : UberGraphs)
 		{
 			TArray<UK2Node_CustomEvent*> CustomEvents;
 			UberGraph->GetNodesOfClass(CustomEvents);
-			for (const auto CustomEvent : CustomEvents)
+			for (const UK2Node_CustomEvent* CustomEvent : CustomEvents)
 			{
 				if (!CustomEvent->bOverrideFunction && (CustomEvent->NodeGuid == FunctionGuid))
 				{
@@ -8097,7 +8061,7 @@ bool FBlueprintEditorUtils::GetFunctionGuidFromClassByFieldName(const UClass* In
 		for (int32 BPIndex = 0; BPIndex < Blueprints.Num(); ++BPIndex)
 		{
 			UBlueprint* Blueprint = Blueprints[BPIndex];
-			for (auto FunctionGraph : Blueprint->FunctionGraphs)
+			for (UEdGraph* FunctionGraph : Blueprint->FunctionGraphs)
 			{
 				if (FunctionGraph && FunctionGraph->GetFName() == FunctionName)
 				{
@@ -8110,7 +8074,7 @@ bool FBlueprintEditorUtils::GetFunctionGuidFromClassByFieldName(const UClass* In
 			if (BaseDelegateSignatureName.RemoveFromEnd(HEADER_GENERATED_DELEGATE_SIGNATURE_SUFFIX))
 			{
 				const FName GraphName(*BaseDelegateSignatureName);
-				for (auto FunctionGraph : Blueprint->DelegateSignatureGraphs)
+				for (UEdGraph* FunctionGraph : Blueprint->DelegateSignatureGraphs)
 				{
 					if (FunctionGraph && FunctionGraph->GetFName() == GraphName)
 					{
@@ -8122,11 +8086,11 @@ bool FBlueprintEditorUtils::GetFunctionGuidFromClassByFieldName(const UClass* In
 
 			TArray<UEdGraph*> UberGraphs;
 			FUberGraphHelper::GetAll(Blueprint, UberGraphs);
-			for (const auto UberGraph : UberGraphs)
+			for (const UEdGraph* UberGraph : UberGraphs)
 			{
 				TArray<UK2Node_CustomEvent*> CustomEvents;
 				UberGraph->GetNodesOfClass(CustomEvents);
-				for (const auto CustomEvent : CustomEvents)
+				for (const UK2Node_CustomEvent* CustomEvent : CustomEvents)
 				{
 					if (!CustomEvent->bOverrideFunction
 						&& (CustomEvent->CustomFunctionName == FunctionName)
@@ -8208,7 +8172,7 @@ void FBlueprintEditorUtils::GetEntryAndResultNodes(const UEdGraph* InGraph, TWea
 				TArray<UK2Node_FunctionResult*> ResultNodes;
 				InGraph->GetNodesOfClass(ResultNodes);
 
-				UK2Node_FunctionResult* ResultNode = ResultNodes.Num() ? ResultNodes[0] : NULL;
+				UK2Node_FunctionResult* ResultNode = ResultNodes.Num() ? ResultNodes[0] : nullptr;
 				// Note:  we assume that if the entry is editable, the result is too (since the entry node is guaranteed to be there on graph creation, but the result isn't)
 				if( ResultNode )
 				{
@@ -8267,7 +8231,7 @@ FKismetUserDeclaredFunctionMetadata* FBlueprintEditorUtils::GetGraphFunctionMeta
 
 FText FBlueprintEditorUtils::GetGraphDescription(const UEdGraph* InGraph)
 {
-	auto FunctionEntryNode = GetEntryNode(InGraph);
+	UK2Node_EditablePinBase* FunctionEntryNode = GetEntryNode(InGraph);
 	if (UK2Node_FunctionEntry* TypedEntryNode = Cast<UK2Node_FunctionEntry>(FunctionEntryNode))
 	{
 		return FText::FromString(TypedEntryNode->MetaData.ToolTip);
@@ -8287,7 +8251,7 @@ bool FBlueprintEditorUtils::CheckIfGraphHasLatentFunctions(UEdGraph* InGraph)
 	{
 		static bool CheckIfGraphHasLatentFunctions(UEdGraph* InGraphToCheck, TArray<UEdGraph*>& InspectedGraphList)
 		{
-			auto EntryNode = GetEntryNode(InGraphToCheck);
+			UK2Node_EditablePinBase* EntryNode = GetEntryNode(InGraphToCheck);
 
 			UK2Node_Tunnel* TunnelNode = ExactCast<UK2Node_Tunnel>(EntryNode);
 			if(!TunnelNode)
@@ -8310,7 +8274,7 @@ bool FBlueprintEditorUtils::CheckIfGraphHasLatentFunctions(UEdGraph* InGraph)
 					if(const UK2Node_CallFunction* CallFunctionNode = Cast<UK2Node_CallFunction>(Node))
 					{
 						// Check any function call nodes to see if they are latent.
-						auto TargetFunction = CallFunctionNode->GetTargetFunction();
+						UFunction* TargetFunction = CallFunctionNode->GetTargetFunction();
 						if (TargetFunction && TargetFunction->HasMetaData(FBlueprintMetadata::MD_Latent))
 						{
 							TunnelNode->MetaData.HasLatentFunctions = 1;
@@ -8381,7 +8345,7 @@ void FBlueprintEditorUtils::PostSetupObjectPinType(UBlueprint* InBlueprint, FBPV
 		// if it's a PC_Object, then it should have an associated UClass object
 		check(InOutVarDesc.VarType.PinSubCategoryObject.IsValid());
 		const UClass* ClassObject = Cast<UClass>(InOutVarDesc.VarType.PinSubCategoryObject.Get());
-		check(ClassObject != NULL);
+		check(ClassObject != nullptr);
 
 		if (ClassObject->IsChildOf(AActor::StaticClass()))
 		{
