Copying //UE4/Dev-Framework to //UE4/Dev-Main (Source: //UE4/Dev-Framework @ 3136612)

diff --git a/Engine/Source/Editor/UnrealEd/Private/Kismet2/KismetReinstanceUtilities.cpp b/Engine/Source/Editor/UnrealEd/Private/Kismet2/KismetReinstanceUtilities.cpp
--- a/Engine/Source/Editor/UnrealEd/Private/Kismet2/KismetReinstanceUtilities.cpp
+++ b/Engine/Source/Editor/UnrealEd/Private/Kismet2/KismetReinstanceUtilities.cpp
@@ -49,7 +49,7 @@ struct FReplaceReferenceHelper
 		OldToNewInstanceMap.Add(OldClass, NewClass);
 		SourceObjects.Add(OldClass);
 
-		if (auto OldCDO = OldClass->GetDefaultObject(false))
+		if (UObject* OldCDO = OldClass->GetDefaultObject(false))
 		{
 			ObjectsToReplace.Add(OldCDO);
 		}
@@ -62,7 +62,7 @@ struct FReplaceReferenceHelper
 		{
 			BP_SCOPED_COMPILER_EVENT_STAT(EKismetReinstancerStats_FindReferencers);
 
-			TFindObjectReferencers<UObject> Referencers(SourceObjects, NULL, false);
+			TFindObjectReferencers<UObject> Referencers(SourceObjects, nullptr, false);
 			for (TFindObjectReferencers<UObject>::TIterator It(Referencers); It; ++It)
 			{
 				UObject* Referencer = It.Value();
@@ -139,8 +139,8 @@ UClass* FBlueprintCompileReinstancer::HotReloadedNewClass = nullptr;
 
 FBlueprintCompileReinstancer::FBlueprintCompileReinstancer(UClass* InClassToReinstance, bool bIsBytecodeOnly, bool bSkipGC, bool bAutoInferSaveOnCompile/* = true*/)
 	: ClassToReinstance(InClassToReinstance)
-	, DuplicatedClass(NULL)
-	, OriginalCDO(NULL)
+	, DuplicatedClass(nullptr)
+	, OriginalCDO(nullptr)
 	, bHasReinstanced(false)
 	, bSkipGarbageCollection(bSkipGC)
 	, ReinstClassType(RCT_Unknown)
@@ -148,7 +148,7 @@ FBlueprintCompileReinstancer::FBlueprintCompileReinstancer(UClass* InClassToRein
 	, bIsRootReinstancer(false)
 	, bAllowResaveAtTheEndIfRequested(false)
 {
-	if( InClassToReinstance != NULL )
+	if( InClassToReinstance != nullptr )
 	{
 		if (FKismetEditorUtilities::IsClassABlueprintSkeleton(ClassToReinstance))
 		{
@@ -172,22 +172,22 @@ FBlueprintCompileReinstancer::FBlueprintCompileReinstancer(UClass* InClassToRein
 		// Duplicate the class we're reinstancing into the transient package.  We'll re-class all objects we find to point to this new class
 		GIsDuplicatingClassForReinstancing = true;
 		ClassToReinstance->ClassFlags |= CLASS_NewerVersionExists;
-		const FName RenistanceName = MakeUniqueObjectName(GetTransientPackage(), ClassToReinstance->GetClass(), *FString::Printf(TEXT("REINST_%s"), *ClassToReinstance->GetName()));
-		DuplicatedClass = (UClass*)StaticDuplicateObject(ClassToReinstance, GetTransientPackage(), RenistanceName, ~RF_Transactional); 
+		const FName ReinstanceName = MakeUniqueObjectName(GetTransientPackage(), ClassToReinstance->GetClass(), *FString::Printf(TEXT("REINST_%s"), *ClassToReinstance->GetName()));
+		DuplicatedClass = (UClass*)StaticDuplicateObject(ClassToReinstance, GetTransientPackage(), ReinstanceName, ~RF_Transactional); 
 		// If you compile a blueprint that is part of the rootset, there's no reason for the REINST version to be part of the rootset:
 		DuplicatedClass->RemoveFromRoot();
 
 		ClassToReinstance->ClassFlags &= ~CLASS_NewerVersionExists;
 		GIsDuplicatingClassForReinstancing = false;
 
-		auto BPClassToReinstance = Cast<UBlueprintGeneratedClass>(ClassToReinstance);
-		auto BPGDuplicatedClass = Cast<UBlueprintGeneratedClass>(DuplicatedClass);
+		UBlueprintGeneratedClass* BPClassToReinstance = Cast<UBlueprintGeneratedClass>(ClassToReinstance);
+		UBlueprintGeneratedClass* BPGDuplicatedClass = Cast<UBlueprintGeneratedClass>(DuplicatedClass);
 		if (BPGDuplicatedClass && BPClassToReinstance && BPClassToReinstance->OverridenArchetypeForCDO)
 		{
 			BPGDuplicatedClass->OverridenArchetypeForCDO = BPClassToReinstance->OverridenArchetypeForCDO;
 		}
 
-		auto DuplicatedClassUberGraphFunction = BPGDuplicatedClass ? BPGDuplicatedClass->UberGraphFunction : nullptr;
+		UFunction* DuplicatedClassUberGraphFunction = BPGDuplicatedClass ? BPGDuplicatedClass->UberGraphFunction : nullptr;
 		if (DuplicatedClassUberGraphFunction)
 		{
 			DuplicatedClassUberGraphFunction->Bind();
@@ -202,7 +202,7 @@ FBlueprintCompileReinstancer::FBlueprintCompileReinstancer(UClass* InClassToRein
 
 		// Temporarily suspend the undo buffer; we don't need to record the duplicated CDO until it is fully resolved
  		ITransaction* CurrentTransaction = GUndo;
- 		GUndo = NULL;
+ 		GUndo = nullptr;
 		DuplicatedClass->ClassDefaultObject = GetClassCDODuplicate(ClassToReinstance->GetDefaultObject(), DuplicatedClass->GetDefaultObjectName());
 
 		// Restore the undo buffer
@@ -221,16 +221,15 @@ FBlueprintCompileReinstancer::FBlueprintCompileReinstancer(UClass* InClassToRein
 			TArray<UObject*> ObjectsToChange;
 			const bool bIncludeDerivedClasses = false;
 			GetObjectsOfClass(ClassToReinstance, ObjectsToChange, bIncludeDerivedClasses);
-			for (auto ObjIt = ObjectsToChange.CreateConstIterator(); ObjIt; ++ObjIt)
+			for (UObject* ObjectToChange : ObjectsToChange)
 			{
-				(*ObjIt)->SetClass(DuplicatedClass);
+				ObjectToChange->SetClass(DuplicatedClass);
 			}
 
 			TArray<UClass*> ChildrenOfClass;
 			GetDerivedClasses(ClassToReinstance, ChildrenOfClass);
-			for ( auto ClassIt = ChildrenOfClass.CreateConstIterator(); ClassIt; ++ClassIt )
+			for ( UClass* ChildClass : ChildrenOfClass )
 			{
-				UClass* ChildClass = *ClassIt;
 				UBlueprint* ChildBP = Cast<UBlueprint>(ChildClass->ClassGeneratedBy);
 				if (ChildBP)
 				{
@@ -304,7 +303,7 @@ void FBlueprintCompileReinstancer::SaveClassFieldMapping(UClass* InClassToReinst
 		PropertyMap.Add(Prop->GetFName(), Prop);
 	}
 
-	for (auto Function : TFieldRange<UFunction>(InClassToReinstance, EFieldIteratorFlags::ExcludeSuper))
+	for (UFunction* Function : TFieldRange<UFunction>(InClassToReinstance, EFieldIteratorFlags::ExcludeSuper))
 	{
 		FunctionMap.Add(Function->GetFName(),Function);
 	}
@@ -316,12 +315,12 @@ void FBlueprintCompileReinstancer::GenerateFieldMappings(TMap<UObject*, UObject*
 
 	FieldMapping.Empty();
 
-	for (auto& Prop : PropertyMap)
+	for (TPair<FName, UProperty*>& Prop : PropertyMap)
 	{
 		FieldMapping.Add(Prop.Value, FindField<UProperty>(ClassToReinstance, *Prop.Key.ToString()));
 	}
 
-	for (auto& Func : FunctionMap)
+	for (TPair<FName, UFunction*>& Func : FunctionMap)
 	{
 		UFunction* NewFunction = ClassToReinstance->FindFunctionByName(Func.Key, EIncludeSuperFlag::ExcludeSuper);
 		FieldMapping.Add(Func.Value, NewFunction);
@@ -421,9 +420,9 @@ public:
 		const bool bIsActor = ClassToReinstance->IsChildOf<AActor>();
 		if (bIsActor)
 		{
-			for (auto Obj : ObjectsToFinalize)
+			for (UObject* Obj : ObjectsToFinalize)
 			{
-				auto Actor = CastChecked<AActor>(Obj);
+				AActor* Actor = CastChecked<AActor>(Obj);
 
 				UWorld* World = Actor->GetWorld();
 				if (World)
@@ -457,7 +456,7 @@ public:
 		//UAnimBlueprintGeneratedClass* AnimClass = Cast<UAnimBlueprintGeneratedClass>(ClassToReinstance);
 		if(bIsAnimInstance)
 		{
-			for(auto Obj : ObjectsToFinalize)
+			for (UObject* Obj : ObjectsToFinalize)
 			{
 				if(USkeletalMeshComponent* SkelComponent = Cast<USkeletalMeshComponent>(Obj->GetOuter()))
 				{
@@ -504,7 +503,7 @@ TSharedPtr<FReinstanceFinalizer> FBlueprintCompileReinstancer::ReinstanceFast()
 	const bool bIsActor = ClassToReinstance->IsChildOf<AActor>();
 	const bool bIsAnimInstance = ClassToReinstance->IsChildOf<UAnimInstance>();
 	const bool bIsComponent = ClassToReinstance->IsChildOf<UActorComponent>();
-	for (auto Obj : Finalizer->ObjectsToReplace)
+	for (UObject* Obj : Finalizer->ObjectsToReplace)
 	{
 		UE_LOG(LogBlueprint, Log, TEXT("  Fast path is refreshing (not replacing) %s"), *Obj->GetFullName());
 
@@ -550,9 +549,8 @@ void FBlueprintCompileReinstancer::CompileChildren()
 	BP_SCOPED_COMPILER_EVENT_STAT(EKismetReinstancerStats_RecompileChildClasses);
 
 	// Reparent all dependent blueprints, and recompile to ensure that they get reinstanced with the new memory layout
-	for (auto ChildBP = Children.CreateIterator(); ChildBP; ++ChildBP)
+	for (UBlueprint* BP : Children)
 	{
-		UBlueprint* BP = *ChildBP;
 		if (BP->ParentClass == ClassToReinstance || BP->ParentClass == DuplicatedClass)
 		{
 			ReparentChild(BP);
@@ -610,7 +608,7 @@ TSharedPtr<FReinstanceFinalizer> FBlueprintCompileReinstancer::ReinstanceInner(b
 
 void FBlueprintCompileReinstancer::ListDependentBlueprintsToRefresh(const TArray<UBlueprint*>& DependentBPs)
 {
-	for (auto& Element : DependentBPs)
+	for (UBlueprint* Element : DependentBPs)
 	{
 		DependentBlueprintsToRefresh.Add(Element);
 	}
@@ -654,137 +652,144 @@ void FBlueprintCompileReinstancer::ReinstanceObjects(bool bForceAlwaysReinstance
 	
 	// Make sure we only reinstance classes once!
 	static TArray<TSharedRef<FBlueprintCompileReinstancer>> QueueToReinstance;
-	TSharedRef<FBlueprintCompileReinstancer> SharedThis = AsShared();
-	bool bAlreadyQueued = QueueToReinstance.Contains(SharedThis);
 
-	// We may already be reinstancing this class, this happens when a dependent blueprint has a compile error and we try to reinstance the stub:
-	for (const auto& Entry : QueueToReinstance)
+	if (!bHasReinstanced)
 	{
-		if (Entry->ClassToReinstance == SharedThis->ClassToReinstance)
-		{
-			bAlreadyQueued = true;
-		}
-	}
-
-	if (!bAlreadyQueued && !bHasReinstanced)
-	{
-		QueueToReinstance.Push(SharedThis);
+		TSharedRef<FBlueprintCompileReinstancer> SharedThis = AsShared();
+		bool bAlreadyQueued = QueueToReinstance.Contains(SharedThis);
 
-		if (ClassToReinstance && DuplicatedClass)
+		// We may already be reinstancing this class, this happens when a dependent blueprint has a compile error and we try to reinstance the stub:
+		if (!bAlreadyQueued)
 		{
-			CompileChildren();
+			for (const TSharedRef<FBlueprintCompileReinstancer>& Entry : QueueToReinstance)
+			{
+				if (Entry->ClassToReinstance == SharedThis->ClassToReinstance)
+				{
+					bAlreadyQueued = true;
+					break;
+				}
+			}
 		}
 
-		if (QueueToReinstance.Num() && (QueueToReinstance[0] == SharedThis))
+		if (!bAlreadyQueued)
 		{
-			// Mark it as the source reinstancer, no other reinstancer can get here until this Blueprint finishes compiling
-			bIsRootReinstancer = true;
-
-			TSet<TWeakObjectPtr<UBlueprint>> CompiledBlueprints;
-			// Blueprints will enqueue dirty and erroring dependents, in case those states would be 
-			// fixed up by having this dependency compiled first. However, this can result in an 
-			// infinite loop where two Blueprints with errors (unrelated to each other) keep 
-			// perpetually queuing the other. 
-			//
-			// To guard against this, we track the recompiled dependents (in order) and break the 
-			// cycle when we see that we've already compiled a dependent after its dependency
-			TArray<UBlueprint*> OrderedRecompiledDependents;
-
-			TSet<TWeakObjectPtr<UBlueprint>> RecompilationQueue = DependentBlueprintsToRecompile;
-			// empty the public facing queue so we can discern between old and new elements (added 
-			// as the result of subsequent recompiles) 
-			DependentBlueprintsToRecompile.Empty();
-
-			while (RecompilationQueue.Num()) 
+			QueueToReinstance.Push(SharedThis);
+
+			if (ClassToReinstance && DuplicatedClass)
 			{
-				auto Iter = RecompilationQueue.CreateIterator();
-				TWeakObjectPtr<UBlueprint> BPPtr = *Iter;
-				Iter.RemoveCurrent();
-				if (UBlueprint* BP = BPPtr.Get())
+				CompileChildren();
+			}
+
+			if (QueueToReinstance.Num() && (QueueToReinstance[0] == SharedThis))
+			{
+				// Mark it as the source reinstancer, no other reinstancer can get here until this Blueprint finishes compiling
+				bIsRootReinstancer = true;
+
+				TSet<TWeakObjectPtr<UBlueprint>> CompiledBlueprints;
+				// Blueprints will enqueue dirty and erroring dependents, in case those states would be 
+				// fixed up by having this dependency compiled first. However, this can result in an 
+				// infinite loop where two Blueprints with errors (unrelated to each other) keep 
+				// perpetually queuing the other. 
+				//
+				// To guard against this, we track the recompiled dependents (in order) and break the 
+				// cycle when we see that we've already compiled a dependent after its dependency
+				TArray<UBlueprint*> OrderedRecompiledDependents;
+
+				TSet<TWeakObjectPtr<UBlueprint>> RecompilationQueue = DependentBlueprintsToRecompile;
+				// empty the public facing queue so we can discern between old and new elements (added 
+				// as the result of subsequent recompiles) 
+				DependentBlueprintsToRecompile.Empty();
+
+				while (RecompilationQueue.Num())
 				{
-					if (IsReinstancingSkeleton())
+					auto Iter = RecompilationQueue.CreateIterator();
+					TWeakObjectPtr<UBlueprint> BPPtr = *Iter;
+					Iter.RemoveCurrent();
+					if (UBlueprint* BP = BPPtr.Get())
 					{
-						const bool bForceRegeneration = true;
-						FKismetEditorUtilities::GenerateBlueprintSkeleton(BP, bForceRegeneration);
-					}
-					else
-					{
-						// it's unsafe to GC in the middle of reinstancing because there may be other reinstancers still alive with references to 
-						// otherwise unreferenced classes:
-						const bool bSkipGC = true;
-						// Full compiles first recompile all skeleton classes, so they are up-to-date
-						const bool bSkeletonUpToDate = true;
-						FKismetEditorUtilities::CompileBlueprint(BP, false, bSkipGC, false, nullptr, bSkeletonUpToDate, true);
-						CompiledBlueprints.Add(BP);
-					}
+						if (IsReinstancingSkeleton())
+						{
+							const bool bForceRegeneration = true;
+							FKismetEditorUtilities::GenerateBlueprintSkeleton(BP, bForceRegeneration);
+						}
+						else
+						{
+							// it's unsafe to GC in the middle of reinstancing because there may be other reinstancers still alive with references to 
+							// otherwise unreferenced classes:
+							const bool bSkipGC = true;
+							// Full compiles first recompile all skeleton classes, so they are up-to-date
+							const bool bSkeletonUpToDate = true;
+							FKismetEditorUtilities::CompileBlueprint(BP, false, bSkipGC, false, nullptr, bSkeletonUpToDate, true);
+							CompiledBlueprints.Add(BP);
+						}
 
-					OrderedRecompiledDependents.Add(BP);
+						OrderedRecompiledDependents.Add(BP);
 
-					// if this BP compiled with an error, then I don't see any reason why we 
-					// should attempt to recompile its dependencies; if a subsequent recompile 
-					// would fix this up, then it'll get re-injected into the queue when that happens
-					if (BP->Status != EBlueprintStatus::BS_Error)
-					{
-						for (TWeakObjectPtr<UBlueprint>& DependentPtr : DependentBlueprintsToRecompile)
+						// if this BP compiled with an error, then I don't see any reason why we 
+						// should attempt to recompile its dependencies; if a subsequent recompile 
+						// would fix this up, then it'll get re-injected into the queue when that happens
+						if (BP->Status != EBlueprintStatus::BS_Error)
 						{
-							if (!DependentPtr.IsValid())
+							for (TWeakObjectPtr<UBlueprint>& DependentPtr : DependentBlueprintsToRecompile)
 							{
-								continue;
-							}
-							UBlueprint* NewDependent = DependentPtr.Get();
-
-							int32 DependentIndex = OrderedRecompiledDependents.FindLast(NewDependent);
-							if (DependentIndex != INDEX_NONE)
-							{
-								// even though we just pushed BP into the list and know that it
-								// exists as the last entry, we want to see if it was compiled 
-								// earlier (once before 'NewDependent'); so we use Find() to search 
-								// out the first entry
-								int32 RecompilingBpIndex = OrderedRecompiledDependents.Find(BP);
-								if (RecompilingBpIndex != INDEX_NONE && RecompilingBpIndex < DependentIndex)
+								if (!DependentPtr.IsValid())
 								{
-									// we've already recompiled this Blueprint once before (here in 
-									// this loop), already after its dependency has been compiled too;
-									// so, to avoid a potential infinite loop we cannot keep trying 
-									// to compile this
-									//
-									// NOTE: this may result in some a compiler error that would have 
-									//       been resolved in another subsequent compile (for example: 
-									//       B depends on A, A is compiled, A has an error, B compiles 
-									//       with an error as a result, C compiles and enqueues A as a 
-									//       dependent, A is recompiled without error now, B is not 
-									//       enqueued again because its already recompiled after A)
-									// 
-									// the true fix is to restructure the compiler so that these sort 
-									// of scenarios don't happen - until then, this is a fair trade 
-									// off... fallback to a byte code compile instead
-									DependentBlueprintsToByteRecompile.Add(DependentPtr);
 									continue;
 								}
+								UBlueprint* NewDependent = DependentPtr.Get();
+
+								int32 DependentIndex = OrderedRecompiledDependents.FindLast(NewDependent);
+								if (DependentIndex != INDEX_NONE)
+								{
+									// even though we just pushed BP into the list and know that it
+									// exists as the last entry, we want to see if it was compiled 
+									// earlier (once before 'NewDependent'); so we use Find() to search 
+									// out the first entry
+									int32 RecompilingBpIndex = OrderedRecompiledDependents.Find(BP);
+									if (RecompilingBpIndex != INDEX_NONE && RecompilingBpIndex < DependentIndex)
+									{
+										// we've already recompiled this Blueprint once before (here in 
+										// this loop), already after its dependency has been compiled too;
+										// so, to avoid a potential infinite loop we cannot keep trying 
+										// to compile this
+										//
+										// NOTE: this may result in some a compiler error that would have 
+										//       been resolved in another subsequent compile (for example: 
+										//       B depends on A, A is compiled, A has an error, B compiles 
+										//       with an error as a result, C compiles and enqueues A as a 
+										//       dependent, A is recompiled without error now, B is not 
+										//       enqueued again because its already recompiled after A)
+										// 
+										// the true fix is to restructure the compiler so that these sort 
+										// of scenarios don't happen - until then, this is a fair trade 
+										// off... fallback to a byte code compile instead
+										DependentBlueprintsToByteRecompile.Add(DependentPtr);
+										continue;
+									}
+								}
+								RecompilationQueue.Add(DependentPtr);
 							}
-							RecompilationQueue.Add(DependentPtr);
 						}
+						DependentBlueprintsToRecompile.Empty();
 					}
-					DependentBlueprintsToRecompile.Empty();
 				}
-			}
 
-			TArray<UBlueprint*> OrderedBytecodeRecompile;
+				TArray<UBlueprint*> OrderedBytecodeRecompile;
 
-			while (DependentBlueprintsToByteRecompile.Num())
-			{
-				auto Iter = DependentBlueprintsToByteRecompile.CreateIterator();
-				if (UBlueprint* BP = Iter->Get())
+				while (DependentBlueprintsToByteRecompile.Num())
 				{
-					OrderedBytecodeRecompile.Add(BP);
+					auto Iter = DependentBlueprintsToByteRecompile.CreateIterator();
+					if (UBlueprint* BP = Iter->Get())
+					{
+						OrderedBytecodeRecompile.Add(BP);
+					}
+					Iter.RemoveCurrent();
 				}
-				Iter.RemoveCurrent();
-			}
 
-			// Make sure we compile classes that are deeper in the class hierarchy later
-			// than ones that are higher:
-			OrderedBytecodeRecompile.Sort(
-				[](const UBlueprint& LHS, const UBlueprint& RHS)
+				// Make sure we compile classes that are deeper in the class hierarchy later
+				// than ones that are higher:
+				OrderedBytecodeRecompile.Sort(
+					[](const UBlueprint& LHS, const UBlueprint& RHS)
 				{
 					int32 LHS_Depth = 0;
 					int32 RHS_Depth = 0;
@@ -807,72 +812,73 @@ void FBlueprintCompileReinstancer::ReinstanceObjects(bool bForceAlwaysReinstance
 					// across editor sessions:
 					return LHS_Depth != RHS_Depth ? (LHS_Depth < RHS_Depth) : LHS.GetName() < RHS.GetName();
 				}
-			);
+				);
 
-			DependentBlueprintsToByteRecompile.Empty();
+				DependentBlueprintsToByteRecompile.Empty();
 
-			for (int I = 0; I != OrderedBytecodeRecompile.Num(); ++I)
-			{
-				UBlueprint* BP = OrderedBytecodeRecompile[I];
-				FKismetEditorUtilities::RecompileBlueprintBytecode(BP, nullptr, true);
-				ensure(0 == DependentBlueprintsToRecompile.Num());
-				CompiledBlueprints.Add(BP);
-			}
+				for (int I = 0; I != OrderedBytecodeRecompile.Num(); ++I)
+				{
+					UBlueprint* BP = OrderedBytecodeRecompile[I];
+					FKismetEditorUtilities::RecompileBlueprintBytecode(BP, nullptr, true);
+					ensure(0 == DependentBlueprintsToRecompile.Num());
+					CompiledBlueprints.Add(BP);
+				}
 
 
-			if (!IsReinstancingSkeleton())
-			{
-				TGuardValue<bool> ReinstancingGuard(GIsReinstancing, true);
+				if (!IsReinstancingSkeleton())
+				{
+					TGuardValue<bool> ReinstancingGuard(GIsReinstancing, true);
 
-				TArray<TSharedPtr<FReinstanceFinalizer>> Finalizers;
+					TArray<TSharedPtr<FReinstanceFinalizer>> Finalizers;
 
-				// All children were recompiled. It's safe to reinstance.
-				for (int32 Idx = 0; Idx < QueueToReinstance.Num(); ++Idx)
-				{
-					auto Finalizer = QueueToReinstance[Idx]->ReinstanceInner(bForceAlwaysReinstance);
-					if (Finalizer.IsValid())
+					// All children were recompiled. It's safe to reinstance.
+					for (int32 Idx = 0; Idx < QueueToReinstance.Num(); ++Idx)
 					{
-						Finalizers.Push(Finalizer);
+						TSharedPtr<FReinstanceFinalizer> Finalizer = QueueToReinstance[Idx]->ReinstanceInner(bForceAlwaysReinstance);
+						if (Finalizer.IsValid())
+						{
+							Finalizers.Push(Finalizer);
+						}
+						QueueToReinstance[Idx]->bHasReinstanced = true;
 					}
-					QueueToReinstance[Idx]->bHasReinstanced = true;
-				}
-				QueueToReinstance.Empty();
+					QueueToReinstance.Empty();
 
-				for (auto Finalizer : Finalizers)
-				{
-					if (Finalizer.IsValid())
+					for (TSharedPtr<FReinstanceFinalizer>& Finalizer : Finalizers)
 					{
-						Finalizer->Finalize();
+						if (Finalizer.IsValid())
+						{
+							Finalizer->Finalize();
+						}
 					}
-				}
 
-				for (auto CompiledBP : CompiledBlueprints)
-				{
-					CompiledBP->BroadcastCompiled();
-				}
+					for (TWeakObjectPtr<UBlueprint>& CompiledBP : CompiledBlueprints)
+					{
+						CompiledBP->BroadcastCompiled();
+					}
 
-				{
-					BP_SCOPED_COMPILER_EVENT_STAT(EKismetCompilerStats_RefreshDependentBlueprintsInReinstancer);
-					for (auto BPPtr : DependentBlueprintsToRefresh)
 					{
-						if (BPPtr.IsValid())
+						BP_SCOPED_COMPILER_EVENT_STAT(EKismetCompilerStats_RefreshDependentBlueprintsInReinstancer);
+						for (TWeakObjectPtr<UBlueprint>& BPPtr : DependentBlueprintsToRefresh)
 						{
-							BPPtr->BroadcastChanged();
+							if (BPPtr.IsValid())
+							{
+								BPPtr->BroadcastChanged();
+							}
 						}
+						DependentBlueprintsToRefresh.Empty();
 					}
-					DependentBlueprintsToRefresh.Empty();
-				}
 
-				if (GEditor)
+					if (GEditor)
+					{
+						GEditor->BroadcastBlueprintCompiled();
+					}
+				}
+				else
 				{
-					GEditor->BroadcastBlueprintCompiled();
+					QueueToReinstance.Empty();
+					DependentBlueprintsToRefresh.Empty();
 				}
 			}
-			else
-			{
-				QueueToReinstance.Empty();
-				DependentBlueprintsToRefresh.Empty();
-			}
 		}
 	}
 }
@@ -883,7 +889,7 @@ void FBlueprintCompileReinstancer::UpdateBytecodeReferences()
 {
 	BP_SCOPED_COMPILER_EVENT_STAT(EKismetReinstancerStats_UpdateBytecodeReferences);
 
-	if(ClassToReinstance != NULL)
+	if(ClassToReinstance != nullptr)
 	{
 		TMap<UObject*, UObject*> FieldMappings;
 		GenerateFieldMappings(FieldMappings);
@@ -934,7 +940,7 @@ void FBlueprintCompileReinstancer::UpdateBytecodeReferences()
 				UE_LOG(LogBlueprint, Log, TEXT("UpdateBytecodeReferences: %d references from %s was replaced in BP %s"), ReplaceInBPAr.GetCount(), *GetPathNameSafe(ClassToReinstance), *GetPathNameSafe(*DependentBP));
 			}
 
-			auto CompiledBlueprint = UBlueprint::GetBlueprintFromClass(ClassToReinstance);
+			UBlueprint* CompiledBlueprint = UBlueprint::GetBlueprintFromClass(ClassToReinstance);
 			if (bBPWasChanged && CompiledBlueprint && !CompiledBlueprint->bIsRegeneratingOnLoad)
 			{
 				DependentBlueprintsToRefresh.Add(*DependentBP);
@@ -1175,7 +1181,7 @@ void FActorReplacementHelper::Finalize(const TMap<UObject*, UObject*>& OldToNewI
 	
 	// Destroy actor and clear references.
 	NewActor->Modify();
-	if (GEditor->Layers.IsValid()) // ensure(NULL != GEditor->Layers) ?? While cooking the Layers is NULL.
+	if (GEditor->Layers.IsValid())
 	{
 		GEditor->Layers->InitializeNewActorLayers(NewActor);
 	}
@@ -1473,25 +1479,32 @@ void FBlueprintCompileReinstancer::ReplaceInstancesOfClass_Inner(TMap<UClass*, U
 				for (int32 OldObjIndex = 0; OldObjIndex < ObjectsToReplace.Num(); ++OldObjIndex)
 				{
 					UObject* OldObject = ObjectsToReplace[OldObjIndex];
-					AActor*  OldActor = Cast<AActor>(OldObject);
-					// Skip non-archetype instances, EXCEPT for component templates
-					if (!OldActor || OldObject->IsPendingKill())
+					if (!OldObject->IsPendingKill())
 					{
-						continue;
+						if (AActor* OldActor = Cast<AActor>(OldObject))
+						{
+							ActorAttachmentData.Add(OldObject, FActorAttachmentData(OldActor));
+						}
 					}
-
-					ActorAttachmentData.Add(OldObject, FActorAttachmentData(OldActor));
 				}
 
 				// Then fix 'real' (non archetype) instances of the class
 				for (int32 OldObjIndex = 0; OldObjIndex < ObjectsToReplace.Num(); ++OldObjIndex)
 				{
 					UObject* OldObject = ObjectsToReplace[OldObjIndex];
-					// Skip non-archetype instances, EXCEPT for component templates
+
+					if (OldObject->IsPendingKill())
+					{
+						continue;
+					}
+
+					AActor* OldActor = Cast<AActor>(OldObject);
+
+					// Skip archetype instances, EXCEPT for component templates and child actor templates
 					const bool bIsComponent = NewClass->IsChildOf(UActorComponent::StaticClass());
-					if ((!bIsComponent && OldObject->IsTemplate()) || OldObject->IsPendingKill())
+					const bool bIsChildActorTemplate = OldActor && OldActor->GetOuter()->IsA<UChildActorComponent>();
+					if (!bIsComponent && !bIsChildActorTemplate && OldObject->IsTemplate())
 					{
-						//OldObject->SetClass(NewClass);
 						continue;
 					}
 
@@ -1517,7 +1530,6 @@ void FBlueprintCompileReinstancer::ReplaceInstancesOfClass_Inner(TMap<UClass*, U
 					};
 					CacheOldObjectState(OldObject);
 
-					AActor*  OldActor = Cast<AActor>(OldObject);
 					UObject* NewUObject = nullptr;
 					// if the object to replace is an actor...
 					if (OldActor != nullptr && OldActor->GetLevel())
@@ -1615,7 +1627,7 @@ void FBlueprintCompileReinstancer::ReplaceInstancesOfClass_Inner(TMap<UClass*, U
 							GEditor->SelectActor(OldActor, /*bInSelected =*/false, /*bNotify =*/false);
 							bSelectionChanged = true;
 						}
-						if (GEditor->Layers.IsValid()) // ensure(NULL != GEditor->Layers) ?? While cooking the Layers is NULL.
+						if (GEditor->Layers.IsValid())
 						{
 							GEditor->Layers->DisassociateActorFromLayers(OldActor);
 						}
@@ -1668,11 +1680,11 @@ void FBlueprintCompileReinstancer::ReplaceInstancesOfClass_Inner(TMap<UClass*, U
 						EObjectFlags OldFlags = OldObject->GetFlags();
 
 						FName OldName(OldObject->GetFName());
-						OldObject->Rename(NULL, OldObject->GetOuter(), REN_DoNotDirty | REN_DontCreateRedirectors);
+						OldObject->Rename(nullptr, OldObject->GetOuter(), REN_DoNotDirty | REN_DontCreateRedirectors);
 						NewUObject = NewObject<UObject>(OldObject->GetOuter(), NewClass, OldName, RF_NoFlags, NewArchetype);
 						check(NewUObject != nullptr);
 
-						auto FlagMask = RF_Public | RF_ArchetypeObject | RF_Transactional | RF_Transient | RF_TextExportTransient | RF_InheritableComponentTemplate; //TODO: what about RF_RootSet and RF_Standalone ?
+						const EObjectFlags FlagMask = RF_Public | RF_ArchetypeObject | RF_Transactional | RF_Transient | RF_TextExportTransient | RF_InheritableComponentTemplate; //TODO: what about RF_RootSet and RF_Standalone ?
 						NewUObject->SetFlags(OldFlags & FlagMask);
 
 						InstancedPropertyUtils::FInstancedPropertyMap InstancedPropertyMap;
@@ -1845,7 +1857,7 @@ void FBlueprintCompileReinstancer::ReplaceInstancesOfClass_Inner(TMap<UClass*, U
 	if (GEditor)
 	{
 		// Refresh any editors for objects that we've updated components for
-		for (auto BlueprintAsset : PotentialEditorsForRefreshing)
+		for (UObject* BlueprintAsset : PotentialEditorsForRefreshing)
 		{
 			FBlueprintEditor* BlueprintEditor = static_cast<FBlueprintEditor*>(FAssetEditorManager::Get().FindEditorForAsset(BlueprintAsset, /*bFocusIfOpen =*/false));
 			if (BlueprintEditor)
@@ -1912,7 +1924,7 @@ void FBlueprintCompileReinstancer::VerifyReplacement()
 	// For each instance, track down references
 	if( SourceObjects.Num() > 0 )
 	{
-		TFindObjectReferencers<UObject> Referencers(SourceObjects, NULL, false);
+		TFindObjectReferencers<UObject> Referencers(SourceObjects, nullptr, false);
 		for (TFindObjectReferencers<UObject>::TIterator It(Referencers); It; ++It)
 		{
 			UObject* CurrentObject = It.Key();
@@ -1945,7 +1957,7 @@ void FBlueprintCompileReinstancer::ReparentChild(UClass* ChildClass)
 {
 	check(ChildClass && ClassToReinstance && DuplicatedClass && ChildClass->GetSuperClass());
 	bool bIsReallyAChild = ChildClass->GetSuperClass() == ClassToReinstance || ChildClass->GetSuperClass() == DuplicatedClass;
-	const auto SuperClassBP = Cast<UBlueprint>(ChildClass->GetSuperClass()->ClassGeneratedBy);
+	const UBlueprint* SuperClassBP = Cast<UBlueprint>(ChildClass->GetSuperClass()->ClassGeneratedBy);
 	if (SuperClassBP && !bIsReallyAChild)
 	{
 		bIsReallyAChild |= (SuperClassBP->SkeletonGeneratedClass == ClassToReinstance) || (SuperClassBP->SkeletonGeneratedClass == DuplicatedClass);
@@ -1991,7 +2003,7 @@ FRecreateUberGraphFrameScope::FRecreateUberGraphFrameScope(UClass* InClass, bool
 		const bool bIncludeDerivedClasses = true;
 		GetObjectsOfClass(RecompiledClass, Objects, bIncludeDerivedClasses, RF_NoFlags);
 
-		for (auto Obj : Objects)
+		for (UObject* Obj : Objects)
 		{
 			RecompiledClass->DestroyPersistentUberGraphFrame(Obj);
 		}
@@ -2001,7 +2013,7 @@ FRecreateUberGraphFrameScope::FRecreateUberGraphFrameScope(UClass* InClass, bool
 FRecreateUberGraphFrameScope::~FRecreateUberGraphFrameScope()
 {
 	BP_SCOPED_COMPILER_EVENT_STAT(EKismetCompilerStats_RecreateUberGraphPersistentFrame);
-	for (auto Obj : Objects)
+	for (UObject* Obj : Objects)
 	{
 		if (IsValid(Obj))
 		{
diff --git a/Engine/Source/Editor/UnrealEd/Private/PhysicsAssetUtils.cpp b/Engine/Source/Editor/UnrealEd/Private/PhysicsAssetUtils.cpp
--- a/Engine/Source/Editor/UnrealEd/Private/PhysicsAssetUtils.cpp
+++ b/Engine/Source/Editor/UnrealEd/Private/PhysicsAssetUtils.cpp
@@ -214,37 +214,30 @@ bool CreateFromSkeletalMeshInternal(UPhysicsAsset* PhysicsAsset, USkeletalMesh*
 	return PhysicsAsset->SkeletalBodySetups.Num() > 0;
 }
 
-bool CreateFromSkeletalMesh(UPhysicsAsset* PhysicsAsset, USkeletalMesh* SkelMesh, FPhysAssetCreateParams& Params, FText& OutErrorMessage)
+bool CreateFromSkeletalMesh(UPhysicsAsset* PhysicsAsset, USkeletalMesh* SkelMesh, FPhysAssetCreateParams& Params, FText& OutErrorMessage, bool bSetToMesh /*= true*/)
 {
 	PhysicsAsset->PreviewSkeletalMesh = SkelMesh;
 
 	bool bSuccess = CreateFromSkeletalMeshInternal(PhysicsAsset, SkelMesh, Params);
-	if (bSuccess)
-	{
-		// sets physics asset here now, so whoever creates 
-		// new physics asset from skeletalmesh will set properly here
-		SkelMesh->PhysicsAsset = PhysicsAsset;
-		SkelMesh->MarkPackageDirty();
-	}
-	else
+	if (!bSuccess)
 	{
 		// try lower minimum bone size 
 		Params.MinBoneSize = 1.f;
 
 		bSuccess = CreateFromSkeletalMeshInternal(PhysicsAsset, SkelMesh, Params);
-		if(bSuccess)
-		{
-			// sets physics asset here now, so whoever creates 
-			// new physics asset from skeletalmesh will set properly here
-			SkelMesh->PhysicsAsset = PhysicsAsset;
-			SkelMesh->MarkPackageDirty();
-		}
-		else
+
+		if(!bSuccess)
 		{
 			OutErrorMessage = FText::Format(NSLOCTEXT("CreatePhysicsAsset", "CreatePhysicsAssetLinkFailed", "The bone size is too small to create Physics Asset '{0}' from Skeletal Mesh '{1}'. You will have to create physics asset manually."), FText::FromString(PhysicsAsset->GetName()), FText::FromString(SkelMesh->GetName()));
 		}
 	}
 
+	if(bSuccess && bSetToMesh)
+	{
+		SkelMesh->PhysicsAsset = PhysicsAsset;
+		SkelMesh->MarkPackageDirty();
+	}
+
 	return bSuccess;
 }
 
diff --git a/Engine/Source/Editor/UnrealEd/Private/PhysicsManipulationMode.cpp b/Engine/Source/Editor/UnrealEd/Private/PhysicsManipulationMode.cpp
--- a/Engine/Source/Editor/UnrealEd/Private/PhysicsManipulationMode.cpp
+++ b/Engine/Source/Editor/UnrealEd/Private/PhysicsManipulationMode.cpp
@@ -84,39 +84,42 @@ bool FPhysicsManipulationEdMode::InputDelta( FEditorViewportClient* InViewportCl
 	}
 }
 
-bool FPhysicsManipulationEdMode::StartTracking( FEditorViewportClient* InViewportClient, FViewport* InViewport )
+bool FPhysicsManipulationEdMode::StartTracking(FEditorViewportClient* InViewportClient, FViewport* InViewport)
 {
 	//UE_LOG(LogEditorPhysMode, Warning, TEXT("Start Tracking"));
 
-	FVector GrabLocation(0,0,0);
-	UPrimitiveComponent* ComponentToGrab = NULL;
-	
-	USelection* Selection = GEditor->GetSelectedActors();
-
-	for (int32 i=0; i<Selection->Num(); ++i)
+	if (InViewportClient->GetCurrentWidgetAxis() != EAxisList::None)
 	{
-		AActor* SelectedActor = Cast<AActor>(Selection->GetSelectedObject(i));
+		FVector GrabLocation(0,0,0);
+		UPrimitiveComponent* ComponentToGrab = NULL;
+	
+		USelection* Selection = GEditor->GetSelectedActors();
 
-		if (SelectedActor != NULL)
+		for (int32 i=0; i<Selection->Num(); ++i)
 		{
-			UPrimitiveComponent* PC = Cast<UPrimitiveComponent>(SelectedActor->GetRootComponent());
+			AActor* SelectedActor = Cast<AActor>(Selection->GetSelectedObject(i));
 
-			if (PC != NULL && PC->BodyInstance.bSimulatePhysics)
+			if (SelectedActor != NULL)
 			{
-				ComponentToGrab = PC;
+				UPrimitiveComponent* PC = Cast<UPrimitiveComponent>(SelectedActor->GetRootComponent());
+
+				if (PC != NULL && PC->BodyInstance.bSimulatePhysics)
+				{
+					ComponentToGrab = PC;
 				
-				HandleTargetLocation = SelectedActor->GetActorLocation();
-				HandleTargetRotation = SelectedActor->GetActorRotation();
-				break;
-			}
+					HandleTargetLocation = SelectedActor->GetActorLocation();
+					HandleTargetRotation = SelectedActor->GetActorRotation();
+					break;
+				}
 
-			if (ComponentToGrab != NULL) { break; }
+				if (ComponentToGrab != NULL) { break; }
+			}
 		}
-	}
 
-	if (ComponentToGrab != NULL)
-	{
-		HandleComp->GrabComponentAtLocationWithRotation(ComponentToGrab, NAME_None, ComponentToGrab->GetOwner()->GetActorLocation(), ComponentToGrab->GetOwner()->GetActorRotation());
+		if (ComponentToGrab != NULL)
+		{
+			HandleComp->GrabComponentAtLocationWithRotation(ComponentToGrab, NAME_None, ComponentToGrab->GetOwner()->GetActorLocation(), ComponentToGrab->GetOwner()->GetActorRotation());
+		}
 	}
 
 	return FEdMode::StartTracking(InViewportClient, InViewport);
@@ -125,7 +128,6 @@ bool FPhysicsManipulationEdMode::StartTracking( FEditorViewportClient* InViewpor
 bool FPhysicsManipulationEdMode::EndTracking( FEditorViewportClient* InViewportClient, FViewport* InViewport )
 {
 	//UE_LOG(LogEditorPhysMode, Warning, TEXT("End Tracking"));
-
 	HandleComp->ReleaseComponent();
 
 	return FEdMode::EndTracking(InViewportClient, InViewport);
diff --git a/Engine/Source/Editor/UnrealEd/Private/PlayLevel.cpp b/Engine/Source/Editor/UnrealEd/Private/PlayLevel.cpp
--- a/Engine/Source/Editor/UnrealEd/Private/PlayLevel.cpp
+++ b/Engine/Source/Editor/UnrealEd/Private/PlayLevel.cpp
@@ -45,7 +45,7 @@
 #include "Components/AudioComponent.h"
 #include "Engine/Note.h"
 #include "UnrealEngine.h"
-#include "GameFramework/GameMode.h"
+#include "GameFramework/GameModeBase.h"
 #include "Engine/NavigationObjectBase.h"
 #include "GameFramework/PlayerController.h"
 #include "GameFramework/PlayerStart.h"
@@ -3470,52 +3470,10 @@ void UEditorEngine::ToggleBetweenPIEandSIE( bool bNewSession )
 				OnSwitchWorldsForPIE(true);
 
 				UWorld* World = GameViewport->GetWorld();
-				AGameMode* AuthGameMode = World->GetAuthGameMode();
+				AGameModeBase* AuthGameMode = World->GetAuthGameMode();
 				if (AuthGameMode && GameViewport->GetGameInstance())	// If there is no GameMode, we are probably the client and cannot RestartPlayer.
 				{
-					APlayerController* PC = GameViewport->GetGameInstance()->GetFirstLocalPlayerController();
-					if (PC != nullptr)
-				{
-					AuthGameMode->RemovePlayerControllerFromPlayerCount(PC);
-					PC->PlayerState->bOnlySpectator = false;
-					AuthGameMode->NumPlayers++;
-
-					bool bNeedsRestart = true;
-					if (PC->GetPawn() == NULL)
-					{
-						// Use the "auto-possess" pawn in the world, if there is one.
-						for (FConstPawnIterator Iterator = World->GetPawnIterator(); Iterator; ++Iterator)
-						{
-							APawn* Pawn = *Iterator;
-							if (Pawn && Pawn->AutoPossessPlayer == EAutoReceiveInput::Player0)
-							{
-								if (Pawn->Controller == nullptr)
-								{
-									PC->Possess(Pawn);
-									bNeedsRestart = false;
-								}
-								break;
-							}
-						}
-					}
-
-					if (bNeedsRestart)
-					{
-						AuthGameMode->RestartPlayer(PC);
-
-						if (PC->GetPawn())
-						{
-							// If there was no player start, then try to place the pawn where the camera was.						
-							if (PC->StartSpot == nullptr || Cast<AWorldSettings>(PC->StartSpot.Get()))
-							{
-								const FVector Location = EditorViewportClient.GetViewLocation();
-								const FRotator Rotation = EditorViewportClient.GetViewRotation();
-								PC->SetControlRotation(Rotation);
-								PC->GetPawn()->TeleportTo(Location, Rotation);
-							}
-						}
-					}
-				}
+					AuthGameMode->SpawnPlayerFromSimulate(EditorViewportClient.GetViewLocation(), EditorViewportClient.GetViewRotation());
 				}
 
 				OnSwitchWorldsForPIE(false);
@@ -3852,7 +3810,7 @@ UWorld* UEditorEngine::CreatePIEWorldFromEntry(FWorldContext &WorldContext, UWor
 	// Force default GameMode class so project specific code doesn't fire off. 
 	// We want this world to truly remain empty while we wait for connect!
 	check(LoadedWorld->GetWorldSettings());
-	LoadedWorld->GetWorldSettings()->DefaultGameMode = AGameMode::StaticClass();
+	LoadedWorld->GetWorldSettings()->DefaultGameMode = AGameModeBase::StaticClass();
 
 	PlayWorldMapName = UGameMapsSettings::GetGameDefaultMap();
 	return LoadedWorld;
diff --git a/Engine/Source/Editor/UnrealEd/Private/PreferenceStubs.cpp b/Engine/Source/Editor/UnrealEd/Private/PreferenceStubs.cpp
--- a/Engine/Source/Editor/UnrealEd/Private/PreferenceStubs.cpp
+++ b/Engine/Source/Editor/UnrealEd/Private/PreferenceStubs.cpp
@@ -49,6 +49,7 @@ UPersonaOptions::UPersonaOptions(const FObjectInitializer& ObjectInitializer)
 	BranchingPointTimingNodeColor = FLinearColor(0.5f, 1.0f, 1.0f);
 
 	bAutoAlignFloorToMesh = true;
+	bUseStandaloneAnimationEditors = true;
 }
 
 void UPersonaOptions::SetViewportBackgroundColor( const FLinearColor& InViewportBackgroundColor)
diff --git a/Engine/Source/Editor/UnrealEd/Private/SkeletalMeshEdit.cpp b/Engine/Source/Editor/UnrealEd/Private/SkeletalMeshEdit.cpp
--- a/Engine/Source/Editor/UnrealEd/Private/SkeletalMeshEdit.cpp
+++ b/Engine/Source/Editor/UnrealEd/Private/SkeletalMeshEdit.cpp
@@ -354,7 +354,6 @@ void UnFbx::FFbxImporter::FillAndVerifyBoneNames(USkeleton* Skeleton, TArray<Fbx
 	}
 
 	const FReferenceSkeleton& RefSkeleton = Skeleton->GetReferenceSkeleton();
-	const USkeleton::FBoneTreeType& BoneTree = Skeleton->GetBoneTree();
 
 	// make sure at least root bone matches
 	if ( OutRawBoneNames[0] != RefSkeleton.GetBoneName(0) )
@@ -563,10 +562,30 @@ UAnimSequence * UnFbx::FFbxImporter::ImportAnimations(USkeleton* Skeleton, UObje
 	return LastCreatedAnim;
 }
 
+int32 GetAnimationCurveRate(FbxAnimCurve* CurrentCurve)
+{
+	int32 KeyCount = CurrentCurve->KeyGetCount();
+
+	FbxTimeSpan TimeInterval(FBXSDK_TIME_INFINITE, FBXSDK_TIME_MINUS_INFINITE);
+	bool bValidTimeInterval = CurrentCurve->GetTimeInterval(TimeInterval);
+
+	if (KeyCount > 1 && bValidTimeInterval)
+	{
+		double KeyAnimLength = TimeInterval.GetDuration().GetSecondDouble();
+		if (KeyAnimLength != 0.0)
+		{
+			// 30 fps animation has 31 keys because it includes index 0 key for 0.0 second
+			return FPlatformMath::RoundToInt((KeyCount - 1) / KeyAnimLength);
+		}
+	}
+
+	return 0;
+}
+
 int32 UnFbx::FFbxImporter::GetMaxSampleRate(TArray<FbxNode*>& SortedLinks, TArray<FbxNode*>& NodeArray)
 {
 	int32 MaxStackResampleRate = 0;
-
+	const FBXImportOptions* ImportOption = GetImportOptions();
 	int32 AnimStackCount = Scene->GetSrcObjectCount<FbxAnimStack>();
 	for( int32 AnimStackIndex = 0; AnimStackIndex < AnimStackCount; AnimStackIndex++)
 	{
@@ -581,12 +600,12 @@ int32 UnFbx::FFbxImporter::GetMaxSampleRate(TArray<FbxNode*>& SortedLinks, TArra
 		double AnimStackStop = AnimStackTimeSpan.GetStop().GetSecondDouble();
 
 		FbxAnimLayer* AnimLayer = (FbxAnimLayer*)CurAnimStack->GetMember(0);
-
 		for(int32 LinkIndex = 0; LinkIndex < SortedLinks.Num(); ++LinkIndex)
 		{
 			FbxNode* CurrentLink = SortedLinks[LinkIndex];
 
-			FbxAnimCurve* Curves[6];
+			const int32 MaxElement = 9;
+			FbxAnimCurve* Curves[MaxElement];
 
 			Curves[0] = CurrentLink->LclTranslation.GetCurve(AnimLayer, FBXSDK_CURVENODE_COMPONENT_X, false);
 			Curves[1] = CurrentLink->LclTranslation.GetCurve(AnimLayer, FBXSDK_CURVENODE_COMPONENT_Y, false);
@@ -594,28 +613,52 @@ int32 UnFbx::FFbxImporter::GetMaxSampleRate(TArray<FbxNode*>& SortedLinks, TArra
 			Curves[3] = CurrentLink->LclRotation.GetCurve(AnimLayer, FBXSDK_CURVENODE_COMPONENT_X, false);
 			Curves[4] = CurrentLink->LclRotation.GetCurve(AnimLayer, FBXSDK_CURVENODE_COMPONENT_Y, false);
 			Curves[5] = CurrentLink->LclRotation.GetCurve(AnimLayer, FBXSDK_CURVENODE_COMPONENT_Z, false);
+			Curves[6] = CurrentLink->LclScaling.GetCurve(AnimLayer, FBXSDK_CURVENODE_COMPONENT_X, false);
+			Curves[7] = CurrentLink->LclScaling.GetCurve(AnimLayer, FBXSDK_CURVENODE_COMPONENT_Y, false);
+			Curves[8] = CurrentLink->LclScaling.GetCurve(AnimLayer, FBXSDK_CURVENODE_COMPONENT_Z, false);
 
-			for(int32 CurveIndex = 0; CurveIndex < 6; ++CurveIndex)
+			for(int32 CurveIndex = 0; CurveIndex < MaxElement; ++CurveIndex)
 			{
 				FbxAnimCurve* CurrentCurve = Curves[CurveIndex];
 				if(CurrentCurve)
 				{
-					int32 KeyCount = CurrentCurve->KeyGetCount();
-
-					FbxTimeSpan TimeInterval(FBXSDK_TIME_INFINITE,FBXSDK_TIME_MINUS_INFINITE);
-					bool bValidTimeInterval = CurrentCurve->GetTimeInterval(TimeInterval);
+					int32 NewRate = GetAnimationCurveRate(CurrentCurve);
+					if (NewRate > 0)
+					{
+						MaxStackResampleRate = FMath::Max(NewRate, MaxStackResampleRate);
+					}
+				}
+			}
+		}
 
- 					if(KeyCount > 1 && bValidTimeInterval)
+		// it doens't matter whether you choose to import morphtarget or not
+		// blendshape are always imported. Import morphtarget is only used for morphtarget for mesh
+		{
+			for (int32 NodeIndex = 0; NodeIndex < NodeArray.Num(); NodeIndex++)
+			{
+				// consider blendshape animation curve
+				FbxGeometry* Geometry = (FbxGeometry*)NodeArray[NodeIndex]->GetNodeAttribute();
+				if (Geometry)
+				{
+					int32 BlendShapeDeformerCount = Geometry->GetDeformerCount(FbxDeformer::eBlendShape);
+					for (int32 BlendShapeIndex = 0; BlendShapeIndex < BlendShapeDeformerCount; ++BlendShapeIndex)
 					{
-						double KeyAnimLength = TimeInterval.GetDuration().GetSecondDouble();
-						double KeyAnimStart = TimeInterval.GetStart().GetSecondDouble();
-						double KeyAnimStop = TimeInterval.GetStop().GetSecondDouble();
+						FbxBlendShape* BlendShape = (FbxBlendShape*)Geometry->GetDeformer(BlendShapeIndex, FbxDeformer::eBlendShape);
 
-						if(KeyAnimLength != 0.0)
+						int32 BlendShapeChannelCount = BlendShape->GetBlendShapeChannelCount();
+						for (int32 ChannelIndex = 0; ChannelIndex < BlendShapeChannelCount; ++ChannelIndex)
 						{
-							// 30 fps animation has 31 keys because it includes index 0 key for 0.0 second
-							int32 NewRate = FPlatformMath::RoundToInt((KeyCount-1) / KeyAnimLength);
-							MaxStackResampleRate = FMath::Max(NewRate, MaxStackResampleRate);
+							FbxBlendShapeChannel* Channel = BlendShape->GetBlendShapeChannel(ChannelIndex);
+
+							if (Channel)
+							{
+								FbxAnimCurve* CurrentCurve = Geometry->GetShapeChannel(BlendShapeIndex, ChannelIndex, AnimLayer);
+								int32 NewRate = GetAnimationCurveRate(CurrentCurve);
+								if (NewRate > 0)
+								{
+									MaxStackResampleRate = FMath::Max(NewRate, MaxStackResampleRate);
+								}
+							}
 						}
 					}
 				}
@@ -1040,12 +1083,16 @@ bool UnFbx::FFbxImporter::ImportAnimation(USkeleton* Skeleton, UAnimSequence * D
 		}
 	}
 
+	USkeleton* MySkeleton = DestSeq->GetSkeleton();
+	check(MySkeleton);
+
 	if (ImportOptions->bDeleteExistingMorphTargetCurves)
 	{
 		for (int32 CurveIdx=0; CurveIdx<DestSeq->RawCurveData.FloatCurves.Num(); ++CurveIdx)
 		{
 			auto& Curve = DestSeq->RawCurveData.FloatCurves[CurveIdx];
-			if (Curve.GetCurveTypeFlag(ACF_DriveMorphTarget))
+			const FCurveMetaData* MetaData = MySkeleton->GetCurveMetaData(Curve.Name);
+			if (MetaData && MetaData->Type.bMorphtarget)
 			{
 				DestSeq->RawCurveData.FloatCurves.RemoveAt(CurveIdx, 1, false);
 				--CurveIdx;
@@ -1098,7 +1145,13 @@ bool UnFbx::FFbxImporter::ImportAnimation(USkeleton* Skeleton, UAnimSequence * D
 								const FText StatusUpate = FText::Format(LOCTEXT("ImportingMorphTargetCurvesDetail", "Importing Morph Target Curves [{BlendShape}]"), Args);
 								GWarn->StatusUpdate(NodeIndex + 1, NodeArray.Num(), StatusUpate);
 								// now see if we have one already exists. If so, just overwrite that. if not, add new one. 
-								ImportCurveToAnimSequence(DestSeq, *ChannelName, Curve,  ACF_DriveMorphTarget | ACF_DriveAttribute, AnimTimeSpan, 0.01f /** for some reason blend shape values are coming as 100 scaled **/);
+
+								if (ImportCurveToAnimSequence(DestSeq, *ChannelName, Curve, 0, AnimTimeSpan, 0.01f /** for some reason blend shape values are coming as 100 scaled **/))
+								{
+									// this one doesn't reset Material curve to false, it just accumulate if true. 
+									MySkeleton->AccumulateCurveMetaData(*ChannelName, false, true);
+								}
+
 							}
 							else
 							{
@@ -1157,26 +1210,28 @@ bool UnFbx::FFbxImporter::ImportAnimation(USkeleton* Skeleton, UAnimSequence * D
 							GWarn->StatusUpdate(CurLinkIndex + 1, TotalLinks, StatusUpate);
 
 							int32 CurveFlags = ACF_DefaultCurve;
-							// first let them override material curve if required
-							if (ImportOptions->bSetMaterialDriveParameterOnCustomAttribute)
+							if (ImportCurveToAnimSequence(DestSeq, FinalCurveName, AnimCurve, CurveFlags, AnimTimeSpan))
 							{
-								CurveFlags |= ACF_DriveMaterial;
-							}
-							else
-							{
-								// if not material set by default, apply naming convention for material
-								for (const auto& Suffix : ImportOptions->MaterialCurveSuffixes)
+								// first let them override material curve if required
+								if (ImportOptions->bSetMaterialDriveParameterOnCustomAttribute)
 								{
-									int32 TotalSuffix = Suffix.Len();
-									if (CurveName.Right(TotalSuffix) == Suffix)
+									// now mark this curve as morphtarget
+									MySkeleton->AccumulateCurveMetaData(FName(*FinalCurveName), true, false);
+								}
+								else
+								{
+									// if not material set by default, apply naming convention for material
+									for (const auto& Suffix : ImportOptions->MaterialCurveSuffixes)
 									{
-										CurveFlags |= ACF_DriveMaterial;
-										break;
+										int32 TotalSuffix = Suffix.Len();
+										if (CurveName.Right(TotalSuffix) == Suffix)
+										{
+											MySkeleton->AccumulateCurveMetaData(FName(*FinalCurveName), true, false);
+											break;
+										}
 									}
 								}
 							}
-
-							ImportCurveToAnimSequence(DestSeq, FinalCurveName, AnimCurve, CurveFlags, AnimTimeSpan);
 						}
 						else
 						{
diff --git a/Engine/Source/Editor/UnrealEd/Private/Toolkits/AssetEditorToolkit.cpp b/Engine/Source/Editor/UnrealEd/Private/Toolkits/AssetEditorToolkit.cpp
--- a/Engine/Source/Editor/UnrealEd/Private/Toolkits/AssetEditorToolkit.cpp
+++ b/Engine/Source/Editor/UnrealEd/Private/Toolkits/AssetEditorToolkit.cpp
@@ -17,6 +17,7 @@
 #include "IIntroTutorials.h"
 #include "SuperSearchModule.h"
 #include "SDockTab.h"
+#include "AssetEditorModeManager.h"
 
 #define LOCTEXT_NAMESPACE "AssetEditorToolkit"
 
@@ -28,6 +29,7 @@ const FName FAssetEditorToolkit::ToolbarTabId( TEXT( "AssetEditorToolkit_Toolbar
 FAssetEditorToolkit::FAssetEditorToolkit()
 	: GCEditingObjects(*this)
 	, bCheckDirtyOnAssetSave(false)
+	, AssetEditorModeManager(nullptr)
 	, bIsToolbarFocusable(false)
 {
 	WorkspaceMenuCategory = FWorkspaceItem::NewGroup(LOCTEXT("WorkspaceMenu_BaseAssetEditor", "Asset Editor"));
@@ -617,6 +619,16 @@ const FSlateBrush* FAssetEditorToolkit::GetDefaultTabIcon() const
 	return IconBrush;
 }
 
+FAssetEditorModeManager* FAssetEditorToolkit::GetAssetEditorModeManager() const
+{
+	return AssetEditorModeManager;
+}
+
+void FAssetEditorToolkit::SetAssetEditorModeManager(FAssetEditorModeManager* InModeManager)
+{
+	AssetEditorModeManager = InModeManager;
+}
+
 void FAssetEditorToolkit::SwitchToStandaloneEditor_Execute( TWeakPtr< FAssetEditorToolkit > ThisToolkitWeakRef )
 {
 	// NOTE: We're being very careful here with pointer handling because we need to make sure the tookit's
diff --git a/Engine/Source/Editor/UnrealEd/Public/AssetEditorModeManager.h b/Engine/Source/Editor/UnrealEd/Public/AssetEditorModeManager.h
--- a/Engine/Source/Editor/UnrealEd/Public/AssetEditorModeManager.h
+++ b/Engine/Source/Editor/UnrealEd/Public/AssetEditorModeManager.h
@@ -7,7 +7,7 @@
 //////////////////////////////////////////////////////////////////////////
 // FAssetEditorModeManager
 
-class UNREALED_API FAssetEditorModeManager : public FEditorModeTools
+class UNREALED_API FAssetEditorModeManager : public FEditorModeTools, public TSharedFromThis<FAssetEditorModeManager>
 {
 public:
 	FAssetEditorModeManager();
@@ -21,6 +21,8 @@ public:
 	// End of FEditorModeTools interface
 
 	void SetPreviewScene(class FPreviewScene* NewPreviewScene);
+	FPreviewScene* GetPreviewScene() const;
+
 protected:
 	class USelection* ActorSet;
 	class USelection* ObjectSet;
diff --git a/Engine/Source/Editor/UnrealEd/Public/AssetNotifications.h b/Engine/Source/Editor/UnrealEd/Public/AssetNotifications.h
--- a/Engine/Source/Editor/UnrealEd/Public/AssetNotifications.h
+++ b/Engine/Source/Editor/UnrealEd/Public/AssetNotifications.h
@@ -9,7 +9,7 @@ class FAssetNotifications
 {
 public:
 	/** Inform about change in skeleton asset that requires saving*/
-	UNREALED_API static void SkeletonNeedsToBeSaved(class USkeleton* Skeleton);
+	UNREALED_API static void SkeletonNeedsToBeSaved(const class USkeleton* Skeleton);
 };
 
 #endif
diff --git a/Engine/Source/Editor/UnrealEd/Public/AssetViewerSettings.h b/Engine/Source/Editor/UnrealEd/Public/AssetViewerSettings.h
--- a/Engine/Source/Editor/UnrealEd/Public/AssetViewerSettings.h
+++ b/Engine/Source/Editor/UnrealEd/Public/AssetViewerSettings.h
@@ -28,7 +28,7 @@ struct FPreviewSceneProfile
 		LightingRigRotation = 0.0f;
 		RotationSpeed = 2.0f;
 		// Set up default cube map texture from editor/engine textures
-		EnvironmentCubeMap = LoadObject<UTextureCube>(NULL, TEXT("/Engine/EngineSky/EpicQuadPanorama_CC+EV1.EpicQuadPanorama_CC+EV1"));		
+		EnvironmentCubeMap = nullptr;
 		bPostProcessingEnabled = true;
 		DirectionalLightRotation = FRotator(-40.f, -67.5f, 0.f);
 	}
@@ -62,8 +62,12 @@ struct FPreviewSceneProfile
 	bool bShowFloor;
 
 	/** Sets environment cube map used for sky lighting and reflections */
-	UPROPERTY(EditAnywhere, Category = Environment)
-	UTextureCube* EnvironmentCubeMap;
+	UPROPERTY(EditAnywhere, transient, Category = Environment)
+	TAssetPtr<UTextureCube> EnvironmentCubeMap;
+
+	/** Storing path to environment cube to prevent it from getting cooked */
+	UPROPERTY(config)
+	FString EnvironmentCubeMapPath;
 
 	/** Manual set post processing settings */
 	UPROPERTY(EditAnywhere, config, Category = PostProcessing, AdvancedDisplay)
@@ -84,6 +88,23 @@ struct FPreviewSceneProfile
 	/** Rotation for directional light */
 	UPROPERTY(config)
 	FRotator DirectionalLightRotation;
+
+	/** Retrieve the environment map texture using the saved path */
+	void LoadEnvironmentMap()
+	{
+		if (EnvironmentCubeMap == nullptr)
+		{
+			// Load cube map from stored path
+			UObject* LoadedObject = nullptr;
+			LoadedObject = LoadObject<UObject>(nullptr, *EnvironmentCubeMapPath);
+			while (UObjectRedirector* Redirector = Cast<UObjectRedirector>(LoadedObject))
+			{
+				LoadedObject = Redirector->DestinationObject;
+			}
+
+			EnvironmentCubeMap = LoadedObject;
+		}		
+	}
 };
 
 /**
diff --git a/Engine/Source/Editor/UnrealEd/Public/Editor.h b/Engine/Source/Editor/UnrealEd/Public/Editor.h
--- a/Engine/Source/Editor/UnrealEd/Public/Editor.h
+++ b/Engine/Source/Editor/UnrealEd/Public/Editor.h
@@ -629,6 +629,9 @@ namespace EditorUtilities
 		TFunction<bool(UProperty&, UObject&)> PropertyFilter;
 	};
 
+	/** Helper function for CopyActorProperties(). Copies a single property form a source object to a target object. */
+	UNREALED_API void CopySingleProperty(const UObject* const InSourceObject, UObject* const InTargetObject, UProperty* const InProperty);
+
 	/**
 	 * Copies properties from one actor to another.  Designed for propagating changes made to PIE actors back to their EditorWorld
 	 * counterpart, or pushing spawned actor changes back to a Blueprint CDO object.  You can pass the 'PreviewOnly' option to
diff --git a/Engine/Source/Editor/UnrealEd/Public/EditorReimportHandler.h b/Engine/Source/Editor/UnrealEd/Public/EditorReimportHandler.h
--- a/Engine/Source/Editor/UnrealEd/Public/EditorReimportHandler.h
+++ b/Engine/Source/Editor/UnrealEd/Public/EditorReimportHandler.h
@@ -64,6 +64,19 @@ public:
 	 */
 	UNREALED_API virtual void ValidateAllSourceFileAndReimport(TArray<UObject*> &ToImportObjects);
 
+	/**
+	* Attempt to reimport multiple objects from its source by giving registered reimport
+	* handlers a chance to try to reimport the object
+	*
+	* @param	Objects	Objects to try reimporting
+	* @param	bAskForNewFileIfMissing If the file is missing, open a dialog to ask for a new one
+	* @param	bShowNotification True to show a notification when complete, false otherwise
+	* @param	PreferedReimportFile if not empty, will be use in case the original file is missing and bAskForNewFileIfMissing is set to false
+	*
+	* @return	true if the objects all imported successfully, for more granular success reporting use FReimportManager::Reimport
+	*/
+	UNREALED_API virtual bool ReimportMultiple( TArrayView<UObject*> Objects, bool bAskForNewFileIfMissing = false, bool bShowNotification = true, FString PreferedReimportFile = TEXT(""), FReimportHandler* SpecifiedReimportHandler = nullptr );
+
 	/**
 	 * Update the reimport paths for the specified object
 	 *
diff --git a/Engine/Source/Editor/UnrealEd/Public/GameModeInfoCustomizer.h b/Engine/Source/Editor/UnrealEd/Public/GameModeInfoCustomizer.h
--- a/Engine/Source/Editor/UnrealEd/Public/GameModeInfoCustomizer.h
+++ b/Engine/Source/Editor/UnrealEd/Public/GameModeInfoCustomizer.h
@@ -5,7 +5,7 @@
 #include "IDocumentation.h"
 #include "KismetEditorUtilities.h"
 #include "EditorClassUtils.h"
-#include "GameFramework/GameMode.h"
+#include "GameFramework/GameModeBase.h"
 #include "Engine/BlueprintGeneratedClass.h"
 
 #define LOCTEXT_NAMESPACE "FGameModeInfoCustomizer"
@@ -26,7 +26,7 @@ public:
 	/** Create widget for the name of a default class property */
 	TSharedRef<SWidget> CreateGameModePropertyLabelWidget(FName PropertyName)
 	{
-		UProperty* Prop = FindFieldChecked<UProperty>(AGameMode::StaticClass(), PropertyName);
+		UProperty* Prop = FindFieldChecked<UProperty>(AGameModeBase::StaticClass(), PropertyName);
 
 		FString DisplayName = Prop->GetDisplayNameText().ToString();
 		if (DisplayName.Len() == 0)
@@ -46,7 +46,7 @@ public:
 	void CustomizeGameModeDefaultClass(IDetailGroup& Group, FName DefaultClassPropertyName)
 	{
 		// Find the metaclass of this property
-		UClassProperty* ClassProp = FindFieldChecked<UClassProperty>(AGameMode::StaticClass(), DefaultClassPropertyName);
+		UClassProperty* ClassProp = FindFieldChecked<UClassProperty>(AGameModeBase::StaticClass(), DefaultClassPropertyName);
 
 		UClass* MetaClass = ClassProp->MetaClass;
 		const bool bAllowNone = !(ClassProp->PropertyFlags & CPF_NoClear);
@@ -129,7 +129,7 @@ public:
 				[
 					SNew(SClassPropertyEntryBox)
 					.AllowNone(bAllowNone)
-					.MetaClass(AGameMode::StaticClass())
+					.MetaClass(AGameModeBase::StaticClass())
 					.SelectedClass(this, &FGameModeInfoCustomizer::GetCurrentGameModeClass)
 					.OnSetClass(FOnSetClass::CreateSP(this, &FGameModeInfoCustomizer::SetCurrentGameModeClass))
 				]
@@ -156,12 +156,12 @@ public:
 		IDetailGroup& Group = CategoryBuilder.AddGroup(SelectedGameModeDetailsName, LOCTEXT("SelectedGameModeDetails", "Selected GameMode"));
 
 		// Then add rows to show key properties and let you edit them
-		CustomizeGameModeDefaultClass(Group, GET_MEMBER_NAME_CHECKED(AGameMode, DefaultPawnClass));
-		CustomizeGameModeDefaultClass(Group, GET_MEMBER_NAME_CHECKED(AGameMode, HUDClass));
-		CustomizeGameModeDefaultClass(Group, GET_MEMBER_NAME_CHECKED(AGameMode, PlayerControllerClass));
-		CustomizeGameModeDefaultClass(Group, GET_MEMBER_NAME_CHECKED(AGameMode, GameStateClass));
-		CustomizeGameModeDefaultClass(Group, GET_MEMBER_NAME_CHECKED(AGameMode, PlayerStateClass));
-		CustomizeGameModeDefaultClass(Group, GET_MEMBER_NAME_CHECKED(AGameMode, SpectatorClass));
+		CustomizeGameModeDefaultClass(Group, GET_MEMBER_NAME_CHECKED(AGameModeBase, DefaultPawnClass));
+		CustomizeGameModeDefaultClass(Group, GET_MEMBER_NAME_CHECKED(AGameModeBase, HUDClass));
+		CustomizeGameModeDefaultClass(Group, GET_MEMBER_NAME_CHECKED(AGameModeBase, PlayerControllerClass));
+		CustomizeGameModeDefaultClass(Group, GET_MEMBER_NAME_CHECKED(AGameModeBase, GameStateClass));
+		CustomizeGameModeDefaultClass(Group, GET_MEMBER_NAME_CHECKED(AGameModeBase, PlayerStateClass));
+		CustomizeGameModeDefaultClass(Group, GET_MEMBER_NAME_CHECKED(AGameModeBase, SpectatorClass));
 	}
 
 	/** Get the currently set GameMode class */
@@ -192,12 +192,12 @@ public:
 	}
 
 	/** Get the CDO from the currently set GameMode class */
-	AGameMode* GetCurrentGameModeCDO() const
+	AGameModeBase* GetCurrentGameModeCDO() const
 	{
 		UClass* GameModeClass = const_cast<UClass*>( GetCurrentGameModeClass() );
 		if (GameModeClass != NULL)
 		{
-			return GameModeClass->GetDefaultObject<AGameMode>();
+			return GameModeClass->GetDefaultObject<AGameModeBase>();
 		}
 		else
 		{
@@ -245,7 +245,7 @@ public:
 
 	void OnMakeNewDefaultClassClicked(FName ClassPropertyName)
 	{
-		UClassProperty* ClassProp = FindFieldChecked<UClassProperty>(AGameMode::StaticClass(), ClassPropertyName);
+		UClassProperty* ClassProp = FindFieldChecked<UClassProperty>(AGameModeBase::StaticClass(), ClassPropertyName);
 
 		UBlueprint* Blueprint = FKismetEditorUtilities::CreateBlueprintFromClass(LOCTEXT("CreateNewBlueprint", "Create New Blueprint"), ClassProp->MetaClass, FString::Printf(TEXT("New%s"),*ClassProp->MetaClass->GetName()));
 
@@ -261,7 +261,7 @@ public:
 	{
 		FEditorDelegates::LoadSelectedAssetsIfNeeded.Broadcast();
 
-		UClassProperty* ClassProp = FindFieldChecked<UClassProperty>(AGameMode::StaticClass(), ClassPropertyName);
+		UClassProperty* ClassProp = FindFieldChecked<UClassProperty>(AGameModeBase::StaticClass(), ClassPropertyName);
 		const UClass* SelectedClass = GEditor->GetFirstSelectedClass(ClassProp->MetaClass);
 		if (SelectedClass)
 		{
@@ -302,7 +302,7 @@ public:
 	{
 		FEditorDelegates::LoadSelectedAssetsIfNeeded.Broadcast();
 
-		const UClass* SelectedClass = GEditor->GetFirstSelectedClass(AGameMode::StaticClass());
+		const UClass* SelectedClass = GEditor->GetFirstSelectedClass(AGameModeBase::StaticClass());
 		if (SelectedClass)
 		{
 			DefaultGameModeClassHandle->SetValueFromFormattedString(SelectedClass->GetPathName());
@@ -317,7 +317,7 @@ public:
 	void OnClickNewGameMode(IDetailLayoutBuilder* DetailLayout)
 	{
 		// Create a new GameMode BP
-		UBlueprint* Blueprint = FKismetEditorUtilities::CreateBlueprintFromClass(LOCTEXT("CreateNewGameMode", "Create New GameMode"), AGameMode::StaticClass(), TEXT("NewGameMode"));
+		UBlueprint* Blueprint = FKismetEditorUtilities::CreateBlueprintFromClass(LOCTEXT("CreateNewGameMode", "Create New GameMode"), AGameModeBase::StaticClass(), TEXT("NewGameMode"));
 		// if that worked, assign it
 		if(Blueprint != NULL && Blueprint->GeneratedClass)
 		{
diff --git a/Engine/Source/Editor/UnrealEd/Public/PhysicsAssetUtils.h b/Engine/Source/Editor/UnrealEd/Public/PhysicsAssetUtils.h
--- a/Engine/Source/Editor/UnrealEd/Public/PhysicsAssetUtils.h
+++ b/Engine/Source/Editor/UnrealEd/Public/PhysicsAssetUtils.h
@@ -48,8 +48,10 @@ namespace FPhysicsAssetUtils
 	 * @param	PhysicsAsset		The PhysicsAsset instance to setup
 	 * @param	SkelMesh			The Skeletal Mesh to create the physics asset from
 	 * @param	Params				Additional creation parameters
+	 * @param	OutErrorMessage		Additional error information
+	 * @param	bSetToMesh			Whether or not to apply the physics asset to SkelMesh immediately
 	 */
-	UNREALED_API bool CreateFromSkeletalMesh(UPhysicsAsset* PhysicsAsset, USkeletalMesh* SkelMesh, FPhysAssetCreateParams& Params, FText& OutErrorMessage);
+	UNREALED_API bool CreateFromSkeletalMesh(UPhysicsAsset* PhysicsAsset, USkeletalMesh* SkelMesh, FPhysAssetCreateParams& Params, FText& OutErrorMessage, bool bSetToMesh = true);
 
 	/** Replaces any collision already in the BodySetup with an auto-generated one using the parameters provided. 
 	 *
diff --git a/Engine/Source/Editor/UnrealEd/Public/SAdvancedPreviewDetailsTab.cpp b/Engine/Source/Editor/UnrealEd/Public/SAdvancedPreviewDetailsTab.cpp
--- a/Engine/Source/Editor/UnrealEd/Public/SAdvancedPreviewDetailsTab.cpp
+++ b/Engine/Source/Editor/UnrealEd/Public/SAdvancedPreviewDetailsTab.cpp
@@ -40,9 +40,9 @@ SAdvancedPreviewDetailsTab::~SAdvancedPreviewDetailsTab()
 	}
 }
 
-void SAdvancedPreviewDetailsTab::Construct(const FArguments& InArgs)
+void SAdvancedPreviewDetailsTab::Construct(const FArguments& InArgs, const TSharedRef<FAdvancedPreviewScene>& InPreviewScene)
 {
-	PreviewScene = InArgs._PreviewScenePtr;	
+	PreviewScenePtr = InPreviewScene;
 	DefaultSettings = UAssetViewerSettings::Get();
 	ProfileIndex = PerProjectSettings->AssetViewerProfileIndex;
 		
@@ -127,7 +127,7 @@ void SAdvancedPreviewDetailsTab::ComboBoxSelectionChanged(TSharedPtr<FString> Ne
 		ProfileIndex = NewSelectionIndex;
 		PerProjectSettings->AssetViewerProfileIndex = ProfileIndex;
 		UpdateSettingsView();	
-		PreviewScene->SetProfileIndex(ProfileIndex);
+		PreviewScenePtr.Pin()->SetProfileIndex(ProfileIndex);
 	}
 }
 
@@ -154,8 +154,39 @@ FReply SAdvancedPreviewDetailsTab::AddProfileButtonClick()
 	// Add new profile to settings instance
 	DefaultSettings->Profiles.AddDefaulted();
 	FPreviewSceneProfile& NewProfile = DefaultSettings->Profiles.Last();
-	NewProfile.ProfileName = FString::Printf(TEXT("Profile_%i"), DefaultSettings->Profiles.Num());
+	
+	// Try to create a valid profile name when one is added
+	bool bFoundValidName = false;
+	int ProfileAppendNum = DefaultSettings->Profiles.Num();
+	FString NewProfileName;
+	while (!bFoundValidName)
+	{
+		NewProfileName = FString::Printf(TEXT("Profile_%i"), ProfileAppendNum);
+
+		bool bValidName = true;
+		for (const FPreviewSceneProfile& Profile : DefaultSettings->Profiles)
+		{
+			if (Profile.ProfileName == NewProfileName)
+			{
+				bValidName = false;				
+				break;
+			}
+		}
+
+		if (!bValidName)
+		{
+			++ProfileAppendNum;
+		}
+
+		bFoundValidName = bValidName;
+	}
+
+	NewProfile.ProfileName = NewProfileName;
 	DefaultSettings->PostEditChange();
+
+	// Change selection to new profile so the user directly sees the profile that was added
+	Refresh();
+	ProfileComboBox->SetSelectedItem(ProfileNames.Last());
 	
 	return FReply::Handled();
 }
@@ -183,7 +214,7 @@ void SAdvancedPreviewDetailsTab::OnAssetViewerSettingsRefresh(const FName& InPro
 		UpdateProfileNames();
 	}
 
-	PreviewScene->UpdateScene(DefaultSettings->Profiles[ProfileIndex], bUpdateSkyLight, bUpdateEnvironment, bUpdatePostProcessing, bUpdateDirectionalLight);
+	PreviewScenePtr.Pin()->UpdateScene(DefaultSettings->Profiles[ProfileIndex], bUpdateSkyLight, bUpdateEnvironment, bUpdatePostProcessing, bUpdateDirectionalLight);
 }
 
 void SAdvancedPreviewDetailsTab::CreateSettingsView()
@@ -211,7 +242,7 @@ void SAdvancedPreviewDetailsTab::Refresh()
 {	
 	ProfileIndex = PerProjectSettings->AssetViewerProfileIndex;	
 	UpdateProfileNames();
-	PreviewScene->SetProfileIndex(ProfileIndex);
+	PreviewScenePtr.Pin()->SetProfileIndex(ProfileIndex);
 	UpdateSettingsView();
 }
 
diff --git a/Engine/Source/Editor/UnrealEd/Public/SAdvancedPreviewDetailsTab.h b/Engine/Source/Editor/UnrealEd/Public/SAdvancedPreviewDetailsTab.h
--- a/Engine/Source/Editor/UnrealEd/Public/SAdvancedPreviewDetailsTab.h
+++ b/Engine/Source/Editor/UnrealEd/Public/SAdvancedPreviewDetailsTab.h
@@ -11,12 +11,8 @@ class UEditorPerProjectUserSettings;
 class UNREALED_API SAdvancedPreviewDetailsTab : public SCompoundWidget
 {
 public:
-	SLATE_BEGIN_ARGS(SAdvancedPreviewDetailsTab) : 
-		_PreviewScenePtr()
-	{
-	}
-		/** The Static Mesh Editor this tool is associated with. */
-		SLATE_ARGUMENT(FAdvancedPreviewScene*, PreviewScenePtr)
+	SLATE_BEGIN_ARGS(SAdvancedPreviewDetailsTab)
+	{}
 
 	SLATE_END_ARGS()
 
@@ -26,7 +22,7 @@ public:
 	~SAdvancedPreviewDetailsTab();
 
 	/** SWidget functions */
-	void Construct(const FArguments& InArgs);
+	void Construct(const FArguments& InArgs, const TSharedRef<FAdvancedPreviewScene>& InPreviewScene);
 
 	void Refresh();
 
@@ -43,7 +39,7 @@ protected:
 	/** Property viewing widget */
 	TSharedPtr<IDetailsView> SettingsView;
 	TSharedPtr<STextComboBox> ProfileComboBox;
-	FAdvancedPreviewScene* PreviewScene;
+	TWeakPtr<FAdvancedPreviewScene> PreviewScenePtr;
 	UAssetViewerSettings* DefaultSettings;
 
 	TArray<TSharedPtr<FString>> ProfileNames;
diff --git a/Engine/Source/Editor/UnrealEd/Public/Toolkits/AssetEditorToolkit.h b/Engine/Source/Editor/UnrealEd/Public/Toolkits/AssetEditorToolkit.h
--- a/Engine/Source/Editor/UnrealEd/Public/Toolkits/AssetEditorToolkit.h
+++ b/Engine/Source/Editor/UnrealEd/Public/Toolkits/AssetEditorToolkit.h
@@ -182,6 +182,12 @@ public:
 	 */
 	static FText GetToolTipTextForObject(const UObject* InObject);
 
+	/** Get the asset editor mode manager we are using */
+	class FAssetEditorModeManager* GetAssetEditorModeManager() const;
+
+	/** Set the asset editor mode manager we are using */
+	void SetAssetEditorModeManager(class FAssetEditorModeManager* InModeManager);
+
 protected:
 
 	/**	Returns the single object currently being edited. Asserts if currently editing no object or multiple objects */
@@ -304,6 +310,9 @@ protected:
 	/** Whether only dirty assets should be prompted about on save - otherwise all edited assets will be prompted to the user for save/check-out */
 	bool bCheckDirtyOnAssetSave;
 
+	/** The asset editor mode manager */
+	FAssetEditorModeManager* AssetEditorModeManager;
+
 private:
 	/** The toolkit standalone host; may be nullptr */
 	TWeakPtr<SStandaloneAssetEditorToolkitHost> StandaloneHost;
diff --git a/Engine/Source/Editor/UnrealEd/UnrealEd.Build.cs b/Engine/Source/Editor/UnrealEd/UnrealEd.Build.cs
--- a/Engine/Source/Editor/UnrealEd/UnrealEd.Build.cs
+++ b/Engine/Source/Editor/UnrealEd/UnrealEd.Build.cs
@@ -61,7 +61,8 @@ public class UnrealEd : ModuleRules
                 "PortalRpc",
                 "PortalServices",
                 "BlueprintNativeCodeGen",
-			}
+                "Persona",
+            }
 		);
 
 		PublicDependencyModuleNames.AddRange(
@@ -253,6 +254,7 @@ public class UnrealEd : ModuleRules
 			(Target.Platform == UnrealTargetPlatform.Win32))
 		{
 			PublicDependencyModuleNames.Add("XAudio2");
+			PublicDependencyModuleNames.Add("AudioMixerXAudio2");
 			PublicDependencyModuleNames.Add("UnrealAudioXAudio2");
 
 			AddEngineThirdPartyPrivateStaticDependencies(Target, 
diff --git a/Engine/Source/Editor/WorldBrowser/Private/Tiles/WorldTileCollectionModel.cpp b/Engine/Source/Editor/WorldBrowser/Private/Tiles/WorldTileCollectionModel.cpp
--- a/Engine/Source/Editor/WorldBrowser/Private/Tiles/WorldTileCollectionModel.cpp
+++ b/Engine/Source/Editor/WorldBrowser/Private/Tiles/WorldTileCollectionModel.cpp
@@ -2049,11 +2049,11 @@ bool FWorldTileCollectionModel::GenerateLODLevels(FLevelModelList InLevelList, i
 			}
 								
 			// This is texture resolution for a landscape mesh, probably needs to be calculated using landscape size
-			LandscapeFlattenMaterial.DiffuseSize = SimplificationDetails.LandscapeMaterialSettings.TextureSize;
-			LandscapeFlattenMaterial.NormalSize = SimplificationDetails.LandscapeMaterialSettings.bNormalMap ? SimplificationDetails.LandscapeMaterialSettings.TextureSize : FIntPoint::ZeroValue;
-			LandscapeFlattenMaterial.MetallicSize = SimplificationDetails.LandscapeMaterialSettings.bMetallicMap ? SimplificationDetails.LandscapeMaterialSettings.TextureSize : FIntPoint::ZeroValue;
-			LandscapeFlattenMaterial.RoughnessSize = SimplificationDetails.LandscapeMaterialSettings.bRoughnessMap ? SimplificationDetails.LandscapeMaterialSettings.TextureSize : FIntPoint::ZeroValue;
-			LandscapeFlattenMaterial.SpecularSize = SimplificationDetails.LandscapeMaterialSettings.bSpecularMap ? SimplificationDetails.LandscapeMaterialSettings.TextureSize : FIntPoint::ZeroValue;
+			LandscapeFlattenMaterial.SetPropertySize(EFlattenMaterialProperties::Diffuse, SimplificationDetails.LandscapeMaterialSettings.TextureSize);
+			LandscapeFlattenMaterial.SetPropertySize(EFlattenMaterialProperties::Normal, SimplificationDetails.LandscapeMaterialSettings.bNormalMap ? SimplificationDetails.LandscapeMaterialSettings.TextureSize : FIntPoint::ZeroValue);
+			LandscapeFlattenMaterial.SetPropertySize(EFlattenMaterialProperties::Metallic, SimplificationDetails.LandscapeMaterialSettings.bMetallicMap ? SimplificationDetails.LandscapeMaterialSettings.TextureSize : FIntPoint::ZeroValue);  
+			LandscapeFlattenMaterial.SetPropertySize(EFlattenMaterialProperties::Roughness, SimplificationDetails.LandscapeMaterialSettings.bRoughnessMap ? SimplificationDetails.LandscapeMaterialSettings.TextureSize : FIntPoint::ZeroValue);  
+			LandscapeFlattenMaterial.SetPropertySize(EFlattenMaterialProperties::Specular, SimplificationDetails.LandscapeMaterialSettings.bSpecularMap ? SimplificationDetails.LandscapeMaterialSettings.TextureSize : FIntPoint::ZeroValue);
 			
 			FMaterialUtilities::ExportLandscapeMaterial(Landscape, PrimitivesToHide, LandscapeFlattenMaterial);
 					
diff --git a/Engine/Source/Runtime/AIModule/Classes/AIController.h b/Engine/Source/Runtime/AIModule/Classes/AIController.h
--- a/Engine/Source/Runtime/AIModule/Classes/AIController.h
+++ b/Engine/Source/Runtime/AIModule/Classes/AIController.h
@@ -271,7 +271,7 @@ public:
 	bool UseBlackboard(UBlackboardData* BlackboardAsset, UBlackboardComponent*& BlackboardComponent);
 
 	/** does this AIController allow given UBlackboardComponent sync data with it */
-	virtual bool ShouldSyncBlackboardWith(const UBlackboardComponent& OtherBlackboardComponent) const { return true; }
+	virtual bool ShouldSyncBlackboardWith(const UBlackboardComponent& OtherBlackboardComponent) const;
 
 	UFUNCTION(BlueprintCallable, Category = "AI|Tasks")
 	void ClaimTaskResource(TSubclassOf<UGameplayTaskResource> ResourceClass);
diff --git a/Engine/Source/Runtime/AIModule/Classes/BehaviorTree/BlackboardData.h b/Engine/Source/Runtime/AIModule/Classes/BehaviorTree/BlackboardData.h
--- a/Engine/Source/Runtime/AIModule/Classes/BehaviorTree/BlackboardData.h
+++ b/Engine/Source/Runtime/AIModule/Classes/BehaviorTree/BlackboardData.h
@@ -29,6 +29,10 @@ struct FBlackboardEntry
 	UPROPERTY(EditAnywhere, Category = Blackboard)
 	uint32 bInstanceSynced : 1;
 
+	FBlackboardEntry()
+		: KeyType(nullptr), bInstanceSynced(0)
+	{}
+
 	bool operator==(const FBlackboardEntry& Other) const;
 };
 
@@ -84,8 +88,8 @@ public:
 #if WITH_EDITOR
 	virtual void PreEditChange(class FEditPropertyChain& PropertyAboutToChange) override;
 	virtual void PostEditChangeProperty(struct FPropertyChangedEvent& PropertyChangedEvent) override;
-	void PropagateKeyChangesToDerivedBlackboardAssets();
 #endif
+	void PropagateKeyChangesToDerivedBlackboardAssets();
 
 	/** @return true if blackboard keys are not conflicting with parent key chain */
 	bool IsValid() const;
@@ -109,12 +113,14 @@ public:
 
 			Keys.Add(Entry);
 			MarkPackageDirty();
+			PropagateKeyChangesToDerivedBlackboardAssets();
 		}
 		else if (KeyID != FBlackboard::InvalidKey && Parent != NULL)
 		{
 			const FBlackboard::FKey KeyIndex = KeyID - FirstKeyID;
 			Keys.RemoveAt(KeyIndex);
 			MarkPackageDirty();
+			PropagateKeyChangesToDerivedBlackboardAssets();
 		}
 
 		return CreatedKeyType;
@@ -135,6 +141,16 @@ public:
 	/** fix entries with deprecated key types */
 	void UpdateDeprecatedKeys();
 
+	/** returns true if OtherAsset is somewhere up the parent chain of this asset. Node that it will return false if *this == OtherAsset */
+	bool IsChildOf(const UBlackboardData& OtherAsset) const;
+
+	/** returns true if OtherAsset is equal to *this, or is it's parent, or *this is OtherAsset's parent */
+	bool IsRelatedTo(const UBlackboardData& OtherAsset) const
+	{
+		return this == &OtherAsset || IsChildOf(OtherAsset) || OtherAsset.IsChildOf(*this)
+			|| (Parent && OtherAsset.Parent && Parent->IsRelatedTo(*OtherAsset.Parent));
+	}
+
 protected:
 
 	enum EKeyLookupMode
@@ -150,5 +166,6 @@ protected:
 	FBlackboard::FKey InternalGetKeyID(const FName& KeyName, EKeyLookupMode LookupMode) const;
 
 	/** check if parent chain contains given blackboard data */
+	DEPRECATED(4.14, "This function is deprecated, please use IsChildOf instead.")
 	bool HasParent(const UBlackboardData* TestParent) const;
 };
diff --git a/Engine/Source/Runtime/AIModule/Classes/EnvironmentQuery/EnvQueryTypes.h b/Engine/Source/Runtime/AIModule/Classes/EnvironmentQuery/EnvQueryTypes.h
--- a/Engine/Source/Runtime/AIModule/Classes/EnvironmentQuery/EnvQueryTypes.h
+++ b/Engine/Source/Runtime/AIModule/Classes/EnvironmentQuery/EnvQueryTypes.h
@@ -852,12 +852,12 @@ public:
 	/** raw data operations */
 	void ReserveItemData(int32 NumAdditionalItems);
 
-	template<typename TypeItem, typename TypeValue>
-	void AddItemData(TypeValue ItemValue)
+	template<typename TypeItem>
+	void AddItemData(typename TypeItem::FValueType ItemValue)
 	{
 		DEC_MEMORY_STAT_BY(STAT_AI_EQS_InstanceMemory, RawData.GetAllocatedSize() + Items.GetAllocatedSize());
 
-		check(GetDefault<TypeItem>()->GetValueSize() == sizeof(TypeValue));
+		check(GetDefault<TypeItem>()->GetValueSize() == sizeof(typename TypeItem::FValueType));
 		check(GetDefault<TypeItem>()->GetValueSize() == ValueSize);
 		const int32 DataOffset = RawData.AddUninitialized(ValueSize);
 		TypeItem::SetValue(RawData.GetData() + DataOffset, ItemValue);
@@ -867,19 +867,19 @@ public:
 	}
 
 	/** AddItemData specialization for arrays if values */
-	template<typename TypeItem, typename TypeValue>
-	void AddItemData(TArray<TypeValue>& ItemCollection)
+	template<typename TypeItem>
+	void AddItemData(TArray<typename TypeItem::FValueType>& ItemCollection)
 	{
 		if (ItemCollection.Num() > 0)
 		{
 			DEC_MEMORY_STAT_BY(STAT_AI_EQS_InstanceMemory, RawData.GetAllocatedSize() + Items.GetAllocatedSize());
 
-			check(GetDefault<TypeItem>()->GetValueSize() == sizeof(TypeValue));
+			check(GetDefault<TypeItem>()->GetValueSize() == sizeof(typename TypeItem::FValueType));
 			check(GetDefault<TypeItem>()->GetValueSize() == ValueSize);
 			int32 DataOffset = RawData.AddUninitialized(ValueSize * ItemCollection.Num());
 			Items.Reserve(Items.Num() + ItemCollection.Num());
 
-			for (TypeValue& Item : ItemCollection)
+			for (typename TypeItem::FValueType& Item : ItemCollection)
 			{
 				TypeItem::SetValue(RawData.GetData() + DataOffset, Item);
 				Items.Add(FEnvQueryItem(DataOffset));
@@ -890,6 +890,21 @@ public:
 		}
 	}
 
+	template<typename TypeItem, typename TypeValue>
+	void AddItemData(TypeValue ItemValue)
+	{
+		AddItemData<TypeItem>((typename TypeItem::FValueType)(ItemValue));
+	}
+
+	template<typename TypeItem, typename TypeValue>
+	void AddItemData(TArray<TypeValue>& ItemCollection)
+	{
+		for (auto& Item : ItemCollection)
+		{
+			AddItemData<TypeItem>((typename TypeItem::FValueType)(Item));
+		}
+	}
+
 protected:
 
 	/** prepare item data after generator has finished */
diff --git a/Engine/Source/Runtime/AIModule/Classes/EnvironmentQuery/Items/EnvQueryItemType.h b/Engine/Source/Runtime/AIModule/Classes/EnvironmentQuery/Items/EnvQueryItemType.h
--- a/Engine/Source/Runtime/AIModule/Classes/EnvironmentQuery/Items/EnvQueryItemType.h
+++ b/Engine/Source/Runtime/AIModule/Classes/EnvironmentQuery/Items/EnvQueryItemType.h
@@ -9,7 +9,15 @@ struct FBlackboardKeySelector;
 UCLASS(Abstract)
 class AIMODULE_API UEnvQueryItemType : public UObject
 {
-	GENERATED_UCLASS_BODY()
+	GENERATED_BODY()
+
+public:
+	/** every EQS item type needs to speficy data type it's using. 
+	 *	Default is void which should trigger a compilation error if it's 
+	  *	not set in a defived class*/
+	typedef void FValueType;
+
+	UEnvQueryItemType(const FObjectInitializer& ObjectInitializer = FObjectInitializer::Get());
 
 	/** get ValueSize */
 	FORCEINLINE uint16 GetValueSize() const { return ValueSize; }
diff --git a/Engine/Source/Runtime/AIModule/Classes/EnvironmentQuery/Items/EnvQueryItemType_Actor.h b/Engine/Source/Runtime/AIModule/Classes/EnvironmentQuery/Items/EnvQueryItemType_Actor.h
--- a/Engine/Source/Runtime/AIModule/Classes/EnvironmentQuery/Items/EnvQueryItemType_Actor.h
+++ b/Engine/Source/Runtime/AIModule/Classes/EnvironmentQuery/Items/EnvQueryItemType_Actor.h
@@ -10,10 +10,14 @@ struct FEnvQueryContextData;
 UCLASS()
 class AIMODULE_API UEnvQueryItemType_Actor : public UEnvQueryItemType_ActorBase
 {
-	GENERATED_UCLASS_BODY()
+	GENERATED_BODY()
+public:
+	typedef const FWeakObjectPtr& FValueType;
+
+	UEnvQueryItemType_Actor(const FObjectInitializer& ObjectInitializer = FObjectInitializer::Get());
 
 	static AActor* GetValue(const uint8* RawData);
-	static void SetValue(uint8* RawData, const AActor* Value);
+	static void SetValue(uint8* RawData, const FWeakObjectPtr& Value);
 
 	static void SetContextHelper(FEnvQueryContextData& ContextData, const AActor* SingleActor);
 	static void SetContextHelper(FEnvQueryContextData& ContextData, const TArray<const AActor*>& MultipleActors);
diff --git a/Engine/Source/Runtime/AIModule/Classes/EnvironmentQuery/Items/EnvQueryItemType_ActorBase.h b/Engine/Source/Runtime/AIModule/Classes/EnvironmentQuery/Items/EnvQueryItemType_ActorBase.h
--- a/Engine/Source/Runtime/AIModule/Classes/EnvironmentQuery/Items/EnvQueryItemType_ActorBase.h
+++ b/Engine/Source/Runtime/AIModule/Classes/EnvironmentQuery/Items/EnvQueryItemType_ActorBase.h
@@ -11,8 +11,9 @@ class UBlackboardComponent;
 UCLASS(Abstract)
 class AIMODULE_API UEnvQueryItemType_ActorBase : public UEnvQueryItemType_VectorBase
 {
-	GENERATED_UCLASS_BODY()
+	GENERATED_BODY()
 
+public:
 	virtual AActor* GetActor(const uint8* RawData) const;
 
 	virtual void AddBlackboardFilters(FBlackboardKeySelector& KeySelector, UObject* FilterOwner) const override;
diff --git a/Engine/Source/Runtime/AIModule/Classes/EnvironmentQuery/Items/EnvQueryItemType_Direction.h b/Engine/Source/Runtime/AIModule/Classes/EnvironmentQuery/Items/EnvQueryItemType_Direction.h
--- a/Engine/Source/Runtime/AIModule/Classes/EnvironmentQuery/Items/EnvQueryItemType_Direction.h
+++ b/Engine/Source/Runtime/AIModule/Classes/EnvironmentQuery/Items/EnvQueryItemType_Direction.h
@@ -9,7 +9,11 @@ struct FEnvQueryContextData;
 UCLASS()
 class AIMODULE_API UEnvQueryItemType_Direction : public UEnvQueryItemType_VectorBase
 {
-	GENERATED_UCLASS_BODY()
+	GENERATED_BODY()
+public:
+	typedef FVector FValueType;
+
+	UEnvQueryItemType_Direction(const FObjectInitializer& ObjectInitializer = FObjectInitializer::Get());
 
 	static FVector GetValue(const uint8* RawData);
 	static void SetValue(uint8* RawData, const FVector& Value);
diff --git a/Engine/Source/Runtime/AIModule/Classes/EnvironmentQuery/Items/EnvQueryItemType_Point.h b/Engine/Source/Runtime/AIModule/Classes/EnvironmentQuery/Items/EnvQueryItemType_Point.h
--- a/Engine/Source/Runtime/AIModule/Classes/EnvironmentQuery/Items/EnvQueryItemType_Point.h
+++ b/Engine/Source/Runtime/AIModule/Classes/EnvironmentQuery/Items/EnvQueryItemType_Point.h
@@ -11,12 +11,14 @@ struct FEnvQueryContextData;
 UCLASS()
 class AIMODULE_API UEnvQueryItemType_Point : public UEnvQueryItemType_VectorBase
 {
-	GENERATED_UCLASS_BODY()
+	GENERATED_BODY()
+public:
+	typedef const FNavLocation FValueType;
 
-	static FVector GetValue(const uint8* RawData);
-	static void SetValue(uint8* RawData, const FVector& Value);
+	UEnvQueryItemType_Point(const FObjectInitializer& ObjectInitializer = FObjectInitializer::Get());
+
+	static FNavLocation GetValue(const uint8* RawData);
 	static void SetValue(uint8* RawData, const FNavLocation& Value);
-	static FNavLocation GetNavValue(const uint8* RawData);
 
 	static void SetContextHelper(FEnvQueryContextData& ContextData, const FVector& SinglePoint);
 	static void SetContextHelper(FEnvQueryContextData& ContextData, const TArray<FVector>& MultiplePoints);
@@ -26,4 +28,8 @@ class AIMODULE_API UEnvQueryItemType_Point : public UEnvQueryItemType_VectorBase
 
 	/** Update location data in existing item */
 	virtual void SetItemNavLocation(uint8* RawData, const FNavLocation& Value) const;
-};
\ No newline at end of file
+};
+
+// a specialization to support saving locations with navigation data already gathered
+template<>
+AIMODULE_API void FEnvQueryInstance::AddItemData<UEnvQueryItemType_Point, FVector>(FVector ItemValue);
\ No newline at end of file
diff --git a/Engine/Source/Runtime/AIModule/Classes/EnvironmentQuery/Items/EnvQueryItemType_VectorBase.h b/Engine/Source/Runtime/AIModule/Classes/EnvironmentQuery/Items/EnvQueryItemType_VectorBase.h
--- a/Engine/Source/Runtime/AIModule/Classes/EnvironmentQuery/Items/EnvQueryItemType_VectorBase.h
+++ b/Engine/Source/Runtime/AIModule/Classes/EnvironmentQuery/Items/EnvQueryItemType_VectorBase.h
@@ -10,8 +10,9 @@ class UBlackboardComponent;
 UCLASS(Abstract)
 class AIMODULE_API UEnvQueryItemType_VectorBase : public UEnvQueryItemType
 {
-	GENERATED_UCLASS_BODY()
+	GENERATED_BODY()
 
+public:
 	virtual FVector GetItemLocation(const uint8* RawData) const;
 	virtual FRotator GetItemRotation(const uint8* RawData) const;
 
diff --git a/Engine/Source/Runtime/AIModule/Classes/Perception/AIPerceptionComponent.h b/Engine/Source/Runtime/AIModule/Classes/Perception/AIPerceptionComponent.h
--- a/Engine/Source/Runtime/AIModule/Classes/Perception/AIPerceptionComponent.h
+++ b/Engine/Source/Runtime/AIModule/Classes/Perception/AIPerceptionComponent.h
@@ -273,7 +273,7 @@ public:
 	UFUNCTION(BlueprintCallable, Category = "AI|Perception")
 	void GetKnownPerceivedActors(TSubclassOf<UAISense> SenseToUse, TArray<AActor*>& OutActors) const;
 	
-	DEPRECATED(4.13, "GetPerceivedActors is deprecated. Use GetPerceivedActors or GetPerceivedActors")
+	DEPRECATED(4.13, "GetPerceivedActors is deprecated. Use GetCurrentlyPerceivedActors or GetKnownPerceivedActors")
 	UFUNCTION(BlueprintCallable, Category = "AI|Perception")
 	void GetPerceivedActors(TSubclassOf<UAISense> SenseToUse, TArray<AActor*>& OutActors) const;
 	
diff --git a/Engine/Source/Runtime/AIModule/Private/AIController.cpp b/Engine/Source/Runtime/AIModule/Private/AIController.cpp
--- a/Engine/Source/Runtime/AIModule/Private/AIController.cpp
+++ b/Engine/Source/Runtime/AIModule/Private/AIController.cpp
@@ -568,7 +568,8 @@ EPathFollowingRequestResult::Type AAIController::MoveToActor(AActor* Goal, float
 	// abort active movement to keep only one request running
 	if (PathFollowingComponent && PathFollowingComponent->GetStatus() != EPathFollowingStatus::Idle)
 	{
-		PathFollowingComponent->AbortMove(*this, FPathFollowingResultFlags::ForcedScript | FPathFollowingResultFlags::NewRequest);
+		PathFollowingComponent->AbortMove(*this, FPathFollowingResultFlags::ForcedScript | FPathFollowingResultFlags::NewRequest
+			, FAIRequestID::CurrentRequest, EPathFollowingVelocityMode::Keep);
 	}
 
 	FAIMoveRequest MoveReq(Goal);
@@ -587,7 +588,8 @@ EPathFollowingRequestResult::Type AAIController::MoveToLocation(const FVector& D
 	// abort active movement to keep only one request running
 	if (PathFollowingComponent && PathFollowingComponent->GetStatus() != EPathFollowingStatus::Idle)
 	{
-		PathFollowingComponent->AbortMove(*this, FPathFollowingResultFlags::ForcedScript | FPathFollowingResultFlags::NewRequest);
+		PathFollowingComponent->AbortMove(*this, FPathFollowingResultFlags::ForcedScript | FPathFollowingResultFlags::NewRequest
+			, FAIRequestID::CurrentRequest, EPathFollowingVelocityMode::Keep);
 	}
 
 	FAIMoveRequest MoveReq(Dest);
@@ -1024,6 +1026,14 @@ bool AAIController::UseBlackboard(UBlackboardData* BlackboardAsset, UBlackboardC
 	return bSuccess;
 }
 
+bool AAIController::ShouldSyncBlackboardWith(const UBlackboardComponent& OtherBlackboardComponent) const 
+{ 
+	return Blackboard != nullptr
+		&& Blackboard->GetBlackboardAsset() != nullptr
+		&& OtherBlackboardComponent.GetBlackboardAsset() != nullptr
+		&& Blackboard->GetBlackboardAsset()->IsRelatedTo(*OtherBlackboardComponent.GetBlackboardAsset());
+}
+
 bool AAIController::SuggestTossVelocity(FVector& OutTossVelocity, FVector Start, FVector End, float TossSpeed, bool bPreferHighArc, float CollisionRadius, bool bOnlyTraceUp)
 {
 	// pawn's physics volume gets 2nd priority
diff --git a/Engine/Source/Runtime/AIModule/Private/Actions/PawnActionsComponent.cpp b/Engine/Source/Runtime/AIModule/Private/Actions/PawnActionsComponent.cpp
--- a/Engine/Source/Runtime/AIModule/Private/Actions/PawnActionsComponent.cpp
+++ b/Engine/Source/Runtime/AIModule/Private/Actions/PawnActionsComponent.cpp
@@ -111,7 +111,7 @@ void FPawnActionStack::PopAction(UPawnAction& ActionToPop)
 		// note StopAction can be null
 		UPawnAction* StopAction = ActionToPop.ParentAction;
 
-		while (ActionBeingRemoved != StopAction)
+		while (ActionBeingRemoved != StopAction && ActionBeingRemoved != nullptr)
 		{
 			checkSlow(ActionBeingRemoved);
 			UPawnAction* NextAction = ActionBeingRemoved->ParentAction;
diff --git a/Engine/Source/Runtime/AIModule/Private/BehaviorTree/BlackboardComponent.cpp b/Engine/Source/Runtime/AIModule/Private/BehaviorTree/BlackboardComponent.cpp
--- a/Engine/Source/Runtime/AIModule/Private/BehaviorTree/BlackboardComponent.cpp
+++ b/Engine/Source/Runtime/AIModule/Private/BehaviorTree/BlackboardComponent.cpp
@@ -5,8 +5,6 @@
 #include "BehaviorTree/Blackboard/BlackboardKeyAllTypes.h"
 #include "BehaviorTree/BlackboardComponent.h"
 
-PRAGMA_DISABLE_OPTIMIZATION
-
 UBlackboardComponent::UBlackboardComponent(const FObjectInitializer& ObjectInitializer) : Super(ObjectInitializer)
 {
 	PrimaryComponentTick.bCanEverTick = false;
@@ -245,7 +243,7 @@ bool UBlackboardComponent::ShouldSyncWithBlackboard(UBlackboardComponent& OtherB
 {
 	return &OtherBlackboardComponent != this && (
 		(BrainComp == nullptr || (BrainComp->GetAIOwner() != nullptr && BrainComp->GetAIOwner()->ShouldSyncBlackboardWith(OtherBlackboardComponent) == true))
-		|| (OtherBlackboardComponent.BrainComp == nullptr || (OtherBlackboardComponent.BrainComp->GetAIOwner() != nullptr && OtherBlackboardComponent.BrainComp->GetAIOwner()->ShouldSyncBlackboardWith(OtherBlackboardComponent) == true)));
+		|| (OtherBlackboardComponent.BrainComp == nullptr || (OtherBlackboardComponent.BrainComp->GetAIOwner() != nullptr && OtherBlackboardComponent.BrainComp->GetAIOwner()->ShouldSyncBlackboardWith(*this) == true)));
 }
 
 UBrainComponent* UBlackboardComponent::GetBrainComponent() const
@@ -746,5 +744,3 @@ bool UBlackboardComponent::GetRotationFromEntry(FBlackboard::FKey KeyID, FRotato
 
 	return false;
 }
-
-PRAGMA_ENABLE_OPTIMIZATION
\ No newline at end of file
diff --git a/Engine/Source/Runtime/AIModule/Private/BehaviorTree/BlackboardData.cpp b/Engine/Source/Runtime/AIModule/Private/BehaviorTree/BlackboardData.cpp
--- a/Engine/Source/Runtime/AIModule/Private/BehaviorTree/BlackboardData.cpp
+++ b/Engine/Source/Runtime/AIModule/Private/BehaviorTree/BlackboardData.cpp
@@ -139,7 +139,7 @@ void UBlackboardData::PostEditChangeProperty(FPropertyChangedEvent& PropertyChan
 		if (PropertyChangedEvent.Property->GetFName() == NAME_Parent)
 		{
 			// look for cycles
-			if (Parent && Parent->HasParent(this))
+			if (Parent && Parent->IsChildOf(*this))
 			{
 				UE_LOG(LogBehaviorTree, Warning, TEXT("Blackboard asset (%s) has (%s) in parent chain! Clearing value to avoid cycle."),
 					*GetNameSafe(Parent), *GetNameSafe(this));
@@ -164,6 +164,7 @@ void UBlackboardData::PostEditChangeProperty(FPropertyChangedEvent& PropertyChan
 		}
 	}
 }
+#endif // WITH_EDITOR
 
 void UBlackboardData::PropagateKeyChangesToDerivedBlackboardAssets()
 {
@@ -178,8 +179,6 @@ void UBlackboardData::PropagateKeyChangesToDerivedBlackboardAssets()
 	}
 }
 
-#endif // WITH_EDITOR
-
 static bool ContainsKeyName(FName KeyName, const TArray<FBlackboardEntry>& Keys, const TArray<FBlackboardEntry>& ParentKeys)
 {
 	for (int32 KeyIndex = 0; KeyIndex < Keys.Num(); KeyIndex++)
@@ -201,16 +200,6 @@ static bool ContainsKeyName(FName KeyName, const TArray<FBlackboardEntry>& Keys,
 	return false;
 }
 
-bool UBlackboardData::HasParent(const UBlackboardData* TestParent) const
-{
-	if (Parent == TestParent)
-	{
-		return true;
-	}
-
-	return Parent ? Parent->HasParent(TestParent) : false;
-}
-
 void UBlackboardData::UpdateParentKeys()
 {
 	if (Parent == this)
@@ -259,3 +248,29 @@ void UBlackboardData::UpdateDeprecatedKeys()
 		}
 	}
 }
+
+bool UBlackboardData::IsChildOf(const UBlackboardData& OtherAsset) const
+{
+	const UBlackboardData* TmpParent = Parent;
+	
+	// rewind
+	while (TmpParent != nullptr && TmpParent != &OtherAsset)
+	{
+		TmpParent = TmpParent->Parent;
+	}
+
+	return (TmpParent == &OtherAsset);
+}
+
+//----------------------------------------------------------------------//
+// DEPRECATED
+//----------------------------------------------------------------------//
+bool UBlackboardData::HasParent(const UBlackboardData* TestParent) const
+{
+	if (Parent == TestParent)
+	{
+		return true;
+	}
+
+	return Parent ? Parent->HasParent(TestParent) : false;
+}
\ No newline at end of file
diff --git a/Engine/Source/Runtime/AIModule/Private/EnvironmentQuery/EnvQueryInstance.cpp b/Engine/Source/Runtime/AIModule/Private/EnvironmentQuery/EnvQueryInstance.cpp
--- a/Engine/Source/Runtime/AIModule/Private/EnvironmentQuery/EnvQueryInstance.cpp
+++ b/Engine/Source/Runtime/AIModule/Private/EnvironmentQuery/EnvQueryInstance.cpp
@@ -382,7 +382,7 @@ void FEnvQueryInstance::Log(const FString Msg) const
 }
 #endif // !NO_LOGGING
 
-void FEnvQueryInstance::ReserveItemData(int32 NumAdditionalItems)
+void FEnvQueryInstance::ReserveItemData(const int32 NumAdditionalItems)
 {
 	DEC_MEMORY_STAT_BY(STAT_AI_EQS_InstanceMemory, RawData.GetAllocatedSize());
 
@@ -391,7 +391,6 @@ void FEnvQueryInstance::ReserveItemData(int32 NumAdditionalItems)
 	INC_MEMORY_STAT_BY(STAT_AI_EQS_InstanceMemory, RawData.GetAllocatedSize());
 }
 
-
 FEnvQueryInstance::ItemIterator::ItemIterator(const UEnvQueryTest* QueryTest, FEnvQueryInstance& QueryInstance, int32 StartingItemIndex)
 	: Instance(&QueryInstance)
 	, CurrentItem(StartingItemIndex != INDEX_NONE ? StartingItemIndex : QueryInstance.CurrentTestStartingItem)
diff --git a/Engine/Source/Runtime/AIModule/Private/EnvironmentQuery/Generators/EnvQueryGenerator_BlueprintBase.cpp b/Engine/Source/Runtime/AIModule/Private/EnvironmentQuery/Generators/EnvQueryGenerator_BlueprintBase.cpp
--- a/Engine/Source/Runtime/AIModule/Private/EnvironmentQuery/Generators/EnvQueryGenerator_BlueprintBase.cpp
+++ b/Engine/Source/Runtime/AIModule/Private/EnvironmentQuery/Generators/EnvQueryGenerator_BlueprintBase.cpp
@@ -44,7 +44,7 @@ void UEnvQueryGenerator_BlueprintBase::AddGeneratedVector(FVector Vector) const
 	check(CachedQueryInstance);
 	if (ensure(ItemType->IsChildOf<UEnvQueryItemType_ActorBase>() == false))
 	{
-		CachedQueryInstance->AddItemData<UEnvQueryItemType_Point>(FNavLocation(Vector));
+		CachedQueryInstance->AddItemData<UEnvQueryItemType_Point>(Vector);
 	}
 	else
 	{
@@ -66,7 +66,7 @@ void UEnvQueryGenerator_BlueprintBase::AddGeneratedActor(AActor* Actor) const
 			, *GetName());
 		if (Actor)
 		{
-			CachedQueryInstance->AddItemData<UEnvQueryItemType_Point>(FNavLocation(Actor->GetActorLocation()));
+			CachedQueryInstance->AddItemData<UEnvQueryItemType_Point>(Actor->GetActorLocation());
 		}
 	}
 }
diff --git a/Engine/Source/Runtime/AIModule/Private/EnvironmentQuery/Items/EnvQueryItemType_Actor.cpp b/Engine/Source/Runtime/AIModule/Private/EnvironmentQuery/Items/EnvQueryItemType_Actor.cpp
--- a/Engine/Source/Runtime/AIModule/Private/EnvironmentQuery/Items/EnvQueryItemType_Actor.cpp
+++ b/Engine/Source/Runtime/AIModule/Private/EnvironmentQuery/Items/EnvQueryItemType_Actor.cpp
@@ -14,7 +14,7 @@ AActor* UEnvQueryItemType_Actor::GetValue(const uint8* RawData)
 	return (AActor*)(WeakObjPtr.Get());
 }
 
-void UEnvQueryItemType_Actor::SetValue(uint8* RawData, const AActor* Value)
+void UEnvQueryItemType_Actor::SetValue(uint8* RawData, const FWeakObjectPtr& Value)
 {
 	FWeakObjectPtr WeakObjPtr(Value);
 	SetValueInMemory<FWeakObjectPtr>(RawData, WeakObjPtr);
diff --git a/Engine/Source/Runtime/AIModule/Private/EnvironmentQuery/Items/EnvQueryItemType_ActorBase.cpp b/Engine/Source/Runtime/AIModule/Private/EnvironmentQuery/Items/EnvQueryItemType_ActorBase.cpp
--- a/Engine/Source/Runtime/AIModule/Private/EnvironmentQuery/Items/EnvQueryItemType_ActorBase.cpp
+++ b/Engine/Source/Runtime/AIModule/Private/EnvironmentQuery/Items/EnvQueryItemType_ActorBase.cpp
@@ -5,10 +5,6 @@
 #include "BehaviorTree/Blackboard/BlackboardKeyType_Object.h"
 #include "EnvironmentQuery/Items/EnvQueryItemType_ActorBase.h"
 
-UEnvQueryItemType_ActorBase::UEnvQueryItemType_ActorBase(const FObjectInitializer& ObjectInitializer) : Super(ObjectInitializer)
-{
-}
-
 void UEnvQueryItemType_ActorBase::AddBlackboardFilters(FBlackboardKeySelector& KeySelector, UObject* FilterOwner) const
 {
 	Super::AddBlackboardFilters(KeySelector, FilterOwner);
diff --git a/Engine/Source/Runtime/AIModule/Private/EnvironmentQuery/Items/EnvQueryItemType_Point.cpp b/Engine/Source/Runtime/AIModule/Private/EnvironmentQuery/Items/EnvQueryItemType_Point.cpp
--- a/Engine/Source/Runtime/AIModule/Private/EnvironmentQuery/Items/EnvQueryItemType_Point.cpp
+++ b/Engine/Source/Runtime/AIModule/Private/EnvironmentQuery/Items/EnvQueryItemType_Point.cpp
@@ -3,19 +3,21 @@
 #include "AIModulePrivate.h"
 #include "EnvironmentQuery/Items/EnvQueryItemType_Point.h"
 
-UEnvQueryItemType_Point::UEnvQueryItemType_Point(const FObjectInitializer& ObjectInitializer) : Super(ObjectInitializer)
+template<>
+void FEnvQueryInstance::AddItemData<UEnvQueryItemType_Point, FVector>(FVector ItemValue)
 {
-	ValueSize = sizeof(FNavLocation);
+	const FNavLocation Item = FNavLocation(ItemValue);
+	AddItemData<UEnvQueryItemType_Point>(Item);
 }
 
-FVector UEnvQueryItemType_Point::GetValue(const uint8* RawData)
+UEnvQueryItemType_Point::UEnvQueryItemType_Point(const FObjectInitializer& ObjectInitializer) : Super(ObjectInitializer)
 {
-	return GetValueFromMemory<FNavLocation>(RawData).Location;
+	ValueSize = sizeof(FNavLocation);
 }
 
-void UEnvQueryItemType_Point::SetValue(uint8* RawData, const FVector& Value)
+FNavLocation UEnvQueryItemType_Point::GetValue(const uint8* RawData)
 {
-	return SetValueInMemory<FNavLocation>(RawData, FNavLocation(Value));
+	return GetValueFromMemory<FNavLocation>(RawData);
 }
 
 void UEnvQueryItemType_Point::SetValue(uint8* RawData, const FNavLocation& Value)
@@ -23,11 +25,6 @@ void UEnvQueryItemType_Point::SetValue(uint8* RawData, const FNavLocation& Value
 	return SetValueInMemory<FNavLocation>(RawData, Value);
 }
 
-FNavLocation UEnvQueryItemType_Point::GetNavValue(const uint8* RawData)
-{
-	return GetValueFromMemory<FNavLocation>(RawData);
-}
-
 FVector UEnvQueryItemType_Point::GetItemLocation(const uint8* RawData) const
 {
 	return UEnvQueryItemType_Point::GetValue(RawData);
@@ -35,7 +32,7 @@ FVector UEnvQueryItemType_Point::GetItemLocation(const uint8* RawData) const
 
 FNavLocation UEnvQueryItemType_Point::GetItemNavLocation(const uint8* RawData) const
 {
-	return UEnvQueryItemType_Point::GetNavValue(RawData);
+	return UEnvQueryItemType_Point::GetValue(RawData);
 }
 
 void UEnvQueryItemType_Point::SetItemNavLocation(uint8* RawData, const FNavLocation& Value) const
@@ -49,7 +46,7 @@ void UEnvQueryItemType_Point::SetContextHelper(FEnvQueryContextData& ContextData
 	ContextData.NumValues = 1;
 	ContextData.RawData.SetNumUninitialized(sizeof(FNavLocation));
 
-	UEnvQueryItemType_Point::SetValue((uint8*)ContextData.RawData.GetData(), SinglePoint);
+	UEnvQueryItemType_Point::SetValue((uint8*)ContextData.RawData.GetData(), FNavLocation(SinglePoint));
 }
 
 void UEnvQueryItemType_Point::SetContextHelper(FEnvQueryContextData& ContextData, const TArray<FVector>& MultiplePoints)
@@ -61,7 +58,7 @@ void UEnvQueryItemType_Point::SetContextHelper(FEnvQueryContextData& ContextData
 	uint8* RawData = (uint8*)ContextData.RawData.GetData();
 	for (int32 PointIndex = 0; PointIndex < MultiplePoints.Num(); PointIndex++)
 	{
-		UEnvQueryItemType_Point::SetValue(RawData, MultiplePoints[PointIndex]);
+		UEnvQueryItemType_Point::SetValue(RawData, FNavLocation(MultiplePoints[PointIndex]));
 		RawData += sizeof(FNavLocation);
 	}
 }
diff --git a/Engine/Source/Runtime/AIModule/Private/EnvironmentQuery/Items/EnvQueryItemType_VectorBase.cpp b/Engine/Source/Runtime/AIModule/Private/EnvironmentQuery/Items/EnvQueryItemType_VectorBase.cpp
--- a/Engine/Source/Runtime/AIModule/Private/EnvironmentQuery/Items/EnvQueryItemType_VectorBase.cpp
+++ b/Engine/Source/Runtime/AIModule/Private/EnvironmentQuery/Items/EnvQueryItemType_VectorBase.cpp
@@ -5,10 +5,6 @@
 #include "BehaviorTree/Blackboard/BlackboardKeyType_Vector.h"
 #include "EnvironmentQuery/Items/EnvQueryItemType_VectorBase.h"
 
-UEnvQueryItemType_VectorBase::UEnvQueryItemType_VectorBase(const FObjectInitializer& ObjectInitializer) : Super(ObjectInitializer)
-{
-}
-
 void UEnvQueryItemType_VectorBase::AddBlackboardFilters(FBlackboardKeySelector& KeySelector, UObject* FilterOwner) const
 {
 	KeySelector.AddVectorFilter(FilterOwner, GetClass()->GetFName());
