Copying //UE4/Dev-Editor to //UE4/Dev-Main (Source: //UE4/Dev-Editor @ 3152045)

diff --git a/Engine/Source/Editor/PropertyEditor/Private/PropertyCustomizationHelpers.cpp b/Engine/Source/Editor/PropertyEditor/Private/PropertyCustomizationHelpers.cpp
--- a/Engine/Source/Editor/PropertyEditor/Private/PropertyCustomizationHelpers.cpp
+++ b/Engine/Source/Editor/PropertyEditor/Private/PropertyCustomizationHelpers.cpp
@@ -677,9 +677,10 @@ public:
 		FOnGenerateWidgetsForMaterial InOnGenerateNameWidgetsForMaterial, 
 		FOnGenerateWidgetsForMaterial InOnGenerateWidgetsForMaterial, 
 		FOnResetMaterialToDefaultClicked InOnResetToDefaultClicked,
-		int32 InMultipleMaterialCount )
+		int32 InMultipleMaterialCount,
+		bool bShowUsedTextures)
 	{
-		return MakeShareable( new FMaterialItemView( Material, InOnMaterialChanged, InOnGenerateNameWidgetsForMaterial, InOnGenerateWidgetsForMaterial, InOnResetToDefaultClicked, InMultipleMaterialCount ) );
+		return MakeShareable( new FMaterialItemView( Material, InOnMaterialChanged, InOnGenerateNameWidgetsForMaterial, InOnGenerateWidgetsForMaterial, InOnResetToDefaultClicked, InMultipleMaterialCount, bShowUsedTextures) );
 	}
 
 	TSharedRef<SWidget> CreateNameContent()
@@ -738,6 +739,7 @@ public:
 								.VAlign( VAlign_Center )
 								.ContentPadding(2)
 								.IsEnabled( this, &FMaterialItemView::IsTexturesMenuEnabled )
+								.Visibility( bShowUsedTextures ? EVisibility::Visible : EVisibility::Hidden )
 								.ButtonContent()
 								[
 									SNew( STextBlock )
@@ -747,7 +749,23 @@ public:
 								]
 							]
 						]
-					
+					]
+					+SHorizontalBox::Slot()
+					.AutoWidth()
+					.VAlign(VAlign_Top)
+					.Padding(4.0f, 9.0f)
+					[
+						// Add a button to reset the material to the base material
+						SNew(SButton)
+						.ToolTipText(LOCTEXT("ResetToBaseMaterial", "Reset to base material"))
+						.ButtonStyle(FEditorStyle::Get(), "NoBorder")
+						.ContentPadding(0)
+						.Visibility(this, &FMaterialItemView::GetReplaceVisibility)
+						.OnClicked(this, &FMaterialItemView::OnResetToBaseClicked)
+						[
+							SNew(SImage)
+							.Image(FEditorStyle::GetBrush("PropertyWindow.DiffersFromDefault"))
+						]
 					]
 				]
 				+SVerticalBox::Slot()
@@ -757,24 +775,8 @@ public:
 				[
 					OnGenerateCustomMaterialWidgets.IsBound() ? OnGenerateCustomMaterialWidgets.Execute( MaterialItem.Material.Get(), MaterialItem.SlotIndex ) : StaticCastSharedRef<SWidget>( SNullWidget::NullWidget )
 				]
-			]
-			+SHorizontalBox::Slot()
-			.AutoWidth()
-			.VAlign( VAlign_Center )
-			.Padding( 4.0f, 0.0f )
-			[
-				// Add a button to reset the material to the base material
-				SNew(SButton)
-				.ToolTipText(LOCTEXT( "ResetToBaseMaterial", "Reset to base material" ))
-				.ButtonStyle( FEditorStyle::Get(), "NoBorder" )
-				.ContentPadding(0) 
-				.Visibility( this, &FMaterialItemView::GetReplaceVisibility )
-				.OnClicked( this, &FMaterialItemView::OnResetToBaseClicked )
-				[
-					SNew(SImage)
-					.Image( FEditorStyle::GetBrush("PropertyWindow.DiffersFromDefault") )
-				]
 			];
+
 	}
 
 private:
@@ -784,13 +786,16 @@ private:
 						FOnGenerateWidgetsForMaterial& InOnGenerateNameWidgets, 
 						FOnGenerateWidgetsForMaterial& InOnGenerateMaterialWidgets, 
 						FOnResetMaterialToDefaultClicked& InOnResetToDefaultClicked,
-						int32 InMultipleMaterialCount )
+						int32 InMultipleMaterialCount,
+						bool bInShowUsedTextures)
+						
 		: MaterialItem( InMaterial )
 		, OnMaterialChanged( InOnMaterialChanged )
 		, OnGenerateCustomNameWidgets( InOnGenerateNameWidgets )
 		, OnGenerateCustomMaterialWidgets( InOnGenerateMaterialWidgets )
 		, OnResetToDefaultClicked( InOnResetToDefaultClicked )
 		, MultipleMaterialCount( InMultipleMaterialCount )
+		, bShowUsedTextures( bInShowUsedTextures )
 	{
 
 	}
@@ -901,14 +906,16 @@ private:
 	FOnGenerateWidgetsForMaterial OnGenerateCustomMaterialWidgets;
 	FOnResetMaterialToDefaultClicked OnResetToDefaultClicked;
 	int32 MultipleMaterialCount;
+	bool bShowUsedTextures;
 };
 
 
-FMaterialList::FMaterialList(IDetailLayoutBuilder& InDetailLayoutBuilder, FMaterialListDelegates& InMaterialListDelegates, bool bInAllowCollapse /*= false*/)
+FMaterialList::FMaterialList(IDetailLayoutBuilder& InDetailLayoutBuilder, FMaterialListDelegates& InMaterialListDelegates, bool bInAllowCollapse, bool bInShowUsedTextures)
 	: MaterialListDelegates( InMaterialListDelegates )
 	, DetailLayoutBuilder( InDetailLayoutBuilder )
 	, MaterialListBuilder( new FMaterialListBuilder )
 	, bAllowCollpase(bInAllowCollapse)
+	, bShowUsedTextures(bInShowUsedTextures)
 {
 }
 
@@ -1089,7 +1096,7 @@ void FMaterialList::AddMaterialItem( FDetailWidgetRow& Row, int32 CurrentSlot, c
 {
 	uint32 NumMaterials = MaterialListBuilder->GetNumMaterialsInSlot(CurrentSlot);
 
-	TSharedRef<FMaterialItemView> NewView = FMaterialItemView::Create( Item, MaterialListDelegates.OnMaterialChanged, MaterialListDelegates.OnGenerateCustomNameWidgets, MaterialListDelegates.OnGenerateCustomMaterialWidgets, MaterialListDelegates.OnResetMaterialToDefaultClicked, NumMaterials );
+	TSharedRef<FMaterialItemView> NewView = FMaterialItemView::Create( Item, MaterialListDelegates.OnMaterialChanged, MaterialListDelegates.OnGenerateCustomNameWidgets, MaterialListDelegates.OnGenerateCustomMaterialWidgets, MaterialListDelegates.OnResetMaterialToDefaultClicked, NumMaterials, bShowUsedTextures );
 
 	TSharedPtr<SWidget> RightSideContent;
 	if( bDisplayLink )
@@ -1120,7 +1127,7 @@ void FMaterialList::AddMaterialItem( FDetailWidgetRow& Row, int32 CurrentSlot, c
 		NewView->CreateNameContent()
 	]
 	.ValueContent()
-	.MinDesiredWidth(250.0f)
+	.MinDesiredWidth(250.f)
 	.MaxDesiredWidth(0.0f) // no maximum
 	[
 		RightSideContent.ToSharedRef()
@@ -1488,22 +1495,23 @@ public:
 				[
 					SNew(SHorizontalBox)
 					+ SHorizontalBox::Slot()
-					.Padding(0.0f, 3.0f, 0.0f, 0.0f)
+					.Padding(0)
 					.VAlign(VAlign_Center)
 					.AutoWidth()
 					[
 						SNew(STextBlock)
 						.Font(IDetailLayoutBuilder::GetDetailFont())
-						.Text(LOCTEXT("SectionListItemMaterialSlotNameLabel", "Material Slot Name: "))
+						.Text(LOCTEXT("SectionListItemMaterialSlotNameLabel", "Material Slot"))
 						.ToolTipText(MaterialSlotNameTooltipText)
 					]
 					+ SHorizontalBox::Slot()
 					.VAlign(VAlign_Center)
 					.AutoWidth()
-					.Padding(15, 0, 0, 0)
+					.Padding(5, 0, 0, 0)
 					[
 						SNew(SBox)
 						.HAlign(HAlign_Left)
+						.VAlign(VAlign_Center)
 						[
 							//Material Slot Name
 							SNew(SComboButton)
@@ -1675,15 +1683,12 @@ void FSectionList::Tick(float DeltaTime)
 
 void FSectionList::GenerateHeaderRowContent(FDetailWidgetRow& NodeRow)
 {
-	if (bAllowCollpase)
-	{
-		NodeRow.NameContent()
-			[
-				SNew(STextBlock)
-				.Text(LOCTEXT("SectionHeaderTitle", "Sections"))
-				.Font(IDetailLayoutBuilder::GetDetailFont())
-			];
-	}
+	NodeRow.NameContent()
+	[
+		SNew(STextBlock)
+		.Text(LOCTEXT("SectionHeaderTitle", "Sections"))
+		.Font(IDetailLayoutBuilder::GetDetailFont())
+	];
 }
 
 void FSectionList::GenerateChildContent(IDetailChildrenBuilder& ChildrenBuilder)
@@ -1703,59 +1708,12 @@ void FSectionList::GenerateChildContent(IDetailChildrenBuilder& ChildrenBuilder)
 		{
 			const FSectionListItem& Section = *It;
 
-			if (CurrentLODIndex != Section.LodIndex)
-			{
-				// We've encountered a new lod.  Make a widget to display that
-				CurrentLODIndex = Section.LodIndex;
-
-				uint32 NumLodSections = SectionListBuilder->GetNumSections(CurrentLODIndex);
-
-				// If an element is expanded we want to display all its Sections
-				bool bWantToDisplayAllSections = NumLodSections > 1 && ExpandedSlots.Contains(CurrentLODIndex);
-
-				// If we are currently displaying an expanded set of Sections for an element add a link to collapse all of them
-				if (bWantToDisplayAllSections)
-				{
-					FDetailWidgetRow& ChildRow = ChildrenBuilder.AddChildContent(LOCTEXT("HideAllSectionSearchString", "Hide All Sections"));
-
-					FFormatNamedArguments Arguments;
-					Arguments.Add(TEXT("LodIndex"), CurrentLODIndex);
-					ChildRow
-						.ValueContent()
-						.MaxDesiredWidth(0.0f)// No Max Width
-						[
-							SNew(SBox)
-							.HAlign(HAlign_Center)
-							[
-								SNew(SHyperlink)
-								.TextStyle(FEditorStyle::Get(), "MaterialList.HyperlinkStyle")
-								.Text(FText::Format(LOCTEXT("HideAllSectionLinkText", "Hide All Sections on LOD {LodIndex}"), Arguments))
-								.OnNavigate(this, &FSectionList::OnHideSectionsForLod, CurrentLODIndex)
-							]
-						];
-				}
-
-				if (NumLodSections > 1 && !bWantToDisplayAllSections)
-				{
-					// The current slot has multiple elements to view
-					bDisplayAllSectionsInSlot = false;
-
-					FDetailWidgetRow& ChildRow = ChildrenBuilder.AddChildContent(FText::GetEmpty());
-
-					AddSectionItem(ChildRow, CurrentLODIndex, FSectionListItem(CurrentLODIndex, Section.SectionIndex, Section.MaterialSlotName, Section.MaterialSlotIndex, Section.OriginalMaterialSlotName, Section.AvailableMaterialSlotName, Section.Material.Get(), Section.IsSectionUsingCloth), !bDisplayAllSectionsInSlot);
-				}
-				else
-				{
-					bDisplayAllSectionsInSlot = true;
-				}
-
-			}
+			CurrentLODIndex = Section.LodIndex;
 
 			// Display each thumbnail element unless we shouldn't display multiple Sections for one slot
 			if (bDisplayAllSectionsInSlot)
 			{
 				FDetailWidgetRow& ChildRow = ChildrenBuilder.AddChildContent(Section.Material.IsValid() ? FText::FromString(Section.Material->GetName()) : FText::GetEmpty());
-
 				AddSectionItem(ChildRow, CurrentLODIndex, FSectionListItem(CurrentLODIndex, Section.SectionIndex, Section.MaterialSlotName, Section.MaterialSlotIndex, Section.OriginalMaterialSlotName, Section.AvailableMaterialSlotName, Section.Material.Get(), Section.IsSectionUsingCloth), !bDisplayAllSectionsInSlot);
 			}
 		}
@@ -1820,5 +1778,65 @@ void FSectionList::AddSectionItem(FDetailWidgetRow& Row, int32 LodIndex, const s
 }
 
 
-#undef LOCTEXT_NAMESPACE
+void SMaterialSlotWidget::Construct(const FArguments& InArgs, int32 SlotIndex, bool bIsMaterialUsed)
+{
+	TSharedPtr<SHorizontalBox> SlotNameBox;
+
+	TSharedRef<SWidget> DeleteButton =
+		PropertyCustomizationHelpers::MakeDeleteButton(
+			InArgs._OnDeleteMaterialSlot,
+			LOCTEXT("CustomNameMaterialNotUsedDeleteTooltip", "Delete this material slot"),
+			InArgs._CanDeleteMaterialSlot);
+
+	if(bIsMaterialUsed)
+	{
+		DeleteButton->SetVisibility(EVisibility::Hidden);
+	}
+
+	ChildSlot
+	[
+		SAssignNew(SlotNameBox, SHorizontalBox)
+		+ SHorizontalBox::Slot()
+		.AutoWidth()
+		.HAlign(HAlign_Left)
+		.VAlign(VAlign_Center)
+		[
+			SNew(SBox)
+			.VAlign(VAlign_Center)
+			// Match the size of the thumbnail
+			.MinDesiredWidth(64.0f)
+			[
+				SNew(STextBlock)
+				.Font(IDetailLayoutBuilder::GetDetailFont())
+				.Text(LOCTEXT("MaterialArrayNameLabelStringKey", "Slot name"))
+			]
+		]
+		+ SHorizontalBox::Slot()
+		.VAlign(VAlign_Center)
+		.AutoWidth()
+		.Padding(12.0f, 3.0f, 0.0f, 3.0f)
+		[
+			SNew(SBox)
+			.VAlign(VAlign_Center)
+			.MinDesiredWidth(125.0f)
+			[
+				SNew(SEditableTextBox)
+				.Text(InArgs._MaterialName)
+				.OnTextChanged(InArgs._OnMaterialNameChanged)
+				.OnTextCommitted(InArgs._OnMaterialNameCommitted)
+				.Font(IDetailLayoutBuilder::GetDetailFont())
+			]
+		]
+	];
+
+	SlotNameBox->AddSlot()
+		.AutoWidth()
+		.VAlign(VAlign_Center)
+		.HAlign(HAlign_Left)
+		.Padding(2)
+		[
+			DeleteButton
+		];
+}
 
+#undef LOCTEXT_NAMESPACE
diff --git a/Engine/Source/Editor/PropertyEditor/Private/PropertyHandleImpl.cpp b/Engine/Source/Editor/PropertyEditor/Private/PropertyHandleImpl.cpp
--- a/Engine/Source/Editor/PropertyEditor/Private/PropertyHandleImpl.cpp
+++ b/Engine/Source/Editor/PropertyEditor/Private/PropertyHandleImpl.cpp
@@ -2905,7 +2905,9 @@ FPropertyAccess::Result FPropertyHandleObject::SetValueFromFormattedString(const
 				UClass* AllowedClass = FindObject<UClass>(ANY_PACKAGE, *ClassName);
 				const bool bIsInterface = AllowedClass && AllowedClass->HasAnyClassFlags(CLASS_Interface);
 				
-				if (AllowedClass && (QualifiedObject->IsA(AllowedClass) || (bIsInterface && QualifiedObject->GetClass()->ImplementsInterface(AllowedClass))))
+				// Check if the object is an allowed class type this property supports
+				// Note: QualifieObject may be null if we're clearing the value.  Allow clears to pass through without a supported class
+				if (AllowedClass && (!QualifiedObject || QualifiedObject->IsA(AllowedClass) || (bIsInterface && QualifiedObject->GetClass()->ImplementsInterface(AllowedClass))))
 				{
 					bSupportedObject = true;
 					break;
diff --git a/Engine/Source/Editor/PropertyEditor/Private/PropertyNode.cpp b/Engine/Source/Editor/PropertyEditor/Private/PropertyNode.cpp
--- a/Engine/Source/Editor/PropertyEditor/Private/PropertyNode.cpp
+++ b/Engine/Source/Editor/PropertyEditor/Private/PropertyNode.cpp
@@ -147,7 +147,7 @@ void FPropertyNode::InitNode( const FPropertyNodeInitParams& InitParams )
 		static const FName Name_ShowInnerProperties("ShowInnerProperties");
 
 		bIsEditInlineNew = bIsObjectOrInterface && GotReadAddresses && MyProperty->HasMetaData(Name_EditInline);
-		bShowInnerObjectProperties = bIsObjectOrInterface && GotReadAddresses && MyProperty->HasMetaData(Name_ShowInnerProperties);
+		bShowInnerObjectProperties = bIsObjectOrInterface && MyProperty->HasMetaData(Name_ShowInnerProperties);
 
 		if(bIsEditInlineNew)
 		{
@@ -342,12 +342,18 @@ FObjectPropertyNode* FPropertyNode::FindRootObjectItemParent()
 }
 
 
+bool FPropertyNode::DoesChildPropertyRequireValidation(UProperty* InChildProp)
+{
+	return InChildProp != nullptr && (Cast<UObjectProperty>(InChildProp) != nullptr || Cast<UStructProperty>(InChildProp) != nullptr);
+}
+
 /** 
  * Used to see if any data has been destroyed from under the property tree.  Should only be called by PropertyWindow::OnIdle
  */
 EPropertyDataValidationResult FPropertyNode::EnsureDataIsValid()
 {
 	bool bValidateChildren = !HasNodeFlags(EPropertyNodeFlags::SkipChildValidation);
+	bool bValidateChildrenKeyNodes = false;		// by default, we don't check this, since it's just for Map properties
 
 	// The root must always be validated
 	if( GetParentNode() == NULL || HasNodeFlags(EPropertyNodeFlags::RequiresValidation) != 0 )
@@ -378,21 +384,26 @@ EPropertyDataValidationResult FPropertyNode::EnsureDataIsValid()
 
 			bool bArrayHasNewItem = false;
 
+			UProperty* ContainerElementProperty = nullptr;
+
 			if (ArrayProperty)
 			{
-				if (!ArrayProperty->Inner->IsA(UObjectProperty::StaticClass()) && !ArrayProperty->Inner->IsA(UStructProperty::StaticClass()))
-				{
-					bValidateChildren = false;
-				}
+				ContainerElementProperty = ArrayProperty->Inner;
 			}
 			else if (SetProperty)
 			{
-				if (!SetProperty->ElementProp->IsA(UObjectProperty::StaticClass()) && !SetProperty->ElementProp->IsA(UStructProperty::StaticClass()))
-				{
-					bValidateChildren = false;
-				}
+				ContainerElementProperty = SetProperty->ElementProp;
+			}
+			else if (MapProperty)
+			{
+				// Need to attempt to validate both the key and value properties...
+				bValidateChildrenKeyNodes = DoesChildPropertyRequireValidation(MapProperty->KeyProp);
+
+				ContainerElementProperty = MapProperty->ValueProp;
 			}
 
+			bValidateChildren = DoesChildPropertyRequireValidation(ContainerElementProperty);
+
 			//verify that the number of object children are the same too
 			UObjectPropertyBase* ObjectProperty = Cast<UObjectPropertyBase>(MyProperty);
 			//check to see, if this an object property, whether the contents are NULL or not.
@@ -539,7 +550,7 @@ EPropertyDataValidationResult FPropertyNode::EnsureDataIsValid()
 			const bool bHasChildren = (GetNumChildNodes() != 0);
 			// If the object property is not null and has no children, its children need to be rebuilt
 			// If the object property is null and this node has children, the node needs to be rebuilt
-			if (ObjectProperty && ((!bObjectPropertyNull && !bHasChildren) || (bObjectPropertyNull && bHasChildren)))
+			if (!HasNodeFlags(EPropertyNodeFlags::ShowInnerObjectProperties) && ObjectProperty && ((!bObjectPropertyNull && !bHasChildren) || (bObjectPropertyNull && bHasChildren)))
 			{
 				RebuildChildren();
 				return EPropertyDataValidationResult::PropertiesChanged;
@@ -556,20 +567,33 @@ EPropertyDataValidationResult FPropertyNode::EnsureDataIsValid()
 	
 	EPropertyDataValidationResult FinalResult = EPropertyDataValidationResult::DataValid;
 
-	//go through my children
-	if (bValidateChildren)
+	// Validate children and/or their key nodes.
+	if (bValidateChildren || bValidateChildrenKeyNodes)
 	{
 		for (int32 Scan = 0; Scan < ChildNodes.Num(); ++Scan)
 		{
 			TSharedPtr<FPropertyNode>& ChildNode = ChildNodes[Scan];
 			check(ChildNode.IsValid());
 
-			EPropertyDataValidationResult ChildDataResult = ChildNode->EnsureDataIsValid();
-			if (FinalResult == EPropertyDataValidationResult::DataValid && ChildDataResult != EPropertyDataValidationResult::DataValid)
+			if (bValidateChildren)
 			{
-				FinalResult = ChildDataResult;
+				EPropertyDataValidationResult ChildDataResult = ChildNode->EnsureDataIsValid();
+				if (FinalResult == EPropertyDataValidationResult::DataValid && ChildDataResult != EPropertyDataValidationResult::DataValid)
+				{
+					FinalResult = ChildDataResult;
+				}
 			}
 
+			// If the child property has a key node that needs validation, validate it here
+			TSharedPtr<FPropertyNode>& ChildKeyNode = ChildNode->GetPropertyKeyNode();
+			if (bValidateChildrenKeyNodes && ChildKeyNode.IsValid())
+			{
+				EPropertyDataValidationResult ChildDataResult = ChildKeyNode->EnsureDataIsValid();
+				if (FinalResult == EPropertyDataValidationResult::DataValid && ChildDataResult != EPropertyDataValidationResult::DataValid)
+				{
+					FinalResult = ChildDataResult;
+				}
+			}
 		}
 	}
 
@@ -2340,7 +2364,7 @@ void FPropertyNode::NotifyPostChange( FPropertyChangedEvent& InPropertyChangedEv
 	}
 
 
-	if( ObjectNode && OriginalActiveProperty )
+	if( OriginalActiveProperty )
 	{
 		//if i have metadata forcing other property windows to rebuild
 		FString MetaData = OriginalActiveProperty->GetMetaData(TEXT("ForceRebuildProperty"));
diff --git a/Engine/Source/Editor/PropertyEditor/Private/PropertyNode.h b/Engine/Source/Editor/PropertyEditor/Private/PropertyNode.h
--- a/Engine/Source/Editor/PropertyEditor/Private/PropertyNode.h
+++ b/Engine/Source/Editor/PropertyEditor/Private/PropertyNode.h
@@ -739,7 +739,9 @@ public:
 		ValueNode->PropertyKeyNode = KeyNode;
 	}
 
-	TSharedPtr<FPropertyNode> GetPropertyKeyNode() const { return PropertyKeyNode; }
+	TSharedPtr<FPropertyNode>& GetPropertyKeyNode() { return PropertyKeyNode; }
+
+	const TSharedPtr<FPropertyNode>& GetPropertyKeyNode() const { return PropertyKeyNode; }
 
 protected:
 
@@ -814,6 +816,13 @@ protected:
 	 */
 	void UpdateEditConstState();
 
+	/**
+	 * Checks to see if the supplied property of a child node requires validation
+	 * @param	InChildProp		The property of the child node
+	 * @return	True if the property requires validation, false otherwise
+	 */
+	static bool DoesChildPropertyRequireValidation(UProperty* InChildProp);
+
 protected:
 	/**
 	 * The node that is the parent of this node or nullptr for the root
diff --git a/Engine/Source/Editor/PropertyEditor/Private/SDetailSingleItemRow.cpp b/Engine/Source/Editor/PropertyEditor/Private/SDetailSingleItemRow.cpp
--- a/Engine/Source/Editor/PropertyEditor/Private/SDetailSingleItemRow.cpp
+++ b/Engine/Source/Editor/PropertyEditor/Private/SDetailSingleItemRow.cpp
@@ -471,29 +471,31 @@ const FSlateBrush* SDetailSingleItemRow::GetBorderImage() const
 
 TSharedRef<SWidget> SDetailSingleItemRow::CreateExtensionWidget(TSharedRef<SWidget> ValueWidget, FDetailLayoutCustomization& InCustomization, TSharedRef<IDetailTreeNode> InTreeNode)
 {
-	IDetailsViewPrivate& DetailsView = InTreeNode->GetDetailsView();
-	TSharedPtr<IDetailPropertyExtensionHandler> ExtensionHandler = DetailsView.GetExtensionHandler();
-
-	if ( ExtensionHandler.IsValid() && InCustomization.HasPropertyNode() )
+	if(InTreeNode->GetParentCategory().IsValid())
 	{
-		TSharedPtr<IPropertyHandle> Handle = PropertyEditorHelpers::GetPropertyHandle(InCustomization.GetPropertyNode().ToSharedRef(), nullptr, nullptr);
+		IDetailsViewPrivate& DetailsView = InTreeNode->GetDetailsView();
+		TSharedPtr<IDetailPropertyExtensionHandler> ExtensionHandler = DetailsView.GetExtensionHandler();
 
-		UClass* ObjectClass = InCustomization.GetPropertyNode()->FindObjectItemParent()->GetObjectBaseClass();
-		if (Handle->IsValidHandle() && ExtensionHandler->IsPropertyExtendable(ObjectClass, *Handle))
+		if(ExtensionHandler.IsValid() && InCustomization.HasPropertyNode())
 		{
-			ValueWidget = SNew(SHorizontalBox)
+			TSharedPtr<IPropertyHandle> Handle = PropertyEditorHelpers::GetPropertyHandle(InCustomization.GetPropertyNode().ToSharedRef(), nullptr, nullptr);
 
-				+ SHorizontalBox::Slot()
-				.FillWidth(1.0f)
-				[
-					ValueWidget
-				]
-					
-				+ SHorizontalBox::Slot()
-				.AutoWidth()
-				[
-					ExtensionHandler->GenerateExtensionWidget(ObjectClass, Handle)
-				];
+			UClass* ObjectClass = InCustomization.GetPropertyNode()->FindObjectItemParent()->GetObjectBaseClass();
+			if(Handle->IsValidHandle() && ExtensionHandler->IsPropertyExtendable(ObjectClass, *Handle))
+			{
+				ValueWidget = SNew(SHorizontalBox)
+					+ SHorizontalBox::Slot()
+					.FillWidth(1.0f)
+					[
+						ValueWidget
+					]
+
+					+ SHorizontalBox::Slot()
+					.AutoWidth()
+					[
+						ExtensionHandler->GenerateExtensionWidget(ObjectClass, Handle)
+					];
+			}
 		}
 	}
 
diff --git a/Engine/Source/Editor/PropertyEditor/Private/UserInterface/PropertyEditor/SPropertyEditorAsset.cpp b/Engine/Source/Editor/PropertyEditor/Private/UserInterface/PropertyEditor/SPropertyEditorAsset.cpp
--- a/Engine/Source/Editor/PropertyEditor/Private/UserInterface/PropertyEditor/SPropertyEditorAsset.cpp
+++ b/Engine/Source/Editor/PropertyEditor/Private/UserInterface/PropertyEditor/SPropertyEditorAsset.cpp
@@ -284,6 +284,7 @@ void SPropertyEditorAsset::Construct( const FArguments& InArgs, const TSharedPtr
 			[
 				SAssignNew( CustomContentBox, SVerticalBox )
 				+ SVerticalBox::Slot()
+				.Padding(0.0f, 4.0f, 0.0f, 0.0f)
 				[
 					AssetComboButton.ToSharedRef()
 				]
diff --git a/Engine/Source/Editor/PropertyEditor/Public/PropertyCustomizationHelpers.h b/Engine/Source/Editor/PropertyEditor/Public/PropertyCustomizationHelpers.h
--- a/Engine/Source/Editor/PropertyEditor/Public/PropertyCustomizationHelpers.h
+++ b/Engine/Source/Editor/PropertyEditor/Public/PropertyCustomizationHelpers.h
@@ -471,7 +471,7 @@ class FMaterialList
 	, public TSharedFromThis<FMaterialList>
 {
 public:
-	PROPERTYEDITOR_API FMaterialList( IDetailLayoutBuilder& InDetailLayoutBuilder, FMaterialListDelegates& MaterialListDelegates, bool bInAllowCollapse = false);
+	PROPERTYEDITOR_API FMaterialList( IDetailLayoutBuilder& InDetailLayoutBuilder, FMaterialListDelegates& MaterialListDelegates, bool bInAllowCollapse = false, bool bInShowUsedTextures = true);
 
 	/**
 	 * @return true if materials are being displayed.                                                          
@@ -538,9 +538,27 @@ private:
 
 	/** Allow Collapse of material header row. Right now if you allow collapse, it will initially collapse. */
 	bool bAllowCollpase;
+	/** Whether or not to use the used textures menu for each material entry */
+	bool bShowUsedTextures;
 };
 
 
+/**
+ * Helper class to create a material slot name widget for material lists
+ */
+class SMaterialSlotWidget : public SCompoundWidget
+{
+	SLATE_BEGIN_ARGS(SMaterialSlotWidget)
+	{}
+		SLATE_ATTRIBUTE(FText, MaterialName)
+		SLATE_EVENT(FOnTextChanged, OnMaterialNameChanged)
+		SLATE_EVENT(FOnTextCommitted, OnMaterialNameCommitted)
+		SLATE_ATTRIBUTE(bool, CanDeleteMaterialSlot)
+		SLATE_EVENT(FSimpleDelegate, OnDeleteMaterialSlot)
+	SLATE_END_ARGS()
+
+	PROPERTYEDITOR_API void Construct(const FArguments& InArgs, int32 SlotIndex, bool bIsMaterialUsed);
+};
 
 //////////////////////////////////////////////////////////////////////////
 //
diff --git a/Engine/Source/Editor/SizeMap/SSizeMap.cpp b/Engine/Source/Editor/SizeMap/SSizeMap.cpp
--- a/Engine/Source/Editor/SizeMap/SSizeMap.cpp
+++ b/Engine/Source/Editor/SizeMap/SSizeMap.cpp
@@ -314,11 +314,11 @@ void SSizeMap::GatherDependenciesRecursively( FAssetRegistryModule& AssetRegistr
 						// @todo sizemap: We should instead have a special EResourceSizeMode that asks for the worst case size.  Some assets (like UTextureCube) currently always report resident mip size, even when asked for inclusive size
 						if( Asset->IsA( UTexture2D::StaticClass() ) )
 						{
-							NodeSizeMapData.AssetSize = Asset->GetResourceSize( EResourceSizeMode::Inclusive );
+							NodeSizeMapData.AssetSize = Asset->GetResourceSizeBytes( EResourceSizeMode::Inclusive );
 						}
 						else
 						{
-							NodeSizeMapData.AssetSize = Asset->GetResourceSize( EResourceSizeMode::Exclusive );
+							NodeSizeMapData.AssetSize = Asset->GetResourceSizeBytes( EResourceSizeMode::Exclusive );
 						}
 
 						NodeSizeMapData.bHasKnownSize = (NodeSizeMapData.AssetSize != 0);
diff --git a/Engine/Source/Editor/SourceControlWindows/Private/SSourceControlSubmit.cpp b/Engine/Source/Editor/SourceControlWindows/Private/SSourceControlSubmit.cpp
--- a/Engine/Source/Editor/SourceControlWindows/Private/SSourceControlSubmit.cpp
+++ b/Engine/Source/Editor/SourceControlWindows/Private/SSourceControlSubmit.cpp
@@ -450,6 +450,7 @@ void FSourceControlWindows::ChoosePackagesToCheckInCompleted(const TArray<UPacka
 		PendingDeletePaths.Add(FPaths::ConvertRelativePathToFull(FPaths::EngineContentDir()));
 		PendingDeletePaths.Add(FPaths::ConvertRelativePathToFull(FPaths::GameContentDir()));
 		PendingDeletePaths.Add(FPaths::ConvertRelativePathToFull(FPaths::GameConfigDir()));
+		PendingDeletePaths.Add(FPaths::ConvertRelativePathToFull(FPaths::GetProjectFilePath()));
 
 		const bool bUseSourceControlStateCache = true;
 		PromptForCheckin(bUseSourceControlStateCache, PackageNames, PendingDeletePaths, ConfigFiles);
@@ -537,6 +538,7 @@ void FSourceControlWindows::ChoosePackagesToCheckIn()
 			Filenames.Add(FPaths::ConvertRelativePathToFull(FPaths::EngineContentDir()));
 			Filenames.Add(FPaths::ConvertRelativePathToFull(FPaths::GameContentDir()));
 			Filenames.Add(FPaths::ConvertRelativePathToFull(FPaths::GameConfigDir()));
+			Filenames.Add(FPaths::ConvertRelativePathToFull(FPaths::GetProjectFilePath()));
 
 			ISourceControlProvider& SourceControlProvider = ISourceControlModule::Get().GetProvider();
 			FSourceControlOperationRef Operation = ISourceControlOperation::Create<FUpdateStatus>();
diff --git a/Engine/Source/Editor/StaticMeshEditor/Private/StaticMeshEditorTools.cpp b/Engine/Source/Editor/StaticMeshEditor/Private/StaticMeshEditorTools.cpp
--- a/Engine/Source/Editor/StaticMeshEditor/Private/StaticMeshEditorTools.cpp
+++ b/Engine/Source/Editor/StaticMeshEditor/Private/StaticMeshEditorTools.cpp
@@ -11,10 +11,10 @@
 #include "PhysicsEngine/BodySetup.h"
 #include "FbxMeshUtils.h"
 #include "SVectorInputBox.h"
-
 #include "Runtime/Analytics/Analytics/Public/Interfaces/IAnalyticsProvider.h"
 #include "EngineAnalytics.h"
 #include "STextComboBox.h"
+#include "ScopedTransaction.h"
 
 const float MaxHullAccuracy = 1.f;
 const float MinHullAccuracy = 0.f;
@@ -1654,13 +1654,14 @@ UStaticMesh& FMeshMaterialsLayout::GetStaticMesh() const
 
 void FMeshMaterialsLayout::AddToCategory(IDetailCategoryBuilder& CategoryBuilder)
 {
-	CategoryBuilder.AddCustomRow(LOCTEXT("AddLODLevelCategories_MaterialArrayOperationAdd", "Materials Operation Add Material Slot"))
+	CategoryBuilder.AddCustomRow(LOCTEXT("AddLODLevelCategories_MaterialArrayOperationAdd", "Add Material Slot"))
 		.NameContent()
 		.HAlign(HAlign_Left)
 		.VAlign(VAlign_Center)
 		[
 			SNew(STextBlock)
-			.Text(LOCTEXT("AddLODLevelCategories_MaterialArrayOperations", "Materials Operations"))
+			.Font(IDetailLayoutBuilder::GetDetailFont())
+			.Text(LOCTEXT("AddLODLevelCategories_MaterialArrayOperations", "Material Slots"))
 		]
 		.ValueContent()
 		.HAlign(HAlign_Left)
@@ -1677,6 +1678,7 @@ void FMeshMaterialsLayout::AddToCategory(IDetailCategoryBuilder& CategoryBuilder
 				[
 					SNew(STextBlock)
 					.Text(this, &FMeshMaterialsLayout::GetMaterialArrayText)
+					.Font(IDetailLayoutBuilder::GetDetailFont())
 				]
 				+ SHorizontalBox::Slot()
 				.AutoWidth()
@@ -1704,7 +1706,6 @@ void FMeshMaterialsLayout::AddToCategory(IDetailCategoryBuilder& CategoryBuilder
 	FMaterialListDelegates MaterialListDelegates;
 	MaterialListDelegates.OnGetMaterials.BindSP(this, &FMeshMaterialsLayout::GetMaterials);
 	MaterialListDelegates.OnMaterialChanged.BindSP(this, &FMeshMaterialsLayout::OnMaterialChanged);
-	MaterialListDelegates.OnGenerateCustomNameWidgets.BindSP(this, &FMeshMaterialsLayout::OnGenerateNameWidgetsForMaterial);
 	MaterialListDelegates.OnGenerateCustomMaterialWidgets.BindSP(this, &FMeshMaterialsLayout::OnGenerateWidgetsForMaterial);
 	MaterialListDelegates.OnMaterialListDirty.BindSP(this, &FMeshMaterialsLayout::OnMaterialListDirty);
 	MaterialListDelegates.OnResetMaterialToDefaultClicked.BindSP(this, &FMeshMaterialsLayout::OnResetMaterialToDefaultClicked);
@@ -1716,13 +1717,12 @@ FReply FMeshMaterialsLayout::AddMaterialSlot()
 {
 	UStaticMesh& StaticMesh = GetStaticMesh();
 
-	GEditor->BeginTransaction(LOCTEXT("FMeshMaterialsLayout_AddMaterialSlot", "Add a material slot"));
+	FScopedTransaction Transaction(LOCTEXT("FMeshMaterialsLayout_AddMaterialSlot", "Add material slot"));
 	StaticMesh.Modify();
 	StaticMesh.StaticMaterials.Add(FStaticMaterial());
 
 	StaticMesh.PostEditChange();
-	// End the transation if we created one
-	GEditor->EndTransaction();
+
 	return FReply::Handled();
 }
 
@@ -1730,7 +1730,7 @@ FText FMeshMaterialsLayout::GetMaterialArrayText() const
 {
 	UStaticMesh& StaticMesh = GetStaticMesh();
 
-	FString MaterialArrayText = TEXT(" Materials Slots");
+	FString MaterialArrayText = TEXT(" Material Slots");
 	int32 SlotNumber = 0;
 	SlotNumber = StaticMesh.StaticMaterials.Num();
 	MaterialArrayText = FString::FromInt(SlotNumber) + MaterialArrayText;
@@ -1769,126 +1769,23 @@ void FMeshMaterialsLayout::OnMaterialChanged(UMaterialInterface* NewMaterial, UM
 	CallPostEditChange(ChangedProperty);
 }
 
-
-TSharedRef<SWidget> FMeshMaterialsLayout::OnGenerateNameWidgetsForMaterial(UMaterialInterface* Material, int32 MaterialIndex)
+TSharedRef<SWidget> FMeshMaterialsLayout::OnGenerateWidgetsForMaterial(UMaterialInterface* Material, int32 SlotIndex)
 {
 	UStaticMesh& StaticMesh = GetStaticMesh();
-	bool MaterialIsUsed = false;
-	if (MaterialUsedMap.Contains(MaterialIndex))
-	{
-		MaterialIsUsed = MaterialUsedMap.Find(MaterialIndex)->Num() > 0;
-	}
-	TSharedRef<SVerticalBox> CustomNameWidget = SNew(SVerticalBox)
-		+ SVerticalBox::Slot()
-		.AutoHeight()
-		.Padding(0, 2, 0, 0)
-		[
-			SNew(SHorizontalBox)
-			+ SHorizontalBox::Slot()
-			.VAlign(VAlign_Center)
-			.Padding(0.0f, 3.0f, 0.0f, 0.0f)
-			.FillWidth(0.4f)
-			[
-				SNew(STextBlock)
-				.Font(IDetailLayoutBuilder::GetDetailFont())
-				.Text(LOCTEXT("MaterialArrayNameLabelStringKey", "Name"))
-				.ToolTipText(this, &FMeshMaterialsLayout::GetOriginalImportMaterialNameText, MaterialIndex)
-			]
-			+ SHorizontalBox::Slot()
-			.Padding(5.0f, 3.0f, 5.0f, 0.0f)
-			.VAlign(VAlign_Center)
-			.FillWidth(1.0f)
-			[
-				SNew(SEditableTextBox)
-				.Text(this, &FMeshMaterialsLayout::GetMaterialNameText, MaterialIndex)
-				.OnTextChanged(this, &FMeshMaterialsLayout::OnMaterialNameChanged, MaterialIndex)
-				.OnTextCommitted(this, &FMeshMaterialsLayout::OnMaterialNameCommitted, MaterialIndex)
-			]
-		];
-	if (!MaterialIsUsed)
-	{
-		CustomNameWidget->AddSlot()
-			.AutoHeight()
-			.Padding(0, 2, 0, 0)
-			[
-				SNew(SHorizontalBox)
-				+ SHorizontalBox::Slot()
-				.Padding(0.0f, 3.0f, 0.0f, 0.0f)
-				.VAlign(VAlign_Center)
-				.FillWidth(0.4f)
-				[
-					SNew(STextBlock)
-					.Font(IDetailLayoutBuilder::GetDetailFont())
-					.Text(LOCTEXT("CustomNameMaterialNotUsed", "Unused"))
-					.ToolTipText(LOCTEXT("CustomNameMaterialNotUsedTooltip", "This material is not use by any LODs section! It will be set to None when the skeletal mesh will be saved."))
-				]
-				+ SHorizontalBox::Slot()
-				.Padding(5.0f, 3.0f, 5.0f, 0.0f)
-				.VAlign(VAlign_Center)
-				.FillWidth(1.0f)
-				[
-					SNew(SButton)
-					.Text(LOCTEXT("CustomNameMaterialNotUsedDelete", "Delete"))
-					.ToolTipText(LOCTEXT("CustomNameMaterialNotUsedDeleteTooltip", "Delete this material slot."))
-					.IsEnabled(this, &FMeshMaterialsLayout::CanDeleteMaterialSlot, MaterialIndex)
-					.OnClicked(this, &FMeshMaterialsLayout::OnDeleteMaterialSlot, MaterialIndex)
-				]
-			];
-	}
-	else
+	bool bMaterialIsUsed = false;
+	if(MaterialUsedMap.Contains(SlotIndex))
 	{
-		CustomNameWidget->AddSlot()
-			.AutoHeight()
-			.Padding(0, 2, 0, 0)
-			[
-				SNew(SHorizontalBox)
-				+ SHorizontalBox::Slot()
-				.VAlign(VAlign_Center)
-				.Padding(0.0f, 3.0f, 0.0f, 0.0f)
-				.FillWidth(0.4f)
-				[
-					SNew(STextBlock)
-					.Font(IDetailLayoutBuilder::GetDetailFont())
-					.Text(LOCTEXT("MaterialSlotUsedByLabel", "Used By"))
-					.ToolTipText(LOCTEXT("MaterialSlotUsedByLabelTooltip", "List all sections that use this material slot"))
-				]
-				+ SHorizontalBox::Slot()
-				.Padding(5.0f, 3.0f, 5.0f, 0.0f)
-				.VAlign(VAlign_Center)
-				.FillWidth(1.0f)
-				[
-					//Material Slot Name
-					SNew(SComboButton)
-					.OnGetMenuContent(this, &FMeshMaterialsLayout::OnGetMaterialSlotUsedByMenuContent, MaterialIndex)
-					.VAlign(VAlign_Center)
-					.ContentPadding(2)
-					.ButtonContent()
-					[
-						SNew(STextBlock)
-						.Font(IDetailLayoutBuilder::GetDetailFont())
-						.Text(this, &FMeshMaterialsLayout::GetFirstMaterialSlotUsedBySection, MaterialIndex)
-					]
-				]
-			];
+		bMaterialIsUsed = MaterialUsedMap.Find(SlotIndex)->Num() > 0;
 	}
-	return CustomNameWidget;
-}
 
-TSharedRef<SWidget> FMeshMaterialsLayout::OnGenerateWidgetsForMaterial(UMaterialInterface* Material, int32 SlotIndex)
-{
-	return SNew(SVerticalBox)
-		+SVerticalBox::Slot()
-		.AutoHeight()
-		[
-			SNew(SCheckBox)
-				.IsChecked(this, &FMeshMaterialsLayout::IsShadowCastingEnabled, SlotIndex)
-				.OnCheckStateChanged(this, &FMeshMaterialsLayout::OnShadowCastingChanged, SlotIndex)
-			[
-				SNew(STextBlock)
-					.Font(FEditorStyle::GetFontStyle("StaticMeshEditor.NormalFont"))
-					.Text(LOCTEXT("CastShadow", "Cast Shadow"))
-			]
-		];
+	return 
+		SNew(SMaterialSlotWidget, SlotIndex, bMaterialIsUsed)
+		.MaterialName(this, &FMeshMaterialsLayout::GetMaterialNameText, SlotIndex)
+		.OnMaterialNameChanged(this, &FMeshMaterialsLayout::OnMaterialNameChanged, SlotIndex)
+		.OnMaterialNameCommitted(this, &FMeshMaterialsLayout::OnMaterialNameCommitted, SlotIndex)
+		.CanDeleteMaterialSlot(this, &FMeshMaterialsLayout::CanDeleteMaterialSlot, SlotIndex)
+		.OnDeleteMaterialSlot(this, &FMeshMaterialsLayout::OnDeleteMaterialSlot, SlotIndex)
+		.ToolTipText(this, &FMeshMaterialsLayout::GetOriginalImportMaterialNameText, SlotIndex);
 }
 
 void FMeshMaterialsLayout::OnResetMaterialToDefaultClicked(UMaterialInterface* Material, int32 MaterialIndex)
@@ -1942,21 +1839,19 @@ bool FMeshMaterialsLayout::CanDeleteMaterialSlot(int32 MaterialIndex) const
 	return (MaterialIndex + 1) == StaticMesh.StaticMaterials.Num();
 }
 
-FReply FMeshMaterialsLayout::OnDeleteMaterialSlot(int32 MaterialIndex)
+void FMeshMaterialsLayout::OnDeleteMaterialSlot(int32 MaterialIndex)
 {
 	UStaticMesh& StaticMesh = GetStaticMesh();
-	if (!CanDeleteMaterialSlot(MaterialIndex))
-		return FReply::Handled();
+	if (CanDeleteMaterialSlot(MaterialIndex))
+	{
+		FScopedTransaction Transaction(LOCTEXT("PersonaDeletedMaterialSlot", "Deleted material slot on skeletal mesh"));
 
-	GEditor->BeginTransaction(LOCTEXT("PersonaDeletedMaterialSlot", "Deleted material slot on skeletal mesh"));
-	StaticMesh.Modify();
-	StaticMesh.StaticMaterials.RemoveAt(MaterialIndex);
+		StaticMesh.Modify();
+		StaticMesh.StaticMaterials.RemoveAt(MaterialIndex);
 
-	StaticMesh.PostEditChange();
-	// End the transation if we created one
-	GEditor->EndTransaction();
+		StaticMesh.PostEditChange();
+	}
 
-	return FReply::Handled();
 }
 
 TSharedRef<SWidget> FMeshMaterialsLayout::OnGetMaterialSlotUsedByMenuContent(int32 MaterialIndex)
diff --git a/Engine/Source/Editor/StaticMeshEditor/Private/StaticMeshEditorTools.h b/Engine/Source/Editor/StaticMeshEditor/Private/StaticMeshEditorTools.h
--- a/Engine/Source/Editor/StaticMeshEditor/Private/StaticMeshEditorTools.h
+++ b/Engine/Source/Editor/StaticMeshEditor/Private/StaticMeshEditorTools.h
@@ -336,7 +336,6 @@ private:
 
 	void GetMaterials(class IMaterialListBuilder& ListBuilder);
 	void OnMaterialChanged(UMaterialInterface* NewMaterial, UMaterialInterface* PrevMaterial, int32 SlotIndex, bool bReplaceAll);
-	TSharedRef<SWidget> OnGenerateNameWidgetsForMaterial(UMaterialInterface* Material, int32 SlotIndex);
 	TSharedRef<SWidget> OnGenerateWidgetsForMaterial(UMaterialInterface* Material, int32 SlotIndex);
 	void OnResetMaterialToDefaultClicked(UMaterialInterface* Material, int32 SlotIndex);
 
@@ -345,7 +344,7 @@ private:
 	void OnMaterialNameCommitted(const FText& InValue, ETextCommit::Type CommitType, int32 MaterialIndex);
 	void OnMaterialNameChanged(const FText& InValue, int32 MaterialIndex);
 	bool CanDeleteMaterialSlot(int32 MaterialIndex) const;
-	FReply OnDeleteMaterialSlot(int32 MaterialIndex);
+	void OnDeleteMaterialSlot(int32 MaterialIndex);
 	TSharedRef<SWidget> OnGetMaterialSlotUsedByMenuContent(int32 MaterialIndex);
 	FText GetFirstMaterialSlotUsedBySection(int32 MaterialIndex) const;
 
diff --git a/Engine/Source/Editor/StatsViewer/Private/StatsPages/PrimitiveStatsPage.cpp b/Engine/Source/Editor/StatsViewer/Private/StatsPages/PrimitiveStatsPage.cpp
--- a/Engine/Source/Editor/StatsViewer/Private/StatsPages/PrimitiveStatsPage.cpp
+++ b/Engine/Source/Editor/StatsViewer/Private/StatsPages/PrimitiveStatsPage.cpp
@@ -261,7 +261,7 @@ struct PrimitiveStatsGenerator
 				NewStatsEntry->Count			= 1;
 				NewStatsEntry->Triangles		= 0;
 				NewStatsEntry->InstTriangles	= 0;
-				NewStatsEntry->ResourceSize		= (float)(FArchiveCountMem(Resource).GetNum() + Resource->GetResourceSize(EResourceSizeMode::Exclusive)) / 1024.0f;
+				NewStatsEntry->ResourceSize		= (float)(FArchiveCountMem(Resource).GetNum() + Resource->GetResourceSizeBytes(EResourceSizeMode::Exclusive)) / 1024.0f;
 				NewStatsEntry->Sections			= 0;
 				NewStatsEntry->InstSections = 0;
 				NewStatsEntry->RadiusMin		= InPrimitiveComponent->Bounds.SphereRadius;
@@ -338,14 +338,16 @@ struct PrimitiveStatsGenerator
 						UniqueTextures.Add(CurrentComponent->HeightmapTexture, &bNotUnique);
 						if (!bNotUnique)
 						{
-							NewStatsEntry->ResourceSize += CurrentComponent->HeightmapTexture->GetResourceSize(EResourceSizeMode::Exclusive);
+							const SIZE_T HeightmapResourceSize = CurrentComponent->HeightmapTexture->GetResourceSizeBytes(EResourceSizeMode::Exclusive);
+							NewStatsEntry->ResourceSize += HeightmapResourceSize;
 						}
 						if (CurrentComponent->XYOffsetmapTexture)
 						{
 							UniqueTextures.Add(CurrentComponent->XYOffsetmapTexture, &bNotUnique);
 							if (!bNotUnique)
 							{
-								NewStatsEntry->ResourceSize += CurrentComponent->XYOffsetmapTexture->GetResourceSize(EResourceSizeMode::Exclusive);
+								const SIZE_T OffsetmapResourceSize = CurrentComponent->XYOffsetmapTexture->GetResourceSizeBytes(EResourceSizeMode::Exclusive);
+								NewStatsEntry->ResourceSize += OffsetmapResourceSize;
 							}
 						}
 
@@ -354,7 +356,8 @@ struct PrimitiveStatsGenerator
 							UniqueTextures.Add((*ItWeightmaps), &bNotUnique);
 							if (!bNotUnique)
 							{
-								NewStatsEntry->ResourceSize += (*ItWeightmaps)->GetResourceSize(EResourceSizeMode::Exclusive);
+								const SIZE_T WeightmapResourceSize = (*ItWeightmaps)->GetResourceSizeBytes(EResourceSizeMode::Exclusive);
+								NewStatsEntry->ResourceSize += WeightmapResourceSize;
 							}
 						}
 					}
diff --git a/Engine/Source/Editor/TextureEditor/Private/TextureEditorToolkit.cpp b/Engine/Source/Editor/TextureEditor/Private/TextureEditorToolkit.cpp
--- a/Engine/Source/Editor/TextureEditor/Private/TextureEditorToolkit.cpp
+++ b/Engine/Source/Editor/TextureEditor/Private/TextureEditorToolkit.cpp
@@ -370,7 +370,7 @@ void FTextureEditorToolkit::PopulateQuickInfo( )
 	CalculateEffectiveTextureDimensions(MipLevel, PreviewEffectiveTextureWidth, PreviewEffectiveTextureHeight);
 
 	// Texture asset size
-	const uint32 Size = (Texture->GetResourceSize(EResourceSizeMode::Exclusive) + 512) / 1024;
+	const uint32 Size = (Texture->GetResourceSizeBytes(EResourceSizeMode::Exclusive) + 512) / 1024;
 
 	FNumberFormattingOptions SizeOptions;
 	SizeOptions.UseGrouping = false;
diff --git a/Engine/Source/Editor/UnrealEd/Classes/Editor/EditorEngine.h b/Engine/Source/Editor/UnrealEd/Classes/Editor/EditorEngine.h
--- a/Engine/Source/Editor/UnrealEd/Classes/Editor/EditorEngine.h
+++ b/Engine/Source/Editor/UnrealEd/Classes/Editor/EditorEngine.h
@@ -1313,11 +1313,12 @@ public:
 	/**
 	 * Deletes all selected actors
 	 *
-	 * @param		InWorld				World context
-	 * @param		bVerifyDeletionCanHappen	[opt] If true (default), verify that deletion can be performed.
-	 * @return									true unless the delete operation was aborted.
+	 * @param	InWorld				World context
+	 * @param	bVerifyDeletionCanHappen	[opt] If true (default), verify that deletion can be performed.
+	 * @param	bWarnAboutReferences		[opt] If true (default), we prompt the user about referenced actours they are about to delete
+	 * @return								true unless the delete operation was aborted.
 	 */
-	virtual bool edactDeleteSelected(UWorld* InWorld, bool bVerifyDeletionCanHappen=true) { return true; }
+	virtual bool edactDeleteSelected(UWorld* InWorld, bool bVerifyDeletionCanHappen=true, bool bWarnAboutReferences = true) { return true; }
 
 	/**
 	 * Checks the state of the selected actors and notifies the user of any potentially unknown destructive actions which may occur as
@@ -1665,7 +1666,7 @@ public:
 	void PlaySessionSingleStepped();
 
 	/** Called when game client received input key */
-	bool ProcessDebuggerCommands(const FKey InKey, const FModifierKeysState ModifierKeyState);
+	bool ProcessDebuggerCommands(const FKey InKey, const FModifierKeysState ModifierKeyState, EInputEvent EventType);
 
 	/**
 	 * Kicks off a "Play From Here" request that was most likely made during a transaction
diff --git a/Engine/Source/Editor/UnrealEd/Classes/Editor/PropertyEditorTestObject.h b/Engine/Source/Editor/UnrealEd/Classes/Editor/PropertyEditorTestObject.h
--- a/Engine/Source/Editor/UnrealEd/Classes/Editor/PropertyEditorTestObject.h
+++ b/Engine/Source/Editor/UnrealEd/Classes/Editor/PropertyEditorTestObject.h
@@ -325,6 +325,9 @@ class UPropertyEditorTestObject : public UObject
 	UPROPERTY(EditAnywhere, Category=TMapTests)
 	TMap<UObject*, int32> ObjectToInt32Map;
 
+	UPROPERTY(EditAnywhere, Category=TMapTests)
+	TMap<UObject*, FLinearColor> ObjectToColorMap;
+
 	UPROPERTY(EditAnywhere, Category=ScriptInterfaces)
 	TScriptInterface<IBlendableInterface> BlendableInterface;
 
diff --git a/Engine/Source/Editor/UnrealEd/Classes/Editor/UnrealEdEngine.h b/Engine/Source/Editor/UnrealEd/Classes/Editor/UnrealEdEngine.h
--- a/Engine/Source/Editor/UnrealEd/Classes/Editor/UnrealEdEngine.h
+++ b/Engine/Source/Editor/UnrealEd/Classes/Editor/UnrealEdEngine.h
@@ -401,9 +401,10 @@ public:
 	 *
 	 * @param	InWorld						World context
 	 * @param	bVerifyDeletionCanHappen	[opt] If true (default), verify that deletion can be performed.
+	 * @param	bWarnAboutReferences		[opt] If true (default), we prompt the user about referenced actours they are about to delete
 	 * @return								true unless the delete operation was aborted.
 	 */
-	virtual bool edactDeleteSelected( UWorld* InWorld, bool bVerifyDeletionCanHappen=true ) override;
+	virtual bool edactDeleteSelected( UWorld* InWorld, bool bVerifyDeletionCanHappen=true, bool bWarnAboutReferences = true) override;
 
 	/**
 	 * Creates a new group from the current selection removing any existing groups.
diff --git a/Engine/Source/Editor/UnrealEd/Classes/Factories/CSVImportFactory.h b/Engine/Source/Editor/UnrealEd/Classes/Factories/CSVImportFactory.h
--- a/Engine/Source/Editor/UnrealEd/Classes/Factories/CSVImportFactory.h
+++ b/Engine/Source/Editor/UnrealEd/Classes/Factories/CSVImportFactory.h
@@ -9,8 +9,41 @@
 
 DECLARE_LOG_CATEGORY_EXTERN(LogCSVImportFactory, Log, All);
 
+/** Enum to indicate what to import CSV as */
+UENUM()
+enum class ECSVImportType : uint8
+{
+	/** Import as UDataTable */
+	ECSV_DataTable,
+	/** Import as UCurveTable */
+	ECSV_CurveTable,
+	/** Import as a UCurveFloat */
+	ECSV_CurveFloat,
+	/** Import as a UCurveVector */
+	ECSV_CurveVector,
+	/** Import as a UCurveLinearColor */
+	ECSV_CurveLinearColor,
+};
+
+USTRUCT()
+struct FCSVImportSettings
+{
+	GENERATED_BODY()
+
+	FCSVImportSettings();
+
+	UPROPERTY()
+	UScriptStruct* ImportRowStruct;
+
+	UPROPERTY()
+	ECSVImportType ImportType;
+
+	UPROPERTY()
+	TEnumAsByte<ERichCurveInterpMode> ImportCurveInterpMode;
+};
+
 UCLASS(hidecategories=Object)
-class UNREALED_API UCSVImportFactory : public UFactory
+class UNREALED_API UCSVImportFactory : public UFactory, public IImportSettingsParser
 {
 	GENERATED_UCLASS_BODY()
 
@@ -18,11 +51,17 @@ class UNREALED_API UCSVImportFactory : public UFactory
 	virtual FText GetDisplayName() const override;
 	virtual UObject* FactoryCreateText( UClass* InClass, UObject* InParent, FName InName, EObjectFlags Flags, UObject* Context, const TCHAR* Type, const TCHAR*& Buffer, const TCHAR* BufferEnd, FFeedbackContext* Warn ) override;
 	virtual bool DoesSupportClass(UClass * Class) override;
-	
-	//~ Begin UFactory Interface
+	virtual bool FactoryCanImport(const FString& Filename) override;
+	virtual	IImportSettingsParser* GetImportSettingsParser() override { return this; }
+
 
 	/* Reimport an object that was created based on a CSV */
 	bool ReimportCSV(UObject* Obj);
+	
+	/**
+	 * IImportSettings interface
+	 */
+	virtual void ParseFromJson(TSharedRef<class FJsonObject> ImportSettingsJson) override;
 
 protected:
 	virtual TArray<FString> DoImportDataTable(class UDataTable* TargetDataTable, const FString& DataToImport);
@@ -32,19 +71,8 @@ protected:
 private:
 	/* Reimport object from the given path*/
 	bool Reimport(UObject* Obj, const FString& Path  );
-};
 
-/** Enum to indicate what to import CSV as */
-enum class ECSVImportType : uint8
-{
-	/** Import as UDataTable */
-	ECSV_DataTable,
-	/** Import as UCurveTable */
-	ECSV_CurveTable,
-	/** Import as a UCurveFloat */
-	ECSV_CurveFloat,
-	/** Import as a UCurveVector */
-	ECSV_CurveVector,
-	/** Import as a UCurveLinearColor */
-	ECSV_CurveLinearColor,
+	UPROPERTY()
+	FCSVImportSettings AutomatedImportSettings;
 };
+
diff --git a/Engine/Source/Editor/UnrealEd/Classes/Factories/Factory.h b/Engine/Source/Editor/UnrealEd/Classes/Factories/Factory.h
--- a/Engine/Source/Editor/UnrealEd/Classes/Factories/Factory.h
+++ b/Engine/Source/Editor/UnrealEd/Classes/Factories/Factory.h
@@ -12,46 +12,10 @@
  * 
  */
 UCLASS(abstract)
-class UNREALED_API UFactory
-	: public UObject
+class UNREALED_API UFactory : public UObject
 {
 	GENERATED_UCLASS_BODY()
 
 public:
 
-	/** The class manufactured by this factory. */
-	UPROPERTY()
-	TSubclassOf<UObject>  SupportedClass;
-
-	/** Class of the context object used to help create the object. */
-	UPROPERTY()
-	TSubclassOf<UObject>  ContextClass;
-
-	/** List of formats supported by the factory. Each entry is of the form "ext;Description" where ext is the file extension. */
-	UPROPERTY()
-	TArray<FString> Formats;
-
-	/** true if the associated editor should be opened after creating a new object. */
-	UPROPERTY()
-	uint32 bEditAfterNew:1;
-
-	/** true if the factory imports objects from files. */
-	UPROPERTY()
-	uint32 bEditorImport:1;
-
-	/** true if the factory imports objects from text. */
-	UPROPERTY()
-	uint32 bText:1;
-
-	/** Determines the order in which factories are tried when importing an object. */
-	DEPRECATED(4.8, "AutoPriority has been replaced with ImportPriority")
-	int32 AutoPriority;
-
-	/** Determines the order in which factories are tried when importing or reimporting an object.
-		Factories with higher priority values will go first. Factories with negative priorities will be excluded. */
-	UPROPERTY()
-	int32 ImportPriority;
-
-public:
-
 	/**
@@ -174,7 +138,6 @@
 
 	/** Do clean up after importing is done. Will be called once for multi batch import. */
 	virtual void CleanUp() {}
-
 	/**
 	 * Creates an asset if it doesn't exist. If it does exist then it overwrites it if possible. If it can not overwrite then it will delete and replace. If it can not delete, it will return nullptr.
 	 * 
@@ -191,6 +154,20 @@
 	/** Returns a new starting point name for newly created assets in the content browser */
 	virtual FString GetDefaultNewAssetName() const;
 
+	/** @return the parser that is capable of parsing a json string of import settings for this factory */
+	virtual class IImportSettingsParser* GetImportSettingsParser() {return nullptr;}
+
+	/**
+	 * Sets the automated import data being used with this factory
+	 *
+	 * @param Data	The automated import data or nullptr if it doesnt exist
+	 */
+	void SetAutomatedAssetImportData(const class UAutomatedAssetImportData* Data);
+
+	/**
+	 * @return true if this factory is being used for automated import.  Dialogs and user input should be disabled if this method returns true
+	 */
+	bool IsAutomatedImport() const { return  GIsAutomationTesting || AutomatedImportData; }
 public:
 
 	/**
@@ -198,7 +175,7 @@
 	 *
 	 * @param Message The message text.
 	 **/
-	static void DisplayOverwriteOptionsDialog(const FText& Message);
+	void DisplayOverwriteOptionsDialog(const FText& Message);
 
 	/** Get the name of the file currently being imported. */
 	static FString GetCurrentFilename()
@@ -218,7 +195,7 @@
 	 * The states are used to suppress messages during multiple object import. 
 	 * It needs to be reset each time a new import is started
 	 */
-	static void ResetState();
+	void ResetState();
 
 	/** Helper function to sort an array of factories by their import priority - use as a predicate for Sort */
 	static bool SortFactoriesByPriority(const UFactory& A, const UFactory& B);
@@ -298,6 +275,45 @@
 	{
 		return FactoryCreateBinary(InClass, InParent, InName, Flags, Context, Type, Buffer, BufferEnd, Warn);
 	}
+public:
 
+	/** The class manufactured by this factory. */
+	UPROPERTY()
+	TSubclassOf<UObject>  SupportedClass;
+
+	/** Class of the context object used to help create the object. */
+	UPROPERTY()
+	TSubclassOf<UObject>  ContextClass;
+
+	/** List of formats supported by the factory. Each entry is of the form "ext;Description" where ext is the file extension. */
+	UPROPERTY()
+	TArray<FString> Formats;
+
+	/** true if the associated editor should be opened after creating a new object. */
+	UPROPERTY()
+	uint32 bEditAfterNew:1;
+
+	/** true if the factory imports objects from files. */
+	UPROPERTY()
+	uint32 bEditorImport:1;
+
+	/** true if the factory imports objects from text. */
+	UPROPERTY()
+	uint32 bText:1;
+
+	/** Determines the order in which factories are tried when importing an object. */
+	DEPRECATED(4.8, "AutoPriority has been replaced with ImportPriority")
+	int32 AutoPriority;
+
+	/** Determines the order in which factories are tried when importing or reimporting an object.
+	Factories with higher priority values will go first. Factories with negative priorities will be excluded. */
+	UPROPERTY()
+	int32 ImportPriority;
+
+	/** Data for how to import files via the automated command line importing interface */
+	UPROPERTY()
+	const class UAutomatedAssetImportData* AutomatedImportData;
+
+
 protected:
 
@@ -307,12 +323,6 @@
 	/** This is the import priority that all factories are given in the default constructor. */
 	static const int32 DefaultImportPriority;
 
-	/**
-	 * For interactive object imports, this value indicates whether the user wants
-	 * objects to be automatically overwritten (See EAppReturnType), or -1 if the
-	 * user should be prompted.
-	 */
-	static int32 OverwriteYesOrNoToAllState;
 
 	/**
 	 * If this value is true, warning messages will be shown once for all objects
@@ -321,4 +331,13 @@
 	 */
 	DEPRECATED(4.8, "bAllowOneTimeWarningMessages is due to be removed in future.")
 	static bool bAllowOneTimeWarningMessages;
+
+	/**
+	 * For interactive object imports, this value indicates whether the user wants
+	 * objects to be automatically overwritten (See EAppReturnType), or -1 if the
+	  * user should be prompted.
+	*/
+	UPROPERTY()
+	int32 OverwriteYesOrNoToAllState;
+
 };
diff --git a/Engine/Source/Editor/UnrealEd/Classes/Factories/FbxFactory.h b/Engine/Source/Editor/UnrealEd/Classes/Factories/FbxFactory.h
--- a/Engine/Source/Editor/UnrealEd/Classes/Factories/FbxFactory.h
+++ b/Engine/Source/Editor/UnrealEd/Classes/Factories/FbxFactory.h
@@ -27,6 +27,8 @@ class UNREALED_API UFbxFactory : public UFactory
 	virtual UClass* ResolveSupportedClass() override;
 	virtual UObject* FactoryCreateBinary(UClass* InClass, UObject* InParent, FName InName, EObjectFlags Flags, UObject* Context, const TCHAR* Type, const uint8*& Buffer, const uint8* BufferEnd, FFeedbackContext* Warn, bool& bOutOperationCanceled) override;
 	virtual bool FactoryCanImport(const FString& Filename) override;
+	virtual IImportSettingsParser* GetImportSettingsParser() override;
+
 	//~ End UFactory Interface
 	
 	/**
diff --git a/Engine/Source/Editor/UnrealEd/Classes/Factories/FbxImportUI.h b/Engine/Source/Editor/UnrealEd/Classes/Factories/FbxImportUI.h
--- a/Engine/Source/Editor/UnrealEd/Classes/Factories/FbxImportUI.h
+++ b/Engine/Source/Editor/UnrealEd/Classes/Factories/FbxImportUI.h
@@ -5,6 +5,7 @@
  */
 
 #pragma once
+#include "ImportSettings.h"
 #include "FbxImportUI.generated.h"
 
 /** Import mesh type */
@@ -24,10 +25,11 @@ enum EFBXImportType
 };
 
 UCLASS(config=EditorPerProjectUserSettings, AutoExpandCategories=(FTransform), HideCategories=Object, MinimalAPI)
-class UFbxImportUI : public UObject
+class UFbxImportUI : public UObject, public IImportSettingsParser
 {
 	GENERATED_UCLASS_BODY()
 
+public:
 	/** Whether or not the imported file is in OBJ format */
 	UPROPERTY()
 	bool bIsObjImport;
@@ -116,9 +118,11 @@ class UFbxImportUI : public UObject
 	UPROPERTY(EditAnywhere, Instanced, Category=Material)
 	class UFbxTextureImportData* TextureImportData;
 
-	//~ Begin UObject Interface
+	/** UObject Interface */
 	virtual bool CanEditChange( const UProperty* InProperty ) const override;
-	//~ End UObject Interface
+
+	/** IImportSettings Interface */
+	virtual void ParseFromJson(TSharedRef<class FJsonObject> ImportSettingsJson) override;
 
 	/** sets MeshTypeToImport */
 	void SetMeshTypeToImport()
@@ -133,4 +137,3 @@ class UFbxImportUI : public UObject
 };
 
 
-
diff --git a/Engine/Source/Editor/UnrealEd/Classes/Factories/ReimportDataTableFactory.h b/Engine/Source/Editor/UnrealEd/Classes/Factories/ReimportDataTableFactory.h
--- a/Engine/Source/Editor/UnrealEd/Classes/Factories/ReimportDataTableFactory.h
+++ b/Engine/Source/Editor/UnrealEd/Classes/Factories/ReimportDataTableFactory.h
@@ -1,6 +1,8 @@
 // Copyright 1998-2016 Epic Games, Inc. All Rights Reserved.
 
 #pragma once
+
+#include "CSVImportFactory.h"
 #include "ReimportDataTableFactory.generated.h"
 
 UCLASS()
diff --git a/Engine/Source/Editor/UnrealEd/Private/AutomatedAssetImportData.cpp b/Engine/Source/Editor/UnrealEd/Private/AutomatedAssetImportData.cpp
new file mode 100644
--- /dev/null
+++ b/Engine/Source/Editor/UnrealEd/Private/AutomatedAssetImportData.cpp
@@ -0,0 +1,100 @@
+// Copyright 1998-2016 Epic Games, Inc. All Rights Reserved.
+
+#include "UnrealEd.h"
+#include "AutomatedAssetImportData.h"
+
+DEFINE_LOG_CATEGORY(LogAutomatedImport);
+
+UAutomatedAssetImportData::UAutomatedAssetImportData()
+	: GroupName()
+	, bReplaceExisting(false)
+	, bSkipReadOnly(false)
+{
+
+}
+
+bool UAutomatedAssetImportData::IsValid() const
+{
+	// This data is valid if there is at least one filename to import, there is a valid destination path
+	// and either no factory was supported (automatic factory picking) or a valid factory was found
+	return Filenames.Num() > 0 && !DestinationPath.IsEmpty() && (FactoryName.IsEmpty() || Factory != nullptr);
+}
+
+void UAutomatedAssetImportData::Initialize()
+{
+	if(Filenames.Num() > 0)
+	{
+		// IF a factory doesn't have a vaild full path assume it is an unreal internal factory
+		if(!FactoryName.IsEmpty() && !FactoryName.StartsWith("/Script/"))
+		{
+			FactoryName = TEXT("/Script/UnrealEd.") + FactoryName;
+		}
+
+		if(!FactoryName.IsEmpty())
+		{
+			UClass* FactoryClass = LoadClass<UObject>(nullptr, *FactoryName, nullptr, LOAD_None, nullptr);
+			if(FactoryClass)
+			{
+				UFactory* NewFactory = NewObject<UFactory>(GetTransientPackage(), FactoryClass);
+
+				if(NewFactory->bEditorImport)
+				{
+					// Check that every file can be imported
+					TArray<FString> InvalidFilesForFactory;
+					for(const FString& Filename : Filenames)
+					{
+						if(!NewFactory->FactoryCanImport(Filename))
+						{
+							InvalidFilesForFactory.Add(Filename);
+						}
+					}
+
+					if(InvalidFilesForFactory.Num() > 0)
+					{
+						UE_LOG(LogAutomatedImport, Error, TEXT("Factory %s could not import one or more of the source files"), *FactoryName);
+						for(const FString& InvalidFiles : InvalidFilesForFactory)
+						{
+							UE_LOG(LogAutomatedImport, Error, TEXT("    %s"), *InvalidFiles);
+						}
+					}
+					else
+					{
+						// All files are valid. Use this factory
+						Factory = NewFactory;
+					}
+				}
+			}
+			else
+			{
+				UE_LOG(LogAutomatedImport, Error, TEXT("Factory %s could not be found"), *FactoryName);
+			}
+		}
+		else
+		{
+			UE_LOG(LogAutomatedImport, Log, TEXT("Factory was not specified, will be set automatically"));
+		}
+	}
+
+	if(!DestinationPath.IsEmpty() && !DestinationPath.StartsWith("/Game"))
+	{
+		DestinationPath = TEXT("/Game") / DestinationPath;
+	}
+
+	FString PackagePath;
+	FString FailureReason;
+	if(!FPackageName::TryConvertFilenameToLongPackageName(DestinationPath, PackagePath, &FailureReason))
+	{
+		UE_LOG(LogAutomatedImport, Error, TEXT("Invalid Destination Path (%s): %s"), *DestinationPath, *FailureReason);
+		DestinationPath.Empty();
+	}
+	else
+	{
+		// Package path is what we will use for importing.  So use that as the dest path now that it has been created
+		DestinationPath = PackagePath;
+	}
+}
+
+FString UAutomatedAssetImportData::GetDisplayName() const
+{
+	return GroupName.Len() > 0 ? GroupName : GetName();
+}
diff --git a/Engine/Source/Editor/UnrealEd/Private/Commandlets/GatherTextFromSourceCommandlet.cpp b/Engine/Source/Editor/UnrealEd/Private/Commandlets/GatherTextFromSourceCommandlet.cpp
--- a/Engine/Source/Editor/UnrealEd/Private/Commandlets/GatherTextFromSourceCommandlet.cpp
+++ b/Engine/Source/Editor/UnrealEd/Private/Commandlets/GatherTextFromSourceCommandlet.cpp
@@ -666,10 +666,9 @@ bool UGatherTextFromSourceCommandlet::ParseSourceText(const FString& Text, const
 			// Check if we're starting comments or string literals. Begins *at* "//" or "/*".
 			if (!ParseCtxt.WithinLineComment && !ParseCtxt.WithinBlockComment && !ParseCtxt.WithinStringLiteral)
 			{
-				const TCHAR* ForwardCursor = Cursor;
-				if(*ForwardCursor == TEXT('/'))
+				if(*Cursor == TEXT('/'))
 				{
-					++ForwardCursor;
+					const TCHAR* const ForwardCursor = Cursor + 1;
 					switch(*ForwardCursor)
 					{
 					case TEXT('/'):
@@ -682,7 +681,7 @@ bool UGatherTextFromSourceCommandlet::ParseSourceText(const FString& Text, const
 						{
 							ParseCtxt.WithinBlockComment = true;
 							ParseCtxt.WithinStartingLine = ParseCtxt.LineText;
-					}
+						}
 						break;
 					}
 				}
@@ -690,79 +689,92 @@ bool UGatherTextFromSourceCommandlet::ParseSourceText(const FString& Text, const
 
 			if (!ParseCtxt.WithinLineComment && !ParseCtxt.WithinBlockComment && !ParseCtxt.WithinStringLiteral)
 			{
-				if (*Cursor == TEXT('\"') && Cursor >= *Line)
+				if (*Cursor == TEXT('\"'))
 				{
-					const TCHAR* ReverseCursor = Cursor;
-					--ReverseCursor;
-					if ((*ReverseCursor != TEXT('\\') && *ReverseCursor != TEXT('\'') && ReverseCursor >= *Line) || ReverseCursor < *Line)
+					if (Cursor == *Line)
 					{
 						ParseCtxt.WithinStringLiteral = true;
 						ParseCtxt.WithinStartingLine = ParseCtxt.LineText;
 					}
-					else 
+					else if (Cursor > *Line)
 					{
-						bool IsEscaped = false;
-						//if the backslash or single quote is itself escaped then the quote is good
-						while (*(--ReverseCursor) == TEXT('\\') && ReverseCursor >= *Line)
-						{
-							IsEscaped = !IsEscaped;
-						}
-
-						if (IsEscaped)
+						const TCHAR* const ReverseCursor = Cursor - 1;
+						if (*ReverseCursor != TEXT('\\') && *ReverseCursor != TEXT('\''))
 						{
 							ParseCtxt.WithinStringLiteral = true;
 							ParseCtxt.WithinStartingLine = ParseCtxt.LineText;
 						}
-						else
+						else 
 						{
-							//   check for '"'
-							ReverseCursor = Cursor;
-							--ReverseCursor;
-							const TCHAR* ForwardCursor = Cursor;
-							++ForwardCursor;
-							if (*ReverseCursor == TEXT('\'') && *ForwardCursor != TEXT('\'') && ReverseCursor >= *Line)
+							bool IsEscaped = false;
+							{
+								//if the backslash or single quote is itself escaped then the quote is good
+								const TCHAR* EscapeCursor = ReverseCursor;
+								while (EscapeCursor > *Line && *(--EscapeCursor) == TEXT('\\'))
+								{
+									IsEscaped = !IsEscaped;
+								}
+							}
+
+							if (IsEscaped)
 							{
 								ParseCtxt.WithinStringLiteral = true;
 								ParseCtxt.WithinStartingLine = ParseCtxt.LineText;
 							}
+							else
+							{
+								//   check for '"'
+								const TCHAR* const ForwardCursor = Cursor + 1;
+								if (*ReverseCursor == TEXT('\'') && *ForwardCursor != TEXT('\''))
+								{
+									ParseCtxt.WithinStringLiteral = true;
+									ParseCtxt.WithinStartingLine = ParseCtxt.LineText;
+								}
+							}
 						}
 					}
 				}
 			}
 			else if (ParseCtxt.WithinStringLiteral)
 			{
-				const TCHAR* ReverseCursor = Cursor;
-				if (*ReverseCursor == TEXT('\"') && ReverseCursor >= *Line)
+				if (*Cursor == TEXT('\"'))
 				{
-					--ReverseCursor;
-					if ((*ReverseCursor != TEXT('\\') && *ReverseCursor != TEXT('\'') && ReverseCursor >= *Line) || ReverseCursor < *Line)
+					if (Cursor == *Line)
 					{
 						ParseCtxt.WithinStringLiteral = false;
 					}
-					else
+					else if (Cursor > *Line)
 					{
-						bool IsEscaped = false;
-						//if the backslash or single quote is itself escaped then the quote is good
-						while (*(--ReverseCursor) == TEXT('\\') && ReverseCursor >= *Line)
-						{
-							IsEscaped = !IsEscaped;
-						}
-
-						if (IsEscaped)
+						const TCHAR* const ReverseCursor = Cursor - 1;
+						if (*ReverseCursor != TEXT('\\') && *ReverseCursor != TEXT('\''))
 						{
 							ParseCtxt.WithinStringLiteral = false;
 						}
 						else
 						{
-							//   check for '"'
-							ReverseCursor = Cursor;
-							--ReverseCursor;
-							const TCHAR* ForwardCursor = Cursor;
-							++ForwardCursor;
-							if (*ReverseCursor == TEXT('\'') && *ForwardCursor != TEXT('\'') && ReverseCursor >= *Line)
+							bool IsEscaped = false;
+							{
+								//if the backslash or single quote is itself escaped then the quote is good
+								const TCHAR* EscapeCursor = ReverseCursor;
+								while (EscapeCursor > *Line && *(--EscapeCursor) == TEXT('\\'))
+								{
+									IsEscaped = !IsEscaped;
+								}
+							}
+
+							if (IsEscaped)
 							{
 								ParseCtxt.WithinStringLiteral = false;
 							}
+							else
+							{
+								//   check for '"'
+								const TCHAR* const ForwardCursor = Cursor + 1;
+								if (*ReverseCursor == TEXT('\'') && *ForwardCursor != TEXT('\''))
+								{
+									ParseCtxt.WithinStringLiteral = false;
+								}
+							}
 						}
 					}
 				}
@@ -771,11 +783,10 @@ bool UGatherTextFromSourceCommandlet::ParseSourceText(const FString& Text, const
 			// Check if we're ending comments. Ends *after* "*/".
 			if(ParseCtxt.WithinBlockComment)
 			{
-				const TCHAR* ReverseCursor = Cursor;
-				if(*ReverseCursor == TEXT('/') && ReverseCursor >= *Line)
+				if(*Cursor == TEXT('/') && Cursor > *Line)
 				{
-					--ReverseCursor;
-					if(*ReverseCursor == TEXT('*')  && ReverseCursor >= *Line)
+					const TCHAR* const ReverseCursor = Cursor - 1;
+					if (*ReverseCursor == TEXT('*'))
 					{
 						ParseCtxt.WithinBlockComment = false;
 					}
diff --git a/Engine/Source/Editor/UnrealEd/Private/Commandlets/ImportAssetsCommandlet.cpp b/Engine/Source/Editor/UnrealEd/Private/Commandlets/ImportAssetsCommandlet.cpp
new file mode 100644
--- /dev/null
+++ b/Engine/Source/Editor/UnrealEd/Private/Commandlets/ImportAssetsCommandlet.cpp
@@ -0,0 +1,348 @@
+// Copyright 1998-2016 Epic Games, Inc. All Rights Reserved.
+
+#include "UnrealEd.h"
+#include "AutomatedAssetImportData.h"
+#include "ImportAssetsCommandlet.h"
+#include "AssetToolsModule.h"
+#include "JsonObjectConverter.h"
+
+
+static void PrintUsage()
+{
+	UE_LOG(LogAutomatedImport, Display, TEXT("LogAutomatedImport Usage: LogAutomatedImport {arglist}"));
+	UE_LOG(LogAutomatedImport, Display, TEXT("Arglist:"));
+
+	UE_LOG(LogAutomatedImport, Display, TEXT("-help or -?"));
+	UE_LOG(LogAutomatedImport, Display, TEXT("\tDisplays this help"));
+
+	UE_LOG(LogAutomatedImport, Display, TEXT("-source=\"path\""));
+	UE_LOG(LogAutomatedImport, Display, TEXT("\tThe source file to import.  This must be specified when importing a single asset\n[IGNORED when using -importparams]"));
+
+	UE_LOG(LogAutomatedImport, Display, TEXT("-dest=\"path/optional_asset_name\""));
+	UE_LOG(LogAutomatedImport, Display, TEXT("\tThe destination path in the project's content directory to import to with an optional asset name.\nIf the asset name is not specified the name of the source file will be used as the asset name.\nThis must be specified when importing a single asset.\n[IGNORED when using -importparams]"));
+
+	UE_LOG(LogAutomatedImport, Display, TEXT("-factory={factory class name}"));
+	UE_LOG(LogAutomatedImport, Display, TEXT("\tForces the asset to be opened with a specific UFactory class type.  If not specified import type will be auto detected.\n[IGNORED when using -importparams]"));
+
+	UE_LOG(LogAutomatedImport, Display, TEXT("-importsettings=\"path to import settings json file\""));
+	UE_LOG(LogAutomatedImport, Display, TEXT("\tPath to a json file that has asset import parameters when importing multiple files. If this argument is used all other import arguments are ignored as they are specified in the json file"));
+
+	UE_LOG(LogAutomatedImport, Display, TEXT("-replaceexisting"));
+	UE_LOG(LogAutomatedImport, Display, TEXT("\tWhether or not to replace existing assets when importing"));
+
+	UE_LOG(LogAutomatedImport, Display, TEXT("-nosourcecontrol"));
+	UE_LOG(LogAutomatedImport, Display, TEXT("\tDisables source control.  Prevents checking out, adding files, and submitting files"));
+
+	UE_LOG(LogAutomatedImport, Display, TEXT("-submitdesc"));
+	UE_LOG(LogAutomatedImport, Display, TEXT("\tSubmit description/comment to use checking in to source control.  If this is empty no files will be submitted"));
+
+	UE_LOG(LogAutomatedImport, Display, TEXT("-skipreadonly"));
+	UE_LOG(LogAutomatedImport, Display, TEXT("\tIf an asset cannot be saved because it is read only, the commandlet will not clear the read only flag and will not save the file"));
+}
+
+UImportAssetsCommandlet::UImportAssetsCommandlet(const FObjectInitializer& ObjectInitializer)
+	: Super(ObjectInitializer)
+{
+}
+
+bool UImportAssetsCommandlet::ParseParams(const FString& InParams)
+{
+	TArray<FString> Tokens;
+	TArray<FString> Params;
+	TMap<FString, FString> ParamVals;
+
+	ParseCommandLine(*InParams, Tokens, Params, ParamVals);
+
+	const bool bEnoughParams = ParamVals.Num() > 1;
+	if( Params.Contains(TEXT("?")) || Params.Contains(TEXT("help") ) )
+	{
+		bShowHelp = true;
+	}
+
+	bAllowSourceControl = !Params.Contains(TEXT("nosourcecontrol"));
+
+	GlobalImportData = NewObject<UAutomatedAssetImportData>(this);
+
+	GlobalImportData->bSkipReadOnly = Params.Contains(TEXT("skipreadonly"));
+
+	FString SourcePathParam = ParamVals.FindRef(TEXT("source"));
+	if(!SourcePathParam.IsEmpty())
+	{
+		GlobalImportData->Filenames.Add(SourcePathParam);
+	}
+	
+	GlobalImportData->DestinationPath = ParamVals.FindRef(TEXT("dest"));
+
+	GlobalImportData->FactoryName = ParamVals.FindRef(TEXT("factoryname"));
+
+	GlobalImportData->bReplaceExisting = Params.Contains(TEXT("replaceexisting"));
+
+	ImportSettingsPath = ParamVals.FindRef(TEXT("importsettings"));
+
+	GlobalImportData->Initialize();
+
+	if(ImportSettingsPath.IsEmpty() && (GlobalImportData->Filenames.Num() == 0 || GlobalImportData->DestinationPath.IsEmpty()))
+	{
+		UE_LOG(LogAutomatedImport, Error, TEXT("Invalid Arguments.  Missing, Source (-source), Destination (-dest), or Import settings file (-importsettings)"));
+	}
+
+	return bEnoughParams;
+}
+
+bool UImportAssetsCommandlet::ParseImportSettings(const FString& InImportSettingsFile)
+{
+	bool bInvalidParse = false;
+	bool bSuccess = false;
+	FString JsonString;
+	if(FFileHelper::LoadFileToString(JsonString, *InImportSettingsFile))
+	{
+		TSharedRef<TJsonReader<>> JsonReader = TJsonReaderFactory<>::Create(JsonString);
+		TSharedPtr<FJsonObject> RootObject;
+		if(FJsonSerializer::Deserialize(JsonReader, RootObject) && RootObject.IsValid())
+		{
+			const TArray< TSharedPtr<FJsonValue> > ImportGroupsJsonArray = RootObject->GetArrayField(TEXT("ImportGroups"));
+			for(const TSharedPtr<FJsonValue>& ImportGroupsJson : ImportGroupsJsonArray)
+			{
+				const TSharedPtr<FJsonObject> ImportGroupsJsonObject = ImportGroupsJson->AsObject();
+				if(ImportGroupsJsonObject.IsValid())
+				{
+					// All import data is based off of the global data defaults
+					UAutomatedAssetImportData* Data = DuplicateObject<UAutomatedAssetImportData>(GlobalImportData, this);
+					
+					// Parse data from the json object
+					if(FJsonObjectConverter::JsonObjectToUStruct(ImportGroupsJsonObject.ToSharedRef(), UAutomatedAssetImportData::StaticClass(), Data, 0, 0 ))
+					{
+						Data->Initialize();
+						if(Data->IsValid())
+						{
+							ImportDataList.Add(Data);
+						}
+
+						UFactory* Factory = Data->Factory;
+						const TSharedPtr<FJsonObject>* ImportSettingsJsonObject = nullptr;
+						ImportGroupsJsonObject->TryGetObjectField(TEXT("ImportSettings"), ImportSettingsJsonObject);
+						if(Factory != nullptr && ImportSettingsJsonObject)
+						{
+							IImportSettingsParser* ImportSettings = Factory->GetImportSettingsParser();
+							if(ImportSettings)
+							{
+								ImportSettings->ParseFromJson(ImportSettingsJsonObject->ToSharedRef());
+							}
+						}
+						else if(Factory == nullptr && ImportSettingsJsonObject)
+						{
+							UE_LOG(LogAutomatedImport, Warning, TEXT("A vaild factory name must be specfied in order to specify settings"));
+						}
+					}
+					else
+					{
+						bInvalidParse = true;
+					}
+					
+				}
+				else
+				{
+					bInvalidParse = true;
+				}
+			}
+		}
+		else
+		{
+			UE_LOG(LogAutomatedImport, Error, TEXT("Json settings file was found but was invalid: %s"), *JsonReader->GetErrorMessage());
+		}
+	}
+	else
+	{
+		UE_LOG(LogAutomatedImport, Error, TEXT("Import settings file %s could not be found"), *InImportSettingsFile);
+	}
+
+	return bSuccess;
+}
+
+static bool SavePackage(UPackage* Package, const FString& PackageFilename)
+{
+	return GEditor->SavePackage(Package, nullptr, RF_Standalone, *PackageFilename, GWarn);
+}
+
+bool UImportAssetsCommandlet::ImportAndSave(const TArray<UAutomatedAssetImportData*>& AssetImportList)
+{
+	bool bImportAndSaveSucceeded = true;
+
+	FAssetToolsModule& AssetToolsModule = FModuleManager::Get().LoadModuleChecked<FAssetToolsModule>("AssetTools");
+
+	ISourceControlProvider& SourceControlProvider = ISourceControlModule::Get().GetProvider();
+
+	for(const UAutomatedAssetImportData* ImportData : AssetImportList)
+	{
+		UE_LOG(LogAutomatedImport, Log, TEXT("Importing group %s"), *ImportData->GetDisplayName() );
+		TArray<UObject*> ImportedAssets = AssetToolsModule.Get().ImportAssetsAutomated(*ImportData);
+		if(ImportedAssets.Num() > 0)
+		{
+			TArray<UPackage*> DirtyPackages;
+			TArray<FSourceControlStateRef> PackageStates;
+
+			FEditorFileUtils::GetDirtyContentPackages(DirtyPackages);
+
+			bool bUseSourceControl = bHasSourceControl && SourceControlProvider.IsAvailable();
+			if(bUseSourceControl)
+			{
+				SourceControlProvider.GetState(DirtyPackages, PackageStates, EStateCacheUsage::ForceUpdate);
+			}
+
+			for(int32 PackageIndex = 0; PackageIndex < DirtyPackages.Num(); ++PackageIndex)
+			{
+				UPackage* PackageToSave = DirtyPackages[PackageIndex];
+				FSourceControlStateRef PackageSCState = PackageStates[PackageIndex];
+
+				FString PackageFilename = SourceControlHelpers::PackageFilename(PackageToSave);
+
+				bool bShouldAttemptToSave = false;
+				bool bShouldAttemptToAdd = false;
+				if(bUseSourceControl)
+				{
+					bool bPackageCanBeCheckedOut = false;
+					if(PackageSCState->IsCheckedOutOther())
+					{
+						// Cannot checkout, file is already checked out
+						UE_LOG(LogAutomatedImport, Error, TEXT("%s is already checked out by someone else, can not submit!"), *PackageFilename);
+						bImportAndSaveSucceeded = false;
+					}
+					else if(!PackageSCState->IsCurrent())
+					{
+						// Cannot checkout, file is not at head revision
+						UE_LOG(LogAutomatedImport, Error, TEXT("%s is not at the head revision and cannot be checked out"), *PackageFilename);
+						bImportAndSaveSucceeded = false;
+					}
+					else if(PackageSCState->CanCheckout())
+					{
+						const bool bWasCheckedOut = SourceControlHelpers::CheckOutFile(PackageFilename);
+						bShouldAttemptToSave = bWasCheckedOut;
+						if(!bWasCheckedOut)
+						{
+							UE_LOG(LogAutomatedImport, Error, TEXT("%s could not be checked out"), *PackageFilename);
+							bImportAndSaveSucceeded = false;
+						}
+					}
+					else
+					{
+						// package was not checked out by another user and is at the current head revision and could not be checked out
+						// this means it should be added after save because it doesnt exist
+						bShouldAttemptToSave = true;
+						bShouldAttemptToAdd = true;
+					}
+				}
+				else
+				{
+					bool bIsReadOnly = IFileManager::Get().IsReadOnly(*PackageFilename);
+					if(bIsReadOnly && ImportData->bSkipReadOnly)
+					{
+						bShouldAttemptToSave = ImportData->bSkipReadOnly;
+						if(bIsReadOnly)
+						{
+							UE_LOG(LogAutomatedImport, Error, TEXT("%s is read only and -skipreadonly was specified.  Will not save"), *PackageFilename);
+							bImportAndSaveSucceeded = false;
+						}
+					}
+					else if(bIsReadOnly)
+					{
+						bShouldAttemptToSave = FPlatformFileManager::Get().GetPlatformFile().SetReadOnly(*PackageFilename, false);
+						if(!bShouldAttemptToSave)
+						{
+							UE_LOG(LogAutomatedImport, Error, TEXT("%s is read only and could not be made writable.  Will not save"), *PackageFilename);
+							bImportAndSaveSucceeded = false;
+						}
+					}
+				}
+
+				if(bShouldAttemptToSave)
+				{
+					SavePackage(PackageToSave, PackageFilename);
+
+					if(bShouldAttemptToAdd)
+					{
+						const bool bWasAdded = SourceControlHelpers::MarkFileForAdd(PackageFilename);
+						if(!bWasAdded)
+						{
+							UE_LOG(LogAutomatedImport, Error, TEXT("%s could not be added to source control"), *PackageFilename);
+							bImportAndSaveSucceeded = false;
+						}
+					}
+				}
+			}
+		}
+		else
+		{
+			bImportAndSaveSucceeded = false;
+			UE_LOG(LogAutomatedImport, Error, TEXT("Failed to import all assets in group %s"), *ImportData->GetDisplayName());
+		}
+	}
+
+	return bImportAndSaveSucceeded;
+}
+
+void UImportAssetsCommandlet::ClearDirtyPackages()
+{
+	TArray<UPackage*> DirtyPackages;
+	FEditorFileUtils::GetDirtyContentPackages(DirtyPackages);
+	for(UPackage* Package : DirtyPackages)
+	{
+		Package->SetDirtyFlag(false);
+	}
+}
+
+int32 UImportAssetsCommandlet::Main(const FString& InParams)
+{
+	bool bEnoughParams = ParseParams(InParams);
+
+	int32 Result = 0;
+
+	if(!bEnoughParams || bShowHelp)
+	{
+		PrintUsage();
+	}
+	else
+	{
+		// Hack:  A huge amount of packages are marked dirty on startup.  This is normally prevented in editor but commandlets have special powers.  
+		// We only want to save assets that were created or modified at import time so clear all existing ones now
+		ClearDirtyPackages();
+
+		if(bAllowSourceControl)
+		{
+			ISourceControlProvider& SourceControlProvider = ISourceControlModule::Get().GetProvider();
+			SourceControlProvider.Init();
+
+			bHasSourceControl = SourceControlProvider.IsEnabled();
+			if(!bHasSourceControl)
+			{
+				UE_LOG(LogAutomatedImport, Error, TEXT("Could not connect to source control!"))
+			}
+		}
+		else
+		{
+			bHasSourceControl = false;
+		}
+
+
+		if(!ImportSettingsPath.IsEmpty())
+		{
+			// Use settings file for importing assets
+			ParseImportSettings(ImportSettingsPath);
+		}
+		else if(GlobalImportData->IsValid())
+		{
+			// Use single import path
+			ImportDataList.Add(GlobalImportData);
+		}
+
+		if(!ImportAndSave(ImportDataList))
+		{
+			UE_LOG(LogAutomatedImport, Error, TEXT("Could not import all groups"));
+		}
+		else
+		{
+			Result = 0;
+		}
+		
+	}
+	return Result;
+}
diff --git a/Engine/Source/Editor/UnrealEd/Private/EditorActor.cpp b/Engine/Source/Editor/UnrealEd/Private/EditorActor.cpp
--- a/Engine/Source/Editor/UnrealEd/Private/EditorActor.cpp
+++ b/Engine/Source/Editor/UnrealEd/Private/EditorActor.cpp
@@ -626,7 +626,7 @@ bool UUnrealEdEngine::CanDeleteSelectedActors( const UWorld* InWorld, const bool
 	return bContainsDeletable;
 }
 
-bool UUnrealEdEngine::edactDeleteSelected( UWorld* InWorld, bool bVerifyDeletionCanHappen)
+bool UUnrealEdEngine::edactDeleteSelected( UWorld* InWorld, bool bVerifyDeletionCanHappen, bool bWarnAboutReferences)
 {
 	if ( bVerifyDeletionCanHappen )
 	{
@@ -723,9 +723,12 @@ bool UUnrealEdEngine::edactDeleteSelected( UWorld* InWorld, bool bVerifyDeletion
 		ClassTypesToIgnore.Add( ALevelScriptActor::StaticClass() );
 		// The delete warning is meant for actor referneces that affect gameplay.  Group actors do not affect gameplay and should not show up as a warning.
 		ClassTypesToIgnore.Add( AGroupActor::StaticClass() );
-		FBlueprintEditorUtils::FindActorsThatReferenceActor( Actor, ClassTypesToIgnore, ReferencingActors );
+		if( bWarnAboutReferences )
+		{
+			FBlueprintEditorUtils::FindActorsThatReferenceActor( Actor, ClassTypesToIgnore, ReferencingActors );
+		}
 
-		bool bReferencedByLevelScript = (NULL != LSB && FBlueprintEditorUtils::FindNumReferencesToActorFromLevelScript(LSB, Actor) > 0);
+		bool bReferencedByLevelScript = bWarnAboutReferences && (nullptr != LSB && FBlueprintEditorUtils::FindNumReferencesToActorFromLevelScript(LSB, Actor) > 0);
 		bool bReferencedByActor = false;
 		bool bReferencedByLODActor = false;
 		for (AActor* ReferencingActor : ReferencingActors)
diff --git a/Engine/Source/Editor/UnrealEd/Private/EditorEngine.cpp b/Engine/Source/Editor/UnrealEd/Private/EditorEngine.cpp
--- a/Engine/Source/Editor/UnrealEd/Private/EditorEngine.cpp
+++ b/Engine/Source/Editor/UnrealEd/Private/EditorEngine.cpp
@@ -719,7 +719,6 @@ void UEditorEngine::Init(IEngineLoop* InEngineLoop)
 			TEXT("SizeMap"),
 			TEXT("MergeActors"),
 			TEXT("NiagaraEditor"),
-			TEXT("EditorAutomation"),
 		};
 
 		FScopedSlowTask ModuleSlowTask(ARRAY_COUNT(ModuleNames));
diff --git a/Engine/Source/Editor/UnrealEd/Private/EditorServer.cpp b/Engine/Source/Editor/UnrealEd/Private/EditorServer.cpp
--- a/Engine/Source/Editor/UnrealEd/Private/EditorServer.cpp
+++ b/Engine/Source/Editor/UnrealEd/Private/EditorServer.cpp
@@ -99,12 +99,12 @@ namespace
 		FWaveCluster(const TCHAR* InName)
 			:	Name( InName )
 			,	Num( 0 )
-			,	Size( 0 )
+			,	Size( )
 		{}
 
 		FString Name;
 		int32 Num;
-		int32 Size;
+		FResourceSizeEx Size;
 	};
 
 	struct FAnimSequenceUsageInfo
@@ -2887,7 +2887,9 @@ public:
 				*OutClipboardContents = *ScratchData;
 			}
 
-			GEditor->edactDeleteSelected( World, false );
+			// Do not warn if we are only copying not pasint as we are not actually deleting anything from the real level
+			const bool bWarnAboutReferencedActors = !bCopyOnly;
+			GEditor->edactDeleteSelected( World, false, bWarnAboutReferencedActors);
 		}
 
 		if( DestLevel )
@@ -6268,15 +6270,14 @@ bool UEditorEngine::HandleBugItGoCommand( const TCHAR* Str, FOutputDevice& Ar )
 bool UEditorEngine::HandleTagSoundsCommand( const TCHAR* Str, FOutputDevice& Ar )
 {
 	int32 NumObjects = 0;
-	int32 TotalSize = 0;
+	SIZE_T TotalSize = 0;
 	for( FObjectIterator It(USoundWave::StaticClass()); It; ++It )
 	{
 		++NumObjects;
 		DebugSoundAnnotation.Set(*It);
 
 		USoundWave* Wave = static_cast<USoundWave*>(*It);
-		const SIZE_T Size = Wave->GetResourceSize(EResourceSizeMode::Exclusive);
-		TotalSize += Size;
+		TotalSize += Wave->GetResourceSizeBytes(EResourceSizeMode::Exclusive);
 	}
 	UE_LOG(LogEditorServer, Log,  TEXT("Marked %i sounds %10.2fMB"), NumObjects, ((float)TotalSize) /(1024.f*1024.f) );
 	return true;
@@ -6315,19 +6316,20 @@ bool UEditorEngine::HandlecheckSoundsCommand( const TCHAR* Str, FOutputDevice& A
 		const int32 NumCoreClusters = Clusters.Num();
 
 		// Output information.
-		int32 TotalSize = 0;
 		UE_LOG(LogEditorServer, Log,  TEXT("=================================================================================") );
 		UE_LOG(LogEditorServer, Log,  TEXT("%60s %10s"), TEXT("Wave Name"), TEXT("Size") );
 		for ( int32 WaveIndex = 0 ; WaveIndex < WaveList.Num() ; ++WaveIndex )
 		{
 			USoundWave* Wave = WaveList[WaveIndex];
-			const SIZE_T WaveSize = Wave->GetResourceSize(EResourceSizeMode::Exclusive);
 			UPackage* WavePackage = Wave->GetOutermost();
 			const FString PackageName( WavePackage->GetName() );
 
+			FResourceSizeEx WaveResourceSize = FResourceSizeEx(EResourceSizeMode::Exclusive);
+			Wave->GetResourceSizeEx(WaveResourceSize);
+
 			// Totals.
 			Clusters[0].Num++;
-			Clusters[0].Size += WaveSize;
+			Clusters[0].Size += WaveResourceSize;
 
 			// Core clusters
 			for ( int32 ClusterIndex = 1 ; ClusterIndex < NumCoreClusters ; ++ClusterIndex )
@@ -6336,7 +6338,7 @@ bool UEditorEngine::HandlecheckSoundsCommand( const TCHAR* Str, FOutputDevice& A
 				if ( PackageName.Find( Cluster.Name ) != -1 )
 				{
 					Cluster.Num++;
-					Cluster.Size += WaveSize;
+					Cluster.Size += WaveResourceSize;
 				}
 			}
 
@@ -6349,7 +6351,7 @@ bool UEditorEngine::HandlecheckSoundsCommand( const TCHAR* Str, FOutputDevice& A
 				{
 					// Found a cluster with this package name.
 					Cluster.Num++;
-					Cluster.Size += WaveSize;
+					Cluster.Size += WaveResourceSize;
 					bFoundMatch = true;
 					break;
 				}
@@ -6359,17 +6361,17 @@ bool UEditorEngine::HandlecheckSoundsCommand( const TCHAR* Str, FOutputDevice& A
 				// Create a new cluster with the package name.
 				FWaveCluster NewCluster( *PackageName );
 				NewCluster.Num = 1;
-				NewCluster.Size = WaveSize;
+				NewCluster.Size = WaveResourceSize;
 				Clusters.Add( NewCluster );
 			}
 
 			// Dump bulk sound list.
-			UE_LOG(LogEditorServer, Log,  TEXT("%70s %10.2fk"), *Wave->GetPathName(), ((float)WaveSize)/1024.f );
+			UE_LOG(LogEditorServer, Log,  TEXT("%70s %10.2fk"), *Wave->GetPathName(), ((float)WaveResourceSize.GetTotalMemoryBytes())/1024.f );
 		}
 		UE_LOG(LogEditorServer, Log,  TEXT("=================================================================================") );
 		UE_LOG(LogEditorServer, Log,  TEXT("%60s %10s %10s"), TEXT("Cluster Name"), TEXT("Num"), TEXT("Size") );
 		UE_LOG(LogEditorServer, Log,  TEXT("=================================================================================") );
-		int32 TotalClusteredSize = 0;
+		FResourceSizeEx TotalClusteredSize;
 		for ( int32 ClusterIndex = 0 ; ClusterIndex < Clusters.Num() ; ++ClusterIndex )
 		{
 			const FWaveCluster& Cluster = Clusters[ClusterIndex];
@@ -6378,10 +6380,10 @@ bool UEditorEngine::HandlecheckSoundsCommand( const TCHAR* Str, FOutputDevice& A
 				UE_LOG(LogEditorServer, Log,  TEXT("---------------------------------------------------------------------------------") );
 				TotalClusteredSize += Cluster.Size;
 			}
-			UE_LOG(LogEditorServer, Log,  TEXT("%60s %10i %10.2fMB"), *Cluster.Name, Cluster.Num, ((float)Cluster.Size)/(1024.f*1024.f) );
+			UE_LOG(LogEditorServer, Log,  TEXT("%60s %10i %10.2fMB"), *Cluster.Name, Cluster.Num, ((float)Cluster.Size.GetTotalMemoryBytes())/(1024.f*1024.f) );
 		}
 		UE_LOG(LogEditorServer, Log,  TEXT("=================================================================================") );
-		UE_LOG(LogEditorServer, Log,  TEXT("Total Clusterd: %10.2fMB"), ((float)TotalClusteredSize)/(1024.f*1024.f) );
+		UE_LOG(LogEditorServer, Log,  TEXT("Total Clusterd: %10.2fMB"), ((float)TotalClusteredSize.GetTotalMemoryBytes())/(1024.f*1024.f) );
 		return true;
 }
 
diff --git a/Engine/Source/Editor/UnrealEd/Private/Factories/CSVImportFactory.cpp b/Engine/Source/Editor/UnrealEd/Private/Factories/CSVImportFactory.cpp
--- a/Engine/Source/Editor/UnrealEd/Private/Factories/CSVImportFactory.cpp
+++ b/Engine/Source/Editor/UnrealEd/Private/Factories/CSVImportFactory.cpp
@@ -14,12 +14,25 @@
 #include "Engine/UserDefinedStruct.h"
 #include "SCSVImportOptions.h"
 #include "DataTableEditorUtils.h"
+#include "Factories/ReimportDataTableFactory.h"
+#include "Factories/ReimportCurveTableFactory.h"
+#include "Factories/ReimportCurveFactory.h"
+#include "JsonObjectConverter.h"
+
 DEFINE_LOG_CATEGORY(LogCSVImportFactory);
 
 #define LOCTEXT_NAMESPACE "CSVImportFactory"
 
 //////////////////////////////////////////////////////////////////////////
 
+FCSVImportSettings::FCSVImportSettings()
+{
+	ImportRowStruct = nullptr;
+	ImportType = ECSVImportType::ECSV_DataTable;
+	ImportCurveInterpMode = ERichCurveInterpMode::RCIM_Linear;
+}
+
+
 static UClass* GetCurveClass( ECSVImportType ImportType )
 {
 	switch( ImportType )
@@ -67,7 +80,18 @@ bool UCSVImportFactory::DoesSupportClass(UClass * Class)
 	return (Class == UDataTable::StaticClass() || Class == UCurveTable::StaticClass() || Class == UCurveFloat::StaticClass() || Class == UCurveVector::StaticClass() || Class == UCurveLinearColor::StaticClass() );
 }
 
-UObject* UCSVImportFactory::FactoryCreateText( UClass* InClass, UObject* InParent, FName InName, EObjectFlags Flags, UObject* Context, const TCHAR* Type, const TCHAR*& Buffer, const TCHAR* BufferEnd, FFeedbackContext* Warn )
+bool UCSVImportFactory::FactoryCanImport(const FString& Filename)
+{
+	const FString Extension = FPaths::GetExtension(Filename);
+
+	if(Extension == TEXT("csv"))
+	{
+		return true;
+	}
+	return false;
+}
+
+UObject* UCSVImportFactory::FactoryCreateText(UClass* InClass, UObject* InParent, FName InName, EObjectFlags Flags, UObject* Context, const TCHAR* Type, const TCHAR*& Buffer, const TCHAR* BufferEnd, FFeedbackContext* Warn)
 {
 	FEditorDelegates::OnAssetPreImport.Broadcast(this, InClass, InParent, InName, Type);
 
@@ -80,9 +104,18 @@ UObject* UCSVImportFactory::FactoryCreateText( UClass* InClass, UObject* InParen
 	bool bHaveInfo = false;
 	UScriptStruct* ImportRowStruct = NULL;
 	ERichCurveInterpMode ImportCurveInterpMode = RCIM_Linear;
-
 	ECSVImportType ImportType = ECSVImportType::ECSV_DataTable;
-	if(ExistingTable != NULL)
+
+	if( IsAutomatedImport() )
+	{
+		ImportRowStruct = AutomatedImportSettings.ImportRowStruct;
+		ImportCurveInterpMode = AutomatedImportSettings.ImportCurveInterpMode;
+		ImportType = AutomatedImportSettings.ImportType;
+
+		// For automated import to work a row struct must be specified for a datatable type or a curve type must be specified
+		bHaveInfo = ImportRowStruct != nullptr || ImportType != ECSVImportType::ECSV_DataTable;
+	}
+	else if(ExistingTable != NULL)
 	{
 		ImportRowStruct = ExistingTable->RowStruct;
 		bHaveInfo = true;
@@ -101,7 +134,7 @@ UObject* UCSVImportFactory::FactoryCreateText( UClass* InClass, UObject* InParen
 	bool bDoImport = true;
 
 	// If we do not have the info we need, pop up window to ask for things
-	if(!bHaveInfo)
+	if(!bHaveInfo && !IsAutomatedImport())
 	{
 		TSharedPtr<SWindow> ParentWindow;
 		// Check if the main frame is loaded.  When using the old main frame it may not be.
@@ -130,6 +163,14 @@ UObject* UCSVImportFactory::FactoryCreateText( UClass* InClass, UObject* InParen
 		ImportCurveInterpMode = ImportOptionsWindow->GetSelectedCurveIterpMode();
 		bDoImport = ImportOptionsWindow->ShouldImport();
 	}
+	else if(!bHaveInfo && IsAutomatedImport())
+	{
+		if(ImportType == ECSVImportType::ECSV_DataTable && !ImportRowStruct)
+		{
+			UE_LOG(LogCSVImportFactory, Error, TEXT("A Data table row type must be specified in the import settings json file for automated import"));
+		}
+		bDoImport = false;
+	}
 
 	UObject* NewAsset = NULL;
 	if(bDoImport)
@@ -216,8 +257,11 @@ UObject* UCSVImportFactory::FactoryCreateText( UClass* InClass, UObject* InParen
 				AllProblems += TEXT("\n");
 			}
 
-			// Pop up any problems for user
-			FMessageDialog::Open( EAppMsgType::Ok, FText::FromString( AllProblems ) );
+			if(!IsAutomatedImport())
+			{
+				// Pop up any problems for user
+				FMessageDialog::Open(EAppMsgType::Ok, FText::FromString(AllProblems));
+			}
 		}
 	}
 
@@ -244,10 +288,15 @@ bool UCSVImportFactory::ReimportCSV( UObject* Obj )
 	return bHandled;
 }
 
-bool UCSVImportFactory::Reimport( UObject* Obj, const FString& Path )
+void UCSVImportFactory::ParseFromJson(TSharedRef<class FJsonObject> ImportSettingsJson)
+{	
+	FJsonObjectConverter::JsonObjectToUStruct(ImportSettingsJson, FCSVImportSettings::StaticStruct(), &AutomatedImportSettings, 0, 0);
+}
+
+bool UCSVImportFactory::Reimport(UObject* Obj, const FString& Path)
 {
 	if(Path.IsEmpty() == false)
-	{
+	{ 
 		FString FilePath = IFileManager::Get().ConvertToRelativePath(*Path);
 
 		FString Data;
@@ -274,16 +323,16 @@ TArray<FString> UCSVImportFactory::DoImportDataTable(UDataTable* TargetDataTable
 	return TargetDataTable->CreateTableFromCSVString(DataToImport);
 }
 
-TArray<FString> UCSVImportFactory::DoImportCurveTable(UCurveTable* TargetCurveTable, const FString& DataToImport, const ERichCurveInterpMode ImportCurveInterpMode)
+TArray<FString> UCSVImportFactory::DoImportCurveTable(UCurveTable* TargetCurveTable, const FString& DataToImport, const ERichCurveInterpMode InImportCurveInterpMode)
 {
 	// Are we importing JSON data?
 	const bool bIsJSON = CurrentFilename.EndsWith(TEXT(".json"));
 	if (bIsJSON)
 	{
-		return TargetCurveTable->CreateTableFromJSONString(DataToImport, ImportCurveInterpMode);
+		return TargetCurveTable->CreateTableFromJSONString(DataToImport, InImportCurveInterpMode);
 	}
 
-	return TargetCurveTable->CreateTableFromCSVString(DataToImport, ImportCurveInterpMode);
+	return TargetCurveTable->CreateTableFromCSVString(DataToImport, InImportCurveInterpMode);
 }
 
 TArray<FString> UCSVImportFactory::DoImportCurve(UCurveBase* TargetCurve, const FString& DataToImport)
@@ -436,3 +485,4 @@ int32 UReimportCurveFactory::GetPriority() const
 
 
 #undef LOCTEXT_NAMESPACE
+
diff --git a/Engine/Source/Editor/UnrealEd/Private/Factories/EditorFactories.cpp b/Engine/Source/Editor/UnrealEd/Private/Factories/EditorFactories.cpp
--- a/Engine/Source/Editor/UnrealEd/Private/Factories/EditorFactories.cpp
+++ b/Engine/Source/Editor/UnrealEd/Private/Factories/EditorFactories.cpp
@@ -117,6 +117,10 @@
 #include "ImageUtils.h"
 #include "Engine/PreviewMeshCollection.h"
 #include "Factories/PreviewMeshCollectionFactory.h"
+#include "Factories/CurveFactory.h"
+#include "Factories/CurveImportFactory.h"
+#include "Factories/DataAssetFactory.h"
+#include "Factories/DataTableFactory.h"
 
 DEFINE_LOG_CATEGORY_STATIC(LogEditorFactories, Log, All);
 
@@ -5692,7 +5696,8 @@ EReimportResult::Type UReimportFbxStaticMeshFactory::Reimport( UObject* Obj )
 		bool bShowOptionDialog = true;
 		bool bOutImportAll = false;
 		bool bIsObjFormat = false;
-		GetImportOptions( FFbxImporter, ReimportUI, bShowOptionDialog, Obj->GetPathName(), bOperationCanceled, bOutImportAll, bIsObjFormat, bForceImportType, FBXIT_StaticMesh );
+		bool bIsAutomated = false;
+		GetImportOptions( FFbxImporter, ReimportUI, bShowOptionDialog, bIsAutomated, Obj->GetPathName(), bOperationCanceled, bOutImportAll, bIsObjFormat, bForceImportType, FBXIT_StaticMesh );
 	}
 
 	if( !bOperationCanceled && ensure(ImportData) )
@@ -5934,14 +5939,13 @@ EReimportResult::Type UReimportFbxSkeletalMeshFactory::Reimport( UObject* Obj )
 		bool bForceImportType = true;
 		bool bOutImportAll = false;
 		bool bIsObjFormat = false;
-
-		// arggg... hate this different option class to confuse everybody
+		bool bIsAutomated = false;
 		// @hack to make sure skeleton is set before opening the dialog
 		ImportOptions->SkeletonForAnimation = SkeletalMesh->Skeleton;
 		ImportOptions->bCreatePhysicsAsset = false;
 		ImportOptions->PhysicsAsset = SkeletalMesh->PhysicsAsset;
 
-		ImportOptions = GetImportOptions( FFbxImporter, ReimportUI, bShowOptionDialog, Obj->GetPathName(), bOperationCanceled, bOutImportAll, bIsObjFormat, bForceImportType, FBXIT_SkeletalMesh );
+		ImportOptions = GetImportOptions( FFbxImporter, ReimportUI, bShowOptionDialog, bIsAutomated, Obj->GetPathName(), bOperationCanceled, bOutImportAll, bIsObjFormat, bForceImportType, FBXIT_SkeletalMesh );
 	}
 
 	if( !bOperationCanceled && ensure(ImportData) )
diff --git a/Engine/Source/Editor/UnrealEd/Private/Factories/Factory.cpp b/Engine/Source/Editor/UnrealEd/Private/Factories/Factory.cpp
--- a/Engine/Source/Editor/UnrealEd/Private/Factories/Factory.cpp
+++ b/Engine/Source/Editor/UnrealEd/Private/Factories/Factory.cpp
@@ -4,7 +4,7 @@
 #include "ObjectTools.h"
 #include "AssetToolsModule.h"
 #include "EditorClassUtils.h"
-
+#include "AutomatedAssetImportData.h"
 
 DEFINE_LOG_CATEGORY_STATIC(LogFactory, Log, All);
 
@@ -13,7 +13,6 @@ FString UFactory::CurrentFilename(TEXT(""));
 
 // This needs to be greater than 0 to allow factories to have both higher and lower priority than the default
 const int32 UFactory::DefaultImportPriority = 100;
-int32 UFactory::OverwriteYesOrNoToAllState = -1;
 bool UFactory::bAllowOneTimeWarningMessages = true;
 
 
@@ -239,7 +238,7 @@ UClass* UFactory::GetSupportedClass() const
 }
 
 
-bool UFactory::DoesSupportClass(UClass * Class)
+bool UFactory::DoesSupportClass(UClass* Class)
 {
 	return (Class == GetSupportedClass());
 }
@@ -264,8 +263,12 @@ void UFactory::ResetState()
 
 void UFactory::DisplayOverwriteOptionsDialog(const FText& Message)
 {
-	// Prompt the user for what to do if a 'To All' response wasn't already given.
-	if (OverwriteYesOrNoToAllState != EAppReturnType::YesAll && OverwriteYesOrNoToAllState != EAppReturnType::NoAll)
+	// if the asset importing is automated, get the override state from the automated settings from there because we cannot prompt
+	if(AutomatedImportData)
+	{
+		OverwriteYesOrNoToAllState =  AutomatedImportData->bReplaceExisting ? EAppReturnType::YesAll : EAppReturnType::NoAll;
+	}
+	else if (OverwriteYesOrNoToAllState != EAppReturnType::YesAll && OverwriteYesOrNoToAllState != EAppReturnType::NoAll)
 	{
 		OverwriteYesOrNoToAllState = FMessageDialog::Open(EAppMsgType::YesNoYesAllNoAllCancel, FText::Format(
 			NSLOCTEXT("UnrealEd", "ImportedAssetAlreadyExists", "{0} Would you like to overwrite the existing settings?\n\nYes or Yes to All: Overwrite the existing settings.\nNo or No to All: Preserve the existing settings.\nCancel: Abort the operation."),
@@ -571,3 +574,8 @@ FString UFactory::GetDefaultNewAssetName() const
 {
 	return FString(TEXT("New")) + GetSupportedClass()->GetName();
 }
+
+void UFactory::SetAutomatedAssetImportData(const UAutomatedAssetImportData* Data)
+{
+	AutomatedImportData = Data;
+}
\ No newline at end of file
