Fixed localization warnings and fixed text that would never be localized

diff --git a/Engine/Source/Editor/UnrealEd/Private/EditorViewportCommands.cpp b/Engine/Source/Editor/UnrealEd/Private/EditorViewportCommands.cpp
--- a/Engine/Source/Editor/UnrealEd/Private/EditorViewportCommands.cpp
+++ b/Engine/Source/Editor/UnrealEd/Private/EditorViewportCommands.cpp
@@ -30,21 +30,30 @@ void FEditorViewportCommands::RegisterCommands()
 	UI_COMMAND( TexStreamAccPrimitiveDistanceMode, "Primitive Distance Accuracy View Mode", "Visualize the accuracy of the primitive distance computed for texture streaming", EUserInterfaceActionType::RadioButton, FInputChord() );
 	UI_COMMAND( TexStreamAccMeshUVDensityMode, "Mesh UV Densities Accuracy View Mode", "Visualize the accuracy of the mesh UV densities computed for texture streaming", EUserInterfaceActionType::RadioButton, FInputChord() );
 	UI_COMMAND( TexStreamAccMeshUVDensityAll, "All UV Channels", "Visualize the densities accuracy of all UV channels", EUserInterfaceActionType::RadioButton, FInputChord() );
+
 	for (int32 TexCoordIndex = 0; TexCoordIndex < TEXSTREAM_MAX_NUM_UVCHANNELS; ++TexCoordIndex)
 	{
-		UI_COMMAND_Function( this, TexStreamAccMeshUVDensitySingle[TexCoordIndex], TEXT(LOCTEXT_NAMESPACE),
-			*FString::Printf(TEXT("TexStreamAccMeshUVDensitySingle%d"), TexCoordIndex), *FString::Printf(TEXT("TexStreamAccMeshUVDensitySingle%d_ToolTip"), TexCoordIndex), nullptr,
-			*FString::Printf(TEXT("UV Channel %d"), TexCoordIndex), *FString::Printf(TEXT("Visualize the size accuracy of UV density for channel %d"), TexCoordIndex),
-			EUserInterfaceActionType::RadioButton, FInputChord());
+		const FName CommandName = *FString::Printf(TEXT("ShowUVChannel%d"), TexCoordIndex);
+		const FText LocalizedName = FText::Format(LOCTEXT("ShowTexCoordCommands", "UV Channel {0}"), TexCoordIndex);
+		const FText LocalizedTooltip = FText::Format(LOCTEXT("ShowTexCoordCommands_ToolTip","Visualize the size accuracy of UV density for channel {0}"), TexCoordIndex);
+
+		TexStreamAccMeshUVDensitySingle[TexCoordIndex]
+			= FUICommandInfoDecl(this->AsShared(), CommandName, LocalizedName, LocalizedTooltip)
+				.UserInterfaceType(EUserInterfaceActionType::RadioButton);
 	}
+
 	UI_COMMAND( TexStreamAccMaterialTextureScaleMode, "Material Texture Scales Accuracy View Mode", "Visualize the accuracy of the material texture scales used for texture streaming", EUserInterfaceActionType::RadioButton, FInputChord() );
 	UI_COMMAND( TexStreamAccMaterialTextureScaleAll, "All Textures", "Visualize the scales accuracy of all textures", EUserInterfaceActionType::RadioButton, FInputChord() );
+
 	for (int32 TextureIndex = 0; TextureIndex < TEXSTREAM_MAX_NUM_TEXTURES_PER_MATERIAL; ++TextureIndex)
 	{
-		UI_COMMAND_Function( this, TexStreamAccMaterialTextureScaleSingle[TextureIndex], TEXT(LOCTEXT_NAMESPACE),
-			*FString::Printf(TEXT("TexStreamAccMaterialTextureScaleSingle%d"), TextureIndex), *FString::Printf(TEXT("TexStreamAccMaterialTextureScaleSingle%d_ToolTip"), TextureIndex), nullptr,
-			*FString::Printf(TEXT("Texture %d"), TextureIndex), *FString::Printf(TEXT("Visualize the scale accuracy of texture %d"), TextureIndex),
-			EUserInterfaceActionType::RadioButton, FInputChord());
+		const FName CommandName = *FString::Printf(TEXT("ShowTexture%d"), TextureIndex);
+		const FText LocalizedName = FText::Format(LOCTEXT("ShowTextureCommands", "Texture {0}"), TextureIndex);
+		const FText LocalizedTooltip = FText::Format(LOCTEXT("ShowTextureCommandss_ToolTip", "Visualize the scale accuracy of texture {0}"), TextureIndex);
+
+		TexStreamAccMaterialTextureScaleSingle[TextureIndex]
+			= FUICommandInfoDecl(this->AsShared(), CommandName, LocalizedName, LocalizedTooltip)
+				.UserInterfaceType(EUserInterfaceActionType::RadioButton);
 	}
 
 	UI_COMMAND( StationaryLightOverlapMode, "Stationary Light Overlap View Mode", "Visualizes overlap of stationary lights", EUserInterfaceActionType::RadioButton, FInputChord() );
