Fixed lightmaps being lost with cook-on-the-fly and old maps

diff --git a/Engine/Source/Editor/UnrealEd/Private/CookOnTheFlyServer.cpp b/Engine/Source/Editor/UnrealEd/Private/CookOnTheFlyServer.cpp
--- a/Engine/Source/Editor/UnrealEd/Private/CookOnTheFlyServer.cpp
+++ b/Engine/Source/Editor/UnrealEd/Private/CookOnTheFlyServer.cpp
@@ -1639,15 +1639,16 @@ uint32 UCookOnTheFlyServer::TickCookOnTheSide( const float TimeSlice, uint32 &Co
 
 				if (Package->ContainsMap())
 				{
+					// load sublevels
+					UWorld* World = UWorld::FindWorldInPackage(Package);
+					check(World);
+
+					World->PersistentLevel->HandleLegacyMapBuildData();
+
 					// child cookers can't process maps
 					// check(IsChildCooker() == false);
 					if ( IsCookByTheBookMode() )
 					{
-						// load sublevels
-						UWorld* World = UWorld::FindWorldInPackage(Package);
-
-						World->PersistentLevel->HandleLegacyMapBuildData();
-
 						// TArray<FString> PreviouslyCookedPackages;
 						if (World->StreamingLevels.Num())
 						{
