Prevent Launch On prompting to save a freshly-opened, non-modified map while still prompting the user to save never-saved maps

diff --git a/Engine/Source/Editor/UnrealEd/Private/PlayLevel.cpp b/Engine/Source/Editor/UnrealEd/Private/PlayLevel.cpp
--- a/Engine/Source/Editor/UnrealEd/Private/PlayLevel.cpp
+++ b/Engine/Source/Editor/UnrealEd/Private/PlayLevel.cpp
@@ -1953,17 +1953,23 @@ void UEditorEngine::PlayUsingLauncher()
 		FWorldContext & EditorContext = GetEditorWorldContext();
 		if (EditorContext.World()->WorldComposition || (LauncherProfile->GetCookMode() == ELauncherProfileCookModes::ByTheBookInEditor) || (LauncherProfile->GetCookMode() == ELauncherProfileCookModes::OnTheFlyInEditor) )
 		{
+			// Prompt the user to save the level if it has not been saved before. 
+			// An unmodified but unsaved blank template level does not appear in the dirty packages check below.
+			if (FEditorFileUtils::GetFilename(GWorld).Len() == 0)
+			{
+				if (!FEditorFileUtils::SaveCurrentLevel())
+				{
+					CancelRequestPlaySession();
+					return;
+				}
+			}
+
 			// Daniel: Only reason we actually need to save any packages is because if a new package is created it won't be on disk yet and CookOnTheFly will early out if the package doesn't exist (even though it could be in memory and not require loading at all)
 			//			future me can optimize this by either adding extra allowances to CookOnTheFlyServer code or only saving packages which doesn't exist if it becomes a problem
 			// if this returns false, it means we should stop what we're doing and return to the editor
 			bool bPromptUserToSave = true;
 			bool bSaveMapPackages = true;
 			bool bSaveContentPackages = true;
-			if (!FEditorFileUtils::SaveCurrentLevel())
-			{
-				CancelRequestPlaySession();
-				return;
-			}
 			if (!FEditorFileUtils::SaveDirtyPackages(bPromptUserToSave, bSaveMapPackages, bSaveContentPackages))
 			{
 				CancelRequestPlaySession();
