Fix non-unity compile errors

diff --git a/Engine/Source/Developer/GameplayDebugger/Private/Editor/GameplayDebuggerEdMode.cpp b/Engine/Source/Developer/GameplayDebugger/Private/Editor/GameplayDebuggerEdMode.cpp
--- a/Engine/Source/Developer/GameplayDebugger/Private/Editor/GameplayDebuggerEdMode.cpp
+++ b/Engine/Source/Developer/GameplayDebugger/Private/Editor/GameplayDebuggerEdMode.cpp
@@ -9,6 +9,9 @@
 #include "GameplayDebuggerPlayerManager.h"
 #include "GameplayDebuggerLocalController.h"
 
+#include "EditorModeManager.h"
+#include "EditorViewportClient.h"
+
 const FName FGameplayDebuggerEdMode::EM_GameplayDebugger = TEXT("EM_GameplayDebugger");
 
 void FGameplayDebuggerEdMode::Enter()
diff --git a/Engine/Source/Developer/GameplayDebugger/Private/Editor/GameplayDebuggerToolkit.h b/Engine/Source/Developer/GameplayDebugger/Private/Editor/GameplayDebuggerToolkit.h
--- a/Engine/Source/Developer/GameplayDebugger/Private/Editor/GameplayDebuggerToolkit.h
+++ b/Engine/Source/Developer/GameplayDebugger/Private/Editor/GameplayDebuggerToolkit.h
@@ -4,6 +4,7 @@
 
 #if WITH_EDITOR
 #include "Toolkits/BaseToolkit.h"
+#include "Layout/Visibility.h"
 
 class FGameplayDebuggerToolkit : public FModeToolkit
 {
diff --git a/Engine/Source/Developer/GameplayDebugger/Private/GameplayDebuggerLocalController.cpp b/Engine/Source/Developer/GameplayDebugger/Private/GameplayDebuggerLocalController.cpp
--- a/Engine/Source/Developer/GameplayDebugger/Private/GameplayDebuggerLocalController.cpp
+++ b/Engine/Source/Developer/GameplayDebugger/Private/GameplayDebuggerLocalController.cpp
@@ -15,6 +15,7 @@
 
 #if WITH_EDITOR
 #include "Editor/GameplayDebuggerEdMode.h"
+#include "EditorModeManager.h"
 #endif // WITH_EDITOR
 
 UGameplayDebuggerLocalController::UGameplayDebuggerLocalController(const FObjectInitializer& ObjectInitializer) : Super(ObjectInitializer)
diff --git a/Engine/Source/Editor/UnrealEd/Public/EdMode.h b/Engine/Source/Editor/UnrealEd/Public/EdMode.h
--- a/Engine/Source/Editor/UnrealEd/Public/EdMode.h
+++ b/Engine/Source/Editor/UnrealEd/Public/EdMode.h
@@ -2,6 +2,28 @@
 
 #pragma once
 
+#include "Editor.h"
+#include "EditorComponents.h"
+#include "EditorModeRegistry.h"
+#include "UnrealWidget.h"
+
+class FEditorViewportClient;
+struct FViewportClick;
+
+/** Outcomes when determining whether it's possible to perform an action on the edit modes*/
+namespace EEditAction
+{
+	enum Type
+	{
+		/** Can't process this action */
+		Skip		= 0,
+		/** Can process this action */
+		Process,
+		/** Stop evaluating other modes (early out) */
+		Halt,
+	};
+};
+
 /**
  * Base class for all editor modes.
  */
@@ -154,7 +176,7 @@ public:
 	void SelectNone();
 	virtual void SelectionChanged() {}
 
-	virtual bool HandleClick(FEditorViewportClient* InViewportClient, HHitProxy *HitProxy, const FViewportClick &Click);
+	virtual bool HandleClick(FEditorViewportClient* InViewportClient, HHitProxy* HitProxy, const FViewportClick& Click);
 
 	/** Handling SelectActor */
 	virtual bool Select( AActor* InActor, bool bInSelected ) { return 0; }
diff --git a/Engine/Source/Editor/UnrealEd/Public/EditorModes.h b/Engine/Source/Editor/UnrealEd/Public/EditorModes.h
--- a/Engine/Source/Editor/UnrealEd/Public/EditorModes.h
+++ b/Engine/Source/Editor/UnrealEd/Public/EditorModes.h
@@ -15,20 +15,6 @@ class FModeTool;
 
 DECLARE_LOG_CATEGORY_EXTERN(LogEditorModes, Log, All);
 
-/** Outcomes when determining whether it's possible to perform an action on the edit modes*/
-namespace EEditAction
-{
-	enum Type
-	{
-		/** Can't process this action */
-		Skip		= 0,
-		/** Can process this action */
-		Process,
-		/** Stop evaluating other modes (early out) */
-		Halt,
-	};
-};
-
 // Builtin editor mode constants
 struct UNREALED_API FBuiltinEditorModes
 {
