Copying //UE4/Dev-Core to //UE4/Dev-Main (Source: //UE4/Dev-Core @ 3208226)

diff --git a/Engine/Source/Developer/MaterialUtilities/Private/MaterialUtilities.cpp b/Engine/Source/Developer/MaterialUtilities/Private/MaterialUtilities.cpp
--- a/Engine/Source/Developer/MaterialUtilities/Private/MaterialUtilities.cpp
+++ b/Engine/Source/Developer/MaterialUtilities/Private/MaterialUtilities.cpp
@@ -24,6 +24,7 @@
 #include "MeshRendering.h"
 #include "MeshMergeData.h"
 #include "Runtime/Engine/Classes/Engine/StaticMesh.h"
+#include "UniquePtr.h"
 
 #if WITH_EDITOR
 #include "UnrealEd.h"
@@ -668,8 +669,8 @@ bool FMaterialUtilities::SupportsExport(EBlendMode InBlendMode, EMaterialPropert
 
 bool FMaterialUtilities::ExportMaterialProperty(UWorld* InWorld, UMaterialInterface* InMaterial, EMaterialProperty InMaterialProperty, UTextureRenderTarget2D* InRenderTarget, TArray<FColor>& OutBMP)
 {
-	TScopedPointer<FExportMaterialProxy> MaterialProxy(new FExportMaterialProxy(InMaterial, InMaterialProperty));
-	if (MaterialProxy == NULL)
+	TUniquePtr<FExportMaterialProxy> MaterialProxy(new FExportMaterialProxy(InMaterial, InMaterialProperty));
+	if (MaterialProxy == nullptr)
 	{
 		return false;
 	}
@@ -686,8 +687,8 @@ bool FMaterialUtilities::ExportMaterialProperty(UWorld* InWorld, UMaterialInterf
 
 bool FMaterialUtilities::ExportMaterialProperty(UWorld* InWorld, UMaterialInterface* InMaterial, EMaterialProperty InMaterialProperty, FIntPoint& OutSize, TArray<FColor>& OutBMP)
 {
-	TScopedPointer<FExportMaterialProxy> MaterialProxy(new FExportMaterialProxy(InMaterial, InMaterialProperty));
-	if (MaterialProxy == NULL)
+	TUniquePtr<FExportMaterialProxy> MaterialProxy(new FExportMaterialProxy(InMaterial, InMaterialProperty));
+	if (MaterialProxy == nullptr)
 	{
 		return false;
 	}
@@ -702,8 +703,8 @@ bool FMaterialUtilities::ExportMaterialProperty(UWorld* InWorld, UMaterialInterf
 
 bool FMaterialUtilities::ExportMaterialProperty(UMaterialInterface* InMaterial, EMaterialProperty InMaterialProperty, TArray<FColor>& OutBMP, FIntPoint& OutSize)
 {
-	TScopedPointer<FExportMaterialProxy> MaterialProxy(new FExportMaterialProxy(InMaterial, InMaterialProperty));
-	if (MaterialProxy == NULL)
+	TUniquePtr<FExportMaterialProxy> MaterialProxy(new FExportMaterialProxy(InMaterial, InMaterialProperty));
+	if (MaterialProxy == nullptr)
 	{
 		return false;
 	}
@@ -718,8 +719,8 @@ bool FMaterialUtilities::ExportMaterialProperty(UMaterialInterface* InMaterial,
 
 bool FMaterialUtilities::ExportMaterialProperty(UMaterialInterface* InMaterial, EMaterialProperty InMaterialProperty, FIntPoint InSize, TArray<FColor>& OutBMP)
 {
-	TScopedPointer<FExportMaterialProxy> MaterialProxy(new FExportMaterialProxy(InMaterial, InMaterialProperty));
-	if (MaterialProxy == NULL)
+	TUniquePtr<FExportMaterialProxy> MaterialProxy(new FExportMaterialProxy(InMaterial, InMaterialProperty));
+	if (MaterialProxy == nullptr)
 	{
 		return false;
 	}
diff --git a/Engine/Source/Developer/MeshBoneReduction/Private/MeshBoneReduction.cpp b/Engine/Source/Developer/MeshBoneReduction/Private/MeshBoneReduction.cpp
--- a/Engine/Source/Developer/MeshBoneReduction/Private/MeshBoneReduction.cpp
+++ b/Engine/Source/Developer/MeshBoneReduction/Private/MeshBoneReduction.cpp
@@ -4,6 +4,7 @@
 #include "MeshBoneReduction.h"
 #include "Developer/MeshUtilities/Public/MeshUtilities.h"
 #include "ComponentReregisterContext.h"
+#include "UniquePtr.h"
 
 class FMeshBoneReductionModule : public IMeshBoneReductionModule
 {
@@ -292,19 +293,19 @@ public:
 	}
 };
 
-TScopedPointer<FMeshBoneReduction> GMeshBoneReduction;
+TUniquePtr<FMeshBoneReduction> GMeshBoneReduction;
 
 void FMeshBoneReductionModule::StartupModule()
 {
-	GMeshBoneReduction = new FMeshBoneReduction();
+	GMeshBoneReduction = MakeUnique<FMeshBoneReduction>();
 }
 
 void FMeshBoneReductionModule::ShutdownModule()
 {
-	GMeshBoneReduction = NULL;
+	GMeshBoneReduction = nullptr;
 }
 
 IMeshBoneReduction* FMeshBoneReductionModule::GetMeshBoneReductionInterface()
 {
-	return GMeshBoneReduction;
+	return GMeshBoneReduction.Get();
 }
diff --git a/Engine/Source/Developer/MeshSimplifier/Private/QuadricMeshReduction.cpp b/Engine/Source/Developer/MeshSimplifier/Private/QuadricMeshReduction.cpp
--- a/Engine/Source/Developer/MeshSimplifier/Private/QuadricMeshReduction.cpp
+++ b/Engine/Source/Developer/MeshSimplifier/Private/QuadricMeshReduction.cpp
@@ -4,8 +4,8 @@
 #include "RawMesh.h"
 #include "MeshUtilities.h"
 #include "MeshBuild.h"
-
 #include "MeshSimplify.h"
+#include "UniquePtr.h"
 
 class FQuadricSimplifierMeshReductionModule : public IMeshReductionModule
 {
@@ -446,21 +446,21 @@ public:
 		return new FQuadricSimplifierMeshReduction;
 	}
 };
-TScopedPointer<FQuadricSimplifierMeshReduction> GQuadricSimplifierMeshReduction;
+TUniquePtr<FQuadricSimplifierMeshReduction> GQuadricSimplifierMeshReduction;
 
 void FQuadricSimplifierMeshReductionModule::StartupModule()
 {
-	GQuadricSimplifierMeshReduction = FQuadricSimplifierMeshReduction::Create();
+	GQuadricSimplifierMeshReduction.Reset(FQuadricSimplifierMeshReduction::Create());
 }
 
 void FQuadricSimplifierMeshReductionModule::ShutdownModule()
 {
-	GQuadricSimplifierMeshReduction = NULL;
+	GQuadricSimplifierMeshReduction = nullptr;
 }
 
 IMeshReduction* FQuadricSimplifierMeshReductionModule::GetStaticMeshReductionInterface()
 {
-	return GQuadricSimplifierMeshReduction;
+	return GQuadricSimplifierMeshReduction.Get();
 }
 
 IMeshReduction* FQuadricSimplifierMeshReductionModule::GetSkeletalMeshReductionInterface()
diff --git a/Engine/Source/Developer/Profiler/Private/ProfilerManager.cpp b/Engine/Source/Developer/Profiler/Private/ProfilerManager.cpp
--- a/Engine/Source/Developer/Profiler/Private/ProfilerManager.cpp
+++ b/Engine/Source/Developer/Profiler/Private/ProfilerManager.cpp
@@ -8,7 +8,7 @@
 #include "ProfilerSession.h"
 #include "SProfilerMiniView.h"
 #include "SProfilerWindow.h"
-
+#include "UniquePtr.h"
 
 #define LOCTEXT_NAMESPACE "FProfilerCommands"
 
@@ -218,7 +218,7 @@ protected:
 
 static int32 GetNumFrameFromCaptureSlow( const FString& ProfilerCaptureFilepath )
 {
-	TAutoPtr<FStatsHeaderReader> Instance( FStatsReader<FStatsHeaderReader>::Create( *ProfilerCaptureFilepath ) );
+	TUniquePtr<FStatsHeaderReader> Instance( FStatsReader<FStatsHeaderReader>::Create( *ProfilerCaptureFilepath ) );
 	return Instance->GetNumFrames();
 }
 
diff --git a/Engine/Source/Developer/Profiler/Private/ProfilerRawStatsForThreadView.cpp b/Engine/Source/Developer/Profiler/Private/ProfilerRawStatsForThreadView.cpp
--- a/Engine/Source/Developer/Profiler/Private/ProfilerRawStatsForThreadView.cpp
+++ b/Engine/Source/Developer/Profiler/Private/ProfilerRawStatsForThreadView.cpp
@@ -149,7 +149,7 @@ void FRawProfilerSession::PrepareLoading()
 		UE_LOG( LogStats, Error, TEXT( "Could not open: %s" ), *Filepath );
 		return;
 	}
-	TAutoPtr<FArchive> FileReader( IFileManager::Get().CreateFileReader( *Filepath ) );
+	TUniquePtr<FArchive> FileReader( IFileManager::Get().CreateFileReader( *Filepath ) );
 	if( !FileReader )
 	{
 		UE_LOG( LogStats, Error, TEXT( "Could not open: %s" ), *Filepath );
diff --git a/Engine/Source/Developer/SimplygonMeshReduction/Private/SimplygonMeshReduction.cpp b/Engine/Source/Developer/SimplygonMeshReduction/Private/SimplygonMeshReduction.cpp
--- a/Engine/Source/Developer/SimplygonMeshReduction/Private/SimplygonMeshReduction.cpp
+++ b/Engine/Source/Developer/SimplygonMeshReduction/Private/SimplygonMeshReduction.cpp
@@ -11,6 +11,7 @@
 #include "Components/SplineMeshComponent.h"
 #include "SimplygonSDK.h"
 #include "ScopedTimers.h"
+#include "UniquePtr.h"
 
 #include "MeshMergeData.h"
 
@@ -2743,18 +2744,18 @@ private:
 	}
 };
 
-TScopedPointer<FSimplygonMeshReduction> GSimplygonMeshReduction;
+TUniquePtr<FSimplygonMeshReduction> GSimplygonMeshReduction;
 
 
 void FSimplygonMeshReductionModule::StartupModule()
 {
-	GSimplygonMeshReduction = FSimplygonMeshReduction::Create();
+	GSimplygonMeshReduction.Reset(FSimplygonMeshReduction::Create());
 }
 
 void FSimplygonMeshReductionModule::ShutdownModule()
 {
 	FSimplygonMeshReduction::Destroy();
-	GSimplygonMeshReduction = NULL;
+	GSimplygonMeshReduction = nullptr;
 }
 
 #define USE_SIMPLYGON_SWARM 0
@@ -2762,7 +2763,7 @@ void FSimplygonMeshReductionModule::ShutdownModule()
 IMeshReduction* FSimplygonMeshReductionModule::GetStaticMeshReductionInterface()
 {
 #if !USE_SIMPLYGON_SWARM
-	return GSimplygonMeshReduction;
+	return GSimplygonMeshReduction.Get();
 #else
 	return nullptr;
 #endif
@@ -2771,7 +2772,7 @@ IMeshReduction* FSimplygonMeshReductionModule::GetStaticMeshReductionInterface()
 IMeshReduction* FSimplygonMeshReductionModule::GetSkeletalMeshReductionInterface()
 {
 #if !USE_SIMPLYGON_SWARM
-	return GSimplygonMeshReduction;
+	return GSimplygonMeshReduction.Get();
 #else
 	return nullptr;
 #endif
@@ -2780,9 +2781,9 @@ IMeshReduction* FSimplygonMeshReductionModule::GetSkeletalMeshReductionInterface
 IMeshMerging* FSimplygonMeshReductionModule::GetMeshMergingInterface()
 {
 #if !USE_SIMPLYGON_SWARM
-	return GSimplygonMeshReduction;
+	return GSimplygonMeshReduction.Get();
 #else
-return nullptr;
+	return nullptr;
 #endif
 }
 
diff --git a/Engine/Source/Developer/SimplygonSwarm/Private/SimplygonSwarm.cpp b/Engine/Source/Developer/SimplygonSwarm/Private/SimplygonSwarm.cpp
--- a/Engine/Source/Developer/SimplygonSwarm/Private/SimplygonSwarm.cpp
+++ b/Engine/Source/Developer/SimplygonSwarm/Private/SimplygonSwarm.cpp
@@ -2,6 +2,7 @@
 
 #include "SimplygonSwarmPrivatePCH.h"
 #include "SimplygonSwarmHelpers.h"
+#include "UniquePtr.h"
 THIRD_PARTY_INCLUDES_START
 #include <algorithm>
 THIRD_PARTY_INCLUDES_END
@@ -1639,12 +1640,12 @@ bool ZipContentsForUpload(FString InputDirectoryPath, FString OutputFileName)
 	}
 };
 
-TScopedPointer<FSimplygonSwarm> GSimplygonMeshReduction;
+TUniquePtr<FSimplygonSwarm> GSimplygonMeshReduction;
 
 
 void FSimplygonSwarmModule::StartupModule()
 {
-	GSimplygonMeshReduction = FSimplygonSwarm::Create();
+	GSimplygonMeshReduction.Reset(FSimplygonSwarm::Create());
 }
 
 void FSimplygonSwarmModule::ShutdownModule()
@@ -1664,7 +1665,7 @@ IMeshReduction* FSimplygonSwarmModule::GetSkeletalMeshReductionInterface()
 
 IMeshMerging* FSimplygonSwarmModule::GetMeshMergingInterface()
 {
-	return GSimplygonMeshReduction;
+	return GSimplygonMeshReduction.Get();
 }
 
 
diff --git a/Engine/Source/Developer/TargetPlatform/Private/TargetPlatformManagerModule.cpp b/Engine/Source/Developer/TargetPlatform/Private/TargetPlatformManagerModule.cpp
--- a/Engine/Source/Developer/TargetPlatform/Private/TargetPlatformManagerModule.cpp
+++ b/Engine/Source/Developer/TargetPlatform/Private/TargetPlatformManagerModule.cpp
@@ -696,8 +696,8 @@ RETRY_SETUPANDVALIDATE:
 		FString SDKInstallManifestFilePath = FPaths::Combine(*TargetSDKRoot, *SDKInstallManifestFileName);
 
 		// If we are using a manual install, then it is valid for there to be no OutputEnvVars file.
-		TAutoPtr<FArchive> InstallManifestFile(IFileManager::Get().CreateFileReader(*SDKInstallManifestFilePath));
-		if (InstallManifestFile.IsValid())
+		TUniquePtr<FArchive> InstallManifestFile(IFileManager::Get().CreateFileReader(*SDKInstallManifestFilePath));
+		if (InstallManifestFile)
 		{
 			TArray<FString> FileLines;
 			int64 FileSize = InstallManifestFile->TotalSize();
@@ -737,8 +737,8 @@ RETRY_SETUPANDVALIDATE:
 		FString EnvVarFileName = FPaths::Combine(*TargetSDKRoot, *SDKEnvironmentVarsFile);		
 
 		// If we are using a manual install, then it is valid for there to be no OutputEnvVars file.
-		TAutoPtr<FArchive> EnvVarFile(IFileManager::Get().CreateFileReader(*EnvVarFileName));
-		if (EnvVarFile.IsValid())
+		TUniquePtr<FArchive> EnvVarFile(IFileManager::Get().CreateFileReader(*EnvVarFileName));
+		if (EnvVarFile)
 		{
 			TArray<FString> FileLines;
 			{
@@ -753,8 +753,8 @@ RETRY_SETUPANDVALIDATE:
 				FileAsString.ParseIntoArrayLines(FileLines);
 
 				FMemory::Free(FileMem);
-				EnvVarFile->Close();				
-			}			
+				EnvVarFile->Close();
+			}
 
 			TArray<FString> PathAdds;
 			TArray<FString> PathRemoves;
diff --git a/Engine/Source/Editor/BlueprintGraph/Private/EdGraphSchema_K2.cpp b/Engine/Source/Editor/BlueprintGraph/Private/EdGraphSchema_K2.cpp
--- a/Engine/Source/Editor/BlueprintGraph/Private/EdGraphSchema_K2.cpp
+++ b/Engine/Source/Editor/BlueprintGraph/Private/EdGraphSchema_K2.cpp
@@ -3238,6 +3238,26 @@ bool UEdGraphSchema_K2::GetPropertyCategoryInfo(const UProperty* TestProperty, F
 			OutSubCategoryObject = ByteProperty->Enum;
 		}
 	}
+	else if (const UEnumProperty* EnumProperty = Cast<const UEnumProperty>(TestProperty))
+	{
+		// K2 only supports byte enums right now - any violations should have been caught by UHT or the editor
+		if (!EnumProperty->GetUnderlyingProperty()->IsA<UByteProperty>())
+		{
+			OutCategory = TEXT("unsupported_enum_type");
+			return false;
+		}
+
+		OutCategory = PC_Byte;
+
+		if (TestProperty->HasMetaData(FBlueprintMetadata::MD_Bitmask))
+		{
+			OutSubCategory = PSC_Bitmask;
+		}
+		else
+		{
+			OutSubCategoryObject = EnumProperty->GetEnum();
+		}
+	}
 	else if (Cast<const UNameProperty>(TestProperty) != NULL)
 	{
 		OutCategory = PC_Name;
diff --git a/Engine/Source/Editor/BlueprintGraph/Private/K2Node_CallFunction.cpp b/Engine/Source/Editor/BlueprintGraph/Private/K2Node_CallFunction.cpp
--- a/Engine/Source/Editor/BlueprintGraph/Private/K2Node_CallFunction.cpp
+++ b/Engine/Source/Editor/BlueprintGraph/Private/K2Node_CallFunction.cpp
@@ -705,22 +705,36 @@ void UK2Node_CallFunction::CreateExecPinsForFunctionCall(const UFunction* Functi
 		if(bWantsEnumToExecExpansion)
 		{
 			const FString& EnumParamName = Function->GetMetaData(FBlueprintMetadata::MD_ExpandEnumAsExecs);
-			UByteProperty* EnumProp = FindField<UByteProperty>(Function, FName(*EnumParamName));
-			if(EnumProp != NULL && EnumProp->Enum != NULL)
+
+			UProperty* Prop = nullptr;
+			UEnum* Enum = nullptr;
+
+			if(UByteProperty* ByteProp = FindField<UByteProperty>(Function, FName(*EnumParamName)))
+			{
+				Prop = ByteProp;
+				Enum = ByteProp->Enum;
+			}
+			else if(UEnumProperty* EnumProp = FindField<UEnumProperty>(Function, FName(*EnumParamName)))
+			{
+				Prop = EnumProp;
+				Enum = EnumProp->GetEnum();
+			}
+
+			if(Prop != nullptr && Enum != nullptr)
 			{
-				const bool bIsFunctionInput = !EnumProp->HasAnyPropertyFlags(CPF_ReturnParm) &&
-					(!EnumProp->HasAnyPropertyFlags(CPF_OutParm) ||
-					 EnumProp->HasAnyPropertyFlags(CPF_ReferenceParm));
+				const bool bIsFunctionInput = !Prop->HasAnyPropertyFlags(CPF_ReturnParm) &&
+					(!Prop->HasAnyPropertyFlags(CPF_OutParm) ||
+					 Prop->HasAnyPropertyFlags(CPF_ReferenceParm));
 				const EEdGraphPinDirection Direction = bIsFunctionInput ? EGPD_Input : EGPD_Output;
 				
 				// yay, found it! Now create exec pin for each
-				int32 NumExecs = (EnumProp->Enum->NumEnums() - 1);
+				int32 NumExecs = (Enum->NumEnums() - 1);
 				for(int32 ExecIdx=0; ExecIdx<NumExecs; ExecIdx++)
 				{
-					bool const bShouldBeHidden = EnumProp->Enum->HasMetaData(TEXT("Hidden"), ExecIdx) || EnumProp->Enum->HasMetaData(TEXT("Spacer"), ExecIdx);
+					bool const bShouldBeHidden = Enum->HasMetaData(TEXT("Hidden"), ExecIdx) || Enum->HasMetaData(TEXT("Spacer"), ExecIdx);
 					if (!bShouldBeHidden)
 					{
-						FString ExecName = EnumProp->Enum->GetEnumName(ExecIdx);
+						FString ExecName = Enum->GetEnumName(ExecIdx);
 						CreatePin(Direction, K2Schema->PC_Exec, TEXT(""), NULL, false, false, ExecName);
 					}
 				}
@@ -764,7 +778,7 @@ void UK2Node_CallFunction::DetermineWantsEnumToExecExpansion(const UFunction* Fu
 	{
 		const FString& EnumParamName = Function->GetMetaData(FBlueprintMetadata::MD_ExpandEnumAsExecs);
 		UByteProperty* EnumProp = FindField<UByteProperty>(Function, FName(*EnumParamName));
-		if(EnumProp != NULL && EnumProp->Enum != NULL)
+		if((EnumProp != NULL && EnumProp->Enum != NULL) || FindField<UEnumProperty>(Function, FName(*EnumParamName)))
 		{
 			bWantsEnumToExecExpansion = true;
 		}
@@ -1967,9 +1981,20 @@ void UK2Node_CallFunction::ExpandNode(class FKismetCompilerContext& CompilerCont
 		{
 			// Get the metadata that identifies which param is the enum, and try and find it
 			const FString& EnumParamName = Function->GetMetaData(FBlueprintMetadata::MD_ExpandEnumAsExecs);
-			UByteProperty* EnumProp = FindField<UByteProperty>(Function, FName(*EnumParamName));
+
+			UEnum* Enum = nullptr;
+
+			if (UByteProperty* ByteProp = FindField<UByteProperty>(Function, FName(*EnumParamName)))
+			{
+				Enum = ByteProp->Enum;
+			}
+			else if (UEnumProperty* EnumProp = FindField<UEnumProperty>(Function, FName(*EnumParamName)))
+			{
+				Enum = EnumProp->GetEnum();
+			}
+
 			UEdGraphPin* EnumParamPin = FindPinChecked(EnumParamName);
-			if(EnumProp != NULL && EnumProp->Enum != NULL)
+			if(Enum != nullptr)
 			{
 				// Expanded as input execs pins
 				if (EnumParamPin->Direction == EGPD_Input)
@@ -1980,7 +2005,7 @@ void UK2Node_CallFunction::ExpandNode(class FKismetCompilerContext& CompilerCont
 					// Create temp enum variable
 					UK2Node_TemporaryVariable* TempEnumVarNode = CompilerContext.SpawnIntermediateNode<UK2Node_TemporaryVariable>(this, SourceGraph);
 					TempEnumVarNode->VariableType.PinCategory = Schema->PC_Byte;
-					TempEnumVarNode->VariableType.PinSubCategoryObject = EnumProp->Enum;
+					TempEnumVarNode->VariableType.PinSubCategoryObject = Enum;
 					TempEnumVarNode->AllocateDefaultPins();
 					// Get the output pin
 					UEdGraphPin* TempEnumVarOutput = TempEnumVarNode->GetVariablePin();
diff --git a/Engine/Source/Editor/ContentBrowser/Private/AssetViewWidgets.cpp b/Engine/Source/Editor/ContentBrowser/Private/AssetViewWidgets.cpp
--- a/Engine/Source/Editor/ContentBrowser/Private/AssetViewWidgets.cpp
+++ b/Engine/Source/Editor/ContentBrowser/Private/AssetViewWidgets.cpp
@@ -1056,12 +1056,25 @@ void SAssetViewItem::CacheToolTipTags()
 					{
 						DisplayName = Field->GetDisplayNameText();
 
+						UProperty* Prop = nullptr;
+						UEnum* Enum = nullptr;
+						if (UByteProperty* ByteProp = Cast<UByteProperty>(Field))
+						{
+							Prop = ByteProp;
+							Enum = ByteProp->Enum;
+						}
+						else if (UEnumProperty* EnumProp = Cast<UEnumProperty>(Field))
+						{
+							Prop = EnumProp;
+							Enum = EnumProp->GetEnum();
+						}
+
 						// Strip off enum prefixes if they exist
-						if (UByteProperty* ByteProperty = Cast<UByteProperty>(Field))
+						if (Prop)
 						{
-							if (ByteProperty->Enum)
+							if (Enum)
 							{
-								const FString EnumPrefix = ByteProperty->Enum->GenerateEnumPrefix();
+								const FString EnumPrefix = Enum->GenerateEnumPrefix();
 								if (EnumPrefix.Len() && ValueString.StartsWith(EnumPrefix))
 								{
 									ValueString = ValueString.RightChop(EnumPrefix.Len() + 1);	// +1 to skip over the underscore
diff --git a/Engine/Source/Editor/DetailCustomizations/Private/MathStructCustomizations.cpp b/Engine/Source/Editor/DetailCustomizations/Private/MathStructCustomizations.cpp
--- a/Engine/Source/Editor/DetailCustomizations/Private/MathStructCustomizations.cpp
+++ b/Engine/Source/Editor/DetailCustomizations/Private/MathStructCustomizations.cpp
@@ -265,6 +265,21 @@ TSharedRef<SWidget> FMathStructCustomization::MakeChildWidget(
 		return MakeNumericWidget<uint8>(StructurePropertyHandle, PropertyHandle);
 	}
 
+	if (PropertyClass == UEnumProperty::StaticClass())
+	{
+		const UEnumProperty* EnumPropertyClass = static_cast<const UEnumProperty*>(PropertyHandle->GetProperty());
+		const UProperty* Enum = EnumPropertyClass->GetUnderlyingProperty();
+		const UClass* EnumClass = Enum->GetClass();
+		if (EnumClass == UByteProperty::StaticClass())
+		{
+			return MakeNumericWidget<uint8>(StructurePropertyHandle, PropertyHandle);
+		}
+		else if (EnumClass == UIntProperty::StaticClass())
+		{
+			return MakeNumericWidget<int32>(StructurePropertyHandle, PropertyHandle);
+		}
+	}
+
 	check(0); // Unsupported class
 	return SNullWidget::NullWidget;
 }
diff --git a/Engine/Source/Editor/FoliageEdit/Private/FoliageTypeCustomizationHelpers.cpp b/Engine/Source/Editor/FoliageEdit/Private/FoliageTypeCustomizationHelpers.cpp
--- a/Engine/Source/Editor/FoliageEdit/Private/FoliageTypeCustomizationHelpers.cpp
+++ b/Engine/Source/Editor/FoliageEdit/Private/FoliageTypeCustomizationHelpers.cpp
@@ -61,7 +61,7 @@ void FFoliageTypeCustomizationHelpers::BindHiddenPropertyVisibilityGetter(const
 			return bState && Result == FPropertyAccess::Success ? EVisibility::Visible : EVisibility::Collapsed;
 		});
 	}
-	else if (PropertyHandle->GetProperty()->IsA<UByteProperty>())
+	else if (PropertyHandle->GetProperty()->IsA<UByteProperty>() || PropertyHandle->GetProperty()->IsA<UEnumProperty>())
 	{
 		// If hidden behind a byte/enum, assume that 0 = disabled and show if nonzero
 		OutVisibilityGetter.BindLambda([=]
diff --git a/Engine/Source/Editor/GraphEditor/Private/SNodePanel.cpp b/Engine/Source/Editor/GraphEditor/Private/SNodePanel.cpp
--- a/Engine/Source/Editor/GraphEditor/Private/SNodePanel.cpp
+++ b/Engine/Source/Editor/GraphEditor/Private/SNodePanel.cpp
@@ -276,9 +276,9 @@ FVector2D SNodePanel::GetViewOffset() const
 
 void SNodePanel::Construct()
 {
-	if (!ZoomLevels.IsValid())
+	if (!ZoomLevels)
 	{
-		ZoomLevels = new FFixedZoomLevelsContainer();
+		ZoomLevels = MakeUnique<FFixedZoomLevelsContainer>();
 	}
 	ZoomLevel = ZoomLevels->GetDefaultZoomLevel();
 	PreviousZoomLevel = ZoomLevels->GetDefaultZoomLevel();
diff --git a/Engine/Source/Editor/GraphEditor/Public/SNodePanel.h b/Engine/Source/Editor/GraphEditor/Public/SNodePanel.h
--- a/Engine/Source/Editor/GraphEditor/Public/SNodePanel.h
+++ b/Engine/Source/Editor/GraphEditor/Public/SNodePanel.h
@@ -4,6 +4,7 @@
 
 #include "GraphEditorModule.h"
 #include "MarqueeOperation.h"
+#include "UniquePtr.h"
 
 //@TODO: Too generic of a name to expose at this scope
 typedef class UObject* SelectedItemType;
@@ -823,7 +824,7 @@ protected:
 	void CancelZoomToFit();
 protected:
 	// The interface for mapping ZoomLevel values to actual node scaling values
-	TScopedPointer<FZoomLevelsContainer> ZoomLevels;
+	TUniquePtr<FZoomLevelsContainer> ZoomLevels;
 
 	/** The position within the graph at which the user is looking */
 	FVector2D ViewOffset;
diff --git a/Engine/Source/Editor/Kismet/Private/BlueprintDetailsCustomization.cpp b/Engine/Source/Editor/Kismet/Private/BlueprintDetailsCustomization.cpp
--- a/Engine/Source/Editor/Kismet/Private/BlueprintDetailsCustomization.cpp
+++ b/Engine/Source/Editor/Kismet/Private/BlueprintDetailsCustomization.cpp
@@ -1684,6 +1684,7 @@ EVisibility FBlueprintVarActionDetails::ExposeToCinematicsVisibility() const
 	{
 		const bool bIsInteger = VariableProperty->IsA(UIntProperty::StaticClass());
 		const bool bIsByte = VariableProperty->IsA(UByteProperty::StaticClass());
+		const bool bIsEnum = VariableProperty->IsA(UEnumProperty::StaticClass());
 		const bool bIsFloat = VariableProperty->IsA(UFloatProperty::StaticClass());
 		const bool bIsBool = VariableProperty->IsA(UBoolProperty::StaticClass());
 		const bool bIsStr = VariableProperty->IsA(UStrProperty::StaticClass());
@@ -1692,7 +1693,7 @@ EVisibility FBlueprintVarActionDetails::ExposeToCinematicsVisibility() const
 		const bool bIsLinearColorStruct = VariableProperty->IsA(UStructProperty::StaticClass()) && Cast<UStructProperty>(VariableProperty)->Struct->GetFName() == NAME_LinearColor;
 		const bool bIsActorProperty = VariableProperty->IsA(UObjectProperty::StaticClass()) && Cast<UObjectProperty>(VariableProperty)->PropertyClass->IsChildOf(AActor::StaticClass());
 
-		if (bIsInteger || bIsByte || bIsFloat || bIsBool || bIsStr || bIsVectorStruct || bIsColorStruct || bIsLinearColorStruct || bIsActorProperty)
+		if (bIsInteger || bIsByte || bIsEnum || bIsFloat || bIsBool || bIsStr || bIsVectorStruct || bIsColorStruct || bIsLinearColorStruct || bIsActorProperty)
 		{
 			return EVisibility::Visible;
 		}
diff --git a/Engine/Source/Editor/KismetCompiler/Private/KismetCompilerMisc.cpp b/Engine/Source/Editor/KismetCompiler/Private/KismetCompilerMisc.cpp
--- a/Engine/Source/Editor/KismetCompiler/Private/KismetCompilerMisc.cpp
+++ b/Engine/Source/Editor/KismetCompiler/Private/KismetCompilerMisc.cpp
@@ -48,8 +48,9 @@ static bool IsTypeCompatibleWithProperty(UEdGraphPin* SourcePin, const FEdGraphP
 	}
 	else if (PinCategory == Schema->PC_Byte)
 	{
-		UByteProperty* SpecificProperty = Cast<UByteProperty>(TestProperty);
-		bTypeMismatch = (SpecificProperty == NULL);
+		UByteProperty* ByteProperty = Cast<UByteProperty>(TestProperty);
+		UEnumProperty* EnumProperty = Cast<UEnumProperty>(TestProperty);
+		bTypeMismatch = (ByteProperty == nullptr) && (EnumProperty == nullptr || !EnumProperty->GetUnderlyingProperty()->IsA<UByteProperty>());
 	}
 	else if ((PinCategory == Schema->PC_Class) || (PinCategory == Schema->PC_AssetClass))
 	{
@@ -599,19 +600,32 @@ void FKismetCompilerUtilities::ValidateEnumProperties(UObject* DefaultObject, FC
 	check(DefaultObject);
 	for (TFieldIterator<UProperty> It(DefaultObject->GetClass()); It; ++It)
 	{
-		const UByteProperty* ByteProperty = Cast<UByteProperty>(*It);
-		if(ByteProperty && !ByteProperty->HasAnyPropertyFlags(CPF_Transient))
+		UProperty* Property = *It;
+
+		if(!Property->HasAnyPropertyFlags(CPF_Transient))
 		{
-			const UEnum* Enum = ByteProperty->GetIntPropertyEnum();
+			const UEnum* Enum = nullptr;
+			const UNumericProperty* UnderlyingProp = nullptr;
+			if (const UEnumProperty* EnumProperty = Cast<UEnumProperty>(Property))
+			{
+				Enum = EnumProperty->GetEnum();
+				UnderlyingProp = EnumProperty->GetUnderlyingProperty();
+			}
+			else if(const UByteProperty* ByteProperty = Cast<UByteProperty>(Property))
+			{
+				Enum = ByteProperty->GetIntPropertyEnum();
+				UnderlyingProp = ByteProperty;
+			}
+
 			if(Enum)
-			{		
-				const uint8 EnumValue = ByteProperty->GetPropertyValue_InContainer(DefaultObject);
+			{
+				const int64 EnumValue = UnderlyingProp->GetSignedIntPropertyValue(Property->ContainerPtrToValuePtr<void>(DefaultObject));
 				if(!Enum->IsValidEnumValue(EnumValue))
 				{
 					MessageLog.Warning(
 						*FString::Printf(
-							*LOCTEXT("InvalidEnumDefaultValue_Error", "Default Enum value '%s' for class '%s' is invalid in object '%s'. EnumVal: %d. EnumAcceptableMax: %d ").ToString(),
-							*ByteProperty->GetName(),
+							*LOCTEXT("InvalidEnumDefaultValue_Error", "Default Enum value '%s' for class '%s' is invalid in object '%s'. EnumVal: %lld. EnumAcceptableMax: %lld ").ToString(),
+							*Property->GetName(),
 							*DefaultObject->GetClass()->GetName(),
 							*DefaultObject->GetName(),
 							EnumValue,
@@ -982,11 +996,26 @@ UProperty* FKismetCompilerUtilities::CreatePrimitiveProperty(UObject* PropertySc
 	}
 	else if (PinCategory == Schema->PC_Byte)
 	{
-		UByteProperty* ByteProp = NewObject<UByteProperty>(PropertyScope, ValidatedPropertyName, ObjectFlags);
-		ByteProp->Enum = Cast<UEnum>(PinSubCategoryObject);
-		ByteProp->SetPropertyFlags(CPF_HasGetValueTypeHash);
+		UEnum* Enum = Cast<UEnum>(PinSubCategoryObject);
+
+		if (Enum && Enum->GetCppForm() == UEnum::ECppForm::EnumClass)
+		{
+			UEnumProperty* EnumProp = new (EC_InternalUseOnlyConstructor, PropertyScope, ValidatedPropertyName, ObjectFlags) UEnumProperty(FObjectInitializer(), CastChecked<UEnum>(PinSubCategoryObject));
+			UNumericProperty* UnderlyingProp = NewObject<UByteProperty>(EnumProp, TEXT("UnderlyingType"), ObjectFlags);
+
+			EnumProp->AddCppProperty(UnderlyingProp);
 
-		NewProperty = ByteProp;
+			NewProperty = EnumProp;
+		}
+		else
+		{
+			UByteProperty* ByteProp = NewObject<UByteProperty>(PropertyScope, ValidatedPropertyName, ObjectFlags);
+			ByteProp->Enum = Cast<UEnum>(PinSubCategoryObject);
+
+			NewProperty = ByteProp;
+		}
+
+		NewProperty->SetPropertyFlags(CPF_HasGetValueTypeHash);
 	}
 	else if (PinCategory == Schema->PC_Name)
 	{
diff --git a/Engine/Source/Editor/KismetCompiler/Private/KismetCompilerVMBackend.cpp b/Engine/Source/Editor/KismetCompiler/Private/KismetCompilerVMBackend.cpp
--- a/Engine/Source/Editor/KismetCompiler/Private/KismetCompilerVMBackend.cpp
+++ b/Engine/Source/Editor/KismetCompiler/Private/KismetCompilerVMBackend.cpp
@@ -455,7 +455,7 @@ public:
 		{
 			if (Property)
 			{
-				return Property->IsA<UByteProperty>();
+				return Property->IsA<UByteProperty>() || (Property->IsA<UEnumProperty>() && static_cast<const UEnumProperty*>(Property)->GetUnderlyingProperty()->IsA<UByteProperty>());
 			}
 			return Type && ((Type->PinCategory == UEdGraphSchema_K2::PC_Byte) || (Type->PinCategory == UEdGraphSchema_K2::PC_Enum));
 		}
@@ -674,9 +674,23 @@ public:
 			{
 				uint8 Value = 0;
 
-				UByteProperty* ByteProp = Cast< UByteProperty >(CoerceProperty);
+				UEnum* EnumPtr = nullptr;
+
+				if (UByteProperty* ByteProp = Cast< UByteProperty >(CoerceProperty))
+				{
+					EnumPtr = ByteProp->Enum;
+				}
+				else if (UEnumProperty* EnumProp = Cast< UEnumProperty >(CoerceProperty))
+				{
+					EnumPtr = EnumProp->GetEnum();
+				}
+
 				//Parameter property can represent a generic byte. we need the actual type to parse the value.
-				UEnum* EnumPtr = (ByteProp && ByteProp->Enum) ? ByteProp->Enum : Cast<UEnum>(Term->Type.PinSubCategoryObject.Get()); 
+				if (!EnumPtr)
+				{
+					EnumPtr = Cast<UEnum>(Term->Type.PinSubCategoryObject.Get());
+				}
+
 				//Check for valid enum object reference
 				if (EnumPtr)
 				{
diff --git a/Engine/Source/Editor/MaterialEditor/Private/MaterialEditorUtilities.cpp b/Engine/Source/Editor/MaterialEditor/Private/MaterialEditorUtilities.cpp
--- a/Engine/Source/Editor/MaterialEditor/Private/MaterialEditorUtilities.cpp
+++ b/Engine/Source/Editor/MaterialEditor/Private/MaterialEditorUtilities.cpp
@@ -24,6 +24,7 @@
 #include "Materials/MaterialInstance.h"
 #include "Engine/TextureStreamingTypes.h"
 #include "MaterialUtilities.h"
+#include "UniquePtr.h"
 
 #define LOCTEXT_NAMESPACE "MaterialEditorUtilities"
 
@@ -202,9 +203,9 @@ void FMaterialEditorUtilities::GetVisibleMaterialParameters(const UMaterial* Mat
 {
 	VisibleExpressions.Empty();
 
-	TScopedPointer<FGetVisibleMaterialParametersFunctionState> FunctionState(new FGetVisibleMaterialParametersFunctionState(NULL));
+	TUniquePtr<FGetVisibleMaterialParametersFunctionState> FunctionState = MakeUnique<FGetVisibleMaterialParametersFunctionState>(nullptr);
 	TArray<FGetVisibleMaterialParametersFunctionState*> FunctionStack;
-	FunctionStack.Push(FunctionState.GetOwnedPointer());
+	FunctionStack.Push(FunctionState.Get());
 
 	for(uint32 i = 0; i < MP_MAX; ++i)
 	{
@@ -500,8 +501,8 @@ void FMaterialEditorUtilities::GetVisibleMaterialParametersFromExpression(
 				checkSlow(FunctionStack[FunctionCallIndex]->FunctionCall != FunctionCallExpression);
 			}
 
-			TScopedPointer<FGetVisibleMaterialParametersFunctionState> NewFunctionState(new FGetVisibleMaterialParametersFunctionState(FunctionCallExpression));
-			FunctionStack.Push(NewFunctionState.GetOwnedPointer());
+			TUniquePtr<FGetVisibleMaterialParametersFunctionState> NewFunctionState = MakeUnique<FGetVisibleMaterialParametersFunctionState>(FunctionCallExpression);
+			FunctionStack.Push(NewFunctionState.Get());
 			
 			GetVisibleMaterialParametersFromExpression(FMaterialExpressionKey(FunctionCallExpression->FunctionOutputs[MaterialExpressionKey.OutputIndex].ExpressionOutput, 0), MaterialInstance, VisibleExpressions, FunctionStack);
 		
diff --git a/Engine/Source/Editor/MeshPaint/Private/MeshPaintEdMode.cpp b/Engine/Source/Editor/MeshPaint/Private/MeshPaintEdMode.cpp
--- a/Engine/Source/Editor/MeshPaint/Private/MeshPaintEdMode.cpp
+++ b/Engine/Source/Editor/MeshPaint/Private/MeshPaintEdMode.cpp
@@ -39,6 +39,7 @@
 #include "VREditorInteractor.h"
 #include "VIBaseTransformGizmo.h"
 #include "EditorWorldManager.h"
+#include "UniquePtr.h"
 
 #define LOCTEXT_NAMESPACE "MeshPaint_Mode"
 
@@ -1486,8 +1487,8 @@ void FEdModeMeshPaint::PaintMeshVertices(
 	// Paint the mesh
 	uint32 NumVerticesInfluencedByBrush = 0;
 	{
-		TScopedPointer< FStaticMeshComponentRecreateRenderStateContext > RecreateRenderStateContext;
-		TScopedPointer< FComponentReregisterContext > ComponentReregisterContext;
+		TUniquePtr< FStaticMeshComponentRecreateRenderStateContext > RecreateRenderStateContext;
+		TUniquePtr< FComponentReregisterContext > ComponentReregisterContext;
 
 
 		FStaticMeshComponentLODInfo* InstanceMeshLODInfo = NULL;
@@ -1497,7 +1498,7 @@ void FEdModeMeshPaint::PaintMeshVertices(
 			{
 				// We're only changing instanced vertices on this specific mesh component, so we
 				// only need to detach our mesh component
-				ComponentReregisterContext.Reset( new FComponentReregisterContext( StaticMeshComponent ) );
+				ComponentReregisterContext = MakeUnique<FComponentReregisterContext>( StaticMeshComponent );
 
 				// Mark the mesh component as modified
 				StaticMeshComponent->SetFlags(RF_Transactional);
@@ -1563,7 +1564,7 @@ void FEdModeMeshPaint::PaintMeshVertices(
 			{
 				// We're changing the mesh itself, so ALL static mesh components in the scene will need
 				// to be unregistered for this (and reregistered afterwards.)
-				RecreateRenderStateContext.Reset( new FStaticMeshComponentRecreateRenderStateContext( StaticMesh ) );
+				RecreateRenderStateContext = MakeUnique<FStaticMeshComponentRecreateRenderStateContext>( StaticMesh );
 
 				// Dirty the mesh
 				StaticMesh->SetFlags(RF_Transactional);
@@ -3760,7 +3761,7 @@ void FEdModeMeshPaint::PasteInstanceVertexColors()
 
 	USelection& SelectedActors = *Owner->GetSelectedActors();
 
-	TScopedPointer< FComponentReregisterContext > ComponentReregisterContext;
+	TUniquePtr< FComponentReregisterContext > ComponentReregisterContext;
 
 	for( int32 ActorIndex = 0; ActorIndex < SelectedActors.Num(); ActorIndex++ )
 	{
@@ -3795,7 +3796,7 @@ void FEdModeMeshPaint::PasteInstanceVertexColors()
 
 					if(FoundColors != NULL)
 					{
-						ComponentReregisterContext.Reset( new FComponentReregisterContext( StaticMeshComponent ) );
+						ComponentReregisterContext = MakeUnique<FComponentReregisterContext>( StaticMeshComponent );
 						StaticMeshComponent->SetFlags(RF_Transactional);
 						StaticMeshComponent->Modify();
 						StaticMeshComponent->SetLODDataCount( NumLods, NumLods );
@@ -4340,8 +4341,8 @@ void FImportVertexTextureHelper::ImportVertexColors(FEditorModeTools* ModeTools,
 
 		FStaticMeshLODResources& LODModel = StaticMesh->RenderData->LODResources[ ImportLOD ];
 
-		TScopedPointer< FStaticMeshComponentRecreateRenderStateContext > RecreateRenderStateContext;
-		TScopedPointer< FComponentReregisterContext > ComponentReregisterContext;
+		TUniquePtr< FStaticMeshComponentRecreateRenderStateContext > RecreateRenderStateContext;
+		TUniquePtr< FComponentReregisterContext > ComponentReregisterContext;
 
 		FStaticMeshComponentLODInfo* InstanceMeshLODInfo = NULL;
 
@@ -4352,7 +4353,7 @@ void FImportVertexTextureHelper::ImportVertexColors(FEditorModeTools* ModeTools,
 
 		if (bComponent)
 		{
-			ComponentReregisterContext.Reset( new FComponentReregisterContext( StaticMeshComponent ) );
+			ComponentReregisterContext = MakeUnique<FComponentReregisterContext>( StaticMeshComponent );
 			StaticMeshComponent->Modify();
 
 			// Ensure LODData has enough entries in it, free not required.
@@ -4400,7 +4401,7 @@ void FImportVertexTextureHelper::ImportVertexColors(FEditorModeTools* ModeTools,
 			}
 			// We're changing the mesh itself, so ALL static mesh components in the scene will need
 			// to be detached for this (and reattached afterwards.)
-			RecreateRenderStateContext.Reset( new FStaticMeshComponentRecreateRenderStateContext( StaticMesh ) );
+			RecreateRenderStateContext = MakeUnique<FStaticMeshComponentRecreateRenderStateContext>( StaticMesh );
 
 			// Dirty the mesh
 			StaticMesh->Modify();
diff --git a/Engine/Source/Editor/MovieSceneTools/Private/TrackEditors/PropertyTrackEditors/BytePropertyTrackEditor.cpp b/Engine/Source/Editor/MovieSceneTools/Private/TrackEditors/PropertyTrackEditors/BytePropertyTrackEditor.cpp
--- a/Engine/Source/Editor/MovieSceneTools/Private/TrackEditors/PropertyTrackEditors/BytePropertyTrackEditor.cpp
+++ b/Engine/Source/Editor/MovieSceneTools/Private/TrackEditors/PropertyTrackEditors/BytePropertyTrackEditor.cpp
@@ -37,10 +37,20 @@ UEnum* GetEnumForByteTrack(TSharedPtr<ISequencer> Sequencer, const FGuid& OwnerO
 		UProperty* Property = RuntimeObject->GetClass()->FindPropertyByName(PropertyName);
 		if (Property != nullptr)
 		{
-			UByteProperty* ByteProperty = Cast<UByteProperty>(Property);
-			if (ByteProperty != nullptr && ByteProperty->Enum != nullptr)
+			UEnum* Enum = nullptr;
+
+			if (UEnumProperty* EnumProperty = Cast<UEnumProperty>(Property))
+			{
+				Enum = EnumProperty->GetEnum();
+			}
+			else if (UByteProperty* ByteProperty = Cast<UByteProperty>(Property))
+			{
+				Enum = ByteProperty->Enum;
+			}
+
+			if (Enum != nullptr)
 			{
-				PropertyEnums.Add(ByteProperty->Enum);
+				PropertyEnums.Add(Enum);
 			}
 		}
 	}
diff --git a/Engine/Source/Editor/PropertyEditor/Private/Presentation/PropertyTable/PropertyTableColumn.cpp b/Engine/Source/Editor/PropertyEditor/Private/Presentation/PropertyTable/PropertyTableColumn.cpp
--- a/Engine/Source/Editor/PropertyEditor/Private/Presentation/PropertyTable/PropertyTableColumn.cpp
+++ b/Engine/Source/Editor/PropertyEditor/Private/Presentation/PropertyTable/PropertyTableColumn.cpp
@@ -11,155 +11,6 @@
 
 #define LOCTEXT_NAMESPACE "PropertyTableColumn"
 
-template< typename UPropertyType >
-struct FCompareUObjectByPropertyAscending
-{
-public:
-
-	FCompareUObjectByPropertyAscending( const TWeakObjectPtr< UPropertyType >& InUProperty )
-		: Property( InUProperty )
-	{}
-
-	FCompareUObjectByPropertyAscending( const UPropertyType* const InUProperty )
-		: Property( InUProperty )
-	{}
-
-	FCompareUObjectByPropertyAscending( const TWeakObjectPtr< UProperty >& InUProperty )
-		: Property( Cast< UPropertyType >( InUProperty.Get() ) )
-	{}
-
-	FCompareUObjectByPropertyAscending( const UProperty* const InUProperty )
-		: Property( Cast< UPropertyType >( InUProperty ) )
-	{}
-
-	FORCEINLINE bool operator()( const TWeakObjectPtr< UObject >& Lhs, const TWeakObjectPtr< UObject >& Rhs ) const
-	{
-		const bool LhsObjectValid = Lhs.IsValid();
-		if ( !LhsObjectValid )
-		{
-			return true;
-		}
-
-		const bool RhsObjectValid = Rhs.IsValid();
-		if ( !RhsObjectValid )
-		{
-			return false;
-		}
-
-		const UStruct* const PropertyTargetType = Property->GetOwnerStruct();
-		const UClass* const LhsClass = Lhs->GetClass();
-
-		if ( !LhsClass->IsChildOf( PropertyTargetType ) )
-		{
-			return true;
-		}
-
-		const UClass* const RhsClass = Rhs->GetClass();
-
-		if ( !RhsClass->IsChildOf( PropertyTargetType ) )
-		{
-			return false;
-		}
-
-		typename UPropertyType::TCppType LhsValue = Property->GetPropertyValue_InContainer(Lhs.Get());
-		typename UPropertyType::TCppType RhsValue = Property->GetPropertyValue_InContainer(Rhs.Get());
-
-		return ComparePropertyValue( LhsValue, RhsValue );
-	}
-
-
-private:
-
-	FORCEINLINE bool ComparePropertyValue( const typename UPropertyType::TCppType& LhsValue, const typename UPropertyType::TCppType& RhsValue ) const
-	{
-		return LhsValue < RhsValue;
-	}
-
-
-private:
-
-	const TWeakObjectPtr< UPropertyType > Property;
-};
-
-// UByteProperty objects may in fact represent Enums - so they need special handling for alphabetic Enum vs. numerical Byte sorting.
-template<>
-FORCEINLINE bool FCompareUObjectByPropertyAscending<UByteProperty>::ComparePropertyValue( const uint8& LhsValue, const uint8& RhsValue ) const
-{
-	// Bytes are trivially sorted numerically
-	UEnum* PropertyEnum = Property->GetIntPropertyEnum();
-	if( PropertyEnum == NULL )
-	{
-		return LhsValue < RhsValue;
-	}
-	else
-	{
-		// But Enums are sorted alphabetically based on the full enum entry name - must be sure that values are within Enum bounds!
-		bool bLhsEnumValid = LhsValue < PropertyEnum->NumEnums();
-		bool bRhsEnumValid = RhsValue < PropertyEnum->NumEnums();
-		if(bLhsEnumValid && bRhsEnumValid)
-		{
-			FName LhsEnumName( PropertyEnum->GetEnum( LhsValue ) );
-			FName RhsEnumName( PropertyEnum->GetEnum( RhsValue ) );
-			return LhsEnumName.Compare( RhsEnumName ) < 0;
-		}
-		else if(bLhsEnumValid)
-		{
-			return true;
-		}
-		else if(bRhsEnumValid)
-		{
-			return false;
-		}
-		else
-		{
-			return LhsValue < RhsValue;
-		}
-	}
-}
-
-template<>
-FORCEINLINE bool FCompareUObjectByPropertyAscending<UNameProperty>::ComparePropertyValue( const FName& LhsValue, const FName& RhsValue ) const
-{
-	return LhsValue.Compare( RhsValue ) < 0;
-}
-
-template<>
-FORCEINLINE bool FCompareUObjectByPropertyAscending<UTextProperty>::ComparePropertyValue( const FText& LhsValue, const FText& RhsValue ) const
-{
-	return LhsValue.CompareTo( RhsValue ) < 0;
-}
-
-template< typename UPropertyType >
-struct FCompareUObjectByPropertyDescending
-{
-public:
-
-	FCompareUObjectByPropertyDescending( const TWeakObjectPtr< UPropertyType >& InUProperty )
-		: Compare( InUProperty )
-	{}
-
-	FCompareUObjectByPropertyDescending( const UPropertyType* const InUProperty )
-		: Compare( InUProperty )
-	{}
-
-	FCompareUObjectByPropertyDescending( const TWeakObjectPtr< UProperty >& InUProperty )
-		: Compare( InUProperty )
-	{}
-
-	FCompareUObjectByPropertyDescending( const UProperty* const InUProperty )
-		: Compare( InUProperty )
-	{}
-
-	FORCEINLINE bool operator()( const TWeakObjectPtr< UObject >& Lhs, const TWeakObjectPtr< UObject >& Rhs ) const
-	{
-		return !Compare(Lhs, Rhs);
-	}
-
-private:
-
-	const FCompareUObjectByPropertyAscending< UPropertyType > Compare;
-};
-
 template< typename UPropertyType >
 struct FCompareRowByColumnAscending
 {
@@ -244,6 +95,42 @@ private:
 };
 
 
+template<>
+FORCEINLINE bool FCompareRowByColumnAscending<UEnumProperty>::ComparePropertyValue( const TSharedPtr< IPropertyHandle >& LhsPropertyHandle, const TSharedPtr< IPropertyHandle >& RhsPropertyHandle ) const
+{
+	// Get the basic uint8 values
+	uint8 LhsValue; 
+	LhsPropertyHandle->GetValue( LhsValue );
+
+	uint8 RhsValue; 
+	RhsPropertyHandle->GetValue( RhsValue );
+
+	// Bytes are trivially sorted numerically
+	UEnum* PropertyEnum = Property->GetEnum();
+
+	// Enums are sorted alphabetically based on the full enum entry name - must be sure that values are within Enum bounds!
+	bool bLhsEnumValid = LhsValue < PropertyEnum->NumEnums();
+	bool bRhsEnumValid = RhsValue < PropertyEnum->NumEnums();
+	if(bLhsEnumValid && bRhsEnumValid)
+	{
+		FName LhsEnumName( PropertyEnum->GetEnum( LhsValue ) );
+		FName RhsEnumName( PropertyEnum->GetEnum( RhsValue ) );
+		return LhsEnumName.Compare( RhsEnumName ) < 0;
+	}
+	else if(bLhsEnumValid)
+	{
+		return true;
+	}
+	else if(bRhsEnumValid)
+	{
+		return false;
+	}
+	else
+	{
+		return LhsValue < RhsValue;
+	}
+}
+
 // UByteProperty objects may in fact represent Enums - so they need special handling for alphabetic Enum vs. numerical Byte sorting.
 template<>
 FORCEINLINE bool FCompareRowByColumnAscending<UByteProperty>::ComparePropertyValue( const TSharedPtr< IPropertyHandle >& LhsPropertyHandle, const TSharedPtr< IPropertyHandle >& RhsPropertyHandle ) const
@@ -535,6 +422,7 @@ bool FPropertyTableColumn::CanSortBy() const
 		return Property->IsA( UByteProperty::StaticClass() )  ||
 			Property->IsA( UIntProperty::StaticClass() )   ||
 			Property->IsA( UBoolProperty::StaticClass() )  ||
+			Property->IsA( UEnumProperty::StaticClass() )  ||
 			Property->IsA( UFloatProperty::StaticClass() ) ||
 			Property->IsA( UNameProperty::StaticClass() )  ||
 			Property->IsA( UStrProperty::StaticClass() )   ||
@@ -566,7 +454,20 @@ void FPropertyTableColumn::Sort( TArray< TSharedRef< class IPropertyTableRow > >
 		return;
 	}
 
-	if ( Property->IsA( UByteProperty::StaticClass() ) )
+	if ( Property->IsA( UEnumProperty::StaticClass() ) )
+	{
+		TWeakObjectPtr<UEnumProperty> EnumProperty( Cast< UEnumProperty >( Property ) );
+
+		if ( SortMode == EColumnSortMode::Ascending )
+		{
+			Rows.Sort( FCompareRowByColumnAscending< UEnumProperty >( SharedThis( this ), EnumProperty ) );
+		}
+		else
+		{
+			Rows.Sort( FCompareRowByColumnDescending< UEnumProperty >( SharedThis( this ), EnumProperty ) );
+		}
+	}
+	else if ( Property->IsA( UByteProperty::StaticClass() ) )
 	{
 		TWeakObjectPtr<UByteProperty> ByteProperty( Cast< UByteProperty >( Property ) );
 
diff --git a/Engine/Source/Editor/PropertyEditor/Private/PropertyEditorHelpers.cpp b/Engine/Source/Editor/PropertyEditor/Private/PropertyEditorHelpers.cpp
--- a/Engine/Source/Editor/PropertyEditor/Private/PropertyEditorHelpers.cpp
+++ b/Engine/Source/Editor/PropertyEditor/Private/PropertyEditorHelpers.cpp
@@ -495,13 +495,18 @@ namespace PropertyEditorHelpers
 		if(Property != NULL)
 		{
 			const UByteProperty* ByteProperty = Cast<UByteProperty>(Property);
-			if(ByteProperty || (Property->IsA(UStrProperty::StaticClass()) && Property->HasMetaData(TEXT("Enum"))))
+			const UEnumProperty* EnumProperty = Cast<UEnumProperty>(Property);
+			if(ByteProperty || EnumProperty || (Property->IsA(UStrProperty::StaticClass()) && Property->HasMetaData(TEXT("Enum"))))
 			{
 				UEnum* Enum = nullptr;
 				if(ByteProperty)
 				{
 					Enum = ByteProperty->Enum;
 				}
+				else if (EnumProperty)
+				{
+					Enum = EnumProperty->GetEnum();
+				}
 				else
 				{
 
diff --git a/Engine/Source/Editor/PropertyEditor/Private/PropertyEditorModule.cpp b/Engine/Source/Editor/PropertyEditor/Private/PropertyEditorModule.cpp
--- a/Engine/Source/Editor/PropertyEditor/Private/PropertyEditorModule.cpp
+++ b/Engine/Source/Editor/PropertyEditor/Private/PropertyEditorModule.cpp
@@ -608,10 +608,21 @@ FPropertyTypeLayoutCallback FPropertyEditorModule::GetPropertyTypeCustomization(
 		const bool bUserDefinedStruct = bStructProperty && StructProperty->Struct->IsA<UUserDefinedStruct>();
 		bStructProperty &= !bUserDefinedStruct;
 
-		const UByteProperty* ByteProperty = Cast<UByteProperty>(Property);
-		bool bEnumProperty = ByteProperty && ByteProperty->Enum;
-		const bool bUserDefinedEnum = bEnumProperty && ByteProperty->Enum->IsA<UUserDefinedEnum>();
-		bEnumProperty &= !bUserDefinedEnum;
+		const UEnum* Enum = nullptr;
+
+		if (const UByteProperty* ByteProperty = Cast<UByteProperty>(Property))
+		{
+			Enum = ByteProperty->Enum;
+		}
+		else if (const UEnumProperty* EnumProperty = Cast<UEnumProperty>(Property))
+		{
+			Enum = EnumProperty->GetEnum();
+		}
+
+		if (Enum && Enum->IsA<UUserDefinedEnum>())
+		{
+			Enum = nullptr;
+		}
 
 		const UObjectProperty* ObjectProperty = Cast<UObjectProperty>(Property);
 		const bool bObjectProperty = ObjectProperty != NULL && ObjectProperty->PropertyClass != NULL;
@@ -621,9 +632,9 @@ FPropertyTypeLayoutCallback FPropertyEditorModule::GetPropertyTypeCustomization(
 		{
 			PropertyTypeName = StructProperty->Struct->GetFName();
 		}
-		else if( bEnumProperty )
+		else if( Enum )
 		{
-			PropertyTypeName = ByteProperty->Enum->GetFName();
+			PropertyTypeName = Enum->GetFName();
 		}
 		else if ( bObjectProperty )
 		{
diff --git a/Engine/Source/Editor/PropertyEditor/Private/PropertyHandleImpl.cpp b/Engine/Source/Editor/PropertyEditor/Private/PropertyHandleImpl.cpp
--- a/Engine/Source/Editor/PropertyEditor/Private/PropertyHandleImpl.cpp
+++ b/Engine/Source/Editor/PropertyEditor/Private/PropertyHandleImpl.cpp
@@ -15,6 +15,7 @@
 
 #include "SNotificationList.h"
 #include "NotificationManager.h"
+#include "EnumProperty.h"
 
 #define LOCTEXT_NAMESPACE "PropertyHandleImplementation"
 
@@ -86,18 +87,31 @@ FPropertyAccess::Result FPropertyValueImpl::GetPropertyValueString( FString& Out
 			{
 				Property->ExportText_Direct( OutString, ValueAddress, ValueAddress, nullptr, PortFlags );
 
-				UByteProperty* ByteProperty = Cast<UByteProperty>(Property);
-				if ( ByteProperty != nullptr && ByteProperty->Enum != nullptr )
+				UEnum* Enum = nullptr;
+				int64 EnumValue = 0;
+				if ( UByteProperty* ByteProperty = Cast<UByteProperty>(Property) )
 				{
-					const uint8 EnumValue = ByteProperty->GetPropertyValue(ValueAddress);
+					if ( ByteProperty->Enum != nullptr )
+					{
+						Enum = ByteProperty->Enum;
+						EnumValue = ByteProperty->GetPropertyValue(ValueAddress);
+					}
+				}
+				else if ( UEnumProperty* EnumProperty = Cast<UEnumProperty>(Property) )
+				{
+					Enum = EnumProperty->GetEnum();
+					EnumValue = EnumProperty->GetUnderlyingProperty()->GetSignedIntPropertyValue(ValueAddress);
+				}
 
-					if (ByteProperty->Enum->IsValidEnumValue(EnumValue))
+				if ( Enum != nullptr )
+				{
+					if (Enum->IsValidEnumValue(EnumValue))
 					{
 						// See if we specified an alternate name for this value using metadata
-						OutString = ByteProperty->Enum->GetDisplayNameTextByValue(EnumValue).ToString();
+						OutString = Enum->GetDisplayNameTextByValue(EnumValue).ToString();
 						if(!bAllowAlternateDisplayValue || OutString.Len() == 0) 
 						{
-							OutString = ByteProperty->Enum->GetEnumNameStringByValue(EnumValue);
+							OutString = Enum->GetEnumNameStringByValue(EnumValue);
 						}
 					}
 					else
@@ -146,18 +160,28 @@ FPropertyAccess::Result FPropertyValueImpl::GetPropertyValueText( FText& OutText
 				FString ExportedTextString;
 				Property->ExportText_Direct(ExportedTextString, ValueAddress, ValueAddress, nullptr, PPF_PropertyWindow );
 
-				UByteProperty* ByteProperty = Cast<UByteProperty>(Property);
-				if ( ByteProperty != nullptr && ByteProperty->Enum != nullptr )
+				UEnum* Enum = nullptr;
+				int32 EnumeratorValue = 0;
+				if (UByteProperty* ByteProperty = Cast<UByteProperty>(Property))
 				{
-					const uint8 EnumValueIndex = ByteProperty->GetPropertyValue(ValueAddress);
-					
-					if (EnumValueIndex >= 0 && EnumValueIndex < ByteProperty->Enum->NumEnums())
+					Enum = ByteProperty->Enum;
+					EnumeratorValue = ByteProperty->GetPropertyValue(ValueAddress);
+				}
+				else if (UEnumProperty* EnumProperty = Cast<UEnumProperty>(Property))
+				{
+					Enum = EnumProperty->GetEnum();
+					EnumeratorValue = EnumProperty->GetUnderlyingProperty()->GetSignedIntPropertyValue(ValueAddress);
+				}
+
+				if (Enum)
+				{
+					if (EnumeratorValue >= 0 && EnumeratorValue < Enum->NumEnums())
 					{
 						// See if we specified an alternate name for this value using metadata
-						OutText = ByteProperty->Enum->GetDisplayNameText(EnumValueIndex);
+						OutText = Enum->GetDisplayNameText(EnumeratorValue);
 						if(!bAllowAlternateDisplayValue || OutText.IsEmpty()) 
 						{
-							OutText = ByteProperty->Enum->GetEnumText(EnumValueIndex);
+							OutText = Enum->GetEnumText(EnumeratorValue);
 						}
 					}
 					else
@@ -2203,24 +2227,25 @@ bool FPropertyHandleBase::GeneratePossibleValues(TArray< TSharedPtr<FString> >&
 	UProperty* Property = GetProperty();
 
 	bool bUsesAlternateDisplayValues = false;
-	const UByteProperty* ByteProperty = Cast<const UByteProperty>( Property );
 
-	if( ByteProperty || ( Property->IsA(UStrProperty::StaticClass()) && Property->HasMetaData( TEXT("Enum") ) ) )
+	UEnum* Enum = nullptr;
+	if( const UByteProperty* ByteProperty = Cast<const UByteProperty>( Property ) )
 	{
-		UEnum* Enum = nullptr;
-		if( ByteProperty )
-		{
-			Enum = ByteProperty->Enum;
-		}
-		else
-		{
-
-			FString EnumName = Property->GetMetaData(TEXT("Enum"));
-			Enum = FindObject<UEnum>(ANY_PACKAGE, *EnumName, true);
-		}
-
+		Enum = ByteProperty->Enum;
+	}
+	else if( const UEnumProperty* EnumProperty = Cast<const UEnumProperty>( Property ) )
+	{
+		Enum = EnumProperty->GetEnum();
+	}
+	else if ( Property->IsA(UStrProperty::StaticClass()) && Property->HasMetaData( TEXT("Enum") ) )
+	{
+		FString EnumName = Property->GetMetaData(TEXT("Enum"));
+		Enum = FindObject<UEnum>(ANY_PACKAGE, *EnumName, true);
 		check( Enum );
+	}
 
+	if( Enum )
+	{
 		const TArray<FName> ValidEnumValues = PropertyEditorHelpers::GetValidEnumsFromPropertyOverride(Property, Enum);
 
 		//NumEnums() - 1, because the last item in an enum is the _MAX item
@@ -2705,7 +2730,7 @@ bool FPropertyHandleByte::Supports( TSharedRef<FPropertyNode> PropertyNode )
 		return false;
 	}
 
-	return Property->IsA( UByteProperty::StaticClass() );
+	return Property->IsA<UByteProperty>() || Property->IsA<UEnumProperty>();
 }
 
 FPropertyAccess::Result FPropertyHandleByte::GetValue( uint8& OutValue ) const
@@ -2715,7 +2740,18 @@ FPropertyAccess::Result FPropertyHandleByte::GetValue( uint8& OutValue ) const
 
 	if( Res == FPropertyAccess::Success )
 	{
-		OutValue = Implementation->GetPropertyValue<UByteProperty>(PropValue);
+		TSharedPtr<FPropertyNode> PropertyNodePin = Implementation->GetPropertyNode();
+
+		UProperty* Property = PropertyNodePin->GetProperty();
+		if( Property->IsA<UByteProperty>() )
+		{
+			OutValue = Implementation->GetPropertyValue<UByteProperty>(PropValue);
+		}
+		else
+		{
+			check(PropertyNodePin.IsValid());
+			OutValue = CastChecked<UEnumProperty>(Property)->GetUnderlyingProperty()->GetUnsignedIntPropertyValue(PropValue);
+		}
 	}
 
 	return Res;
@@ -2725,12 +2761,23 @@ FPropertyAccess::Result FPropertyHandleByte::SetValue( const uint8& NewValue, EP
 {
 	FPropertyAccess::Result Res;
 	FString ValueStr;
-	
-	UByteProperty* ByteProperty = Cast<UByteProperty>(GetProperty());
-	if (ByteProperty && ByteProperty->Enum)
+
+	UProperty* Property = GetProperty();
+
+	UEnum* Enum = nullptr;
+	if (UByteProperty* ByteProperty = Cast<UByteProperty>(Property))
+	{
+		Enum = ByteProperty->Enum;
+	}
+	else if (UEnumProperty* EnumProperty = Cast<UEnumProperty>(Property))
+	{
+		Enum = EnumProperty->GetEnum();
+	}
+
+	if (Enum)
 	{
 		// Handle Enums using enum names to make sure they're compatible with UByteProperty::ExportText.
-		ValueStr = ByteProperty->Enum->GetEnumName(NewValue);
+		ValueStr = Enum->GetEnumName(NewValue);
 	}
 	else
 	{
diff --git a/Engine/Source/Editor/PropertyEditor/Private/PropertyNode.cpp b/Engine/Source/Editor/PropertyEditor/Private/PropertyNode.cpp
--- a/Engine/Source/Editor/PropertyEditor/Private/PropertyNode.cpp
+++ b/Engine/Source/Editor/PropertyEditor/Private/PropertyNode.cpp
@@ -1690,10 +1690,19 @@ FString FPropertyNode::GetDefaultValueAsStringForObject( FPropertyItemValueDataT
 				{
 					InProperty->ExportTextItem( DefaultValue, ValueTracker.GetPropertyDefaultAddress(), ValueTracker.GetPropertyDefaultAddress(), InObject, PortFlags, NULL );
 
-					UByteProperty* ByteProperty = Cast<UByteProperty>(InProperty);
-					if ( ByteProperty != NULL && ByteProperty->Enum != NULL )
+					UEnum* Enum = nullptr;
+					if (UByteProperty* ByteProperty = Cast<UByteProperty>(InProperty))
 					{
-						AdjustEnumPropDisplayName(ByteProperty->Enum, DefaultValue);
+						Enum = ByteProperty->Enum;
+					}
+					else if (UEnumProperty* EnumProperty = Cast<UEnumProperty>(InProperty))
+					{
+						Enum = EnumProperty->GetEnum();
+					}
+
+					if ( Enum )
+					{
+						AdjustEnumPropDisplayName(Enum, DefaultValue);
 					}
 				}
 			}
diff --git a/Engine/Source/Editor/PropertyEditor/Private/UserInterface/PropertyEditor/SPropertyEditorCombo.cpp b/Engine/Source/Editor/PropertyEditor/Private/UserInterface/PropertyEditor/SPropertyEditorCombo.cpp
--- a/Engine/Source/Editor/PropertyEditor/Private/UserInterface/PropertyEditor/SPropertyEditorCombo.cpp
+++ b/Engine/Source/Editor/PropertyEditor/Private/UserInterface/PropertyEditor/SPropertyEditorCombo.cpp
@@ -51,6 +51,7 @@ bool SPropertyEditorCombo::Supports( const TSharedRef< class FPropertyEditor >&
 	int32 ArrayIndex = PropertyNode->GetArrayIndex();
 
 	if(	((Property->IsA(UByteProperty::StaticClass()) && Cast<const UByteProperty>(Property)->Enum)
+		||	Property->IsA(UEnumProperty::StaticClass())
 		||	(Property->IsA(UNameProperty::StaticClass()) && Property->GetFName() == NAME_InitialState)
 		||	(Property->IsA(UStrProperty::StaticClass()) && Property->HasMetaData(TEXT("Enum")))
 		)
@@ -75,34 +76,47 @@ void SPropertyEditorCombo::Construct( const FArguments& InArgs, const TSharedRef
 	TArray<TSharedPtr<SToolTip>> RichToolTips;
 
 	// For enums, look for rich tooltip information
-	if(PropertyEditor.IsValid() && PropertyEditor->GetProperty())
+	if(PropertyEditor.IsValid())
 	{
-		const UProperty* Property = PropertyEditor->GetProperty();
-		if(UEnum* Enum = CastChecked<UByteProperty>(Property)->Enum)
+		if(const UProperty* Property = PropertyEditor->GetProperty())
 		{
-			TArray<FName> AllowedPropertyEnums = PropertyEditorHelpers::GetValidEnumsFromPropertyOverride(Property, Enum);
-
-			// Get enum doc link (not just GetDocumentationLink as that is the documentation for the struct we're in, not the enum documentation)
-			FString DocLink = PropertyEditorHelpers::GetEnumDocumentationLink(Property);
+			UEnum* Enum = nullptr;
 			
-			for(int32 EnumIdx = 0; EnumIdx < Enum->NumEnums() - 1; ++EnumIdx)
+			if(const UByteProperty* ByteProperty = Cast<UByteProperty>(Property))
+			{
+				Enum = ByteProperty->Enum;
+			}
+			else if (const UEnumProperty* EnumProperty = Cast<UEnumProperty>(Property))
 			{
-				FString Excerpt = Enum->GetEnumName(EnumIdx);
+				Enum = EnumProperty->GetEnum();
+			}
 
-				bool bShouldBeHidden = Enum->HasMetaData(TEXT("Hidden"), EnumIdx) || Enum->HasMetaData(TEXT("Spacer"), EnumIdx);
-				if( !bShouldBeHidden && AllowedPropertyEnums.Num() != 0 )
-				{
-					bShouldBeHidden = AllowedPropertyEnums.Find(Enum->GetEnum(EnumIdx)) == INDEX_NONE;
-				}
+			if(Enum)
+			{
+				TArray<FName> AllowedPropertyEnums = PropertyEditorHelpers::GetValidEnumsFromPropertyOverride(Property, Enum);
 
-				if(!bShouldBeHidden)
+				// Get enum doc link (not just GetDocumentationLink as that is the documentation for the struct we're in, not the enum documentation)
+				FString DocLink = PropertyEditorHelpers::GetEnumDocumentationLink(Property);
+			
+				for(int32 EnumIdx = 0; EnumIdx < Enum->NumEnums() - 1; ++EnumIdx)
 				{
-					RichToolTips.Add(IDocumentation::Get()->CreateToolTip(Enum->GetToolTipText(EnumIdx), nullptr, DocLink, Excerpt));
+					FString Excerpt = Enum->GetEnumName(EnumIdx);
+
+					bool bShouldBeHidden = Enum->HasMetaData(TEXT("Hidden"), EnumIdx) || Enum->HasMetaData(TEXT("Spacer"), EnumIdx);
+					if( !bShouldBeHidden && AllowedPropertyEnums.Num() != 0 )
+					{
+						bShouldBeHidden = AllowedPropertyEnums.Find(Enum->GetEnum(EnumIdx)) == INDEX_NONE;
+					}
+
+					if(!bShouldBeHidden)
+					{
+						RichToolTips.Add(IDocumentation::Get()->CreateToolTip(Enum->GetToolTipText(EnumIdx), nullptr, DocLink, Excerpt));
+					}
 				}
 			}
 		}
 	}
-	
+
 	SAssignNew(ComboBox, SPropertyComboBox)
 		.Font( InArgs._Font )
 		.RichToolTipList( RichToolTips )
@@ -124,10 +138,9 @@ void SPropertyEditorCombo::Construct( const FArguments& InArgs, const TSharedRef
 FString SPropertyEditorCombo::GetDisplayValueAsString() const
 {
 	const UProperty* Property = PropertyEditor->GetProperty();
-	const UByteProperty* ByteProperty = Cast<const UByteProperty>( Property );
-	const bool bStringEnumProperty = Property && Property->IsA(UStrProperty::StaticClass()) && Property->HasMetaData(TEXT("Enum"));	
+	const bool bStringEnumProperty = Cast<UStrProperty>(Property) && Property->HasMetaData(TEXT("Enum"));
 
-	if ( !(ByteProperty || bStringEnumProperty) )
+	if ( !Cast<UByteProperty>(Property) && !Cast<UEnumProperty>(Property) && !bStringEnumProperty )
 	{
 		UObject* ObjectValue = NULL;
 		FPropertyAccess::Result Result = PropertyEditor->GetPropertyHandle()->GetValue( ObjectValue );
@@ -174,7 +187,15 @@ void SPropertyEditorCombo::SendToObjects( const FString& NewValue )
 		// currently only enum properties can use alternate display values; this 
 		// might change, so assert here so that if support is expanded to other 
 		// property types without updating this block of code, we'll catch it quickly
-		UEnum* Enum = CastChecked<UByteProperty>(Property)->Enum;
+		UEnum* Enum = nullptr;
+		if (UByteProperty* ByteProperty = Cast<UByteProperty>(Property))
+		{
+			Enum = ByteProperty->Enum;
+		}
+		else if (UEnumProperty* EnumProperty = Cast<UEnumProperty>(Property))
+		{
+			Enum = EnumProperty->GetEnum();
+		}
 		check(Enum != nullptr);
 
 		const int32 Index = FindEnumValueIndex(Enum, NewValue);
diff --git a/Engine/Source/Editor/SequenceRecorder/Private/Sections/MovieSceneMultiPropertyRecorder.cpp b/Engine/Source/Editor/SequenceRecorder/Private/Sections/MovieSceneMultiPropertyRecorder.cpp
--- a/Engine/Source/Editor/SequenceRecorder/Private/Sections/MovieSceneMultiPropertyRecorder.cpp
+++ b/Engine/Source/Editor/SequenceRecorder/Private/Sections/MovieSceneMultiPropertyRecorder.cpp
@@ -61,6 +61,10 @@ void FMovieSceneMultiPropertyRecorder::CreateSection(UObject* InObjectToRecord,
 			{
 				PropertyRecorders.Add(MakeShareable(new FMovieScenePropertyRecorder<uint8>(Binding)));
 			}
+			else if (Property->IsA<UEnumProperty>())
+			{
+				PropertyRecorders.Add(MakeShareable(new FMovieScenePropertyRecorderEnum(Binding)));
+			}
 			else if (Property->IsA<UFloatProperty>())
 			{
 				PropertyRecorders.Add(MakeShareable(new FMovieScenePropertyRecorder<float>(Binding)));
@@ -108,6 +112,10 @@ bool FMovieSceneMultiPropertyRecorder::CanPropertyBeRecorded(const UProperty& In
 	{
 		return true;
 	}
+	else if (InProperty.IsA<UEnumProperty>())
+	{
+		return true;
+	}
 	else if (InProperty.IsA<UFloatProperty>())
 	{
 		return true;
diff --git a/Engine/Source/Editor/SequenceRecorder/Private/Sections/MovieScenePropertyRecorder.cpp b/Engine/Source/Editor/SequenceRecorder/Private/Sections/MovieScenePropertyRecorder.cpp
--- a/Engine/Source/Editor/SequenceRecorder/Private/Sections/MovieScenePropertyRecorder.cpp
+++ b/Engine/Source/Editor/SequenceRecorder/Private/Sections/MovieScenePropertyRecorder.cpp
@@ -6,6 +6,8 @@
 #include "MovieSceneBoolTrack.h"
 #include "MovieSceneByteSection.h"
 #include "MovieSceneByteTrack.h"
+#include "MovieSceneEnumSection.h"
+#include "MovieSceneEnumTrack.h"
 #include "MovieSceneFloatSection.h"
 #include "MovieSceneFloatTrack.h"
 #include "MovieSceneColorSection.h"
@@ -99,6 +101,44 @@ void FMovieScenePropertyRecorder<uint8>::ReduceKeys(UMovieSceneSection* InSectio
 {
 }
 
+bool FMovieScenePropertyRecorderEnum::ShouldAddNewKey(const int64& InNewValue) const
+{
+	return InNewValue != PreviousValue;
+}
+
+UMovieSceneSection* FMovieScenePropertyRecorderEnum::AddSection(UObject* InObjectToRecord, UMovieScene* InMovieScene, const FGuid& InGuid, float InTime)
+{
+	UMovieSceneEnumTrack* Track = InMovieScene->AddTrack<UMovieSceneEnumTrack>(InGuid);
+	if (Track)
+	{
+		if (InObjectToRecord)
+		{
+			Track->SetPropertyNameAndPath(*Binding.GetProperty(*InObjectToRecord)->GetDisplayNameText().ToString(), Binding.GetPropertyPath());
+		}
+
+		UMovieSceneEnumSection* Section = Cast<UMovieSceneEnumSection>(Track->CreateNewSection());
+		Section->SetDefault(PreviousValue);
+		Section->SetStartTime(InTime);
+		Section->SetEndTime(InTime);
+		Section->AddKey(InTime, PreviousValue, EMovieSceneKeyInterpolation::Break);
+
+		Track->AddSection(*Section);
+
+		return Section;
+	}
+
+	return nullptr;
+}
+
+void FMovieScenePropertyRecorderEnum::AddKeyToSection(UMovieSceneSection* InSection, const FPropertyKey<int64>& InKey)
+{
+	CastChecked<UMovieSceneEnumSection>(InSection)->AddKey(InKey.Time, InKey.Value, EMovieSceneKeyInterpolation::Break);
+}
+
+void FMovieScenePropertyRecorderEnum::ReduceKeys(UMovieSceneSection* InSection)
+{
+}
+
 template <>
 bool FMovieScenePropertyRecorder<float>::ShouldAddNewKey(const float& InNewValue) const
 {
diff --git a/Engine/Source/Editor/SequenceRecorder/Private/Sections/MovieScenePropertyRecorder.h b/Engine/Source/Editor/SequenceRecorder/Private/Sections/MovieScenePropertyRecorder.h
--- a/Engine/Source/Editor/SequenceRecorder/Private/Sections/MovieScenePropertyRecorder.h
+++ b/Engine/Source/Editor/SequenceRecorder/Private/Sections/MovieScenePropertyRecorder.h
@@ -103,4 +103,84 @@ private:
 
 	/** Previous value we use to establish whether we should key */
 	PropertyType PreviousValue;
-};
\ No newline at end of file
+};
+
+/** Recorder for a simple property of type enum */
+class FMovieScenePropertyRecorderEnum : public IMovieScenePropertyRecorder
+{
+public:
+	FMovieScenePropertyRecorderEnum(const FTrackInstancePropertyBindings& InBinding)
+		: Binding(InBinding)
+	{
+	}
+	virtual ~FMovieScenePropertyRecorderEnum() { }
+
+	virtual void Create(UObject* InObjectToRecord, class UMovieScene* InMovieScene, const FGuid& InGuid, float InTime) override
+	{
+		if (InObjectToRecord)
+		{
+			PreviousValue = Binding.GetCurrentValueForEnum(*InObjectToRecord);
+		}
+
+		MovieSceneSection = AddSection(InObjectToRecord, InMovieScene, InGuid, InTime);
+	}
+
+	virtual void Record(UObject* InObjectToRecord, float InCurrentTime) override
+	{
+		if (InObjectToRecord != nullptr)
+		{
+			MovieSceneSection->SetEndTime(InCurrentTime);
+
+			int64 NewValue = Binding.GetCurrentValueForEnum(*InObjectToRecord);
+			if (ShouldAddNewKey(NewValue))
+			{
+				FPropertyKey<int64> Key;
+				Key.Time = InCurrentTime;
+				Key.Value = NewValue;
+
+				Keys.Add(Key);
+
+				PreviousValue = NewValue;
+			}
+		}
+	}
+
+	virtual void Finalize(UObject* InObjectToRecord) override
+	{
+		for (const FPropertyKey<int64>& Key : Keys)
+		{
+			AddKeyToSection(MovieSceneSection, Key);
+		}
+
+		ReduceKeys(MovieSceneSection);
+	}
+
+private:
+	/** 
+	 * Helper function, specialized by type, used to check if we do capture-time filtering of keys based
+	 * on previous values
+	 */
+	bool ShouldAddNewKey(const int64& InNewValue) const;
+
+	/** Helper function, specialized by type, used to add an appropriate section to the movie scene */
+	class UMovieSceneSection* AddSection(UObject* InObjectToRecord, class UMovieScene* InMovieScene, const FGuid& InGuid, float InTime);
+
+	/** Helper function, specialized by type, used to add keys to the movie scene section at Finalize() time */
+	void AddKeyToSection(UMovieSceneSection* InSection, const FPropertyKey<int64>& InKey);
+
+	/** Helper function, specialized by type, used to reduce keys */
+	void ReduceKeys(UMovieSceneSection* InSection);
+
+private:
+	/** Binding for this property */
+	FTrackInstancePropertyBindings Binding;
+
+	/** The keys that are being recorded */
+	TArray<FPropertyKey<int64>> Keys;
+
+	/** Section we are recording */
+	UMovieSceneSection* MovieSceneSection;
+
+	/** Previous value we use to establish whether we should key */
+	int64 PreviousValue;
+};
diff --git a/Engine/Source/Editor/StaticMeshEditor/Private/StaticMeshEditorTools.cpp b/Engine/Source/Editor/StaticMeshEditor/Private/StaticMeshEditorTools.cpp
--- a/Engine/Source/Editor/StaticMeshEditor/Private/StaticMeshEditorTools.cpp
+++ b/Engine/Source/Editor/StaticMeshEditor/Private/StaticMeshEditorTools.cpp
@@ -1362,7 +1362,7 @@ void FMeshSectionSettingsLayout::OnGetSectionsForView(ISectionListBuilder& OutSe
 {
 	check(LODIndex == ForLODIndex);
 	UStaticMesh& StaticMesh = GetStaticMesh();
-	FStaticMeshRenderData* RenderData = StaticMesh.RenderData;
+	FStaticMeshRenderData* RenderData = StaticMesh.RenderData.Get();
 	if (RenderData && RenderData->LODResources.IsValidIndex(LODIndex))
 	{
 		FStaticMeshLODResources& LOD = RenderData->LODResources[LODIndex];
@@ -1413,7 +1413,7 @@ void FMeshSectionSettingsLayout::OnSectionChanged(int32 ForLODIndex, int32 Secti
 	}
 	check(NewStaticMaterialIndex != INDEX_NONE);
 	check(StaticMesh.RenderData);
-	FStaticMeshRenderData* RenderData = StaticMesh.RenderData;
+	FStaticMeshRenderData* RenderData = StaticMesh.RenderData.Get();
 	if (RenderData && RenderData->LODResources.IsValidIndex(LODIndex))
 	{
 		FStaticMeshLODResources& LOD = RenderData->LODResources[LODIndex];
@@ -2286,7 +2286,7 @@ void FLevelOfDetailSettingsLayout::AddLODLevelCategories( IDetailLayoutBuilder&
 	if( StaticMesh )
 	{
 		const int32 StaticMeshLODCount = StaticMesh->GetNumLODs();
-		FStaticMeshRenderData* RenderData = StaticMesh->RenderData;
+		FStaticMeshRenderData* RenderData = StaticMesh->RenderData.Get();
 
 		//Add the Materials array
 		{
diff --git a/Engine/Source/Editor/StaticMeshEditor/Private/StaticMeshEditorViewportClient.cpp b/Engine/Source/Editor/StaticMeshEditor/Private/StaticMeshEditorViewportClient.cpp
--- a/Engine/Source/Editor/StaticMeshEditor/Private/StaticMeshEditorViewportClient.cpp
+++ b/Engine/Source/Editor/StaticMeshEditor/Private/StaticMeshEditorViewportClient.cpp
@@ -744,7 +744,7 @@ void FStaticMeshEditorViewportClient::DrawCanvas( FViewport& InViewport, FSceneV
 	int32 CurrentLODLevel = StaticMeshEditor->GetCurrentLODLevel();
 	if (CurrentLODLevel == 0)
 	{
-		CurrentLODLevel = ComputeStaticMeshLOD(StaticMesh->RenderData, StaticMeshComponent->Bounds.Origin, StaticMeshComponent->Bounds.SphereRadius, View, StaticMesh->MinLOD);
+		CurrentLODLevel = ComputeStaticMeshLOD(StaticMesh->RenderData.Get(), StaticMeshComponent->Bounds.Origin, StaticMeshComponent->Bounds.SphereRadius, View, StaticMesh->MinLOD);
 	}
 	else
 	{
@@ -833,7 +833,7 @@ void FStaticMeshEditorViewportClient::DrawUVsForMesh(FViewport* InViewport, FCan
 
 	int32 UVChannel = StaticMeshEditorPtr.Pin()->GetCurrentUVChannel();
 
-	DrawUVs(InViewport, InCanvas, InTextYPos, LODLevel, UVChannel, SelectedEdgeTexCoords[UVChannel], StaticMeshComponent->GetStaticMesh()->RenderData, NULL);
+	DrawUVs(InViewport, InCanvas, InTextYPos, LODLevel, UVChannel, SelectedEdgeTexCoords[UVChannel], StaticMeshComponent->GetStaticMesh()->RenderData.Get(), NULL);
 }
 
 void FStaticMeshEditorViewportClient::MouseMove(FViewport* InViewport,int32 x, int32 y)
diff --git a/Engine/Source/Editor/UMGEditor/Private/Designer/SDesignSurface.cpp b/Engine/Source/Editor/UMGEditor/Private/Designer/SDesignSurface.cpp
--- a/Engine/Source/Editor/UMGEditor/Private/Designer/SDesignSurface.cpp
+++ b/Engine/Source/Editor/UMGEditor/Private/Designer/SDesignSurface.cpp
@@ -104,9 +104,9 @@ struct FFixedZoomLevelsContainerDesignSurface : public FZoomLevelsContainer
 
 void SDesignSurface::Construct(const FArguments& InArgs)
 {
-	if ( !ZoomLevels.IsValid() )
+	if ( !ZoomLevels )
 	{
-		ZoomLevels = new FFixedZoomLevelsContainerDesignSurface();
+		ZoomLevels = MakeUnique<FFixedZoomLevelsContainerDesignSurface>();
 	}
 	ZoomLevel = ZoomLevels->GetDefaultZoomLevel();
 	PreviousZoomLevel = ZoomLevels->GetDefaultZoomLevel();
diff --git a/Engine/Source/Editor/UMGEditor/Private/Designer/SDesignSurface.h b/Engine/Source/Editor/UMGEditor/Private/Designer/SDesignSurface.h
--- a/Engine/Source/Editor/UMGEditor/Private/Designer/SDesignSurface.h
+++ b/Engine/Source/Editor/UMGEditor/Private/Designer/SDesignSurface.h
@@ -4,6 +4,7 @@
 
 #include "SCompoundWidget.h"
 #include "SNodePanel.h"
+#include "UniquePtr.h"
 
 class SDesignSurface : public SCompoundWidget
 {
@@ -92,7 +93,7 @@ protected:
 	FCurveSequence ZoomLevelFade;
 
 	// The interface for mapping ZoomLevel values to actual node scaling values
-	TScopedPointer<FZoomLevelsContainer> ZoomLevels;
+	TUniquePtr<FZoomLevelsContainer> ZoomLevels;
 
 	bool bAllowContinousZoomInterpolation;
 
diff --git a/Engine/Source/Editor/UMGEditor/Private/WidgetBlueprintCompiler.cpp b/Engine/Source/Editor/UMGEditor/Private/WidgetBlueprintCompiler.cpp
--- a/Engine/Source/Editor/UMGEditor/Private/WidgetBlueprintCompiler.cpp
+++ b/Engine/Source/Editor/UMGEditor/Private/WidgetBlueprintCompiler.cpp
@@ -272,7 +272,8 @@ void FWidgetBlueprintCompiler::FinishCompilingClass(UClass* Class)
 		{
 			FBlueprintEditorUtils::ForceLoadMembers(WidgetBP->WidgetTree);
 		}
-		BPGClass->WidgetTree = DuplicateObject<UWidgetTree>(WidgetBP->WidgetTree, BPGClass);
+
+		BPGClass->WidgetTree = Cast<UWidgetTree>(StaticDuplicateObject(WidgetBP->WidgetTree, BPGClass, NAME_None, RF_AllFlags & ~RF_DefaultSubObject));
 
 		for ( const UWidgetAnimation* Animation : WidgetBP->Animations )
 		{
diff --git a/Engine/Source/Editor/UnrealEd/Classes/Commandlets/CookCommandlet.h b/Engine/Source/Editor/UnrealEd/Classes/Commandlets/CookCommandlet.h
--- a/Engine/Source/Editor/UnrealEd/Classes/Commandlets/CookCommandlet.h
+++ b/Engine/Source/Editor/UnrealEd/Classes/Commandlets/CookCommandlet.h
@@ -6,6 +6,7 @@
 
 #pragma once
 #include "Commandlets/Commandlet.h"
+#include "UniquePtr.h"
 #include "CookCommandlet.generated.h"
 
 UCLASS(config=Editor)
@@ -116,7 +117,7 @@ public:
 private:
 
 	/** Holds the sandbox file wrapper to handle sandbox path conversion. */
-	TAutoPtr<class FSandboxPlatformFile> SandboxFile;
+	TUniquePtr<class FSandboxPlatformFile> SandboxFile;
 
 	/** We hook this up to a delegate to avoid reloading textures and whatnot */
 	TSet<FString> PackagesToNotReload;
diff --git a/Engine/Source/Editor/UnrealEd/Classes/CookOnTheSide/CookOnTheFlyServer.h b/Engine/Source/Editor/UnrealEd/Classes/CookOnTheSide/CookOnTheFlyServer.h
--- a/Engine/Source/Editor/UnrealEd/Classes/CookOnTheSide/CookOnTheFlyServer.h
+++ b/Engine/Source/Editor/UnrealEd/Classes/CookOnTheSide/CookOnTheFlyServer.h
@@ -818,7 +818,7 @@ private:
 	/** Max number of conncurrent shader jobs reducing this too low will increase cook time */
 	int32 MaxConcurrentShaderJobs;
 	ECookInitializationFlags CookFlags;
-	TAutoPtr<class FSandboxPlatformFile> SandboxFile;
+	TUniquePtr<class FSandboxPlatformFile> SandboxFile;
 	bool bIsInitializingSandbox; // stop recursion into callbacks when we are initializing sandbox
 	bool bIsSavingPackage; // used to stop recursive mark package dirty functions
 
@@ -874,7 +874,7 @@ private:
 	FString GetCachedPackageFilename( const UPackage* Package ) const;
 	FString GetCachedStandardPackageFilename( const UPackage* Package ) const;
 	FName GetCachedStandardPackageFileFName( const UPackage* Package ) const;
-	const FString& GetCachedSandboxFilename( const UPackage* Package, TAutoPtr<class FSandboxPlatformFile>& SandboxFile ) const;
+	const FString& GetCachedSandboxFilename( const UPackage* Package, TUniquePtr<class FSandboxPlatformFile>& SandboxFile ) const;
 	const FName* GetCachedPackageFilenameToPackageFName(const FName& StandardPackageFilename) const;
 	const FCachedPackageFilename& Cache(const FName& PackageName) const;
 	void ClearPackageFilenameCache() const;
diff --git a/Engine/Source/Editor/UnrealEd/Private/Commandlets/ChunkManifestGenerator.cpp b/Engine/Source/Editor/UnrealEd/Private/Commandlets/ChunkManifestGenerator.cpp
--- a/Engine/Source/Editor/UnrealEd/Private/Commandlets/ChunkManifestGenerator.cpp
+++ b/Engine/Source/Editor/UnrealEd/Private/Commandlets/ChunkManifestGenerator.cpp
@@ -14,6 +14,7 @@
 #include "ChunkManifestGenerator.h"
 #include "ChunkDependencyInfo.h"
 #include "IPlatformFileSandboxWrapper.h"
+#include "UniquePtr.h"
 
 #include "JsonWriter.h"
 #include "JsonReader.h"
@@ -150,16 +151,16 @@ bool FChunkManifestGenerator::GenerateStreamingInstallManifest(const FString& Pl
 	
 	// open a file for writing the list of pak file lists that we've generated
 	FString PakChunkListFilename = TmpPackagingDir / TEXT("pakchunklist.txt");
-	TAutoPtr<FArchive> PakChunkListFile(IFileManager::Get().CreateFileWriter(*PakChunkListFilename));
+	TUniquePtr<FArchive> PakChunkListFile(IFileManager::Get().CreateFileWriter(*PakChunkListFilename));
 
-	if (!PakChunkListFile.IsValid())
+	if (!PakChunkListFile)
 	{
 		UE_LOG(LogChunkManifestGenerator, Error, TEXT("Failed to open output pakchunklist file %s"), *PakChunkListFilename);
 		return false;
 	}
 
 	FString PakChunkLayerInfoFilename = FString::Printf(TEXT("%s/pakchunklayers.txt"), *TmpPackagingDir);
-	TAutoPtr<FArchive> ChunkLayerFile(IFileManager::Get().CreateFileWriter(*PakChunkLayerInfoFilename));
+	TUniquePtr<FArchive> ChunkLayerFile(IFileManager::Get().CreateFileWriter(*PakChunkLayerInfoFilename));
 
 	// generate per-chunk pak list files
 	for (int32 Index = 0; Index < FinalChunkManifests.Num(); ++Index)
@@ -170,9 +171,9 @@ bool FChunkManifestGenerator::GenerateStreamingInstallManifest(const FString& Pl
 			continue;
 		}
 		FString PakListFilename = FString::Printf(TEXT("%s/pakchunk%d.txt"), *TmpPackagingDir, Index);
-		TAutoPtr<FArchive> PakListFile(IFileManager::Get().CreateFileWriter(*PakListFilename));
+		TUniquePtr<FArchive> PakListFile(IFileManager::Get().CreateFileWriter(*PakListFilename));
 
-		if (!PakListFile.IsValid())
+		if (!PakListFile)
 		{
 			UE_LOG(LogChunkManifestGenerator, Error, TEXT("Failed to open output paklist file %s"), *PakListFilename);
 			return false;
diff --git a/Engine/Source/Editor/UnrealEd/Private/Commandlets/CookCommandlet.cpp b/Engine/Source/Editor/UnrealEd/Private/Commandlets/CookCommandlet.cpp
--- a/Engine/Source/Editor/UnrealEd/Private/Commandlets/CookCommandlet.cpp
+++ b/Engine/Source/Editor/UnrealEd/Private/Commandlets/CookCommandlet.cpp
@@ -1752,7 +1752,7 @@ bool UCookCommandlet::Cook(const TArray<ITargetPlatform*>& Platforms, TArray<FSt
 	// but will not be used to actually write/read files so we can safely
 	// use [Platform] token in the sandbox directory name and then replace it
 	// with the actual platform name.
-	SandboxFile = new FSandboxPlatformFile(false);
+	SandboxFile = MakeUnique<FSandboxPlatformFile>(false);
 
 	// Output directory override.	
 	FString OutputDirectory = GetOutputDirectoryOverride();
@@ -1983,7 +1983,7 @@ bool UCookCommandlet::Cook(const TArray<ITargetPlatform*>& Platforms, TArray<FSt
 
 	{
 		// Always try to save the manifests, this is required to make the asset registry work, but doesn't necessarily write a file
-		ManifestGenerator.SaveManifests(SandboxFile.GetOwnedPointer());
+		ManifestGenerator.SaveManifests(SandboxFile.Get());
 
 		// Save modified asset registry with all streaming chunk info generated during cook
 		FString RegistryFilename = FPaths::GameDir() / TEXT("AssetRegistry.bin");
diff --git a/Engine/Source/Editor/UnrealEd/Private/Commandlets/GatherTextFromAssetsCommandlet.cpp b/Engine/Source/Editor/UnrealEd/Private/Commandlets/GatherTextFromAssetsCommandlet.cpp
--- a/Engine/Source/Editor/UnrealEd/Private/Commandlets/GatherTextFromAssetsCommandlet.cpp
+++ b/Engine/Source/Editor/UnrealEd/Private/Commandlets/GatherTextFromAssetsCommandlet.cpp
@@ -11,6 +11,7 @@
 #include "Sound/DialogueVoice.h"
 #include "Engine/DataTable.h"
 #include "EditorObjectVersion.h"
+#include "UniquePtr.h"
 
 DEFINE_LOG_CATEGORY_STATIC(LogGatherTextFromAssetsCommandlet, Log, All);
 
@@ -423,7 +424,7 @@ int32 UGatherTextFromAssetsCommandlet::Main(const FString& Params)
 	TArray< FString > PackageFileNamesToLoad;
 	for (FString& PackageFile : PackageFileNamesToProcess)
 	{
-		TScopedPointer< FArchive > FileReader( IFileManager::Get().CreateFileReader( *PackageFile ) );
+		TUniquePtr< FArchive > FileReader( IFileManager::Get().CreateFileReader( *PackageFile ) );
 		if( FileReader )
 		{
 			// Read package file summary from the file
diff --git a/Engine/Source/Editor/UnrealEd/Private/Commandlets/GenerateTextLocalizationResourceCommandlet.cpp b/Engine/Source/Editor/UnrealEd/Private/Commandlets/GenerateTextLocalizationResourceCommandlet.cpp
--- a/Engine/Source/Editor/UnrealEd/Private/Commandlets/GenerateTextLocalizationResourceCommandlet.cpp
+++ b/Engine/Source/Editor/UnrealEd/Private/Commandlets/GenerateTextLocalizationResourceCommandlet.cpp
@@ -2,6 +2,7 @@
 
 #include "UnrealEd.h"
 #include "TextLocalizationResourceGenerator.h"
+#include "UniquePtr.h"
 
 DEFINE_LOG_CATEGORY_STATIC(LogGenerateTextLocalizationResourceCommandlet, Log, All);
 
@@ -140,7 +141,7 @@ int32 UGenerateTextLocalizationResourceCommandlet::Main(const FString& Params)
 		{
 			bool bSaved = false;
 
-			TAutoPtr<FArchive> TextLocalizationResourceArchive(IFileManager::Get().CreateFileWriter(*InSaveFileName));
+			TUniquePtr<FArchive> TextLocalizationResourceArchive(IFileManager::Get().CreateFileWriter(*InSaveFileName));
 			if (TextLocalizationResourceArchive)
 			{
 				bSaved = FTextLocalizationResourceGenerator::Generate(LocTextHelper, CultureName, bSkipSourceCheck, *TextLocalizationResourceArchive);
diff --git a/Engine/Source/Editor/UnrealEd/Private/Commandlets/StabilizeLocalizationKeys.cpp b/Engine/Source/Editor/UnrealEd/Private/Commandlets/StabilizeLocalizationKeys.cpp
--- a/Engine/Source/Editor/UnrealEd/Private/Commandlets/StabilizeLocalizationKeys.cpp
+++ b/Engine/Source/Editor/UnrealEd/Private/Commandlets/StabilizeLocalizationKeys.cpp
@@ -5,6 +5,7 @@
 #include "JsonInternationalizationArchiveSerializer.h"
 #include "EditorObjectVersion.h"
 #include "TextPackageNamespaceUtil.h"
+#include "UniquePtr.h"
 
 DEFINE_LOG_CATEGORY_STATIC(LogStabilizeLocalizationKeys, Log, All);
 
@@ -233,7 +234,7 @@ int32 UStabilizeLocalizationKeysCommandlet::Main(const FString& Params)
 	TArray<FString> UnstablePackages;
 	for (const FString& PackageFilename : AllPackages)
 	{
-		if (TScopedPointer<FArchive> FileReader = TScopedPointer<FArchive>(IFileManager::Get().CreateFileReader(*PackageFilename)))
+		if (TUniquePtr<FArchive> FileReader = TUniquePtr<FArchive>(IFileManager::Get().CreateFileReader(*PackageFilename)))
 		{
 			// Read package file summary from the file so we can test the version
 			FPackageFileSummary PackageFileSummary;
diff --git a/Engine/Source/Editor/UnrealEd/Private/CookOnTheFlyServer.cpp b/Engine/Source/Editor/UnrealEd/Private/CookOnTheFlyServer.cpp
--- a/Engine/Source/Editor/UnrealEd/Private/CookOnTheFlyServer.cpp
+++ b/Engine/Source/Editor/UnrealEd/Private/CookOnTheFlyServer.cpp
@@ -640,7 +640,7 @@ bool UCookOnTheFlyServer::ClearPackageFilenameCacheForPackage( const UPackage* P
 	return PackageFilenameCache.Remove( Package->GetFName() ) >= 1;
 }
 
-const FString& UCookOnTheFlyServer::GetCachedSandboxFilename( const UPackage* Package, TAutoPtr<class FSandboxPlatformFile>& InSandboxFile ) const 
+const FString& UCookOnTheFlyServer::GetCachedSandboxFilename( const UPackage* Package, TUniquePtr<class FSandboxPlatformFile>& InSandboxFile ) const 
 {
 	FName PackageFName = Package->GetFName();
 	static TMap<FName, FString> CachedSandboxFilenames;
@@ -5490,9 +5490,9 @@ void UCookOnTheFlyServer::CookByTheBookFinished()
 				LongCookedPackageNames.Add(FName(*LongPackageName));
 			}
 
-			Manifest.Value->BuildChunkManifest(LongCookedPackageNames, SandboxFile, CookByTheBookOptions->bGenerateStreamingInstallManifests);
+			Manifest.Value->BuildChunkManifest(LongCookedPackageNames, SandboxFile.Get(), CookByTheBookOptions->bGenerateStreamingInstallManifests);
 			// Always try to save the manifests, this is required to make the asset registry work, but doesn't necessarily write a file
-			Manifest.Value->SaveManifests(SandboxFile.GetOwnedPointer());
+			Manifest.Value->SaveManifests(SandboxFile.Get());
 
 			const FName& PlatformName = Manifest.Key;
 			TArray<FName> IgnorePackageFilenames;
@@ -5648,7 +5648,7 @@ void UCookOnTheFlyServer::CancelCookByTheBook()
 		CookRequests.DequeueAllRequests(CookByTheBookOptions->PreviousCookRequests);
 		CookByTheBookOptions->bRunning = false;
 
-		SandboxFile = NULL;
+		SandboxFile = nullptr;
 	}	
 }
 
@@ -5720,8 +5720,8 @@ void UCookOnTheFlyServer::CreateSandboxFile()
 	// but will not be used to actually write/read files so we can safely
 	// use [Platform] token in the sandbox directory name and then replace it
 	// with the actual platform name.
-	check( SandboxFile == NULL );
-	SandboxFile = new FSandboxPlatformFile(false);
+	check( SandboxFile == nullptr );
+	SandboxFile = MakeUnique<FSandboxPlatformFile>(false);
 
 	// Output directory override.	
 	FString OutputDirectory = GetOutputDirectoryOverride();
@@ -5732,7 +5732,7 @@ void UCookOnTheFlyServer::CreateSandboxFile()
 
 void UCookOnTheFlyServer::InitializeSandbox()
 {
-	if ( SandboxFile == NULL )
+	if ( SandboxFile == nullptr )
 	{
 		ITargetPlatformManagerModule& TPM = GetTargetPlatformManagerRef();
 		const TArray<ITargetPlatform*>& TargetPlatforms = TPM.GetCookingTargetPlatforms();
@@ -5753,7 +5753,7 @@ void UCookOnTheFlyServer::TermSandbox()
 {
 	ClearAllCookedData();
 	ClearPackageFilenameCache();
-	SandboxFile = NULL;
+	SandboxFile = nullptr;
 }
 
 void UCookOnTheFlyServer::ValidateCookByTheBookSettings() const
diff --git a/Engine/Source/Editor/UnrealEd/Private/DistanceFieldBuildNotification.cpp b/Engine/Source/Editor/UnrealEd/Private/DistanceFieldBuildNotification.cpp
--- a/Engine/Source/Editor/UnrealEd/Private/DistanceFieldBuildNotification.cpp
+++ b/Engine/Source/Editor/UnrealEd/Private/DistanceFieldBuildNotification.cpp
@@ -6,6 +6,7 @@
 #include "DistanceFieldAtlas.h"
 #include "SNotificationList.h"
 #include "NotificationManager.h"
+#include "UniquePtr.h"
 
 /** Notification class for asynchronous distance field building. */
 class FDistanceFieldBuildNotificationImpl
@@ -44,17 +45,17 @@ private:
 };
 
 /** Global notification object. */
-TScopedPointer<FDistanceFieldBuildNotificationImpl> GDistanceFieldBuildNotification;
+TUniquePtr<FDistanceFieldBuildNotificationImpl> GDistanceFieldBuildNotification;
 
 void SetupDistanceFieldBuildNotification()
 {
 	// Create explicitly to avoid relying on static initialization order
-	GDistanceFieldBuildNotification = new FDistanceFieldBuildNotificationImpl;
+	GDistanceFieldBuildNotification = MakeUnique<FDistanceFieldBuildNotificationImpl>();
 }
 
 void TearDownDistanceFieldBuildNotification()
 {
-	GDistanceFieldBuildNotification = NULL;
+	GDistanceFieldBuildNotification = nullptr;
 }
 
 void FDistanceFieldBuildNotificationImpl::DistanceFieldBuildStarted()
diff --git a/Engine/Source/Editor/UnrealEd/Private/Factories/PackFactory.cpp b/Engine/Source/Editor/UnrealEd/Private/Factories/PackFactory.cpp
--- a/Engine/Source/Editor/UnrealEd/Private/Factories/PackFactory.cpp
+++ b/Engine/Source/Editor/UnrealEd/Private/Factories/PackFactory.cpp
@@ -16,6 +16,7 @@
 #include "Factories/PackFactory.h"
 #include "GameFramework/InputSettings.h"
 #include "Dialogs/SOutputLogDialog.h"
+#include "UniquePtr.h"
 
 DEFINE_LOG_CATEGORY_STATIC(LogPackFactory, Log, All);
 
@@ -506,9 +507,9 @@ UObject* UPackFactory::FactoryCreateBinary
 					DestFilename = ContentDestinationRoot / DestFilename;
 					UE_LOG(LogPackFactory, Log, TEXT("%s (%ld) -> %s"), *It.Filename(), Entry.Size, *DestFilename);
 
-					TAutoPtr<FArchive> FileHandle(IFileManager::Get().CreateFileWriter(*DestFilename));
+					TUniquePtr<FArchive> FileHandle(IFileManager::Get().CreateFileWriter(*DestFilename));
 
-					if (FileHandle.IsValid())
+					if (FileHandle)
 					{
 						PackFactoryHelper::ExtractFile(Entry, PakReader, CopyBuffer, PersistentCompressionBuffer, *FileHandle);
 						WrittenFiles.Add(*DestFilename);
diff --git a/Engine/Source/Editor/UnrealEd/Private/Fbx/FbxMainExport.cpp b/Engine/Source/Editor/UnrealEd/Private/Fbx/FbxMainExport.cpp
--- a/Engine/Source/Editor/UnrealEd/Private/Fbx/FbxMainExport.cpp
+++ b/Engine/Source/Editor/UnrealEd/Private/Fbx/FbxMainExport.cpp
@@ -633,7 +633,7 @@ void FFbxExporter::ExportModel(UModel* Model, FbxNode* Node, const char* Name)
 	
 	for (auto MaterialIterator = Model->MaterialIndexBuffers.CreateIterator(); MaterialIterator; ++MaterialIterator)
 	{
-		BeginReleaseResource(MaterialIterator.Value());
+		BeginReleaseResource(MaterialIterator->Value.Get());
 	}
 	FlushRenderingCommands();
 
diff --git a/Engine/Source/Editor/UnrealEd/Private/Kismet2/BlueprintEditorUtils.cpp b/Engine/Source/Editor/UnrealEd/Private/Kismet2/BlueprintEditorUtils.cpp
--- a/Engine/Source/Editor/UnrealEd/Private/Kismet2/BlueprintEditorUtils.cpp
+++ b/Engine/Source/Editor/UnrealEd/Private/Kismet2/BlueprintEditorUtils.cpp
@@ -7478,6 +7478,20 @@ bool FBlueprintEditorUtils::PropertyValueFromString(const UProperty* Property, c
 			bParseSucceeded = bParseSucceeded && ( IntValue <= 255 ) && ( IntValue >= 0 );
 			ByteProperty->SetPropertyValue_InContainer(DefaultObject, IntValue);
 		}
+		else if (const UEnumProperty* EnumProperty = Cast<const UEnumProperty>(Property))
+		{
+			int64 IntValue = IntValue = EnumProperty->GetEnum()->GetValueByName(FName(*Value));
+			bParseSucceeded = (INDEX_NONE != IntValue);
+
+			// If the parse did not succeed, clear out the int to keep the enum value valid
+			if(!bParseSucceeded)
+			{
+				IntValue = 0;
+			}
+			bParseSucceeded = bParseSucceeded && ( IntValue <= 255 ) && ( IntValue >= 0 );
+			void* PropAddr = EnumProperty->ContainerPtrToValuePtr<void>(DefaultObject);
+			EnumProperty->GetUnderlyingProperty()->SetIntPropertyValue(PropAddr, IntValue);
+		}
 		else if( Property->IsA(UStrProperty::StaticClass()) )
 		{
 			CastChecked<UStrProperty>(Property)->SetPropertyValue_InContainer(DefaultObject, Value);
diff --git a/Engine/Source/Editor/UnrealEd/Private/Kismet2/EnumEditorUtils.cpp b/Engine/Source/Editor/UnrealEd/Private/Kismet2/EnumEditorUtils.cpp
--- a/Engine/Source/Editor/UnrealEd/Private/Kismet2/EnumEditorUtils.cpp
+++ b/Engine/Source/Editor/UnrealEd/Private/Kismet2/EnumEditorUtils.cpp
@@ -34,7 +34,7 @@ UEnum* FEnumEditorUtils::CreateUserDefinedEnum(UObject* InParent, FName EnumName
 
 	if (NULL != Enum)
 	{
-		TArray<TPair<FName, uint8>> EmptyNames;
+		TArray<TPair<FName, int64>> EmptyNames;
 		Enum->SetEnums(EmptyNames, UEnum::ECppForm::Namespaced);
 		Enum->SetMetaData(TEXT("BlueprintType"), TEXT("true"));
 	}
@@ -51,13 +51,13 @@ void FEnumEditorUtils::UpdateAfterPathChanged(UEnum* Enum)
 {
 	check(NULL != Enum);
 
-	TArray<TPair<FName, uint8>> NewEnumeratorsNames;
+	TArray<TPair<FName, int64>> NewEnumeratorsNames;
 	const int32 EnumeratorsToCopy = Enum->NumEnums() - 1; // skip _MAX
 	for (int32 Index = 0; Index < EnumeratorsToCopy; Index++)
 	{
 		const FString ShortName = Enum->GetEnumName(Index);
 		const FString NewFullName = Enum->GenerateFullEnumName(*ShortName);
-		NewEnumeratorsNames.Add(TPairInitializer<FName, uint8>(FName(*NewFullName), Index));
+		NewEnumeratorsNames.Add(TPairInitializer<FName, int64>(FName(*NewFullName), Index));
 	}
 
 	Enum->SetEnums(NewEnumeratorsNames, UEnum::ECppForm::Namespaced);
@@ -66,7 +66,7 @@ void FEnumEditorUtils::UpdateAfterPathChanged(UEnum* Enum)
 //////////////////////////////////////////////////////////////////////////
 // Enumerators
 
-void FEnumEditorUtils::CopyEnumeratorsWithoutMax(const UEnum* Enum, TArray<TPair<FName, uint8>>& OutEnumNames)
+void FEnumEditorUtils::CopyEnumeratorsWithoutMax(const UEnum* Enum, TArray<TPair<FName, int64>>& OutEnumNames)
 {
 	if (Enum == nullptr)
 	{
@@ -76,7 +76,7 @@ void FEnumEditorUtils::CopyEnumeratorsWithoutMax(const UEnum* Enum, TArray<TPair
 	const int32 EnumeratorsToCopy = Enum->NumEnums() - 1;
 	for	(int32 Index = 0; Index < EnumeratorsToCopy; Index++)
 	{
-		OutEnumNames.Add(TPairInitializer<FName, int8>(Enum->GetNameByIndex(Index), Enum->GetValueByIndex(Index)));
+		OutEnumNames.Add(TPairInitializer<FName, int64>(Enum->GetNameByIndex(Index), Enum->GetValueByIndex(Index)));
 	}
 }
 
@@ -90,13 +90,13 @@ void FEnumEditorUtils::AddNewEnumeratorForUserDefinedEnum(UUserDefinedEnum* Enum
 
 	PrepareForChange(Enum);
 
-	TArray<TPair<FName, uint8>> OldNames, Names;
+	TArray<TPair<FName, int64>> OldNames, Names;
 	CopyEnumeratorsWithoutMax(Enum, OldNames);
 	Names = OldNames;
 
 	FString EnumNameString = Enum->GenerateNewEnumeratorName();
 	const FString FullNameStr = Enum->GenerateFullEnumName(*EnumNameString);
-	Names.Add(TPairInitializer<FName, uint8>(FName(*FullNameStr), Enum->GetMaxEnumValue()));
+	Names.Add(TPairInitializer<FName, int64>(FName(*FullNameStr), Enum->GetMaxEnumValue()));
 
 	// Clean up enum values.
 	for (int32 i = 0; i < Names.Num(); ++i)
@@ -108,7 +108,7 @@ void FEnumEditorUtils::AddNewEnumeratorForUserDefinedEnum(UUserDefinedEnum* Enum
 	EnsureAllDisplayNamesExist(Enum);
 
 	BroadcastChanges(Enum, OldNames);
-		
+
 	Enum->MarkPackageDirty();
 }
 
@@ -121,7 +121,7 @@ void FEnumEditorUtils::RemoveEnumeratorFromUserDefinedEnum(UUserDefinedEnum* Enu
 
 	PrepareForChange(Enum);
 
-	TArray<TPair<FName, uint8>> OldNames, Names;
+	TArray<TPair<FName, int64>> OldNames, Names;
 	CopyEnumeratorsWithoutMax(Enum, OldNames);
 	Names = OldNames;
 
@@ -162,7 +162,7 @@ void FEnumEditorUtils::SetEnumeratorBitflagsTypeState(UUserDefinedEnum* Enum, bo
 			Enum->RemoveMetaData(*FBlueprintMetadata::MD_Bitflags.ToString());
 		}
 
-		TArray<TPair<FName, uint8>> Names;
+		TArray<TPair<FName, int64>> Names;
 		CopyEnumeratorsWithoutMax(Enum, Names);
 		BroadcastChanges(Enum, Names);
 
@@ -180,7 +180,7 @@ void FEnumEditorUtils::MoveEnumeratorInUserDefinedEnum(UUserDefinedEnum* Enum, i
 
 	PrepareForChange(Enum);
 
-	TArray<TPair<FName, uint8>> OldNames, Names;
+	TArray<TPair<FName, int64>> OldNames, Names;
 	CopyEnumeratorsWithoutMax(Enum, OldNames);
 	Names = OldNames;
 
@@ -232,9 +232,9 @@ class FArchiveEnumeratorResolver : public FArchiveUObject
 {
 public:
 	const UEnum* Enum;
-	const TArray<TPair<FName, uint8>>& OldNames;
+	const TArray<TPair<FName, int64>>& OldNames;
 
-	FArchiveEnumeratorResolver(const UEnum* InEnum, const TArray<TPair<FName, uint8>>& InOldNames)
+	FArchiveEnumeratorResolver(const UEnum* InEnum, const TArray<TPair<FName, int64>>& InOldNames)
 		: FArchiveUObject(), Enum(InEnum), OldNames(InOldNames)
 	{
 	}
@@ -250,7 +250,7 @@ void FEnumEditorUtils::PrepareForChange(const UUserDefinedEnum* Enum)
 	FEnumEditorManager::Get().PreChange(Enum, EEnumEditorChangeInfo::Changed);
 }
 
-void FEnumEditorUtils::BroadcastChanges(const UUserDefinedEnum* Enum, const TArray<TPair<FName, uint8>>& OldNames, bool bResolveData)
+void FEnumEditorUtils::BroadcastChanges(const UUserDefinedEnum* Enum, const TArray<TPair<FName, int64>>& OldNames, bool bResolveData)
 {
 	check(NULL != Enum);
 	if (bResolveData)
@@ -258,9 +258,8 @@ void FEnumEditorUtils::BroadcastChanges(const UUserDefinedEnum* Enum, const TArr
 		FArchiveEnumeratorResolver EnumeratorResolver(Enum, OldNames);
 
 		TArray<UClass*> ClassesToCheck;
-		for (TObjectIterator<UByteProperty> PropertyIter; PropertyIter; ++PropertyIter)
+		for (const UByteProperty* ByteProperty : TObjectRange<UByteProperty>())
 		{
-			const UByteProperty* ByteProperty = *PropertyIter;
 			if (ByteProperty && (Enum == ByteProperty->GetIntPropertyEnum()))
 			{
 				UClass* OwnerClass = ByteProperty->GetOwnerClass();
@@ -270,12 +269,23 @@ void FEnumEditorUtils::BroadcastChanges(const UUserDefinedEnum* Enum, const TArr
 				}
 			}
 		}
+		for (const UEnumProperty* EnumProperty : TObjectRange<UEnumProperty>())
+		{
+			if (EnumProperty && (Enum == EnumProperty->GetEnum()))
+			{
+				UClass* OwnerClass = EnumProperty->GetOwnerClass();
+				if (OwnerClass)
+				{
+					ClassesToCheck.Add(OwnerClass);
+				}
+			}
+		}
 
 		for (FObjectIterator ObjIter; ObjIter; ++ObjIter)
 		{
-			for (auto ClassIter = ClassesToCheck.CreateConstIterator(); ClassIter; ++ClassIter)
+			for (UClass* Class : ClassesToCheck)
 			{
-				if (ObjIter->IsA(*ClassIter))
+				if (ObjIter->IsA(Class))
 				{
 					ObjIter->Serialize(EnumeratorResolver);
 					break;
@@ -400,13 +410,13 @@ void FEnumEditorUtils::BroadcastChanges(const UUserDefinedEnum* Enum, const TArr
 	FEnumEditorManager::Get().PostChange(Enum, EEnumEditorChangeInfo::Changed);
 }
 
-int32 FEnumEditorUtils::ResolveEnumerator(const UEnum* Enum, FArchive& Ar, int32 EnumeratorValue)
+int64 FEnumEditorUtils::ResolveEnumerator(const UEnum* Enum, FArchive& Ar, int64 EnumeratorValue)
 {
 	check(Ar.UseToResolveEnumerators());
 	const FArchiveEnumeratorResolver* EnumeratorResolver = (FArchiveEnumeratorResolver*)(&Ar);
 	if(Enum == EnumeratorResolver->Enum)
 	{
-		for (TPair<FName, uint8> OldName : EnumeratorResolver->OldNames)
+		for (TPair<FName, int64> OldName : EnumeratorResolver->OldNames)
 		{
 			if (OldName.Value == EnumeratorValue)
 			{
@@ -441,7 +451,7 @@ bool FEnumEditorUtils::SetEnumeratorDisplayName(UUserDefinedEnum* Enum, int32 En
 			PrepareForChange(Enum);
 			Enum->SetMetaData(FEnumEditorUtilsHelper::DisplayName(), *NewDisplayName, EnumeratorIndex);
 			EnsureAllDisplayNamesExist(Enum);
-			BroadcastChanges(Enum, TArray<TPair<FName, uint8>>(), false);
+			BroadcastChanges(Enum, TArray<TPair<FName, int64>>(), false);
 			return true;
 		}
 	}
diff --git a/Engine/Source/Editor/UnrealEd/Private/ObjectTools.cpp b/Engine/Source/Editor/UnrealEd/Private/ObjectTools.cpp
--- a/Engine/Source/Editor/UnrealEd/Private/ObjectTools.cpp
+++ b/Engine/Source/Editor/UnrealEd/Private/ObjectTools.cpp
@@ -42,6 +42,7 @@
 #include "Engine/SCS_Node.h"
 #include "Engine/UserDefinedStruct.h"
 #include "ShaderCompiler.h"
+#include "UniquePtr.h"
 
 DEFINE_LOG_CATEGORY_STATIC(LogObjectTools, Log, All);
 
@@ -3916,9 +3917,9 @@ namespace ThumbnailTools
 		if ( ObjectFullName.Len() > 0 && DestPackage != NULL )
 		{
 			// Create a new thumbnail map if we don't have one already
-			if( !DestPackage->ThumbnailMap.IsValid() )
+			if( !DestPackage->ThumbnailMap )
 			{
-				DestPackage->ThumbnailMap.Reset( new FThumbnailMap() );
+				DestPackage->ThumbnailMap = MakeUnique<FThumbnailMap>();
 			}
 
 			// @todo thumbnails: Backwards compat
@@ -4062,8 +4063,8 @@ namespace ThumbnailTools
 	bool LoadThumbnailsFromPackage( const FString& InPackageFileName, const TSet< FName >& InObjectFullNames, FThumbnailMap& InOutThumbnails )
 	{
 		// Create a file reader to load the file
-		TScopedPointer< FArchive > FileReader( IFileManager::Get().CreateFileReader( *InPackageFileName ) );
-		if( FileReader == NULL )
+		TUniquePtr< FArchive > FileReader( IFileManager::Get().CreateFileReader( *InPackageFileName ) );
+		if( FileReader == nullptr )
 		{
 			// Couldn't open the file
 			return false;
diff --git a/Engine/Source/Editor/UnrealEd/Private/StaticLightingSystem/StaticLightingDebug.cpp b/Engine/Source/Editor/UnrealEd/Private/StaticLightingSystem/StaticLightingDebug.cpp
--- a/Engine/Source/Editor/UnrealEd/Private/StaticLightingSystem/StaticLightingDebug.cpp
+++ b/Engine/Source/Editor/UnrealEd/Private/StaticLightingSystem/StaticLightingDebug.cpp
@@ -354,7 +354,7 @@ void SetDebugLightmapSample(TArray<UActorComponent*>* Components, UModel* Model,
 					for (int32 ElementIndex = 0; ElementIndex < CurrentComponent->GetElements().Num(); ElementIndex++)
 					{
 						FModelElement& Element = CurrentComponent->GetElements()[ElementIndex];
-						TScopedPointer<FRawIndexBuffer16or32>* IndexBufferRef = Model->MaterialIndexBuffers.Find(Element.Material);
+						TUniquePtr<FRawIndexBuffer16or32>* IndexBufferRef = Model->MaterialIndexBuffers.Find(Element.Material);
 						check(IndexBufferRef);
 						for(uint32 TriangleIndex = Element.FirstIndex; TriangleIndex < Element.FirstIndex + Element.NumTriangles * 3; TriangleIndex += 3)
 						{
@@ -395,7 +395,7 @@ void SetDebugLightmapSample(TArray<UActorComponent*>* Components, UModel* Model,
 				FModelVertex* ModelVertices = (FModelVertex*)Model->VertexBuffer.Vertices.GetData();
 
 				FModelElement& Element = ClosestComponent->GetElements()[ClosestElementIndex];
-				TScopedPointer<FRawIndexBuffer16or32>* IndexBufferRef = Model->MaterialIndexBuffers.Find(Element.Material);
+				TUniquePtr<FRawIndexBuffer16or32>* IndexBufferRef = Model->MaterialIndexBuffers.Find(Element.Material);
 
 				uint32 Index0 = (*IndexBufferRef)->Indices[ClosestTriangleIndex];
 				uint32 Index1 = (*IndexBufferRef)->Indices[ClosestTriangleIndex + 1];
diff --git a/Engine/Source/Editor/UnrealEd/Private/StaticLightingSystem/StaticLightingPrivate.h b/Engine/Source/Editor/UnrealEd/Private/StaticLightingSystem/StaticLightingPrivate.h
--- a/Engine/Source/Editor/UnrealEd/Private/StaticLightingSystem/StaticLightingPrivate.h
+++ b/Engine/Source/Editor/UnrealEd/Private/StaticLightingSystem/StaticLightingPrivate.h
@@ -11,6 +11,7 @@
 #include "GenericOctree.h"
 #include "StaticLighting.h"
 #include "LightingBuildOptions.h"
+#include "UniquePtr.h"
 
 struct FSelectedLightmapSample;
 
@@ -223,7 +224,7 @@ private:
 	class FStaticLightingSystem* ActiveStaticLightingSystem;
 
 	/** The system for kicking off the asynchronous lightmass */
-	TArray<TScopedPointer<class FStaticLightingSystem>> StaticLightingSystems;
+	TArray<TUniquePtr<class FStaticLightingSystem>> StaticLightingSystems;
 
 	/** Notification we hold on to that indicates progress. */
 	TWeakPtr<SNotificationItem> LightBuildNotification;
