Copying //UE4/Dev-Editor to //UE4/Dev-Main (Source: //UE4/Dev-Editor @ 3228984)

diff --git a/Engine/Source/Editor/UnrealEd/Private/AutomatedAssetImportData.cpp b/Engine/Source/Editor/UnrealEd/Private/AutomatedAssetImportData.cpp
--- a/Engine/Source/Editor/UnrealEd/Private/AutomatedAssetImportData.cpp
+++ b/Engine/Source/Editor/UnrealEd/Private/AutomatedAssetImportData.cpp
@@ -22,8 +22,9 @@ bool UAutomatedAssetImportData::IsValid() const
 	return Filenames.Num() > 0 && !DestinationPath.IsEmpty() && (FactoryName.IsEmpty() || Factory != nullptr);
 }
 
-void UAutomatedAssetImportData::Initialize()
+void UAutomatedAssetImportData::Initialize(TSharedPtr<FJsonObject> InImportGroupJsonData)
 {
+	ImportGroupJsonData = InImportGroupJsonData;
 	if(Filenames.Num() > 0)
 	{
 		// IF a factory doesn't have a vaild full path assume it is an unreal internal factory
diff --git a/Engine/Source/Editor/UnrealEd/Private/Commandlets/ImportAssetsCommandlet.cpp b/Engine/Source/Editor/UnrealEd/Private/Commandlets/ImportAssetsCommandlet.cpp
--- a/Engine/Source/Editor/UnrealEd/Private/Commandlets/ImportAssetsCommandlet.cpp
+++ b/Engine/Source/Editor/UnrealEd/Private/Commandlets/ImportAssetsCommandlet.cpp
@@ -90,7 +90,7 @@ bool UImportAssetsCommandlet::ParseParams(const FString& InParams)
 
 	ImportSettingsPath = ParamVals.FindRef(TEXT("importsettings"));
 
-	GlobalImportData->Initialize();
+	GlobalImportData->Initialize(nullptr);
 
 	if(ImportSettingsPath.IsEmpty() && (GlobalImportData->Filenames.Num() == 0 || GlobalImportData->DestinationPath.IsEmpty()))
 	{
@@ -123,27 +123,11 @@ bool UImportAssetsCommandlet::ParseImportSettings(const FString& InImportSetting
 					// Parse data from the json object
 					if(FJsonObjectConverter::JsonObjectToUStruct(ImportGroupsJsonObject.ToSharedRef(), UAutomatedAssetImportData::StaticClass(), Data, 0, 0 ))
 					{
-						Data->Initialize();
+						Data->Initialize(ImportGroupsJsonObject);
 						if(Data->IsValid())
 						{
 							ImportDataList.Add(Data);
 						}
-
-						UFactory* Factory = Data->Factory;
-						const TSharedPtr<FJsonObject>* ImportSettingsJsonObject = nullptr;
-						ImportGroupsJsonObject->TryGetObjectField(TEXT("ImportSettings"), ImportSettingsJsonObject);
-						if(Factory != nullptr && ImportSettingsJsonObject)
-						{
-							IImportSettingsParser* ImportSettings = Factory->GetImportSettingsParser();
-							if(ImportSettings)
-							{
-								ImportSettings->ParseFromJson(ImportSettingsJsonObject->ToSharedRef());
-							}
-						}
-						else if(Factory == nullptr && ImportSettingsJsonObject)
-						{
-							UE_LOG(LogAutomatedImport, Warning, TEXT("A vaild factory name must be specfied in order to specify settings"));
-						}
 					}
 					else
 					{
@@ -186,6 +170,27 @@ bool UImportAssetsCommandlet::ImportAndSave(const TArray<UAutomatedAssetImportDa
 	for(const UAutomatedAssetImportData* ImportData : AssetImportList)
 	{
 		UE_LOG(LogAutomatedImport, Log, TEXT("Importing group %s"), *ImportData->GetDisplayName() );
+
+		UFactory* Factory = ImportData->Factory;
+		const TSharedPtr<FJsonObject>* ImportSettingsJsonObject = nullptr;
+		if(ImportData->ImportGroupJsonData.IsValid())
+		{
+			ImportData->ImportGroupJsonData->TryGetObjectField(TEXT("ImportSettings"), ImportSettingsJsonObject);
+		}
+
+		if(Factory != nullptr && ImportSettingsJsonObject)
+		{
+			IImportSettingsParser* ImportSettings = Factory->GetImportSettingsParser();
+			if(ImportSettings)
+			{
+				ImportSettings->ParseFromJson(ImportSettingsJsonObject->ToSharedRef());
+			}
+		}
+		else if(Factory == nullptr && ImportSettingsJsonObject)
+		{
+			UE_LOG(LogAutomatedImport, Warning, TEXT("A vaild factory name must be specfied in order to specify settings"));
+		}
+
 		TArray<UObject*> ImportedAssets = AssetToolsModule.Get().ImportAssetsAutomated(*ImportData);
 		if(ImportedAssets.Num() > 0)
 		{
@@ -203,7 +208,6 @@ bool UImportAssetsCommandlet::ImportAndSave(const TArray<UAutomatedAssetImportDa
 			for(int32 PackageIndex = 0; PackageIndex < DirtyPackages.Num(); ++PackageIndex)
 			{
 				UPackage* PackageToSave = DirtyPackages[PackageIndex];
-				FSourceControlStateRef PackageSCState = PackageStates[PackageIndex];
 
 				FString PackageFilename = SourceControlHelpers::PackageFilename(PackageToSave);
 
@@ -211,6 +215,8 @@ bool UImportAssetsCommandlet::ImportAndSave(const TArray<UAutomatedAssetImportDa
 				bool bShouldAttemptToAdd = false;
 				if(bUseSourceControl)
 				{
+					FSourceControlStateRef PackageSCState = PackageStates[PackageIndex];
+
 					bool bPackageCanBeCheckedOut = false;
 					if(PackageSCState->IsCheckedOutOther())
 					{
diff --git a/Engine/Source/Editor/UnrealEd/Private/Commandlets/InternationalizationExportCommandlet.cpp b/Engine/Source/Editor/UnrealEd/Private/Commandlets/InternationalizationExportCommandlet.cpp
--- a/Engine/Source/Editor/UnrealEd/Private/Commandlets/InternationalizationExportCommandlet.cpp
+++ b/Engine/Source/Editor/UnrealEd/Private/Commandlets/InternationalizationExportCommandlet.cpp
@@ -4,7 +4,10 @@
 #include "Misc/FileHelper.h"
 #include "Misc/Paths.h"
 #include "Internationalization/InternationalizationMetadata.h"
+#include "Serialization/JsonInternationalizationMetadataSerializer.h"
 #include "Localization/PortableObjectFormatDOM.h"
+#include "TextNamespaceUtil.h"
+#include "LocalizationTargetTypes.h"
 
 DEFINE_LOG_CATEGORY_STATIC(LogInternationalizationExportCommandlet, Log, All);
 
@@ -28,10 +31,53 @@ uint32 GetTypeHash(const FPortableObjectEntryIdentity& ID)
 
 namespace
 {
+	struct FCaseSensitiveStringPair
+	{
+	public:
+		FCaseSensitiveStringPair(FString InFirst, FString InSecond)
+			: First(MoveTemp(InFirst))
+			, Second(MoveTemp(InSecond))
+		{
+		}
+
+		FORCEINLINE bool operator==(const FCaseSensitiveStringPair& Other) const
+		{
+			return First.Equals(Other.First, ESearchCase::CaseSensitive)
+				&& Second.Equals(Other.Second, ESearchCase::CaseSensitive);
+		}
+
+		FORCEINLINE bool operator!=(const FCaseSensitiveStringPair& Other) const
+		{
+			return !First.Equals(Other.First, ESearchCase::CaseSensitive)
+				|| !Second.Equals(Other.Second, ESearchCase::CaseSensitive);
+		}
+
+		friend inline uint32 GetTypeHash(const FCaseSensitiveStringPair& Id)
+		{
+			uint32 Hash = 0;
+			Hash = FCrc::StrCrc32(*Id.First, Hash);
+			Hash = FCrc::StrCrc32(*Id.Second, Hash);
+			return Hash;
+		}
+
+		FString First;
+		FString Second;
+	};
+	typedef TMultiMap<FCaseSensitiveStringPair, FCaseSensitiveStringPair> FCaseSensitiveStringPairMultiMap;
+
+	struct FCollapsedData
+	{
+		/** Mapping between a collapsed namespace (First) and key (Second), to an expanded namespace (First) and key (Second) */
+		FCaseSensitiveStringPairMultiMap CollapsedNSKeyToExpandedNSKey;
+
+		/** Mapping between a collapsed namespace (First) and source string/native translation (Second), to an expanded namespace (First) and key (Second) */
+		FCaseSensitiveStringPairMultiMap CollapsedNSSourceStringToExpandedNSKey;
+	};
+
 	/**
 	 * Declarations
 	 */
-	FString ConditionIdentityForPOMsgCtxt(const FString& Namespace, const FString& Key, const TSharedPtr<FLocMetadataObject>& KeyMetaData);
+	FString ConditionIdentityForPOMsgCtxt(const FString& Namespace, const FString& Key, const TSharedPtr<FLocMetadataObject>& KeyMetaData, const ELocalizedTextCollapseMode InTextCollapseMode);
 	void ParsePOMsgCtxtForIdentity(const FString& MsgCtxt, FString& OutNamespace, FString& OutKey);
 	FString ConditionArchiveStrForPo(const FString& InStr);
 	FString ConditionPoStringForArchive(const FString& InStr);
@@ -39,11 +85,12 @@ namespace
 	FString GetConditionedKeyForExtractedComment(const FString& Key);
 	FString GetConditionedReferenceForExtractedComment(const FString& PORefString);
 	FString GetConditionedInfoMetaDataForExtractedComment(const FString& KeyName, const FString& ValueString);
+	TSharedRef<FInternationalizationManifest> BuildCollapsedManifest(FLocTextHelper& InLocTextHelper, const ELocalizedTextCollapseMode InTextCollapseMode, FCollapsedData& OutCollapsedData);
 
 	/**
 	 * Definitions
 	 */
-	FString ConditionIdentityForPOMsgCtxt(const FString& Namespace, const FString& Key, const TSharedPtr<FLocMetadataObject>& KeyMetaData)
+	FString ConditionIdentityForPOMsgCtxt(const FString& Namespace, const FString& Key, const TSharedPtr<FLocMetadataObject>& KeyMetaData, const ELocalizedTextCollapseMode InTextCollapseMode)
 	{
 		const auto& EscapeMsgCtxtParticle = [](const FString& InStr) -> FString
 		{
@@ -62,7 +109,8 @@ namespace
 		const FString EscapedNamespace = EscapeMsgCtxtParticle(Namespace);
 		const FString EscapedKey = EscapeMsgCtxtParticle(Key);
 
-		const FString MsgCtxt = FString::Printf(TEXT("%s,%s"), *EscapedNamespace, *EscapedKey);
+		const bool bAppendKey = InTextCollapseMode != ELocalizedTextCollapseMode::IdenticalNamespaceAndSource || KeyMetaData.IsValid();
+		const FString MsgCtxt = bAppendKey ? FString::Printf(TEXT("%s,%s"), *EscapedNamespace, *EscapedKey) : EscapedNamespace;
 		return ConditionArchiveStrForPo(MsgCtxt);
 	}
 
@@ -152,112 +200,112 @@ namespace
 
 	FString ConditionArchiveStrForPo(const FString& InStr)
 	{
-	FString Result;
-	for (const TCHAR C : InStr)
-	{
-		switch (C)
+		FString Result;
+		for (const TCHAR C : InStr)
 		{
-		case TEXT('\\'):	Result += TEXT("\\\\");	break;
-		case TEXT('"'):		Result += TEXT("\\\"");	break;
-		case TEXT('\r'):	Result += TEXT("\\r");	break;
-		case TEXT('\n'):	Result += TEXT("\\n");	break;
-		case TEXT('\t'):	Result += TEXT("\\t");	break;
-		default:			Result += C;			break;
+			switch (C)
+			{
+			case TEXT('\\'):	Result += TEXT("\\\\");	break;
+			case TEXT('"'):		Result += TEXT("\\\"");	break;
+			case TEXT('\r'):	Result += TEXT("\\r");	break;
+			case TEXT('\n'):	Result += TEXT("\\n");	break;
+			case TEXT('\t'):	Result += TEXT("\\t");	break;
+			default:			Result += C;			break;
+			}
+		}
+		return Result;
 		}
-	}
-	return Result;
-	}
-
-	FString ConditionPoStringForArchive(const FString& InStr)
-	{
-	FString Result;
-	FString EscapeSequenceBuffer;
 
-	auto HandleEscapeSequenceBuffer = [&]()
-	{
-		// Insert unescaped sequence if needed.
-		if (!EscapeSequenceBuffer.IsEmpty())
+		FString ConditionPoStringForArchive(const FString& InStr)
 		{
-			bool EscapeSequenceIdentified = true;
+		FString Result;
+		FString EscapeSequenceBuffer;
 
-			// Identify escape sequence
-			TCHAR UnescapedCharacter = 0;
-			if (EscapeSequenceBuffer == TEXT("\\\\"))
-			{
-				UnescapedCharacter = '\\';
-			}
-			else if (EscapeSequenceBuffer == TEXT("\\\""))
-			{
-				UnescapedCharacter = '"';
-			}
-			else if (EscapeSequenceBuffer == TEXT("\\r"))
-			{
-				UnescapedCharacter = '\r';
-			}
-			else if (EscapeSequenceBuffer == TEXT("\\n"))
-			{
-				UnescapedCharacter = '\n';
-			}
-			else if (EscapeSequenceBuffer == TEXT("\\t"))
-			{
-				UnescapedCharacter = '\t';
-			}
-			else
+		auto HandleEscapeSequenceBuffer = [&]()
+		{
+			// Insert unescaped sequence if needed.
+			if (!EscapeSequenceBuffer.IsEmpty())
 			{
-				EscapeSequenceIdentified = false;
-			}
+				bool EscapeSequenceIdentified = true;
 
-			// If identified, append the processed sequence as the unescaped character.
-			if (EscapeSequenceIdentified)
-			{
-				Result += UnescapedCharacter;
-			}
-			// If it was not identified, preserve the escape sequence and append it.
-			else
-			{
-				Result += EscapeSequenceBuffer;
+				// Identify escape sequence
+				TCHAR UnescapedCharacter = 0;
+				if (EscapeSequenceBuffer == TEXT("\\\\"))
+				{
+					UnescapedCharacter = '\\';
+				}
+				else if (EscapeSequenceBuffer == TEXT("\\\""))
+				{
+					UnescapedCharacter = '"';
+				}
+				else if (EscapeSequenceBuffer == TEXT("\\r"))
+				{
+					UnescapedCharacter = '\r';
+				}
+				else if (EscapeSequenceBuffer == TEXT("\\n"))
+				{
+					UnescapedCharacter = '\n';
+				}
+				else if (EscapeSequenceBuffer == TEXT("\\t"))
+				{
+					UnescapedCharacter = '\t';
+				}
+				else
+				{
+					EscapeSequenceIdentified = false;
+				}
+
+				// If identified, append the processed sequence as the unescaped character.
+				if (EscapeSequenceIdentified)
+				{
+					Result += UnescapedCharacter;
+				}
+				// If it was not identified, preserve the escape sequence and append it.
+				else
+				{
+					Result += EscapeSequenceBuffer;
+				}
+				// Either way, we've appended something based on the buffer and it should be reset.
+				EscapeSequenceBuffer.Empty();
 			}
-			// Either way, we've appended something based on the buffer and it should be reset.
-			EscapeSequenceBuffer.Empty();
-		}
-	};
+		};
 
-	for (const TCHAR C : InStr)
-	{
-		// Not in an escape sequence.
-		if (EscapeSequenceBuffer.IsEmpty())
+		for (const TCHAR C : InStr)
 		{
-			// Regular character, just copy over.
-			if (C != TEXT('\\'))
+			// Not in an escape sequence.
+			if (EscapeSequenceBuffer.IsEmpty())
 			{
-				Result += C;
+				// Regular character, just copy over.
+				if (C != TEXT('\\'))
+				{
+					Result += C;
+				}
+				// Start of an escape sequence, put in escape sequence buffer.
+				else
+				{
+					EscapeSequenceBuffer += C;
+				}
 			}
-			// Start of an escape sequence, put in escape sequence buffer.
+			// If already in an escape sequence.
 			else
 			{
+				// Append to escape sequence buffer.
 				EscapeSequenceBuffer += C;
-			}
-		}
-		// If already in an escape sequence.
-		else
-		{
-			// Append to escape sequence buffer.
-			EscapeSequenceBuffer += C;
 
-			HandleEscapeSequenceBuffer();
+				HandleEscapeSequenceBuffer();
+			}
 		}
-	}
-	// Catch any trailing backslashes.
-	HandleEscapeSequenceBuffer();
-	return Result;
+		// Catch any trailing backslashes.
+		HandleEscapeSequenceBuffer();
+		return Result;
 	}
 
 	FString ConvertSrcLocationToPORef(const FString& InSrcLocation)
 	{
-	// Source location format: /Path1/Path2/file.cpp - line 123
-	// PO Reference format: /Path1/Path2/file.cpp:123
-	// @TODO: Note, we assume the source location format here but it could be arbitrary.
-	return InSrcLocation.Replace(TEXT(" - line "), TEXT(":"));
+		// Source location format: /Path1/Path2/file.cpp - line 123
+		// PO Reference format: /Path1/Path2/file.cpp:123
+		// @TODO: Note, we assume the source location format here but it could be arbitrary.
+		return InSrcLocation.Replace(TEXT(" - line "), TEXT(":"));
 	}
 
 	FString GetConditionedKeyForExtractedComment(const FString& Key)
@@ -274,6 +322,96 @@ namespace
 	{
 		return FString::Printf(TEXT("InfoMetaData:\t\"%s\" : \"%s\""), *KeyName, *ValueString);
 	}
+
+	TSharedRef<FInternationalizationManifest> BuildCollapsedManifest(FLocTextHelper& InLocTextHelper, const ELocalizedTextCollapseMode InTextCollapseMode, FCollapsedData& OutCollapsedData)
+	{
+		TSharedRef<FInternationalizationManifest> CollapsedManifest = MakeShared<FInternationalizationManifest>();
+
+		InLocTextHelper.EnumerateSourceTexts([&](TSharedRef<FManifestEntry> InManifestEntry) -> bool
+		{
+			const FString CollapsedNamespace = InTextCollapseMode == ELocalizedTextCollapseMode::IdenticalPackageIdTextIdAndSource ? InManifestEntry->Namespace : TextNamespaceUtil::StripPackageNamespace(InManifestEntry->Namespace);
+
+			for (const FManifestContext& Context : InManifestEntry->Contexts)
+			{
+				bool bAddedContext = false;
+
+				// Check if the entry already exists in the manifest
+				TSharedPtr<FManifestEntry> ExistingEntry = CollapsedManifest->FindEntryByContext(CollapsedNamespace, Context);
+				if (ExistingEntry.IsValid())
+				{
+					if (InManifestEntry->Source.IsExactMatch(ExistingEntry->Source))
+					{
+						bAddedContext = true;
+					}
+					else
+					{
+						// Grab the source location of the conflicting context
+						const FManifestContext* ConflictingContext = ExistingEntry->FindContext(Context.Key, Context.KeyMetadataObj);
+
+						const FString Message = FLocTextHelper::SanitizeLogOutput(
+							FString::Printf(TEXT("Found previously entered localized string: %s [%s] %s %s=\"%s\" %s. It was previously \"%s\" %s in %s."),
+								*Context.SourceLocation,
+								*CollapsedNamespace,
+								*Context.Key,
+								*FJsonInternationalizationMetaDataSerializer::MetadataToString(Context.KeyMetadataObj),
+								*InManifestEntry->Source.Text,
+								*FJsonInternationalizationMetaDataSerializer::MetadataToString(InManifestEntry->Source.MetadataObj),
+								*ExistingEntry->Source.Text,
+								*FJsonInternationalizationMetaDataSerializer::MetadataToString(ExistingEntry->Source.MetadataObj),
+								*ConflictingContext->SourceLocation
+							)
+						);
+						UE_LOG(LogInternationalizationExportCommandlet, Warning, TEXT("%s"), *Message);
+
+						InLocTextHelper.AddConflict(CollapsedNamespace, Context.Key, Context.KeyMetadataObj, InManifestEntry->Source, Context.SourceLocation);
+						InLocTextHelper.AddConflict(CollapsedNamespace, Context.Key, Context.KeyMetadataObj, ExistingEntry->Source, ConflictingContext->SourceLocation);
+					}
+				}
+				else
+				{
+					if (CollapsedManifest->AddSource(CollapsedNamespace, InManifestEntry->Source, Context))
+					{
+						bAddedContext = true;
+					}
+					else
+					{
+						UE_LOG(LogInternationalizationExportCommandlet, Error, TEXT("Could not process localized string: %s [%s] %s=\"%s\" %s."),
+							*Context.SourceLocation,
+							*CollapsedNamespace,
+							*Context.Key,
+							*InManifestEntry->Source.Text,
+							*FJsonInternationalizationMetaDataSerializer::MetadataToString(InManifestEntry->Source.MetadataObj)
+						);
+					}
+				}
+
+				if (bAddedContext)
+				{
+					// Add this collapsed namespace/key pair to our mapping so we can expand it again during import
+					OutCollapsedData.CollapsedNSKeyToExpandedNSKey.AddUnique(FCaseSensitiveStringPair(CollapsedNamespace, Context.Key), FCaseSensitiveStringPair(InManifestEntry->Namespace, Context.Key));
+
+					// Add this collapsed namespace/source string pair to our mapping so we expand it again during import (also map it against any native "translation" as that's what foreign imports will use as their source for translations)
+					if (!Context.KeyMetadataObj.IsValid())
+					{
+						OutCollapsedData.CollapsedNSSourceStringToExpandedNSKey.AddUnique(FCaseSensitiveStringPair(CollapsedNamespace, InManifestEntry->Source.Text), FCaseSensitiveStringPair(InManifestEntry->Namespace, Context.Key));
+
+						if (InLocTextHelper.HasNativeArchive())
+						{
+							TSharedPtr<FArchiveEntry> NativeTranslation = InLocTextHelper.FindTranslation(InLocTextHelper.GetNativeCulture(), InManifestEntry->Namespace, Context.Key, nullptr);
+							if (NativeTranslation.IsValid() && !NativeTranslation->Translation.Text.Equals(InManifestEntry->Source.Text))
+							{
+								OutCollapsedData.CollapsedNSSourceStringToExpandedNSKey.AddUnique(FCaseSensitiveStringPair(CollapsedNamespace, NativeTranslation->Translation.Text), FCaseSensitiveStringPair(InManifestEntry->Namespace, Context.Key));
+							}
+						}
+					}
+				}
+			}
+
+			return true; // continue enumeration
+		}, true);
+
+		return CollapsedManifest;
+	}
 }
 
 /**
@@ -319,6 +457,21 @@ bool UInternationalizationExportCommandlet::DoExport( const FString& SourcePath,
 		return false;
 	}
 
+	// Read in the text collapse mode to use
+	ELocalizedTextCollapseMode TextCollapseMode = ELocalizedTextCollapseMode::IdenticalTextIdAndSource;
+	{
+		FString TextCollapseModeName;
+		if (GetStringFromConfig(*SectionName, TEXT("LocalizedTextCollapseMode"), TextCollapseModeName, ConfigPath))
+		{
+			UEnum* LocalizedTextCollapseModeEnum = FindObjectChecked<UEnum>(ANY_PACKAGE, TEXT("ELocalizedTextCollapseMode"));
+			const int64 TextCollapseModeInt = LocalizedTextCollapseModeEnum->GetValueByName(*TextCollapseModeName);
+			if (TextCollapseModeInt != INDEX_NONE)
+			{
+				TextCollapseMode = (ELocalizedTextCollapseMode)TextCollapseModeInt;
+			}
+		}
+	}
+
 	bool ShouldAddSourceLocationsAsComments = true;
 	GetBoolFromConfig(*SectionName, TEXT("ShouldAddSourceLocationsAsComments"), ShouldAddSourceLocationsAsComments, ConfigPath);
 
@@ -333,6 +486,10 @@ bool UInternationalizationExportCommandlet::DoExport( const FString& SourcePath,
 		}
 	}
 
+	// Build the collapsed manifest data to export
+	FCollapsedData CollapsedData;
+	TSharedRef<FInternationalizationManifest> CollapsedManifest = BuildCollapsedManifest(LocTextHelper, TextCollapseMode, CollapsedData);
+
 	// Process the desired cultures
 	for (const FString& CultureName : CulturesToGenerate)
 	{
@@ -349,20 +506,25 @@ bool UInternationalizationExportCommandlet::DoExport( const FString& SourcePath,
 		NewPortableObject.CreateNewHeader();
 
 		// Add each manifest entry to the PO file
-		LocTextHelper.EnumerateSourceTexts([&LocTextHelper, &CultureName, &NewPortableObject](TSharedRef<FManifestEntry> InManifestEntry) -> bool
+		for (FManifestEntryByStringContainer::TConstIterator ManifestIterator = CollapsedManifest->GetEntriesByKeyIterator(); ManifestIterator; ++ManifestIterator)
 		{
+			const TSharedRef<FManifestEntry> ManifestEntry = ManifestIterator.Value();
+
 			// For each context, we may need to create a different or even multiple PO entries.
-			for (const FManifestContext& Context : InManifestEntry->Contexts)
+			for (const FManifestContext& Context : ManifestEntry->Contexts)
 			{
 				TSharedRef<FPortableObjectEntry> PoEntry = MakeShareable(new FPortableObjectEntry());
 
+				// For export we just use the first expanded namespace/key pair to find the current translation (they should all be identical due to how the import works)
+				const FCaseSensitiveStringPair& ExportNamespaceKeyPair = CollapsedData.CollapsedNSKeyToExpandedNSKey.FindChecked(FCaseSensitiveStringPair(ManifestEntry->Namespace, Context.Key));
+
 				// Find the correct translation based upon the native source text
 				FLocItem ExportedSource;
 				FLocItem ExportedTranslation;
-				LocTextHelper.GetExportText(CultureName, InManifestEntry->Namespace, Context.Key, Context.KeyMetadataObj, ELocTextExportSourceMethod::NativeText, InManifestEntry->Source, ExportedSource, ExportedTranslation);
+				LocTextHelper.GetExportText(CultureName, ExportNamespaceKeyPair.First, ExportNamespaceKeyPair.Second, Context.KeyMetadataObj, ELocTextExportSourceMethod::NativeText, ManifestEntry->Source, ExportedSource, ExportedTranslation);
 
 				PoEntry->MsgId = ConditionArchiveStrForPo(ExportedSource.Text);
-				PoEntry->MsgCtxt = ConditionIdentityForPOMsgCtxt(InManifestEntry->Namespace, Context.Key, Context.KeyMetadataObj);
+				PoEntry->MsgCtxt = ConditionIdentityForPOMsgCtxt(ManifestEntry->Namespace, Context.Key, Context.KeyMetadataObj, TextCollapseMode);
 				PoEntry->MsgStr.Add(ConditionArchiveStrForPo(ExportedTranslation.Text));
 
 				//@TODO: We support additional metadata entries that can be translated.  How do those fit in the PO file format?  Ex: isMature
@@ -389,9 +551,7 @@ bool UInternationalizationExportCommandlet::DoExport( const FString& SourcePath,
 
 				NewPortableObject.AddEntry(PoEntry);
 			}
-
-			return true; // continue enumeration
-		}, true);
+		}
 
 		// Write out the Portable Object to .po file.
 		{
@@ -490,6 +650,21 @@ bool UInternationalizationExportCommandlet::DoImport(const FString& SourcePath,
 		return false;
 	}
 
+	// Read in the text collapse mode to use
+	ELocalizedTextCollapseMode TextCollapseMode = ELocalizedTextCollapseMode::IdenticalTextIdAndSource;
+	{
+		FString TextCollapseModeName;
+		if (GetStringFromConfig(*SectionName, TEXT("LocalizedTextCollapseMode"), TextCollapseModeName, ConfigPath))
+		{
+			UEnum* LocalizedTextCollapseModeEnum = FindObjectChecked<UEnum>(ANY_PACKAGE, TEXT("ELocalizedTextCollapseMode"));
+			const int64 TextCollapseModeInt = LocalizedTextCollapseModeEnum->GetValueByName(*TextCollapseModeName);
+			if (TextCollapseModeInt != INDEX_NONE)
+			{
+				TextCollapseMode = (ELocalizedTextCollapseMode)TextCollapseModeInt;
+			}
+		}
+	}
+
 	// Load the manifest and all archives
 	FLocTextHelper LocTextHelper(DestinationPath, ManifestName, ArchiveName, NativeCultureName, CulturesToGenerate, MakeShareable(new FLocFileSCCNotifies(SourceControlInfo)));
 	{
@@ -501,6 +676,10 @@ bool UInternationalizationExportCommandlet::DoImport(const FString& SourcePath,
 		}
 	}
 
+	// Build the collapsed manifest data needed to import
+	FCollapsedData CollapsedData;
+	TSharedRef<FInternationalizationManifest> CollapsedManifest = BuildCollapsedManifest(LocTextHelper, TextCollapseMode, CollapsedData);
+
 	// Process the desired cultures
 	for (const FString& CultureName : CulturesToGenerate)
 	{
@@ -544,34 +723,39 @@ bool UInternationalizationExportCommandlet::DoImport(const FString& SourcePath,
 					UE_LOG(LogInternationalizationExportCommandlet, Error, TEXT("Portable Object entry has plural form we did not process.  File: %s  MsgCtxt: %s  MsgId: %s"), *POFilePath, *POEntry->MsgCtxt, *POEntry->MsgId);
 				}
 
-				FString Namespace;
-				TArray<FString> Keys;
-				const FString& SourceText = ConditionPoStringForArchive(POEntry->MsgId);
-				const FString& Translation = ConditionPoStringForArchive(POEntry->MsgStr[0]);
+				const FString SourceText = ConditionPoStringForArchive(POEntry->MsgId);
+				const FString Translation = ConditionPoStringForArchive(POEntry->MsgStr[0]);
 
+				TArray<FCaseSensitiveStringPair> NamespacesAndKeys; // Namespace (First) and Key (Second)
 				{
-					FString Key;
-					ParsePOMsgCtxtForIdentity(POEntry->MsgCtxt, Namespace, Key);
+					FString ParsedNamespace;
+					FString ParsedKey;
+					ParsePOMsgCtxtForIdentity(POEntry->MsgCtxt, ParsedNamespace, ParsedKey);
 
-					if (Key.IsEmpty())
+					if (ParsedKey.IsEmpty())
 					{
-						// Legacy non-keyed PO entry - need to look up and perform the import using the keys inferred from the manifest
-						LocTextHelper.FindKeysForLegacyTranslation(CultureName, Namespace, SourceText, nullptr, Keys);
+						// Legacy non-keyed PO entry - need to look-up the expanded namespace key/pairs via the namespace and source string
+						CollapsedData.CollapsedNSSourceStringToExpandedNSKey.MultiFind(FCaseSensitiveStringPair(ParsedNamespace, SourceText), NamespacesAndKeys);
 					}
 					else
 					{
-						Keys.Add(MoveTemp(Key));
+						// Keyed PO entry - need to look-up the expanded namespace/key pairs via the namespace and key
+						CollapsedData.CollapsedNSKeyToExpandedNSKey.MultiFind(FCaseSensitiveStringPair(ParsedNamespace, ParsedKey), NamespacesAndKeys);
 					}
 				}
 
-				if (Keys.Num() == 0)
+				if (NamespacesAndKeys.Num() == 0)
 				{
-					UE_LOG(LogInternationalizationExportCommandlet, Warning, TEXT("Could not import a PO entry as no key was specified, nor could a key be inferred from the manifest.  File: %s  MsgCtxt: %s  MsgId: %s"), *POFilePath, *POEntry->MsgCtxt, *POEntry->MsgId);
+					UE_LOG(LogInternationalizationExportCommandlet, Warning, TEXT("Could not import PO entry as it did not map to any known entries in the collapsed manifest data.  File: %s  MsgCtxt: %s  MsgId: %s"), *POFilePath, *POEntry->MsgCtxt, *POEntry->MsgId);
 					continue;
 				}
 				
-				for (const FString& Key : Keys)
+				for (const FCaseSensitiveStringPair& NamespaceAndKey : NamespacesAndKeys)
 				{
+					// Alias for convenience of reading
+					const FString& Namespace = NamespaceAndKey.First;
+					const FString& Key = NamespaceAndKey.Second;
+
 					// Get key metadata from the manifest, using the namespace and key.
 					const FManifestContext* ItemContext = nullptr;
 					{
diff --git a/Engine/Source/Editor/UnrealEd/Private/ComponentTypeRegistry.cpp b/Engine/Source/Editor/UnrealEd/Private/ComponentTypeRegistry.cpp
--- a/Engine/Source/Editor/UnrealEd/Private/ComponentTypeRegistry.cpp
+++ b/Engine/Source/Editor/UnrealEd/Private/ComponentTypeRegistry.cpp
@@ -118,6 +118,26 @@ void FComponentTypeRegistryData::AddBasicShapeComponents(TArray<FComponentClassC
 		}
 	}
 
+	{
+		FComponentEntryCustomizationArgs Args;
+		Args.AssetOverride = FindOrLoadObject<UStaticMesh>(UActorFactoryBasicShape::BasicPlane.ToString());
+		Args.OnComponentCreated = FOnComponentCreated::CreateStatic(OnBasicShapeCreated);
+		Args.ComponentNameOverride = LOCTEXT("BasicPlaneShapeDisplayName", "Plane").ToString();
+		Args.IconOverrideBrushName = FName("ClassIcon.Plane");
+		Args.SortPriority = 2;
+
+		{
+			FComponentClassComboEntryPtr NewShape = MakeShareable(new FComponentClassComboEntry(BasicShapesHeading, UStaticMeshComponent::StaticClass(), true, EComponentCreateAction::SpawnExistingClass, Args));
+			SortedClassList.Add(NewShape);
+		}
+
+		{
+			//Quad also goes in the common group
+			FComponentClassComboEntryPtr NewShape = MakeShareable(new FComponentClassComboEntry(CommonClassGroup, UStaticMeshComponent::StaticClass(), false, EComponentCreateAction::SpawnExistingClass, Args));
+			SortedClassList.Add(NewShape);
+		}
+	}
+
 	{
 		FComponentEntryCustomizationArgs Args;
 		Args.AssetOverride = FindOrLoadObject<UStaticMesh>(UActorFactoryBasicShape::BasicSphere.ToString());
diff --git a/Engine/Source/Editor/UnrealEd/Private/ComponentVisualizerManager.cpp b/Engine/Source/Editor/UnrealEd/Private/ComponentVisualizerManager.cpp
--- a/Engine/Source/Editor/UnrealEd/Private/ComponentVisualizerManager.cpp
+++ b/Engine/Source/Editor/UnrealEd/Private/ComponentVisualizerManager.cpp
@@ -8,6 +8,12 @@
 #include "UnrealEdGlobals.h"
 #include "SEditorViewport.h"
 
+
+FComponentVisualizerManager::FComponentVisualizerManager()
+	: EditedVisualizerViewportClient(nullptr)
+{
+}
+
 /** Handle a click on the specified editor viewport client */
 bool FComponentVisualizerManager::HandleClick(FEditorViewportClient* InViewportClient, HHitProxy* HitProxy, const FViewportClick& Click)
 {
@@ -57,6 +63,7 @@ bool FComponentVisualizerManager::HandleProxyForComponentVis(FEditorViewportClie
 					}
 
 					EditedVisualizerPtr = Visualizer;
+					EditedVisualizerViewportClient = InViewportClient;
 					return true;
 				}
 			}
@@ -77,8 +84,10 @@ void FComponentVisualizerManager::ClearActiveComponentVis()
 	if (EditedVisualizer.IsValid())
 	{
 		EditedVisualizer->EndEditing();
-		EditedVisualizerPtr = NULL;
 	}
+
+	EditedVisualizerPtr.Reset();
+	EditedVisualizerViewportClient = nullptr;
 }
 
 bool FComponentVisualizerManager::HandleInputKey(FEditorViewportClient* ViewportClient, FViewport* Viewport, FKey Key, EInputEvent Event) const
@@ -100,7 +109,7 @@ bool FComponentVisualizerManager::HandleInputDelta(FEditorViewportClient* InView
 {
 	TSharedPtr<FComponentVisualizer> EditedVisualizer = EditedVisualizerPtr.Pin();
 
-	if (EditedVisualizer.IsValid() && InViewportClient->GetCurrentWidgetAxis() != EAxisList::None)
+	if (EditedVisualizer.IsValid() && EditedVisualizerViewportClient == InViewportClient && InViewportClient->GetCurrentWidgetAxis() != EAxisList::None)
 	{
 		if (EditedVisualizer->HandleInputDelta(InViewportClient, InViewport, InDrag, InRot, InScale))
 		{
diff --git a/Engine/Source/Editor/UnrealEd/Private/DataTableEditorUtils.cpp b/Engine/Source/Editor/UnrealEd/Private/DataTableEditorUtils.cpp
--- a/Engine/Source/Editor/UnrealEd/Private/DataTableEditorUtils.cpp
+++ b/Engine/Source/Editor/UnrealEd/Private/DataTableEditorUtils.cpp
@@ -196,6 +196,9 @@ void FDataTableEditorUtils::CacheDataTableForEditing(const UDataTable* DataTable
 		return;
 	}
 
+	TArray<FDataTableEditorColumnHeaderDataPtr> OldColumns = OutAvailableColumns;
+	TArray<FDataTableEditorRowListViewDataPtr> OldRows = OutAvailableRows;
+
 	// First build array of properties
 	TArray<const UProperty*> StructProps;
 	for (TFieldIterator<const UProperty> It(DataTable->RowStruct); It; ++It)
@@ -211,11 +214,25 @@ void FDataTableEditorUtils::CacheDataTableForEditing(const UDataTable* DataTable
 
 	// Populate the column data
 	OutAvailableColumns.Reset(StructProps.Num());
-	for (const UProperty* Prop : StructProps)
+	for (int32 Index = 0; Index < StructProps.Num(); ++Index)
 	{
-		FDataTableEditorColumnHeaderDataPtr CachedColumnData = MakeShareable(new FDataTableEditorColumnHeaderData());
-		CachedColumnData->ColumnId = Prop->GetFName();
-		CachedColumnData->DisplayName = FText::FromString(DataTableUtils::GetPropertyDisplayName(Prop, Prop->GetName()));
+		const UProperty* Prop = StructProps[Index];
+		const FText PropertyDisplayName = FText::FromString(DataTableUtils::GetPropertyDisplayName(Prop, Prop->GetName()));
+
+		FDataTableEditorColumnHeaderDataPtr CachedColumnData;
+		
+		// If at all possible, attempt to reuse previous columns if their data has not changed
+		if (Index >= OldColumns.Num() || OldColumns[Index]->ColumnId != Prop->GetFName() || !OldColumns[Index]->DisplayName.EqualTo(PropertyDisplayName))
+		{
+			CachedColumnData = MakeShareable(new FDataTableEditorColumnHeaderData());
+			CachedColumnData->ColumnId = Prop->GetFName();
+			CachedColumnData->DisplayName = PropertyDisplayName;
+		}
+		else
+		{
+			CachedColumnData = OldColumns[Index];
+		}
+
 		CachedColumnData->DesiredColumnWidth = FontMeasure->Measure(CachedColumnData->DisplayName, CellTextStyle.Font).X + CellPadding;
 
 		OutAvailableColumns.Add(CachedColumnData);
@@ -223,14 +240,29 @@ void FDataTableEditorUtils::CacheDataTableForEditing(const UDataTable* DataTable
 
 	// Populate the row data
 	OutAvailableRows.Reset(DataTable->RowMap.Num());
-	for (auto RowIt = DataTable->RowMap.CreateConstIterator(); RowIt; ++RowIt)
+	int32 Index = 0;
+	for (auto RowIt = DataTable->RowMap.CreateConstIterator(); RowIt; ++RowIt, ++Index)
 	{
-		FDataTableEditorRowListViewDataPtr CachedRowData = MakeShareable(new FDataTableEditorRowListViewData());
-		CachedRowData->RowId = RowIt->Key;
-		CachedRowData->DisplayName = FText::FromName(RowIt->Key);
+		FText RowName = FText::FromName(RowIt->Key);
+		FDataTableEditorRowListViewDataPtr CachedRowData;
+
+		// If at all possible, attempt to reuse previous rows if their data has not changed.
+		if (Index >= OldRows.Num() || OldRows[Index]->RowId != RowIt->Key || !OldRows[Index]->DisplayName.EqualTo(RowName))
+		{
+			CachedRowData = MakeShareable(new FDataTableEditorRowListViewData());
+			CachedRowData->RowId = RowIt->Key;
+			CachedRowData->DisplayName = RowName;
+			CachedRowData->CellData.Reserve(StructProps.Num());
+		}
+		else
+		{
+			CachedRowData = OldRows[Index];
+			CachedRowData->CellData.Reset(StructProps.Num());
+		}
+
 		CachedRowData->DesiredRowHeight = FontMeasure->GetMaxCharacterHeight(CellTextStyle.Font);
 
-		CachedRowData->CellData.Reserve(StructProps.Num());
+		// Always rebuild cell data
 		{
 			uint8* RowData = RowIt.Value();
 			for (int32 ColumnIndex = 0; ColumnIndex < StructProps.Num(); ++ColumnIndex)
diff --git a/Engine/Source/Editor/UnrealEd/Private/Editor.cpp b/Engine/Source/Editor/UnrealEd/Private/Editor.cpp
--- a/Engine/Source/Editor/UnrealEd/Private/Editor.cpp
+++ b/Engine/Source/Editor/UnrealEd/Private/Editor.cpp
@@ -362,20 +362,27 @@ void FReimportManager::ValidateAllSourceFileAndReimport(TArray<UObject*> &ToImpo
 		TArray<FString> SourceFilenames;
 		if (this->CanReimport(Asset, &SourceFilenames))
 		{
-			bool bMissingFile = false;
-			for (FString SourceFilename : SourceFilenames)
+			if (SourceFilenames.Num() == 0)
 			{
-				if (SourceFilename.IsEmpty() || IFileManager::Get().FileSize(*SourceFilename) == INDEX_NONE)
+				MissingFileSelectedAssets.Add(Asset);
+			}
+			else
+			{
+				bool bMissingFile = false;
+				for (FString SourceFilename : SourceFilenames)
 				{
-					MissingFileSelectedAssets.Add(Asset);
-					bMissingFile = true;
-					break;
+					if (SourceFilename.IsEmpty() || IFileManager::Get().FileSize(*SourceFilename) == INDEX_NONE)
+					{
+						MissingFileSelectedAssets.Add(Asset);
+						bMissingFile = true;
+						break;
+					}
 				}
-			}
 
-			if (!bMissingFile)
-			{
-				CopyOfSelectedAssets.Add(Asset);
+				if (!bMissingFile)
+				{
+					CopyOfSelectedAssets.Add(Asset);
+				}
 			}
 		}
 	}
@@ -441,6 +448,18 @@ void FReimportManager::ValidateAllSourceFileAndReimport(TArray<UObject*> &ToImpo
 	FReimportManager::Instance()->ReimportMultiple(CopyOfSelectedAssets, /*bAskForNewFileIfMissing=*/false);
 }
 
+void FReimportManager::AddReferencedObjects( FReferenceCollector& Collector )
+{
+	for(FReimportHandler* Handler : Handlers)
+	{
+		const UObject* Obj = Handler->GetFactoryObject();
+		if(Obj)
+		{
+			Collector.AddReferencedObject(Obj);
+		}
+	}
+}
+
 bool FReimportManager::ReimportMultiple(TArrayView<UObject*> Objects, bool bAskForNewFileIfMissing /*= false*/, bool bShowNotification /*= true*/, FString PreferedReimportFile /*= TEXT("")*/, FReimportHandler* SpecifiedReimportHandler /*= nullptr */)
 {
 	bool bBulkSuccess = true;
diff --git a/Engine/Source/Editor/UnrealEd/Private/EditorActor.cpp b/Engine/Source/Editor/UnrealEd/Private/EditorActor.cpp
--- a/Engine/Source/Editor/UnrealEd/Private/EditorActor.cpp
+++ b/Engine/Source/Editor/UnrealEd/Private/EditorActor.cpp
@@ -57,6 +57,7 @@
 #include "UnrealExporter.h"
 #include "LevelEditor.h"
 #include "Engine/LODActor.h"
+#include "Settings/LevelEditorMiscSettings.h"
 
 #define LOCTEXT_NAMESPACE "UnrealEd.EditorActor"
 
@@ -335,10 +336,18 @@ void UUnrealEdEngine::edactPasteSelected(UWorld* InWorld, bool bDuplicate, bool
 		}
 		const TCHAR* Paste = *PasteString;
 
+		// Turn off automatic BSP update while pasting to save rebuilding geometry potentially multiple times
+		const bool bBSPAutoUpdate = GetDefault<ULevelEditorMiscSettings>()->bBSPAutoUpdate;
+		GetMutableDefault<ULevelEditorMiscSettings>()->bBSPAutoUpdate = false;
+
 		// Import the actors.
 		auto Factory = NewObject<ULevelFactory>();
 		Factory->FactoryCreateText(ULevel::StaticClass(), InWorld->GetCurrentLevel(), InWorld->GetCurrentLevel()->GetFName(), RF_Transactional, NULL, bDuplicate ? TEXT("move") : TEXT("paste"), Paste, Paste + FCString::Strlen(Paste), GWarn);
 
+		// Reinstate old BSP update setting, and force a rebuild - any levels whose geometry has changed while pasting will be rebuilt
+		GetMutableDefault<ULevelEditorMiscSettings>()->bBSPAutoUpdate = bBSPAutoUpdate;
+		RebuildAlteredBSP();
+
 		// Fire ULevel::LevelDirtiedEvent when falling out of scope.
 		FScopedLevelDirtied			LevelDirtyCallback;
 
@@ -738,6 +747,18 @@ bool UUnrealEdEngine::edactDeleteSelected( UWorld* InWorld, bool bVerifyDeletion
 	int32		DeleteCount = 0;
 
 	USelection* SelectedActors = GetSelectedActors();
+	TMap<AActor*, TArray<AActor*> > ReferencingActorsMap;
+	TArray<UClass*> ClassTypesToIgnore;
+	ClassTypesToIgnore.Add(ALevelScriptActor::StaticClass());
+	// The delete warning is meant for actor references that affect gameplay.  Group actors do not affect gameplay and should not show up as a warning.
+	ClassTypesToIgnore.Add(AGroupActor::StaticClass());
+
+	// If we want to warn about references to the actors to be deleted, it is a lot more efficient to query
+	// the world first and build a map of actors referenced by other actors. We can then quickly look this up later on in the loop.
+	if (bWarnAboutReferences)
+	{
+		FBlueprintEditorUtils::GetActorReferenceMap(InWorld, ClassTypesToIgnore, ReferencingActorsMap);
+	}
 
 	for ( int32 ActorIndex = 0 ; ActorIndex < ActorsToDelete.Num() ; ++ActorIndex )
 	{
@@ -745,41 +766,45 @@ bool UUnrealEdEngine::edactDeleteSelected( UWorld* InWorld, bool bVerifyDeletion
 
 		//If actor is referenced by script, ask user if they really want to delete
 		ULevelScriptBlueprint* LSB = Actor->GetLevel()->GetLevelScriptBlueprint(true);
-		TArray<AActor*> ReferencingActors;
-		TArray<UClass*> ClassTypesToIgnore;
-		ClassTypesToIgnore.Add( ALevelScriptActor::StaticClass() );
-		// The delete warning is meant for actor referneces that affect gameplay.  Group actors do not affect gameplay and should not show up as a warning.
-		ClassTypesToIgnore.Add( AGroupActor::StaticClass() );
+
+		// Get the array of actors that reference this actor from the cached map we built above.
+		TArray<AActor*>* ReferencingActors = nullptr;
 		if( bWarnAboutReferences )
 		{
-			FBlueprintEditorUtils::FindActorsThatReferenceActor( Actor, ClassTypesToIgnore, ReferencingActors );
+			ReferencingActors = ReferencingActorsMap.Find(Actor);
 		}
 
 		bool bReferencedByLevelScript = bWarnAboutReferences && (nullptr != LSB && FBlueprintEditorUtils::FindNumReferencesToActorFromLevelScript(LSB, Actor) > 0);
 		bool bReferencedByActor = false;
 		bool bReferencedByLODActor = false;
-		for (AActor* ReferencingActor : ReferencingActors)
+
+		// If there are any referencing actors, make sure that they are reference types that we care about.
+		if (ReferencingActors != nullptr)
 		{
-			if (ReferencingActor->IsA(ALODActor::StaticClass()))
-			{
-				bReferencedByLODActor = true;
-				break;
-			}
-			// If the referencing actor is a child actor that is referencing us, do not treat it
-			// as referencing for the purposes of warning about deletion
-			UChildActorComponent* ParentComponent = ReferencingActor->GetParentComponent();
-			if (ParentComponent == nullptr || ParentComponent->GetOwner() != Actor)
+			for (AActor* ReferencingActor : (*ReferencingActors))
 			{
-				bReferencedByActor = true;
-				break;
+				if (ReferencingActor->IsA(ALODActor::StaticClass()))
+				{
+					bReferencedByLODActor = true;
+					break;
+				}
+				// If the referencing actor is a child actor that is referencing us, do not treat it
+				// as referencing for the purposes of warning about deletion
+				UChildActorComponent* ParentComponent = ReferencingActor->GetParentComponent();
+				if (ParentComponent == nullptr || ParentComponent->GetOwner() != Actor)
+				{
+					bReferencedByActor = true;
+					break;
+				}
 			}
 		}
 
+		// We have references from one or more sources, prompt the user for feedback.
 		if (bReferencedByLevelScript || bReferencedByActor || bReferencedByLODActor)
 		{
-			if (( bReferencedByLevelScript && !bRequestedDeleteAllByLevel ) ||
-				( bReferencedByActor && !bRequestedDeleteAllByActor ) ||
-				(bReferencedByLODActor && !bRequestedDeleteAllByLODActor) )
+			if ((bReferencedByLevelScript && !bRequestedDeleteAllByLevel) ||
+				(bReferencedByActor && !bRequestedDeleteAllByActor) ||
+				(bReferencedByLODActor && !bRequestedDeleteAllByLODActor))
 			{
 				FText ConfirmDelete;
 
@@ -789,60 +814,60 @@ bool UUnrealEdEngine::edactDeleteSelected( UWorld* InWorld, bool bVerifyDeletion
 				{
 					ConfirmDelete = FText::Format(LOCTEXT("ConfirmDeleteActorReferencedByHLOD",
 						"Actor {0} is referenced by LODActor, do you really want to delete it?"),
-						FText::FromString(Actor->GetName()));
+						FText::FromString(Actor->GetActorLabel()));
 				}
 				else if (bReferencedByLevelScript && bReferencedByActor)
 				{
-					ConfirmDelete = FText::Format(LOCTEXT( "ConfirmDeleteActorReferenceByScriptAndActor",
-														   "Actor {0} is referenced by the level blueprint and another Actor, do you really want to delete it?"),
-														   FText::FromString( Actor->GetName() ) );
+					ConfirmDelete = FText::Format(LOCTEXT("ConfirmDeleteActorReferenceByScriptAndActor",
+						"Actor {0} is referenced by the level blueprint and another Actor, do you really want to delete it?"),
+						FText::FromString(Actor->GetActorLabel()));
 				}
-				else if ( bReferencedByLevelScript )
+				else if (bReferencedByLevelScript)
 				{
-					ConfirmDelete = FText::Format(LOCTEXT( "ConfirmDeleteActorReferencedByScript",
-														   "Actor {0} is referenced by the level blueprint, do you really want to delete it?"),
-														   FText::FromString( Actor->GetName() ) );
+					ConfirmDelete = FText::Format(LOCTEXT("ConfirmDeleteActorReferencedByScript",
+						"Actor {0} is referenced by the level blueprint, do you really want to delete it?"),
+						FText::FromString(Actor->GetActorLabel()));
 				}
 				else
 				{
-					ConfirmDelete = FText::Format(LOCTEXT( "ConfirmDeleteActorReferencedByActor",
-														   "Actor {0} is referenced by another Actor, do you really want to delete it?"),
-														   FText::FromString( Actor->GetName() ) );
+					ConfirmDelete = FText::Format(LOCTEXT("ConfirmDeleteActorReferencedByActor",
+						"Actor {0} is referenced by another Actor, do you really want to delete it?"),
+						FText::FromString(Actor->GetActorLabel()));
 				}
 
-				int32 Result = FMessageDialog::Open( MessageType, ConfirmDelete );
-				if ( Result == EAppReturnType::YesAll )
+				int32 Result = FMessageDialog::Open(MessageType, ConfirmDelete);
+				if (Result == EAppReturnType::YesAll)
 				{
 					bRequestedDeleteAllByLevel |= bReferencedByLevelScript;
 					bRequestedDeleteAllByActor |= bReferencedByActor;
 					bRequestedDeleteAllByLODActor |= bReferencedByLODActor;
 				}
-				else if ( Result == EAppReturnType::NoAll )
+				else if (Result == EAppReturnType::NoAll)
 				{
 					break;
 				}
-				else if( Result == EAppReturnType::No || Result == EAppReturnType::Cancel )
-				{ 
+				else if (Result == EAppReturnType::No || Result == EAppReturnType::Cancel)
+				{
 					continue;
 				}
 			}
 
-			if ( bReferencedByLevelScript )
+			if (bReferencedByLevelScript)
 			{
 				FBlueprintEditorUtils::ModifyActorReferencedGraphNodes(LSB, Actor);
 			}
-			if ( bReferencedByActor )
+			if (bReferencedByActor && ReferencingActors != nullptr)
 			{
-				for ( int32 ReferencingActorIndex = 0; ReferencingActorIndex < ReferencingActors.Num(); ReferencingActorIndex++ )
+				for (int32 ReferencingActorIndex = 0; ReferencingActorIndex < ReferencingActors->Num(); ReferencingActorIndex++)
 				{
-					ReferencingActors[ReferencingActorIndex]->Modify();
+					(*ReferencingActors)[ReferencingActorIndex]->Modify();
 				}
 			}
-			if (bReferencedByLODActor)
+			if (bReferencedByLODActor && ReferencingActors != nullptr)
 			{
-				for (int32 ReferencingActorIndex = 0; ReferencingActorIndex < ReferencingActors.Num(); ReferencingActorIndex++)
+				for (int32 ReferencingActorIndex = 0; ReferencingActorIndex < ReferencingActors->Num(); ReferencingActorIndex++)
 				{
-					ALODActor* LODActor = Cast<ALODActor>(ReferencingActors[ReferencingActorIndex]);
+					ALODActor* LODActor = Cast<ALODActor>((*ReferencingActors)[ReferencingActorIndex]);
 					// it's possible other actor is referencing this
 					if (LODActor)
 					{
@@ -852,6 +877,7 @@ bool UUnrealEdEngine::edactDeleteSelected( UWorld* InWorld, bool bVerifyDeletion
 			}
 		}
 	
+	
 		bool bRebuildNavigation = false;
 
 		ABrush* Brush = Cast< ABrush >( Actor );
diff --git a/Engine/Source/Editor/UnrealEd/Private/EditorCsg.cpp b/Engine/Source/Editor/UnrealEd/Private/EditorCsg.cpp
--- a/Engine/Source/Editor/UnrealEd/Private/EditorCsg.cpp
+++ b/Engine/Source/Editor/UnrealEd/Private/EditorCsg.cpp
@@ -285,7 +285,7 @@ void UEditorEngine::csgRebuild( UWorld* InWorld )
 	for( TActorIterator<ABrush> It(InWorld); It; ++It )
 	{
 		ABrush* B=*It;
-		if ( B->Brush && !B->IsStaticBrush() )
+		if ( B && B->GetLevel() == InWorld->GetCurrentLevel() && B->Brush && !B->IsStaticBrush() )
 		{
 			DynamicBrushes.Add(B);
 		}
@@ -1360,7 +1360,10 @@ static void SendTo( UWorld* InWorld, int32 bSendToFirst )
 	ULevel*	Level = InWorld->GetCurrentLevel();
 	for (AActor* Actor : Level->Actors)
 	{
-		Actor->Modify();
+		if(Actor)
+		{
+			Actor->Modify();
+		}
 	}
 	
 	// Fire ULevel::LevelDirtiedEvent when falling out of scope.
diff --git a/Engine/Source/Editor/UnrealEd/Private/EditorEngine.cpp b/Engine/Source/Editor/UnrealEd/Private/EditorEngine.cpp
--- a/Engine/Source/Editor/UnrealEd/Private/EditorEngine.cpp
+++ b/Engine/Source/Editor/UnrealEd/Private/EditorEngine.cpp
@@ -87,6 +87,7 @@
 #include "Engine/NetDriver.h"
 #include "Net/NetworkProfiler.h"
 #include "Interfaces/IPluginManager.h"
+#include "PackageReload.h"
 
 // needed for the RemotePropagator
 #include "AudioDevice.h"
@@ -107,7 +108,7 @@
 #include "Kismet2/BlueprintEditorUtils.h"
 #include "Kismet2/KismetEditorUtilities.h"
 #include "Kismet2/KismetDebugUtilities.h"
-
+#include "Kismet2/KismetReinstanceUtilities.h"
 
 #include "AssetRegistryModule.h"
 #include "IContentBrowserSingleton.h"
@@ -177,6 +178,7 @@
 #include "PhysicsPublic.h"
 #include "Engine/CoreSettings.h"
 #include "ShaderCompiler.h"
+#include "DistanceFieldAtlas.h"
 
 #include "PixelInspectorModule.h"
 
@@ -311,6 +313,7 @@ UEditorEngine::UEditorEngine(const FObjectInitializer& ObjectInitializer)
 	bDisableDeltaModification = false;
 	bPlayOnLocalPcSession = false;
 	bAllowMultiplePIEWorlds = true;
+	bIsEndingPlay = false;
 	NumOnlinePIEInstances = 0;
 	DefaultWorldFeatureLevel = GMaxRHIFeatureLevel;
 
@@ -619,6 +622,8 @@ void UEditorEngine::InitEditor(IEngineLoop* InEngineLoop)
 	FCoreUObjectDelegates::IsPackageOKToSaveDelegate.BindUObject(this, &UEditorEngine::IsPackageOKToSave);
 	FCoreUObjectDelegates::AutoPackageBackupDelegate.BindStatic(&FAutoPackageBackup::BackupPackage);
 
+	FCoreUObjectDelegates::OnPackageReloaded.AddUObject(this, &UEditorEngine::HandlePackageReloaded);
+
 	extern void SetupDistanceFieldBuildNotification();
 	SetupDistanceFieldBuildNotification();
 
@@ -643,6 +648,124 @@ bool UEditorEngine::HandleOpenAsset(UObject* Asset)
 	return FAssetEditorManager::Get().OpenEditorForAsset(Asset);
 }
 
+void UEditorEngine::HandlePackageReloaded(const EPackageReloadPhase InPackageReloadPhase, FPackageReloadedEvent* InPackageReloadedEvent)
+{
+	static TSet<UBlueprint*> BlueprintsToRecompileThisBatch;
+
+	if (InPackageReloadPhase == EPackageReloadPhase::PrePackageFixup)
+	{
+		NotifyToolsOfObjectReplacement(InPackageReloadedEvent->GetRepointedObjects());
+
+		// Notify any Blueprint assets that are about to be unloaded.
+		ForEachObjectWithOuter(InPackageReloadedEvent->GetOldPackage(), [&](UObject* InObject)
+		{
+			if (InObject->IsAsset())
+			{
+				// Notify about any BP assets that are about to be unloaded
+				if (UBlueprint* BP = Cast<UBlueprint>(InObject))
+				{
+					FKismetEditorUtilities::OnBlueprintUnloaded.Broadcast(BP);
+				}
+			}
+		}, false, RF_Transient, EInternalObjectFlags::PendingKill);
+	}
+
+	if (InPackageReloadPhase == EPackageReloadPhase::OnPackageFixup)
+	{
+		for (const auto& RepointedObjectPair : InPackageReloadedEvent->GetRepointedObjects())
+		{
+			UObject* OldObject = RepointedObjectPair.Key;
+			UObject* NewObject = RepointedObjectPair.Value;
+		
+			if (OldObject->IsAsset())
+			{
+				if (const UBlueprint* OldBlueprint = Cast<UBlueprint>(OldObject))
+				{
+					FBlueprintCompileReinstancer::ReplaceInstancesOfClass(OldBlueprint->GeneratedClass, NewObject ? CastChecked<UBlueprint>(NewObject)->GeneratedClass : nullptr);
+				}
+			}
+		}
+	}
+
+	if (InPackageReloadPhase == EPackageReloadPhase::PostPackageFixup)
+	{
+		for (TWeakObjectPtr<UObject> ObjectReferencer : InPackageReloadedEvent->GetObjectReferencers())
+		{
+			UObject* ObjectReferencerPtr = ObjectReferencer.Get();
+			if (!ObjectReferencerPtr)
+			{
+				continue;
+			}
+
+			FPropertyChangedEvent PropertyEvent(nullptr, EPropertyChangeType::Redirected);
+			ObjectReferencerPtr->PostEditChangeProperty(PropertyEvent);
+
+			// We need to recompile any Blueprints that had properties changed to make sure their generated class is up-to-date and has no lingering references to the old objects
+			UBlueprint* BlueprintToRecompile = nullptr;
+			if (UBlueprint* BlueprintReferencer = Cast<UBlueprint>(ObjectReferencerPtr))
+			{
+				BlueprintToRecompile = BlueprintReferencer;
+			}
+			else if (UClass* ClassReferencer = Cast<UClass>(ObjectReferencerPtr))
+			{
+				BlueprintToRecompile = Cast<UBlueprint>(ClassReferencer->ClassGeneratedBy);
+			}
+			else
+			{
+				BlueprintToRecompile = ObjectReferencerPtr->GetTypedOuter<UBlueprint>();
+			}
+			
+			if (BlueprintToRecompile)
+			{
+				BlueprintsToRecompileThisBatch.Add(BlueprintToRecompile);
+			}
+		}
+	}
+
+	if (InPackageReloadPhase == EPackageReloadPhase::PreBatch)
+	{
+		// If this fires then ReloadPackages has probably bee called recursively :(
+		check(BlueprintsToRecompileThisBatch.Num() == 0);
+
+		// Flush all pending render commands, as reloading the package may invalidate render resources.
+		FlushRenderingCommands();
+	}
+
+	if (InPackageReloadPhase == EPackageReloadPhase::PostBatchPreGC)
+	{
+		// Make sure we don't have any lingering transaction buffer references.
+		GEditor->Trans->Reset(NSLOCTEXT("UnrealEd", "ReloadedPackage", "Reloaded Package"));
+
+		// Recompile any BPs that had their references updated
+		if (BlueprintsToRecompileThisBatch.Num() > 0)
+		{
+			FScopedSlowTask CompilingBlueprintsSlowTask(BlueprintsToRecompileThisBatch.Num(), NSLOCTEXT("UnrealEd", "CompilingBlueprints", "Compiling Blueprints"));
+
+			for (UBlueprint* BlueprintToRecompile : BlueprintsToRecompileThisBatch)
+			{
+				CompilingBlueprintsSlowTask.EnterProgressFrame(1.0f);
+
+				//FBlueprintEditorUtils::MarkBlueprintAsModified(BlueprintToRecompile, FPropertyChangedEvent(nullptr, EPropertyChangeType::Redirected));
+				FKismetEditorUtilities::CompileBlueprint(BlueprintToRecompile, /*bIsRegeneratingOnLoad*/false, /*bSkipGarbageCollection*/true);
+			}
+		}
+		BlueprintsToRecompileThisBatch.Reset();
+	}
+
+	if (InPackageReloadPhase == EPackageReloadPhase::PostBatchPostGC)
+	{
+		// Tick some things that aren't processed while we're reloading packages and can result in excessive memory usage if not periodically updated.
+		if (GShaderCompilingManager)
+		{
+			GShaderCompilingManager->ProcessAsyncResults(true, false);
+		}
+		if (GDistanceFieldAsyncQueue)
+		{
+			GDistanceFieldAsyncQueue->ProcessAsyncTasks();
+		}
+	}
+}
+
 void UEditorEngine::HandleSettingChanged( FName Name )
 {
 	// When settings are reset to default, the property name will be "None" so make sure that case is handled.
diff --git a/Engine/Source/Editor/UnrealEd/Private/EditorSelectUtils.cpp b/Engine/Source/Editor/UnrealEd/Private/EditorSelectUtils.cpp
--- a/Engine/Source/Editor/UnrealEd/Private/EditorSelectUtils.cpp
+++ b/Engine/Source/Editor/UnrealEd/Private/EditorSelectUtils.cpp
@@ -315,7 +315,7 @@ void UUnrealEdEngine::UpdatePivotLocationForSelection( bool bOnChange )
 			{
 				USelection* SelectedActors = GetSelectedActors();
 				const bool bIsOwnerSelected = SelectedActors->IsSelected(ComponentOwner);
-				check(bIsOwnerSelected);
+				ensureMsgf(bIsOwnerSelected, TEXT("Owner(%s) of %s is not selected"), *ComponentOwner->GetFullName(), *Component->GetFullName());
 
 				if (ComponentOwner->GetWorld() == GWorld)
 				{
diff --git a/Engine/Source/Editor/UnrealEd/Private/EditorViewportClient.cpp b/Engine/Source/Editor/UnrealEd/Private/EditorViewportClient.cpp
--- a/Engine/Source/Editor/UnrealEd/Private/EditorViewportClient.cpp
+++ b/Engine/Source/Editor/UnrealEd/Private/EditorViewportClient.cpp
@@ -338,6 +338,8 @@ FEditorViewportClient::FEditorViewportClient(FEditorModeTools* InModeTools, FPre
 	, FarPlane(0.0f)
 	, bInGameViewMode(false)
 	, bShouldInvalidateViewportWidget(false)
+	, DragStartView(nullptr)
+	, DragStartViewFamily(nullptr)
 {
 	InitViewOptionsArray();
 	if (ModeTools == nullptr)
@@ -1997,15 +1999,17 @@ void FEditorViewportClient::UpdateMouseDelta()
 				//if Absolute Translation, and not just moving the camera around
 				if (IsUsingAbsoluteTranslation())
 				{
-					// Compute a view.
-					FSceneViewFamilyContext ViewFamily(FSceneViewFamily::ConstructionValues(
-						Viewport,
-						GetScene(),
-						EngineShowFlags)
-						.SetRealtimeUpdate( IsRealtime() ) );
-					FSceneView* View = CalcSceneView( &ViewFamily );
-
-					MouseDeltaTracker->AbsoluteTranslationConvertMouseToDragRot(View, this, Drag, Rot, Scale);
+					if (DragStartView == nullptr)
+					{
+						// Compute a view.
+						DragStartViewFamily = new FSceneViewFamily(FSceneViewFamily::ConstructionValues(
+							Viewport,
+							GetScene(),
+							EngineShowFlags)
+							.SetRealtimeUpdate(IsRealtime()));
+						DragStartView = CalcSceneView(DragStartViewFamily);
+					}
+					MouseDeltaTracker->AbsoluteTranslationConvertMouseToDragRot(DragStartView, this, Drag, Rot, Scale);
 				} 
 				else
 				{
@@ -2198,7 +2202,7 @@ void FEditorViewportClient::MarkMouseMovedSinceClick()
 bool FEditorViewportClient::IsUsingAbsoluteTranslation() const
 {
 	bool bIsHotKeyAxisLocked = Viewport->KeyState(EKeys::LeftControl) || Viewport->KeyState(EKeys::RightControl);
-	bool bCameraLockedToWidget = Viewport->KeyState(EKeys::LeftShift) || Viewport->KeyState(EKeys::RightShift);
+	bool bCameraLockedToWidget = !(Widget && Widget->GetCurrentAxis() & EAxisList::Screen) && (Viewport->KeyState(EKeys::LeftShift) || Viewport->KeyState(EKeys::RightShift));
 	// Screen-space movement must always use absolute translation
 	bool bScreenSpaceTransformation = Widget && (Widget->GetCurrentAxis() == EAxisList::Screen);
 	bool bAbsoluteMovementEnabled = GetDefault<ULevelEditorViewportSettings>()->bUseAbsoluteTranslation || bScreenSpaceTransformation;
@@ -2414,6 +2418,12 @@ void FEditorViewportClient::StopTracking()
 {
 	if( bIsTracking )
 	{
+		DragStartView = nullptr;
+		if (DragStartViewFamily != nullptr)
+		{
+			delete DragStartViewFamily;
+			DragStartViewFamily = nullptr;
+		}
 		MouseDeltaTracker->EndTracking( this );
 
 		Widget->SetCurrentAxis( EAxisList::None );
diff --git a/Engine/Source/Editor/UnrealEd/Private/Factories/ActorFactory.cpp b/Engine/Source/Editor/UnrealEd/Private/Factories/ActorFactory.cpp
--- a/Engine/Source/Editor/UnrealEd/Private/Factories/ActorFactory.cpp
+++ b/Engine/Source/Editor/UnrealEd/Private/Factories/ActorFactory.cpp
@@ -399,6 +399,7 @@ const FName UActorFactoryBasicShape::BasicCube("/Engine/BasicShapes/Cube.Cube");
 const FName UActorFactoryBasicShape::BasicSphere("/Engine/BasicShapes/Sphere.Sphere");
 const FName UActorFactoryBasicShape::BasicCylinder("/Engine/BasicShapes/Cylinder.Cylinder");
 const FName UActorFactoryBasicShape::BasicCone("/Engine/BasicShapes/Cone.Cone");
+const FName UActorFactoryBasicShape::BasicPlane("/Engine/BasicShapes/Plane.Plane");
 
 UActorFactoryBasicShape::UActorFactoryBasicShape(const FObjectInitializer& ObjectInitializer)
 	: Super(ObjectInitializer)
@@ -410,7 +411,7 @@ UActorFactoryBasicShape::UActorFactoryBasicShape(const FObjectInitializer& Objec
 
 bool UActorFactoryBasicShape::CanCreateActorFrom( const FAssetData& AssetData, FText& OutErrorMsg )
 {
-	if(AssetData.IsValid() && (AssetData.ObjectPath == BasicCube || AssetData.ObjectPath == BasicSphere || AssetData.ObjectPath == BasicCone || AssetData.ObjectPath == BasicCylinder ) )
+	if(AssetData.IsValid() && (AssetData.ObjectPath == BasicCube || AssetData.ObjectPath == BasicSphere || AssetData.ObjectPath == BasicCone || AssetData.ObjectPath == BasicCylinder || AssetData.ObjectPath == BasicPlane) )
 	{
 		return true;
 	}
diff --git a/Engine/Source/Editor/UnrealEd/Private/Factories/EditorFactories.cpp b/Engine/Source/Editor/UnrealEd/Private/Factories/EditorFactories.cpp
--- a/Engine/Source/Editor/UnrealEd/Private/Factories/EditorFactories.cpp
+++ b/Engine/Source/Editor/UnrealEd/Private/Factories/EditorFactories.cpp
@@ -191,6 +191,7 @@
 #include "Interfaces/IImageWrapperModule.h"
 
 #include "FbxImporter.h"
+#include "FbxErrors.h"
 
 #include "AssetRegistryModule.h"
 #include "IContentBrowserSingleton.h"
@@ -4396,7 +4397,10 @@ UObject* UFontFileImportFactory::FactoryCreateBinary(UClass* InClass, UObject* I
 	if (FontFace)
 	{
 		FontFace->SourceFilename = GetCurrentFilename();
-		FontFace->FontFaceData.Append(InBuffer, InBufferEnd - InBuffer);
+
+		TArray<uint8> FontData;
+		FontData.Append(InBuffer, InBufferEnd - InBuffer);
+		FontFace->FontFaceData->SetData(MoveTemp(FontData));
 	}
 
 	FEditorDelegates::OnAssetPostImport.Broadcast(this, FontFace);
@@ -4931,6 +4935,10 @@ EReimportResult::Type UReimportFbxStaticMeshFactory::Reimport( UObject* Obj )
 	
 	UnFbx::FFbxImporter* FFbxImporter = UnFbx::FFbxImporter::GetInstance();
 	UnFbx::FBXImportOptions* ImportOptions = FFbxImporter->GetImportOptions();
+	
+	//Pop the message log in case of error
+	UnFbx::FFbxLoggerSetter Logger(FFbxImporter, true);
+
 	//Clean up the options
 	UnFbx::FBXImportOptions::ResetOptions(ImportOptions);
 
@@ -4939,19 +4947,12 @@ EReimportResult::Type UReimportFbxStaticMeshFactory::Reimport( UObject* Obj )
 	UFbxImportUI* ReimportUI = NewObject<UFbxImportUI>();
 	ReimportUI->MeshTypeToImport = FBXIT_StaticMesh;
 	ReimportUI->bOverrideFullName = false;
-	ReimportUI->bCombineMeshes = true;
+	ReimportUI->StaticMeshImportData->bCombineMeshes = true;
 
 	if (!ImportUI)
 	{
 		ImportUI = NewObject<UFbxImportUI>(this, NAME_None, RF_Public);
 	}
-	else
-	{
-		//Set misc options
-		ReimportUI->bConvertScene = ImportUI->bConvertScene;
-		ReimportUI->bForceFrontXAxis = ImportUI->bForceFrontXAxis;
-		ReimportUI->bConvertSceneUnit = ImportUI->bConvertSceneUnit;
-	}
 
 	if( ImportData )
 	{
@@ -5061,6 +5062,8 @@ EReimportResult::Type UReimportFbxStaticMeshFactory::Reimport( UObject* Obj )
 				{
 					Mesh->MarkPackageDirty();
 				}
+
+				FFbxImporter->ImportStaticMeshGlobalSockets(Mesh);
 			}
 			else
 			{
@@ -5183,13 +5186,6 @@ EReimportResult::Type UReimportFbxSkeletalMeshFactory::Reimport( UObject* Obj )
 	{
 		ImportUI = NewObject<UFbxImportUI>(this, NAME_None, RF_Public);
 	}
-	else
-	{
-		//Set misc options
-		ReimportUI->bConvertScene = ImportUI->bConvertScene;
-		ReimportUI->bForceFrontXAxis = ImportUI->bForceFrontXAxis;
-		ReimportUI->bConvertSceneUnit = ImportUI->bConvertSceneUnit;
-	}
 
 	bool bSuccess = false;
 
@@ -5403,6 +5399,9 @@ EReimportResult::Type UReimportFbxAnimSequenceFactory::Reimport( UObject* Obj )
 
 	UnFbx::FFbxImporter* Importer = UnFbx::FFbxImporter::GetInstance();
 
+	//Pop the message log in case of error
+	UnFbx::FFbxLoggerSetter Logger(Importer, true);
+
 	CurrentFilename = Filename;
 
 	USkeleton* Skeleton = AnimSequence->GetSkeleton();
@@ -5415,6 +5414,7 @@ EReimportResult::Type UReimportFbxAnimSequenceFactory::Reimport( UObject* Obj )
 			// If skeleton wasn't found or the user canceled out of the dialog, we cannot proceed, but this reimport factory 
 			// has still technically "handled" the reimport, so return true instead of false
 			UE_LOG(LogEditorFactories, Warning, TEXT("-- import failed") );
+			Importer->AddTokenizedErrorMessage(FTokenizedMessage::Create(EMessageSeverity::Error, LOCTEXT("Error_CouldNotFindSkeleton", "Cannot re-import animation with no skeleton.\nImport failed.")), FFbxErrors::SkeletalMesh_NoBoneFound);
 			return EReimportResult::Succeeded;
 		}
 	}
@@ -5439,6 +5439,7 @@ EReimportResult::Type UReimportFbxAnimSequenceFactory::Reimport( UObject* Obj )
 	else
 	{
 		UE_LOG(LogEditorFactories, Warning, TEXT("-- import failed") );
+		Importer->AddTokenizedErrorMessage(FTokenizedMessage::Create(EMessageSeverity::Error, LOCTEXT("Error_CouldNotReimportAnimation", "Cannot re-import animation.")), FFbxErrors::Generic_ReimportingObjectFailed);
 	}
 
 	Importer->ReleaseScene(); 
diff --git a/Engine/Source/Editor/UnrealEd/Private/Factories/SkeletalMeshImport.cpp b/Engine/Source/Editor/UnrealEd/Private/Factories/SkeletalMeshImport.cpp
--- a/Engine/Source/Editor/UnrealEd/Private/Factories/SkeletalMeshImport.cpp
+++ b/Engine/Source/Editor/UnrealEd/Private/Factories/SkeletalMeshImport.cpp
@@ -308,6 +308,9 @@ void ProcessImportMeshInfluences(FSkeletalMeshImportData& ImportData)
 	float TotalWeight		= 0.f;
 	const float MINWEIGHT   = 0.01f;
 
+	int MaxVertexInfluence = 0;
+	float MaxIgnoredWeight = 0.0f;
+
 	//We have to normalize the data before filtering influences
 	//Because influence filtering is base on the normalize value.
 	//Some DCC like Daz studio don't have normalized weight
@@ -328,10 +331,28 @@ void ProcessImportMeshInfluences(FSkeletalMeshImportData& ImportData)
 					Influences[i - r].Weight *= OneOverTotalWeight;
 				}
 			}
+			
+			if (MaxVertexInfluence < InfluenceCount)
+			{
+				MaxVertexInfluence = InfluenceCount;
+			}
+
 			// clear to count next one
 			InfluenceCount = 0;
 			TotalWeight = 0.f;
 		}
+
+		if (InfluenceCount > MAX_TOTAL_INFLUENCES &&  Influences[i].Weight > MaxIgnoredWeight)
+		{
+			MaxIgnoredWeight = Influences[i].Weight;
+		}
+	}
+ 
+ 	// warn about too many influences
+	if (MaxVertexInfluence > MAX_TOTAL_INFLUENCES)
+	{
+		UnFbx::FFbxImporter* FFbxImporter = UnFbx::FFbxImporter::GetInstance();
+		FFbxImporter->AddTokenizedErrorMessage(FTokenizedMessage::Create(EMessageSeverity::Warning, FText::Format(LOCTEXT("WarningTooManySkelInfluences", "Warning skeletal mesh influence count of {0} exceeds max count of {1}. Influence truncation will occur. Maximum Ignored Weight {2}"), MaxVertexInfluence, MAX_TOTAL_INFLUENCES, MaxIgnoredWeight)), FFbxErrors::SkeletalMesh_TooManyInfluences);
 	}
 
 	for( int32 i=0; i<Influences.Num(); i++ )
diff --git a/Engine/Source/Editor/UnrealEd/Private/Fbx/FbxAssetImportData.cpp b/Engine/Source/Editor/UnrealEd/Private/Fbx/FbxAssetImportData.cpp
--- a/Engine/Source/Editor/UnrealEd/Private/Fbx/FbxAssetImportData.cpp
+++ b/Engine/Source/Editor/UnrealEd/Private/Fbx/FbxAssetImportData.cpp
@@ -7,6 +7,9 @@ UFbxAssetImportData::UFbxAssetImportData(const FObjectInitializer& ObjectInitial
 	, ImportTranslation(0)
 	, ImportRotation(0)
 	, ImportUniformScale(1.0f)
+	, bConvertScene(true)
+	, bForceFrontXAxis(false)
+	, bConvertSceneUnit(true)
 	, bImportAsScene(false)
 	, FbxSceneImportDataReference(nullptr)
 {
diff --git a/Engine/Source/Editor/UnrealEd/Private/Fbx/FbxFactory.cpp b/Engine/Source/Editor/UnrealEd/Private/Fbx/FbxFactory.cpp
--- a/Engine/Source/Editor/UnrealEd/Private/Fbx/FbxFactory.cpp
+++ b/Engine/Source/Editor/UnrealEd/Private/Fbx/FbxFactory.cpp
@@ -91,7 +91,7 @@ bool UFbxFactory::DetectImportType(const FString& InFilename)
 		FFbxImporter->AddTokenizedErrorMessage(FTokenizedMessage::Create(EMessageSeverity::Warning, LOCTEXT("NoImportTypeDetected", "Can't detect import type. No mesh is found or animation track.")), FFbxErrors::Generic_CannotDetectImportType);
 		return false;
 	}
-	else
+	else if(!IsAutomatedImport() || ImportUI->bAutomatedImportShouldDetectType)
 	{
 		ImportUI->MeshTypeToImport = EFBXImportType(ImportType);
 		ImportUI->OriginalImportType = ImportUI->MeshTypeToImport;
@@ -247,7 +247,8 @@ UObject* UFbxFactory::FactoryCreateBinary
 			TArray< TArray<FbxNode*>* > SkelMeshArray;
 
 			bool bImportStaticMeshLODs = ImportUI->StaticMeshImportData->bImportMeshLODs;
-			bool bCombineMeshes = ImportUI->bCombineMeshes;
+			bool bCombineMeshes = ImportUI->StaticMeshImportData->bCombineMeshes;
+			bool bCombineMeshesLOD = false;
 
 			if ( ImportUI->MeshTypeToImport == FBXIT_SkeletalMesh )
 			{
@@ -271,9 +272,11 @@ UObject* UFbxFactory::FactoryCreateBinary
 					InterestingNodeCount = FbxImporter->GetFbxMeshCount(RootNodeToImport,bCountLODGroupMeshes,NumLODGroups);
 
 					// if there were LODs in the file, do not combine meshes even if requested
-					if( bImportStaticMeshLODs && bCombineMeshes )
+					if( bImportStaticMeshLODs && bCombineMeshes && NumLODGroups > 0)
 					{
-						bCombineMeshes = NumLODGroups == 0;
+						bCombineMeshes = false;
+						//Combine all the LOD together and export one mesh with LODs
+						bCombineMeshesLOD = true;
 					}
 				}
 			}
@@ -305,6 +308,53 @@ UObject* UFbxFactory::FactoryCreateBinary
 
 						ImportedMeshCount = NewStaticMesh ? 1 : 0;
 					}
+					else if (bCombineMeshesLOD)
+					{
+						TArray<FbxNode*> FbxMeshArray;
+						TArray<FbxNode*> FbxLodGroups;
+						TArray<TArray<FbxNode*>> FbxMeshesLod;
+						FbxImporter->FillFbxMeshAndLODGroupArray(RootNodeToImport, FbxLodGroups, FbxMeshArray);
+						FbxMeshesLod.Add(FbxMeshArray);
+						for (FbxNode* LODGroup : FbxLodGroups)
+						{
+							if (LODGroup->GetNodeAttribute() && LODGroup->GetNodeAttribute()->GetAttributeType() == FbxNodeAttribute::eLODGroup && LODGroup->GetChildCount() > 0)
+							{
+								for (int32 GroupLodIndex = 0; GroupLodIndex < LODGroup->GetChildCount(); ++GroupLodIndex)
+								{
+									TArray<FbxNode*> AllNodeInLod;
+									FbxImporter->FindAllLODGroupNode(AllNodeInLod, LODGroup, GroupLodIndex);
+									if (AllNodeInLod.Num() > 0)
+									{
+										if (FbxMeshesLod.Num() <= GroupLodIndex)
+										{
+											FbxMeshesLod.Add(AllNodeInLod);
+										}
+										else
+										{
+											TArray<FbxNode*> &LODGroupArray = FbxMeshesLod[GroupLodIndex];
+											for (FbxNode* NodeToAdd : AllNodeInLod)
+											{
+												LODGroupArray.Add(NodeToAdd);
+											}
+										}
+									}
+								}
+							}
+						}
+
+						//Import the LOD root
+						if (FbxMeshesLod.Num() > 0)
+						{
+							TArray<FbxNode*> &LODMeshesArray = FbxMeshesLod[0];
+							NewStaticMesh = FbxImporter->ImportStaticMeshAsSingle(InParent, LODMeshesArray, Name, Flags, ImportUI->StaticMeshImportData, NULL, 0);
+						}
+						//Import all LODs
+						for (int32 LODIndex = 1; LODIndex < FbxMeshesLod.Num(); ++LODIndex)
+						{
+							TArray<FbxNode*> &LODMeshesArray = FbxMeshesLod[LODIndex];
+							FbxImporter->ImportStaticMeshAsSingle(InParent, LODMeshesArray, Name, Flags, ImportUI->StaticMeshImportData, NewStaticMesh, LODIndex);
+						}
+					}
 					else
 					{
 						TArray<UObject*> AllNewAssets;
@@ -326,10 +376,10 @@ UObject* UFbxFactory::FactoryCreateBinary
 						ImportedMeshCount = AllNewAssets.Num();
 					}
 
-					// Importing static mesh sockets only works if one mesh is being imported
-					if( ImportedMeshCount == 1 && NewStaticMesh )
+					// Importing static mesh global sockets only if one mesh is imported
+					if( ImportedMeshCount == 1 && NewStaticMesh)
 					{
-						FbxImporter->ImportStaticMeshSockets( NewStaticMesh );
+						FbxImporter->ImportStaticMeshGlobalSockets( NewStaticMesh );
 					}
 
 					NewObject = NewStaticMesh;
@@ -671,8 +721,7 @@ IImportSettingsParser* UFbxFactory::GetImportSettingsParser()
 UFbxImportUI::UFbxImportUI(const FObjectInitializer& ObjectInitializer)
 	: Super(ObjectInitializer)
 {
-	bCombineMeshes = true;
-
+	bAutomatedImportShouldDetectType = true;
 	StaticMeshImportData = CreateDefaultSubobject<UFbxStaticMeshImportData>(TEXT("StaticMeshImportData"));
 	SkeletalMeshImportData = CreateDefaultSubobject<UFbxSkeletalMeshImportData>(TEXT("SkeletalMeshImportData"));
 	AnimSequenceImportData = CreateDefaultSubobject<UFbxAnimSequenceImportData>(TEXT("AnimSequenceImportData"));
@@ -711,6 +760,13 @@ void UFbxImportUI::ParseFromJson(TSharedRef<class FJsonObject> ImportSettingsJso
 	int64 SkipFlags = CPF_InstancedReference;
 	FJsonObjectConverter::JsonObjectToUStruct(ImportSettingsJson, GetClass(), this, 0, SkipFlags);
 
+	bAutomatedImportShouldDetectType = true;
+	if(ImportSettingsJson->TryGetField("MeshTypeToImport").IsValid())
+	{
+		// Import type was specified by the user if MeshTypeToImport exists
+		bAutomatedImportShouldDetectType = false;
+	}
+
 	const TSharedPtr<FJsonObject>* StaticMeshImportJson = nullptr;
 	ImportSettingsJson->TryGetObjectField(TEXT("StaticMeshImportData"), StaticMeshImportJson);
 	if(StaticMeshImportJson)
diff --git a/Engine/Source/Editor/UnrealEd/Private/Fbx/FbxMainExport.cpp b/Engine/Source/Editor/UnrealEd/Private/Fbx/FbxMainExport.cpp
--- a/Engine/Source/Editor/UnrealEd/Private/Fbx/FbxMainExport.cpp
+++ b/Engine/Source/Editor/UnrealEd/Private/Fbx/FbxMainExport.cpp
@@ -973,7 +973,32 @@ void FFbxExporter::ExportStaticMesh( UStaticMesh* StaticMesh, const TArray<FStat
 	StaticMesh->GetName(MeshName);
 	FbxNode* MeshNode = FbxNode::Create(Scene, TCHAR_TO_UTF8(*MeshName));
 	Scene->GetRootNode()->AddChild(MeshNode);
-	ExportStaticMeshToFbx(StaticMesh, 0, *MeshName, MeshNode, -1, NULL, MaterialOrder);
+
+	if (StaticMesh->GetNumLODs() > 1)
+	{
+		FString LodGroup_MeshName = MeshName + ("_LodGroup");
+		FbxLODGroup *FbxLodGroupAttribute = FbxLODGroup::Create(Scene, TCHAR_TO_UTF8(*LodGroup_MeshName));
+		MeshNode->AddNodeAttribute(FbxLodGroupAttribute);
+		FbxLodGroupAttribute->ThresholdsUsedAsPercentage = true;
+		//Export an Fbx Mesh Node for every LOD and child them to the fbx node (LOD Group)
+		for (int CurrentLodIndex = 0; CurrentLodIndex < StaticMesh->GetNumLODs(); ++CurrentLodIndex)
+		{
+			FString FbxLODNodeName = MeshName + TEXT("_LOD") + FString::FromInt(CurrentLodIndex);
+			FbxNode* FbxActorLOD = FbxNode::Create(Scene, TCHAR_TO_UTF8(*FbxLODNodeName));
+			MeshNode->AddChild(FbxActorLOD);
+			if (CurrentLodIndex + 1 < StaticMesh->GetNumLODs())
+			{
+				//Convert the screen size to a threshold, it is just to be sure that we set some threshold, there is no way to convert this precisely
+				double LodScreenSize = (double)(10.0f / StaticMesh->RenderData->ScreenSize[CurrentLodIndex]);
+				FbxLodGroupAttribute->AddThreshold(LodScreenSize);
+			}
+			ExportStaticMeshToFbx(StaticMesh, CurrentLodIndex, *MeshName, FbxActorLOD, -1, nullptr, MaterialOrder);
+		}
+	}
+	else
+	{
+		ExportStaticMeshToFbx(StaticMesh, 0, *MeshName, MeshNode, -1, NULL, MaterialOrder);
+	}
 }
 
 void FFbxExporter::ExportStaticMeshLightMap( UStaticMesh* StaticMesh, int32 LODIndex, int32 UVChannel )
diff --git a/Engine/Source/Editor/UnrealEd/Private/Fbx/FbxMainImport.cpp b/Engine/Source/Editor/UnrealEd/Private/Fbx/FbxMainImport.cpp
--- a/Engine/Source/Editor/UnrealEd/Private/Fbx/FbxMainImport.cpp
+++ b/Engine/Source/Editor/UnrealEd/Private/Fbx/FbxMainImport.cpp
@@ -183,6 +183,7 @@ void ApplyImportUIToImportOptions(UFbxImportUI* ImportUI, FBXImportOptions& InOu
 	check(ImportUI);
 	InOutImportOptions.bImportMaterials = ImportUI->bImportMaterials;
 	InOutImportOptions.bInvertNormalMap = ImportUI->TextureImportData->bInvertNormalMaps;
+	InOutImportOptions.MaterialSearchLocation = ImportUI->TextureImportData->MaterialSearchLocation;
 	UMaterialInterface* BaseMaterialInterface = Cast<UMaterialInterface>(ImportUI->TextureImportData->BaseMaterialName.TryLoad());
 	if (BaseMaterialInterface) {
 		InOutImportOptions.BaseMaterial = BaseMaterialInterface;
@@ -195,9 +196,6 @@ void ApplyImportUIToImportOptions(UFbxImportUI* ImportUI, FBXImportOptions& InOu
 	}
 	InOutImportOptions.bImportTextures = ImportUI->bImportTextures;
 	InOutImportOptions.bUsedAsFullName = ImportUI->bOverrideFullName;
-	InOutImportOptions.bConvertScene = ImportUI->bConvertScene;
-	InOutImportOptions.bForceFrontXAxis = ImportUI->bForceFrontXAxis;
-	InOutImportOptions.bConvertSceneUnit = ImportUI->bConvertSceneUnit;
 	InOutImportOptions.bImportAnimations = ImportUI->bImportAnimations;
 	InOutImportOptions.SkeletonForAnimation = ImportUI->Skeleton;
 
@@ -212,6 +210,9 @@ void ApplyImportUIToImportOptions(UFbxImportUI* ImportUI, FBXImportOptions& InOu
 		InOutImportOptions.bTransformVertexToAbsolute = StaticMeshData->bTransformVertexToAbsolute;
 		InOutImportOptions.bBakePivotInVertex		= StaticMeshData->bBakePivotInVertex;
 		InOutImportOptions.bImportStaticMeshLODs	= StaticMeshData->bImportMeshLODs;
+		InOutImportOptions.bConvertScene			= StaticMeshData->bConvertScene;
+		InOutImportOptions.bForceFrontXAxis			= StaticMeshData->bForceFrontXAxis;
+		InOutImportOptions.bConvertSceneUnit		= StaticMeshData->bConvertSceneUnit;
 	}
 	else if ( ImportUI->MeshTypeToImport == FBXIT_SkeletalMesh )
 	{
@@ -224,6 +225,9 @@ void ApplyImportUIToImportOptions(UFbxImportUI* ImportUI, FBXImportOptions& InOu
 		InOutImportOptions.bTransformVertexToAbsolute	= SkeletalMeshData->bTransformVertexToAbsolute;
 		InOutImportOptions.bBakePivotInVertex			= SkeletalMeshData->bBakePivotInVertex;
 		InOutImportOptions.bImportSkeletalMeshLODs		= SkeletalMeshData->bImportMeshLODs;
+		InOutImportOptions.bConvertScene				= SkeletalMeshData->bConvertScene;
+		InOutImportOptions.bForceFrontXAxis				= SkeletalMeshData->bForceFrontXAxis;
+		InOutImportOptions.bConvertSceneUnit			= SkeletalMeshData->bConvertSceneUnit;
 
 		if(ImportUI->bImportAnimations)
 		{
@@ -232,6 +236,9 @@ void ApplyImportUIToImportOptions(UFbxImportUI* ImportUI, FBXImportOptions& InOu
 			AnimData->ImportTranslation				= SkeletalMeshData->ImportTranslation;
 			AnimData->ImportRotation				= SkeletalMeshData->ImportRotation;
 			AnimData->ImportUniformScale			= SkeletalMeshData->ImportUniformScale;
+			AnimData->bConvertScene					= SkeletalMeshData->bConvertScene;
+			AnimData->bForceFrontXAxis				= SkeletalMeshData->bForceFrontXAxis;
+			AnimData->bConvertSceneUnit				= SkeletalMeshData->bConvertSceneUnit;
 		}
 
 	}
@@ -242,6 +249,9 @@ void ApplyImportUIToImportOptions(UFbxImportUI* ImportUI, FBXImportOptions& InOu
 		InOutImportOptions.ImportTranslation	= AnimData->ImportTranslation;
 		InOutImportOptions.ImportRotation		= AnimData->ImportRotation;
 		InOutImportOptions.ImportUniformScale	= AnimData->ImportUniformScale;
+		InOutImportOptions.bConvertScene		= AnimData->bConvertScene;
+		InOutImportOptions.bForceFrontXAxis		= AnimData->bForceFrontXAxis;
+		InOutImportOptions.bConvertSceneUnit	= AnimData->bConvertSceneUnit;
 	}
 
 	InOutImportOptions.bImportMorph = ImportUI->SkeletalMeshImportData->bImportMorphTargets;
@@ -250,7 +260,7 @@ void ApplyImportUIToImportOptions(UFbxImportUI* ImportUI, FBXImportOptions& InOu
 	InOutImportOptions.bUseT0AsRefPose = ImportUI->SkeletalMeshImportData->bUseT0AsRefPose;
 	InOutImportOptions.bPreserveSmoothingGroups = ImportUI->SkeletalMeshImportData->bPreserveSmoothingGroups;
 	InOutImportOptions.bKeepOverlappingVertices = ImportUI->SkeletalMeshImportData->bKeepOverlappingVertices;
-	InOutImportOptions.bCombineToSingle = ImportUI->bCombineMeshes;
+	InOutImportOptions.bCombineToSingle = ImportUI->StaticMeshImportData->bCombineMeshes;
 	InOutImportOptions.VertexColorImportOption = ImportUI->StaticMeshImportData->VertexColorImportOption;
 	InOutImportOptions.VertexOverrideColor = ImportUI->StaticMeshImportData->VertexOverrideColor;
 	InOutImportOptions.bRemoveDegenerates = ImportUI->StaticMeshImportData->bRemoveDegenerates;
@@ -951,6 +961,50 @@ bool FFbxImporter::ImportFile(FString Filename, bool bPreventMaterialNameClash /
 	return Result;
 }
 
+void FFbxImporter::ConvertScene()
+{
+	if (GetImportOptions()->bConvertScene)
+	{
+		// we use -Y as forward axis here when we import. This is odd considering our forward axis is technically +X
+		// but this is to mimic Maya/Max behavior where if you make a model facing +X facing, 
+		// when you import that mesh, you want +X facing in engine. 
+		// only thing that doesn't work is hand flipping because Max/Maya is RHS but UE is LHS
+		// On the positive note, we now have import transform set up you can do to rotate mesh if you don't like default setting
+		FbxAxisSystem::ECoordSystem CoordSystem = FbxAxisSystem::eRightHanded;
+		FbxAxisSystem::EUpVector UpVector = FbxAxisSystem::eZAxis;
+		FbxAxisSystem::EFrontVector FrontVector = (FbxAxisSystem::EFrontVector) - FbxAxisSystem::eParityOdd;
+		if (GetImportOptions()->bForceFrontXAxis)
+		{
+			FrontVector = FbxAxisSystem::eParityEven;
+		}
+
+
+		FbxAxisSystem UnrealImportAxis(UpVector, FrontVector, CoordSystem);
+
+		FbxAxisSystem SourceSetup = Scene->GetGlobalSettings().GetAxisSystem();
+
+		if (SourceSetup != UnrealImportAxis)
+		{
+			FbxRootNodeUtility::RemoveAllFbxRoots(Scene);
+			UnrealImportAxis.ConvertScene(Scene);
+			FbxAMatrix JointOrientationMatrix;
+			JointOrientationMatrix.SetIdentity();
+			if (GetImportOptions()->bForceFrontXAxis)
+			{
+				JointOrientationMatrix.SetR(FbxVector4(-90.0, -90.0, 0.0));
+			}
+			FFbxDataConverter::SetJointPostConversionMatrix(JointOrientationMatrix);
+		}
+	}
+	// Convert the scene's units to what is used in this program, if needed.
+	// The base unit used in both FBX and Unreal is centimeters.  So unless the units 
+	// are already in centimeters (ie: scalefactor 1.0) then it needs to be converted
+	if (GetImportOptions()->bConvertSceneUnit && Scene->GetGlobalSettings().GetSystemUnit() != FbxSystemUnit::cm)
+	{
+		FbxSystemUnit::cm.ConvertScene(Scene);
+	}
+}
+
 //-------------------------------------------------------------------------
 //
 //-------------------------------------------------------------------------
@@ -981,48 +1035,8 @@ bool FFbxImporter::ImportFromFile(const FString& Filename, const FString& Type,
 			// The imported axis system is unknown for obj files
 			if( !Type.Equals( Obj, ESearchCase::IgnoreCase ) )
 			{
-				const FBXImportOptions* ImportOption = GetImportOptions();
-				if (ImportOption->bConvertScene)
-				{
-					// we use -Y as forward axis here when we import. This is odd considering our forward axis is technically +X
-					// but this is to mimic Maya/Max behavior where if you make a model facing +X facing, 
-					// when you import that mesh, you want +X facing in engine. 
-					// only thing that doesn't work is hand flipping because Max/Maya is RHS but UE is LHS
-					// On the positive note, we now have import transform set up you can do to rotate mesh if you don't like default setting
-					FbxAxisSystem::ECoordSystem CoordSystem = FbxAxisSystem::eRightHanded;
-					FbxAxisSystem::EUpVector UpVector = FbxAxisSystem::eZAxis;
-					FbxAxisSystem::EFrontVector FrontVector = (FbxAxisSystem::EFrontVector) - FbxAxisSystem::eParityOdd;
-					if (ImportOption->bForceFrontXAxis)
-					{
-						FrontVector = FbxAxisSystem::eParityEven;
-					}
-					
-					
-					FbxAxisSystem UnrealImportAxis(UpVector, FrontVector, CoordSystem);
-					
-					FbxAxisSystem SourceSetup = Scene->GetGlobalSettings().GetAxisSystem();
-
-					if (SourceSetup != UnrealImportAxis)
-					{
-						FbxRootNodeUtility::RemoveAllFbxRoots(Scene);
-						UnrealImportAxis.ConvertScene(Scene);
-						FbxAMatrix JointOrientationMatrix;
-						JointOrientationMatrix.SetIdentity();
-						if (ImportOption->bForceFrontXAxis)
-						{
-							JointOrientationMatrix.SetR(FbxVector4(-90.0, -90.0, 0.0));
-						}
-						FFbxDataConverter::SetJointPostConversionMatrix(JointOrientationMatrix);
-					}
-				}
-
-				// Convert the scene's units to what is used in this program, if needed.
-				// The base unit used in both FBX and Unreal is centimeters.  So unless the units 
-				// are already in centimeters (ie: scalefactor 1.0) then it needs to be converted
-				if (ImportOption->bConvertSceneUnit && Scene->GetGlobalSettings().GetSystemUnit() != FbxSystemUnit::cm)
-				{
-					FbxSystemUnit::cm.ConvertScene(Scene);
-				}
+				//Convert the scene
+				ConvertScene();
 
 				// do analytics on getting Fbx data
 				FbxDocumentInfo* DocInfo = Scene->GetSceneInfo();
@@ -1056,6 +1070,8 @@ bool FFbxImporter::ImportFromFile(const FString& Filename, const FString& Type,
 
 			//Warn the user if there is some geometry that cannot be imported because they are not reference by any scene node attribute
 			ValidateAllMeshesAreReferenceByNodeAttribute();
+
+			MeshNamesCache.Empty();
 		}
 		
 	default:
@@ -1158,16 +1174,28 @@ FName FFbxImporter::MakeNameForMesh(FString InName, FbxObject* FbxObject)
 			NewName = Name;
 		}
 
-		if ( InName == FString("None"))
+		int32 NameCount = 0;
+		FString ComposeName;
+		do 
 		{
-			OutputName = FName( *FString::Printf(TEXT("%s"), UTF8_TO_TCHAR(NewName )) );
-		}
-		else
-		{
-			OutputName = FName( *FString::Printf(TEXT("%s_%s"), *InName,UTF8_TO_TCHAR(NewName)) );
-		}
+			if ( InName == FString("None"))
+			{
+				ComposeName = FString::Printf(TEXT("%s"), UTF8_TO_TCHAR(NewName ));
+			}
+			else
+			{
+				ComposeName = FString::Printf(TEXT("%s_%s"), *InName,UTF8_TO_TCHAR(NewName));
+			}
+			if (NameCount > 0)
+			{
+				ComposeName += TEXT("_") + FString::FromInt(NameCount);
+			}
+			NameCount++;
+		} while (MeshNamesCache.Contains(ComposeName));
+		OutputName = FName(*ComposeName);
 	}
 	
+	MeshNamesCache.Add(OutputName.ToString());
 	return OutputName;
 }
 
@@ -1323,6 +1351,34 @@ void FFbxImporter::FillFbxMeshArray(FbxNode* Node, TArray<FbxNode*>& outMeshArra
 	}
 }
 
+void FFbxImporter::FillFbxMeshAndLODGroupArray(FbxNode* Node, TArray<FbxNode*>& outLODGroupArray, TArray<FbxNode*>& outMeshArray)
+{
+	// Is this node an LOD group
+	bool bLODGroup = Node->GetNodeAttribute() && Node->GetNodeAttribute()->GetAttributeType() == FbxNodeAttribute::eLODGroup;
+
+	if (bLODGroup)
+	{
+		outLODGroupArray.Add(Node);
+		//Do not do LOD group childrens
+		return;
+	}
+	
+	if (Node->GetMesh())
+	{
+		if (!FillCollisionModelList(Node) && Node->GetMesh()->GetPolygonVertexCount() > 0)
+		{
+			outMeshArray.Add(Node);
+		}
+	}
+	
+	// Cycle the childrens
+	int32 ChildIndex;
+	for (ChildIndex = 0; ChildIndex < Node->GetChildCount(); ++ChildIndex)
+	{
+		FillFbxMeshAndLODGroupArray(Node->GetChild(ChildIndex), outLODGroupArray, outMeshArray);
+	}
+}
+
 /**
 * Get all Fbx skeletal mesh objects
 *
diff --git a/Engine/Source/Editor/UnrealEd/Private/Fbx/FbxMaterialImport.cpp b/Engine/Source/Editor/UnrealEd/Private/Fbx/FbxMaterialImport.cpp
--- a/Engine/Source/Editor/UnrealEd/Private/Fbx/FbxMaterialImport.cpp
+++ b/Engine/Source/Editor/UnrealEd/Private/Fbx/FbxMaterialImport.cpp
@@ -24,9 +24,12 @@
 #include "AssetRegistryModule.h"
 #include "IAssetTools.h"
 #include "AssetToolsModule.h"
+#include "Misc/FbxErrors.h"
 
 DEFINE_LOG_CATEGORY_STATIC(LogFbxMaterialImport, Log, All);
 
+#define LOCTEXT_NAMESPACE "FbxMaterialImport"
+
 using namespace UnFbx;
 
 UTexture* UnFbx::FFbxImporter::ImportTexture( FbxFileTexture* FbxTexture, bool bSetupAsNormalMap )
@@ -536,6 +539,38 @@ void UnFbx::FFbxImporter::CreateUnrealMaterial(FbxSurfaceMaterial& FbxMaterial,
 	else
 	{
 		UMaterialInterface* FoundMaterial = LoadObject<UMaterialInterface>(NULL, *ObjectPath.ToString(), NULL, LOAD_Quiet | LOAD_NoWarn);
+		FBXImportOptions* FbxImportOptions = GetImportOptions();
+		if (FoundMaterial == nullptr && FbxImportOptions->MaterialSearchLocation != EMaterialSearchLocation::Local)
+		{
+			// Search recursively in asset's folder
+			FString SearchPath = FPaths::GetPath(BasePackageName);
+			FoundMaterial = FindExistingUnrealMaterial(SearchPath, MaterialFullName);
+
+			if (FoundMaterial == nullptr && 
+				(FbxImportOptions->MaterialSearchLocation == EMaterialSearchLocation::UnderParent ||
+				FbxImportOptions->MaterialSearchLocation == EMaterialSearchLocation::UnderRoot ||
+				FbxImportOptions->MaterialSearchLocation == EMaterialSearchLocation::AllAssets))
+			{
+				// Search recursively in parent's folder
+				SearchPath = FPaths::GetPath(SearchPath);
+				FoundMaterial = FindExistingUnrealMaterial(SearchPath, MaterialFullName);
+			}
+			if (FoundMaterial == nullptr && 
+				(FbxImportOptions->MaterialSearchLocation == EMaterialSearchLocation::UnderRoot || 
+				FbxImportOptions->MaterialSearchLocation == EMaterialSearchLocation::AllAssets))
+			{
+				// Search recursively in root folder of asset
+				FString OutPackageRoot, OutPackagePath, OutPackageName;
+				FPackageName::SplitLongPackageName(SearchPath, OutPackageRoot, OutPackagePath, OutPackageName);
+				FoundMaterial = FindExistingUnrealMaterial(OutPackageRoot, MaterialFullName);
+			}
+			if (FoundMaterial == nullptr && 
+				FbxImportOptions->MaterialSearchLocation == EMaterialSearchLocation::AllAssets)
+			{
+				// Search everywhere
+				FoundMaterial = FindExistingUnrealMaterial(TEXT("/"), MaterialFullName);
+			}
+		}
 		// do not override existing materials
 		if (FoundMaterial)
 		{
@@ -737,14 +772,99 @@ void UnFbx::FFbxImporter::CreateUnrealMaterial(FbxSurfaceMaterial& FbxMaterial,
 int32 UnFbx::FFbxImporter::CreateNodeMaterials(FbxNode* FbxNode, TArray<UMaterialInterface*>& OutMaterials, TArray<FString>& UVSets, bool bForSkeletalMesh)
 {
 	int32 MaterialCount = FbxNode->GetMaterialCount();
+	TArray<FbxSurfaceMaterial*> UsedSurfaceMaterials;
+	FbxMesh *MeshNode = FbxNode->GetMesh();
+	TSet<int32> UsedMaterialIndexes;
+	if (MeshNode)
+	{
+		for (int32 ElementMaterialIndex = 0; ElementMaterialIndex < MeshNode->GetElementMaterialCount(); ++ElementMaterialIndex)
+		{
+			FbxGeometryElementMaterial *ElementMaterial = MeshNode->GetElementMaterial(ElementMaterialIndex);
+			switch (ElementMaterial->GetMappingMode())
+			{
+			case FbxLayerElement::eAllSame:
+				{
+					if (ElementMaterial->GetIndexArray().GetCount() > 0)
+					{
+						UsedMaterialIndexes.Add(ElementMaterial->GetIndexArray()[0]);
+					}
+				}
+				break;
+			case FbxLayerElement::eByPolygon:
+				{
+					for (int32 MaterialIndex = 0; MaterialIndex < ElementMaterial->GetIndexArray().GetCount(); ++MaterialIndex)
+					{
+						UsedMaterialIndexes.Add(ElementMaterial->GetIndexArray()[MaterialIndex]);
+					}
+				}
+				break;
+			}
+		}
+	}
 	for(int32 MaterialIndex=0; MaterialIndex < MaterialCount; ++MaterialIndex)
 	{
-		FbxSurfaceMaterial *FbxMaterial = FbxNode->GetMaterial(MaterialIndex);
+		//Create only the material used by the mesh element material
+		if (MeshNode == nullptr || UsedMaterialIndexes.Contains(MaterialIndex))
+		{
+			FbxSurfaceMaterial *FbxMaterial = FbxNode->GetMaterial(MaterialIndex);
 
-		if( FbxMaterial )
+			if (FbxMaterial)
+			{
+				CreateUnrealMaterial(*FbxMaterial, OutMaterials, UVSets, bForSkeletalMesh);
+			}
+		}
+		else
 		{
-			CreateUnrealMaterial(*FbxMaterial, OutMaterials, UVSets, bForSkeletalMesh);
+			OutMaterials.Add(nullptr);
 		}
 	}
 	return MaterialCount;
 }
+
+UMaterialInterface* UnFbx::FFbxImporter::FindExistingUnrealMaterial(const FString& BasePath, const FString& MaterialName)
+{
+	UMaterialInterface* Material = nullptr;
+
+	FAssetRegistryModule& AssetRegistryModule = FModuleManager::LoadModuleChecked<FAssetRegistryModule>("AssetRegistry");
+	IAssetRegistry& AssetRegistry = AssetRegistryModule.Get();
+	TArray<FAssetData> AssetData;
+	FARFilter Filter;
+
+	AssetRegistry.SearchAllAssets(true);
+
+	Filter.bRecursiveClasses = true;
+	Filter.bRecursivePaths = true;
+	Filter.ClassNames.Add(UMaterialInterface::StaticClass()->GetFName());
+	Filter.PackagePaths.Add(FName(*BasePath));
+
+	AssetRegistry.GetAssets(Filter, AssetData);
+
+	TArray<UMaterialInterface*> FoundAssets;
+	for (const FAssetData& Data : AssetData)
+	{
+		if (Data.AssetName == FName(*MaterialName))
+		{
+			Material = Cast<UMaterialInterface>(Data.GetAsset());
+			if (Material != nullptr)
+			{
+				FoundAssets.Add(Material);
+			}
+		}
+	}
+
+	if (FoundAssets.Num() > 1)
+	{
+		check(Material != nullptr);
+		AddTokenizedErrorMessage(
+			FTokenizedMessage::Create(EMessageSeverity::Warning,
+				FText::Format(LOCTEXT("FbxMaterialImport_MultipleMaterialsFound", "While importing '{0}', found {1} materials matching name '{2}'. Using '{3}'."), 
+					FText::FromString(Parent->GetOutermost()->GetName()), 
+					FText::FromString(FString::FromInt(FoundAssets.Num())), 
+					FText::FromString(MaterialName), 
+					FText::FromString(Material->GetOutermost()->GetName()))),
+			FFbxErrors::Generic_LoadingSceneFailed);
+	}
+	return Material;
+}
+
+#undef LOCTEXT_NAMESPACE
