Copying //UE4/Dev-Editor to //UE4/Dev-Main (Source: //UE4/Dev-Editor @ 3228984)

diff --git a/Engine/Source/Editor/UnrealEd/Private/Fbx/FbxSceneImportData.cpp b/Engine/Source/Editor/UnrealEd/Private/Fbx/FbxSceneImportData.cpp
--- a/Engine/Source/Editor/UnrealEd/Private/Fbx/FbxSceneImportData.cpp
+++ b/Engine/Source/Editor/UnrealEd/Private/Fbx/FbxSceneImportData.cpp
@@ -15,6 +15,7 @@ UFbxSceneImportData::UFbxSceneImportData(const FObjectInitializer& ObjectInitial
 	BluePrintFullName.Empty();
 	SceneInfoSourceData = nullptr;
 	bCreateFolderHierarchy = false;
+	bForceFrontXAxis = false;
 	bImportScene = true;
 	HierarchyType = FBXSOCHT_CreateLevelActors;
 #endif
@@ -267,11 +268,12 @@ FString UFbxSceneImportData::ToJson() const
 	Json.Reserve(1024);
 	//data array start
 	Json += TEXT("[ { ");
-	Json += FString::Printf(TEXT("\"bImportScene\" : \"%d\", \"bCreateFolderHierarchy\" : \"%d\", \"HierarchyType\" : \"%d\", \"BluePrintFullName\" : \"%s\" "),
+	Json += FString::Printf(TEXT("\"bImportScene\" : \"%d\", \"bCreateFolderHierarchy\" : \"%d\", \"HierarchyType\" : \"%d\", \"BluePrintFullName\" : \"%s\", \"bForceFrontXAxis\" : \"%d\" "),
 		bImportScene ? 1 : 0,
 		bCreateFolderHierarchy ? 1 : 0,
 		HierarchyType,
-		*BluePrintFullName
+		*BluePrintFullName,
+		bForceFrontXAxis ? 1 : 0
 		);
 		
 	if (NameOptionsMap.Num() > 0)
@@ -368,164 +370,173 @@ void UFbxSceneImportData::FromJson(FString InJsonString)
 		SceneInfoObj->TryGetBoolField("bCreateFolderHierarchy", bCreateFolderHierarchy);
 		SceneInfoObj->TryGetNumberField("HierarchyType", HierarchyType);
 		SceneInfoObj->TryGetStringField("BluePrintFullName", BluePrintFullName);
+		SceneInfoObj->TryGetBoolField("bForceFrontXAxis", bForceFrontXAxis);
 
 		//Read Options
 		const TArray<TSharedPtr<FJsonValue>>* JSONOptions;
-		SceneInfoObj->TryGetArrayField("FbxOptions", JSONOptions);
-		for (const auto& OptionJsonValue : *JSONOptions)
+		if (SceneInfoObj->TryGetArrayField("FbxOptions", JSONOptions))
 		{
-			FString OptionName;
-			UnFbx::FBXImportOptions *Option = JSONToFbxOption(OptionJsonValue, OptionName);
-			if (Option != nullptr)
+			for (const auto& OptionJsonValue : *JSONOptions)
 			{
-				NameOptionsMap.Add(OptionName, Option);
+				FString OptionName;
+				UnFbx::FBXImportOptions *Option = JSONToFbxOption(OptionJsonValue, OptionName);
+				if (Option != nullptr)
+				{
+					NameOptionsMap.Add(OptionName, Option);
+				}
 			}
 		}
 
 		//Read Meshes
 		const TArray<TSharedPtr<FJsonValue>>* JSONMeshes;
-		SceneInfoObj->TryGetArrayField("MeshInfo", JSONMeshes);
-		for (const auto& MeshJsonValue : *JSONMeshes)
+		if (SceneInfoObj->TryGetArrayField("MeshInfo", JSONMeshes))
 		{
-			const TSharedPtr<FJsonObject>& MeshInfoObj = MeshJsonValue->AsObject();
-			if (!MeshInfoObj.IsValid())
+			for (const auto& MeshJsonValue : *JSONMeshes)
 			{
-				continue;
-			}
-			TSharedPtr<FFbxMeshInfo> MeshInfo = MakeShareable(new FFbxMeshInfo());
+				const TSharedPtr<FJsonObject>& MeshInfoObj = MeshJsonValue->AsObject();
+				if (!MeshInfoObj.IsValid())
+				{
+					continue;
+				}
+				TSharedPtr<FFbxMeshInfo> MeshInfo = MakeShareable(new FFbxMeshInfo());
 
-			if (!MeshInfoObj->TryGetStringField("Name", MeshInfo->Name)) continue;
+				if (!MeshInfoObj->TryGetStringField("Name", MeshInfo->Name)) continue;
 
-			FString UniqueIDStr;
-			if (!MeshInfoObj->TryGetStringField("UniqueId", UniqueIDStr)) continue;
-			MeshInfo->UniqueId = FCString::Atoi64(*UniqueIDStr);
+				FString UniqueIDStr;
+				if (!MeshInfoObj->TryGetStringField("UniqueId", UniqueIDStr)) continue;
+				MeshInfo->UniqueId = FCString::Atoi64(*UniqueIDStr);
 
-			MeshInfoObj->TryGetBoolField("bImportAttribute", MeshInfo->bImportAttribute);
+				MeshInfoObj->TryGetBoolField("bImportAttribute", MeshInfo->bImportAttribute);
 
-			MeshInfoObj->TryGetStringField("OptionName", MeshInfo->OptionName);
+				MeshInfoObj->TryGetStringField("OptionName", MeshInfo->OptionName);
 
-			MeshInfoObj->TryGetBoolField("bIsSkelMesh", MeshInfo->bIsSkelMesh);
+				MeshInfoObj->TryGetBoolField("bIsSkelMesh", MeshInfo->bIsSkelMesh);
 
-			MeshInfoObj->TryGetStringField("OriginalImportPath", MeshInfo->OriginalImportPath);
+				MeshInfoObj->TryGetStringField("OriginalImportPath", MeshInfo->OriginalImportPath);
 
-			MeshInfoObj->TryGetStringField("OriginalFullImportName", MeshInfo->OriginalFullImportName);
+				MeshInfoObj->TryGetStringField("OriginalFullImportName", MeshInfo->OriginalFullImportName);
 
-			MeshInfoObj->TryGetBoolField("bOverridePath", MeshInfo->bOverridePath);
+				MeshInfoObj->TryGetBoolField("bOverridePath", MeshInfo->bOverridePath);
 
-			MeshInfoObj->TryGetStringField("OverrideImportPath", MeshInfo->OverrideImportPath);
+				MeshInfoObj->TryGetStringField("OverrideImportPath", MeshInfo->OverrideImportPath);
 
-			MeshInfoObj->TryGetStringField("OverrideFullImportName", MeshInfo->OverrideFullImportName);
+				MeshInfoObj->TryGetStringField("OverrideFullImportName", MeshInfo->OverrideFullImportName);
 
-			if (MeshInfoObj->TryGetStringField("PivotNodeUid", UniqueIDStr))
-			{
-				MeshInfo->PivotNodeUid = FCString::Atoi64(*UniqueIDStr);
-			}
+				if (MeshInfoObj->TryGetStringField("PivotNodeUid", UniqueIDStr))
+				{
+					MeshInfo->PivotNodeUid = FCString::Atoi64(*UniqueIDStr);
+				}
+
+				MeshInfoObj->TryGetStringField("LODGroup", MeshInfo->LODGroup);
 
-			MeshInfoObj->TryGetStringField("LODGroup", MeshInfo->LODGroup);
+				MeshInfoObj->TryGetNumberField("LODLevel", MeshInfo->LODLevel);
 
-			MeshInfoObj->TryGetNumberField("LODLevel", MeshInfo->LODLevel);
+				MeshInfoObj->TryGetBoolField("IsLod", MeshInfo->IsLod);
 
-			MeshInfoObj->TryGetBoolField("IsLod", MeshInfo->IsLod);
-				
-			SceneInfoSourceData->MeshInfo.Add(MeshInfo);
+				SceneInfoSourceData->MeshInfo.Add(MeshInfo);
+			}
 		}
 
 		//Read Hierarchy
 		const TArray<TSharedPtr<FJsonValue>>* JSONHierarchyNodes;
-		SceneInfoObj->TryGetArrayField("Hierarchy", JSONHierarchyNodes);
-		for (const auto& NodeJsonValue : *JSONHierarchyNodes)
+		if (SceneInfoObj->TryGetArrayField("Hierarchy", JSONHierarchyNodes))
 		{
-			const TSharedPtr<FJsonObject>& NodeInfoObj = NodeJsonValue->AsObject();
-			if (!NodeInfoObj.IsValid())
+			for (const auto& NodeJsonValue : *JSONHierarchyNodes)
 			{
-				continue;
-			}
-			TSharedPtr<FFbxNodeInfo> NodeInfo = MakeShareable(new FFbxNodeInfo());
+				const TSharedPtr<FJsonObject>& NodeInfoObj = NodeJsonValue->AsObject();
+				if (!NodeInfoObj.IsValid())
+				{
+					continue;
+				}
+				TSharedPtr<FFbxNodeInfo> NodeInfo = MakeShareable(new FFbxNodeInfo());
 
-			if (!NodeInfoObj->TryGetStringField("NodeName", NodeInfo->NodeName)) continue;
+				if (!NodeInfoObj->TryGetStringField("NodeName", NodeInfo->NodeName)) continue;
 
-			NodeInfoObj->TryGetStringField("NodeHierarchyPath", NodeInfo->NodeHierarchyPath);
+				NodeInfoObj->TryGetStringField("NodeHierarchyPath", NodeInfo->NodeHierarchyPath);
 
-			FString UniqueIDStr;
-			if (!NodeInfoObj->TryGetStringField("UniqueId", UniqueIDStr)) continue;
-			NodeInfo->UniqueId = FCString::Atoi64(*UniqueIDStr);
+				FString UniqueIDStr;
+				if (!NodeInfoObj->TryGetStringField("UniqueId", UniqueIDStr)) continue;
+				NodeInfo->UniqueId = FCString::Atoi64(*UniqueIDStr);
 
-			NodeInfoObj->TryGetBoolField("bImportNode", NodeInfo->bImportNode);
+				NodeInfoObj->TryGetBoolField("bImportNode", NodeInfo->bImportNode);
 
-			//Find the parent
-			if (NodeInfoObj->TryGetStringField("ParentUniqueId", UniqueIDStr))
-			{
-				uint64 ParentUniqueId = FCString::Atoi64(*UniqueIDStr);
-				if (ParentUniqueId != 0)
+				//Find the parent
+				if (NodeInfoObj->TryGetStringField("ParentUniqueId", UniqueIDStr))
 				{
-					//Find the parent
-					for (TSharedPtr<FFbxNodeInfo> SearchNodeInfo : SceneInfoSourceData->HierarchyInfo)
+					uint64 ParentUniqueId = FCString::Atoi64(*UniqueIDStr);
+					if (ParentUniqueId != 0)
 					{
-						if (SearchNodeInfo->UniqueId == ParentUniqueId)
+						//Find the parent
+						for (TSharedPtr<FFbxNodeInfo> SearchNodeInfo : SceneInfoSourceData->HierarchyInfo)
 						{
-							NodeInfo->ParentNodeInfo = SearchNodeInfo;
-							SearchNodeInfo->Childrens.Add(NodeInfo);
-							break;
+							if (SearchNodeInfo->UniqueId == ParentUniqueId)
+							{
+								NodeInfo->ParentNodeInfo = SearchNodeInfo;
+								SearchNodeInfo->Childrens.Add(NodeInfo);
+								break;
+							}
 						}
 					}
 				}
-			}
 
-			NodeInfoObj->TryGetStringField("AttributeType", NodeInfo->AttributeType);
+				NodeInfoObj->TryGetStringField("AttributeType", NodeInfo->AttributeType);
 
-			if (NodeInfoObj->TryGetStringField("AttributeUniqueId", UniqueIDStr))
-			{
-				NodeInfo->AttributeUniqueId = FCString::Atoi64(*UniqueIDStr);
-				if (NodeInfo->AttributeUniqueId != 0)
+				if (NodeInfoObj->TryGetStringField("AttributeUniqueId", UniqueIDStr))
 				{
-					//Find the attribute info
-					for (TSharedPtr<FFbxMeshInfo> MeshInfo : SceneInfoSourceData->MeshInfo)
+					NodeInfo->AttributeUniqueId = FCString::Atoi64(*UniqueIDStr);
+					if (NodeInfo->AttributeUniqueId != 0)
 					{
-						if (MeshInfo->UniqueId == NodeInfo->AttributeUniqueId)
+						//Find the attribute info
+						for (TSharedPtr<FFbxMeshInfo> MeshInfo : SceneInfoSourceData->MeshInfo)
 						{
-							NodeInfo->AttributeInfo = MeshInfo;
-							break;
+							if (MeshInfo->UniqueId == NodeInfo->AttributeUniqueId)
+							{
+								NodeInfo->AttributeInfo = MeshInfo;
+								break;
+							}
 						}
 					}
 				}
-			}
 
-			//Materials
-			const TArray<TSharedPtr<FJsonValue>>* JSONMaterials;
-			NodeInfoObj->TryGetArrayField("Materials", JSONMaterials);
-			for (const auto& MaterialJsonValue : *JSONMaterials)
-			{
-				const TSharedPtr<FJsonObject>& MaterialObj = MaterialJsonValue->AsObject();
-				if (!MaterialObj.IsValid())
+				//Materials
+				const TArray<TSharedPtr<FJsonValue>>* JSONMaterials;
+				if (NodeInfoObj->TryGetArrayField("Materials", JSONMaterials))
 				{
-					continue;
-				}
+					for (const auto& MaterialJsonValue : *JSONMaterials)
+					{
+						const TSharedPtr<FJsonObject>& MaterialObj = MaterialJsonValue->AsObject();
+						if (!MaterialObj.IsValid())
+						{
+							continue;
+						}
 
-				TSharedPtr<FFbxMaterialInfo> MaterialInfo = MakeShareable(new FFbxMaterialInfo());
+						TSharedPtr<FFbxMaterialInfo> MaterialInfo = MakeShareable(new FFbxMaterialInfo());
 
-				if (!MaterialObj->TryGetStringField("Name", MaterialInfo->Name)) continue;
+						if (!MaterialObj->TryGetStringField("Name", MaterialInfo->Name)) continue;
 
-				MaterialObj->TryGetStringField("HierarchyPath", MaterialInfo->HierarchyPath);
+						MaterialObj->TryGetStringField("HierarchyPath", MaterialInfo->HierarchyPath);
 
-				if (!MaterialObj->TryGetStringField("UniqueId", UniqueIDStr)) continue;
-				MaterialInfo->UniqueId = FCString::Atoi64(*UniqueIDStr);
+						if (!MaterialObj->TryGetStringField("UniqueId", UniqueIDStr)) continue;
+						MaterialInfo->UniqueId = FCString::Atoi64(*UniqueIDStr);
 
-				MaterialObj->TryGetBoolField("bImportAttribute", MaterialInfo->bImportAttribute);
+						MaterialObj->TryGetBoolField("bImportAttribute", MaterialInfo->bImportAttribute);
 
-				MaterialObj->TryGetStringField("OriginalImportPath", MaterialInfo->OriginalImportPath);
+						MaterialObj->TryGetStringField("OriginalImportPath", MaterialInfo->OriginalImportPath);
 
-				MaterialObj->TryGetStringField("OriginalFullImportName", MaterialInfo->OriginalFullImportName);
+						MaterialObj->TryGetStringField("OriginalFullImportName", MaterialInfo->OriginalFullImportName);
 
-				MaterialObj->TryGetBoolField("bOverridePath", MaterialInfo->bOverridePath);
+						MaterialObj->TryGetBoolField("bOverridePath", MaterialInfo->bOverridePath);
 
-				MaterialObj->TryGetStringField("OverrideImportPath", MaterialInfo->OverrideImportPath);
+						MaterialObj->TryGetStringField("OverrideImportPath", MaterialInfo->OverrideImportPath);
 
-				MaterialObj->TryGetStringField("OverrideFullImportName", MaterialInfo->OverrideFullImportName);
+						MaterialObj->TryGetStringField("OverrideFullImportName", MaterialInfo->OverrideFullImportName);
 
-				NodeInfo->Materials.Add(MaterialInfo);
-			}
+						NodeInfo->Materials.Add(MaterialInfo);
+					}
+				}
 
-			SceneInfoSourceData->HierarchyInfo.Add(NodeInfo);
+				SceneInfoSourceData->HierarchyInfo.Add(NodeInfo);
+			}
 		}
 	}
 }
diff --git a/Engine/Source/Editor/UnrealEd/Private/Fbx/FbxSceneImportFactory.cpp b/Engine/Source/Editor/UnrealEd/Private/Fbx/FbxSceneImportFactory.cpp
--- a/Engine/Source/Editor/UnrealEd/Private/Fbx/FbxSceneImportFactory.cpp
+++ b/Engine/Source/Editor/UnrealEd/Private/Fbx/FbxSceneImportFactory.cpp
@@ -29,15 +29,15 @@
 #include "Factories/FbxSceneImportOptionsStaticMesh.h"
 #include "EngineGlobals.h"
 #include "Camera/CameraComponent.h"
-#include "Components/PointLightComponent.h"
 #include "Components/SpotLightComponent.h"
+#include "Components/PointLightComponent.h"
 #include "Engine/StaticMesh.h"
 #include "Engine/StaticMeshActor.h"
 #include "Components/DirectionalLightComponent.h"
 #include "AssetData.h"
 #include "Editor.h"
 #include "FileHelpers.h"
-
+#include "CineCameraComponent.h"
 
 #include "AssetSelection.h"
 
@@ -118,7 +118,6 @@ bool GetFbxSceneImportOptions(UnFbx::FFbxImporter* FbxImporter
 	GlobalImportSettings->ImportUniformScale = 1.0f;
 
 	GlobalImportSettings->bConvertScene = true;
-	GlobalImportSettings->bForceFrontXAxis = false;
 	GlobalImportSettings->bConvertSceneUnit = true;
 
 	GlobalImportSettings->bBakePivotInVertex = SceneImportOptions->bBakePivotInVertex;
@@ -160,6 +159,7 @@ bool GetFbxSceneImportOptions(UnFbx::FFbxImporter* FbxImporter
 	}
 
 	//setup all options
+	GlobalImportSettings->bForceFrontXAxis = SceneImportOptions->bForceFrontXAxis;
 	GlobalImportSettings->bBakePivotInVertex = SceneImportOptions->bBakePivotInVertex;
 	GlobalImportSettings->bImportStaticMeshLODs = SceneImportOptions->bImportStaticMeshLODs;
 	GlobalImportSettings->bImportSkeletalMeshLODs = SceneImportOptions->bImportSkeletalMeshLODs;
@@ -373,6 +373,7 @@ void FetchFbxCameraInScene(UnFbx::FFbxImporter *FbxImporter, FbxNode* ParentNode
 			CameraInfo->ProjectionPerspective = CameraAttribute->ProjectionType.Get() == FbxCamera::ePerspective;
 			CameraInfo->OrthoZoom = CameraAttribute->OrthoZoom.Get();
 			CameraInfo->FieldOfView = CameraAttribute->FieldOfView.Get();
+			CameraInfo->FocalLength = CameraAttribute->FocalLength.Get();						
 			SceneInfoPtr->CameraInfo.Add(CameraInfo->UniqueId, CameraInfo);
 		}
 	}
@@ -808,6 +809,7 @@ UFbxSceneImportData* CreateReImportAsset(const FString &PackagePath, const FStri
 	
 	ReImportAsset->SourceFbxFile = FPaths::ConvertRelativePathToFull(FbxImportFileName);
 	ReImportAsset->bCreateFolderHierarchy = SceneImportOptions->bCreateContentFolderHierarchy;
+	ReImportAsset->bForceFrontXAxis = SceneImportOptions->bForceFrontXAxis;
 	ReImportAsset->HierarchyType = SceneImportOptions->HierarchyType.GetValue();
 	return ReImportAsset;
 }
@@ -833,6 +835,68 @@ UObject* UFbxSceneImportFactory::FactoryCreateBinary
 	return ReturnObject;
 }
 
+TSharedPtr<FFbxNodeInfo> GetNodeInfoPtrById(TArray<TSharedPtr<FFbxNodeInfo>> &HierarchyInfo, uint64 SearchId)
+{
+	for (TSharedPtr<FFbxNodeInfo> NodeInfoPtr : HierarchyInfo)
+	{
+		if (NodeInfoPtr->UniqueId == SearchId)
+		{
+			return NodeInfoPtr;
+		}
+	}
+	return nullptr;
+}
+
+void UFbxSceneImportFactory::ChangeFrontAxis(void* VoidFbxImporter, void* VoidSceneInfo, TSharedPtr<FFbxSceneInfo> SceneInfoPtr)
+{
+	UnFbx::FFbxImporter* FbxImporter = (UnFbx::FFbxImporter*)VoidFbxImporter;
+	UnFbx::FbxSceneInfo* SceneInfo = (UnFbx::FbxSceneInfo*)VoidSceneInfo;
+
+	FbxImporter->ConvertScene();
+	//Adjust the root node with the new apply scene conversion
+	FbxNode* RootNode = FbxImporter->Scene->GetRootNode();
+	if (SceneInfo->HierarchyInfo.Num() > 0)
+	{
+		//Set the fbx data
+		UnFbx::FbxNodeInfo &RootNodeInfo = SceneInfo->HierarchyInfo[0];
+		check(RootNodeInfo.UniqueId == RootNode->GetUniqueID());
+		RootNodeInfo.Transform = RootNode->EvaluateGlobalTransform();
+		//Set the UE4 data
+		TSharedPtr<FFbxNodeInfo> RootNodeInfoPtr = GetNodeInfoPtrById(SceneInfoPtr->HierarchyInfo, RootNodeInfo.UniqueId);
+		if (RootNodeInfoPtr.IsValid())
+		{
+			RootNodeInfoPtr->Transform = FTransform::Identity;
+			FbxVector4 NewLocalT = RootNodeInfo.Transform.GetT();
+			FbxVector4 NewLocalS = RootNodeInfo.Transform.GetS();
+			FbxQuaternion NewLocalQ = RootNodeInfo.Transform.GetQ();
+			RootNodeInfoPtr->Transform.SetTranslation(UnFbx::FFbxDataConverter::ConvertPos(NewLocalT));
+			RootNodeInfoPtr->Transform.SetScale3D(UnFbx::FFbxDataConverter::ConvertScale(NewLocalS));
+			RootNodeInfoPtr->Transform.SetRotation(UnFbx::FFbxDataConverter::ConvertRotToQuat(NewLocalQ));
+			for (int32 NodeIndex = 1; NodeIndex < SceneInfo->HierarchyInfo.Num(); ++NodeIndex)
+			{
+				UnFbx::FbxNodeInfo &LocalNodeInfo = SceneInfo->HierarchyInfo[NodeIndex];
+				FbxNode *RealFbxNode = FindFbxNodeById(FbxImporter, nullptr, LocalNodeInfo.UniqueId);
+				if (!RealFbxNode)
+				{
+					continue;
+				}
+				LocalNodeInfo.Transform = RealFbxNode->EvaluateLocalTransform();
+				TSharedPtr<FFbxNodeInfo> LocalNodeInfoPtr = GetNodeInfoPtrById(SceneInfoPtr->HierarchyInfo, LocalNodeInfo.UniqueId);
+				if (LocalNodeInfoPtr.IsValid())
+				{
+					LocalNodeInfoPtr->Transform = FTransform::Identity;
+					NewLocalT = LocalNodeInfo.Transform.GetT();
+					NewLocalS = LocalNodeInfo.Transform.GetS();
+					NewLocalQ = LocalNodeInfo.Transform.GetQ();
+					LocalNodeInfoPtr->Transform.SetTranslation(UnFbx::FFbxDataConverter::ConvertPos(NewLocalT));
+					LocalNodeInfoPtr->Transform.SetScale3D(UnFbx::FFbxDataConverter::ConvertScale(NewLocalS));
+					LocalNodeInfoPtr->Transform.SetRotation(UnFbx::FFbxDataConverter::ConvertRotToQuat(NewLocalQ));
+				}
+			}
+		}
+	}
+}
+
 UObject* UFbxSceneImportFactory::FactoryCreateBinary
 (
 UClass*				Class,
@@ -878,12 +942,11 @@ FFeedbackContext*	Warn
 	
 	//Always convert the scene
 	GlobalImportSettings->bConvertScene = true;
-	GlobalImportSettings->bForceFrontXAxis = false;
 	GlobalImportSettings->bConvertSceneUnit = true;
 
 	//Set the import option in importscene mode
 	GlobalImportSettings->bImportScene = true;
-
+	bool OriginalForceFrontXAxis = GlobalImportSettings->bForceFrontXAxis;
 	//Read the fbx and store the hierarchy's information so we can reuse it after importing all the model in the fbx file
 	if (!FbxImporter->ImportFromFile(*FbxImportFileName, Type, true))
 	{
@@ -945,6 +1008,14 @@ FFeedbackContext*	Warn
 
 	SFbxSceneOptionWindow::CopyFbxOptionsToFbxOptions(GlobalImportSettingsReference, GlobalImportSettings);
 
+	//Convert the scene to the correct axis system. Option like force front X or ConvertScene affect the scene conversion
+	//We need to get the new convert transform
+	if(OriginalForceFrontXAxis != GlobalImportSettings->bForceFrontXAxis)
+	{
+		ChangeFrontAxis(FbxImporter, &SceneInfo, SceneInfoPtr);
+	}
+
+
 	FillSceneHierarchyPath(SceneInfoPtr);
 
 	ReimportData = CreateReImportAsset(Path, FbxImportFileName, SceneImportOptions, SceneInfoPtr, NameOptionsMap);
@@ -1115,12 +1186,19 @@ bool UFbxSceneImportFactory::SetStaticMeshComponentOverrideMaterial(UStaticMeshC
 
 USceneComponent *CreateCameraComponent(AActor *ParentActor, TSharedPtr<FFbxCameraInfo> CameraInfo)
 {
-	UCameraComponent *CameraComponent = NewObject<UCameraComponent>(ParentActor, *(CameraInfo->Name));
+	UCineCameraComponent *CameraComponent = NewObject<UCineCameraComponent>(ParentActor, *(CameraInfo->Name));
 	CameraComponent->SetProjectionMode(CameraInfo->ProjectionPerspective ? ECameraProjectionMode::Perspective : ECameraProjectionMode::Orthographic);
 	CameraComponent->SetAspectRatio(CameraInfo->AspectWidth / CameraInfo->AspectHeight);
 	CameraComponent->SetOrthoNearClipPlane(CameraInfo->NearPlane);
 	CameraComponent->SetOrthoFarClipPlane(CameraInfo->FarPlane);
 	CameraComponent->SetOrthoWidth(CameraInfo->AspectWidth);
+	CameraComponent->SetFieldOfView(CameraInfo->FieldOfView);
+	CameraComponent->FilmbackSettings.SensorWidth = 2 * CameraInfo->FocalLength * FMath::Tan(FMath::DegreesToRadians(0.5f * CameraInfo->FieldOfView));
+	CameraComponent->FilmbackSettings.SensorHeight = CameraComponent->FilmbackSettings.SensorWidth * CameraInfo->AspectHeight / CameraInfo->AspectWidth;
+	CameraComponent->LensSettings.MaxFocalLength = CameraInfo->FocalLength;
+	CameraComponent->LensSettings.MinFocalLength = CameraInfo->FocalLength;
+	CameraComponent->FocusSettings.FocusMethod = ECameraFocusMethod::None;
+
 	return CameraComponent;
 }
 
@@ -1910,7 +1988,7 @@ void UFbxSceneImportFactory::ImportAllStaticMesh(void* VoidRootNodeToImport, voi
 	ImportedMeshCount = AllNewAssets.Num();
 	if (ImportedMeshCount == 1 && NewStaticMesh)
 	{
-		FbxImporter->ImportStaticMeshSockets(NewStaticMesh);
+		FbxImporter->ImportStaticMeshGlobalSockets(NewStaticMesh);
 	}
 }
 
diff --git a/Engine/Source/Editor/UnrealEd/Private/Fbx/FbxSceneImportOptions.cpp b/Engine/Source/Editor/UnrealEd/Private/Fbx/FbxSceneImportOptions.cpp
--- a/Engine/Source/Editor/UnrealEd/Private/Fbx/FbxSceneImportOptions.cpp
+++ b/Engine/Source/Editor/UnrealEd/Private/Fbx/FbxSceneImportOptions.cpp
@@ -11,6 +11,7 @@ UFbxSceneImportOptions::UFbxSceneImportOptions(const FObjectInitializer& ObjectI
 	bCreateContentFolderHierarchy = false;
 	bImportAsDynamic = false;
 	HierarchyType = FBXSOCHT_CreateBlueprint;
+	bForceFrontXAxis = false;
 	bImportStaticMeshLODs = false;
 	bImportSkeletalMeshLODs = false;
 	bInvertNormalMaps = false;
diff --git a/Engine/Source/Editor/UnrealEd/Private/Fbx/FbxSkeletalMeshExport.cpp b/Engine/Source/Editor/UnrealEd/Private/Fbx/FbxSkeletalMeshExport.cpp
--- a/Engine/Source/Editor/UnrealEd/Private/Fbx/FbxSkeletalMeshExport.cpp
+++ b/Engine/Source/Editor/UnrealEd/Private/Fbx/FbxSkeletalMeshExport.cpp
@@ -65,9 +65,11 @@ FbxNode* FFbxExporter::CreateSkeleton(const USkeletalMesh* SkelMesh, TArray<FbxN
 		FVector UnrealRotation = BoneTransform.GetRotation().Euler();
 		FbxVector4 LocalPos = Converter.ConvertToFbxPos(BoneTransform.GetTranslation());
 		FbxVector4 LocalRot = Converter.ConvertToFbxRot(UnrealRotation);
+		FbxVector4 LocalScale = Converter.ConvertToFbxScale(BoneTransform.GetScale3D());
 
 		BoneNode->LclTranslation.Set(LocalPos);
 		BoneNode->LclRotation.Set(LocalRot);
+		BoneNode->LclScaling.Set(LocalScale);
 
 
 		// If this is not the root bone, attach it to its parent
diff --git a/Engine/Source/Editor/UnrealEd/Private/Fbx/FbxSkeletalMeshImport.cpp b/Engine/Source/Editor/UnrealEd/Private/Fbx/FbxSkeletalMeshImport.cpp
--- a/Engine/Source/Editor/UnrealEd/Private/Fbx/FbxSkeletalMeshImport.cpp
+++ b/Engine/Source/Editor/UnrealEd/Private/Fbx/FbxSkeletalMeshImport.cpp
@@ -2179,6 +2179,11 @@ void UnFbx::FFbxImporter::SetMaterialOrderByName(FSkeletalMeshImportData& Import
 
 void UnFbx::FFbxImporter::CleanUpUnusedMaterials(FSkeletalMeshImportData& ImportData)
 {
+	if (ImportData.Materials.Num() <= 0)
+	{
+		return;
+	}
+
 	TArray< VMaterial > ExistingMatList = ImportData.Materials;
 
 	TArray<uint8> UsedMaterialIndex;
@@ -2190,7 +2195,7 @@ void UnFbx::FFbxImporter::CleanUpUnusedMaterials(FSkeletalMeshImportData& Import
 		UsedMaterialIndex.AddUnique(Triangle.MatIndex);
 	}
 	//Remove any unused material.
-	if (UsedMaterialIndex.Num() != ExistingMatList.Num())
+	if (UsedMaterialIndex.Num() < ExistingMatList.Num())
 	{
 		TArray<int32> RemapIndex;
 		TArray< VMaterial > &NewMatList = ImportData.Materials;
@@ -3545,16 +3550,21 @@ FFbxLogger::~FFbxLogger()
 			}
 		}
 	}
-	if(ShowLogMessage && TokenizedErrorMessages.Num() > 0)
-	{
-		const TCHAR* LogTitle = TEXT("FBXImport");
-		FMessageLogModule& MessageLogModule = FModuleManager::LoadModuleChecked<FMessageLogModule>("MessageLog");
-		TSharedPtr<class IMessageLogListing> LogListing = MessageLogModule.GetLogListing(LogTitle);
-		LogListing->SetLabel(FText::FromString("FBX Import"));
-		LogListing->ClearMessages();
 
+	//Always clear the old message after an import or re-import
+	const TCHAR* LogTitle = TEXT("FBXImport");
+	FMessageLogModule& MessageLogModule = FModuleManager::LoadModuleChecked<FMessageLogModule>("MessageLog");
+	TSharedPtr<class IMessageLogListing> LogListing = MessageLogModule.GetLogListing(LogTitle);
+	LogListing->SetLabel(FText::FromString("FBX Import"));
+	LogListing->ClearMessages();
+
+	if(TokenizedErrorMessages.Num() > 0)
+	{
 		LogListing->AddMessages(TokenizedErrorMessages);
-		MessageLogModule.OpenMessageLog(LogTitle);
+		if (ShowLogMessage)
+		{
+			MessageLogModule.OpenMessageLog(LogTitle);
+		}
 	}
 }
 #undef LOCTEXT_NAMESPACE
diff --git a/Engine/Source/Editor/UnrealEd/Private/Fbx/FbxStaticMeshImport.cpp b/Engine/Source/Editor/UnrealEd/Private/Fbx/FbxStaticMeshImport.cpp
--- a/Engine/Source/Editor/UnrealEd/Private/Fbx/FbxStaticMeshImport.cpp
+++ b/Engine/Source/Editor/UnrealEd/Private/Fbx/FbxStaticMeshImport.cpp
@@ -41,7 +41,7 @@ static const int32 LARGE_MESH_MATERIAL_INDEX_THRESHOLD = 64;
 using namespace UnFbx;
 
 struct ExistingStaticMeshData;
-extern ExistingStaticMeshData* SaveExistingStaticMeshData(UStaticMesh* ExistingMesh, bool bSaveMaterials);
+extern ExistingStaticMeshData* SaveExistingStaticMeshData(UStaticMesh* ExistingMesh, FBXImportOptions* ImportOptions);
 extern void RestoreExistingMeshSettings(struct ExistingStaticMeshData* ExistingMesh, UStaticMesh* NewMesh, int32 LODIndex);
 extern void RestoreExistingMeshData(struct ExistingStaticMeshData* ExistingMeshDataPtr, UStaticMesh* NewMesh);
 
@@ -813,7 +813,7 @@ UStaticMesh* UnFbx::FFbxImporter::ReimportSceneStaticMesh(uint64 FbxNodeUniqueId
 		return Mesh;
 	}
 
-	struct ExistingStaticMeshData* ExistMeshDataPtr = SaveExistingStaticMeshData(Mesh, !ImportOptions->bImportMaterials);
+	struct ExistingStaticMeshData* ExistMeshDataPtr = SaveExistingStaticMeshData(Mesh, ImportOptions);
 
 	if (Node)
 	{
@@ -938,7 +938,7 @@ UStaticMesh* UnFbx::FFbxImporter::ReimportStaticMesh(UStaticMesh* Mesh, UFbxStat
 	ImportOptions->bImportMaterials = false;
 	ImportOptions->bImportTextures = false;
 
-	struct ExistingStaticMeshData* ExistMeshDataPtr = SaveExistingStaticMeshData(Mesh, !ImportOptions->bImportMaterials);
+	struct ExistingStaticMeshData* ExistMeshDataPtr = SaveExistingStaticMeshData(Mesh, ImportOptions);
 
 	if (Node)
 	{
@@ -1576,6 +1576,8 @@ UStaticMesh* UnFbx::FFbxImporter::ImportStaticMeshAsSingle(UObject* InParent, TA
 		
 		//warnings based on geometry
 		VerifyGeometry(StaticMesh);
+
+		ImportStaticMeshLocalSockets(StaticMesh, MeshNodeArray);
 	}
 
 	return StaticMesh;
@@ -1784,7 +1786,73 @@ static void FindMeshSockets( FbxNode* StartNode, TArray<FbxSocketNode>& OutFbxSo
 	}
 }
 
-void UnFbx::FFbxImporter::ImportStaticMeshSockets( UStaticMesh* StaticMesh )
+void UnFbx::FFbxImporter::ImportStaticMeshLocalSockets(UStaticMesh* StaticMesh, TArray<FbxNode*>& MeshNodeArray)
+{
+	check(MeshNodeArray.Num());
+	FbxNode *MeshRootNode = MeshNodeArray[0];
+	const FbxAMatrix &MeshTotalMatrix = ComputeTotalMatrix(MeshRootNode);
+	for (FbxNode* RootNode : MeshNodeArray)
+	{
+		// Find all nodes that are sockets
+		TArray<FbxSocketNode> SocketNodes;
+		FindMeshSockets(RootNode, SocketNodes);
+
+		// Create a UStaticMeshSocket for each fbx socket
+		for (int32 SocketIndex = 0; SocketIndex < SocketNodes.Num(); ++SocketIndex)
+		{
+			FbxSocketNode& SocketNode = SocketNodes[SocketIndex];
+
+			UStaticMeshSocket* Socket = StaticMesh->FindSocket(SocketNode.SocketName);
+			if (!Socket)
+			{
+				// If the socket didn't exist create a new one now
+				Socket = NewObject<UStaticMeshSocket>(StaticMesh);
+				check(Socket);
+
+				Socket->SocketName = SocketNode.SocketName;
+				StaticMesh->Sockets.Add(Socket);
+			}
+
+			if (Socket)
+			{
+				const FbxAMatrix& SocketMatrix = Scene->GetAnimationEvaluator()->GetNodeLocalTransform(SocketNode.Node);
+				FbxAMatrix FinalSocketMatrix = MeshTotalMatrix * SocketMatrix;
+				FTransform SocketTransform;
+				SocketTransform.SetTranslation(Converter.ConvertPos(FinalSocketMatrix.GetT()));
+				SocketTransform.SetRotation(Converter.ConvertRotToQuat(FinalSocketMatrix.GetQ()));
+				SocketTransform.SetScale3D(Converter.ConvertScale(FinalSocketMatrix.GetS()));
+
+				Socket->RelativeLocation = SocketTransform.GetLocation();
+				Socket->RelativeRotation = SocketTransform.GetRotation().Rotator();
+				Socket->RelativeScale = SocketTransform.GetScale3D();
+			}
+		}
+		// Delete mesh sockets that were removed from the import data
+		if (StaticMesh->Sockets.Num() != SocketNodes.Num())
+		{
+			for (int32 MeshSocketIx = 0; MeshSocketIx < StaticMesh->Sockets.Num(); ++MeshSocketIx)
+			{
+				bool Found = false;
+				UStaticMeshSocket* MeshSocket = StaticMesh->Sockets[MeshSocketIx];
+				for (int32 FbxSocketIx = 0; FbxSocketIx < SocketNodes.Num(); FbxSocketIx++)
+				{
+					if (SocketNodes[FbxSocketIx].SocketName == MeshSocket->SocketName)
+					{
+						Found = true;
+						break;
+					}
+				}
+				if (!Found)
+				{
+					StaticMesh->Sockets.RemoveAt(MeshSocketIx);
+					MeshSocketIx--;
+				}
+			}
+		}
+	}
+}
+
+void UnFbx::FFbxImporter::ImportStaticMeshGlobalSockets( UStaticMesh* StaticMesh )
 {
 	FbxNode* RootNode = Scene->GetRootNode();
 
@@ -1806,11 +1874,8 @@ void UnFbx::FFbxImporter::ImportStaticMeshSockets( UStaticMesh* StaticMesh )
 
 			Socket->SocketName = SocketNode.SocketName;
 			StaticMesh->Sockets.Add(Socket);
-		}
 
-		if( Socket )
-		{
-			FbxAMatrix& SocketMatrix = Scene->GetAnimationEvaluator()->GetNodeLocalTransform(SocketNode.Node);
+			const FbxAMatrix& SocketMatrix = Scene->GetAnimationEvaluator()->GetNodeGlobalTransform(SocketNode.Node);
 			FTransform SocketTransform;
 			SocketTransform.SetTranslation(Converter.ConvertPos(SocketMatrix.GetT()));
 			SocketTransform.SetRotation(Converter.ConvertRotToQuat(SocketMatrix.GetQ()));
@@ -1821,6 +1886,29 @@ void UnFbx::FFbxImporter::ImportStaticMeshSockets( UStaticMesh* StaticMesh )
 			Socket->RelativeScale = SocketTransform.GetScale3D();
 		}
 	}
+
+	// Delete mesh sockets that were removed from the import data
+	if (StaticMesh->Sockets.Num() != SocketNodes.Num())
+	{
+		for (int32 MeshSocketIx = 0; MeshSocketIx < StaticMesh->Sockets.Num(); ++MeshSocketIx)
+		{
+			bool Found = false;
+			UStaticMeshSocket* MeshSocket = StaticMesh->Sockets[MeshSocketIx];
+			for (int32 FbxSocketIx = 0; FbxSocketIx < SocketNodes.Num(); FbxSocketIx++)
+			{
+				if (SocketNodes[FbxSocketIx].SocketName == MeshSocket->SocketName)
+				{
+					Found = true;
+					break;
+				}
+			}
+			if (!Found)
+			{
+				StaticMesh->Sockets.RemoveAt(MeshSocketIx);
+				MeshSocketIx--;
+			}
+		}
+	}
 }
 
 
diff --git a/Engine/Source/Editor/UnrealEd/Private/Fbx/ReimportFbxSceneFactory.cpp b/Engine/Source/Editor/UnrealEd/Private/Fbx/ReimportFbxSceneFactory.cpp
--- a/Engine/Source/Editor/UnrealEd/Private/Fbx/ReimportFbxSceneFactory.cpp
+++ b/Engine/Source/Editor/UnrealEd/Private/Fbx/ReimportFbxSceneFactory.cpp
@@ -239,7 +239,6 @@ bool GetFbxSceneReImportOptions(UnFbx::FFbxImporter* FbxImporter
 	GlobalImportSettings->ImportUniformScale = 1.0f;
 
 	GlobalImportSettings->bConvertScene = true;
-	GlobalImportSettings->bForceFrontXAxis = false;
 	GlobalImportSettings->bConvertSceneUnit = true;
 
 
@@ -285,6 +284,7 @@ bool GetFbxSceneReImportOptions(UnFbx::FFbxImporter* FbxImporter
 	//Set the bakepivot option in the SceneImportOptions
 	SceneImportOptions->bBakePivotInVertex = GlobalImportSettings->bBakePivotInVertex;
 	//setup all options
+	GlobalImportSettings->bForceFrontXAxis = SceneImportOptions->bForceFrontXAxis;
 	GlobalImportSettings->bImportStaticMeshLODs = SceneImportOptions->bImportStaticMeshLODs;
 	GlobalImportSettings->bImportSkeletalMeshLODs = SceneImportOptions->bImportSkeletalMeshLODs;
 	SceneImportOptions->bInvertNormalMaps = GlobalImportSettings->bInvertNormalMap;
@@ -381,7 +381,6 @@ EReimportResult::Type UReimportFbxSceneFactory::Reimport(UObject* Obj)
 
 	//Always convert the scene
 	GlobalImportSettings->bConvertScene = true;
-	GlobalImportSettings->bForceFrontXAxis = false;
 	GlobalImportSettings->bConvertSceneUnit = true;
 	GlobalImportSettings->bImportScene = ReimportData->bImportScene;
 	if (ReimportData->NameOptionsMap.Contains(DefaultOptionName))
@@ -390,7 +389,7 @@ EReimportResult::Type UReimportFbxSceneFactory::Reimport(UObject* Obj)
 		GlobalImportSettings->bBakePivotInVertex = DefaultOption->bBakePivotInVertex;
 		GlobalImportSettings->bInvertNormalMap = DefaultOption->bInvertNormalMap;
 	}
-
+	bool OriginalForceFrontXAxis = GlobalImportSettings->bForceFrontXAxis;
 	//Read the fbx and store the hierarchy's information so we can reuse it after importing all the model in the fbx file
 	if (!FbxImporter->ImportFromFile(*FbxImportFileName, FPaths::GetExtension(FbxImportFileName), true))
 	{
@@ -456,6 +455,7 @@ EReimportResult::Type UReimportFbxSceneFactory::Reimport(UObject* Obj)
 	FbxSceneReimportStatusMap NodeStatusMap;
 	bool bCanReimportHierarchy = ReimportData->HierarchyType == (int32)EFBXSceneOptionsCreateHierarchyType::FBXSOCHT_CreateBlueprint && !ReimportData->BluePrintFullName.IsEmpty();
 
+	SceneImportOptions->bForceFrontXAxis = GlobalImportSettings->bForceFrontXAxis;
 	if (!GetFbxSceneReImportOptions(FbxImporter
 		, SceneInfoPtr
 		, ReimportData->SceneInfoSourceData
@@ -480,6 +480,13 @@ EReimportResult::Type UReimportFbxSceneFactory::Reimport(UObject* Obj)
 	GlobalImportSettingsReference = new UnFbx::FBXImportOptions();
 	SFbxSceneOptionWindow::CopyFbxOptionsToFbxOptions(GlobalImportSettings, GlobalImportSettingsReference);
 
+	//Convert the scene to the correct axis system. Option like force front X
+	//We need to get the new convert transform
+	if (OriginalForceFrontXAxis != GlobalImportSettings->bForceFrontXAxis)
+	{
+		ChangeFrontAxis(FbxImporter, &SceneInfo, SceneInfoPtr);
+	}
+
 	//Overwrite the reimport asset data with the new data
 	ReimportData->SceneInfoSourceData = SceneInfoPtr;
 	ReimportData->SourceFbxFile = FPaths::ConvertRelativePathToFull(FbxImportFileName);
@@ -1158,6 +1165,12 @@ EReimportResult::Type UReimportFbxSceneFactory::ImportStaticMesh(void* VoidFbxIm
 		}
 		return EReimportResult::Failed;
 	}
+	else
+	{
+		//Mark any re-imported package dirty
+		NewObject->MarkPackageDirty();
+	}
+
 	AllNewAssets.Add(MeshInfo, NewObject);
 	AssetToSyncContentBrowser.Add(NewObject);
 	return EReimportResult::Succeeded;
diff --git a/Engine/Source/Editor/UnrealEd/Private/Fbx/SFbxSceneOptionWindow.cpp b/Engine/Source/Editor/UnrealEd/Private/Fbx/SFbxSceneOptionWindow.cpp
--- a/Engine/Source/Editor/UnrealEd/Private/Fbx/SFbxSceneOptionWindow.cpp
+++ b/Engine/Source/Editor/UnrealEd/Private/Fbx/SFbxSceneOptionWindow.cpp
@@ -190,6 +190,7 @@ void SFbxSceneOptionWindow::OnFinishedChangingPropertiesSceneTabDetailView(const
 		MaterialsTabListView->UpdateMaterialBasePath();
 	}
 	//Set the Global Import setting
+	GlobalImportSettings->bForceFrontXAxis = SceneImportOptionsDisplay->bForceFrontXAxis;
 	GlobalImportSettings->bBakePivotInVertex = SceneImportOptionsDisplay->bBakePivotInVertex;
 	GlobalImportSettings->bInvertNormalMap = SceneImportOptionsDisplay->bInvertNormalMaps;
 	GlobalImportSettings->ImportTranslation = SceneImportOptionsDisplay->ImportTranslation;
diff --git a/Engine/Source/Editor/UnrealEd/Private/FbxImporter.h b/Engine/Source/Editor/UnrealEd/Private/FbxImporter.h
--- a/Engine/Source/Editor/UnrealEd/Private/FbxImporter.h
+++ b/Engine/Source/Editor/UnrealEd/Private/FbxImporter.h
@@ -9,6 +9,7 @@
 #include "Factories/FbxImportUI.h"
 #include "Logging/TokenizedMessage.h"
 #include "Factories/FbxStaticMeshImportData.h"
+#include "Factories/FbxTextureImportData.h"
 #include "Factories/FbxSceneImportFactory.h"
 
 class AActor;
@@ -144,6 +145,7 @@ struct FBXImportOptions
 	FString BaseNormalTextureName;
 	FString BaseEmmisiveTextureName;
 	FString BaseSpecularTextureName;
+	EMaterialSearchLocation::Type MaterialSearchLocation;
 	// Skeletal Mesh options
 	bool bImportMorph;
 	bool bImportAnimations;
@@ -398,6 +400,19 @@ struct FbxSceneInfo
 	double FrameRate;
 	double TotalTime;
 
+	void Reset()
+	{
+		NonSkinnedMeshNum = 0;
+		SkinnedMeshNum = 0;
+		TotalGeometryNum = 0;
+		TotalMaterialNum = 0;
+		TotalTextureNum = 0;
+		MeshInfo.Empty();
+		HierarchyInfo.Empty();
+		bHasAnimation = false;
+		FrameRate = 0.0;
+		TotalTime = 0.0;
+	}
 };
 
 /**
@@ -505,6 +520,12 @@ public:
 	 */
 	bool ImportFile(FString Filename, bool bPreventMaterialNameClash = false);
 	
+	/**
+	 * Convert the scene from the current options.
+	 * The scene will be converted to RH -Y or RH X depending if we force a front X axis or not
+	 */
+	void ConvertScene();
+
 	/**
 	 * Attempt to load an FBX scene from a given filename.
 	 *
@@ -593,7 +614,8 @@ public:
 	*/
 	UNREALED_API bool ImportSubDSurface(USubDSurface* Out, UObject* InParent, TArray<FbxNode*>& MeshNodeArray, const FName InName, EObjectFlags Flags, UFbxStaticMeshImportData* TemplateImportData);
 
-	void ImportStaticMeshSockets( UStaticMesh* StaticMesh );
+	void ImportStaticMeshGlobalSockets( UStaticMesh* StaticMesh );
+	void ImportStaticMeshLocalSockets( UStaticMesh* StaticMesh, TArray<FbxNode*>& MeshNodeArray);
 
 	/**
 	 * re-import Unreal static mesh from updated Fbx file
@@ -776,6 +798,15 @@ public:
 	*/
 	UNREALED_API void FillFbxMeshArray(FbxNode* Node, TArray<FbxNode*>& outMeshArray, UnFbx::FFbxImporter* FFbxImporter);
 	
+	/**
+	* Get all Fbx mesh objects not under a LOD group and all LOD group node
+	*
+	* @param Node Root node to find meshes
+	* @param outLODGroupArray return Fbx LOD group
+	* @param outMeshArray return Fbx meshes with no LOD group
+	*/
+	UNREALED_API void FillFbxMeshAndLODGroupArray(FbxNode* Node, TArray<FbxNode*>& outLODGroupArray, TArray<FbxNode*>& outMeshArray);
+
 	/**
 	* Fill FBX skeletons to OutSortedLinks recursively
 	*
@@ -1223,6 +1254,15 @@ protected:
 	 */
 	void CreateUnrealMaterial(FbxSurfaceMaterial& FbxMaterial, TArray<UMaterialInterface*>& OutMaterials, TArray<FString>& UVSets, bool bForSkeletalMesh);
 	
+	/**
+	* Search for an existing Unreal material based on import option settings
+	*
+	* @param BasePath Folder to start looking from, recursively.
+	* @param MaterialName Name of the material to search for.
+	* @return Material found
+	*/
+	UMaterialInterface* FindExistingUnrealMaterial(const FString& BasePath, const FString& MaterialName);
+
 	/**
 	 * Visit all materials of one node, import textures from materials.
 	 *
@@ -1413,6 +1453,9 @@ private:
 
 	FImportedMaterialData ImportedMaterialData;
 
+	//Cache to create unique name for mesh. This is use to fix name clash
+	TArray<FString> MeshNamesCache;
+
 private:
 
 
diff --git a/Engine/Source/Editor/UnrealEd/Private/FileHelpers.cpp b/Engine/Source/Editor/UnrealEd/Private/FileHelpers.cpp
--- a/Engine/Source/Editor/UnrealEd/Private/FileHelpers.cpp
+++ b/Engine/Source/Editor/UnrealEd/Private/FileHelpers.cpp
@@ -565,7 +565,8 @@ static bool SaveWorld(UWorld* World,
 			SaveErrors.Flush();
 		}
 
-		if (bSuccess)
+		// @todo Autosaving should save build data as well
+		if (bSuccess && !bAutosaving)
 		{
 			// Also save MapBuildData packages when saving the current level
 			FEditorFileUtils::SaveMapDataPackages(DuplicatedWorld ? DuplicatedWorld : World, bCheckDirty);
diff --git a/Engine/Source/Editor/UnrealEd/Private/Kismet2/BlueprintEditorUtils.cpp b/Engine/Source/Editor/UnrealEd/Private/Kismet2/BlueprintEditorUtils.cpp
--- a/Engine/Source/Editor/UnrealEd/Private/Kismet2/BlueprintEditorUtils.cpp
+++ b/Engine/Source/Editor/UnrealEd/Private/Kismet2/BlueprintEditorUtils.cpp
@@ -586,14 +586,7 @@ void FBlueprintEditorUtils::ReconstructAllNodes(UBlueprint* Blueprint)
 
 void FBlueprintEditorUtils::ReplaceDeprecatedNodes(UBlueprint* Blueprint)
 {
-	TArray<UEdGraph*> Graphs;
-	Blueprint->GetAllGraphs(Graphs);
-	for (auto It = Graphs.CreateIterator(); It; ++It)
-	{
-		UEdGraph* const Graph = *It;
-		const UEdGraphSchema* Schema = Graph->GetSchema();
-		Schema->BackwardCompatibilityNodeConversion(Graph, true);
-	}
+	Blueprint->ReplaceDeprecatedNodes();
 }
 
 void FBlueprintEditorUtils::RefreshExternalBlueprintDependencyNodes(UBlueprint* Blueprint, UStruct* RefreshOnlyChild)
@@ -6703,6 +6696,44 @@ void FBlueprintEditorUtils::FindActorsThatReferenceActor( AActor* InActor, TArra
 	}
 }
 
+void FBlueprintEditorUtils::GetActorReferenceMap(UWorld* InWorld, TArray<UClass*>& InClassesToIgnore, TMap<AActor*, TArray<AActor*> >& OutReferencingActors)
+{
+	// Iterate all actors in the same world as InActor
+	for (FActorIterator ActorIt(InWorld); ActorIt; ++ActorIt)
+	{
+		AActor* CurrentActor = *ActorIt;
+		if (CurrentActor)
+		{
+			bool bShouldIgnore = false;
+
+			// Ignore Actors if they are of a type we were instructed to ignore.
+			for (int32 IgnoreIndex = 0; IgnoreIndex < InClassesToIgnore.Num() && !bShouldIgnore; IgnoreIndex++)
+			{
+				if (CurrentActor->IsA(InClassesToIgnore[IgnoreIndex]))
+				{
+					bShouldIgnore = true;
+				}
+			}
+
+			if (!bShouldIgnore)
+			{
+				// Get all references from CurrentActor and see if any Actors
+				TArray<UObject*> References;
+				FReferenceFinder Finder(References);
+				Finder.FindReferences(CurrentActor);
+
+				for (int32 RefIdx = 0; RefIdx < References.Num(); RefIdx++)
+				{
+					if (References[RefIdx] && References[RefIdx]->IsA(AActor::StaticClass()))
+					{
+						OutReferencingActors.FindOrAdd(Cast<AActor>(References[RefIdx])).Add(CurrentActor);
+					}
+				}
+			}
+		}
+	}
+}
+
 FBlueprintEditorUtils::FFixLevelScriptActorBindingsEvent& FBlueprintEditorUtils::OnFixLevelScriptActorBindings()
 {
 	return FixLevelScriptActorBindingsEvent;
diff --git a/Engine/Source/Editor/UnrealEd/Private/Kismet2/EnumEditorUtils.cpp b/Engine/Source/Editor/UnrealEd/Private/Kismet2/EnumEditorUtils.cpp
--- a/Engine/Source/Editor/UnrealEd/Private/Kismet2/EnumEditorUtils.cpp
+++ b/Engine/Source/Editor/UnrealEd/Private/Kismet2/EnumEditorUtils.cpp
@@ -13,6 +13,8 @@
 #include "K2Node_Variable.h"
 #include "NodeDependingOnEnumInterface.h"
 #include "Kismet2/BlueprintEditorUtils.h"
+#include "Internationalization/TextNamespaceUtil.h"
+#include "Internationalization/TextPackageNamespaceUtil.h"
 
 #define LOCTEXT_NAMESPACE "Enum"
 
@@ -134,8 +136,6 @@ void FEnumEditorUtils::RemoveEnumeratorFromUserDefinedEnum(UUserDefinedEnum* Enu
 
 	Names.RemoveAt(EnumeratorIndex);
 
-	Enum->RemoveMetaData(FEnumEditorUtilsHelper::DisplayName(), EnumeratorIndex);
-
 	// Clean up enum values.
 	for (int32 i = 0; i < Names.Num(); ++i)
 	{
@@ -440,43 +440,75 @@ int64 FEnumEditorUtils::ResolveEnumerator(const UEnum* Enum, FArchive& Ar, int64
 	return EnumeratorValue;
 }
 
-FString FEnumEditorUtils::GetEnumeratorDisplayName(const UUserDefinedEnum* Enum, const int32 EnumeratorIndex)
+FText FEnumEditorUtils::GetEnumeratorDisplayName(const UUserDefinedEnum* Enum, const int32 EnumeratorIndex)
 {
-	if (Enum && Enum->DisplayNames.IsValidIndex(EnumeratorIndex))
+	if (Enum)
 	{
-		return Enum->DisplayNames[EnumeratorIndex].ToString();
+		const FName EnumEntryName = *Enum->GetEnumName(EnumeratorIndex);
+
+		if (!EnumEntryName.IsNone())
+		{
+			const FText* EnumEntryDisplayName = Enum->DisplayNameMap.Find(EnumEntryName);
+			if (EnumEntryDisplayName)
+			{
+				return *EnumEntryDisplayName;
+			}
+		}
 	}
-	return FString();
+
+	return FText::GetEmpty();
 }
 
-bool FEnumEditorUtils::SetEnumeratorDisplayName(UUserDefinedEnum* Enum, int32 EnumeratorIndex, FString NewDisplayName)
+bool FEnumEditorUtils::SetEnumeratorDisplayName(UUserDefinedEnum* Enum, int32 EnumeratorIndex, FText NewDisplayName)
 {
 	if (Enum && (EnumeratorIndex >= 0) && (EnumeratorIndex < Enum->NumEnums()))
 	{
-		if (IsEnumeratorDisplayNameValid(Enum, NewDisplayName))
+		if (IsEnumeratorDisplayNameValid(Enum, EnumeratorIndex, NewDisplayName))
 		{
 			PrepareForChange(Enum);
-			Enum->SetMetaData(FEnumEditorUtilsHelper::DisplayName(), *NewDisplayName, EnumeratorIndex);
-			EnsureAllDisplayNamesExist(Enum);
+			
+			const FName EnumEntryName = *Enum->GetEnumName(EnumeratorIndex);
+			check(!EnumEntryName.IsNone());
+
+			FText DisplayNameToSet = NewDisplayName;
+
+#if USE_STABLE_LOCALIZATION_KEYS
+			// Make sure the package namespace for this display name is up-to-date
+			{
+				const FString PackageNamespace = TextNamespaceUtil::EnsurePackageNamespace(Enum);
+				if (!PackageNamespace.IsEmpty())
+				{
+					const FString DisplayNameNamespace = FTextInspector::GetNamespace(DisplayNameToSet).Get(FString());
+					const FString FullNamespace = TextNamespaceUtil::BuildFullNamespace(DisplayNameNamespace, PackageNamespace);
+					if (!DisplayNameNamespace.Equals(FullNamespace, ESearchCase::CaseSensitive))
+					{
+						// We may assign a new key when if we don't have the correct package namespace in order to avoid identity conflicts when instancing
+						DisplayNameToSet = FText::ChangeKey(FullNamespace, FGuid::NewGuid().ToString(), DisplayNameToSet);
+					}
+				}
+			}
+#endif // USE_STABLE_LOCALIZATION_KEYS
+
+			Enum->DisplayNameMap.Add(EnumEntryName, DisplayNameToSet);
+
 			BroadcastChanges(Enum, TArray<TPair<FName, int64>>(), false);
 			return true;
 		}
 	}
+
 	return false;
 }
 
-bool FEnumEditorUtils::IsEnumeratorDisplayNameValid(const UUserDefinedEnum* Enum, const FString NewDisplayName)
+bool FEnumEditorUtils::IsEnumeratorDisplayNameValid(const UUserDefinedEnum* Enum, int32 EnumeratorIndex, const FText NewDisplayName)
 {
-	if (!Enum || 
-		NewDisplayName.IsEmpty() || 
-		(NewDisplayName == FEnumEditorUtilsHelper::InvalidName()))
+	if (!Enum || NewDisplayName.IsEmptyOrWhitespace() || NewDisplayName.ToString() == FEnumEditorUtilsHelper::InvalidName())
 	{
 		return false;
 	}
 
 	for	(int32 Index = 0; Index < Enum->NumEnums(); Index++)
 	{
-		if (NewDisplayName == GetEnumeratorDisplayName(Enum, Index))
+		if (Index != EnumeratorIndex && NewDisplayName.ToString() == GetEnumeratorDisplayName(Enum, Index).ToString())
 		{
 			return false;
 		}
@@ -489,33 +521,103 @@ void FEnumEditorUtils::EnsureAllDisplayNamesExist(UUserDefinedEnum* Enum)
 {
 	if (Enum)
 	{
-		const int32 EnumeratorsToEnsure = (Enum->NumEnums() > 0) ? (Enum->NumEnums() - 1) : 0;
-		if (Enum->DisplayNames.Num() != EnumeratorsToEnsure)
-		{
-			Enum->DisplayNames.Empty(EnumeratorsToEnsure);
-		}
-		for	(int32 Index = 0; Index < EnumeratorsToEnsure; Index++)
+		const int32 EnumeratorsToEnsure = FMath::Max(Enum->NumEnums() - 1, 0);
+
+		// Remove any stale display names
 		{
-			FText DisplayNameMetaData = Enum->UEnum::GetEnumText(Index);
-			if (DisplayNameMetaData.IsEmpty())
+			TSet<FName> KnownEnumEntryNames;
+			KnownEnumEntryNames.Reserve(EnumeratorsToEnsure);
+
+			for (int32 Index = 0; Index < EnumeratorsToEnsure; ++Index)
+			{
+				const FName EnumEntryName = *Enum->GetEnumName(Index);
+				check(!EnumEntryName.IsNone());
+
+				KnownEnumEntryNames.Add(EnumEntryName);
+			}
+
+			for (auto DisplayNameIt = Enum->DisplayNameMap.CreateIterator(); DisplayNameIt; ++DisplayNameIt)
 			{
-				const FString EnumName = Enum->GetEnumName(Index);
-				DisplayNameMetaData = FText::FromString(EnumName);
-				for(int32 AddIndex = 0; !IsEnumeratorDisplayNameValid(Enum, DisplayNameMetaData.ToString()); ++AddIndex)
+				if (!KnownEnumEntryNames.Contains(DisplayNameIt->Key))
 				{
-					DisplayNameMetaData = FText::FromString(FString::Printf(TEXT("%s%d"), *EnumName, AddIndex));
+					DisplayNameIt.RemoveCurrent();
 				}
-				Enum->SetMetaData(FEnumEditorUtilsHelper::DisplayName(), *DisplayNameMetaData.ToString(), Index);
 			}
-			if (!Enum->DisplayNames.IsValidIndex(Index))
+		}
+
+		Enum->DisplayNameMap.Reserve(EnumeratorsToEnsure);
+
+		// Add any missing display names
+		for	(int32 Index = 0; Index < EnumeratorsToEnsure; ++Index)
+		{
+			const FName EnumEntryName = *Enum->GetEnumName(Index);
+			check(!EnumEntryName.IsNone());
+
+			if (!Enum->DisplayNameMap.Contains(EnumEntryName))
 			{
-				Enum->DisplayNames.Add(DisplayNameMetaData);
+				FText DisplayNameToSet;
+
+#if USE_STABLE_LOCALIZATION_KEYS
+				// Give the new text instance the full package-based namespace, and give it a brand new key; these will become its new stable identity
+				{
+					const FString PackageNamespace = GIsEditor ? TextNamespaceUtil::EnsurePackageNamespace(Enum) : TextNamespaceUtil::GetPackageNamespace(Enum);
+					const FString TextNamespace = TextNamespaceUtil::BuildFullNamespace(FString(), PackageNamespace, true);
+					const FString TextKey = FGuid::NewGuid().ToString();
+					const FString EnumEntryDisplayName = EnumEntryName.ToString();
+					DisplayNameToSet = FInternationalization::ForUseOnlyByLocMacroAndGraphNodeTextLiterals_CreateText(*EnumEntryDisplayName, *TextNamespace, *TextKey);
+				}
+#else  // USE_STABLE_LOCALIZATION_KEYS
+				DisplayNameToSet = FText::FromName(EnumEntryName);
+#endif // USE_STABLE_LOCALIZATION_KEYS
+
+				Enum->DisplayNameMap.Add(EnumEntryName, DisplayNameToSet);
 			}
-			else if (!Enum->DisplayNames[Index].EqualTo(DisplayNameMetaData))
+		}
+	}
+}
+
+void FEnumEditorUtils::UpgradeDisplayNamesFromMetaData(UUserDefinedEnum* Enum)
+{
+	if (Enum)
+	{
+		const int32 EnumeratorsToEnsure = FMath::Max(Enum->NumEnums() - 1, 0);
+		Enum->DisplayNameMap.Empty(EnumeratorsToEnsure);
+
+		bool bDidUpgradeDisplayNames = false;
+		for (int32 Index = 0; Index < EnumeratorsToEnsure; ++Index)
+		{
+			const FString MetaDataEntryDisplayName = Enum->GetMetaData(FEnumEditorUtilsHelper::DisplayName(), Index);
+			if (!MetaDataEntryDisplayName.IsEmpty())
 			{
-				Enum->DisplayNames[Index] = DisplayNameMetaData;
+				bDidUpgradeDisplayNames = true;
+
+				const FName EnumEntryName = *Enum->GetEnumName(Index);
+				check(!EnumEntryName.IsNone());
+
+				FText DisplayNameToSet;
+
+#if USE_STABLE_LOCALIZATION_KEYS
+				// Give the new text instance the full package-based namespace, and give it a brand new key; these will become its new stable identity
+				{
+					const FString PackageNamespace = GIsEditor ? TextNamespaceUtil::EnsurePackageNamespace(Enum) : TextNamespaceUtil::GetPackageNamespace(Enum);
+					const FString TextNamespace = TextNamespaceUtil::BuildFullNamespace(TEXT(""), PackageNamespace, true);
+					const FString TextKey = FGuid::NewGuid().ToString();
+					DisplayNameToSet = FInternationalization::ForUseOnlyByLocMacroAndGraphNodeTextLiterals_CreateText(*MetaDataEntryDisplayName, *TextNamespace, *TextKey);
+				}
+#else  // USE_STABLE_LOCALIZATION_KEYS
+				DisplayNameToSet = FText::FromName(EnumEntryName);
+#endif // USE_STABLE_LOCALIZATION_KEYS
+
+				Enum->DisplayNameMap.Add(EnumEntryName, DisplayNameToSet);
 			}
 		}
+
+#if USE_STABLE_LOCALIZATION_KEYS
+		if (bDidUpgradeDisplayNames)
+		{
+			UE_LOG(LogClass, Warning, TEXT("Enum '%s' was upgraded to use FText to store its display name data. Please re-save this asset to avoid issues with localization and determinstic cooking."), *Enum->GetPathName());
+		}
+#endif // USE_STABLE_LOCALIZATION_KEYS
 	}
 }
 
diff --git a/Engine/Source/Editor/UnrealEd/Private/LevelEditorViewport.cpp b/Engine/Source/Editor/UnrealEd/Private/LevelEditorViewport.cpp
--- a/Engine/Source/Editor/UnrealEd/Private/LevelEditorViewport.cpp
+++ b/Engine/Source/Editor/UnrealEd/Private/LevelEditorViewport.cpp
@@ -2054,14 +2054,6 @@ void FLevelEditorViewportClient::Tick(float DeltaTime)
 
 	FEditorViewportClient::Tick(DeltaTime);
 
-	if (!GUnrealEd->IsPivotMovedIndependently() && GCurrentLevelEditingViewportClient == this &&
-		bIsRealtime &&
-		(Widget == NULL || !Widget->IsDragging()))
-	{
-		// @todo SIE: May be very expensive for lots of actors
-		GUnrealEd->UpdatePivotLocationForSelection();
-	}
-
 	// Update the preview mesh for the preview mesh mode. 
 	GEditor->UpdatePreviewMesh();
 
diff --git a/Engine/Source/Editor/UnrealEd/Private/MaterialGraph.cpp b/Engine/Source/Editor/UnrealEd/Private/MaterialGraph.cpp
--- a/Engine/Source/Editor/UnrealEd/Private/MaterialGraph.cpp
+++ b/Engine/Source/Editor/UnrealEd/Private/MaterialGraph.cpp
@@ -11,7 +11,9 @@
 #include "Materials/MaterialExpressionComment.h"
 #include "Materials/MaterialExpressionFunctionOutput.h"
 #include "Materials/MaterialExpressionCustomOutput.h"
+#include "Materials/MaterialExpressionReroute.h"
 
+#include "MaterialGraphNode_Knot.h"
 
 #define LOCTEXT_NAMESPACE "MaterialGraph"
 
@@ -82,7 +84,18 @@ void UMaterialGraph::RebuildGraph()
 UMaterialGraphNode* UMaterialGraph::AddExpression(UMaterialExpression* Expression)
 {
 	UMaterialGraphNode* NewNode = NULL;
-	if (Expression)
+	if (Expression && Expression->IsA(UMaterialExpressionReroute::StaticClass()))
+	{
+		Modify();
+		FGraphNodeCreator<UMaterialGraphNode_Knot> NodeCreator(*this);
+		NewNode = NodeCreator.CreateNode(false);
+		NewNode->MaterialExpression = Expression;
+		NewNode->RealtimeDelegate = RealtimeDelegate;
+		NewNode->MaterialDirtyDelegate = MaterialDirtyDelegate;
+		Expression->GraphNode = NewNode;
+		NodeCreator.Finalize();
+	}
+	else if (Expression)
 	{
 		Modify();
 		FGraphNodeCreator<UMaterialGraphNode> NodeCreator(*this);
diff --git a/Engine/Source/Editor/UnrealEd/Private/MaterialGraphNode.cpp b/Engine/Source/Editor/UnrealEd/Private/MaterialGraphNode.cpp
--- a/Engine/Source/Editor/UnrealEd/Private/MaterialGraphNode.cpp
+++ b/Engine/Source/Editor/UnrealEd/Private/MaterialGraphNode.cpp
@@ -425,6 +425,51 @@ void UMaterialGraphNode::GetContextMenuActions(const FGraphNodeContextMenuBuilde
 	}
 }
 
+FString UMaterialGraphNode::GetShortenPinName(const FString& PinName)
+{
+	FString InputName = PinName;
+
+	// Shorten long expression input names.
+	if (!FCString::Stricmp(*PinName, TEXT("Coordinates")))
+	{
+		InputName = TEXT("UVs");
+	}
+	else if (!FCString::Stricmp(*PinName, TEXT("TextureObject")))
+	{
+		InputName = TEXT("Tex");
+	}
+	else if (!FCString::Stricmp(*PinName, TEXT("Input")))
+	{
+		InputName = TEXT("");
+	}
+	else if (!FCString::Stricmp(*PinName, TEXT("Exponent")))
+	{
+		InputName = TEXT("Exp");
+	}
+	else if (!FCString::Stricmp(*PinName, TEXT("AGreaterThanB")))
+	{
+		InputName = TEXT("A > B");
+	}
+	else if (!FCString::Stricmp(*PinName, TEXT("AEqualsB")))
+	{
+		InputName = TEXT("A == B");
+	}
+	else if (!FCString::Stricmp(*PinName, TEXT("ALessThanB")))
+	{
+		InputName = TEXT("A < B");
+	}
+	else if (!FCString::Stricmp(*PinName, TEXT("MipLevel")))
+	{
+		InputName = TEXT("Level");
+	}
+	else if (!FCString::Stricmp(*PinName, TEXT("MipBias")))
+	{
+		InputName = TEXT("Bias");
+	}
+
+	return InputName;
+}
+
 void UMaterialGraphNode::CreateInputPins()
 {
 	const TArray<FExpressionInput*> ExpressionInputs = MaterialExpression->GetInputs();
@@ -434,43 +479,7 @@ void UMaterialGraphNode::CreateInputPins()
 		FExpressionInput* Input = ExpressionInputs[Index];
 		FString InputName = MaterialExpression->GetInputName(Index);
 
-		// Shorten long expression input names.
-		if ( !FCString::Stricmp( *InputName, TEXT("Coordinates") ) )
-		{
-			InputName = TEXT("UVs");
-		}
-		else if ( !FCString::Stricmp( *InputName, TEXT("TextureObject") ) )
-		{
-			InputName = TEXT("Tex");
-		}
-		else if ( !FCString::Stricmp( *InputName, TEXT("Input") ) )
-		{
-			InputName = TEXT("");
-		}
-		else if ( !FCString::Stricmp( *InputName, TEXT("Exponent") ) )
-		{
-			InputName = TEXT("Exp");
-		}
-		else if ( !FCString::Stricmp( *InputName, TEXT("AGreaterThanB") ) )
-		{
-			InputName = TEXT("A > B");
-		}
-		else if ( !FCString::Stricmp( *InputName, TEXT("AEqualsB") ) )
-		{
-			InputName = TEXT("A == B");
-		}
-		else if ( !FCString::Stricmp( *InputName, TEXT("ALessThanB") ) )
-		{
-			InputName = TEXT("A < B");
-		}
-		else if ( !FCString::Stricmp( *InputName, TEXT("MipLevel") ) )
-		{
-			InputName = TEXT("Level");
-		}
-		else if ( !FCString::Stricmp( *InputName, TEXT("MipBias") ) )
-		{
-			InputName = TEXT("Bias");
-		}
+		InputName = GetShortenPinName(InputName);
 
 		const UMaterialGraphSchema* Schema = CastChecked<UMaterialGraphSchema>(GetSchema());
 		FString PinCategory = MaterialExpression->IsInputConnectionRequired(Index) ? Schema->PC_Required : Schema->PC_Optional;
diff --git a/Engine/Source/Editor/UnrealEd/Private/MaterialGraphNode_Root.cpp b/Engine/Source/Editor/UnrealEd/Private/MaterialGraphNode_Root.cpp
--- a/Engine/Source/Editor/UnrealEd/Private/MaterialGraphNode_Root.cpp
+++ b/Engine/Source/Editor/UnrealEd/Private/MaterialGraphNode_Root.cpp
@@ -52,8 +52,9 @@ void UMaterialGraphNode_Root::CreateInputPins()
 
 	for (int32 Index = 0; Index < MaterialGraph->MaterialInputs.Num(); ++Index)
 	{
-		UEdGraphPin* InputPin = CreatePin(EGPD_Input, Schema->PC_MaterialInput, TEXT(""), NULL, /*bIsArray=*/ false, /*bIsReference=*/ false, MaterialGraph->MaterialInputs[Index].GetName().ToString());
+		UEdGraphPin* InputPin = CreatePin(EGPD_Input, Schema->PC_MaterialInput, FString::Printf(TEXT("%d"), (int32)MaterialGraph->MaterialInputs[Index].GetProperty()), NULL, /*bIsArray=*/ false, /*bIsReference=*/ false, MaterialGraph->MaterialInputs[Index].GetName().ToString());
 	}
+
 }
 
 int32 UMaterialGraphNode_Root::GetInputIndex(const UEdGraphPin* InputPin) const
diff --git a/Engine/Source/Editor/UnrealEd/Private/MaterialGraphSchema.cpp b/Engine/Source/Editor/UnrealEd/Private/MaterialGraphSchema.cpp
--- a/Engine/Source/Editor/UnrealEd/Private/MaterialGraphSchema.cpp
+++ b/Engine/Source/Editor/UnrealEd/Private/MaterialGraphSchema.cpp
@@ -28,11 +28,14 @@
 #include "Materials/MaterialExpressionFunctionInput.h"
 #include "Materials/MaterialExpressionTextureSample.h"
 #include "Materials/MaterialExpressionFunctionOutput.h"
+#include "Materials/MaterialExpressionReroute.h"
 
 #include "ScopedTransaction.h"
 #include "MaterialEditorUtilities.h"
 #include "GraphEditorActions.h"
 #include "AssetRegistryModule.h"
+#include "MaterialEditorActions.h"
+#include "MaterialGraphNode_Knot.h"
 
 #define LOCTEXT_NAMESPACE "MaterialGraphSchema"
 
@@ -412,6 +415,12 @@ void UMaterialGraphSchema::GetContextMenuActions(const UEdGraph* CurrentGraph, c
 					((UMaterialGraphSchema*const)this)->GetBreakLinkToSubMenuActions(*MenuBuilder, const_cast<UEdGraphPin*>(InGraphPin));
 				}
 			}
+
+			// Only display Promote to Parameters on input pins
+			if (InGraphPin->Direction == EEdGraphPinDirection::EGPD_Input)
+			{
+				MenuBuilder->AddMenuEntry(FMaterialEditorCommands::Get().PromoteToParameter);
+			}
 		}
 		MenuBuilder->EndSection();
 
@@ -899,4 +908,39 @@ void UMaterialGraphSchema::ForceVisualizationCacheClear() const
 	++CurrentCacheRefreshID;
 }
 
+void UMaterialGraphSchema::OnPinConnectionDoubleCicked(UEdGraphPin* PinA, UEdGraphPin* PinB, const FVector2D& GraphPosition) const
+{
+	const FScopedTransaction Transaction(LOCTEXT("CreateRerouteNodeOnWire", "Create Reroute Node"));
+
+	//@TODO: This constant is duplicated from inside of SGraphNodeKnot
+	const FVector2D NodeSpacerSize(42.0f, 24.0f);
+	const FVector2D KnotTopLeft = GraphPosition - (NodeSpacerSize * 0.5f);
+
+	// Create a new knot
+	UEdGraph* ParentGraph = PinA->GetOwningNode()->GetGraph();
+
+	{
+		UMaterialExpression* Expression = FMaterialEditorUtilities::CreateNewMaterialExpression(ParentGraph, UMaterialExpressionReroute::StaticClass(), KnotTopLeft, true, true);
+
+		// Move the connections across (only notifying the knot, as the other two didn't really change)
+		PinA->BreakLinkTo(PinB);
+		PinA->MakeLinkTo((PinA->Direction == EGPD_Output) ? CastChecked<UMaterialGraphNode_Knot>(Expression->GraphNode)->GetInputPin() : CastChecked<UMaterialGraphNode_Knot>(Expression->GraphNode)->GetOutputPin());
+		PinB->MakeLinkTo((PinB->Direction == EGPD_Output) ? CastChecked<UMaterialGraphNode_Knot>(Expression->GraphNode)->GetInputPin() : CastChecked<UMaterialGraphNode_Knot>(Expression->GraphNode)->GetOutputPin());
+		FMaterialEditorUtilities::UpdateMaterialAfterGraphChange(ParentGraph);
+	}
+}
+
+bool UMaterialGraphSchema::SafeDeleteNodeFromGraph(UEdGraph* Graph, UEdGraphNode* NodeToDelete) const
+{
+	if (NodeToDelete == nullptr || Graph == nullptr || NodeToDelete->GetGraph() != Graph)
+	{
+		return false;
+	}
+
+	TArray<UEdGraphNode*> NodesToDelete;
+	NodesToDelete.Add(NodeToDelete);
+	FMaterialEditorUtilities::DeleteNodes(Graph, NodesToDelete);
+	return true;
+}
+
 #undef LOCTEXT_NAMESPACE
diff --git a/Engine/Source/Editor/UnrealEd/Private/ObjectTools.cpp b/Engine/Source/Editor/UnrealEd/Private/ObjectTools.cpp
--- a/Engine/Source/Editor/UnrealEd/Private/ObjectTools.cpp
+++ b/Engine/Source/Editor/UnrealEd/Private/ObjectTools.cpp
@@ -754,7 +754,7 @@ namespace ObjectTools
 
 			for ( TArray<UProperty*>::TConstIterator RefPropIter( RefPropArray ); RefPropIter; ++RefPropIter )
 			{
-				FPropertyChangedEvent PropertyEvent(*RefPropIter);
+				FPropertyChangedEvent PropertyEvent(*RefPropIter, EPropertyChangeType::Redirected);
 				CurReplaceObj->PostEditChangeProperty( PropertyEvent );
 			}
 
@@ -921,6 +921,8 @@ namespace ObjectTools
 
 			FEditorDelegates::OnAssetsPreDelete.Broadcast(ReplaceInfo.ReplaceableObjects);
 
+			TSet<FString> AlreadyMappedObjectPaths;
+
 			// With all references to the objects to consolidate to eliminated from objects that are currently loaded, it should now be safe to delete
 			// the objects to be consolidated themselves, leaving behind a redirector in their place to fix up objects that were not currently loaded at the time
 			// of this operation.
@@ -931,7 +933,8 @@ namespace ObjectTools
 				UObject* CurObjToConsolidate = *ConsolIter;
 				UObject* CurObjOuter = CurObjToConsolidate->GetOuter();
 				UPackage* CurObjPackage = CurObjToConsolidate->GetOutermost();
-				FName CurObjName = CurObjToConsolidate->GetFName();
+				const FName CurObjName = CurObjToConsolidate->GetFName();
+				const FString CurObjPath = CurObjToConsolidate->GetPathName();
 				UBlueprint* BlueprintToConsolidate = Cast<UBlueprint>(CurObjToConsolidate);
 
 				// Attempt to delete the object that was consolidated
@@ -940,6 +943,11 @@ namespace ObjectTools
 					// DONT GC YET!!! we still need these objects around to notify other tools that they are gone and to create redirectors
 					ConsolidatedObjects.Add(CurObjToConsolidate);
 
+					if ( AlreadyMappedObjectPaths.Contains(CurObjPath) )
+					{
+						continue;
+					}
+
 					// Create a redirector with a unique name
 					// It will have the same name as the object that was consolidated after the garbage collect
 					UObjectRedirector* Redirector = NewObject<UObjectRedirector>(CurObjOuter, NAME_None, RF_Standalone | RF_Public);
@@ -949,10 +957,8 @@ namespace ObjectTools
 					Redirector->DestinationObject = ObjectToConsolidateTo;
 
 					// Keep track of the object name so we can rename the redirector later
-					if (!RedirectorToObjectNameMap.FindKey(CurObjName))
-					{
-						RedirectorToObjectNameMap.Add(Redirector, CurObjName);
-					}
+					RedirectorToObjectNameMap.Add(Redirector, CurObjName);
+					AlreadyMappedObjectPaths.Add(CurObjPath);
 
 					// If consolidating blueprints, make sure redirectors are created for the consolidated blueprint class and CDO
 					if ( BlueprintToConsolidateTo != NULL && BlueprintToConsolidate != NULL )
@@ -962,12 +968,14 @@ namespace ObjectTools
 						check( ClassRedirector );
 						ClassRedirector->DestinationObject = BlueprintToConsolidateTo->GeneratedClass;
 						RedirectorToObjectNameMap.Add(ClassRedirector, BlueprintToConsolidate->GeneratedClass->GetFName());
+						AlreadyMappedObjectPaths.Add(BlueprintToConsolidate->GeneratedClass->GetPathName());
 
 						// One redirector for the CDO
 						UObjectRedirector* CDORedirector = NewObject<UObjectRedirector>(CurObjOuter, NAME_None, RF_Standalone | RF_Public);
 						check( CDORedirector );
 						CDORedirector->DestinationObject = BlueprintToConsolidateTo->GeneratedClass->GetDefaultObject();
 						RedirectorToObjectNameMap.Add(CDORedirector, BlueprintToConsolidate->GeneratedClass->GetDefaultObject()->GetFName());
+						AlreadyMappedObjectPaths.Add(BlueprintToConsolidate->GeneratedClass->GetDefaultObject()->GetPathName());
 					}
 
 					DirtiedPackages.AddUnique( CurObjPackage );
diff --git a/Engine/Source/Editor/UnrealEd/Private/PackageTools.cpp b/Engine/Source/Editor/UnrealEd/Private/PackageTools.cpp
--- a/Engine/Source/Editor/UnrealEd/Private/PackageTools.cpp
+++ b/Engine/Source/Editor/UnrealEd/Private/PackageTools.cpp
@@ -2,6 +2,7 @@
 
 
 #include "PackageTools.h"
+#include "PackageReload.h"
 #include "Misc/MessageDialog.h"
 #include "Misc/Paths.h"
 #include "Misc/Guid.h"
@@ -24,14 +25,19 @@
 #include "Dialogs/Dialogs.h"
 
 #include "ObjectTools.h"
+#include "KismetEditorUtilities.h"
 #include "BusyCursor.h"
 
 
+#include "AssetRegistryModule.h"
 #include "Logging/MessageLog.h"
 #include "UObject/UObjectIterator.h"
 #include "ComponentReregisterContext.h"
 #include "Engine/Selection.h"
 
+#include "ShaderCompiler.h"
+#include "DistanceFieldAtlas.h"
+
 #define LOCTEXT_NAMESPACE "PackageTools"
 
 DEFINE_LOG_CATEGORY_STATIC(LogPackageTools, Log, All);
@@ -310,6 +316,18 @@ namespace PackageTools
 
 				GWarn->StatusUpdate( PackageIndex, PackagesToUnload.Num(), FText::Format(NSLOCTEXT("UnrealEd", "Unloadingf", "Unloading {0}..."), FText::FromString(PackageBeingUnloaded->GetName()) ) );
 
+				// Flush all pending render commands, as unloading the package may invalidate render resources.
+				FlushRenderingCommands();
+
+				// Close any open asset editors
+				ForEachObjectWithOuter(PackageBeingUnloaded, [](UObject* Obj)
+				{
+					if (Obj->IsAsset())
+					{
+						FAssetEditorManager::Get().CloseAllEditorsForAsset(Obj);
+					}
+				}, false);
+
 				PackageBeingUnloaded->bHasBeenFullyLoaded = false;
 				PackageBeingUnloaded->ClearFlags(RF_WasLoaded);
 				if ( PackageBeingUnloaded->HasAnyPackageFlags(PKG_ContainsScript) )
@@ -396,6 +414,227 @@ namespace PackageTools
 		return bResult;
 	}
 
+
+	bool ReloadPackages( const TArray<UPackage*>& TopLevelPackages )
+	{
+		FText ErrorMessage;
+		const bool bResult = ReloadPackages(TopLevelPackages, ErrorMessage, /*bInteractive*/true);
+		
+		if (!ErrorMessage.IsEmpty())
+		{
+			FMessageDialog::Open(EAppMsgType::Ok, ErrorMessage);
+		}
+
+		return bResult;
+	}
+
+
+	bool ReloadPackages( const TArray<UPackage*>& TopLevelPackages, FText& OutErrorMessage, const bool bInteractive )
+	{
+		bool bResult = false;
+
+		FTextBuilder ErrorMessageBuilder;
+
+		// Split the set of selected top level packages into packages which are dirty or in-memory (and thus cannot be reloaded) and packages that are not dirty (and thus can be reloaded).
+		TArray<UPackage*> PackagesToReload;
+		{
+			TArray<UPackage*> DirtyPackages;
+			TArray<UPackage*> InMemoryPackages;
+			for (UPackage* TopLevelPackage : TopLevelPackages)
+			{
+				if (TopLevelPackage)
+				{
+					// Get outermost packages, in case groups were selected.
+					UPackage* RealPackage = TopLevelPackage->GetOutermost() ? TopLevelPackage->GetOutermost() : TopLevelPackage;
+
+					if (RealPackage->IsDirty())
+					{
+						DirtyPackages.AddUnique(RealPackage);
+					}
+					else if (RealPackage->HasAnyPackageFlags(PKG_InMemoryOnly))
+					{
+						InMemoryPackages.AddUnique(RealPackage);
+					}
+					else
+					{
+						PackagesToReload.AddUnique(RealPackage);
+					}
+				}
+
+				// Ask the user whether dirty packages should be reloaded.
+				if (bInteractive && DirtyPackages.Num() > 0)
+				{
+					FTextBuilder ReloadDirtyPackagesMsgBuilder;
+					ReloadDirtyPackagesMsgBuilder.AppendLine(NSLOCTEXT("UnrealEd", "ShouldReloadDirtyPackagesHeader", "The following packages have been modified:"));
+					{
+						ReloadDirtyPackagesMsgBuilder.Indent();
+						for (UPackage* DirtyPackage : DirtyPackages)
+						{
+							ReloadDirtyPackagesMsgBuilder.AppendLine(DirtyPackage->GetFName());
+						}
+						ReloadDirtyPackagesMsgBuilder.Unindent();
+					}
+					ReloadDirtyPackagesMsgBuilder.AppendLine(NSLOCTEXT("UnrealEd", "ShouldReloadDirtyPackagesFooter", "Would you like to reload these packages? This will revert any changes you have made."));
+
+					if (FMessageDialog::Open(EAppMsgType::YesNo, ReloadDirtyPackagesMsgBuilder.ToText()) == EAppReturnType::Yes)
+					{
+						for (UPackage* DirtyPackage : DirtyPackages)
+						{
+							DirtyPackage->SetDirtyFlag(false);
+							PackagesToReload.AddUnique(DirtyPackage);
+						}
+						DirtyPackages.Reset();
+					}
+				}
+			}
+
+			// Inform the user that dirty packages won't be reloaded.
+			if (DirtyPackages.Num() > 0)
+			{
+				if (!ErrorMessageBuilder.IsEmpty())
+				{
+					ErrorMessageBuilder.AppendLine();
+				}
+
+				ErrorMessageBuilder.AppendLine(NSLOCTEXT("UnrealEd", "Error_ReloadDirtyPackagesHeader", "The following packages have been modified and cannot be reloaded:"));
+				{
+					ErrorMessageBuilder.Indent();
+					for (UPackage* DirtyPackage : DirtyPackages)
+					{
+						ErrorMessageBuilder.AppendLine(DirtyPackage->GetFName());
+					}
+					ErrorMessageBuilder.Unindent();
+				}
+				ErrorMessageBuilder.AppendLine(NSLOCTEXT("UnrealEd", "Error_ReloadDirtyPackagesFooter", "Saving these packages will allow them to be reloaded."));
+			}
+
+			// Inform the user that in-memory packages won't be reloaded.
+			if (InMemoryPackages.Num() > 0)
+			{
+				if (!ErrorMessageBuilder.IsEmpty())
+				{
+					ErrorMessageBuilder.AppendLine();
+				}
+
+				ErrorMessageBuilder.AppendLine(NSLOCTEXT("UnrealEd", "Error_ReloadInMemoryPackagesHeader", "The following packages are in-memory only and cannot be reloaded:"));
+				{
+					ErrorMessageBuilder.Indent();
+					for (UPackage* InMemoryPackage : InMemoryPackages)
+					{
+						ErrorMessageBuilder.AppendLine(InMemoryPackage->GetFName());
+					}
+					ErrorMessageBuilder.Unindent();
+				}
+			}
+		}
+
+		// Check to see if we need to reload the current world.
+		FName WorldNameToReload;
+		PackagesToReload.RemoveAll([&](UPackage* PackageToReload) -> bool
+		{
+			// Is this the currently loaded world? If so, we just reset the current world and load it again at the end rather than let it go 
+			// through ReloadPackage (which doesn't work for the editor due to some assumptions it makes about worlds, and their lifetimes).
+			if (UWorld* EditorWorld = GEditor->GetEditorWorldContext().World())
+			{
+				const bool bIsCurrentWorld = EditorWorld->GetOutermost() == PackageToReload;
+				if (bIsCurrentWorld)
+				{
+					WorldNameToReload = *EditorWorld->GetPathName();
+					return true; // remove world package
+				}
+			}
+			return false; // leave non-world package
+		});
+
+		// Unload the current world (if needed).
+		if (!WorldNameToReload.IsNone())
+		{
+			GEditor->CreateNewMapForEditing();
+		}
+
+		if (PackagesToReload.Num() > 0)
+		{
+			const FScopedBusyCursor BusyCursor;
+
+			// We need to sort the packages to reload so that dependencies are reloaded before the assets that depend on them
+			::SortPackagesForReload(PackagesToReload);
+
+			// Remove potential references to to-be deleted objects from the global selection set.
+			GEditor->GetSelectedObjects()->DeselectAll();
+
+			// Detach all components while loading a package.
+			// This is necessary for the cases where the load replaces existing objects which may be referenced by the attached components.
+			FGlobalComponentReregisterContext ReregisterContext;
+
+			bool bScriptPackageWasReloaded = false;
+			TArray<FReloadPackageData> PackagesToReloadData;
+			PackagesToReloadData.Reserve(PackagesToReload.Num());
+			for (UPackage* PackageToReload : PackagesToReload)
+			{
+				bScriptPackageWasReloaded |= PackageToReload->HasAnyPackageFlags(PKG_ContainsScript);
+				PackagesToReloadData.Emplace(PackageToReload, LOAD_None);
+			}
+
+			TArray<UPackage*> ReloadedPackages;
+			::ReloadPackages(PackagesToReloadData, ReloadedPackages, 500);
+
+			TArray<UPackage*> FailedPackages;
+			for (int32 PackageIndex = 0; PackageIndex < PackagesToReload.Num(); ++PackageIndex)
+			{
+				UPackage* ExistingPackage = PackagesToReload[PackageIndex];
+				UPackage* ReloadedPackage = ReloadedPackages[PackageIndex];
+
+				if (ReloadedPackage)
+				{
+					bScriptPackageWasReloaded |= ReloadedPackage->HasAnyPackageFlags(PKG_ContainsScript);
+					bResult = true;
+				}
+				else
+				{
+					FailedPackages.Add(ExistingPackage);
+				}
+			}
+
+			// Inform the user of any packages that failed to reload.
+			if (FailedPackages.Num() > 0)
+			{
+				if (!ErrorMessageBuilder.IsEmpty())
+				{
+					ErrorMessageBuilder.AppendLine();
+				}
+
+				ErrorMessageBuilder.AppendLine(NSLOCTEXT("UnrealEd", "Error_ReloadFailedPackagesHeader", "The following packages failed to reload:"));
+				{
+					ErrorMessageBuilder.Indent();
+					for (UPackage* FailedPackage : FailedPackages)
+					{
+						ErrorMessageBuilder.AppendLine(FailedPackage->GetFName());
+					}
+					ErrorMessageBuilder.Unindent();
+				}
+			}
+
+			// Update the actor browser if a script package was reloaded.
+			if (bScriptPackageWasReloaded)
+			{
+				GEditor->BroadcastClassPackageLoadedOrUnloaded();
+			}
+		}
+
+		// Load the previous world (if needed).
+		if (!WorldNameToReload.IsNone())
+		{
+			TArray<FName> WorldNamesToReload;
+			WorldNamesToReload.Add(WorldNameToReload);
+			FAssetEditorManager::Get().OpenEditorsForAssets(WorldNamesToReload);
+		}
+
+		OutErrorMessage = ErrorMessageBuilder.ToText();
+
+		return bResult;
+	}
+
+
 	/**
 	 * Wrapper method for multiple objects at once.
 	 *
diff --git a/Engine/Source/Editor/UnrealEd/Private/PlayLevel.cpp b/Engine/Source/Editor/UnrealEd/Private/PlayLevel.cpp
--- a/Engine/Source/Editor/UnrealEd/Private/PlayLevel.cpp
+++ b/Engine/Source/Editor/UnrealEd/Private/PlayLevel.cpp
@@ -149,6 +149,13 @@ public:
 
 void UEditorEngine::EndPlayMap()
 {
+	if ( bIsEndingPlay )
+	{
+		return;
+	}
+
+	TGuardValue<bool> GuardIsEndingPlay(bIsEndingPlay, true);
+
 	FlushAsyncLoading();
 
 	// Monitoring when PIE corrupts references between the World and the PIE generated World for UE-20486
diff --git a/Engine/Source/Editor/UnrealEd/Private/SAdvancedPreviewDetailsTab.cpp b/Engine/Source/Editor/UnrealEd/Private/SAdvancedPreviewDetailsTab.cpp
--- a/Engine/Source/Editor/UnrealEd/Private/SAdvancedPreviewDetailsTab.cpp
+++ b/Engine/Source/Editor/UnrealEd/Private/SAdvancedPreviewDetailsTab.cpp
@@ -25,6 +25,8 @@ SAdvancedPreviewDetailsTab::SAdvancedPreviewDetailsTab()
 		RefreshDelegate = DefaultSettings->OnAssetViewerSettingsChanged().AddRaw(this, &SAdvancedPreviewDetailsTab::OnAssetViewerSettingsRefresh);
 			
 		AddRemoveProfileDelegate = DefaultSettings->OnAssetViewerProfileAddRemoved().AddLambda([this]() { this->Refresh(); });
+
+		PostUndoDelegate = DefaultSettings->OnAssetViewerSettingsPostUndo().AddRaw(this, &SAdvancedPreviewDetailsTab::OnAssetViewerSettingsPostUndo);
 	}
 
 	PerProjectSettings = GetMutableDefault<UEditorPerProjectUserSettings>();
@@ -38,6 +40,7 @@ SAdvancedPreviewDetailsTab::~SAdvancedPreviewDetailsTab()
 	{
 		DefaultSettings->OnAssetViewerSettingsChanged().Remove(RefreshDelegate);
 		DefaultSettings->OnAssetViewerProfileAddRemoved().Remove(AddRemoveProfileDelegate);
+		DefaultSettings->OnAssetViewerSettingsPostUndo().Remove(PostUndoDelegate);
 		DefaultSettings->Save();
 	}
 }
@@ -257,4 +260,10 @@ void SAdvancedPreviewDetailsTab::Refresh()
 	UpdateSettingsView();
 }
 
-#undef LOCTEXT_NAMESPACE
+void SAdvancedPreviewDetailsTab::OnAssetViewerSettingsPostUndo()
+{
+	Refresh();
+	PreviewScenePtr.Pin()->UpdateScene(DefaultSettings->Profiles[ProfileIndex]);
+}
+
+#undef LOCTEXT_NAMESPACE
\ No newline at end of file
diff --git a/Engine/Source/Editor/UnrealEd/Private/SClassPickerDialog.cpp b/Engine/Source/Editor/UnrealEd/Private/SClassPickerDialog.cpp
--- a/Engine/Source/Editor/UnrealEd/Private/SClassPickerDialog.cpp
+++ b/Engine/Source/Editor/UnrealEd/Private/SClassPickerDialog.cpp
@@ -235,7 +235,12 @@ TSharedRef<ITableRow> SClassPickerDialog::GenerateListRow(TSharedPtr<FClassPicke
 			+SHorizontalBox::Slot()
 			.AutoWidth()
 			[
-				FEditorClassUtils::GetDocumentationLinkWidget(ItemClass)
+				SNew(SBox)
+				.WidthOverride(32)
+				.HeightOverride(32)
+				[
+					FEditorClassUtils::GetDocumentationLinkWidget(ItemClass)
+				]
 			]
 		]
 	];
diff --git a/Engine/Source/Editor/UnrealEd/Private/SCurveEditor.cpp b/Engine/Source/Editor/UnrealEd/Private/SCurveEditor.cpp
--- a/Engine/Source/Editor/UnrealEd/Private/SCurveEditor.cpp
+++ b/Engine/Source/Editor/UnrealEd/Private/SCurveEditor.cpp
@@ -448,6 +448,8 @@ void SCurveEditor::Construct(const FArguments& InArgs)
 	{
 		GEditor->RegisterForUndo(this);
 	}
+
+	FCoreUObjectDelegates::OnObjectPropertyChanged.AddRaw(this, &SCurveEditor::OnObjectPropertyChanged);
 }
 
 FText SCurveEditor::GetIsCurveVisibleToolTip(TSharedPtr<FCurveViewModel> CurveViewModel) const
@@ -537,6 +539,8 @@ SCurveEditor::~SCurveEditor()
 	{
 		GEditor->UnregisterForUndo(this);
 	}
+
+	FCoreUObjectDelegates::OnObjectPropertyChanged.RemoveAll(this);
 }
 
 TSharedRef<SWidget> SCurveEditor::CreateCurveSelectionWidget() const
@@ -3578,6 +3582,14 @@ void SCurveEditor::RedoAction()
 	GEditor->RedoTransaction();
 }
 
+void SCurveEditor::OnObjectPropertyChanged(UObject* Object, FPropertyChangedEvent& PropertyChangedEvent)
+{
+	if ( CurveOwner && CurveOwner->GetOwners().Contains(Object) )
+	{
+		ValidateSelection();
+	}
+}
+
 void SCurveEditor::PostUndo(bool bSuccess)
 {
 	ValidateSelection();
diff --git a/Engine/Source/Editor/UnrealEd/Private/SSocketManager.cpp b/Engine/Source/Editor/UnrealEd/Private/SSocketManager.cpp
--- a/Engine/Source/Editor/UnrealEd/Private/SSocketManager.cpp
+++ b/Engine/Source/Editor/UnrealEd/Private/SSocketManager.cpp
@@ -487,46 +487,43 @@ void SSocketManager::DeleteSelectedSocket()
 
 void SSocketManager::RefreshSocketList()
 {
-	if( StaticMesh.IsValid() )
+	// The static mesh might not be the same one we built the SocketListView with
+	// check it here and update it if necessary.
+	TSharedPtr<IStaticMeshEditor> StaticMeshEditorPinned = StaticMeshEditorPtr.Pin();
+	if (StaticMeshEditorPinned.IsValid())
 	{
-		// The static mesh might not be the same one we built the SocketListView with
-		// check it here and update it if necessary.
 		bool bIsSameStaticMesh = true;
-		TSharedPtr<IStaticMeshEditor> StaticMeshEditorPinned = StaticMeshEditorPtr.Pin();
-		if (StaticMeshEditorPinned.IsValid())
+		UStaticMesh* CurrentStaticMesh = StaticMeshEditorPinned->GetStaticMesh();
+		if (!StaticMesh.IsValid() || CurrentStaticMesh != StaticMesh.Get())
 		{
-			UStaticMesh* CurrentStaticMesh = StaticMeshEditorPinned->GetStaticMesh();
-			if (CurrentStaticMesh != StaticMesh.Get())
+			StaticMesh = CurrentStaticMesh;
+			bIsSameStaticMesh = false;
+		}
+		// Only rebuild the socket list if it differs from the static meshes socket list
+		// This is done so that an undo on a socket property doesn't cause the selected
+		// socket to be de-selected, thus hiding the socket properties on the detail view.
+		// NB: Also force a rebuild if the underlying StaticMesh has been changed.
+		if (StaticMesh->Sockets.Num() != SocketList.Num() || !bIsSameStaticMesh)
+		{
+			SocketList.Empty();
+			for (int32 i = 0; i < StaticMesh->Sockets.Num(); i++)
 			{
-				StaticMesh = CurrentStaticMesh;
-				bIsSameStaticMesh = false;
+				UStaticMeshSocket* Socket = StaticMesh->Sockets[i];
+				SocketList.Add(MakeShareable(new SocketListItem(Socket)));
 			}
-			// Only rebuild the socket list if it differs from the static meshes socket list
-			// This is done so that an undo on a socket property doesn't cause the selected
-			// socket to be de-selected, thus hiding the socket properties on the detail view.
-			// NB: Also force a rebuild if the underlying StaticMesh has been changed.
-			if (StaticMesh->Sockets.Num() != SocketList.Num() || !bIsSameStaticMesh)
-			{
-				SocketList.Empty();
-				for (int32 i = 0; i < StaticMesh->Sockets.Num(); i++)
-				{
-					UStaticMeshSocket* Socket = StaticMesh->Sockets[i];
-					SocketList.Add(MakeShareable(new SocketListItem(Socket)));
-				}
 
-				SocketListView->RequestListRefresh();
-			}
-
-			// Set the socket on the detail view to keep it in sync with the sockets properties
-			if (SocketListView->GetSelectedItems().Num())
-			{
-				TArray< UObject* > ObjectList;
-				ObjectList.Add(SocketListView->GetSelectedItems()[0]->Socket);
-				SocketDetailsView->SetObjects(ObjectList, true);
-			}
+			SocketListView->RequestListRefresh();
+		}
 
-			StaticMeshEditorPinned->RefreshViewport();
+		// Set the socket on the detail view to keep it in sync with the sockets properties
+		if (SocketListView->GetSelectedItems().Num())
+		{
+			TArray< UObject* > ObjectList;
+			ObjectList.Add(SocketListView->GetSelectedItems()[0]->Socket);
+			SocketDetailsView->SetObjects(ObjectList, true);
 		}
+
+		StaticMeshEditorPinned->RefreshViewport();
 	}
 	else
 	{
diff --git a/Engine/Source/Editor/UnrealEd/Private/Settings/SettingsClasses.cpp b/Engine/Source/Editor/UnrealEd/Private/Settings/SettingsClasses.cpp
--- a/Engine/Source/Editor/UnrealEd/Private/Settings/SettingsClasses.cpp
+++ b/Engine/Source/Editor/UnrealEd/Private/Settings/SettingsClasses.cpp
@@ -10,11 +10,13 @@
 #include "EditorStyleSettings.h"
 #include "AI/Navigation/NavigationSystem.h"
 #include "Model.h"
+#include "ISourceControlModule.h"
 #include "Settings/ContentBrowserSettings.h"
 #include "Settings/DestructableMeshEditorSettings.h"
 #include "Settings/LevelEditorPlaySettings.h"
 #include "Settings/LevelEditorViewportSettings.h"
-#include "ISourceControlModule.h"
+#include "Settings/EditorProjectSettings.h"
+#include "Settings/ClassViewerSettings.h"
 #include "Settings/EditorExperimentalSettings.h"
 #include "Settings/EditorLoadingSavingSettings.h"
 #include "Settings/EditorMiscSettings.h"
@@ -29,10 +31,7 @@
 #include "CrashReporterSettings.h"
 #include "AutoReimport/AutoReimportUtilities.h"
 
-#include "Settings/EditorProjectSettings.h"
-
 #include "SourceCodeNavigation.h"
-
 #include "Developer/BlueprintProfiler/Public/BlueprintProfilerModule.h"
 
 #define LOCTEXT_NAMESPACE "SettingsClasses"
@@ -64,6 +63,32 @@ void UContentBrowserSettings::PostEditChangeProperty( struct FPropertyChangedEve
 	SettingChangedEvent.Broadcast(Name);
 }
 
+/* UClassViewerSettings interface
+*****************************************************************************/
+
+UClassViewerSettings::FSettingChangedEvent UClassViewerSettings::SettingChangedEvent;
+
+UClassViewerSettings::UClassViewerSettings(const FObjectInitializer& ObjectInitializer)
+	: Super(ObjectInitializer)
+{
+}
+
+
+void UClassViewerSettings::PostEditChangeProperty(struct FPropertyChangedEvent& PropertyChangedEvent)
+{
+	Super::PostEditChangeProperty(PropertyChangedEvent);
+
+	const FName Name = (PropertyChangedEvent.Property != nullptr)
+		? PropertyChangedEvent.Property->GetFName()
+		: NAME_None;
+
+	if (!FUnrealEdMisc::Get().IsDeletePreferences())
+	{
+		SaveConfig();
+	}
+
+	SettingChangedEvent.Broadcast(Name);
+}
 
 /* UDestructableMeshEditorSettings interface
  *****************************************************************************/
diff --git a/Engine/Source/Editor/UnrealEd/Private/SkeletalMeshEdit.cpp b/Engine/Source/Editor/UnrealEd/Private/SkeletalMeshEdit.cpp
--- a/Engine/Source/Editor/UnrealEd/Private/SkeletalMeshEdit.cpp
+++ b/Engine/Source/Editor/UnrealEd/Private/SkeletalMeshEdit.cpp
@@ -106,13 +106,11 @@ UAnimSequence * UEditorEngine::ImportFbxAnimation( USkeleton* Skeleton, UObject*
 bool UEditorEngine::ReimportFbxAnimation( USkeleton* Skeleton, UAnimSequence* AnimSequence, UFbxAnimSequenceImportData* ImportData, const TCHAR* InFilename)
 {
 	check(Skeleton);
-
+	bool bResult = true;
 	GWarn->BeginSlowTask( LOCTEXT("ImportingFbxAnimations", "Importing FBX animations"), true );
 
 	UnFbx::FFbxImporter* FbxImporter = UnFbx::FFbxImporter::GetInstance();
-	// logger for all error/warnings
-	// this one prints all messages that are stored in FFbxImporter	
-	UnFbx::FFbxLoggerSetter Logger(FbxImporter);	
+	
 	const bool bPrevImportMorph = (AnimSequence->RawCurveData.FloatCurves.Num() > 0) ;
 
 	if ( ImportData )
@@ -135,7 +133,7 @@ bool UEditorEngine::ReimportFbxAnimation( USkeleton* Skeleton, UAnimSequence* An
 	{
 		// Log the error message and fail the import.
 		FbxImporter->FlushToTokenizedErrorMessage(EMessageSeverity::Error);
-
+		bResult = false;
 	}
 	else
 	{
@@ -152,76 +150,74 @@ bool UEditorEngine::ReimportFbxAnimation( USkeleton* Skeleton, UAnimSequence* An
 		if (!SkeletonRoot)
 		{
 			FbxImporter->AddTokenizedErrorMessage(FTokenizedMessage::Create(EMessageSeverity::Error, FText::Format(LOCTEXT("Error_CouldNotFindFbxTrack", "Mesh contains {0} bone as root but animation doesn't contain the root track.\nImport failed."), FText::FromName(Skeleton->GetReferenceSkeleton().GetBoneName(0)))), FFbxErrors::Animation_CouldNotFindTrack);
-
-			FbxImporter->ReleaseScene();
-			GWarn->EndSlowTask();
-			return false;
+			bResult = false;
 		}
 
-		// for now import all the time?
-		bool bImportMorphTracks = true;
-		// Check for blend shape curves that are not skinned.  Unskinned geometry can still contain morph curves
-		if( bImportMorphTracks )
+		if (bResult)
 		{
-			TArray<FbxNode*> MeshNodes;
-			FbxImporter->FillFbxMeshArray( FbxImporter->Scene->GetRootNode(), MeshNodes, FbxImporter );
-
-			for( int32 NodeIndex = 0; NodeIndex < MeshNodes.Num(); ++NodeIndex )
+			// for now import all the time?
+			bool bImportMorphTracks = true;
+			// Check for blend shape curves that are not skinned.  Unskinned geometry can still contain morph curves
+			if (bImportMorphTracks)
 			{
-				// Its possible the nodes already exist so make sure they are only added once
-				FBXMeshNodeArray.AddUnique( MeshNodes[NodeIndex] );
+				TArray<FbxNode*> MeshNodes;
+				FbxImporter->FillFbxMeshArray(FbxImporter->Scene->GetRootNode(), MeshNodes, FbxImporter);
+
+				for (int32 NodeIndex = 0; NodeIndex < MeshNodes.Num(); ++NodeIndex)
+				{
+					// Its possible the nodes already exist so make sure they are only added once
+					FBXMeshNodeArray.AddUnique(MeshNodes[NodeIndex]);
+				}
 			}
-		}
 
-		TArray<FbxNode*> SortedLinks;
-		FbxImporter->RecursiveBuildSkeleton(SkeletonRoot, SortedLinks);
+			TArray<FbxNode*> SortedLinks;
+			FbxImporter->RecursiveBuildSkeleton(SkeletonRoot, SortedLinks);
 
-		if(SortedLinks.Num() == 0)
-		{
-			FbxImporter->AddTokenizedErrorMessage(FTokenizedMessage::Create(EMessageSeverity::Warning, LOCTEXT("Error_CouldNotBuildValidSkeleton", "Could not create a valid skeleton from the import data that matches the given Skeletal Mesh.  Check the bone names of both the Skeletal Mesh for this AnimSet and the animation data you are trying to import.")), FFbxErrors::Animation_CouldNotBuildSkeleton);
-		}
-		else
-		{
-			check(ImportData);
-
-			// find the correct animation based on import data
-			FbxAnimStack* CurAnimStack = nullptr;
-			
-			//ignore the source animation name if there's only one animation in the file.
-			//this is to make it easier for people who use content creation programs that only export one animation and/or ones that don't allow naming animations			
-			if (FbxImporter->Scene->GetSrcObjectCount(FbxCriteria::ObjectType(FbxAnimStack::ClassId)) > 1 && !ImportData->SourceAnimationName.IsEmpty())
+			if (SortedLinks.Num() == 0)
 			{
-				CurAnimStack = FbxCast<FbxAnimStack>(FbxImporter->Scene->FindSrcObject(FbxCriteria::ObjectType(FbxAnimStack::ClassId), TCHAR_TO_UTF8(*ImportData->SourceAnimationName), 0));
+				FbxImporter->AddTokenizedErrorMessage(FTokenizedMessage::Create(EMessageSeverity::Warning, LOCTEXT("Error_CouldNotBuildValidSkeleton", "Could not create a valid skeleton from the import data that matches the given Skeletal Mesh.  Check the bone names of both the Skeletal Mesh for this AnimSet and the animation data you are trying to import.")), FFbxErrors::Animation_CouldNotBuildSkeleton);
 			}
 			else
 			{
-				CurAnimStack = FbxCast<FbxAnimStack>(FbxImporter->Scene->GetSrcObject(FbxCriteria::ObjectType(FbxAnimStack::ClassId), 0));
-			}
-			
-			if (CurAnimStack)
-			{
-				// set current anim stack
-				int32 ResampleRate = DEFAULT_SAMPLERATE;
-				if (FbxImporter->ImportOptions->bResample)
+				check(ImportData);
+
+				// find the correct animation based on import data
+				FbxAnimStack* CurAnimStack = nullptr;
+
+				//ignore the source animation name if there's only one animation in the file.
+				//this is to make it easier for people who use content creation programs that only export one animation and/or ones that don't allow naming animations			
+				if (FbxImporter->Scene->GetSrcObjectCount(FbxCriteria::ObjectType(FbxAnimStack::ClassId)) > 1 && !ImportData->SourceAnimationName.IsEmpty())
 				{
-					ResampleRate = FbxImporter->GetMaxSampleRate(SortedLinks, FBXMeshNodeArray);
+					CurAnimStack = FbxCast<FbxAnimStack>(FbxImporter->Scene->FindSrcObject(FbxCriteria::ObjectType(FbxAnimStack::ClassId), TCHAR_TO_UTF8(*ImportData->SourceAnimationName), 0));
 				}
-				FbxTimeSpan AnimTimeSpan = FbxImporter->GetAnimationTimeSpan(SortedLinks[0], CurAnimStack, ResampleRate);
-				// for now it's not importing morph - in the future, this should be optional or saved with asset
-				if (FbxImporter->ValidateAnimStack(SortedLinks, FBXMeshNodeArray, CurAnimStack, ResampleRate, bImportMorphTracks, AnimTimeSpan))
+				else
 				{
-					FbxImporter->ImportAnimation( Skeleton, AnimSequence, Filename, SortedLinks, FBXMeshNodeArray, CurAnimStack, ResampleRate, AnimTimeSpan);
+					CurAnimStack = FbxCast<FbxAnimStack>(FbxImporter->Scene->GetSrcObject(FbxCriteria::ObjectType(FbxAnimStack::ClassId), 0));
 				}
-			}
-			else
-			{
-				// no track is found
 
-				FbxImporter->AddTokenizedErrorMessage(FTokenizedMessage::Create(EMessageSeverity::Error, LOCTEXT("Error_CouldNotFindTrack", "Could not find needed track.")), FFbxErrors::Animation_CouldNotFindTrack);
+				if (CurAnimStack)
+				{
+					// set current anim stack
+					int32 ResampleRate = DEFAULT_SAMPLERATE;
+					if (FbxImporter->ImportOptions->bResample)
+					{
+						ResampleRate = FbxImporter->GetMaxSampleRate(SortedLinks, FBXMeshNodeArray);
+					}
+					FbxTimeSpan AnimTimeSpan = FbxImporter->GetAnimationTimeSpan(SortedLinks[0], CurAnimStack, ResampleRate);
+					// for now it's not importing morph - in the future, this should be optional or saved with asset
+					if (FbxImporter->ValidateAnimStack(SortedLinks, FBXMeshNodeArray, CurAnimStack, ResampleRate, bImportMorphTracks, AnimTimeSpan))
+					{
+						FbxImporter->ImportAnimation(Skeleton, AnimSequence, Filename, SortedLinks, FBXMeshNodeArray, CurAnimStack, ResampleRate, AnimTimeSpan);
+					}
+				}
+				else
+				{
+					// no track is found
+
+					FbxImporter->AddTokenizedErrorMessage(FTokenizedMessage::Create(EMessageSeverity::Error, LOCTEXT("Error_CouldNotFindTrack", "Could not find needed track.")), FFbxErrors::Animation_CouldNotFindTrack);
 
-				FbxImporter->ReleaseScene();
-				GWarn->EndSlowTask();
-				return false;
+					bResult = false;
+				}
 			}
 		}
 	}
@@ -230,7 +226,7 @@ bool UEditorEngine::ReimportFbxAnimation( USkeleton* Skeleton, UAnimSequence* An
 	FbxImporter->ReleaseScene();
 	GWarn->EndSlowTask();
 
-	return true;
+	return bResult;
 }
 
 // The Unroll filter expects only rotation curves, we need to walk the scene and extract the
diff --git a/Engine/Source/Editor/UnrealEd/Private/StaticMeshEdit.cpp b/Engine/Source/Editor/UnrealEd/Private/StaticMeshEdit.cpp
--- a/Engine/Source/Editor/UnrealEd/Private/StaticMeshEdit.cpp
+++ b/Engine/Source/Editor/UnrealEd/Private/StaticMeshEdit.cpp
@@ -21,6 +21,8 @@
 #include "PhysicsEngine/BoxElem.h"
 #include "PhysicsEngine/SphereElem.h"
 #include "PhysicsEngine/BodySetup.h"
+#include "FbxImporter.h"
+
 
 bool GBuildStaticMeshCollision = 1;
 
@@ -868,12 +870,13 @@ struct ExistingStaticMeshData
 };
 
 
-ExistingStaticMeshData* SaveExistingStaticMeshData(UStaticMesh* ExistingMesh, bool bSaveMaterials)
+ExistingStaticMeshData* SaveExistingStaticMeshData(UStaticMesh* ExistingMesh, UnFbx::FBXImportOptions* ImportOptions)
 {
 	struct ExistingStaticMeshData* ExistingMeshDataPtr = NULL;
 
 	if (ExistingMesh)
 	{
+		bool bSaveMaterials = !ImportOptions->bImportMaterials;
 		ExistingMeshDataPtr = new ExistingStaticMeshData();
 
 		ExistingMeshDataPtr->ImportVersion = ExistingMesh->ImportVersion;
@@ -925,6 +928,11 @@ ExistingStaticMeshData* SaveExistingStaticMeshData(UStaticMesh* ExistingMesh, bo
 				}
 			}
 
+			//The normals, tangent and tangent space build setting depend of the import options, so we cannot restore them, we have to set them when re-importing
+			ExistingMesh->SourceModels[i].BuildSettings.bRecomputeNormals = ImportOptions->NormalImportMethod == FBXNIM_ComputeNormals;
+			ExistingMesh->SourceModels[i].BuildSettings.bRecomputeTangents = ImportOptions->NormalImportMethod != FBXNIM_ImportNormalsAndTangents;
+			ExistingMesh->SourceModels[i].BuildSettings.bUseMikkTSpace = (ImportOptions->NormalGenerationMethod == EFBXNormalGenerationMethod::MikkTSpace) && (!ImportOptions->ShouldImportNormals() || !ImportOptions->ShouldImportTangents());
+
 			ExistingMeshDataPtr->ExistingLODData[i].ExistingBuildSettings = ExistingMesh->SourceModels[i].BuildSettings;
 			ExistingMeshDataPtr->ExistingLODData[i].ExistingReductionSettings = ExistingMesh->SourceModels[i].ReductionSettings;
 			ExistingMeshDataPtr->ExistingLODData[i].ExistingScreenSize = ExistingMesh->SourceModels[i].ScreenSize;
diff --git a/Engine/Source/Editor/UnrealEd/Private/Toolkits/AssetEditorManager.cpp b/Engine/Source/Editor/UnrealEd/Private/Toolkits/AssetEditorManager.cpp
--- a/Engine/Source/Editor/UnrealEd/Private/Toolkits/AssetEditorManager.cpp
+++ b/Engine/Source/Editor/UnrealEd/Private/Toolkits/AssetEditorManager.cpp
@@ -15,6 +15,7 @@
 #include "Toolkits/AssetEditorToolkit.h"
 #include "Toolkits/SimpleAssetEditor.h"
 #include "LevelEditor.h"
+#include "PackageReload.h"
 #include "EngineAnalytics.h"
 #include "AnalyticsEventAttribute.h"
 #include "Interfaces/IAnalyticsProvider.h"
@@ -52,10 +53,14 @@ FAssetEditorManager::FAssetEditorManager()
 
 	TickDelegate = FTickerDelegate::CreateRaw(this, &FAssetEditorManager::HandleTicker);
 	FTicker::GetCoreTicker().AddTicker(TickDelegate, 1.f);
+
+	FCoreUObjectDelegates::OnPackageReloaded.AddRaw(this, &FAssetEditorManager::HandlePackageReloaded);
 }
 
 void FAssetEditorManager::OnExit()
 {
+	FCoreUObjectDelegates::OnPackageReloaded.RemoveAll(this);
+
 	SaveOpenAssetEditors(true);
 
 	TGuardValue<bool> GuardOnShutdown(bSavingOnShutdown, true);
@@ -704,6 +709,38 @@ void FAssetEditorManager::SaveOpenAssetEditors(bool bOnShutdown)
 	}
 }
 
+void FAssetEditorManager::HandlePackageReloaded(const EPackageReloadPhase InPackageReloadPhase, FPackageReloadedEvent* InPackageReloadedEvent)
+{
+	if (InPackageReloadPhase == EPackageReloadPhase::PrePackageFixup)
+	{
+		TArray<UObject*> OldAssets;
+		TArray<UObject*> NewAssets;
+
+		for (auto AssetEditorPair : OpenedAssets)
+		{
+			UObject* NewAsset = nullptr;
+			if (InPackageReloadedEvent->GetRepointedObject(AssetEditorPair.Key, NewAsset))
+			{
+				OldAssets.Add(AssetEditorPair.Key);
+				if (NewAsset)
+				{
+					NewAssets.Add(NewAsset);
+				}
+			}
+		}
+
+		for (UObject* OldAsset : OldAssets)
+		{
+			CloseAllEditorsForAsset(OldAsset);
+		}
+
+		for (UObject* NewAsset : NewAssets)
+		{
+			OpenEditorForAsset(NewAsset);
+		}
+	}
+}
+
 void FAssetEditorManager::OpenEditorsForAssets(const TArray<FString>& AssetsToOpen)
 {
 	for (const FString& AssetName : AssetsToOpen)
diff --git a/Engine/Source/Editor/UnrealEd/Private/Toolkits/AssetEditorToolkit.cpp b/Engine/Source/Editor/UnrealEd/Private/Toolkits/AssetEditorToolkit.cpp
--- a/Engine/Source/Editor/UnrealEd/Private/Toolkits/AssetEditorToolkit.cpp
+++ b/Engine/Source/Editor/UnrealEd/Private/Toolkits/AssetEditorToolkit.cpp
@@ -111,7 +111,7 @@ void FAssetEditorToolkit::InitAssetEditor( const EToolkitMode::Type Mode, const
 			const UEditorStyleSettings* StyleSettings = GetDefault<UEditorStyleSettings>();
 
 			FName PlaceholderId(TEXT("StandaloneToolkit"));
-			FTabManager::FSearchPreference* SearchPreference = nullptr;
+			TSharedPtr<FTabManager::FSearchPreference> SearchPreference = nullptr;
 			if ( StyleSettings->AssetEditorOpenLocation == EAssetEditorOpenLocation::Default )
 			{
 				// Work out where we should create this asset editor
@@ -124,37 +124,37 @@ void FAssetEditorToolkit::InitAssetEditor( const EToolkitMode::Type Mode, const
 					);
 
 				PlaceholderId = ( SavedAssetEditorToolkitTabLocation == EAssetEditorToolkitTabLocation::Docked ) ? TEXT("DockedToolkit") : TEXT("StandaloneToolkit");
-				SearchPreference = new FTabManager::FLiveTabSearch();
+				SearchPreference = MakeShareable(new FTabManager::FLiveTabSearch());
 			}
 			else if ( StyleSettings->AssetEditorOpenLocation == EAssetEditorOpenLocation::NewWindow )
 			{
 				PlaceholderId = TEXT("StandaloneToolkit");
-				SearchPreference = new FTabManager::FRequireClosedTab();
+				SearchPreference = MakeShareable(new FTabManager::FRequireClosedTab());
 			}
 			else if ( StyleSettings->AssetEditorOpenLocation == EAssetEditorOpenLocation::MainWindow )
 			{
 				PlaceholderId = TEXT("DockedToolkit");
-				SearchPreference = new FTabManager::FLiveTabSearch(TEXT("LevelEditor"));
+				SearchPreference = MakeShareable(new FTabManager::FLiveTabSearch(TEXT("LevelEditor")));
 			}
 			else if ( StyleSettings->AssetEditorOpenLocation == EAssetEditorOpenLocation::ContentBrowser )
 			{
 				PlaceholderId = TEXT("DockedToolkit");
-				SearchPreference = new FTabManager::FLiveTabSearch(TEXT("ContentBrowserTab1"));
+				SearchPreference = MakeShareable(new FTabManager::FLiveTabSearch(TEXT("ContentBrowserTab1")));
 			}
 			else if ( StyleSettings->AssetEditorOpenLocation == EAssetEditorOpenLocation::LastDockedWindowOrNewWindow )
 			{
 				PlaceholderId = TEXT("StandaloneToolkit");
-				SearchPreference = new FTabManager::FLastMajorOrNomadTab(NAME_None);
+				SearchPreference = MakeShareable(new FTabManager::FLastMajorOrNomadTab(NAME_None));
 			}
 			else if ( StyleSettings->AssetEditorOpenLocation == EAssetEditorOpenLocation::LastDockedWindowOrMainWindow )
 			{
 				PlaceholderId = TEXT("StandaloneToolkit");
-				SearchPreference = new FTabManager::FLastMajorOrNomadTab(TEXT("LevelEditor"));
+				SearchPreference = MakeShareable(new FTabManager::FLastMajorOrNomadTab(TEXT("LevelEditor")));
 			}
 			else if ( StyleSettings->AssetEditorOpenLocation == EAssetEditorOpenLocation::LastDockedWindowOrContentBrowser )
 			{
 				PlaceholderId = TEXT("StandaloneToolkit");
-				SearchPreference = new FTabManager::FLastMajorOrNomadTab(TEXT("ContentBrowserTab1"));
+				SearchPreference = MakeShareable(new FTabManager::FLastMajorOrNomadTab(TEXT("ContentBrowserTab1")));
 			}
 			else
 			{
@@ -163,7 +163,14 @@ void FAssetEditorToolkit::InitAssetEditor( const EToolkitMode::Type Mode, const
 			}
 
 			FGlobalTabmanager::Get()->InsertNewDocumentTab(PlaceholderId, *SearchPreference, NewMajorTab.ToSharedRef());
-			delete SearchPreference;
+
+			// Bring the window to front.  The tab manager will not do this for us to avoid intrusive stealing focus behavior
+			// However, here the expectation is that opening an new asset editor is something that should steal focus so the user can see their asset
+			TSharedPtr<SWindow> Window = NewMajorTab->GetParentWindow();
+			if(Window.IsValid())
+			{
+				Window->BringToFront();
+			}
 		}
 
 #if PLATFORM_MAC
